Fix for push model "CanSkipUpdate" with fast arrays and conditional custom delta properties

diff --git a/Engine/Source/Runtime/Engine/Private/DataReplication.cpp b/Engine/Source/Runtime/Engine/Private/DataReplication.cpp
--- a/Engine/Source/Runtime/Engine/Private/DataReplication.cpp
+++ b/Engine/Source/Runtime/Engine/Private/DataReplication.cpp
@@ -22,6 +22,7 @@
 #include "Net/Core/Trace/NetTrace.h"
 #include "ProfilingDebugging/CsvProfiler.h"
 #include "HAL/LowLevelMemStats.h"
+#include "PushModelPerNetDriverState.h"
 
 DECLARE_LLM_MEMORY_STAT(TEXT("NetObjReplicator"), STAT_NetObjReplicatorLLM, STATGROUP_LLMFULL);
 LLM_DEFINE_TAG(NetObjReplicator, NAME_None, TEXT("Networking"), GET_STATFNAME(STAT_NetObjReplicatorLLM), GET_STATFNAME(STAT_NetworkingSummaryLLM));
@@ -282,6 +283,11 @@ public:
 		return RepLayout.GetNumLifetimeCustomDeltaProperties();
 	}
 
+	static uint16 GetLifetimeCustomDeltaPropertyRepIndex(const FRepLayout& RepLayout, const uint16 CustomDeltaPropertyIndex)
+	{
+		return RepLayout.GetLifetimeCustomDeltaPropertyRepIndex(CustomDeltaPropertyIndex);
+	}
+
 	static FProperty* GetLifetimeCustomDeltaProperty(const FRepLayout& RepLayout, const uint16 CustomDeltaPropertyIndex)
 	{
 		return RepLayout.GetLifetimeCustomDeltaProperty(CustomDeltaPropertyIndex);
@@ -872,6 +878,8 @@ void FObjectReplicator::ReceivedNak( int32 NakPacketId )
 	{
 		if (FSendingRepState* SendingRepState = RepState.IsValid() ? RepState->GetSendingRepState() : nullptr)
 		{
+			SendingRepState->CustomDeltaChangeIndex--;
+			
 			// Go over properties tracked with histories, and mark them as needing to be resent.
 			for (int32 i = SendingRepState->HistoryStart; i < SendingRepState->HistoryEnd; ++i)
 			{
@@ -1548,12 +1556,15 @@ static FORCEINLINE FPropertyRetirement** UpdateAckedRetirements(
 	return Rec;
 }
 
-void FObjectReplicator::ReplicateCustomDeltaProperties( FNetBitWriter & Bunch, FReplicationFlags RepFlags )
+void FObjectReplicator::ReplicateCustomDeltaProperties( FNetBitWriter & Bunch, FReplicationFlags RepFlags, bool& bSkippedPropertyCondition)
 {
 	CONDITIONAL_SCOPE_CYCLE_COUNTER(STAT_NetReplicateCustomDeltaPropTime, CVarNetEnableDetailedScopeCounters.GetValueOnAnyThread() > 0);
 
 	check(RepLayout);
 	const FRepLayout& LocalRepLayout = *RepLayout;
+
+	bSkippedPropertyCondition = false;
+
 	const int32 NumLifetimeCustomDeltaProperties = FNetSerializeCB::GetNumLifetimeCustomDeltaProperties(LocalRepLayout);
 
 	if (NumLifetimeCustomDeltaProperties <= 0)
@@ -1609,6 +1620,7 @@ void FObjectReplicator::ReplicateCustomDeltaProperties( FNetBitWriter & Bunch, F
 		if (!ConditionMap[RepCondition])
 		{
 			// We didn't pass the condition so don't replicate us
+			bSkippedPropertyCondition = true;
 			continue;
 		}
 
@@ -1732,6 +1744,7 @@ bool FObjectReplicator::CanSkipUpdate(FReplicationFlags RepFlags)
 			SendingRepState.LastChangelistIndex == RepChangelistState.HistoryEnd &&
 			Connection->ResendAllDataState == EResendAllDataState::None &&
 			!OwningChannel->bForceCompareProperties &&
+			(!GbPushModelSkipUndirtiedFastArrays || (RepChangelistState.CustomDeltaChangeIndex == SendingRepState.CustomDeltaChangeIndex)) &&
 			!RepChangelistState.HasAnyDirtyProperties())
 		))
 	{
@@ -1807,6 +1820,30 @@ bool FObjectReplicator::ReplicateProperties_r( FOutBunch & Bunch, FReplicationFl
 
 	FSendingRepState* SendingRepState = (bUseCheckpointRepState && CheckpointRepState.IsValid()) ? CheckpointRepState->GetSendingRepState() : RepState->GetSendingRepState();
 
+#if WITH_PUSH_MODEL
+	FRepChangelistState* ChangelistState = ChangelistMgr->GetRepChangelistState();
+
+	bool bHasDirtyCustomDeltaProperties = false;
+
+	const UEPushModelPrivate::FPushModelPerNetDriverHandle Handle = ChangelistState->GetPushModelObjectHandle();
+	if (Handle.IsValid())
+	{
+		if (UEPushModelPrivate::FPushModelPerNetDriverState* PushModelState = UEPushModelPrivate::GetPerNetDriverState(Handle))
+		{
+			const uint16 NumLifetimeCustomDeltaProperties = FNetSerializeCB::GetNumLifetimeCustomDeltaProperties(*RepLayout);
+
+			for (uint16 CustomDeltaProperty = 0; CustomDeltaProperty < NumLifetimeCustomDeltaProperties; ++CustomDeltaProperty)
+			{
+				if (PushModelState->IsPropertyDirty(FNetSerializeCB::GetLifetimeCustomDeltaPropertyRepIndex(*RepLayout, CustomDeltaProperty)))
+				{
+					bHasDirtyCustomDeltaProperties = true;
+					break;
+				}
+			}
+		}
+	}
+#endif
+
 	const ERepLayoutResult UpdateResult = FNetSerializeCB::UpdateChangelistMgr(*RepLayout, SendingRepState, *ChangelistMgr, Object, Connection->Driver->ReplicationFrame, RepFlags, OwningChannel->bForceCompareProperties || bUseCheckpointRepState);
 
 	if (UNLIKELY(ERepLayoutResult::FatalError == UpdateResult))
@@ -1826,7 +1863,38 @@ bool FObjectReplicator::ReplicateProperties_r( FOutBunch & Bunch, FReplicationFl
 	);
 
 	// Replicate all the custom delta properties (fast arrays, etc)
-	ReplicateCustomDeltaProperties(Writer, RepFlags);
+
+	{
+#if WITH_PUSH_MODEL
+		const int32 WriterBits = Writer.GetNumBits();
+#endif // WITH_PUSH_MODEL
+
+		bool bSkippedPropertyCondition = false;
+		ReplicateCustomDeltaProperties(Writer, RepFlags, bSkippedPropertyCondition);
+
+#if WITH_PUSH_MODEL
+		if (WriterBits != Writer.GetNumBits())
+		{
+			if (bHasDirtyCustomDeltaProperties)
+			{
+				ChangelistState->CustomDeltaChangeIndex++;
+			}
+
+			SendingRepState->CustomDeltaChangeIndex = ChangelistState->CustomDeltaChangeIndex;
+		}
+		else
+		{
+			// also increment changes if we had to skip any properties for conditionals, since another connection may generate data
+			if (bHasDirtyCustomDeltaProperties && bSkippedPropertyCondition)
+			{
+				ChangelistState->CustomDeltaChangeIndex++;
+
+				// this skipping connection is up to date
+				SendingRepState->CustomDeltaChangeIndex = ChangelistState->CustomDeltaChangeIndex;
+			}
+		}
+#endif //WITH_PUSH_MODEL
+	}
 
 	if ( Connection->ResendAllDataState != EResendAllDataState::None )
 	{
diff --git a/Engine/Source/Runtime/Engine/Private/RepLayout.cpp b/Engine/Source/Runtime/Engine/Private/RepLayout.cpp
--- a/Engine/Source/Runtime/Engine/Private/RepLayout.cpp
+++ b/Engine/Source/Runtime/Engine/Private/RepLayout.cpp
@@ -8212,19 +8212,32 @@ const uint16 FRepLayout::GetNumLifetimeCustomDeltaProperties() const
 	return LifetimeCustomPropertyState.IsValid() ? LifetimeCustomPropertyState->GetNumCustomDeltaProperties() : 0;
 }
 
+const uint16 FRepLayout::GetLifetimeCustomDeltaPropertyRepIndex(const uint16 CustomDeltaPropertyIndex) const
+{
+	checkSlow(LifetimeCustomPropertyState.IsValid());
+
+	return LifetimeCustomPropertyState->GetCustomDeltaProperty(CustomDeltaPropertyIndex).PropertyRepIndex;
+}
+
 FProperty* FRepLayout::GetLifetimeCustomDeltaProperty(const uint16 CustomDeltaPropertyIndex) const
 {
+	checkSlow(LifetimeCustomPropertyState.IsValid());
+
 	const FLifetimeCustomDeltaProperty& CustomDeltaProperty = LifetimeCustomPropertyState->GetCustomDeltaProperty(CustomDeltaPropertyIndex);
 	return Parents[CustomDeltaProperty.PropertyRepIndex].Property;
 }
 
 const uint16 FRepLayout::GetCustomDeltaIndexFromPropertyRepIndex(const uint16 PropertyRepIndex) const
 {
+	checkSlow(LifetimeCustomPropertyState.IsValid());
+
 	return LifetimeCustomPropertyState->GetCustomDeltaIndexFromPropertyRepIndex(PropertyRepIndex);
 }
 
 const ELifetimeCondition FRepLayout::GetLifetimeCustomDeltaPropertyCondition(const uint16 CustomDeltaPropertyIndex) const
 {
+	checkSlow(LifetimeCustomPropertyState.IsValid());
+
 	const FLifetimeCustomDeltaProperty& CustomDeltaProperty = LifetimeCustomPropertyState->GetCustomDeltaProperty(CustomDeltaPropertyIndex);
 	return Parents[CustomDeltaProperty.PropertyRepIndex].Condition;
 }
diff --git a/Engine/Source/Runtime/Engine/Public/Net/DataReplication.h b/Engine/Source/Runtime/Engine/Public/Net/DataReplication.h
--- a/Engine/Source/Runtime/Engine/Public/Net/DataReplication.h
+++ b/Engine/Source/Runtime/Engine/Public/Net/DataReplication.h
@@ -157,7 +157,14 @@ public:
 	void CountBytes(FArchive& Ar) const;
 
 	/** Writes dirty properties to bunch */
-	void ReplicateCustomDeltaProperties(FNetBitWriter& Bunch, FReplicationFlags RepFlags);
+	UE_DEPRECATED(5.1, "Now takes an additional out param")
+	void ReplicateCustomDeltaProperties(FNetBitWriter& Bunch, FReplicationFlags RepFlags)
+	{
+		bool bSkippedPropertyCondition = false;
+		ReplicateCustomDeltaProperties(Bunch, RepFlags, bSkippedPropertyCondition);
+	}
+
+	void ReplicateCustomDeltaProperties(FNetBitWriter& Bunch, FReplicationFlags RepFlags, bool& bSkippedPropertyCondition);
 	bool ReplicateProperties(FOutBunch& Bunch, FReplicationFlags RepFlags);
 	bool ReplicateProperties(FOutBunch& Bunch, FReplicationFlags RepFlags, FNetBitWriter& Writer);
 	bool ReplicateProperties_r(FOutBunch& Bunch, FReplicationFlags RepFlags, FNetBitWriter& Writer);
diff --git a/Engine/Source/Runtime/Engine/Public/Net/RepLayout.h b/Engine/Source/Runtime/Engine/Public/Net/RepLayout.h
--- a/Engine/Source/Runtime/Engine/Public/Net/RepLayout.h
+++ b/Engine/Source/Runtime/Engine/Public/Net/RepLayout.h
@@ -450,6 +450,9 @@ public:
 	/** Number of times that properties have been compared */
 	int32 CompareIndex;
 
+	/** Tracking custom delta sends, for comparison against sending rep state. */
+	uint32 CustomDeltaChangeIndex = 0;
+	
 	/** Latest state of all property data. Not used on Clients, only used on Servers if Shadow State is enabled. */
 	FRepStateStaticBuffer StaticBuffer;
 
@@ -604,6 +607,9 @@ public:
 	 */
 	int32 LastCompareIndex;
 
+	/** Tracking custom delta sends, for comparison against the changelist state. */
+	uint32 CustomDeltaChangeIndex = 0;
+
 	FReplicationFlags RepFlags;
 
 	TSharedPtr<FRepChangedPropertyTracker> RepChangedPropertyTracker;
@@ -1837,9 +1843,11 @@ private:
 
 	const uint16 GetNumLifetimeCustomDeltaProperties() const;
 
+	const uint16 GetLifetimeCustomDeltaPropertyRepIndex(const uint16 RepIndCustomDeltaPropertyIndex) const;
+
 	FProperty* GetLifetimeCustomDeltaProperty(const uint16 CustomDeltaPropertyIndex) const;
 
-	const ELifetimeCondition GetLifetimeCustomDeltaPropertyCondition(const uint16 RepIndCustomDeltaPropertyIndexex) const;
+	const ELifetimeCondition GetLifetimeCustomDeltaPropertyCondition(const uint16 RepIndCustomDeltaPropertyIndex) const;
 
 	ERepLayoutFlags Flags;
 
