Include inherited component template overrides in Blueprint data validation.

diff --git a/Engine/Source/Runtime/Engine/Classes/Engine/InheritableComponentHandler.h b/Engine/Source/Runtime/Engine/Classes/Engine/InheritableComponentHandler.h
--- a/Engine/Source/Runtime/Engine/Classes/Engine/InheritableComponentHandler.h
+++ b/Engine/Source/Runtime/Engine/Classes/Engine/InheritableComponentHandler.h
@@ -127,6 +127,9 @@ public:
 	//~ Begin UObject Interface
 	virtual void Serialize(FArchive& Ar) override;
 	virtual void PostLoad() override;
+#if WITH_EDITOR
+	virtual EDataValidationResult IsDataValid(TArray<FText>& ValidationErrors) override;
+#endif	// WITH_EDITOR
 	//~ End UObject Interface
 
 	void PreloadAllTemplates();
diff --git a/Engine/Source/Runtime/Engine/Private/Blueprint.cpp b/Engine/Source/Runtime/Engine/Private/Blueprint.cpp
--- a/Engine/Source/Runtime/Engine/Private/Blueprint.cpp
+++ b/Engine/Source/Runtime/Engine/Private/Blueprint.cpp
@@ -2011,6 +2011,12 @@ EDataValidationResult UBlueprint::IsDataValid(TArray<FText>& ValidationErrors)
 		IsValid = CombineDataValidationResults(IsValid, IsSCSValid);
 	}
 
+	if (InheritableComponentHandler)
+	{
+		EDataValidationResult IsICHValid = InheritableComponentHandler->IsDataValid(ValidationErrors);
+		IsValid = CombineDataValidationResults(IsValid, IsICHValid);
+	}
+
 	for (UActorComponent* Component : ComponentTemplates)
 	{
 		if (Component)
diff --git a/Engine/Source/Runtime/Engine/Private/InheritableComponentHandler.cpp b/Engine/Source/Runtime/Engine/Private/InheritableComponentHandler.cpp
--- a/Engine/Source/Runtime/Engine/Private/InheritableComponentHandler.cpp
+++ b/Engine/Source/Runtime/Engine/Private/InheritableComponentHandler.cpp
@@ -78,6 +78,28 @@ void UInheritableComponentHandler::PostLoad()
 }
 
 #if WITH_EDITOR
+EDataValidationResult UInheritableComponentHandler::IsDataValid(TArray<FText>& ValidationErrors)
+{
+	EDataValidationResult Result = Super::IsDataValid(ValidationErrors);
+
+	// If nothing has performed any validation yet, start with a valid result.
+	if (Result == EDataValidationResult::NotValidated)
+	{
+		Result = EDataValidationResult::Valid;
+	}
+
+	for (const FComponentOverrideRecord& Record : Records)
+	{
+		if (Record.ComponentTemplate)
+		{
+			EDataValidationResult OverrideResult = Record.ComponentTemplate->IsDataValid(ValidationErrors);
+			Result = CombineDataValidationResults(Result, OverrideResult);
+		}
+	}
+
+	return Result;
+}
+
 UActorComponent* UInheritableComponentHandler::CreateOverridenComponentTemplate(const FComponentKey& Key)
 {
 	for (int32 Index = 0; Index < Records.Num(); ++Index)
