Detect invalid texture streaming data when remapping actor texture streaming built data to level (at cook).
When invalid, log details and invalidate built data. Usually, user will need to rebuild texture streaming and or resave textures with duplicate guids.

diff --git a/Engine/Source/Runtime/Engine/Classes/Engine/TextureStreamingTypes.h b/Engine/Source/Runtime/Engine/Classes/Engine/TextureStreamingTypes.h
--- a/Engine/Source/Runtime/Engine/Classes/Engine/TextureStreamingTypes.h
+++ b/Engine/Source/Runtime/Engine/Classes/Engine/TextureStreamingTypes.h
@@ -103,6 +103,9 @@ struct FStreamingRenderAssetPrimitiveInfo
 	}
 };
 
+// Invalid streamable texture registration index
+static const uint16 InvalidRegisteredStreamableTexture = (uint16)INDEX_NONE;
+
 /**
  * Interface for texture streaming container
  */
diff --git a/Engine/Source/Runtime/Engine/Private/Components/StaticMeshComponent.cpp b/Engine/Source/Runtime/Engine/Private/Components/StaticMeshComponent.cpp
--- a/Engine/Source/Runtime/Engine/Private/Components/StaticMeshComponent.cpp
+++ b/Engine/Source/Runtime/Engine/Private/Components/StaticMeshComponent.cpp
@@ -886,7 +886,12 @@ bool UStaticMeshComponent::RemapActorTextureStreamingBuiltDataToLevel(const UAct
 	FGuid TextureGuid;
 	for (FStreamingTextureBuildInfo& BuildInfo : StreamingTextureData)
 	{
-		if (!InActorTextureBuildData->GetStreamableTexture(BuildInfo.TextureLevelIndex, TextureName, TextureGuid))
+		uint16 TextureLevelIndex = InvalidRegisteredStreamableTexture;
+		if (InActorTextureBuildData->GetStreamableTexture(BuildInfo.TextureLevelIndex, TextureName, TextureGuid))
+		{
+			TextureLevelIndex = Level->RegisterStreamableTexture(TextureName, TextureGuid);
+		}
+		if (TextureLevelIndex == InvalidRegisteredStreamableTexture)
 		{
 			// If remapping failed, invalidate built texture streaming data (this should not happen with newly generated texture streaming build data)
 			UE_LOG(LogStaticMesh, Warning, TEXT("Clearing invalid texture streaming built data for %s"), *GetFullName());
@@ -894,7 +899,7 @@ bool UStaticMeshComponent::RemapActorTextureStreamingBuiltDataToLevel(const UAct
 			return false;
 		}
 		// Update BuildInfo's TextureLevelIndex
-		BuildInfo.TextureLevelIndex = Level->RegisterStreamableTexture(TextureName, TextureGuid);
+		BuildInfo.TextureLevelIndex = TextureLevelIndex;
 	}
 	return true;
 }
diff --git a/Engine/Source/Runtime/Engine/Private/Level.cpp b/Engine/Source/Runtime/Engine/Private/Level.cpp
--- a/Engine/Source/Runtime/Engine/Private/Level.cpp
+++ b/Engine/Source/Runtime/Engine/Private/Level.cpp
@@ -2097,7 +2097,9 @@ uint16 ULevel::RegisterStreamableTexture(UTexture* InTexture)
 			InTexture->SetLightingGuid();
 		}
 
-		InTexture->LevelIndex = (int32)RegisterStreamableTexture(InTexture->GetPathName(), InTexture->GetLightingGuid());
+		uint16 RegisteredIndex = RegisterStreamableTexture(InTexture->GetPathName(), InTexture->GetLightingGuid());
+		check(RegisteredIndex != InvalidRegisteredStreamableTexture);
+		InTexture->LevelIndex = (int32)RegisteredIndex;
 	}
 	check(StreamingTextureGuids.IsValidIndex(InTexture->LevelIndex));
 	check(StreamingTextureGuids[InTexture->LevelIndex] == InTexture->GetLightingGuid());
@@ -2107,10 +2109,25 @@ uint16 ULevel::RegisterStreamableTexture(UTexture* InTexture)
 uint16 ULevel::RegisterStreamableTexture(const FString& InTextureName, const FGuid& InTextureGuid)
 {
 	check(StreamingTextures.Num() == StreamingTextureGuids.Num());
-	uint16 Index = StreamingTextureGuids.AddUnique(InTextureGuid);
-	StreamingTextures.AddUnique(FName(*InTextureName));
-	check(StreamingTextures.Num() == StreamingTextureGuids.Num());
-	return Index;
+	const int32 TextureNameIndex = StreamingTextures.Find(FName(InTextureName));
+	const int32 TextureGuidIndex = StreamingTextureGuids.Find(InTextureGuid);
+	if (TextureNameIndex != TextureGuidIndex)
+	{
+		UE_CLOG(TextureNameIndex != INDEX_NONE, LogLevel, Warning, TEXT("Failed to register streamable texture Name = %s Guid = %s: An entry already exists for this Name with a different Guid = %s. Consider rebuilding texture streaming."), *InTextureName, *InTextureGuid.ToString(), *StreamingTextureGuids[TextureNameIndex].ToString());
+		UE_CLOG(TextureGuidIndex != INDEX_NONE, LogLevel, Warning, TEXT("Failed to register streamable texture Name = %s Guid = %s: An entry already exists for this Guid with a different Name = %s. Consider modifying & resaving one of these textures (will change its guid) and rebuiling texture streaming."), *InTextureName, *InTextureGuid.ToString(), *StreamingTextures[TextureGuidIndex].ToString());
+		return InvalidRegisteredStreamableTexture;
+	}
+	else if (TextureNameIndex != INDEX_NONE)
+	{
+		return TextureNameIndex;
+	}
+	else
+	{
+		uint16 Index = StreamingTextureGuids.Add(InTextureGuid);
+		StreamingTextures.Add(FName(*InTextureName));
+		check(StreamingTextures.Num() == StreamingTextureGuids.Num());
+		return Index;
+	}
 }
 
 #endif
