Display (read only) role and remote role in the editor actor details panel under the networking category

This change helps debug netsyncing issues in gameplay code, particularly in cases where we are doing something weird like with client predicted projectiles. These roles are a foundational concept in creating multiplayer experiences, and covered in the public unreal docs. I think they should be easily observable when using PIE.

diff --git a/Engine/Source/Runtime/Engine/Classes/GameFramework/Actor.h b/Engine/Source/Runtime/Engine/Classes/GameFramework/Actor.h
--- a/Engine/Source/Runtime/Engine/Classes/GameFramework/Actor.h
+++ b/Engine/Source/Runtime/Engine/Classes/GameFramework/Actor.h
@@ -590,7 +590,7 @@ public:
 	
 private:
 	/** Describes how much control the remote machine has over the actor. */
-	UPROPERTY(Replicated, Transient)
+	UPROPERTY(Replicated, Transient, VisibleInstanceOnly, Category=Networking)
 	TEnumAsByte<enum ENetRole> RemoteRole;
 
 	/** The RayTracingGroupId this actor and its components belong to. (For components that did not specify any) */
@@ -655,7 +655,7 @@ public:
 
 private:
 	/** Describes how much control the local machine has over the actor. */
-	UPROPERTY(Replicated)
+	UPROPERTY(Replicated, VisibleInstanceOnly, Category=Networking)
 	TEnumAsByte<enum ENetRole> Role;
 
 public:
