Added arm64 suffix to the SkeletalMesh DDC key to avoid potential consistency issues. Amendment to merge conflict in 21891333.

diff --git a/Engine/Source/Runtime/Engine/Private/SkeletalMesh.cpp b/Engine/Source/Runtime/Engine/Private/SkeletalMesh.cpp
--- a/Engine/Source/Runtime/Engine/Private/SkeletalMesh.cpp
+++ b/Engine/Source/Runtime/Engine/Private/SkeletalMesh.cpp
@@ -4392,6 +4392,13 @@ FString USkeletalMesh::BuildDerivedDataKey(const ITargetPlatform* TargetPlatform
 	const bool bUnlimitedBoneInfluences = FGPUBaseSkinVertexFactory::GetUnlimitedBoneInfluences();
 	KeySuffix += bUnlimitedBoneInfluences ? "1" : "0";
 
+#if PLATFORM_CPU_ARM_FAMILY
+	// Separate out arm keys as x64 and arm64 clang do not generate the same data for a given
+	// input. Add the arm specifically so that a) we avoid rebuilding the current DDC and
+	// b) we can remove it once we get arm64 to be consistent.
+	KeySuffix.Append(TEXT("_arm64"));
+#endif
+
 	return FDerivedDataCacheInterface::BuildCacheKey(
 		TEXT("SKELETALMESH"),
 		*GetSkeletalMeshDerivedDataVersion(),
