Fix editor assert when staticmesh property is overwritten during component replication

diff --git a/Engine/Source/Runtime/Engine/Private/Components/StaticMeshComponent.cpp b/Engine/Source/Runtime/Engine/Private/Components/StaticMeshComponent.cpp
--- a/Engine/Source/Runtime/Engine/Private/Components/StaticMeshComponent.cpp
+++ b/Engine/Source/Runtime/Engine/Private/Components/StaticMeshComponent.cpp
@@ -239,19 +239,21 @@ void UStaticMeshComponent::GetLifetimeReplicatedProps(TArray< FLifetimeProperty
 
 /// @endcond
 
-void UStaticMeshComponent::OnRep_StaticMesh(class UStaticMesh *OldStaticMesh)
+void UStaticMeshComponent::OnRep_StaticMesh(class UStaticMesh* OldStaticMesh)
 {
-	PRAGMA_DISABLE_DEPRECATION_WARNINGS
 	// Only do stuff if this actually changed from the last local value
-	if (OldStaticMesh!= StaticMesh)
+	if (OldStaticMesh != StaticMesh)
 	{
-		// We have to force a call to SetStaticMesh with a new StaticMesh
-		UStaticMesh *NewStaticMesh = StaticMesh;
-		StaticMesh = OldStaticMesh;
+		// Properly handle replicated StaticMesh property change by putting the old value back
+		// and applying the modification through a proper call to SetStaticMesh.
+		UStaticMesh* NewStaticMesh = StaticMesh;
+
+		// Put back the old value with minimal logic involved
+		SetStaticMeshInternal(OldStaticMesh);
 		
+		// Go through all the logic required to properly apply a new static mesh.
 		SetStaticMesh(NewStaticMesh);
 	}
-	PRAGMA_ENABLE_DEPRECATION_WARNINGS
 }
 
 bool UStaticMeshComponent::HasAnySockets() const
@@ -2097,9 +2099,7 @@ bool UStaticMeshComponent::ShouldCreatePhysicsState() const
 
 void UStaticMeshComponent::SetStaticMeshInternal(UStaticMesh* NewMesh)
 {
-	PRAGMA_DISABLE_DEPRECATION_WARNINGS
 	StaticMesh = NewMesh;
-	PRAGMA_ENABLE_DEPRECATION_WARNINGS
 
 	NotifyIfStaticMeshChanged();
 }
