Fixed crash when LevelInstance references a Level that fails to load (for instance, if package was saved with a newer engine version).

diff --git a/Engine/Source/Runtime/Engine/Private/LevelInstance/LevelInstanceLevelStreaming.cpp b/Engine/Source/Runtime/Engine/Private/LevelInstance/LevelInstanceLevelStreaming.cpp
--- a/Engine/Source/Runtime/Engine/Private/LevelInstance/LevelInstanceLevelStreaming.cpp
+++ b/Engine/Source/Runtime/Engine/Private/LevelInstance/LevelInstanceLevelStreaming.cpp
@@ -225,45 +225,51 @@ ULevelStreamingLevelInstance* ULevelStreamingLevelInstance::LoadInstance(ILevelI
 			// by setting the objects RF_Transient and !RF_Transactional we can check when unloading if those flags
 			// have been changed and figure out if we need to clear the transaction buffer or not.
 			// It might not be the final solution to support Undo/Redo in LevelInstances but it handles most of the non-editing part
-			ULevel* Level = LevelStreaming->GetLoadedLevel();
-			check(Level);
-			check(LevelStreaming->GetCurrentState() == ULevelStreaming::ECurrentState::LoadedVisible);
+			if (ULevel* Level = LevelStreaming->GetLoadedLevel())
+			{
+				check(LevelStreaming->GetCurrentState() == ULevelStreaming::ECurrentState::LoadedVisible);
 
-			Level->OnLoadedActorAddedToLevelEvent.AddUObject(LevelStreaming, &ULevelStreamingLevelInstance::OnLoadedActorAddedToLevel);
-			Level->OnLoadedActorRemovedFromLevelEvent.AddUObject(LevelStreaming, &ULevelStreamingLevelInstance::OnLoadedActorRemovedFromLevel);
+				Level->OnLoadedActorAddedToLevelEvent.AddUObject(LevelStreaming, &ULevelStreamingLevelInstance::OnLoadedActorAddedToLevel);
+				Level->OnLoadedActorRemovedFromLevelEvent.AddUObject(LevelStreaming, &ULevelStreamingLevelInstance::OnLoadedActorRemovedFromLevel);
 
-			UWorld* OuterWorld = Level->GetTypedOuter<UWorld>();
-			OuterWorld->ClearFlags(RF_Transactional);
-			OuterWorld->SetFlags(RF_Transient);
+				UWorld* OuterWorld = Level->GetTypedOuter<UWorld>();
+				OuterWorld->ClearFlags(RF_Transactional);
+				OuterWorld->SetFlags(RF_Transient);
 			
-			OuterWorld->GetPackage()->ClearFlags(RF_Transactional);
-			OuterWorld->GetPackage()->SetFlags(RF_Transient);
-
-			ForEachObjectWithOuter(OuterWorld, [&](UObject* Obj)
-			{
-				Obj->ClearFlags(RF_Transactional);
-				Obj->SetFlags(RF_Transient);
-			}, true);
+				OuterWorld->GetPackage()->ClearFlags(RF_Transactional);
+				OuterWorld->GetPackage()->SetFlags(RF_Transient);
 
-			for (AActor* LevelActor : Level->Actors)
-			{
-				if (LevelActor)
+				ForEachObjectWithOuter(OuterWorld, [&](UObject* Obj)
 				{
-					LevelStreaming->PrepareLevelInstanceLoadedActor(*LevelActor, LevelInstance, false);
-				}
-			}
+					Obj->ClearFlags(RF_Transactional);
+					Obj->SetFlags(RF_Transient);
+				}, true);
 
-			Level->ForEachActorFolder([](UActorFolder* ActorFolder)
-			{
-				if (ActorFolder->IsPackageExternal())
+				for (AActor* LevelActor : Level->Actors)
 				{
-					ActorFolder->GetPackage()->SetFlags(RF_Transient);
+					if (LevelActor)
+					{
+						LevelStreaming->PrepareLevelInstanceLoadedActor(*LevelActor, LevelInstance, false);
+					}
 				}
-				return true;
-			});
 
-			// Create special actor that will handle selection and transform
-			LevelStreaming->LevelInstanceEditorInstanceActor = ALevelInstanceEditorInstanceActor::Create(LevelInstance, Level);
+				Level->ForEachActorFolder([](UActorFolder* ActorFolder)
+				{
+					if (ActorFolder->IsPackageExternal())
+					{
+						ActorFolder->GetPackage()->SetFlags(RF_Transient);
+					}
+					return true;
+				});
+
+				// Create special actor that will handle selection and transform
+				LevelStreaming->LevelInstanceEditorInstanceActor = ALevelInstanceEditorInstanceActor::Create(LevelInstance, Level);
+			}
+			else
+			{
+				// Failed to load package
+				return nullptr;
+			}
 		}
 #endif
 		return LevelStreaming;
