Convert serialized external actor/component references to weak ptrs in the debug camera controller so that they don't prevent GC.

diff --git a/Engine/Source/Developer/LogVisualizer/Private/VisualLoggerCameraController.cpp b/Engine/Source/Developer/LogVisualizer/Private/VisualLoggerCameraController.cpp
--- a/Engine/Source/Developer/LogVisualizer/Private/VisualLoggerCameraController.cpp
+++ b/Engine/Source/Developer/LogVisualizer/Private/VisualLoggerCameraController.cpp
@@ -151,6 +151,6 @@ void AVisualLoggerCameraController::Select(FHitResult const& Hit)
 	if (SelectedActor != PickedActor)
 	{
 		PickedActor = SelectedActor;
-		OnActorSelected.ExecuteIfBound(SelectedActor);
+		OnActorSelected.ExecuteIfBound(SelectedActor.Get());
 	}
 }
diff --git a/Engine/Source/Developer/LogVisualizer/Private/VisualLoggerCameraController.h b/Engine/Source/Developer/LogVisualizer/Private/VisualLoggerCameraController.h
--- a/Engine/Source/Developer/LogVisualizer/Private/VisualLoggerCameraController.h
+++ b/Engine/Source/Developer/LogVisualizer/Private/VisualLoggerCameraController.h
@@ -18,7 +18,7 @@ public:
 	GENERATED_UCLASS_BODY()
 
 	UPROPERTY()
-	TObjectPtr<class AActor> PickedActor;
+	TWeakObjectPtr<class AActor> PickedActor;
 
 	virtual void PostInitializeComponents() override;
 
diff --git a/Engine/Source/Developer/LogVisualizer/Private/VisualLoggerHUD.cpp b/Engine/Source/Developer/LogVisualizer/Private/VisualLoggerHUD.cpp
--- a/Engine/Source/Developer/LogVisualizer/Private/VisualLoggerHUD.cpp
+++ b/Engine/Source/Developer/LogVisualizer/Private/VisualLoggerHUD.cpp
@@ -63,7 +63,7 @@ void AVisualLoggerHUD::PostRender()
 			Canvas->DrawItem( TextItem, X, Y );
 			Y += TextItem.DrawnSize.Y;
 
-			if (DebugCamController->PickedActor != NULL)
+			if (DebugCamController->PickedActor.IsValid())
 			{
 				TextItem.Text = FText::FromString(FString::Printf(TEXT("Selected: '%s'"), *DebugCamController->PickedActor->GetName()));
 				Canvas->DrawItem( TextItem, X, Y );				
diff --git a/Engine/Source/Runtime/Engine/Classes/Engine/DebugCameraController.h b/Engine/Source/Runtime/Engine/Classes/Engine/DebugCameraController.h
--- a/Engine/Source/Runtime/Engine/Classes/Engine/DebugCameraController.h
+++ b/Engine/Source/Runtime/Engine/Classes/Engine/DebugCameraController.h
@@ -177,7 +177,7 @@ public:
 
 	/** Currently selected actor, may be invalid */
 	UPROPERTY()
-	TObjectPtr<class AActor> SelectedActor;
+	TWeakObjectPtr<class AActor> SelectedActor;
 
 	/** Returns the currently selected actor, or null if it is invalid or not set */
 	UFUNCTION(BlueprintCallable, Category="Debug Camera")
@@ -185,7 +185,7 @@ public:
 	
 	/** Currently selected component, may be invalid */
 	UPROPERTY()
-	TObjectPtr<class UPrimitiveComponent> SelectedComponent;
+	TWeakObjectPtr<class UPrimitiveComponent> SelectedComponent;
 
 	/** Selected hit point */
 	UPROPERTY()
diff --git a/Engine/Source/Runtime/Engine/Private/DebugCameraController.cpp b/Engine/Source/Runtime/Engine/Private/DebugCameraController.cpp
--- a/Engine/Source/Runtime/Engine/Private/DebugCameraController.cpp
+++ b/Engine/Source/Runtime/Engine/Private/DebugCameraController.cpp
@@ -31,8 +31,6 @@ static const float MIN_ORBIT_RADIUS = 30.0f;
 ADebugCameraController::ADebugCameraController(const FObjectInitializer& ObjectInitializer)
 	: Super(ObjectInitializer)
 {
-	SelectedActor = nullptr;
-	SelectedComponent = nullptr;
 	OriginalControllerRef = nullptr;
 	OriginalPlayer = nullptr;
 
@@ -246,25 +244,27 @@ void ADebugCameraController::OnFingerMove(ETouchIndex::Type FingerIndex, FVector
 
 AActor* ADebugCameraController::GetSelectedActor() const
 {
-	return SelectedActor;
+	return SelectedActor.Get();
 }
 
 void ADebugCameraController::Select( FHitResult const& Hit )
 {
+	AActor* HitActor = Hit.HitObjectHandle.FetchActor();
+
 	// store selection
-	SelectedActor = Hit.HitObjectHandle.FetchActor();
-	SelectedComponent = Hit.Component.Get();
+	SelectedActor = HitActor;
+	SelectedComponent = Hit.Component;
 	SelectedHitPoint = Hit;
 
 	//BP Event
-	ReceiveOnActorSelected(SelectedActor, Hit.ImpactPoint, Hit.ImpactNormal, Hit);
+	ReceiveOnActorSelected(HitActor, Hit.ImpactPoint, Hit.ImpactNormal, Hit);
 }
 
 
 void ADebugCameraController::Unselect()
 {	
-	SelectedActor = nullptr;
-	SelectedComponent = nullptr;
+	SelectedActor.Reset();
+	SelectedComponent.Reset();
 }
 
 FString ADebugCameraController::ConsoleCommand(const FString& Cmd,bool bWriteToLog)
@@ -615,7 +615,7 @@ void ADebugCameraController::UpdateRotationForOrbit(float DeltaTime)
 
 bool ADebugCameraController::GetPivotForOrbit(FVector& PivotLocation) const
 {
-	if (SelectedActor)
+	if (SelectedActor.IsValid())
 	{
 		if (bOrbitPivotUseCenter)
 		{
@@ -625,7 +625,7 @@ bool ADebugCameraController::GetPivotForOrbit(FVector& PivotLocation) const
 			// Use the center of the bounding box of the current selected actor as the pivot point for orbiting the camera
 			int32 NumSelectedActors = 0;
 
-			TInlineComponentArray<UMeshComponent*> MeshComponents(SelectedActor);
+			TInlineComponentArray<UMeshComponent*> MeshComponents(SelectedActor.Get());
 
 			for (int32 ComponentIndex = 0; ComponentIndex < MeshComponents.Num(); ++ComponentIndex)
 			{
diff --git a/Engine/Source/Runtime/Engine/Private/DebugCameraHUD.cpp b/Engine/Source/Runtime/Engine/Private/DebugCameraHUD.cpp
--- a/Engine/Source/Runtime/Engine/Private/DebugCameraHUD.cpp
+++ b/Engine/Source/Runtime/Engine/Private/DebugCameraHUD.cpp
@@ -190,7 +190,7 @@ void ADebugCameraHUD::PostRender()
 				{
 					DrawDebugLine(GetWorld(), Hit.Location, Hit.Location + Hit.Normal*30.f, FColor::White);
 
-					if (DCC->SelectedActor)
+					if (DCC->SelectedActor.IsValid())
 					{
 						DrawDebugLine(GetWorld(), DCC->SelectedHitPoint.Location, DCC->SelectedHitPoint.Location + DCC->SelectedHitPoint.Normal*30.f, FColor::Red);
 					}
@@ -202,7 +202,7 @@ void ADebugCameraHUD::PostRender()
 				Canvas->DrawText( RenderFont, TEXT("No trace Hit"), X, yl, 1.f, 1.f, FontRenderInfo);
 			}
 
-			if (DCC->bShowSelectedInfo && DCC->SelectedActor != NULL)
+			if (DCC->bShowSelectedInfo && DCC->SelectedActor.IsValid())
 			{
 				if (DCC->bIsOrbitingSelectedActor)
 				{
