Fix crash during cook
 - FShaderCompileUtilities::FetchGBufferParamsPipeline() in ShaderGenerationUtil.cpp never set the ShaderPlatform member of FGBufferParams.
 - It was left as default (SP_NumPlatforms) which isn't valid when its later passed to a DDPI function.

diff --git a/Engine/Source/Runtime/Engine/Private/ShaderCompiler/ShaderGenerationUtil.cpp b/Engine/Source/Runtime/Engine/Private/ShaderCompiler/ShaderGenerationUtil.cpp
--- a/Engine/Source/Runtime/Engine/Private/ShaderCompiler/ShaderGenerationUtil.cpp
+++ b/Engine/Source/Runtime/Engine/Private/ShaderCompiler/ShaderGenerationUtil.cpp
@@ -2109,6 +2109,8 @@ void FShaderCompileUtilities::GenerateBrdfHeaders(const FName& ShaderFormat)
 FGBufferParams FShaderCompileUtilities::FetchGBufferParamsPipeline(EShaderPlatform Platform)
 {
 	FGBufferParams Ret = {};
+	Ret.ShaderPlatform = Platform;
+
 #if WITH_EDITOR
 	{
 		static const auto CVar = IConsoleManager::Get().FindTConsoleVariableDataInt(TEXT("r.AllowStaticLighting"));
