Adding experimental project setting to render with HeterogeneousVolumes.

Found within Rendering > Experimental > Enable Heterogeneous Volumes

diff --git a/Engine/Source/Runtime/Engine/Classes/Engine/RendererSettings.h b/Engine/Source/Runtime/Engine/Classes/Engine/RendererSettings.h
--- a/Engine/Source/Runtime/Engine/Classes/Engine/RendererSettings.h
+++ b/Engine/Source/Runtime/Engine/Classes/Engine/RendererSettings.h
@@ -841,6 +841,12 @@ class ENGINE_API URendererSettings : public UDeveloperSettings
 		ConfigRestartRequired = true))
 		uint32 bMeshStreaming : 1;
 
+	UPROPERTY(config, EditAnywhere, Category = Experimental, meta = (
+		ConsoleVariable = "r.HeterogeneousVolumes", DisplayName = "Enable Heterogeneous Volumes",
+		ToolTip = "Enable rendering with the heterogeneous volumes subsystem.",
+		ConfigRestartRequired = false))
+	uint32 bEnableHeterogeneousVolumes : 1;
+
 	UPROPERTY(config, EditAnywhere, Category=Editor, meta=(
 		ConsoleVariable="r.WireframeCullThreshold",DisplayName="Wireframe Cull Threshold",
 		ToolTip="Screen radius at which wireframe objects are culled. Larger values can improve performance when viewing a scene in wireframe."))
diff --git a/Engine/Source/Runtime/Renderer/Private/HeterogeneousVolumes/HeterogeneousVolumesLiveShadingPipeline.cpp b/Engine/Source/Runtime/Renderer/Private/HeterogeneousVolumes/HeterogeneousVolumesLiveShadingPipeline.cpp
--- a/Engine/Source/Runtime/Renderer/Private/HeterogeneousVolumes/HeterogeneousVolumesLiveShadingPipeline.cpp
+++ b/Engine/Source/Runtime/Renderer/Private/HeterogeneousVolumes/HeterogeneousVolumesLiveShadingPipeline.cpp
@@ -74,6 +74,7 @@ class FRenderTransmittanceVolumeWithLiveShadingCS : public FMeshMaterialShader
 	{
 		return DoesPlatformSupportHeterogeneousVolumes(Parameters.Platform)
 			&& (Parameters.MaterialParameters.MaterialDomain == MD_Volume)
+			// Restricting compilation to materials bound to Niagara meshes
 			&& Parameters.MaterialParameters.bIsUsedWithNiagaraMeshParticles;
 	}
 
diff --git a/Engine/Source/Runtime/Renderer/Private/HeterogeneousVolumes/HeterogeneousVolumesMaterialBakingPipeline.cpp b/Engine/Source/Runtime/Renderer/Private/HeterogeneousVolumes/HeterogeneousVolumesMaterialBakingPipeline.cpp
--- a/Engine/Source/Runtime/Renderer/Private/HeterogeneousVolumes/HeterogeneousVolumesMaterialBakingPipeline.cpp
+++ b/Engine/Source/Runtime/Renderer/Private/HeterogeneousVolumes/HeterogeneousVolumesMaterialBakingPipeline.cpp
@@ -64,6 +64,7 @@ class FHeterogeneousVolumesBakeMaterialCS : public FMeshMaterialShader
 	{
 		return DoesPlatformSupportHeterogeneousVolumes(Parameters.Platform)
 			&& (Parameters.MaterialParameters.MaterialDomain == MD_Volume)
+			// Restricting compilation to materials bound to Niagara meshes
 			&& Parameters.MaterialParameters.bIsUsedWithNiagaraMeshParticles;
 	}
 
