Change TArray::Empty() + TArray::AddUninitialized() call to TArray::SetNumUninitialized in FRenderTarget::ReadFloat16Pixels to remove reallocation.

#8953

diff --git a/Engine/Source/Runtime/Engine/Private/UnrealClient.cpp b/Engine/Source/Runtime/Engine/Private/UnrealClient.cpp
--- a/Engine/Source/Runtime/Engine/Private/UnrealClient.cpp
+++ b/Engine/Source/Runtime/Engine/Private/UnrealClient.cpp
@@ -174,8 +174,7 @@ bool FRenderTarget::ReadFloat16Pixels(FFloat16Color* OutImageData,ECubeFace Cube
 bool FRenderTarget::ReadFloat16Pixels(TArray<FFloat16Color>& OutputBuffer,ECubeFace CubeFace)
 {
 	// Copy the surface data into the output array.
-	OutputBuffer.Empty();
-	OutputBuffer.AddUninitialized(GetSizeXY().X * GetSizeXY().Y);
+	OutputBuffer.SetNumUninitialized(GetSizeXY().X * GetSizeXY().Y, true);
 	return ReadFloat16Pixels((FFloat16Color*)&(OutputBuffer[0]), CubeFace);
 }
 
