PR #9400:  Pass along success/fail return value to AActor AttachToComponent and AttachToActor functions.

diff --git a/Engine/Source/Runtime/Engine/Classes/GameFramework/Actor.h b/Engine/Source/Runtime/Engine/Classes/GameFramework/Actor.h
--- a/Engine/Source/Runtime/Engine/Classes/GameFramework/Actor.h
+++ b/Engine/Source/Runtime/Engine/Classes/GameFramework/Actor.h
@@ -1668,17 +1668,19 @@ public:
 	 * @param RotationRule				How to handle rotation when attaching.
 	 * @param ScaleRule					How to handle scale when attaching.
 	 * @param bWeldSimulatedBodies		Whether to weld together simulated physics bodies.
+	 * @return							Whether the attachment was successful or not
 	 */
 	UFUNCTION(BlueprintCallable, meta = (DisplayName = "Attach Actor To Component", ScriptName = "AttachToComponent", bWeldSimulatedBodies = true), Category = "Transformation")
-	void K2_AttachToComponent(USceneComponent* Parent, FName SocketName, EAttachmentRule LocationRule, EAttachmentRule RotationRule, EAttachmentRule ScaleRule, bool bWeldSimulatedBodies);
+	bool K2_AttachToComponent(USceneComponent* Parent, FName SocketName, EAttachmentRule LocationRule, EAttachmentRule RotationRule, EAttachmentRule ScaleRule, bool bWeldSimulatedBodies);
 
 	/**
 	 * Attaches the RootComponent of this Actor to the supplied component, optionally at a named socket. It is not valid to call this on components that are not Registered.
 	 * @param  Parent					Parent to attach to.
 	 * @param  AttachmentRules			How to handle transforms and welding when attaching.
 	 * @param  SocketName				Optional socket to attach to on the parent.
+	 * @return							Whether the attachment was successful or not
 	 */
-	void AttachToComponent(USceneComponent* Parent, const FAttachmentTransformRules& AttachmentRules, FName SocketName = NAME_None);
+	bool AttachToComponent(USceneComponent* Parent, const FAttachmentTransformRules& AttachmentRules, FName SocketName = NAME_None);
 
 	/** DEPRECATED - Use AttachToActor() instead */
 	UE_DEPRECATED(4.17, "Use AttachToActor() instead.")
@@ -1690,8 +1692,9 @@ public:
 	 * @param ParentActor				Actor to attach this actor's RootComponent to
 	 * @param AttachmentRules			How to handle transforms and modification when attaching.
 	 * @param SocketName				Socket name to attach to, if any
+	 * @return							Whether the attachment was successful or not
 	 */
-	void AttachToActor(AActor* ParentActor, const FAttachmentTransformRules& AttachmentRules, FName SocketName = NAME_None);
+	bool AttachToActor(AActor* ParentActor, const FAttachmentTransformRules& AttachmentRules, FName SocketName = NAME_None);
 
 	/**
 	 * Attaches the RootComponent of this Actor to the supplied actor, optionally at a named socket.
@@ -1701,9 +1704,10 @@ public:
 	 * @param RotationRule				How to handle rotation when attaching.
 	 * @param ScaleRule					How to handle scale when attaching.
 	 * @param bWeldSimulatedBodies		Whether to weld together simulated physics bodies.
+	 * @return							Whether the attachment was successful or not
 	 */
 	UFUNCTION(BlueprintCallable, meta = (DisplayName = "Attach Actor To Actor", ScriptName = "AttachToActor", bWeldSimulatedBodies=true), Category = "Transformation")
-	void K2_AttachToActor(AActor* ParentActor, FName SocketName, EAttachmentRule LocationRule, EAttachmentRule RotationRule, EAttachmentRule ScaleRule, bool bWeldSimulatedBodies);
+	bool K2_AttachToActor(AActor* ParentActor, FName SocketName, EAttachmentRule LocationRule, EAttachmentRule RotationRule, EAttachmentRule ScaleRule, bool bWeldSimulatedBodies);
 
 	/** DEPRECATED - Use DetachFromActor() instead */
 	UE_DEPRECATED(4.17, "Use DetachFromActor() instead")
diff --git a/Engine/Source/Runtime/Engine/Private/Actor.cpp b/Engine/Source/Runtime/Engine/Private/Actor.cpp
--- a/Engine/Source/Runtime/Engine/Private/Actor.cpp
+++ b/Engine/Source/Runtime/Engine/Private/Actor.cpp
@@ -2105,17 +2105,19 @@ PRAGMA_DISABLE_DEPRECATION_WARNINGS
 PRAGMA_ENABLE_DEPRECATION_WARNINGS
 }
 
-void AActor::K2_AttachToComponent(USceneComponent* Parent, FName SocketName, EAttachmentRule LocationRule, EAttachmentRule RotationRule, EAttachmentRule ScaleRule, bool bWeldSimulatedBodies)
+bool AActor::K2_AttachToComponent(USceneComponent* Parent, FName SocketName, EAttachmentRule LocationRule, EAttachmentRule RotationRule, EAttachmentRule ScaleRule, bool bWeldSimulatedBodies)
 {
-	AttachToComponent(Parent, FAttachmentTransformRules(LocationRule, RotationRule, ScaleRule, bWeldSimulatedBodies), SocketName);
+	return AttachToComponent(Parent, FAttachmentTransformRules(LocationRule, RotationRule, ScaleRule, bWeldSimulatedBodies), SocketName);
 }
 
-void AActor::AttachToComponent(USceneComponent* Parent, const FAttachmentTransformRules& AttachmentRules, FName SocketName)
+bool AActor::AttachToComponent(USceneComponent* Parent, const FAttachmentTransformRules& AttachmentRules, FName SocketName)
 {
 	if (RootComponent && Parent)
 	{
-		RootComponent->AttachToComponent(Parent, AttachmentRules, SocketName);
+		return RootComponent->AttachToComponent(Parent, AttachmentRules, SocketName);
 	}
+
+	return false;
 }
 
 void AActor::OnRep_AttachmentReplication()
@@ -2177,21 +2179,23 @@ PRAGMA_ENABLE_DEPRECATION_WARNINGS
 	}
 }
 
-void AActor::K2_AttachToActor(AActor* ParentActor, FName SocketName, EAttachmentRule LocationRule, EAttachmentRule RotationRule, EAttachmentRule ScaleRule, bool bWeldSimulatedBodies)
+bool AActor::K2_AttachToActor(AActor* ParentActor, FName SocketName, EAttachmentRule LocationRule, EAttachmentRule RotationRule, EAttachmentRule ScaleRule, bool bWeldSimulatedBodies)
 {
-	AttachToActor(ParentActor, FAttachmentTransformRules(LocationRule, RotationRule, ScaleRule, bWeldSimulatedBodies), SocketName);
+	return AttachToActor(ParentActor, FAttachmentTransformRules(LocationRule, RotationRule, ScaleRule, bWeldSimulatedBodies), SocketName);
 }
 
-void AActor::AttachToActor(AActor* ParentActor, const FAttachmentTransformRules& AttachmentRules, FName SocketName)
+bool AActor::AttachToActor(AActor* ParentActor, const FAttachmentTransformRules& AttachmentRules, FName SocketName)
 {
 	if (RootComponent && ParentActor)
 	{
 		USceneComponent* ParentDefaultAttachComponent = ParentActor->GetDefaultAttachComponent();
 		if (ParentDefaultAttachComponent)
 		{
-			RootComponent->AttachToComponent(ParentDefaultAttachComponent, AttachmentRules, SocketName);
+			return RootComponent->AttachToComponent(ParentDefaultAttachComponent, AttachmentRules, SocketName);
 		}
 	}
+
+	return false;
 }
 
 void AActor::DetachRootComponentFromParent(bool bMaintainWorldPosition)
