Fix ensure() in bsp code when input spans zero range and exactly aligns to Scale (which happens at LWC scales due to float precision)
+ do more of the computation explicitly in double precision

diff --git a/Engine/Source/Runtime/Engine/Private/ModelLight.cpp b/Engine/Source/Runtime/Engine/Private/ModelLight.cpp
--- a/Engine/Source/Runtime/Engine/Private/ModelLight.cpp
+++ b/Engine/Source/Runtime/Engine/Private/ModelLight.cpp
@@ -397,9 +397,9 @@ void UModelComponent::GetSurfaceLightMapResolution( int32 SurfaceIndex, int32 Qu
 	FBspSurf& Surf = Model->Surfs[SurfaceIndex];
 
 	// Find a plane parallel to the surface.
-	FVector3f MapX;
-	FVector3f MapY;
-	Surf.Plane.FindBestAxisVectors(MapX,MapY);
+	FVector3f MapXf, MapYf;
+	Surf.Plane.FindBestAxisVectors(MapXf,MapYf);
+	FVector3d MapX(MapXf), MapY(MapYf);
 
 	// Find the surface's nodes and the part of the plane they map to.
 	bool bFoundNode = false;
@@ -419,8 +419,8 @@ void UModelComponent::GetSurfaceLightMapResolution( int32 SurfaceIndex, int32 Qu
 			{
 				bFoundNode = true;
 
-				FVector3f	Position = Model->Points[Model->Verts[Node.iVertPool + VertexIndex].pVertex];
-				float	X = MapX | Position,
+				FVector3d	Position = (FVector3d)Model->Points[Model->Verts[Node.iVertPool + VertexIndex].pVertex];
+				double X = MapX | Position,
 					Y = MapY | Position;
 				MinUV.X = FMath::Min(X,MinUV.X);
 				MinUV.Y = FMath::Min(Y,MinUV.Y);
@@ -432,19 +432,24 @@ void UModelComponent::GetSurfaceLightMapResolution( int32 SurfaceIndex, int32 Qu
 
 	if (bFoundNode)
 	{
-		float Scale = Surf.LightMapScale * QualityScale;
+		double Scale = Surf.LightMapScale * QualityScale;
 		MinUV.X = FMath::FloorToFloat(MinUV.X / Scale) * Scale;
 		MinUV.Y = FMath::FloorToFloat(MinUV.Y / Scale) * Scale;
 		MaxUV.X = FMath::CeilToFloat(MaxUV.X / Scale) * Scale;
 		MaxUV.Y = FMath::CeilToFloat(MaxUV.Y / Scale) * Scale;
-
-		Width = FMath::Clamp(FMath::CeilToInt32((MaxUV.X - MinUV.X) / (Surf.LightMapScale * QualityScale)),4,SHADOWMAP_MAX_WIDTH);
-		Height = FMath::Clamp(FMath::CeilToInt32((MaxUV.Y - MinUV.Y) / (Surf.LightMapScale * QualityScale)),4,SHADOWMAP_MAX_HEIGHT);
+		// The intent of the above is that Min and Max should be separated by at least Scale, from the floor/ceil,
+		// but in the unlucky case that the range is ~0 and exactly aligned to Scale, Min and Max can take the same value.
+		// To avoid this causing a divide by zero below, we enforce the at-least-Scale delta explicitly.
+		double UVDx = FMath::Max(Scale, (MaxUV.X - MinUV.X));
+		double UVDy = FMath::Max(Scale, (MaxUV.Y - MinUV.Y));
+
+		Width = FMath::Clamp(FMath::CeilToInt32((UVDx) / (Surf.LightMapScale * QualityScale)),4,SHADOWMAP_MAX_WIDTH);
+		Height = FMath::Clamp(FMath::CeilToInt32((UVDy) / (Surf.LightMapScale * QualityScale)),4,SHADOWMAP_MAX_HEIGHT);
 		WorldToMap = FMatrix(
-			FPlane(MapX.X / (MaxUV.X - MinUV.X),	MapY.X / (MaxUV.Y - MinUV.Y),	Surf.Plane.X,	0),
-			FPlane(MapX.Y / (MaxUV.X - MinUV.X),	MapY.Y / (MaxUV.Y - MinUV.Y),	Surf.Plane.Y,	0),
-			FPlane(MapX.Z / (MaxUV.X - MinUV.X),	MapY.Z / (MaxUV.Y - MinUV.Y),	Surf.Plane.Z,	0),
-			FPlane(-MinUV.X / (MaxUV.X - MinUV.X),	-MinUV.Y / (MaxUV.Y - MinUV.Y),	-Surf.Plane.W,	1)
+			FPlane(MapX.X / UVDx,	MapY.X / UVDy,	Surf.Plane.X,	0),
+			FPlane(MapX.Y / UVDx,	MapY.Y / UVDy,	Surf.Plane.Y,	0),
+			FPlane(MapX.Z / UVDx,	MapY.Z / UVDy,	Surf.Plane.Z,	0),
+			FPlane(-MinUV.X / UVDx,	-MinUV.Y / UVDy,	-Surf.Plane.W,	1)
 			);
 	}
 	else
