- Removed the "PerObjectConfig" flag from TextureLODSettings, because it was causing the class defaults to not be read in from BaseDeviceProfiles.ini (DeviceProfile needs to be PerObjectConfig, but not TextureLODSettings, which has the defaults in a normal class ini section)

diff --git a/Engine/Source/Runtime/Engine/Classes/Engine/TextureLODSettings.h b/Engine/Source/Runtime/Engine/Classes/Engine/TextureLODSettings.h
--- a/Engine/Source/Runtime/Engine/Classes/Engine/TextureLODSettings.h
+++ b/Engine/Source/Runtime/Engine/Classes/Engine/TextureLODSettings.h
@@ -136,7 +136,7 @@ struct ENGINE_API FTextureLODGroup
  * Structure containing all information related to an LOD group and providing helper functions to calculate
  * the LOD bias of a given group.
  */
-UCLASS(config=DeviceProfiles, perObjectConfig)
+UCLASS(config=DeviceProfiles)
 class ENGINE_API UTextureLODSettings : public UObject
 {
 	GENERATED_UCLASS_BODY()
