Pass package name from ULevelStreamingDynamic::LoadLevelInstance to inner function instead of object path.
Also truncate package name to object path in FLoadLevelInstanceParams for other potential callers.

diff --git a/Engine/Source/Runtime/Engine/Classes/Engine/LevelStreamingDynamic.h b/Engine/Source/Runtime/Engine/Classes/Engine/LevelStreamingDynamic.h
--- a/Engine/Source/Runtime/Engine/Classes/Engine/LevelStreamingDynamic.h
+++ b/Engine/Source/Runtime/Engine/Classes/Engine/LevelStreamingDynamic.h
@@ -32,7 +32,7 @@ class ENGINE_API ULevelStreamingDynamic : public ULevelStreaming
 	{
 		FLoadLevelInstanceParams(UWorld* InWorld, const FString& InLongPackageName, FTransform InLevelTransform)
 			: World(InWorld)
-			, LongPackageName(InLongPackageName)
+			, LongPackageName(FPackageName::ObjectPathToPackageName(InLongPackageName))
 			, LevelTransform(InLevelTransform) 
 		{}
 
diff --git a/Engine/Source/Runtime/Engine/Private/LevelStreaming.cpp b/Engine/Source/Runtime/Engine/Private/LevelStreaming.cpp
--- a/Engine/Source/Runtime/Engine/Private/LevelStreaming.cpp
+++ b/Engine/Source/Runtime/Engine/Private/LevelStreaming.cpp
@@ -2140,14 +2140,6 @@ ULevelStreamingDynamic* ULevelStreamingDynamic::LoadLevelInstance(UObject* World
 		FString ObjectName;
 		FString SubObjectName;
 		FPackageName::SplitFullObjectPath(LevelObjectPath, UnusedClassName, PackageName, ObjectName, SubObjectName);
-		if (!ObjectName.IsEmpty())
-		{
-			ObjectRelativePath = FString(TEXT(".")) + ObjectName;
-			if (!SubObjectName.IsEmpty())
-			{
-				ObjectRelativePath += FString(SUBOBJECT_DELIMITER) + SubObjectName;
-			}
-		}
 	}
 	else if (!FPackageName::TryConvertFilenameToLongPackageName(LevelObjectPath, PackageName))
 	{
@@ -2162,9 +2154,8 @@ ULevelStreamingDynamic* ULevelStreamingDynamic::LoadLevelInstance(UObject* World
 	{
 		return nullptr;
 	}
-	const FString LongPackageNameObjectPath = ExistingPackageName.ToString() + ObjectRelativePath;
-
-	FLoadLevelInstanceParams Params(World, LongPackageNameObjectPath, FTransform(Rotation, Location));
+	const FString LongPackageName = ExistingPackageName.ToString();
+	FLoadLevelInstanceParams Params(World, LongPackageName, FTransform(Rotation, Location));
 	Params.OptionalLevelNameOverride = OptionalLevelNameOverride.IsEmpty() ? nullptr : &OptionalLevelNameOverride;
 	Params.OptionalLevelStreamingClass = OptionalLevelStreamingClass;
 	Params.bLoadAsTempPackage = bLoadAsTempPackage;
