World Partition - HLOD: Fix HLODs not showing up for maps created from the template map
* Removed the world name from cells names
* Strip the world name from AWorldPartitionHLOD::SourceCellName upon loading

diff --git a/Engine/Source/Runtime/Core/Public/UObject/FortniteMainBranchObjectVersion.h b/Engine/Source/Runtime/Core/Public/UObject/FortniteMainBranchObjectVersion.h
--- a/Engine/Source/Runtime/Core/Public/UObject/FortniteMainBranchObjectVersion.h
+++ b/Engine/Source/Runtime/Core/Public/UObject/FortniteMainBranchObjectVersion.h
@@ -241,6 +241,9 @@ struct CORE_API FFortniteMainBranchObjectVersion
 		// Added FHierarchicalSimplification::SimplificationMethod
 		HierarchicalSimplificationMethodEnumAdded,
 
+		// Changed how world partition streaming cells are named
+		WorldPartitionStreamingCellsNamingShortened,
+
 		// -----<new versions can be added above this line>-------------------------------------------------
 		VersionPlusOne,
 		LatestVersion = VersionPlusOne - 1
diff --git a/Engine/Source/Runtime/Engine/Private/WorldPartition/HLOD/HLODActor.cpp b/Engine/Source/Runtime/Engine/Private/WorldPartition/HLOD/HLODActor.cpp
--- a/Engine/Source/Runtime/Engine/Private/WorldPartition/HLOD/HLODActor.cpp
+++ b/Engine/Source/Runtime/Engine/Private/WorldPartition/HLOD/HLODActor.cpp
@@ -5,6 +5,7 @@
 #include "Components/PrimitiveComponent.h"
 #include "UObject/UE5MainStreamObjectVersion.h"
 #include "UObject/UE5PrivateFrostyStreamObjectVersion.h"
+#include "UObject/FortniteMainBranchObjectVersion.h"
 #include "UObject/ObjectSaveContext.h"
 
 #if WITH_EDITOR
@@ -68,6 +69,7 @@ void AWorldPartitionHLOD::Serialize(FArchive& Ar)
 {
 	Ar.UsingCustomVersion(FUE5MainStreamObjectVersion::GUID);
 	Ar.UsingCustomVersion(FUE5PrivateFrostyStreamObjectVersion::GUID);
+	Ar.UsingCustomVersion(FFortniteMainBranchObjectVersion::GUID);
 
 	Super::Serialize(Ar);
 
@@ -92,6 +94,31 @@ void AWorldPartitionHLOD::Serialize(FArchive& Ar)
 }
 
 #if WITH_EDITOR
+
+void AWorldPartitionHLOD::PostLoad()
+{
+	Super::PostLoad();
+
+	if (GetLinkerCustomVersion(FFortniteMainBranchObjectVersion::GUID) < FFortniteMainBranchObjectVersion::WorldPartitionStreamingCellsNamingShortened)
+	{
+		if (UWorldPartition* WorldPartition = GetWorld()->GetWorldPartition())
+		{
+			if (FWorldPartitionActorDesc* ActorDesc = WorldPartition->GetActorDesc(GetActorGuid()))
+			{
+				// As we may be dealing with an unsaved world created from a template map, get
+				// the source package name of this HLOD actor and figure out the world name from there
+				FString ExternalActorsPath = ActorDesc->GetContainer()->GetExternalActorPath();
+				FString WorldName = FPackageName::GetShortName(ExternalActorsPath);
+				
+				// Strip "WorldName_" from the cell name
+				FString CellName = SourceCellName.ToString();
+				CellName.RemoveFromStart(WorldName + TEXT("_"), ESearchCase::CaseSensitive);
+				SourceCellName = *CellName;
+			}
+		}
+	}
+}
+
 void AWorldPartitionHLOD::RerunConstructionScripts()
 {}
 
diff --git a/Engine/Source/Runtime/Engine/Private/WorldPartition/RuntimeSpatialHash/RuntimeSpatialHashHLOD.cpp b/Engine/Source/Runtime/Engine/Private/WorldPartition/RuntimeSpatialHash/RuntimeSpatialHashHLOD.cpp
--- a/Engine/Source/Runtime/Engine/Private/WorldPartition/RuntimeSpatialHash/RuntimeSpatialHashHLOD.cpp
+++ b/Engine/Source/Runtime/Engine/Private/WorldPartition/RuntimeSpatialHash/RuntimeSpatialHashHLOD.cpp
@@ -174,13 +174,13 @@ static TArray<FGuid> GenerateHLODsForGrid(UWorldPartition* WorldPartition, const
 				SlowTask.EnterProgressFrame(1);
 				// todo_ow : Generating hlods is not yet supported for content bundles. Pass an invalid UID to specified the cell is not coming from ContentBundles.
 				FGuid ContentBundleInvalidUID;
-				FName CellName = UWorldPartitionRuntimeSpatialHash::GetCellName(WorldPartition, RuntimeGrid.GridName, CellGlobalCoord, GridCellDataChunk.GetDataLayersID(), ContentBundleInvalidUID);
+				FString CellName = UWorldPartitionRuntimeSpatialHash::GetCellNameString(RuntimeGrid.GridName, CellGlobalCoord, GridCellDataChunk.GetDataLayersID(), ContentBundleInvalidUID);
 
-				UE_LOG(LogWorldPartitionRuntimeSpatialHashHLOD, Verbose, TEXT("Creating HLOD for cell %s at %s..."), *CellName.ToString(), *CellCoord.ToString());
+				UE_LOG(LogWorldPartitionRuntimeSpatialHashHLOD, Verbose, TEXT("Creating HLOD for cell %s at %s..."), *CellName, *CellCoord.ToString());
 
-				TRACE_CPUPROFILER_EVENT_SCOPE_TEXT(*CellName.ToString());
+				TRACE_CPUPROFILER_EVENT_SCOPE_TEXT(*CellName);
 
-				UE_LOG(LogWorldPartitionRuntimeSpatialHashHLOD, Display, TEXT("[%d / %d] Processing cell %s..."), (int32)SlowTask.CompletedWork + 1, (int32)SlowTask.TotalAmountOfWork, *CellName.ToString());
+				UE_LOG(LogWorldPartitionRuntimeSpatialHashHLOD, Display, TEXT("[%d / %d] Processing cell %s..."), (int32)SlowTask.CompletedWork + 1, (int32)SlowTask.TotalAmountOfWork, *CellName);
 
 				FHLODCreationParams CreationParams;
 				CreationParams.WorldPartition = WorldPartition;
@@ -188,7 +188,7 @@ static TArray<FGuid> GenerateHLODsForGrid(UWorldPartition* WorldPartition, const
 				CreationParams.GridIndexY = CellCoord.Y;
 				CreationParams.GridIndexZ = CellCoord.Z;
 				CreationParams.DataLayersID = GridCellDataChunk.GetDataLayersID();
-				CreationParams.CellName = CellName;
+				CreationParams.CellName = FName(CellName);
 				CreationParams.CellBounds = CellBounds;
 				CreationParams.HLODLevel = HLODLevel;
 				CreationParams.MinVisibleDistance = RuntimeGrid.LoadingRange;
diff --git a/Engine/Source/Runtime/Engine/Private/WorldPartition/WorldPartitionRuntimeSpatialHash.cpp b/Engine/Source/Runtime/Engine/Private/WorldPartition/WorldPartitionRuntimeSpatialHash.cpp
--- a/Engine/Source/Runtime/Engine/Private/WorldPartition/WorldPartitionRuntimeSpatialHash.cpp
+++ b/Engine/Source/Runtime/Engine/Private/WorldPartition/WorldPartitionRuntimeSpatialHash.cpp
@@ -959,19 +959,16 @@ void UWorldPartitionRuntimeSpatialHash::DumpStateLog(FHierarchicalLogArchive& Ar
 	});
 }
 
-FName UWorldPartitionRuntimeSpatialHash::GetCellName(UWorldPartition* WorldPartition, FName InGridName, const FGridCellCoord& InCellGlobalCoord, const FDataLayersID& InDataLayerID, const FGuid& InContentBundleID)
+FString UWorldPartitionRuntimeSpatialHash::GetCellNameString(FName InGridName, const FGridCellCoord& InCellGlobalCoord, const FDataLayersID& InDataLayerID, const FGuid& InContentBundleID)
 {
-	const FString PackageName = FPackageName::GetShortName(WorldPartition->GetPackage());
-	const FString PackageNameNoPIEPrefix = UWorld::RemovePIEPrefix(PackageName);
-
-	FString CellName = FString::Printf(TEXT("%s_%s_%s_DL%X"), *PackageNameNoPIEPrefix, *InGridName.ToString(), *GetCellCoordString(InCellGlobalCoord), InDataLayerID.GetHash());
+	FString CellName = FString::Printf(TEXT("%s_%s_DL%X"), *InGridName.ToString(), *GetCellCoordString(InCellGlobalCoord), InDataLayerID.GetHash());
 
 	if (InContentBundleID.IsValid())
 	{
 		CellName += TEXT("_CB") + InContentBundleID.ToString(EGuidFormats::Short);
 	}
 
-	return FName(*CellName);
+	return CellName;
 }
 
 bool UWorldPartitionRuntimeSpatialHash::GetPreviewGrids() const
@@ -985,12 +982,6 @@ void UWorldPartitionRuntimeSpatialHash::SetPreviewGrids(bool bInPreviewGrids)
 	bPreviewGrids = bInPreviewGrids;
 }
 
-FName UWorldPartitionRuntimeSpatialHash::GetCellName(FName InGridName, const FGridCellCoord& InCellGlobalCoord, const FDataLayersID& InDataLayerID, const FGuid& InContentBundleID) const
-{
-	UWorldPartition* WorldPartition = GetOuterUWorldPartition();
-	return UWorldPartitionRuntimeSpatialHash::GetCellName(WorldPartition, InGridName, InCellGlobalCoord, InDataLayerID, InContentBundleID);
-}
-
 bool UWorldPartitionRuntimeSpatialHash::CreateStreamingGrid(const FSpatialHashRuntimeGrid& RuntimeGrid, const FSquare2DGridHelper& PartionedActors, UWorldPartitionStreamingPolicy* StreamingPolicy, TArray<FString>* OutPackagesToGenerate)
 {
 	TRACE_CPUPROFILER_EVENT_SCOPE(CreateStreamingGrid);
@@ -1091,9 +1082,9 @@ bool UWorldPartitionRuntimeSpatialHash::CreateStreamingGrid(const FSpatialHashRu
 				
 				FGridCellCoord CellGlobalCoords;
 				verify(PartionedActors.GetCellGlobalCoords(FGridCellCoord(CellCoordX, CellCoordY, Level), CellGlobalCoords));
-				FName CellName = GetCellName(CurrentStreamingGrid.GridName, CellGlobalCoords, GridCellDataChunk.GetDataLayersID(), GridCellDataChunk.GetContentBundleID());
+				FString CellName = GetCellNameString(CurrentStreamingGrid.GridName, CellGlobalCoords, GridCellDataChunk.GetDataLayersID(), GridCellDataChunk.GetContentBundleID());
 
-				UWorldPartitionRuntimeSpatialHashCell* StreamingCell = NewObject<UWorldPartitionRuntimeSpatialHashCell>(WorldPartition, StreamingPolicy->GetRuntimeCellClass(), CellName);
+				UWorldPartitionRuntimeSpatialHashCell* StreamingCell = NewObject<UWorldPartitionRuntimeSpatialHashCell>(WorldPartition, StreamingPolicy->GetRuntimeCellClass(), FName(CellName));
 
 				int32 LayerCellIndex;
 				int32* LayerCellIndexPtr = GridLevel.LayerCellsMapping.Find(CellIndex);
diff --git a/Engine/Source/Runtime/Engine/Public/WorldPartition/HLOD/HLODActor.h b/Engine/Source/Runtime/Engine/Public/WorldPartition/HLOD/HLODActor.h
--- a/Engine/Source/Runtime/Engine/Public/WorldPartition/HLOD/HLODActor.h
+++ b/Engine/Source/Runtime/Engine/Public/WorldPartition/HLOD/HLODActor.h
@@ -69,6 +69,7 @@ protected:
 	//~ Begin UObject Interface.
 	virtual void Serialize(FArchive& Ar) override;
 #if WITH_EDITOR
+	virtual void PostLoad() override;
 	virtual void RerunConstructionScripts() override;
 	virtual bool CanEditChange(const FProperty* InProperty) const { return false; }
 	virtual bool CanEditChangeComponent(const UActorComponent* Component, const FProperty* InProperty) const { return false; }
diff --git a/Engine/Source/Runtime/Engine/Public/WorldPartition/WorldPartitionRuntimeSpatialHash.h b/Engine/Source/Runtime/Engine/Public/WorldPartition/WorldPartitionRuntimeSpatialHash.h
--- a/Engine/Source/Runtime/Engine/Public/WorldPartition/WorldPartitionRuntimeSpatialHash.h
+++ b/Engine/Source/Runtime/Engine/Public/WorldPartition/WorldPartitionRuntimeSpatialHash.h
@@ -250,8 +250,7 @@ public:
 
 	virtual URuntimeHashExternalStreamingObjectBase* StoreToExternalStreamingObject(UObject* StreamingObjectOuter, FName StreamingObjectName) override;
 
-	FName GetCellName(FName InGridName, const FGridCellCoord& InCellGlobalCoord, const FDataLayersID& InDataLayerID, const FGuid& InContentBundleID) const;
-	static FName GetCellName(UWorldPartition* WorldPartition, FName InGridName, const FGridCellCoord& InCellGlobalCoord, const FDataLayersID& InDataLayerID, const FGuid& InContentBundleID);
+	static FString GetCellNameString(FName InGridName, const FGridCellCoord& InCellGlobalCoord, const FDataLayersID& InDataLayerID, const FGuid& InContentBundleID);
 
 	bool GetPreviewGrids() const;
 	void SetPreviewGrids(bool bInPreviewGrids);
