Expose PerQualityLevelsMinLod Set/Get to blueprints. Can also set the Default value if user input a value >= 0.

diff --git a/Engine/Source/Editor/DetailCustomizations/Private/PerQualityLevelPropertyCustomization.cpp b/Engine/Source/Editor/DetailCustomizations/Private/PerQualityLevelPropertyCustomization.cpp
--- a/Engine/Source/Editor/DetailCustomizations/Private/PerQualityLevelPropertyCustomization.cpp
+++ b/Engine/Source/Editor/DetailCustomizations/Private/PerQualityLevelPropertyCustomization.cpp
@@ -230,7 +230,7 @@ void FPerQualityLevelPropertyCustomization<OverrideType>::CustomizeHeader(TShare
 		]
 	.ValueContent()
 		.MinDesiredWidth(FPerQualityLevelPropertyCustomization<OverrideType>::CalcDesiredWidth(StructPropertyHandle))
-		.MaxDesiredWidth((float)(static_cast<int32>(QualityLevelProperty::EQualityLevels::Num) + 1) * 125.0f)
+		.MaxDesiredWidth((float)(static_cast<int32>(EPerQualityLevels::Num) + 1) * 125.0f)
 		[
 			SNew(SPerQualityLevelPropertiesWidget)
 			.OnGenerateWidget(this, &FPerQualityLevelPropertyCustomization<OverrideType>::GetWidget, StructPropertyHandle)
diff --git a/Engine/Source/Editor/Persona/Private/PersonaMeshDetails.cpp b/Engine/Source/Editor/Persona/Private/PersonaMeshDetails.cpp
--- a/Engine/Source/Editor/Persona/Private/PersonaMeshDetails.cpp
+++ b/Engine/Source/Editor/Persona/Private/PersonaMeshDetails.cpp
@@ -3253,7 +3253,7 @@ void FPersonaMeshDetails::CustomizeLODSettingsCategories(IDetailLayoutBuilder& D
 	]
 	.ValueContent()
 	.MinDesiredWidth(float(SkelMesh->GetNoRefStreamingLODBias().PerQuality.Num() + 1) * 125.f)
-	.MaxDesiredWidth(float((int32)QualityLevelProperty::EQualityLevels::Num + 1) * 125.f)
+	.MaxDesiredWidth(float((int32)EPerQualityLevels::Num + 1) * 125.f)
 	[
 		SNew(SPerQualityLevelPropertiesWidget)
 		.OnGenerateWidget(this, &FPersonaMeshDetails::GetNoRefStreamingLODBiasWidget)
@@ -3313,7 +3313,7 @@ void FPersonaMeshDetails::CustomizeLODSettingsCategories(IDetailLayoutBuilder& D
 		]
 		.ValueContent()
 		.MinDesiredWidth((float)(SkelMesh->GetQualityLevelMinLod().PerQuality.Num() + 1) * 125.0f)
-		.MaxDesiredWidth((float)((int32)QualityLevelProperty::EQualityLevels::Num + 1) * 125.0f)
+		.MaxDesiredWidth((float)((int32)EPerQualityLevels::Num + 1) * 125.0f)
 		[
 			SNew(SPerQualityLevelPropertiesWidget)
 			.OnGenerateWidget(this, &FPersonaMeshDetails::GetMinQualityLevelLodWidget)
diff --git a/Engine/Source/Editor/PropertyEditor/Private/SPerQualityLevelPropertiesWidget.cpp b/Engine/Source/Editor/PropertyEditor/Private/SPerQualityLevelPropertiesWidget.cpp
--- a/Engine/Source/Editor/PropertyEditor/Private/SPerQualityLevelPropertiesWidget.cpp
+++ b/Engine/Source/Editor/PropertyEditor/Private/SPerQualityLevelPropertiesWidget.cpp
@@ -176,7 +176,7 @@ void SPerQualityLevelPropertiesWidget::ConstructChildren()
 		AddMenuBuilder.BeginSection(Section, FText::FromString(TEXT("Quality Levels")));
 
 		int32 QualityLevelCount = 0;
-		while (QualityLevelCount < static_cast<int32>(QualityLevelProperty::EQualityLevels::Num))
+		while (QualityLevelCount < static_cast<int32>(EPerQualityLevels::Num))
 		{
 			FName QualityName = QualityLevelProperty::QualityLevelToFName(QualityLevelCount++);
 			const FTextFormat Format = NSLOCTEXT("SPerQualityLevelPropertiesWidget", "AddOverrideGroupFor", "Add Override specifically for {0}");
diff --git a/Engine/Source/Editor/StaticMeshEditor/Private/StaticMeshEditorTools.cpp b/Engine/Source/Editor/StaticMeshEditor/Private/StaticMeshEditorTools.cpp
--- a/Engine/Source/Editor/StaticMeshEditor/Private/StaticMeshEditorTools.cpp
+++ b/Engine/Source/Editor/StaticMeshEditor/Private/StaticMeshEditorTools.cpp
@@ -3399,7 +3399,7 @@ void FLevelOfDetailSettingsLayout::AddToDetailsPanel( IDetailLayoutBuilder& Deta
 		]
 		.ValueContent()
 		.MinDesiredWidth((float)(StaticMesh->GetQualityLevelMinLOD().PerQuality.Num() + 1)*125.0f)
-		.MaxDesiredWidth((float)((int32)QualityLevelProperty::EQualityLevels::Num + 1)*125.0f)
+		.MaxDesiredWidth((float)((int32)EPerQualityLevels::Num + 1)*125.0f)
 		[
 			SNew(SPerQualityLevelPropertiesWidget)
 			.OnGenerateWidget(this, &FLevelOfDetailSettingsLayout::GetMinQualityLevelLODWidget)
@@ -3419,7 +3419,7 @@ void FLevelOfDetailSettingsLayout::AddToDetailsPanel( IDetailLayoutBuilder& Deta
 		]
 		.ValueContent()
 		.MinDesiredWidth(float(StaticMesh->GetNoRefStreamingLODBias().PerQuality.Num() + 1) * 125.f)
-		.MaxDesiredWidth(float((int32)QualityLevelProperty::EQualityLevels::Num + 1) * 125.f)
+		.MaxDesiredWidth(float((int32)EPerQualityLevels::Num + 1) * 125.f)
 		[
 			SNew(SPerQualityLevelPropertiesWidget)
 			.OnGenerateWidget(this, &FLevelOfDetailSettingsLayout::GetNoRefStreamingLODBiasWidget)
diff --git a/Engine/Source/Runtime/Engine/Classes/Engine/SkeletalMesh.h b/Engine/Source/Runtime/Engine/Classes/Engine/SkeletalMesh.h
--- a/Engine/Source/Runtime/Engine/Classes/Engine/SkeletalMesh.h
+++ b/Engine/Source/Runtime/Engine/Classes/Engine/SkeletalMesh.h
@@ -849,6 +849,30 @@ public:
 	UPROPERTY(EditAnywhere, Category = LODSettings, meta = (DisplayName = "Quality Level Minimum LOD"))
 	FPerQualityLevelInt MinQualityLevelLOD;
 
+	UFUNCTION(BlueprintCallable, Category = StaticMesh, Meta = (ToolTip = "Allow to override min lod quality levels on a skeletalMesh and it Default value (-1 value for Default dont override its value)."))
+	void SetMinLODForQualityLevels(const TMap<EPerQualityLevels, int32>& QualityLevelMinimumLODs, int32 Default = -1)
+	{
+#if WITH_EDITORONLY_DATA
+		WaitUntilAsyncPropertyReleased(ESkeletalMeshAsyncProperties::MinLod);
+		PRAGMA_DISABLE_DEPRECATION_WARNINGS
+		MinQualityLevelLOD.PerQuality = QualityLevelProperty::ConvertQualtiyLevelData(QualityLevelMinimumLODs);
+		MinQualityLevelLOD.Default = Default >=0 ? Default : MinQualityLevelLOD.Default;
+		PRAGMA_ENABLE_DEPRECATION_WARNINGS
+#endif
+	}
+
+	UFUNCTION(BlueprintPure, Category = StaticMesh)
+	void GetMinLODForQualityLevels(TMap<EPerQualityLevels, int32>& QualityLevelMinimumLODs, int32& Default) const
+	{
+#if WITH_EDITORONLY_DATA
+		WaitUntilAsyncPropertyReleased(ESkeletalMeshAsyncProperties::MinLod, EAsyncPropertyLockType::ReadOnly);
+		PRAGMA_DISABLE_DEPRECATION_WARNINGS
+		QualityLevelMinimumLODs = QualityLevelProperty::ConvertQualtiyLevelData(MinQualityLevelLOD.PerQuality);
+		Default = MinQualityLevelLOD.Default;
+		PRAGMA_ENABLE_DEPRECATION_WARNINGS
+#endif
+	}
+
 	static FName GetQualityLevelMinLodMemberName()
 	{
 		PRAGMA_DISABLE_DEPRECATION_WARNINGS
diff --git a/Engine/Source/Runtime/Engine/Classes/Engine/StaticMesh.h b/Engine/Source/Runtime/Engine/Classes/Engine/StaticMesh.h
--- a/Engine/Source/Runtime/Engine/Classes/Engine/StaticMesh.h
+++ b/Engine/Source/Runtime/Engine/Classes/Engine/StaticMesh.h
@@ -788,6 +788,30 @@ public:
 		return INDEX_NONE;
 	}
 
+	UFUNCTION(BlueprintCallable, Category = StaticMesh, Meta = (ToolTip = "Allow to override min lod quality levels on a staticMesh and it Default value (-1 value for Default dont override its value)."))
+	void SetMinLODForQualityLevels(const TMap<EPerQualityLevels, int32>& QualityLevelMinimumLODs, int32 Default = -1)
+	{
+#if WITH_EDITORONLY_DATA
+		WaitUntilAsyncPropertyReleased(EStaticMeshAsyncProperties::MinLOD);
+		PRAGMA_DISABLE_DEPRECATION_WARNINGS
+		MinQualityLevelLOD.PerQuality = QualityLevelProperty::ConvertQualtiyLevelData(QualityLevelMinimumLODs);
+		MinQualityLevelLOD.Default = Default >= 0 ? Default : MinQualityLevelLOD.Default;
+		PRAGMA_ENABLE_DEPRECATION_WARNINGS
+#endif
+	}
+
+	UFUNCTION(BlueprintPure, Category = StaticMesh)
+	void GetMinLODForQualityLevels(TMap<EPerQualityLevels, int32>& QualityLevelMinimumLODs, int32& Default) const
+	{
+#if WITH_EDITORONLY_DATA
+		WaitUntilAsyncPropertyReleased(EStaticMeshAsyncProperties::MinLOD);
+		PRAGMA_DISABLE_DEPRECATION_WARNINGS
+		QualityLevelMinimumLODs = QualityLevelProperty::ConvertQualtiyLevelData(MinQualityLevelLOD.PerQuality);
+		Default = MinQualityLevelLOD.Default;
+		PRAGMA_ENABLE_DEPRECATION_WARNINGS
+#endif
+	}
+
 	/*Choose either PerPlatform or PerQuality override. Note: Enable PerQuality override in the Project Settings/ General Settings/ UseStaticMeshMinLODPerQualityLevels*/
 	ENGINE_API int32 GetMinLODIdx(bool bForceLowestLODIdx = false) const;
 	ENGINE_API int32 GetDefaultMinLOD() const;
diff --git a/Engine/Source/Runtime/Engine/Private/PerQualityLevelProperties.cpp b/Engine/Source/Runtime/Engine/Private/PerQualityLevelProperties.cpp
--- a/Engine/Source/Runtime/Engine/Private/PerQualityLevelProperties.cpp
+++ b/Engine/Source/Runtime/Engine/Private/PerQualityLevelProperties.cpp
@@ -15,9 +15,33 @@ namespace QualityLevelProperty
 	static TArray<FName> QualityLevelNames = { FName("Low"), FName("Medium"), FName("High"), FName("Epic"), FName("Cinematic") };
 	static FString QualityLevelMappingStr = TEXT("QualityLevelMapping");
 
+	TMap<int32, int32> ConvertQualtiyLevelData(const TMap<EPerQualityLevels, int32>& Data)
+	{
+		TMap<int32, int32> ConvertedData;
+
+		for (const TPair<EPerQualityLevels, int32>& Pair : Data)
+		{
+			ConvertedData.Add((int32)Pair.Key,Pair.Value);
+		}
+
+		return ConvertedData;
+	}
+
+	TMap<EPerQualityLevels, int32> ConvertQualtiyLevelData(const TMap<int32, int32>& Data)
+	{
+		TMap<EPerQualityLevels, int32> ConvertedData;
+
+		for (const TPair<int32, int32>& Pair : Data)
+		{
+			ConvertedData.Add((EPerQualityLevels)Pair.Key, Pair.Value);
+		}
+
+		return ConvertedData;
+	}
+
 	FName QualityLevelToFName(int32 QL)
 	{
-		if (QL >= 0 && QL < static_cast<int32>(EQualityLevels::Num))
+		if (QL >= 0 && QL < static_cast<int32>(EPerQualityLevels::Num))
 		{
 			return QualityLevelNames[QL];
 		}
@@ -177,9 +201,9 @@ FSupportedQualityLevelArray FPerQualityLevelProperty<_StructType, _ValueType, _B
 	FConfigCacheIni::LoadLocalIniFile(ScalabilitySettings, TEXT("Scalability"), true, InPlatformName);
 
 	//check all possible quality levels specify in the scalability ini 
-	for (int32 QualityLevel = 0; QualityLevel < (int32)QualityLevelProperty::EQualityLevels::Num; ++QualityLevel)
+	for (int32 QualityLevel = 0; QualityLevel < (int32)EPerQualityLevels::Num; ++QualityLevel)
 	{
-		FString QualitLevelSectionName = Scalability::GetScalabilitySectionString(*ScalabilitySection, QualityLevel, (int32)QualityLevelProperty::EQualityLevels::Num);
+		FString QualitLevelSectionName = Scalability::GetScalabilitySectionString(*ScalabilitySection, QualityLevel, (int32)EPerQualityLevels::Num);
 		PropertyQualityLevel = -1;
 		ScalabilitySettings.GetInt(*QualitLevelSectionName, *CVarName, PropertyQualityLevel);
 
diff --git a/Engine/Source/Runtime/Engine/Private/StaticMesh.cpp b/Engine/Source/Runtime/Engine/Private/StaticMesh.cpp
--- a/Engine/Source/Runtime/Engine/Private/StaticMesh.cpp
+++ b/Engine/Source/Runtime/Engine/Private/StaticMesh.cpp
@@ -7921,7 +7921,7 @@ ENGINE_API int32 UStaticMesh::GetMinLODIdx(bool bForceLowestLODIdx) const
 		extern int32 GUseMobileLODBiasOnDesktopES31;
 		if (GUseMobileLODBiasOnDesktopES31 != 0 && GMaxRHIFeatureLevel == ERHIFeatureLevel::ES3_1)
 		{
-			CurrentMinLodQualityLevel = (int32)QualityLevelProperty::EQualityLevels::Low;
+			CurrentMinLodQualityLevel = (int32)EPerQualityLevels::Low;
 		}
 #endif
 		return bForceLowestLODIdx ? GetQualityLevelMinLOD().GetLowestValue() : GetQualityLevelMinLOD().GetValue(CurrentMinLodQualityLevel);
diff --git a/Engine/Source/Runtime/Engine/Public/PerQualityLevelProperties.h b/Engine/Source/Runtime/Engine/Public/PerQualityLevelProperties.h
--- a/Engine/Source/Runtime/Engine/Public/PerQualityLevelProperties.h
+++ b/Engine/Source/Runtime/Engine/Public/PerQualityLevelProperties.h
@@ -23,10 +23,22 @@ FPerQualityLevelProperties.h: Property types that can be overridden on a quality
 	typedef TSet<int32> FSupportedQualityLevelArray;
 #endif
 
+
+UENUM(BlueprintType)
+enum class EPerQualityLevels : uint8
+{
+	Low,
+	Medium,
+	High,
+	Epic,
+	Cinematic,
+	Num
+};
+
 namespace QualityLevelProperty
 {
 	UENUM()
-	enum class EQualityLevels : uint8
+	enum class UE_DEPRECATED(5.1, "Use EPerQualityLevels instead since we need to expose as an ENUM in blueprint.") EQualityLevels : uint8
 	{
 		Low,
 		Medium,
@@ -38,6 +50,8 @@ namespace QualityLevelProperty
 
 	ENGINE_API FName QualityLevelToFName(int32 QL);
 	ENGINE_API int32 FNameToQualityLevel(FName QL);
+	ENGINE_API TMap<int32, int32> ConvertQualtiyLevelData(const TMap<EPerQualityLevels, int32>& Data);
+	ENGINE_API TMap<EPerQualityLevels, int32> ConvertQualtiyLevelData(const TMap<int32, int32>& Data);
 #if WITH_EDITOR
 	ENGINE_API FSupportedQualityLevelArray PerPlatformOverrideMapping(FString& InPlatformName);
 #endif
