Compile Fix

diff --git a/Engine/Source/Runtime/Engine/Classes/Engine/SkeletalMesh.h b/Engine/Source/Runtime/Engine/Classes/Engine/SkeletalMesh.h
--- a/Engine/Source/Runtime/Engine/Classes/Engine/SkeletalMesh.h
+++ b/Engine/Source/Runtime/Engine/Classes/Engine/SkeletalMesh.h
@@ -865,7 +865,7 @@ public:
 	void GetMinLODForQualityLevels(TMap<EPerQualityLevels, int32>& QualityLevelMinimumLODs, int32& Default) const
 	{
 #if WITH_EDITORONLY_DATA
-		WaitUntilAsyncPropertyReleased(ESkeletalMeshAsyncProperties::MinLod, EAsyncPropertyLockType::ReadOnly);
+		WaitUntilAsyncPropertyReleased(ESkeletalMeshAsyncProperties::MinLod, ESkinnedAssetAsyncPropertyLockType::ReadOnly);
 		PRAGMA_DISABLE_DEPRECATION_WARNINGS
 		QualityLevelMinimumLODs = QualityLevelProperty::ConvertQualtiyLevelData(MinQualityLevelLOD.PerQuality);
 		Default = MinQualityLevelLOD.Default;
