When deleting an actor folder, allow to cleanup other folders marked for delete if their external package is writable.
Because this is now allowed, removed warnings related to missing parent folder on actor and actor folder.

diff --git a/Engine/Source/Runtime/Engine/Private/ActorEditor.cpp b/Engine/Source/Runtime/Engine/Private/ActorEditor.cpp
--- a/Engine/Source/Runtime/Engine/Private/ActorEditor.cpp
+++ b/Engine/Source/Runtime/Engine/Private/ActorEditor.cpp
@@ -1263,7 +1263,7 @@ void AActor::FixupActorFolder()
 	{
 		if (FolderGuid.IsValid())
 		{
-			UE_LOG(LogLevel, Warning, TEXT("Actor folder %s for actor %s encountered when not using actor folders"), *FolderGuid.ToString(), *GetName());
+			UE_LOG(LogActor, Warning, TEXT("Actor folder %s for actor %s encountered when not using actor folders"), *FolderGuid.ToString(), *GetName());
 			FolderGuid = FGuid();
 		}
 	}
@@ -1286,16 +1286,17 @@ void AActor::FixupActorFolder()
 			}
 		}
 
-		// If still invalid, warn and fallback to root
+		// If still invalid, fallback to root
 		if (!IsActorFolderValid())
 		{
-			UE_LOG(LogLevel, Warning, TEXT("Missing actor folder for actor %s"), *GetName());
+			// Here we don't warn anymore since there's a supported workflow where we allow to delete actor folders 
+			// even when referenced by actors (i.e. when the end result remains the root)
 			SetFolderGuidInternal(FGuid(), /*bBroadcastChange*/ false);
 		}
 
 		if (!FolderPath.IsNone())
 		{
-			UE_LOG(LogLevel, Warning, TEXT("Actor folder path %s for actor %s encountered when using actor folders"), *FolderPath.ToString(), *GetName());
+			UE_LOG(LogActor, Warning, TEXT("Actor folder path %s for actor %s encountered when using actor folders"), *FolderPath.ToString(), *GetName());
 			FolderPath = NAME_None;
 		}
 	}
diff --git a/Engine/Source/Runtime/Engine/Private/ActorFolder.cpp b/Engine/Source/Runtime/Engine/Private/ActorFolder.cpp
--- a/Engine/Source/Runtime/Engine/Private/ActorFolder.cpp
+++ b/Engine/Source/Runtime/Engine/Private/ActorFolder.cpp
@@ -264,7 +264,8 @@ void UActorFolder::FixupParentFolder()
 {
 	if (ParentFolderGuid.IsValid() && !GetOuterULevel()->GetActorFolder(ParentFolderGuid, /*bSkipDeleted*/ false))
 	{
-		UE_LOG(LogLevel, Warning, TEXT("Missing parent actor folder for actor folder %s (%s)"), *FolderLabel, *GetName());
+		// Here we don't warn anymore since there's a supported workflow where we allow to delete actor folders 
+		// even when referenced by other folders (i.e. when the end result remains the root)
 		Modify();
 		ParentFolderGuid.Invalidate();
 	}
diff --git a/Engine/Source/Runtime/Engine/Private/Level.cpp b/Engine/Source/Runtime/Engine/Private/Level.cpp
--- a/Engine/Source/Runtime/Engine/Private/Level.cpp
+++ b/Engine/Source/Runtime/Engine/Private/Level.cpp
@@ -79,6 +79,7 @@ Level.cpp: Level-related functions
 #include "UObject/MetaData.h"
 #include "UObject/LinkerLoad.h"
 #include "WorldPartition/WorldPartitionHelpers.h"
+#include "HAL/PlatformFileManager.h"
 #endif
 #include "WorldPartition/WorldPartition.h"
 #include "WorldPartition/WorldPartitionSubsystem.h"
@@ -2685,9 +2686,37 @@ void ULevel::OnFolderMarkAsDeleted(UActorFolder* InActorFolder)
 		{
 			FolderLabelToActorFolders.Remove(InActorFolder->GetLabel());
 		}
-		if (GAllowCleanupActorFolders && GetDeletedAndUnreferencedActorFolders().Contains(InActorFolder->GetGuid()))
+		if (GAllowCleanupActorFolders)
 		{
-			RemoveActorFolder(InActorFolder);
+			// Get all unreferenced folders or folders that resolves to the root ("/")
+			for (const FGuid& GuidFolder : GetDeletedAndUnreferencedActorFolders())
+			{
+				// Remove our folder if it's part of these folders
+				if (GuidFolder == InActorFolder->GetGuid())
+				{
+					RemoveActorFolder(InActorFolder);
+				}
+				// Also allow to remove other folders of this list only if their external package is writable
+				else if (UActorFolder* ActorFolder = GetActorFolder(GuidFolder, false))
+				{
+					if (!ActorFolder->GetParent())
+					{
+						if (UPackage* Package = ActorFolder->GetExternalPackage())
+						{
+							FPackagePath PackagePath(Package->GetLoadedPath());
+							if (FPackageName::DoesPackageExist(PackagePath, &PackagePath))
+							{
+								FString FilePath = PackagePath.GetLocalFullPath();
+								const bool bReadOnly = FPlatformFileManager::Get().GetPlatformFile().IsReadOnly(*FilePath);
+								if (!bReadOnly)
+								{
+									RemoveActorFolder(ActorFolder);
+								}
+							}
+						}
+					}
+				}
+			}
 		}
 	}
 }
@@ -3707,6 +3736,10 @@ TSet<FGuid> ULevel::GetDeletedAndUnreferencedActorFolders() const
 		if (ParentFolder)
 		{
 			FoldersToDelete.Remove(ParentFolder->GetGuid());
+			if (FoldersToDelete.IsEmpty())
+			{
+				break;
+			}
 		}
 	}
 
@@ -3716,7 +3749,7 @@ TSet<FGuid> ULevel::GetDeletedAndUnreferencedActorFolders() const
 		{
 			AActor* Actor = ActorDesc->GetActor();
 			FoldersToDelete.Remove(Actor ? Actor->GetFolderGuid() : ActorDesc->GetFolderGuid());
-			return true;
+			return !FoldersToDelete.IsEmpty();
 		});
 	}
 
@@ -3726,6 +3759,10 @@ TSet<FGuid> ULevel::GetDeletedAndUnreferencedActorFolders() const
 		if (IsValid(Actor))
 		{
 			FoldersToDelete.Remove(Actor->GetFolderGuid());
+			if (FoldersToDelete.IsEmpty())
+			{
+				break;
+			}
 		}
 	}
 
@@ -3737,7 +3774,7 @@ TSet<FGuid> ULevel::GetDeletedAndUnreferencedActorFolders() const
 void ULevel::CleanupDeletedAndUnreferencedActorFolders()
 {
 	// Remove actor folders marked as deleted
-	for (const FGuid GuidFolderToDelete : GetDeletedAndUnreferencedActorFolders())
+	for (const FGuid& GuidFolderToDelete : GetDeletedAndUnreferencedActorFolders())
 	{
 		if (UActorFolder* ActorFolderToDelete = GetActorFolder(GuidFolderToDelete, false /*bSkipDeleted*/))
 		{
