Lumen Mesh Cards - fixed randomly missing mesh cards cooks. Now mesh cards always wait on distance fields on which they depend and they fill appropriate platform RenderData. Also fixed MovieQueue not waiting for mesh cards to build before rendering.

diff --git a/Engine/Source/Editor/MovieSceneTools/Private/AutomatedLevelSequenceCapture.cpp b/Engine/Source/Editor/MovieSceneTools/Private/AutomatedLevelSequenceCapture.cpp
--- a/Engine/Source/Editor/MovieSceneTools/Private/AutomatedLevelSequenceCapture.cpp
+++ b/Engine/Source/Editor/MovieSceneTools/Private/AutomatedLevelSequenceCapture.cpp
@@ -22,6 +22,7 @@
 #include "Protocols/AudioCaptureProtocol.h"
 #include "ShaderCompiler.h"
 #include "DistanceFieldAtlas.h"
+#include "MeshCardRepresentation.h"
 #include "EntitySystem/MovieSceneEntitySystemLinker.h"
 #include "Systems/MovieSceneMotionVectorSimulationSystem.h"
 
@@ -670,6 +671,13 @@ void UAutomatedLevelSequenceCapture::OnTick(float DeltaSeconds)
 		UE_LOG(LogMovieSceneCapture, Log, TEXT("[%d] Done waiting for Mesh Distance Fields to build."), GFrameCounter);
 	}
 
+	if (GCardRepresentationAsyncQueue && GCardRepresentationAsyncQueue->GetNumOutstandingTasks() > 0)
+	{
+		UE_LOG(LogMovieSceneCapture, Log, TEXT("[%d] Waiting for %d Mesh Cards to finish building..."), GFrameCounter, GCardRepresentationAsyncQueue->GetNumOutstandingTasks());
+		GCardRepresentationAsyncQueue->BlockUntilAllBuildsComplete();
+		UE_LOG(LogMovieSceneCapture, Log, TEXT("[%d] Done waiting for Mesh Cards to build."), GFrameCounter);
+	}
+
 	// Setup the automated capture
 	if (CaptureState == ELevelSequenceCaptureState::Setup)
 	{
diff --git a/Engine/Source/Editor/UnrealEd/Private/Commandlets/GatherTextFromAssetsCommandlet.cpp b/Engine/Source/Editor/UnrealEd/Private/Commandlets/GatherTextFromAssetsCommandlet.cpp
--- a/Engine/Source/Editor/UnrealEd/Private/Commandlets/GatherTextFromAssetsCommandlet.cpp
+++ b/Engine/Source/Editor/UnrealEd/Private/Commandlets/GatherTextFromAssetsCommandlet.cpp
@@ -24,6 +24,7 @@
 #include "PackageHelperFunctions.h"
 #include "ShaderCompiler.h"
 #include "DistanceFieldAtlas.h"
+#include "MeshCardRepresentation.h"
 #include "Templates/UniquePtr.h"
 #include "CollectionManagerModule.h"
 #include "ICollectionManager.h"
@@ -845,6 +846,10 @@ PRAGMA_ENABLE_DEPRECATION_WARNINGS
 		{
 			GDistanceFieldAsyncQueue->ProcessAsyncTasks();
 		}
+		if (GCardRepresentationAsyncQueue)
+		{
+			GCardRepresentationAsyncQueue->ProcessAsyncTasks();
+		}
 
 		// Because packages may not have been resaved after this flagging was implemented, we may have added packages to load that weren't flagged - potential false positives.
 		// The loading process should have reflagged said packages so that only true positives will have this flag.
diff --git a/Engine/Source/Runtime/Engine/Private/DistanceFieldAtlas.cpp b/Engine/Source/Runtime/Engine/Private/DistanceFieldAtlas.cpp
--- a/Engine/Source/Runtime/Engine/Private/DistanceFieldAtlas.cpp
+++ b/Engine/Source/Runtime/Engine/Private/DistanceFieldAtlas.cpp
@@ -806,7 +806,12 @@ void FDistanceFieldAsyncQueue::ProcessAsyncTasks(bool bLimitExecutionTime)
 				COOK_STAT(Timer.AddMiss(DerivedData.Num()));
 			}
 
-			BeginCacheMeshCardRepresentation(Task->TargetPlatform, Task->StaticMesh, *Task->StaticMesh->GetRenderData(), Task->DDCKey, &Task->SourceMeshData);
+			BeginCacheMeshCardRepresentation(
+				Task->TargetPlatform,
+				Task->StaticMesh,
+				Task->StaticMesh->GetPlatformStaticMeshRenderData(Task->StaticMesh, Task->TargetPlatform),
+				Task->DDCKey,
+				&Task->SourceMeshData);
 		}
 
 		delete Task;
diff --git a/Engine/Source/Runtime/Engine/Private/MeshCardRepresentation.cpp b/Engine/Source/Runtime/Engine/Private/MeshCardRepresentation.cpp
--- a/Engine/Source/Runtime/Engine/Private/MeshCardRepresentation.cpp
+++ b/Engine/Source/Runtime/Engine/Private/MeshCardRepresentation.cpp
@@ -501,7 +501,7 @@ void FCardRepresentationAsyncQueue::RescheduleBackgroundTask(FAsyncCardRepresent
 	}
 }
 
-void FCardRepresentationAsyncQueue::BlockUntilBuildComplete(UStaticMesh* StaticMesh, bool bWarnIfBlocked)
+void FCardRepresentationAsyncQueue::BlockUntilBuildComplete(UStaticMesh* InStaticMesh, bool bWarnIfBlocked)
 {
 	// We will track the wait time here, but only the cycles used.
 	// This function is called whether or not an async task is pending, 
@@ -513,7 +513,11 @@ void FCardRepresentationAsyncQueue::BlockUntilBuildComplete(UStaticMesh* StaticM
 	double StartTime = 0;
 
 #if WITH_EDITOR
-	FStaticMeshCompilingManager::Get().FinishCompilation({ StaticMesh });
+	FStaticMeshCompilingManager::Get().FinishCompilation({ InStaticMesh });
+	if (GDistanceFieldAsyncQueue)
+	{
+		GDistanceFieldAsyncQueue->BlockUntilBuildComplete(InStaticMesh, bWarnIfBlocked);
+	}	
 #endif
 
 	TSet<UStaticMesh*> RequiredFinishCompilation;
@@ -527,8 +531,8 @@ void FCardRepresentationAsyncQueue::BlockUntilBuildComplete(UStaticMesh* StaticM
 			FScopeLock Lock(&CriticalSection);
 			for (FAsyncCardRepresentationTask* Task : ReferencedTasks)
 			{
-				if (Task->StaticMesh == StaticMesh ||
-					Task->GenerateSource == StaticMesh)
+				if (Task->StaticMesh == InStaticMesh ||
+					Task->GenerateSource == InStaticMesh)
 				{
 					bReferenced = true;
 
@@ -560,6 +564,13 @@ void FCardRepresentationAsyncQueue::BlockUntilBuildComplete(UStaticMesh* StaticM
 		if (RequiredFinishCompilation.Num())
 		{
 			FStaticMeshCompilingManager::Get().FinishCompilation(RequiredFinishCompilation.Array());
+			if (GDistanceFieldAsyncQueue)
+			{
+				for (UStaticMesh* StaticMesh : RequiredFinishCompilation)
+				{
+					GDistanceFieldAsyncQueue->BlockUntilBuildComplete(StaticMesh, bWarnIfBlocked);
+				}
+			}
 		}
 #endif
 
@@ -585,7 +596,7 @@ void FCardRepresentationAsyncQueue::BlockUntilBuildComplete(UStaticMesh* StaticM
 	{
 		UE_LOG(LogStaticMesh, Display, TEXT("Main thread blocked for %.3fs for async card representation build of %s to complete!  This can happen if the mesh is rebuilt excessively."),
 			(float)(FPlatformTime::Seconds() - StartTime), 
-			*StaticMesh->GetName());
+			*InStaticMesh->GetName());
 	}
 }
 
@@ -596,6 +607,10 @@ void FCardRepresentationAsyncQueue::BlockUntilAllBuildsComplete()
 	{
 #if WITH_EDITOR
 		FStaticMeshCompilingManager::Get().FinishAllCompilation();
+		if (GDistanceFieldAsyncQueue)
+		{
+			GDistanceFieldAsyncQueue->BlockUntilAllBuildsComplete();
+		}
 #endif
 		// Reschedule as highest prio since we're explicitly waiting on them
 		{
diff --git a/Engine/Source/Runtime/Engine/Private/StaticMesh.cpp b/Engine/Source/Runtime/Engine/Private/StaticMesh.cpp
--- a/Engine/Source/Runtime/Engine/Private/StaticMesh.cpp
+++ b/Engine/Source/Runtime/Engine/Private/StaticMesh.cpp
@@ -1632,8 +1632,6 @@ void FStaticMeshRenderData::SerializeInlineDataRepresentations(FArchive& Ar, USt
 	{
 		if (Ar.IsSaving())
 		{
-			// Need to block on GDistanceFieldAsyncQueue as it can generate GCardRepresentationAsyncQueue tasks
-			GDistanceFieldAsyncQueue->BlockUntilBuildComplete(Owner, false);
 			GCardRepresentationAsyncQueue->BlockUntilBuildComplete(Owner, false);
 		}
 
diff --git a/Engine/Source/Runtime/Engine/Public/MeshCardRepresentation.h b/Engine/Source/Runtime/Engine/Public/MeshCardRepresentation.h
--- a/Engine/Source/Runtime/Engine/Public/MeshCardRepresentation.h
+++ b/Engine/Source/Runtime/Engine/Public/MeshCardRepresentation.h
@@ -343,7 +343,7 @@ public:
 	ENGINE_API void CancelAllOutstandingBuilds();
 
 	/** Blocks the main thread until the async build of the specified mesh is complete. */
-	ENGINE_API void BlockUntilBuildComplete(UStaticMesh* StaticMesh, bool bWarnIfBlocked);
+	ENGINE_API void BlockUntilBuildComplete(UStaticMesh* InStaticMesh, bool bWarnIfBlocked);
 
 	/** Blocks the main thread until all async builds complete. */
 	ENGINE_API void BlockUntilAllBuildsComplete();
