Only call FCoreDelegates::OnActorLabelChanged when the ActorLabel actually changes

#rnx

diff --git a/Engine/Source/Runtime/Engine/Private/ActorEditor.cpp b/Engine/Source/Runtime/Engine/Private/ActorEditor.cpp
--- a/Engine/Source/Runtime/Engine/Private/ActorEditor.cpp
+++ b/Engine/Source/Runtime/Engine/Private/ActorEditor.cpp
@@ -1139,14 +1139,14 @@ void AActor::SetActorLabel(const FString& NewActorLabelDirty, bool bMarkDirty)
 				// Store new label
 				Modify(bMarkDirty);
 				ActorLabel = MoveTemp(NewActorLabel);
+
+				FPropertyChangedEvent PropertyEvent(FindFProperty<FProperty>(AActor::StaticClass(), "ActorLabel"));
+				PostEditChangeProperty(PropertyEvent);
+
+				FCoreDelegates::OnActorLabelChanged.Broadcast(this);
 			}
 		}
 	}
-
-	FPropertyChangedEvent PropertyEvent( FindFProperty<FProperty>( AActor::StaticClass(), "ActorLabel" ) );
-	PostEditChangeProperty(PropertyEvent);
-
-	FCoreDelegates::OnActorLabelChanged.Broadcast(this);
 }
 
 bool AActor::IsActorLabelEditable() const
@@ -1156,8 +1156,11 @@ bool AActor::IsActorLabelEditable() const
 
 void AActor::ClearActorLabel()
 {
-	ActorLabel.Reset();
-	FCoreDelegates::OnActorLabelChanged.Broadcast(this);
+	if (!ActorLabel.IsEmpty())
+	{
+		ActorLabel.Reset();
+		FCoreDelegates::OnActorLabelChanged.Broadcast(this);
+	}
 }
 
 FFolder AActor::GetFolder() const
