Move parse exec commands to utility function that can be re-used

diff --git a/Engine/Source/Runtime/Engine/Private/ParseExecCommands.cpp b/Engine/Source/Runtime/Engine/Private/ParseExecCommands.cpp
new file mode 100644
--- /dev/null
+++ b/Engine/Source/Runtime/Engine/Private/ParseExecCommands.cpp
@@ -0,0 +1,78 @@
+// Copyright Epic Games, Inc. All Rights Reserved.
+
+#include "ParseExecCommands.h"
+#include "Misc/Parse.h"
+#include "Engine/Engine.h"
+
+namespace ParseExecCommands
+{
+	TArray<FString> ParseExecCmds(const FString& Line)
+	{
+		// Read the command array, ignoring any commas in single-quotes. 
+				// Convert any single-quotes to double-quotes and skip leading whitespace
+				// This allows passing of strings, e:g -execcmds="exampleCvar '0,1,2,3'"
+		TArray<FString> CommandArray;
+		FString CurrentCommand = "";
+		bool bInQuotes = false;
+		bool bSkippingWhitespace = true;
+		for (int i = 0; i < Line.Len(); i++)
+		{
+			TCHAR CurrentChar = Line[i];
+			if (CurrentChar == '\'')
+			{
+				bInQuotes = !bInQuotes;
+				CurrentCommand += "\"";
+			}
+			else if (CurrentChar == ',' && !bInQuotes)
+			{
+				if (CurrentCommand.Len() > 0)
+				{
+					CommandArray.Add(CurrentCommand);
+					CurrentCommand = "";
+				}
+				bSkippingWhitespace = true;
+			}
+			else
+			{
+				if (bSkippingWhitespace)
+				{
+					bSkippingWhitespace = FChar::IsWhitespace(CurrentChar);
+				}
+				if (!bSkippingWhitespace)
+				{
+					CurrentCommand += CurrentChar;
+				}
+			}
+		}
+		if (CurrentCommand.Len() > 0)
+		{
+			CommandArray.Add(CurrentCommand);
+		}
+
+		return CommandArray;
+	}
+
+	TArray<FString> ParseExecCmdsFromCommandLine(const FString& InKey)
+	{
+		FString Line;
+		if (FParse::Value(FCommandLine::Get(), *(InKey + TEXT("=")), Line, /*bShouldStopOnSeparator*/false))
+		{
+			return ParseExecCmds(Line);
+		}
+
+		return TArray<FString>();
+	}
+
+	void QueueDeferredCommands(const TArray<FString>& CommandArray)
+	{
+		for (const FString& Command : CommandArray)
+		{
+			GEngine->DeferredCommands.Add(Command);
+		}
+	}
+
+	void QueueDeferredCommands(const FString& Line)
+	{
+		QueueDeferredCommands(ParseExecCmds(Line));
+	}
+}
diff --git a/Engine/Source/Runtime/Engine/Private/UnrealEngine.cpp b/Engine/Source/Runtime/Engine/Private/UnrealEngine.cpp
--- a/Engine/Source/Runtime/Engine/Private/UnrealEngine.cpp
+++ b/Engine/Source/Runtime/Engine/Private/UnrealEngine.cpp
@@ -223,6 +223,10 @@ UnrealEngine.cpp: Implements the UEngine class and helpers.
 #include "IAutomationWorkerModule.h"
 #endif	// UE_BUILD_SHIPPING
 
+#if !(UE_BUILD_SHIPPING) || ENABLE_PGO_PROFILE
+#include "ParseExecCommands.h"
+#endif
+
 #if ENABLE_LOC_TESTING
 #include "LocalizationModule.h"
 #endif
@@ -1956,55 +1960,7 @@ void UEngine::Init(IEngineLoop* InEngineLoop)
 	}
 
 	// Optionally exec commands passed in the command line.
-	FString ExecCmds;
-	if( FParse::Value(FCommandLine::Get(), TEXT("ExecCmds="), ExecCmds, false) )
-	{
-		// Read the command array, ignoring any commas in single-quotes. 
-		// Convert any single-quotes to double-quotes and skip leading whitespace
-		// This allows passing of strings, e:g -execcmds="exampleCvar '0,1,2,3'"
-		TArray<FString> CommandArray;
-		FString CurrentCommand = "";
-		bool bInQuotes = false;
-		bool bSkippingWhitespace = true;
-		for (int i = 0; i < ExecCmds.Len(); i++)
-		{
-			TCHAR CurrentChar = ExecCmds[i];
-			if (CurrentChar == '\'')
-			{
-				bInQuotes = !bInQuotes;
-				CurrentCommand += "\"";
-			}
-			else if (CurrentChar == ',' && !bInQuotes)
-			{
-				if (CurrentCommand.Len() > 0)
-				{
-					CommandArray.Add(CurrentCommand);
-					CurrentCommand = "";
-				}
-				bSkippingWhitespace = true;
-			}
-			else
-			{
-				if (bSkippingWhitespace)
-				{
-					bSkippingWhitespace = FChar::IsWhitespace(CurrentChar);
-				}
-				if (!bSkippingWhitespace)
-				{
-					CurrentCommand += CurrentChar;
-				}
-			}
-		}
-		if (CurrentCommand.Len() > 0)
-		{
-			CommandArray.Add(CurrentCommand);
-		}
-
-		for( int32 Cx = 0; Cx < CommandArray.Num(); ++Cx )
-		{
-			new(GEngine->DeferredCommands) FString(*CommandArray[Cx]);
-		}
-	}
+	ParseExecCommands::QueueDeferredCommands(ParseExecCommands::ParseExecCmdsFromCommandLine(TEXT("ExecCmds")));
 
 	// optionally set the vsync console variable
 	if( FParse::Param(FCommandLine::Get(), TEXT("vsync")) )
diff --git a/Engine/Source/Runtime/Engine/Public/ParseExecCommands.h b/Engine/Source/Runtime/Engine/Public/ParseExecCommands.h
new file mode 100644
--- /dev/null
+++ b/Engine/Source/Runtime/Engine/Public/ParseExecCommands.h
@@ -0,0 +1,14 @@
+// Copyright Epic Games, Inc. All Rights Reserved.
+
+#pragma once
+
+#include "CoreTypes.h"
+#include "CoreFwd.h"
+
+namespace ParseExecCommands
+{
+	ENGINE_API TArray<FString> ParseExecCmds(const FString& Line);
+	ENGINE_API TArray<FString> ParseExecCmdsFromCommandLine(const FString& InKey);
+	ENGINE_API void QueueDeferredCommands(const TArray<FString>& CommandArray);
+	ENGINE_API void QueueDeferredCommands(const FString& Line);
+}
\ No newline at end of file
