Niagara:  Corruption in GPU sim particle rendering, caused by FGPUSortManager not running across all GPUs.  This system does incremental updates to indirect buffers, which need to be coherent across all GPUs.  Cross GPU copy wouldn't make sense here, as the GPU work done is tiny, and the cross GPU sync cost would be significantly larger than the work.  Checking this into IRON branch, will repeat fix and have it reviewed in UE5.

#rnx

diff --git a/Engine/Source/Runtime/Engine/Private/GPUSortManager.cpp b/Engine/Source/Runtime/Engine/Private/GPUSortManager.cpp
--- a/Engine/Source/Runtime/Engine/Private/GPUSortManager.cpp
+++ b/Engine/Source/Runtime/Engine/Private/GPUSortManager.cpp
@@ -724,6 +724,7 @@ void FGPUSortManager::UpdateSortBuffersPool()
 void FGPUSortManager::OnPreRender(FRDGBuilder& GraphBuilder)
 {
 	RDG_CSV_STAT_EXCLUSIVE_SCOPE(GraphBuilder, GPUSort);
+	RDG_GPU_MASK_SCOPE(GraphBuilder, FRHIGPUMask::All());
 
 	AddPass(
 		GraphBuilder,
@@ -781,6 +782,7 @@ void FGPUSortManager::OnPostRenderOpaque(FRDGBuilder& GraphBuilder)
 {
 	LLM_SCOPE(ELLMTag::GPUSort);
 	RDG_CSV_STAT_EXCLUSIVE_SCOPE(GraphBuilder, GPUSort);
+	RDG_GPU_MASK_SCOPE(GraphBuilder, FRHIGPUMask::All());
 
 	AddPass(
 		GraphBuilder,
