check CompositeTextureMode out of valid range at Texture load time and warn about it

diff --git a/Engine/Source/Runtime/Engine/Private/Texture.cpp b/Engine/Source/Runtime/Engine/Private/Texture.cpp
--- a/Engine/Source/Runtime/Engine/Private/Texture.cpp
+++ b/Engine/Source/Runtime/Engine/Private/Texture.cpp
@@ -889,6 +889,14 @@ void UTexture::Serialize(FArchive& Ar)
 				}
 			}
 		}
+
+		// CTM_MAX is mis-named, it's actually one higher than the maximum valid value
+		if ( (uint32)CompositeTextureMode >= (uint32)CTM_MAX )
+		{
+			UE_LOG(LogTexture, Warning, TEXT("CompositeTextureMode was invalid in uasset, disabled.  Fix content and re-save : %s"), *GetPathName());
+			
+			CompositeTextureMode = CTM_Disabled;
+		}
 	}
 #endif // #if WITH_EDITORONLY_DATA
 }
