BlendSpace, BlendSpaceUtilities : Added protection for the case where animations are deleted while running

diff --git a/Engine/Source/Runtime/Engine/Private/Animation/BlendSpace.cpp b/Engine/Source/Runtime/Engine/Private/Animation/BlendSpace.cpp
--- a/Engine/Source/Runtime/Engine/Private/Animation/BlendSpace.cpp
+++ b/Engine/Source/Runtime/Engine/Private/Animation/BlendSpace.cpp
@@ -459,28 +459,31 @@ void UBlendSpace::TickAssetPlayer(FAnimTickRecord& Instance, struct FAnimNotifyQ
 					FBlendSampleData& SampleDataItem = SampleDataList[HighestMarkerSyncWeightIndex];
 					const FBlendSample& Sample = SampleData[SampleDataItem.SampleDataIndex];
 
-					bool bResetMarkerDataOnFollowers = false;
-					if (!Instance.MarkerTickRecord->IsValid(Instance.bLooping))
+					if (Sample.Animation)
 					{
-						SampleDataItem.MarkerTickRecord.Reset();
-						bResetMarkerDataOnFollowers = true;
-						SampleDataItem.Time = NormalizedCurrentTime * Sample.Animation->GetPlayLength();
-					}
-					else if (!SampleDataItem.MarkerTickRecord.IsValid(Instance.bLooping) && Context.MarkerTickContext.GetMarkerSyncStartPosition().IsValid())
-					{
-						Sample.Animation->GetMarkerIndicesForPosition(Context.MarkerTickContext.GetMarkerSyncStartPosition(), true, SampleDataItem.MarkerTickRecord.PreviousMarker, SampleDataItem.MarkerTickRecord.NextMarker, SampleDataItem.Time, Instance.MirrorDataTable);
-					}
+						bool bResetMarkerDataOnFollowers = false;
+						if (!Instance.MarkerTickRecord->IsValid(Instance.bLooping))
+						{
+							SampleDataItem.MarkerTickRecord.Reset();
+							bResetMarkerDataOnFollowers = true;
+							SampleDataItem.Time = NormalizedCurrentTime * Sample.Animation->GetPlayLength();
+						}
+						else if (!SampleDataItem.MarkerTickRecord.IsValid(Instance.bLooping) && Context.MarkerTickContext.GetMarkerSyncStartPosition().IsValid())
+						{
+							Sample.Animation->GetMarkerIndicesForPosition(Context.MarkerTickContext.GetMarkerSyncStartPosition(), true, SampleDataItem.MarkerTickRecord.PreviousMarker, SampleDataItem.MarkerTickRecord.NextMarker, SampleDataItem.Time, Instance.MirrorDataTable);
+						}
 
-					const float NewDeltaTime = Context.GetDeltaTime() * Instance.PlayRateMultiplier * Sample.RateScale * Sample.Animation->RateScale;
-					if (!FMath::IsNearlyZero(NewDeltaTime))
-					{
-						Context.SetLeaderDelta(NewDeltaTime);
-						Sample.Animation->TickByMarkerAsLeader(SampleDataItem.MarkerTickRecord, Context.MarkerTickContext, SampleDataItem.Time, SampleDataItem.PreviousTime, NewDeltaTime, Instance.bLooping, Instance.MirrorDataTable);
-						check(!Instance.bLooping || Context.MarkerTickContext.IsMarkerSyncStartValid());
-						TickFollowerSamples(SampleDataList, HighestMarkerSyncWeightIndex, Context, bResetMarkerDataOnFollowers, Instance.MirrorDataTable);
+						const float NewDeltaTime = Context.GetDeltaTime() * Instance.PlayRateMultiplier * Sample.RateScale * Sample.Animation->RateScale;
+						if (!FMath::IsNearlyZero(NewDeltaTime))
+						{
+							Context.SetLeaderDelta(NewDeltaTime);
+							Sample.Animation->TickByMarkerAsLeader(SampleDataItem.MarkerTickRecord, Context.MarkerTickContext, SampleDataItem.Time, SampleDataItem.PreviousTime, NewDeltaTime, Instance.bLooping, Instance.MirrorDataTable);
+							check(!Instance.bLooping || Context.MarkerTickContext.IsMarkerSyncStartValid());
+							TickFollowerSamples(SampleDataList, HighestMarkerSyncWeightIndex, Context, bResetMarkerDataOnFollowers, Instance.MirrorDataTable);
+						}
+						NormalizedCurrentTime = SampleDataItem.Time / Sample.Animation->GetPlayLength();
+						*Instance.MarkerTickRecord = SampleDataItem.MarkerTickRecord;
 					}
-					NormalizedCurrentTime = SampleDataItem.Time / Sample.Animation->GetPlayLength();
-					*Instance.MarkerTickRecord = SampleDataItem.MarkerTickRecord;
 				}
 				else
 				{
@@ -508,17 +511,20 @@ void UBlendSpace::TickAssetPlayer(FAnimTickRecord& Instance, struct FAnimNotifyQ
 					FBlendSampleData& SampleDataItem = SampleDataList[HighestWeightIndex];
 					const FBlendSample& Sample = SampleData[SampleDataItem.SampleDataIndex];
 
-					if (Context.GetDeltaTime() != 0.f)
+					if (Sample.Animation)
 					{
-						if (!Instance.MarkerTickRecord->IsValid(Instance.bLooping))
+						if (Context.GetDeltaTime() != 0.f)
 						{
-							SampleDataItem.Time = NormalizedCurrentTime * Sample.Animation->GetPlayLength();
-						}
+							if (!Instance.MarkerTickRecord->IsValid(Instance.bLooping))
+							{
+								SampleDataItem.Time = NormalizedCurrentTime * Sample.Animation->GetPlayLength();
+							}
 
-						TickFollowerSamples(SampleDataList, -1, Context, false, Instance.MirrorDataTable);
+							TickFollowerSamples(SampleDataList, -1, Context, false, Instance.MirrorDataTable);
+						}
+						*Instance.MarkerTickRecord = SampleDataItem.MarkerTickRecord;
+						NormalizedCurrentTime = SampleDataItem.Time / Sample.Animation->GetPlayLength();
 					}
-					*Instance.MarkerTickRecord = SampleDataItem.MarkerTickRecord;
-					NormalizedCurrentTime = SampleDataItem.Time / Sample.Animation->GetPlayLength();
 				}
 				else
 				{
diff --git a/Engine/Source/Runtime/Engine/Private/Animation/BlendSpaceUtilities.cpp b/Engine/Source/Runtime/Engine/Private/Animation/BlendSpaceUtilities.cpp
--- a/Engine/Source/Runtime/Engine/Private/Animation/BlendSpaceUtilities.cpp
+++ b/Engine/Source/Runtime/Engine/Private/Animation/BlendSpaceUtilities.cpp
@@ -75,7 +75,8 @@ int32 FBlendSpaceUtilities::GetHighestWeightMarkerSyncSample(const TArray<struct
 	{
 		const FBlendSampleData& SampleData = SampleDataList[I];
 		if (SampleData.GetClampedWeight() > HighestWeight &&
-			BlendSamples[SampleData.SampleDataIndex].Animation->AuthoredSyncMarkers.Num() > 0)
+			(BlendSamples[SampleData.SampleDataIndex].Animation != nullptr &&
+			 BlendSamples[SampleData.SampleDataIndex].Animation->AuthoredSyncMarkers.Num() > 0))
 		{
 			HighestWeightIndex = I;
 			HighestWeight = SampleData.GetClampedWeight();
