Fix some object searching for their level owning world at load time not being able to find it because it was assigned too late

diff --git a/Engine/Source/Runtime/Engine/Private/Level.cpp b/Engine/Source/Runtime/Engine/Private/Level.cpp
--- a/Engine/Source/Runtime/Engine/Private/Level.cpp
+++ b/Engine/Source/Runtime/Engine/Private/Level.cpp
@@ -1078,6 +1078,21 @@ void ULevel::PostLoad()
 {
 	Super::PostLoad();
 
+	// Ensure that the level is pointed to the owning world.  For streamed levels, this will be the world of the P map
+	// they are streamed in to which we cached when the package loading was invoked
+	// We also need to do this before doing additional loading to avoid reinstantiation of BP actors
+	// assigned to this level to be without a World.
+	OwningWorld = ULevel::StreamedLevelsOwningWorld.FindRef(GetOutermost()->GetFName()).Get();
+	if (OwningWorld == nullptr)
+	{
+		OwningWorld = CastChecked<UWorld>(GetOuter());
+	}
+	else
+	{
+		// This entry will not be used anymore, remove it
+		ULevel::StreamedLevelsOwningWorld.Remove(GetOutermost()->GetFName());
+	}
+
 #if WITH_EDITOR
 	if (IsUsingActorFolders() && IsUsingExternalObjects() && IsActorFolderObjectsFeatureAvailable())
 	{
@@ -1167,19 +1182,6 @@ void ULevel::PostLoad()
 	}
 #endif
 
-	// Ensure that the level is pointed to the owning world.  For streamed levels, this will be the world of the P map
-	// they are streamed in to which we cached when the package loading was invoked
-	OwningWorld = ULevel::StreamedLevelsOwningWorld.FindRef(GetOutermost()->GetFName()).Get();
-	if (OwningWorld == NULL)
-	{
-		OwningWorld = CastChecked<UWorld>(GetOuter());
-	}
-	else
-	{
-		// This entry will not be used anymore, remove it
-		ULevel::StreamedLevelsOwningWorld.Remove(GetOutermost()->GetFName());
-	}
-
 	UWorldComposition::OnLevelPostLoad(this);
 		
 #if WITH_EDITOR
