AssetManager: Change SearchAssetRegistryPaths to get the list of Blueprint subclasses from the AssetRegistry instead of reading them from UObject global loaded class datas. This takes about the same amount of cputime, but is more consistent with the rest of the AssetManager's queries of AssetDatas.
#rnx

diff --git a/Engine/Source/Runtime/Engine/Private/AssetManager.cpp b/Engine/Source/Runtime/Engine/Private/AssetManager.cpp
--- a/Engine/Source/Runtime/Engine/Private/AssetManager.cpp
+++ b/Engine/Source/Runtime/Engine/Private/AssetManager.cpp
@@ -565,12 +565,8 @@ int32 UAssetManager::SearchAssetRegistryPaths(TArray<FAssetData>& OutAssetDataLi
 		else
 		{
 			// Search for all blueprints and then check derived classes later
-			TArray<UClass*> BlueprintCoreDerivedClasses;
-			GetDerivedClasses(UBlueprintCore::StaticClass(), BlueprintCoreDerivedClasses);
-			for (UClass* BPCoreClass : BlueprintCoreDerivedClasses)
-			{
-				ARFilter.ClassPaths.Add(BPCoreClass->GetClassPathName());
-			}
+			ARFilter.ClassPaths.Add(UBlueprintCore::StaticClass()->GetClassPathName());
+			ARFilter.bRecursiveClasses = true;
 
 			GetAssetRegistry().GetDerivedClassNames(CompiledRules.AssetClassNames, TSet<FTopLevelAssetPath>(), CompiledRules.DerivedClassNames);
 		}
