WorldPartition: Fix Actor unloading after being saved in WP with Streaming disabled (Level Instance)

#rnx

diff --git a/Engine/Source/Runtime/Engine/Private/WorldPartition/WorldPartition.cpp b/Engine/Source/Runtime/Engine/Private/WorldPartition/WorldPartition.cpp
--- a/Engine/Source/Runtime/Engine/Private/WorldPartition/WorldPartition.cpp
+++ b/Engine/Source/Runtime/Engine/Private/WorldPartition/WorldPartition.cpp
@@ -1205,7 +1205,7 @@ void UWorldPartition::Tick(float DeltaSeconds)
 			if (!DirtyActorIt.Key().IsValid() || !DirtyActorIt.Value()->GetPackage()->IsDirty())
 			{
 				// If we hold the last reference to that actor (or no reference are held at all), pin it to avoid unloading
-				if (IsStreamingEnabled() && DirtyActorIt.Key().IsValid() && DirtyActorIt.Key()->GetHardRefCount() <= 1)
+				if (DirtyActorIt.Key().IsValid() && DirtyActorIt.Key()->GetHardRefCount() <= 1)
 				{
 					PinnedActors->AddActors({ DirtyActorIt.Key().ToHandle() });
 				}
