BlendSpace: Added protection for the case where animations are deleted while previewing a Blend Space

diff --git a/Engine/Source/Runtime/Engine/Private/Animation/BlendSpace.cpp b/Engine/Source/Runtime/Engine/Private/Animation/BlendSpace.cpp
--- a/Engine/Source/Runtime/Engine/Private/Animation/BlendSpace.cpp
+++ b/Engine/Source/Runtime/Engine/Private/Animation/BlendSpace.cpp
@@ -1921,7 +1921,17 @@ void UBlendSpace::UpdateBlendSpacesUsingAnimSequence(UAnimSequenceBase* Sequence
 
 TArray<FName>* UBlendSpace::GetUniqueMarkerNames()
 {
-	return (SampleIndexWithMarkers != INDEX_NONE && SampleData.Num() > 0) ? SampleData[SampleIndexWithMarkers].Animation->GetUniqueMarkerNames() : nullptr;
+	if (SampleIndexWithMarkers != INDEX_NONE
+		&& SampleData.Num() > SampleIndexWithMarkers)
+	{
+		FBlendSample& BlendSample = SampleData[SampleIndexWithMarkers];
+		if (BlendSample.Animation != nullptr)
+		{
+			return BlendSample.Animation->GetUniqueMarkerNames();
+		}
+	}
+
+	return nullptr;
 }
 
 void UBlendSpace::GetRawSamplesFromBlendInput(const FVector& BlendInput, TArray<FGridBlendSample, TInlineAllocator<4> >& OutBlendSamples) const
