Proper Multi-axis support for forced input commands.

#rnx

diff --git a/Engine/Plugins/EnhancedInput/Source/EnhancedInput/Private/EnhancedInputSubsystemInterface.cpp b/Engine/Plugins/EnhancedInput/Source/EnhancedInput/Private/EnhancedInputSubsystemInterface.cpp
--- a/Engine/Plugins/EnhancedInput/Source/EnhancedInput/Private/EnhancedInputSubsystemInterface.cpp
+++ b/Engine/Plugins/EnhancedInput/Source/EnhancedInput/Private/EnhancedInputSubsystemInterface.cpp
@@ -779,7 +779,7 @@ void InjectKey(T* InjectVia, FKey Key, const FInputActionValue& Value, float Del
 	// TODO: Multi axis FKey support
 	if (Key.IsAnalog())
 	{
-		InjectVia->InputKey(FInputKeyParams(Key, static_cast<double>(Value.Get<float>()), DeltaTime, 1, Key.IsGamepadKey()));
+		InjectVia->InputKey(FInputKeyParams(Key, Value.Get<FVector>(), DeltaTime, 1, Key.IsGamepadKey()));
 	}
 	else
 	{
diff --git a/Engine/Source/Runtime/Engine/Classes/GameFramework/PlayerInput.h b/Engine/Source/Runtime/Engine/Classes/GameFramework/PlayerInput.h
--- a/Engine/Source/Runtime/Engine/Classes/GameFramework/PlayerInput.h
+++ b/Engine/Source/Runtime/Engine/Classes/GameFramework/PlayerInput.h
@@ -354,6 +354,15 @@ struct ENGINE_API FInputKeyParams
 		, bIsGamepadOverride(bGamepadOverride)
 	{};
 
+	FInputKeyParams(FKey InKey, FVector InDelta, float InDeltaTime, int32 InNumSamples, bool bGamepadOverride = false, FInputDeviceId InInputDevice = INPUTDEVICEID_NONE)
+		: Key(InKey)
+		, InputDevice(InInputDevice)
+		, NumSamples(InNumSamples)
+		, DeltaTime(InDeltaTime)
+		, Delta(InDelta)
+		, bIsGamepadOverride(bGamepadOverride)
+	{};
+
 	/** The key that has been pressed */
 	FKey Key = EKeys::Invalid;
 
