Fixing RemoveDestroyedGuidsByLevel in the wrong scope

diff --git a/Engine/Source/Runtime/Engine/Private/NetDriver.cpp b/Engine/Source/Runtime/Engine/Private/NetDriver.cpp
--- a/Engine/Source/Runtime/Engine/Private/NetDriver.cpp
+++ b/Engine/Source/Runtime/Engine/Private/NetDriver.cpp
@@ -6314,48 +6314,48 @@ void UNetDriver::OnLevelRemovedFromWorld(class ULevel* Level, class UWorld* InWo
 			}
 		}	
 		else
-	{
-			for (AActor* Actor : Level->Actors)
 		{
-			if (Actor)
+			for (AActor* Actor : Level->Actors)
 			{
+				if (Actor)
+				{
 					// Keep Startup actors alive, because they haven't been destroyed yet.
 					// Technically, Dynamic actors may not have been destroyed either, but this
 					// resembles relevancy.
 					if (!Actor->IsNetStartupActor())
 					{
-				NotifyActorLevelUnloaded(Actor);
+						NotifyActorLevelUnloaded(Actor);
 					}
 					else
 					{
 						// We still want to remove Startup actors from the Network list so they aren't processed anymore.
 						GetNetworkObjectList().Remove(Actor);
 					}
+				}
 			}
-		}
 
-		TArray<FNetworkGUID> RemovedGUIDs;
-		for (auto It = DestroyedStartupOrDormantActors.CreateIterator(); It; ++It)
-		{
+			TArray<FNetworkGUID> RemovedGUIDs;
+			for (auto It = DestroyedStartupOrDormantActors.CreateIterator(); It; ++It)
+			{
 				FActorDestructionInfo* DestructInfo = It->Value.Get();
 				
 				// Until the level is actually unloaded / reloaded, any Static Actors that have been
 				// destroyed will not be recreated, so we still need to track them.
 				if (DestructInfo->Level == Level && !DestructInfo->NetGUID.IsStatic())
-			{
-				for (UNetConnection* Connection : ClientConnections)
 				{
-					Connection->RemoveDestructionInfo(DestructInfo);
-				}
+					for (UNetConnection* Connection : ClientConnections)
+					{
+						Connection->RemoveDestructionInfo(DestructInfo);
+					}
 
-				RemovedGUIDs.Add(It->Key);
-				It.RemoveCurrent();
+					RemovedGUIDs.Add(It->Key);
+					It.RemoveCurrent();
+				}
 			}
 
 			RemoveDestroyedGuidsByLevel(Level, RemovedGUIDs);
 		}
 	}
-	}
 }
 
 bool UNetDriver::ShouldSkipRepNotifies() const
