Add Telemetry Event to Session Summary to Distinguish UE arch on M1

diff --git a/Engine/Source/Runtime/Engine/Private/Analytics/EngineAnalyticsSessionSummary.cpp b/Engine/Source/Runtime/Engine/Private/Analytics/EngineAnalyticsSessionSummary.cpp
--- a/Engine/Source/Runtime/Engine/Private/Analytics/EngineAnalyticsSessionSummary.cpp
+++ b/Engine/Source/Runtime/Engine/Private/Analytics/EngineAnalyticsSessionSummary.cpp
@@ -151,6 +151,13 @@ FEngineAnalyticsSessionSummary::FEngineAnalyticsSessionSummary(TSharedPtr<IAnaly
 
 	// Create the immutable metrics.
 	Store->Set(TEXT("Platform"), FString(FPlatformProperties::IniPlatformName()));
+#if PLATFORM_MAC
+#if PLATFORM_MAC_ARM64
+    Store->Set(TEXT("UEBuildArch"), TEXT("AppleSilicon"));
+#else
+    Store->Set(TEXT("UEBuildArch"), TEXT("Intel(Mac)"));
+#endif
+#endif
 	Store->Set(TEXT("OSMajor"), OSMajor);
 	Store->Set(TEXT("OSMinor"), OSMinor);
 	Store->Set(TEXT("OSVersion"), FPlatformMisc::GetOSVersion());
