Add some additional trace scopes for engine startup

#rnx

diff --git a/Engine/Source/Runtime/Engine/Private/GameEngine.cpp b/Engine/Source/Runtime/Engine/Private/GameEngine.cpp
--- a/Engine/Source/Runtime/Engine/Private/GameEngine.cpp
+++ b/Engine/Source/Runtime/Engine/Private/GameEngine.cpp
@@ -62,6 +62,7 @@
 #include "DynamicResolutionState.h"
 #include "MoviePlayerProxy.h"
 #include "ProfilingDebugging/CsvProfiler.h"
+#include "ProfilingDebugging/CpuProfilerTrace.h"
 #include "RenderTargetPool.h"
 #include "RenderGraphBuilder.h"
 #include "CustomResourcePool.h"
@@ -242,6 +243,7 @@ void UGameEngine::CreateGameViewportWidget( UGameViewportClient* GameViewportCli
 
 void UGameEngine::CreateGameViewport( UGameViewportClient* GameViewportClient )
 {
+	TRACE_CPUPROFILER_EVENT_SCOPE(CreateGameViewport);
 	check(GameViewportWindow.IsValid());
 
 	if( !GameViewportWidget.IsValid() )
@@ -1101,6 +1103,7 @@ public:
 
 void UGameEngine::Init(IEngineLoop* InEngineLoop)
 {
+	TRACE_CPUPROFILER_EVENT_SCOPE(UGameEngine::Init);
 	DECLARE_SCOPE_CYCLE_COUNTER(TEXT("UGameEngine Init"), STAT_GameEngineStartup, STATGROUP_LoadTime);
 
 	if (!GIsEditor)
@@ -1120,11 +1123,15 @@ void UGameEngine::Init(IEngineLoop* InEngineLoop)
 #endif
 
 	// Load and apply user game settings
-	GetGameUserSettings()->LoadSettings();
-	GetGameUserSettings()->ApplyNonResolutionSettings();
+	{
+		TRACE_CPUPROFILER_EVENT_SCOPE(InitGameUserSettings);
+		GetGameUserSettings()->LoadSettings();
+		GetGameUserSettings()->ApplyNonResolutionSettings();
+	}
 
 	// Create game instance.  For GameEngine, this should be the only GameInstance that ever gets created.
 	{
+		TRACE_CPUPROFILER_EVENT_SCOPE(InitGameInstance);
 		FSoftClassPath GameInstanceClassName = GetDefault<UGameMapsSettings>()->GameInstanceClass;
 		UClass* GameInstanceClass = (GameInstanceClassName.IsValid() ? LoadObject<UClass>(NULL, *GameInstanceClassName.ToString()) : UGameInstance::StaticClass());
 		
@@ -1158,6 +1165,7 @@ void UGameEngine::Init(IEngineLoop* InEngineLoop)
 	UGameViewportClient* ViewportClient = NULL;
 	if(GIsClient)
 	{
+		TRACE_CPUPROFILER_EVENT_SCOPE(InitGameViewPortClient);
 		ViewportClient = NewObject<UGameViewportClient>(this, GameViewportClientClass);
 		ViewportClient->Init(*GameInstance->GetWorldContext(), GameInstance);
 		GameViewport = ViewportClient;
@@ -1169,6 +1177,7 @@ void UGameEngine::Init(IEngineLoop* InEngineLoop)
 	// Attach the viewport client to a new viewport.
 	if(ViewportClient)
 	{
+		TRACE_CPUPROFILER_EVENT_SCOPE(AttachGameViewport);
 		// This must be created before any gameplay code adds widgets
 		bool bWindowAlreadyExists = GameViewportWindow.IsValid();
 		if (!bWindowAlreadyExists)
@@ -1190,6 +1199,7 @@ void UGameEngine::Init(IEngineLoop* InEngineLoop)
 			UE_LOG(LogEngine, Fatal,TEXT("%s"),*Error);
 		}
 
+		TRACE_CPUPROFILER_EVENT_SCOPE(BroadcastOnViewPortCreated);
 		UGameViewportClient::OnViewportCreated().Broadcast();
 	}
 
@@ -1201,6 +1211,7 @@ void UGameEngine::Init(IEngineLoop* InEngineLoop)
 
 void UGameEngine::Start()
 {
+	TRACE_CPUPROFILER_EVENT_SCOPE(UGameEngine::Start);
 	UE_LOG(LogInit, Display, TEXT("Starting Game."));
 
 	GameInstance->StartGameInstance();
diff --git a/Engine/Source/Runtime/Engine/Private/GameInstance.cpp b/Engine/Source/Runtime/Engine/Private/GameInstance.cpp
--- a/Engine/Source/Runtime/Engine/Private/GameInstance.cpp
+++ b/Engine/Source/Runtime/Engine/Private/GameInstance.cpp
@@ -31,6 +31,7 @@
 #include "Net/ReplayPlaylistTracker.h"
 #include "ReplaySubsystem.h"
 #include "GenericPlatform/GenericPlatformInputDeviceMapper.h"
+#include "ProfilingDebugging/CpuProfilerTrace.h"
 
 #if WITH_EDITOR
 #include "Settings/LevelEditorPlaySettings.h"
@@ -85,6 +86,7 @@ UEngine* UGameInstance::GetEngine() const
 
 void UGameInstance::Init()
 {
+	TRACE_CPUPROFILER_EVENT_SCOPE(UGameInstance::Init);
 	ReceiveInit();
 
 	if (!IsRunningCommandlet())
@@ -572,6 +574,7 @@ UGameViewportClient* UGameInstance::GetGameViewportClient() const
 
 void UGameInstance::StartGameInstance()
 {
+	TRACE_CPUPROFILER_EVENT_SCOPE(UGameInstance::StartGameInstance);
 	UEngine* const Engine = GetEngine();
 
 	// Create default URL.
diff --git a/Engine/Source/Runtime/Engine/Private/UnrealEngine.cpp b/Engine/Source/Runtime/Engine/Private/UnrealEngine.cpp
--- a/Engine/Source/Runtime/Engine/Private/UnrealEngine.cpp
+++ b/Engine/Source/Runtime/Engine/Private/UnrealEngine.cpp
@@ -1803,6 +1803,7 @@ void UEngine::PerformGarbageCollectionAndCleanupActors()
 //
 void UEngine::Init(IEngineLoop* InEngineLoop)
 {
+	TRACE_CPUPROFILER_EVENT_SCOPE(UEngine::Init);
 	UE_LOG(LogEngine, Log, TEXT("Initializing Engine..."));
 	DECLARE_SCOPE_CYCLE_COUNTER(TEXT("Engine Initialized"), STAT_EngineStartup, STATGROUP_LoadTime);
 
@@ -17914,6 +17915,7 @@ static FAutoConsoleCommand SetupThreadConfigCmd(
 
 void InitThreadConfig()
 {
+	TRACE_CPUPROFILER_EVENT_SCOPE(InitThreadConfig);
 	GetDefaultThreadConfigs();
 
 	if (GConfig)
