Fix Strata crash when converting a material using material expression but having no node connected to its shading model input.

diff --git a/Engine/Source/Runtime/Engine/Private/Materials/Material.cpp b/Engine/Source/Runtime/Engine/Private/Materials/Material.cpp
--- a/Engine/Source/Runtime/Engine/Private/Materials/Material.cpp
+++ b/Engine/Source/Runtime/Engine/Private/Materials/Material.cpp
@@ -3263,10 +3263,19 @@ void UMaterial::ConvertMaterialToStrataMaterial()
 				// rebuilding the final Shading model (see RebuildShadingModelField())
 				if (ShadingModel == MSM_FromMaterialExpression)
 				{
-					check(EditorOnly->ShadingModelFromMaterialExpression.IsConnected());
-
-					// Reconnect the shading model expression
-					MoveConnectionTo(EditorOnly->ShadingModelFromMaterialExpression, ConvertNode, 19);
+					//check(EditorOnly->ShadingModelFromMaterialExpression.IsConnected());
+					if (!EditorOnly->ShadingModelFromMaterialExpression.IsConnected())
+					{
+						// Now setup the mean free path with a hard coded transmittance of 0.75 when viewing the surface perpendicularly
+						UMaterialExpressionConstant* DefaultShadingModel = NewObject<UMaterialExpressionConstant>(this);
+						DefaultShadingModel->R = 1; // Default Lit shading model
+						ConvertNode->ShadingModel.Connect(0, DefaultShadingModel);
+					}
+					else
+					{
+						// Reconnect the shading model expression
+						MoveConnectionTo(EditorOnly->ShadingModelFromMaterialExpression, ConvertNode, 19);
+					}
 
 					// Store strata shading model of the converted material. 
 					GatherCustomNodes();
