Clear RF_Standalone flag on objects inside external packages when loaded : this guarantees that external packages of actors that are destroyed during the PIE session will properly get GC'ed and will allow future edits/modifications of OFPA actors.

diff --git a/Engine/Source/Runtime/Engine/Private/Level.cpp b/Engine/Source/Runtime/Engine/Private/Level.cpp
--- a/Engine/Source/Runtime/Engine/Private/Level.cpp
+++ b/Engine/Source/Runtime/Engine/Private/Level.cpp
@@ -1157,7 +1157,7 @@ void ULevel::PostLoad()
 				ActorPackage = LoadPackage(ActorPackage, *ActorPackageName, bPackageForPIE ? LOAD_PackageForPIE : LOAD_None, nullptr, &InstancingContext);
 
 				int32 PreviousActorCount = Actors.Num();
-				ForEachObjectWithPackage(ActorPackage, [this, &ActorsSet](UObject* PackageObject)
+				ForEachObjectWithPackage(ActorPackage, [this, &ActorsSet, bPackageForPIE](UObject* PackageObject)
 				{
 					// There might be multiple actors per package in the case where an actor as a child actor component as we put child actor in the same package as their parent
 					if (PackageObject->IsA<AActor>() && !PackageObject->IsTemplate())
@@ -1169,6 +1169,13 @@ void ULevel::PostLoad()
 							Actors.Add(Actor);
 						}
 					}
+					// In PIE, we make sure to clear RF_Standalone flag on objects in external packages (UMetaData) 
+					// This guarantees that external packages of actors that are destroyed during the PIE session will
+					// properly get GC'ed and will allow future edits/modifications of OFPA actors.
+					if (bPackageForPIE)
+					{
+						PackageObject->ClearFlags(RF_Standalone);
+					}
 					return true;
 				}, false);
 
diff --git a/Engine/Source/Runtime/Engine/Private/WorldPartition/WorldPartitionLevelHelper.cpp b/Engine/Source/Runtime/Engine/Private/WorldPartition/WorldPartitionLevelHelper.cpp
--- a/Engine/Source/Runtime/Engine/Private/WorldPartition/WorldPartitionLevelHelper.cpp
+++ b/Engine/Source/Runtime/Engine/Private/WorldPartition/WorldPartitionLevelHelper.cpp
@@ -357,6 +357,18 @@ bool FWorldPartitionLevelHelper::LoadActors(UWorld* InOwningWorld, ULevel* InDes
 			check(LoadProgress->NumPendingLoadRequests);
 			LoadProgress->NumPendingLoadRequests--;
 
+			// In PIE, we make sure to clear RF_Standalone flag on objects in external packages (UMetaData) 
+			// This guarantees that external packages of actors that are destroyed during the PIE session will
+			// properly get GC'ed and will allow future edits/modifications of OFPA actors.
+			if (LoadedPackage && InDestLevel && InDestLevel->GetPackage()->HasAnyPackageFlags(PKG_PlayInEditor))
+			{
+				ForEachObjectWithPackage(LoadedPackage, [](UObject* Object)
+				{
+					Object->ClearFlags(RF_Standalone);
+					return true;
+				}, false);
+			}
+
 			AActor* Actor = LoadedPackage ? FindObject<AActor>(LoadedPackage, *ActorName.ToString()) : nullptr;
 
 			if (Actor)
