Add support for adding textures of different resoution but same aspect ratio to a Texture2DArray.

diff --git a/Engine/Source/Runtime/Engine/Classes/Engine/Texture2DArray.cpp b/Engine/Source/Runtime/Engine/Classes/Engine/Texture2DArray.cpp
--- a/Engine/Source/Runtime/Engine/Classes/Engine/Texture2DArray.cpp
+++ b/Engine/Source/Runtime/Engine/Classes/Engine/Texture2DArray.cpp
@@ -339,50 +339,42 @@ EPixelFormat UTexture2DArray::GetPixelFormat() const
 
 ENGINE_API bool UTexture2DArray::CheckArrayTexturesCompatibility()
 {
-	bool bError = false;
+	if (!SourceTextures.Num())
+	{
+		return true;
+	}
 
-	for (int32 TextureIndex = 0; TextureIndex < SourceTextures.Num(); ++TextureIndex)
+	if (!SourceTextures[0])
 	{
 		// Do not create array till all texture slots are filled.
+		return false;
+	}
+
+	const FTextureSource& BaseTextureSource = SourceTextures[0]->Source;
+	const FString BaseTextureName = SourceTextures[0]->GetFName().ToString();
+
+	for (int32 TextureIndex = 1; TextureIndex < SourceTextures.Num(); ++TextureIndex)
+	{
 		if (!SourceTextures[TextureIndex])
 		{
+			// Do not create array till all texture slots are filled.
 			return false;
 		}
 
-		const FTextureSource& TextureSource = SourceTextures[TextureIndex]->Source;
-		// const int32 FormatDataSize = TextureSource.GetBytesPerPixel();
-		const ETextureSourceFormat SourceFormat = TextureSource.GetFormat();
-		const int32 SizeX = TextureSource.GetSizeX();
-		const int32 SizeY = TextureSource.GetSizeY();
-
-		for (int32 TextureCmpIndex = TextureIndex + 1; TextureCmpIndex < SourceTextures.Num(); ++TextureCmpIndex)
+		if (SourceTextures[TextureIndex]->Source.GetSizeX() * BaseTextureSource.GetSizeY() != SourceTextures[TextureIndex]->Source.GetSizeY() * BaseTextureSource.GetSizeX())
 		{
-			// Do not create array till all texture slots are filled.
-			if (!SourceTextures[TextureCmpIndex]) 
-			{
-				return false;
-			}
-
-			const FTextureSource& TextureSourceCmp = SourceTextures[TextureCmpIndex]->Source;
-			const FString TextureName = SourceTextures[TextureIndex]->GetFName().ToString();
-			const FString TextureNameCmp = SourceTextures[TextureCmpIndex]->GetFName().ToString();
-			const ETextureSourceFormat SourceFormatCmp = TextureSourceCmp.GetFormat();
-
-			if (TextureSourceCmp.GetSizeX() != SizeX || TextureSourceCmp.GetSizeY() != SizeY)
-			{
-				UE_LOG(LogTexture, Warning, TEXT("Texture2DArray creation failed. Textures %s and %s have different sizes."), *TextureName, *TextureNameCmp);
-				bError = true;
-			}
+			UE_LOG(LogTexture, Warning, TEXT("Texture2DArray creation failed. Textures %s and %s have different aspect ratios."), *SourceTextures[TextureIndex]->GetFName().ToString(), *BaseTextureName);
+			return false;
+		}
 
-			if (SourceFormatCmp != SourceFormat)
-			{
-				UE_LOG(LogTexture, Warning, TEXT("Texture2DArray creation failed. Textures %s and %s have incompatible pixel formats."), *TextureName, *TextureNameCmp);
-				bError = true;
-			}
+		if (SourceTextures[TextureIndex]->Source.GetFormat() != BaseTextureSource.GetFormat())
+		{
+			UE_LOG(LogTexture, Warning, TEXT("Texture2DArray creation failed. Textures %s and %s have incompatible pixel formats."), *SourceTextures[TextureIndex]->GetFName().ToString(), *BaseTextureName);
+			return false;
 		}
 	}
 
-	return (!bError);
+	return true;
 }
 
 
@@ -397,7 +389,18 @@ ENGINE_API bool UTexture2DArray::UpdateSourceFromSourceTextures(bool bCreatingNe
 	{
 		Modify();
 
-		FTextureSource& InitialSource = SourceTextures[0]->Source;
+		int32 InitialSourceTextureIndex = 0;
+		for (int32 SourceTextureIndex = 1; SourceTextureIndex < SourceTextures.Num(); ++SourceTextureIndex)
+		{
+			// Dimensions of the texture array source are calculated based on the dimensions of the largest texture source.
+			// It is sufficient to compare SizeX because CheckArrayTexturesCompatibility() guarantees that all sources have the same aspect ratio.
+			if (SourceTextures[SourceTextureIndex]->Source.GetSizeX() > SourceTextures[InitialSourceTextureIndex]->Source.GetSizeX())
+			{
+				InitialSourceTextureIndex = SourceTextureIndex;
+			}
+		}
+
+		FTextureSource& InitialSource = SourceTextures[InitialSourceTextureIndex]->Source;
 		// Format and format size.
 		ETextureSourceFormat Format = InitialSource.GetFormat();
 		int32 FormatDataSize = InitialSource.GetBytesPerPixel();
@@ -405,8 +408,6 @@ ENGINE_API bool UTexture2DArray::UpdateSourceFromSourceTextures(bool bCreatingNe
 		int32 SizeX = InitialSource.GetSizeX();
 		int32 SizeY = InitialSource.GetSizeY();
 		uint32 ArraySize = SourceTextures.Num();
-		// Only copy the first mip from the source textures to array texture.
-		uint32 NumMips = 1;
 
 		// This should be false when texture is updated to avoid overriding user settings.
 		if (bCreatingNewTexture)
@@ -419,36 +420,36 @@ ENGINE_API bool UTexture2DArray::UpdateSourceFromSourceTextures(bool bCreatingNe
 			NeverStream = true;
 		}
 
-		// Create the source texture for this UTexture.
-		Source.Init(SizeX, SizeY, ArraySize, NumMips, Format);
+		// Create the source texture for this Texture2DArray.
+		// Currently only single-mip Texture2DArray's are supported, therefore only the first mip is copied from each element source.
+		Source.Init(SizeX, SizeY, ArraySize, 1, Format);
 
-		// We only copy the top level Mip map.
-		TArray<uint8*, TInlineAllocator<MAX_TEXTURE_MIP_COUNT> > DestMipData;
-		TArray<uint64, TInlineAllocator<MAX_TEXTURE_MIP_COUNT> > MipSizeBytes;
-		DestMipData.AddZeroed(NumMips);
-		MipSizeBytes.AddZeroed(NumMips);
-			
-		for (uint32 MipIndex = 0; MipIndex < NumMips; ++MipIndex)
-		{
-			DestMipData[MipIndex] =  Source.LockMip(MipIndex);
-			MipSizeBytes[MipIndex] = Source.CalcMipSize(MipIndex) / ArraySize;
-		}
+		uint8* DestMipData = Source.LockMip(0);
+		int64 MipSizeBytes = Source.CalcMipSize(0) / ArraySize;
 
 		for (int32 SourceTexIndex = 0; SourceTexIndex < SourceTextures.Num(); ++SourceTexIndex)
 		{
-			for (uint32 MipIndex = 0; MipIndex < NumMips; ++MipIndex)
+			FTextureSource& TextureSource = SourceTextures[SourceTexIndex]->Source;
+			void* DestSliceData = DestMipData + MipSizeBytes * SourceTexIndex;
+			if (TextureSource.SizeX == SizeX && TextureSource.SizeY == SizeY)
 			{
-				TArray64<uint8> Ref2DData;
-				SourceTextures[SourceTexIndex]->Source.GetMipData(Ref2DData, MipIndex);
-				void* Dst = DestMipData[MipIndex] + MipSizeBytes[MipIndex] * SourceTexIndex;
-				FMemory::Memcpy(Dst, Ref2DData.GetData(), MipSizeBytes[MipIndex]);
+				void* SourceData = TextureSource.LockMip(0);
+				check(TextureSource.CalcMipSize(0) == MipSizeBytes);
+				FMemory::Memcpy(DestSliceData, SourceData, MipSizeBytes);
+				TextureSource.UnlockMip(0);
+			}
+			else
+			{
+				FImage SourceImage;
+				SourceTextures[SourceTexIndex]->Source.GetMipImage(SourceImage, 0);
+				FImage ResizedSourceImage;
+				SourceImage.ResizeTo(ResizedSourceImage, SizeX, SizeY, SourceImage.Format, SourceImage.GammaSpace);
+				check(ResizedSourceImage.RawData.Num() == MipSizeBytes);
+				FMemory::Memcpy(DestSliceData, ResizedSourceImage.RawData.GetData(), MipSizeBytes);
 			}
 		}
 
-		for (uint32 MipIndex = 0; MipIndex < NumMips; ++MipIndex)
-		{
-			Source.UnlockMip(MipIndex);
-		}
+		Source.UnlockMip(0);
 
 		ValidateSettingsAfterImportOrEdit();
 		SetLightingGuid();
