Reintroducing AnimationDataModel refactor with additional memory and cooktime optimizations.

diff --git a/Engine/Source/Editor/UnrealEd/Private/Fbx/FbxAnimationExport.cpp b/Engine/Source/Editor/UnrealEd/Private/Fbx/FbxAnimationExport.cpp
--- a/Engine/Source/Editor/UnrealEd/Private/Fbx/FbxAnimationExport.cpp
+++ b/Engine/Source/Editor/UnrealEd/Private/Fbx/FbxAnimationExport.cpp
@@ -52,7 +52,7 @@ namespace UnFbx
 		// set time correctly
 		FbxTime ExportedStartTime, ExportedStopTime;
 		ExportedStartTime.SetSecondDouble(0.f);
-		ExportedStopTime.SetSecondDouble(AnimSeq->GetPlayLength());
+		ExportedStopTime.SetSecondDouble(AnimSeq->GetDataModel()->GetFrameRate().AsSeconds(AnimSeq->GetDataModel()->GetNumberOfFrames()));
 
 		FbxTimeSpan ExportedTimeSpan;
 		ExportedTimeSpan.Set(ExportedStartTime, ExportedStopTime);
@@ -69,6 +69,14 @@ void FFbxExporter::ExportAnimSequenceToFbx(const UAnimSequence* AnimSeq,
 									 float AnimEndOffset,
 									 float AnimPlayRate,
 									 float StartTime)
+{
+	ExportAnimSequenceToFbx(AnimSeq, SkelMesh, BoneNodes, InAnimLayer,
+		AnimSeq->GetDataModel()->GetFrameRate().AsFrameTime(AnimStartOffset),
+		AnimSeq->GetDataModel()->GetFrameRate().AsFrameTime(AnimSeq->GetPlayLength() - AnimEndOffset),
+		AnimPlayRate, StartTime);
+}
+
+void FFbxExporter::ExportAnimSequenceToFbx(const UAnimSequence* AnimSeq, const USkeletalMesh* SkelMesh, TArray<FbxNode*>& BoneNodes, FbxAnimLayer* InAnimLayer, FFrameTime StartFrameTime, FFrameTime EndFrameTime, float FrameRateScale, float StartTime)
 {
 	// stack allocator for extracting curve
 	FMemMark Mark(FMemStack::Get());
@@ -112,7 +120,7 @@ void FFbxExporter::ExportAnimSequenceToFbx(const UAnimSequence* AnimSeq,
 		}
 	}
 
-	ExportCustomAnimCurvesToFbx(CustomCurveMap, AnimSeq, AnimStartOffset, AnimEndOffset, AnimPlayRate, StartTime);
+	ExportCustomAnimCurvesToFbx(CustomCurveMap, AnimSeq, StartFrameTime, EndFrameTime, FrameRateScale, StartTime);
 
 	TArray<const FAnimatedBoneAttribute*> CustomAttributes;
 
@@ -120,10 +128,10 @@ void FFbxExporter::ExportAnimSequenceToFbx(const UAnimSequence* AnimSeq,
 	for(int32 BoneIndex = 0; BoneIndex < BoneNodes.Num(); ++BoneIndex)
 	{
 		FbxNode* CurrentBoneNode = BoneNodes[BoneIndex];
-		int32 BoneTreeIndex = Skeleton->GetSkeletonBoneIndexFromMeshBoneIndex(SkelMesh, BoneIndex);
-		int32 BoneTrackIndex = Skeleton->GetRawAnimationTrackIndex(BoneTreeIndex, AnimSeq);
-		FName BoneName = Skeleton->GetReferenceSkeleton().GetBoneName(BoneTreeIndex);
-		
+		const int32 BoneTreeIndex = Skeleton->GetSkeletonBoneIndexFromMeshBoneIndex(SkelMesh, BoneIndex);
+		const int32 BoneTrackIndex = Skeleton->GetRawAnimationTrackIndex(BoneTreeIndex, AnimSeq);
+		const FName BoneName = Skeleton->GetReferenceSkeleton().GetBoneName(BoneTreeIndex);
+
 		CustomAttributes.Reset();
 		AnimSeq->GetDataModel()->GetAttributesForBone(BoneName, CustomAttributes);
 
@@ -207,7 +215,7 @@ void FFbxExporter::ExportAnimSequenceToFbx(const UAnimSequence* AnimSeq,
 			Curve->KeyModifyBegin();
 		}
 
-		auto ExportLambda = [&](float AnimTime, FbxTime ExportTime, bool bLastKey) {
+		auto ExportLambda = [&](double AnimTime, FbxTime ExportTime, bool bLastKey) {
 			FTransform BoneAtom;
 			AnimSeq->GetBoneTransform(BoneAtom, BoneTrackIndex, AnimTime, true);
 			FbxAMatrix FbxMatrix = Converter.ConvertMatrix(BoneAtom.ToMatrixWithScale());
@@ -265,7 +273,7 @@ void FFbxExporter::ExportAnimSequenceToFbx(const UAnimSequence* AnimSeq,
 			}
 		};
 
-		IterateInsideAnimSequence(AnimSeq, AnimStartOffset, AnimEndOffset, AnimPlayRate, StartTime, ExportLambda);
+		IterateInsideAnimSequence(AnimSeq, StartFrameTime, EndFrameTime, FrameRateScale, StartTime, ExportLambda);
 
 		for (FbxAnimCurve* Curve : Curves)
 		{
@@ -287,6 +295,15 @@ void FFbxExporter::ExportAnimSequenceToFbx(const UAnimSequence* AnimSeq,
 
 void FFbxExporter::ExportCustomAnimCurvesToFbx(const TMap<FName, FbxAnimCurve*>& CustomCurves, const UAnimSequence* AnimSeq, 
 	float AnimStartOffset, float AnimEndOffset, float AnimPlayRate, float StartTime, float ValueScale)
+{
+	ExportCustomAnimCurvesToFbx(CustomCurves, AnimSeq,
+		AnimSeq->GetDataModel()->GetFrameRate().AsFrameTime(AnimStartOffset),
+		AnimSeq->GetDataModel()->GetFrameRate().AsFrameTime(AnimSeq->GetPlayLength() - AnimEndOffset),
+		AnimPlayRate, StartTime, ValueScale);
+}
+	
+void FFbxExporter::ExportCustomAnimCurvesToFbx(const TMap<FName, FbxAnimCurve*>& CustomCurves, const UAnimSequence* AnimSeq,
+	FFrameTime AnimStartOffset, FFrameTime AnimEndOffset, float FrameRateScale, float StartTime, float ValueScale /*= 1.f*/)
 {
 	// stack allocator for extracting curve
 	FMemMark Mark(FMemStack::Get());
@@ -336,7 +353,7 @@ void FFbxExporter::ExportCustomAnimCurvesToFbx(const TMap<FName, FbxAnimCurve*>&
 		}
 	};
 
-	IterateInsideAnimSequence(AnimSeq, AnimStartOffset, AnimEndOffset, AnimPlayRate, StartTime, ExportLambda);
+	IterateInsideAnimSequence(AnimSeq, AnimStartOffset, AnimEndOffset, FrameRateScale, StartTime, ExportLambda);
 
 	for (auto CustomCurve : CustomCurves)
 	{
@@ -346,10 +363,21 @@ void FFbxExporter::ExportCustomAnimCurvesToFbx(const TMap<FName, FbxAnimCurve*>&
 
 void FFbxExporter::IterateInsideAnimSequence(const UAnimSequence* AnimSeq, float AnimStartOffset, float AnimEndOffset, float AnimPlayRate, float StartTime, TFunctionRef<void(float, FbxTime, bool)> IterationLambda)
 {
-	float AnimTime = AnimStartOffset;
-	float AnimEndTime = (AnimSeq->GetPlayLength() - AnimEndOffset);
-	double TimePerKey = AnimSeq->GetDataModel()->GetFrameRate().AsInterval();
-	const float AnimTimeIncrement = TimePerKey * AnimPlayRate;
+	IterateInsideAnimSequence(AnimSeq,
+		AnimSeq->GetDataModel()->GetFrameRate().AsFrameTime(AnimStartOffset),
+		AnimSeq->GetDataModel()->GetFrameRate().AsFrameTime(AnimSeq->GetPlayLength() - AnimEndOffset),
+		AnimPlayRate,
+		AnimPlayRate,
+		[IterationLambda](double T, FbxTime FT, bool B) -> void
+		{
+			IterationLambda((float)T, FT, B);
+		});
+}
+
+void FFbxExporter::IterateInsideAnimSequence(const UAnimSequence* AnimSeq, FFrameTime StartFrameTime, FFrameTime EndFrameTime, float FrameRateScale, float StartTime, TFunctionRef<void(double, FbxTime, bool)> IterationLambda)
+{
+	const double TimePerKey = AnimSeq->GetDataModel()->GetFrameRate().AsInterval();
+	const double AnimTimeIncrement = TimePerKey * FrameRateScale;
 	uint32 AnimFrameIndex = 0;
 
 	FbxTime ExportTime;
@@ -360,15 +388,16 @@ void FFbxExporter::IterateInsideAnimSequence(const UAnimSequence* AnimSeq, float
 
 	// Step through each frame and add custom curve data
 	bool bLastKey = false;
+	FFrameTime FrameTime = StartFrameTime;
+	const FFrameRate& FrameRate = AnimSeq->GetDataModel()->GetFrameRate();
 	while (!bLastKey)
 	{
-		bLastKey = (AnimTime + KINDA_SMALL_NUMBER) > AnimEndTime;
-
-		IterationLambda(AnimTime, ExportTime, bLastKey);
+		bLastKey = FrameTime > EndFrameTime;
+		IterationLambda(FrameRate.AsSeconds(FrameTime), ExportTime, bLastKey);
 
 		ExportTime += ExportTimeIncrement;
 		AnimFrameIndex++;
-		AnimTime = AnimStartOffset + ((float)AnimFrameIndex * AnimTimeIncrement);
+		FrameTime += FFrameTime::FromDecimal(FrameRateScale);
 	}
 }
 
@@ -423,9 +452,9 @@ void FFbxExporter::CorrectAnimTrackInterpolation( TArray<FbxNode*>& BoneNodes, F
 
 FbxNode* FFbxExporter::ExportAnimSequence( const UAnimSequence* AnimSeq, const USkeletalMesh* SkelMesh, bool bExportSkelMesh, const TCHAR* MeshName, FbxNode* ActorRootNode, const TArray<UMaterialInterface*>* OverrideMaterials /*= nullptr*/ )
 {
-	if( Scene == NULL || AnimSeq == NULL || SkelMesh == NULL )
+	if( Scene == nullptr || AnimSeq == nullptr || SkelMesh == nullptr )
 	{
- 		return NULL;
+ 		return nullptr;
 	}
 
 
@@ -451,10 +480,14 @@ FbxNode* FFbxExporter::ExportAnimSequence( const UAnimSequence* AnimSeq, const U
 			SkelMesh,
 			BoneNodes,
 			AnimLayer,
-			0.f,		// AnimStartOffset
-			0.f,		// AnimEndOffset
-			1.f,		// AnimPlayRate
-			0.f);		// StartTime
+			// Start frame to export
+			FFrameTime(0),
+			// Final frame to export
+			FFrameTime(AnimSeq->GetDataModel()->GetNumberOfFrames()),
+			// Frame rate scale
+			1.f,
+			// FBX StartTime
+			0.f);
 
 		CorrectAnimTrackInterpolation(BoneNodes, AnimLayer);
 	}
diff --git a/Engine/Source/Editor/UnrealEd/Private/Fbx/FbxSkeletalMeshExport.cpp b/Engine/Source/Editor/UnrealEd/Private/Fbx/FbxSkeletalMeshExport.cpp
--- a/Engine/Source/Editor/UnrealEd/Private/Fbx/FbxSkeletalMeshExport.cpp
+++ b/Engine/Source/Editor/UnrealEd/Private/Fbx/FbxSkeletalMeshExport.cpp
@@ -29,7 +29,7 @@ FbxNode* FFbxExporter::CreateSkeleton(const USkeletalMesh* SkelMesh, TArray<FbxN
 
 	if(RefSkeleton.GetRawBoneNum() == 0)
 	{
-		return NULL;
+		return nullptr;
 	}
 
 	// Create a list of the nodes we create for each bone, so that children can 
@@ -116,13 +116,13 @@ FbxNode* FFbxExporter::CreateMesh(const USkeletalMesh* SkelMesh, const TCHAR* Me
 	const int32 VertexCount = SourceModel.GetNumNonClothingVertices();
 
 	// Verify the integrity of the mesh.
-	if (VertexCount == 0) return NULL;
+	if (VertexCount == 0) return nullptr;
 
 	// Copy all the vertex data from the various chunks to a single buffer.
 	// Makes the rest of the code in this function cleaner and easier to maintain.  
 	TArray<FSoftSkinVertex> Vertices;
 	SourceModel.GetNonClothVertices(Vertices);
-	if (Vertices.Num() != VertexCount) return NULL;
+	if (Vertices.Num() != VertexCount) return nullptr;
 
 	FbxMesh* Mesh = FbxMesh::Create(Scene, TCHAR_TO_UTF8(MeshName));
 
@@ -137,7 +137,7 @@ FbxNode* FFbxExporter::CreateMesh(const USkeletalMesh* SkelMesh, const TCHAR* Me
 
 	// Create Layer 0 to hold the normals
 	FbxLayer* LayerZero = Mesh->GetLayer(0);
-	if (LayerZero == NULL)
+	if (LayerZero == nullptr)
 	{
 		Mesh->CreateLayer();
 		LayerZero = Mesh->GetLayer(0);
@@ -169,7 +169,7 @@ FbxNode* FFbxExporter::CreateMesh(const USkeletalMesh* SkelMesh, const TCHAR* Me
 	for (int32 TexCoordSourceIndex = 0; TexCoordSourceIndex < TexCoordSourceCount; ++TexCoordSourceIndex)
 	{
 		FbxLayer* Layer = Mesh->GetLayer(TexCoordSourceIndex);
-		if (Layer == NULL)
+		if (Layer == nullptr)
 		{
 			Mesh->CreateLayer();
 			Layer = Mesh->GetLayer(TexCoordSourceIndex);
@@ -328,10 +328,10 @@ FbxNode* FFbxExporter::CreateMesh(const USkeletalMesh* SkelMesh, const TCHAR* Me
 		if(AnimSeq && BlendShapeCurvesMap.Num() > 0)
 		{
 			ExportCustomAnimCurvesToFbx(BlendShapeCurvesMap, AnimSeq,
-				0.f,		// AnimStartOffset
-				0.f,		// AnimEndOffset
-				1.f,		// AnimPlayRate
-				0.f,		// StartTime
+				FFrameTime(0),													// Start frame to export
+				FFrameTime(AnimSeq->GetDataModel()->GetNumberOfFrames()),		// Final frame to export
+				1.f,												// Frame rate scale
+				0.f,												// FBX StartTime
 				100.f);		// ValueScale, for some reason we need to scale BlendShape curves by a factor of 100.
 		}
 	}
@@ -339,8 +339,6 @@ FbxNode* FFbxExporter::CreateMesh(const USkeletalMesh* SkelMesh, const TCHAR* Me
 	FbxNode* MeshNode = FbxNode::Create(Scene, TCHAR_TO_UTF8(MeshName));
 	MeshNode->SetNodeAttribute(Mesh);
 
-
-
 	// Add the materials for the mesh
 	const TArray<FSkeletalMaterial>& SkelMeshMaterials = SkelMesh->GetMaterials();
 	int32 MaterialCount = SkelMeshMaterials.Num();
@@ -358,7 +356,7 @@ FbxNode* FFbxExporter::CreateMesh(const USkeletalMesh* SkelMesh, const TCHAR* Me
 			MatInterface = SkelMeshMaterials[MaterialIndex].MaterialInterface;
 		}
 
-		FbxSurfaceMaterial* FbxMaterial = NULL;
+		FbxSurfaceMaterial* FbxMaterial = nullptr;
 		if (LODIndex == 0)
 		{
 			if (MatInterface && !FbxMaterials.Find(MatInterface))
@@ -590,7 +588,7 @@ void FFbxExporter::ExportSkeletalMeshComponent(USkeletalMeshComponent* SkelMeshC
 	if (SkelMeshComp && SkelMeshComp->GetSkeletalMeshAsset())
 	{
 		UAnimSequence* AnimSeq = (bSaveAnimSeq && SkelMeshComp->GetAnimationMode() == EAnimationMode::AnimationSingleNode) ? 
-			Cast<UAnimSequence>(SkelMeshComp->AnimationData.AnimToPlay) : NULL;
+			Cast<UAnimSequence>(SkelMeshComp->AnimationData.AnimToPlay) : nullptr;
 		FbxNode* SkeletonRootNode = ExportSkeletalMeshToFbx(SkelMeshComp->GetSkeletalMeshAsset(), AnimSeq, MeshName, ActorRootNode, &ToRawPtrTArrayUnsafe(SkelMeshComp->OverrideMaterials));
 		if(SkeletonRootNode)
 		{
@@ -766,7 +764,7 @@ FbxNode* FFbxExporter::ExportSkeletalMeshToFbx(const USkeletalMesh* SkeletalMesh
 		return SkeletonRootNode;
 	}
 
-	return NULL;
+	return nullptr;
 }
 
 } // namespace UnFbx
diff --git a/Engine/Source/Editor/UnrealEd/Private/FbxExporter.h b/Engine/Source/Editor/UnrealEd/Private/FbxExporter.h
--- a/Engine/Source/Editor/UnrealEd/Private/FbxExporter.h
+++ b/Engine/Source/Editor/UnrealEd/Private/FbxExporter.h
@@ -183,7 +183,7 @@ public:
 	 * @param StaticMesh	The static mesh to export
 	 * @param MaterialOrder	Optional ordering of materials to set up correct material ID's across multiple meshes being export such as BSP surfaces which share common materials. Should be used sparingly
 	 */
-	virtual void ExportStaticMesh( UStaticMesh* StaticMesh, const TArray<FStaticMaterial>* MaterialOrder = NULL );
+	virtual void ExportStaticMesh( UStaticMesh* StaticMesh, const TArray<FStaticMaterial>* MaterialOrder = nullptr );
 
 	/**
 	 * Exports BSP
@@ -215,7 +215,7 @@ public:
 	/**
 	 * Exports a single UAnimSequence, and optionally a skeletal mesh
 	 */
-	FbxNode* ExportAnimSequence( const UAnimSequence* AnimSeq, const USkeletalMesh* SkelMesh, bool bExportSkelMesh, const TCHAR* MeshNames=NULL, FbxNode* ActorRootNode=NULL, const TArray<UMaterialInterface*>* OverrideMaterials = nullptr);
+	FbxNode* ExportAnimSequence( const UAnimSequence* AnimSeq, const USkeletalMesh* SkelMesh, bool bExportSkelMesh, const TCHAR* MeshNames=nullptr, FbxNode* ActorRootNode=nullptr, const TArray<UMaterialInterface*>* OverrideMaterials = nullptr);
 
 	/** A node name adapter for a level sequence. */
 	class UNREALED_API FLevelSequenceNodeNameAdapter : public INodeNameAdapter
@@ -290,7 +290,7 @@ private:
 	 * @param MaterialOrderOverride	Optional ordering of materials to set up correct material ID's across multiple meshes being export such as BSP surfaces which share common materials. Should be used sparingly
 	 * @param OverrideMaterials	Optional array of materials to be used instead of the static mesh materials. Used for material overrides in static mesh components.
 	 */
-	FbxNode* ExportStaticMeshToFbx(const UStaticMesh* StaticMesh, int32 ExportLOD, const TCHAR* MeshName, FbxNode* FbxActor, int32 LightmapUVChannel = -1, const FColorVertexBuffer* ColorBuffer = NULL, const TArray<FStaticMaterial>* MaterialOrderOverride = NULL, const TArray<UMaterialInterface*>* OverrideMaterials = NULL);
+	FbxNode* ExportStaticMeshToFbx(const UStaticMesh* StaticMesh, int32 ExportLOD, const TCHAR* MeshName, FbxNode* FbxActor, int32 LightmapUVChannel = -1, const FColorVertexBuffer* ColorBuffer = nullptr, const TArray<FStaticMaterial>* MaterialOrderOverride = nullptr, const TArray<UMaterialInterface*>* OverrideMaterials = nullptr);
 
 	/**
 	 * Exports a spline mesh
@@ -372,18 +372,28 @@ private:
 	/**
 	 * Add the given animation sequence as rotation and translation tracks to the given list of bone nodes
 	 */
+	void ExportAnimSequenceToFbx(const UAnimSequence* AnimSeq, const USkeletalMesh* SkelMesh, TArray<FbxNode*>& BoneNodes, FbxAnimLayer* AnimLayer, FFrameTime StartFrameTime, FFrameTime EndFrameTime, float FrameRateScale, float StartTime);
+
+	UE_DEPRECATED(5.1, "ExportAnimSequenceToFbx is deprecated, use different signature")
 	void ExportAnimSequenceToFbx(const UAnimSequence* AnimSeq, const USkeletalMesh* SkelMesh, TArray<FbxNode*>& BoneNodes, FbxAnimLayer* AnimLayer,
 		float AnimStartOffset, float AnimEndOffset, float AnimPlayRate, float StartTime);
 
 	/**
 	 * Add the custom Curve data to the FbxAnimCurves passed in parameter by matching their name to the skeletal mesh custom curves.
 	 */
+	void ExportCustomAnimCurvesToFbx(const TMap<FName, FbxAnimCurve*>& CustomCurves, const UAnimSequence* AnimSeq,
+		FFrameTime AnimStartOffset, FFrameTime AnimEndOffset, float FrameRateScale, float StartTime, float ValueScale = 1.f);
+	
+	UE_DEPRECATED(5.1, "ExportCustomAnimCurvesToFbx is deprecated, use different signature")
 	void ExportCustomAnimCurvesToFbx(const TMap<FName, FbxAnimCurve*>& CustomCurves, const UAnimSequence* AnimSeq,
 		float AnimStartOffset, float AnimEndOffset, float AnimPlayRate, float StartTime, float ValueScale = 1.f);
 
 	/**
 	 * Used internally to reuse the AnimSequence iteration code when exporting various kind of curves.
 	 */
+	void IterateInsideAnimSequence(const UAnimSequence* AnimSeq, FFrameTime StartFrameTime, FFrameTime EndFrameTime, float FrameRateScale, float StartTime, TFunctionRef<void(double, FbxTime, bool)> IterationLambda);
+
+	UE_DEPRECATED(5.1, "IterateInsideAnimSequence is deprecated, use different signature")
 	void IterateInsideAnimSequence(const UAnimSequence* AnimSeq, float AnimStartOffset, float AnimEndOffset, float AnimPlayRate, float StartTime, TFunctionRef<void(float, FbxTime, bool)> IterationLambda);
 
 	/** 
diff --git a/Engine/Source/Editor/UnrealEd/Private/SkeletalMeshEdit.cpp b/Engine/Source/Editor/UnrealEd/Private/SkeletalMeshEdit.cpp
--- a/Engine/Source/Editor/UnrealEd/Private/SkeletalMeshEdit.cpp
+++ b/Engine/Source/Editor/UnrealEd/Private/SkeletalMeshEdit.cpp
@@ -49,7 +49,7 @@ UAnimSequence * UEditorEngine::ImportFbxAnimation( USkeleton* Skeleton, UObject*
 {
 	check(Skeleton);
 
-	UAnimSequence * NewAnimation=NULL;
+	UAnimSequence * NewAnimation=nullptr;
 
 	UnFbx::FFbxImporter* FFbxImporter = UnFbx::FFbxImporter::GetInstance();
 
@@ -76,7 +76,7 @@ UAnimSequence * UEditorEngine::ImportFbxAnimation( USkeleton* Skeleton, UObject*
 			FFbxImporter->AddTokenizedErrorMessage(FTokenizedMessage::Create(EMessageSeverity::Error, FText::Format(LOCTEXT("Error_CouldNotFindFbxTrack", "Mesh contains {0} bone as root but animation doesn't contain the root track.\nImport failed."), FText::FromName(Skeleton->GetReferenceSkeleton().GetBoneName(0)))), FFbxErrors::Animation_CouldNotFindRootTrack);
 
 			FFbxImporter->ReleaseScene();
-			return NULL;
+			return nullptr;
 		}
 
 		// Check for blend shape curves that are not skinned.  Unskinned geometry can still contain morph curves
@@ -568,19 +568,19 @@ void UnFbx::FFbxImporter::GetAnimationIntervalMultiLayer(FbxNode* RootNode, FbxA
 UAnimSequence * UnFbx::FFbxImporter::ImportAnimations(USkeleton* Skeleton, UObject* Outer, TArray<FbxNode*>& SortedLinks, const FString& Name, UFbxAnimSequenceImportData* TemplateImportData, TArray<FbxNode*>& NodeArray)
 {
 	// we need skeleton to create animsequence
-	if (Skeleton == NULL)
+	if (Skeleton == nullptr)
 	{
-		return NULL;
+		return nullptr;
 	}
 
 	int32 ValidTakeCount = 0;
 	if (IsValidAnimationData(SortedLinks, NodeArray, ValidTakeCount) == false)
 	{
 		AddTokenizedErrorMessage(FTokenizedMessage::Create(EMessageSeverity::Warning, LOCTEXT("FBXImport_InvalidAnimationData", "This does not contain any valid animation takes.")), FFbxErrors::Animation_InvalidData);
-		return NULL;
+		return nullptr;
 	}
 
-	UAnimSequence* LastCreatedAnim = NULL;
+	UAnimSequence* LastCreatedAnim = nullptr;
 
 	int32 ResampleRate = DEFAULT_SAMPLERATE;
 	if ( ImportOptions->bResample )
@@ -635,7 +635,7 @@ UAnimSequence * UnFbx::FFbxImporter::ImportAnimations(USkeleton* Skeleton, UObje
 
 		FString 	ParentPath = FString::Printf(TEXT("%s/%s"), *FPackageName::GetLongPackagePath(*Outer->GetName()), *SequenceName);
 		UObject* 	ParentPackage = CreatePackage( *ParentPath);
-		UObject* Object = LoadObject<UObject>(ParentPackage, *SequenceName, NULL, (LOAD_Quiet | LOAD_NoWarn), NULL);
+		UObject* Object = LoadObject<UObject>(ParentPackage, *SequenceName, nullptr, (LOAD_Quiet | LOAD_NoWarn), nullptr);
 		UAnimSequence * DestSeq = Cast<UAnimSequence>(Object);
 		// if object with same name exists, warn user
 		if (Object && !DestSeq)
@@ -645,12 +645,12 @@ UAnimSequence * UnFbx::FFbxImporter::ImportAnimations(USkeleton* Skeleton, UObje
 		}
 
 		// If not, create new one now.
-		if(!DestSeq)
+
+		const bool bCreateAsset = !DestSeq;
+		if(bCreateAsset)
 		{
 			DestSeq = NewObject<UAnimSequence>(ParentPackage, *SequenceName, RF_Public | RF_Standalone);
 			CreatedObjects.Add(DestSeq);
-			// Notify the asset registry
-			FAssetRegistryModule::AssetCreated(DestSeq);
 		}
 
 		DestSeq->SetSkeleton(Skeleton);
@@ -662,8 +662,16 @@ UAnimSequence * UnFbx::FFbxImporter::ImportAnimations(USkeleton* Skeleton, UObje
 		DestSeq->ImportFileFramerate = GetOriginalFbxFramerate();
 		DestSeq->ImportResampleFramerate = ResampleRate;
 
+		DestSeq->GetController().InitializeModel();
+
 		ImportAnimation(Skeleton, DestSeq, Name, SortedLinks, NodeArray, CurAnimStack, ResampleRate, AnimTimeSpan);
 
+		if (bCreateAsset)
+		{			
+			// Notify the asset registry
+			FAssetRegistryModule::AssetCreated(DestSeq);
+		}
+
 		LastCreatedAnim = DestSeq;
 	}
 
@@ -1337,7 +1345,7 @@ bool UnFbx::FFbxImporter::ImportCurveToAnimSequence(class UAnimSequence * Target
 		FAnimationCurveIdentifier FloatCurveId(NewName, ERawCurveTrackTypes::RCT_Float);
 
 
-		UAnimDataModel* DataModel = TargetSequence->GetDataModel();
+		IAnimationDataModel* DataModel = TargetSequence->GetDataModel();
 		IAnimationDataController& Controller = TargetSequence->GetController();
 
 		const FFloatCurve* TargetCurve = DataModel->FindFloatCurve(FloatCurveId);
@@ -1400,8 +1408,7 @@ bool UnFbx::FFbxImporter::ImportRichCurvesToAnimSequence(UAnimSequence* TargetSe
 
 			FAnimationCurveIdentifier FloatCurveId(NewName, ERawCurveTrackTypes::RCT_Float);
 
-
-			UAnimDataModel* DataModel = TargetSequence->GetDataModel();
+			const IAnimationDataModel* DataModel = TargetSequence->GetDataModel();
 			IAnimationDataController& Controller = TargetSequence->GetController();
 
 			const FFloatCurve* TargetCurve = DataModel->FindFloatCurve(FloatCurveId);
@@ -1721,8 +1728,14 @@ bool UnFbx::FFbxImporter::ImportAnimation(USkeleton* Skeleton, UAnimSequence * D
 	//This destroy all previously imported animation raw data
 	Controller.RemoveAllBoneTracks();
 
+	// First set frame rate
+	const FFrameRate ResampleFrameRate(ResampleRate, 1);
+	Controller.SetFrameRate(ResampleFrameRate);
+
 	// if you have one pose(thus 0.f duration), it still contains animation, so we'll need to consider that as MINIMUM_ANIMATION_LENGTH time length
-	Controller.SetPlayLength(FGenericPlatformMath::Max<float>(SequenceLength.GetSecondDouble(), MINIMUM_ANIMATION_LENGTH));
+	const FFrameNumber NumberOfFrames = ResampleFrameRate.AsFrameNumber(SequenceLength.GetSecondDouble());
+	Controller.SetNumberOfFrames(FGenericPlatformMath::Max<int32>(NumberOfFrames.Value, 1));
+
 	if(PreviousSequenceLength > MINIMUM_ANIMATION_LENGTH && DestSeq->GetDataModel()->GetNumberOfFloatCurves() > 0)
 	{
 		// The sequence already existed when we began the import. We need to scale the key times for all curves to match the new 
@@ -2024,7 +2037,7 @@ void UnFbx::FFbxImporter::ImportAnimationCustomAttribute(FAnimCurveImportSetting
 	UAnimSequence* DestSeq = AnimImportSettings.DestSeq;
 	USkeleton* MySkeleton = DestSeq->GetSkeleton();
 
-	const UAnimDataModel* DataModel = DestSeq->GetDataModel();
+	const IAnimationDataModel* DataModel = DestSeq->GetDataModel();
 	const int32 NumFloatCurves = DataModel->GetNumberOfFloatCurves();
 	const FAnimationCurveData& CurveData = DataModel->GetCurveData();
 
@@ -2140,7 +2153,7 @@ void UnFbx::FFbxImporter::ImportBoneTracks(USkeleton* Skeleton, FAnimCurveImport
 	const FbxTime TimeComparisonThreshold = (KINDA_SMALL_NUMBER * static_cast<float>(FBXSDK_TC_SECOND));
 	const FReferenceSkeleton& RefSkeleton = Skeleton->GetReferenceSkeleton();
 
-	UAnimDataModel* DataModel = DestSeq->GetDataModel();
+	IAnimationDataModel* DataModel = DestSeq->GetDataModel();
 	IAnimationDataController& Controller = DestSeq->GetController();	
 	Controller.SetFrameRate(FFrameRate(ResampleRate, 1));
 
@@ -2301,7 +2314,7 @@ void UnFbx::FFbxImporter::ImportBoneTracks(USkeleton* Skeleton, FAnimCurveImport
 								if (AttributeIdentifier.IsValid())
 								{
 									// remove any existing attribute with the same identifier
-									if (const UAnimDataModel* Model = Controller.GetModel())
+									if (const IAnimationDataModel* Model = Controller.GetModel())
 									{
 										if (Model->FindAttribute(AttributeIdentifier))
 										{
diff --git a/Engine/Source/Editor/UnrealEd/Public/ChangeTransactor.h b/Engine/Source/Editor/UnrealEd/Public/ChangeTransactor.h
--- a/Engine/Source/Editor/UnrealEd/Public/ChangeTransactor.h
+++ b/Engine/Source/Editor/UnrealEd/Public/ChangeTransactor.h
@@ -42,7 +42,7 @@ struct FChangeTransactor
 		CheckTransactionObject();
 		checkf(!IsTransactionPending(), TEXT("A transaction is already pending"));
 		
-		if (GEngine && GEngine->CanTransact() && ensure(!GIsTransacting))
+		if (CanTransactChanges())
 		{
 			CompoundChangeData = MakeUnique<FCompoundChangeInput>();
 			PendingTransactionIndex = GEngine->BeginTransaction(TEXT("FChangeTransactor"), TransactionDescription, TransactionObject.Get());
@@ -100,7 +100,7 @@ struct FChangeTransactor
 
 		CheckTransactionObject();
 
-		if (GEngine && GEngine->CanTransact() && ensure(!GIsTransacting))
+		if (CanTransactChanges())
 		{
 			checkf(IsTransactionPending(), TEXT("No transaction was previously opened"));
 
@@ -118,6 +118,11 @@ struct FChangeTransactor
 			}
 		}
 	}
+	
+	static bool CanTransactChanges()
+	{
+		return GEngine && GEngine->CanTransact() && !GIsTransacting;
+	};
 private:	
 	void CheckTransactionObject()
 	{
@@ -136,5 +141,28 @@ private:
 	int32 TransactedChanges;
 };
 
+struct FScopedCompoundTransaction
+{
+	FScopedCompoundTransaction(FChangeTransactor& InTransactor, const FText& InDescription) : Transactor(InTransactor), bCreated(false)
+	{
+		if (FChangeTransactor::CanTransactChanges() && !Transactor.IsTransactionPending())
+		{
+			Transactor.OpenTransaction(InDescription);
+			bCreated = true;
+		}
+	}
+
+	~FScopedCompoundTransaction()
+	{
+		if (bCreated)
+		{
+			Transactor.CloseTransaction();
+		}
+	}
+
+	FChangeTransactor& Transactor;
+	bool bCreated;
+};
+
 }
 #endif // WITH_EDITOR
diff --git a/Engine/Source/Runtime/AnimGraphRuntime/Private/AnimNodes/AnimNode_BlendSpacePlayer.cpp b/Engine/Source/Runtime/AnimGraphRuntime/Private/AnimNodes/AnimNode_BlendSpacePlayer.cpp
--- a/Engine/Source/Runtime/AnimGraphRuntime/Private/AnimNodes/AnimNode_BlendSpacePlayer.cpp
+++ b/Engine/Source/Runtime/AnimGraphRuntime/Private/AnimNodes/AnimNode_BlendSpacePlayer.cpp
@@ -125,7 +125,7 @@ void FAnimNode_BlendSpacePlayerBase::Evaluate_AnyThread(FPoseContext& Output)
 	if ((CurrentBlendSpace != nullptr) && (Output.AnimInstanceProxy->IsSkeletonCompatible(CurrentBlendSpace->GetSkeleton())))
 	{
 		FAnimationPoseData AnimationPoseData(Output);
-		CurrentBlendSpace->GetAnimationPose(BlendSampleDataCache, FAnimExtractContext(InternalTimeAccumulator, Output.AnimInstanceProxy->ShouldExtractRootMotion(), DeltaTimeRecord, GetLoop()), AnimationPoseData);
+		CurrentBlendSpace->GetAnimationPose(BlendSampleDataCache, FAnimExtractContext(static_cast<double>(InternalTimeAccumulator), Output.AnimInstanceProxy->ShouldExtractRootMotion(), DeltaTimeRecord, GetLoop()), AnimationPoseData);
 	}
 	else
 	{
diff --git a/Engine/Source/Runtime/AnimGraphRuntime/Private/AnimNodes/AnimNode_RandomPlayer.cpp b/Engine/Source/Runtime/AnimGraphRuntime/Private/AnimNodes/AnimNode_RandomPlayer.cpp
--- a/Engine/Source/Runtime/AnimGraphRuntime/Private/AnimNodes/AnimNode_RandomPlayer.cpp
+++ b/Engine/Source/Runtime/AnimGraphRuntime/Private/AnimNodes/AnimNode_RandomPlayer.cpp
@@ -297,8 +297,8 @@ void FAnimNode_RandomPlayer::Evaluate_AnyThread(FPoseContext& Output)
 		FAnimationPoseData CurrentPoseData(Poses[0], Curves[0], Attributes[0]);
 		FAnimationPoseData NextPoseData(Poses[1], Curves[1], Attributes[1]);
 
-		CurrentSequence->GetAnimationPose(CurrentPoseData, FAnimExtractContext(CurrentData.CurrentPlayTime, AnimProxy->ShouldExtractRootMotion(), CurrentData.DeltaTimeRecord, CurrentData.RemainingLoops > 0));
-		NextSequence->GetAnimationPose(NextPoseData, FAnimExtractContext(NextData.CurrentPlayTime, AnimProxy->ShouldExtractRootMotion(), NextData.DeltaTimeRecord, NextData.RemainingLoops > 0));
+		CurrentSequence->GetAnimationPose(CurrentPoseData, FAnimExtractContext(static_cast<double>(CurrentData.CurrentPlayTime), AnimProxy->ShouldExtractRootMotion(), CurrentData.DeltaTimeRecord, CurrentData.RemainingLoops > 0));
+		NextSequence->GetAnimationPose(NextPoseData, FAnimExtractContext(static_cast<double>(NextData.CurrentPlayTime), AnimProxy->ShouldExtractRootMotion(), NextData.DeltaTimeRecord, NextData.RemainingLoops > 0));
 
 		FAnimationPoseData AnimationPoseData(Output);
 		FAnimationRuntime::BlendPosesTogether(Poses, Curves, Attributes, Weights, AnimationPoseData);
@@ -307,7 +307,7 @@ void FAnimNode_RandomPlayer::Evaluate_AnyThread(FPoseContext& Output)
 	{
 		// Single animation, no blending needed.
 		FAnimationPoseData AnimationPoseData(Output);
-		CurrentSequence->GetAnimationPose(AnimationPoseData, FAnimExtractContext(CurrentData.CurrentPlayTime, Output.AnimInstanceProxy->ShouldExtractRootMotion(), CurrentData.DeltaTimeRecord, CurrentData.RemainingLoops > 0));
+		CurrentSequence->GetAnimationPose(AnimationPoseData, FAnimExtractContext(static_cast<double>(CurrentData.CurrentPlayTime), Output.AnimInstanceProxy->ShouldExtractRootMotion(), CurrentData.DeltaTimeRecord, CurrentData.RemainingLoops > 0));
 	}
 }
 
diff --git a/Engine/Source/Runtime/AnimGraphRuntime/Private/AnimNodes/AnimNode_SequenceEvaluator.cpp b/Engine/Source/Runtime/AnimGraphRuntime/Private/AnimNodes/AnimNode_SequenceEvaluator.cpp
--- a/Engine/Source/Runtime/AnimGraphRuntime/Private/AnimNodes/AnimNode_SequenceEvaluator.cpp
+++ b/Engine/Source/Runtime/AnimGraphRuntime/Private/AnimNodes/AnimNode_SequenceEvaluator.cpp
@@ -90,7 +90,7 @@ void FAnimNode_SequenceEvaluatorBase::Evaluate_AnyThread(FPoseContext& Output)
 	if ((CurrentSequence != nullptr) && (Output.AnimInstanceProxy->IsSkeletonCompatible(CurrentSequence->GetSkeleton())))
 	{
 		FAnimationPoseData AnimationPoseData(Output);
-		CurrentSequence->GetAnimationPose(AnimationPoseData, FAnimExtractContext(InternalTimeAccumulator, Output.AnimInstanceProxy->ShouldExtractRootMotion(), DeltaTimeRecord, GetShouldLoop()));
+		CurrentSequence->GetAnimationPose(AnimationPoseData, FAnimExtractContext(static_cast<double>(InternalTimeAccumulator), Output.AnimInstanceProxy->ShouldExtractRootMotion(), DeltaTimeRecord, GetShouldLoop()));
 	}
 	else
 	{
diff --git a/Engine/Source/Runtime/Core/Public/UObject/UE5MainStreamObjectVersion.h b/Engine/Source/Runtime/Core/Public/UObject/UE5MainStreamObjectVersion.h
--- a/Engine/Source/Runtime/Core/Public/UObject/UE5MainStreamObjectVersion.h
+++ b/Engine/Source/Runtime/Core/Public/UObject/UE5MainStreamObjectVersion.h
@@ -280,6 +280,9 @@ struct CORE_API FUE5MainStreamObjectVersion
 		// Replace FNames for class/actor paths with FSoftObjectPath
 		WorldPartitionActorDescActorAndClassPaths,
 
+		// Reintroducing AnimationDataModelInterface_BackedOut changes
+		ReintroduceAnimationDataModelInterface,
+
 		// -----<new versions can be added above this line>-------------------------------------------------
 		VersionPlusOne,
 		LatestVersion = VersionPlusOne - 1
diff --git a/Engine/Source/Runtime/Engine/Classes/Animation/AnimCompositeBase.h b/Engine/Source/Runtime/Engine/Classes/Animation/AnimCompositeBase.h
--- a/Engine/Source/Runtime/Engine/Classes/Animation/AnimCompositeBase.h
+++ b/Engine/Source/Runtime/Engine/Classes/Animation/AnimCompositeBase.h
@@ -39,7 +39,7 @@ struct FRootMotionExtractionStep
 	float EndPosition;
 
 	FRootMotionExtractionStep() 
-		: AnimSequence(NULL)
+		: AnimSequence(nullptr)
 		, StartPosition(0.f)
 		, EndPosition(0.f)
 		{
@@ -348,6 +348,10 @@ class UAnimCompositeBase : public UAnimSequenceBase
 	//~ Begin UObject Interface
 	virtual void PostLoad() override;
 	//~ End UObject Interface
+	
+	//~ Begin UAnimSequenceBase Interface
+	virtual FFrameRate GetSamplingFrameRate() const override;
+	//~ End UAnimSequenceBase Interface
 
 	// Extracts root motion from the supplied FAnimTrack between the Start End range specified
 	ENGINE_API void ExtractRootMotionFromTrack(const FAnimTrack &SlotAnimTrack, float StartTrackPosition, float EndTrackPosition, FRootMotionMovementParams &RootMotion) const;
@@ -362,5 +366,9 @@ class UAnimCompositeBase : public UAnimSequenceBase
 	virtual bool ContainRecursive(TArray<UAnimCompositeBase*>& CurrentAccumulatedList) PURE_VIRTUAL(UAnimCompositeBase::ContainRecursive, return false; );
 
 	virtual void SetCompositeLength(float InLength) PURE_VIRTUAL(UAnimCompositeBase::SetCompositeLength, );
+
+#if WITH_EDITOR
+	virtual void PopulateWithExistingModel(TScriptInterface<IAnimationDataModel> ExistingDataModel) override;
+#endif // WITH_EDITOR
 };
 
diff --git a/Engine/Source/Runtime/Engine/Classes/Animation/AnimData/AnimDataModel.h b/Engine/Source/Runtime/Engine/Classes/Animation/AnimData/AnimDataModel.h
--- a/Engine/Source/Runtime/Engine/Classes/Animation/AnimData/AnimDataModel.h
+++ b/Engine/Source/Runtime/Engine/Classes/Animation/AnimData/AnimDataModel.h
@@ -4,279 +4,75 @@
 
 #include "CoreMinimal.h"
 #include "UObject/ObjectMacros.h"
-#include "AnimDataNotifications.h"
-#include "Animation/AttributeCurve.h"
 
-#include "AnimDataModel.generated.h"
-
-/**
- * Structure encapsulating a single bone animation track.
- */
-USTRUCT(BlueprintType)
-struct ENGINE_API FBoneAnimationTrack
-{
-	GENERATED_BODY()
-
-	/** Internally stored data representing the animation bone data */
-	UPROPERTY(VisibleAnywhere, BlueprintReadOnly, Category = "Model")
-	FRawAnimSequenceTrack InternalTrackData;
-
-	/** Index corresponding to the bone this track corresponds to within the target USkeleton */
-	UPROPERTY(VisibleAnywhere, BlueprintReadOnly, Category = "Model")
-	int32 BoneTreeIndex = INDEX_NONE;
-
-	/** Name of the bone this track corresponds to */
-	UPROPERTY(VisibleAnywhere, BlueprintReadOnly, Category = "Model")
-	FName Name;
-};
-
-/**
- * Structure encapsulating animated curve data. Currently only contains Float and Transform curves.
- */
-USTRUCT(BlueprintType)
-struct ENGINE_API FAnimationCurveData
-{
-	GENERATED_BODY()
-
-	/** Float-based animation curves */
-	UPROPERTY(VisibleAnywhere, BlueprintReadOnly, Category = "Model")
-	TArray<FFloatCurve>	FloatCurves;
-
-	/** FTransform-based animation curves, used for animation layer editing */
-	UPROPERTY(VisibleAnywhere, BlueprintReadOnly, Category = "Model")
-	TArray<FTransformCurve>	TransformCurves;
-};
-
-/**
- * Structure encapsulating animated (bone) attribute data.
- */
-USTRUCT(BlueprintType)
-struct ENGINE_API FAnimatedBoneAttribute
-{
-	GENERATED_BODY()
-
-	/** Identifier to reference this attribute by */
-	UPROPERTY(VisibleAnywhere, BlueprintReadOnly, Category = "Model")
-	FAnimationAttributeIdentifier Identifier;	
+#include "IAnimationDataModel.h"
 
-	/** Curve containing the (animated) attribute data */
-	UPROPERTY(VisibleAnywhere, BlueprintReadOnly, Category = "Model")
-	FAttributeCurve Curve;
-};
-
-struct FTrackToSkeletonMap;
-struct FAnimationCurveIdentifier;
+#include "AnimDataModel.generated.h"
 
 /**
  * The Model represents the source data for animations. It contains both bone animation data as well as animated curves.
  * They are currently only a sub-object of a AnimSequenceBase instance. The instance derives all runtime data from the source data. 
  */
 UCLASS(BlueprintType, meta=(DebugTreeLeaf))
-class ENGINE_API UAnimDataModel : public UObject
+class ENGINE_API UAnimDataModel : public UObject, public IAnimationDataModel
 {
 	GENERATED_BODY()
 public:
-
-	/** Begin UAnimDataModel overrides */
+	/** Begin UObject overrides */
 	virtual void PostLoad() override;
 	virtual void PostDuplicate(bool bDuplicateForPIE) override;
 	virtual bool IsEditorOnly() const override { return true; }
-	/** End UAnimDataModel overrides */
-
-	/**
-	* @return	Total length of play-able animation data 
-	*/
-	UFUNCTION(BlueprintPure, Category=AnimationDataModel)
-	float GetPlayLength() const;
+	/** End UObject overrides */
 	
-	/**
-	* @return	Total number of frames of animation data stored 
-	*/
-	UFUNCTION(BlueprintPure, Category = AnimationDataModel)
-	int32 GetNumberOfFrames() const;
-
-	/**
-	* @return	Total number of animation data keys stored 
-	*/
-	UFUNCTION(BlueprintPure, Category = AnimationDataModel)
-	int32 GetNumberOfKeys() const;
-
-	/**
-	* @return	Frame rate at which the animation data is key-ed 
-	*/
-	UFUNCTION(BlueprintPure, Category = AnimationDataModel)
-	const FFrameRate& GetFrameRate() const;
-	
-	/**
-	* @return	Array containg all bone animation tracks 
-	*/
-	UFUNCTION(BlueprintPure, Category = AnimationDataModel)
-	const TArray<FBoneAnimationTrack>& GetBoneAnimationTracks() const;
-	
-	/**
-	* @return	Bone animation track for the provided index
-	*/
-	UFUNCTION(BlueprintPure, Category = AnimationDataModel)
-	const FBoneAnimationTrack& GetBoneTrackByIndex(int32 TrackIndex) const;
-
-	/**
-	* @return	Bone animation track for the provided (bone) name
-	*/
-	UFUNCTION(BlueprintPure, Category = AnimationDataModel)
-	const FBoneAnimationTrack& GetBoneTrackByName(FName TrackName) const;
-	
-	/**
-	* @return	Bone animation track for the provided (bone) name if found, otherwise returns a nullptr 
-	*/
-	const FBoneAnimationTrack* FindBoneTrackByName(FName Name) const;
-
-	/**
-	* @return	Bone animation track for the provided index if valid, otherwise returns a nullptr 
-	*/
-	const FBoneAnimationTrack* FindBoneTrackByIndex(int32 BoneIndex) const;
-
-	/**
-	* @return	Internal track index for the provided bone animation track if found, otherwise returns INDEX_NONE 
-	*/
-	UFUNCTION(BlueprintPure, Category = AnimationDataModel)
-	int32 GetBoneTrackIndex(const FBoneAnimationTrack& Track) const;
-
-	/**
-	* @return	Internal track index for the provided (bone) name if found, otherwise returns INDEX_NONE 
-	*/
-	UFUNCTION(BlueprintPure, Category = AnimationDataModel)
-	int32 GetBoneTrackIndexByName(FName TrackName) const;
+	/** Begin IAnimationDataModel overrides */
+	virtual double GetPlayLength() const override;
+	virtual int32 GetNumberOfFrames() const override;
+	virtual int32 GetNumberOfKeys() const override;
+	virtual FFrameRate GetFrameRate() const override;
+	virtual const TArray<FBoneAnimationTrack>& GetBoneAnimationTracks() const override;
+	virtual const FBoneAnimationTrack& GetBoneTrackByIndex(int32 TrackIndex) const override;
+	virtual const FBoneAnimationTrack& GetBoneTrackByName(FName TrackName) const override;
+	virtual const FBoneAnimationTrack* FindBoneTrackByName(FName Name) const override;
+	virtual const FBoneAnimationTrack* FindBoneTrackByIndex(int32 BoneIndex) const override;
+	virtual int32 GetBoneTrackIndex(const FBoneAnimationTrack& Track) const override;
+	virtual int32 GetBoneTrackIndexByName(FName TrackName) const override;
+	virtual bool IsValidBoneTrackIndex(int32 TrackIndex) const override;
+	virtual int32 GetNumBoneTracks() const override;
+	virtual void GetBoneTrackNames(TArray<FName>& OutNames) const override;
+
+	virtual const FAnimCurveBase* FindCurve(const FAnimationCurveIdentifier& CurveIdentifier) const override;
+	virtual const FFloatCurve* FindFloatCurve(const FAnimationCurveIdentifier& CurveIdentifier) const override;
+	virtual const FTransformCurve* FindTransformCurve(const FAnimationCurveIdentifier& CurveIdentifier) const override;
+	virtual const FRichCurve* FindRichCurve(const FAnimationCurveIdentifier& CurveIdentifier) const override;
+	virtual const FAnimationCurveData& GetCurveData() const override;
+	virtual int32 GetNumberOfTransformCurves() const override;
+	virtual int32 GetNumberOfFloatCurves() const override;
+	virtual const TArray<struct FFloatCurve>& GetFloatCurves() const override;
+	virtual const TArray<struct FTransformCurve>& GetTransformCurves() const override;	
+	virtual const FAnimCurveBase& GetCurve(const FAnimationCurveIdentifier& CurveIdentifier) const override;
+	virtual const FFloatCurve& GetFloatCurve(const FAnimationCurveIdentifier& CurveIdentifier) const override;
+	virtual const FTransformCurve& GetTransformCurve(const FAnimationCurveIdentifier& CurveIdentifier) const override;
+	virtual const FRichCurve& GetRichCurve(const FAnimationCurveIdentifier& CurveIdentifier) const override;
+	virtual TArrayView<const FAnimatedBoneAttribute> GetAttributes() const override;
+	virtual int32 GetNumberOfAttributes() const override;
+	virtual int32 GetNumberOfAttributesForBoneIndex(const int32 BoneIndex) const override;
+	virtual void GetAttributesForBone(const FName& BoneName, TArray<const FAnimatedBoneAttribute*>& OutBoneAttributes) const override;
+	virtual const FAnimatedBoneAttribute& GetAttribute(const FAnimationAttributeIdentifier& AttributeIdentifier) const override;
+	virtual const FAnimatedBoneAttribute* FindAttribute(const FAnimationAttributeIdentifier& AttributeIdentifier) const override;
+	virtual UAnimSequence* GetAnimationSequence() const override;
+	virtual FAnimDataModelModifiedEvent& GetModifiedEvent() override { return ModifiedEvent; }
+	virtual FGuid GenerateGuid() const override;
+#if WITH_EDITOR
+	virtual void Evaluate(FAnimationPoseData& InOutPoseData, const UE::Anim::DataModel::FEvaluationContext& EvaluationContext) const override;
+#endif
+	virtual TScriptInterface<IAnimationDataController> GetController() override;
+	virtual bool HasBeenPopulated() const override { return bPopulated; }
+protected:
+	virtual IAnimationDataModel::FModelNotifier& GetNotifier() override;
+	virtual FAnimDataModelModifiedDynamicEvent& GetModifiedDynamicEvent() override { return ModifiedEventDynamic; }
+	virtual void OnNotify(const EAnimDataModelNotifyType& NotifyType, const FAnimDataModelNotifPayload& Payload) override {}
+	/** End IAnimationDataModel overrides */
 
-	/**
-	* @return	Whether or not the provided track index is valid 
-	*/
-	UFUNCTION(BlueprintPure, Category = AnimationDataModel)
-	bool IsValidBoneTrackIndex(int32 TrackIndex) const;
-
-	/**
-	* @return	Total number of bone animation tracks
-	*/
-	UFUNCTION(BlueprintPure, Category = AnimationDataModel)
-	const int32 GetNumBoneTracks() const;
-
-	/**
-	* Populates the provided array with all contained (bone) track names
-	*
-	* @param	OutNames	[out] Array containing all bone track names
-	*/
-	UFUNCTION(BlueprintPure, Category = AnimationDataModel)
-	void GetBoneTrackNames(TArray<FName>& OutNames) const;
-
-	/** Returns all contained curve animation data */
-	const FAnimationCurveData& GetCurveData() const;
-
-	/**
-	* @return	Total number of stored FTransform curves
-	*/
-	UFUNCTION(BlueprintPure, Category = AnimationDataModel)
-	int32 GetNumberOfTransformCurves() const;
-
-	/**
-	* @return	Total number of stored float curves
-	*/
-	UFUNCTION(BlueprintPure, Category = AnimationDataModel)
-	int32 GetNumberOfFloatCurves() const;
-
-	/**
-	* @return	Array containing all stored float curves 
-	*/
-	const TArray<struct FFloatCurve>& GetFloatCurves() const;
-
-	/**
-	* @return	Array containing all stored FTransform curves 
-	*/
-	const TArray<struct FTransformCurve>& GetTransformCurves() const;
-
-	/**
-	* @return	Curve ptr for the provided identifier if valid, otherwise returns a nullptr 
-	*/
-	const FAnimCurveBase* FindCurve(const FAnimationCurveIdentifier& CurveIdentifier) const;
-
-	/**
-	* @return	Float Curve ptr for the provided identifier if valid, otherwise returns a nullptr
-	*/
-	const FFloatCurve* FindFloatCurve(const FAnimationCurveIdentifier& CurveIdentifier) const;
-
-	/**
-	* @return	Transform Curve ptr for the provided identifier if valid, otherwise returns a nullptr
-	*/
-	const FTransformCurve* FindTransformCurve(const FAnimationCurveIdentifier& CurveIdentifier) const;
-
-	/**
-	* @return	Rich curve ptr for the provided identifier if valid, otherwise returns a nullptr
-	*/
-	const FRichCurve* FindRichCurve(const FAnimationCurveIdentifier& CurveIdentifier) const;
-
-	/**
-	* @return	Curve object for the provided identifier if valid
-	*/
-	const FAnimCurveBase& GetCurve(const FAnimationCurveIdentifier& CurveIdentifier) const;
-
-	/**
-	* @return	Float Curve object for the provided identifier if valid
-	*/
-	const FFloatCurve& GetFloatCurve(const FAnimationCurveIdentifier& CurveIdentifier) const;
-
-	/**
-	* @return	Transform Curve object for the provided identifier if valid
-	*/
-	const FTransformCurve& GetTransformCurve(const FAnimationCurveIdentifier& CurveIdentifier) const;
-
-	/**
-	* @return	Rich Curve object for the provided identifier if valid
-	*/
-	const FRichCurve& GetRichCurve(const FAnimationCurveIdentifier& CurveIdentifier) const;
-		
-	/**
-	* @return	Animated (bone) attributes stored
-	*/
-	TArrayView<const FAnimatedBoneAttribute> GetAttributes() const;
-
-	/**
-	* @return	Number of animated (bone) attributes stored
-	*/
-	int32 GetNumberOfAttributes() const;
-
-	/**
-	* @return	Number of animated (bone) attributes stored for the specified bone index
-	*/
-	int32 GetNumberOfAttributesForBoneIndex(const int32 BoneIndex) const;
-
-	/**
-	* @return	All animated (bone) attributes stored for the specified bone name
-	*/
-	void GetAttributesForBone(const FName& BoneName, TArray<const FAnimatedBoneAttribute*>& OutBoneAttributes) const;
-
-	/**
-	* @return	Animated (bone) attribute object for the provided identifier if valid
-	*/
-	const FAnimatedBoneAttribute& GetAttribute(const FAnimationAttributeIdentifier& AttributeIdentifier) const;
-
-	/**
-	* @return	Animated (bone) attribute ptr for the provided identifier if valid, otherwise returns a nullptr 
-	*/
-	const FAnimatedBoneAttribute* FindAttribute(const FAnimationAttributeIdentifier& AttributeIdentifier) const;
-		
-	/**
-	* @return	The outer UAnimSequence object if found, otherwise returns a nullptr 
-	*/
-	UFUNCTION(BlueprintPure, Category = AnimationDataModel)
-	UAnimSequence* GetAnimationSequence() const;
-
-	/**
-	* @return	Multicast delegate which is broadcasted to propagated changes to any internal data, see FAnimDataModelModifiedEvent and EAnimDataModelNotifyType
-	*/
-	FAnimDataModelModifiedEvent& GetModifiedEvent() { return ModifiedEvent; }
-
-	/**
-	* @return	GUID representing the contained data and state 
-	*/
-	FGuid GenerateGuid() const;
 private:
 	/** Helper functionality used by UAnimDataController to retrieve mutable data */ 
 	FBoneAnimationTrack* FindMutableBoneTrackByName(FName Name);
@@ -286,77 +82,10 @@ private:
 	FAnimCurveBase* FindMutableCurveById(const FAnimationCurveIdentifier& CurveIdentifier);	   
 	FRichCurve* GetMutableRichCurve(const FAnimationCurveIdentifier& CurveIdentifier);
 
-	/**
-	* Broadcasts a new EAnimDataModelNotifyType with the provided payload data alongside it.
-	*
-	* @param	NotifyType			Type of notify to broadcast
-	* @param	PayloadData			Typed payload data
-	*/
-	template<typename T>
-	void Notify(EAnimDataModelNotifyType NotifyType, const T& PayloadData)
-	{
-		UScriptStruct* TypeScriptStruct = T::StaticStruct();
-
-		const FAnimDataModelNotifPayload Payload((int8*)&PayloadData, TypeScriptStruct);
-		ModifiedEvent.Broadcast(NotifyType, this, Payload);
-
-		if (ModifiedEventDynamic.IsBound())
-		{
-			ModifiedEventDynamic.Broadcast(NotifyType, this, Payload);
-		}
-
-		// Only regenerate transient data when not in a bracket, or at the end of one
-		{
-			if (NotifyType == EAnimDataModelNotifyType::BracketOpened)
-			{
-				++BracketCounter;
-			}
-			if (NotifyType == EAnimDataModelNotifyType::BracketClosed)
-			{
-				--BracketCounter;
-			}
-
-			check(BracketCounter >= 0);
-		}
-	}
-
-	/**
-	* Broadcasts a new EAnimDataModelNotifyType alongside of an empty payload.
-	*
-	* @param	NotifyType			Type of notify to broadcast
-	*/
-	void Notify(EAnimDataModelNotifyType NotifyType)
-	{
-		FEmptyPayload EmptyPayload;
-		const FAnimDataModelNotifPayload Payload((int8*)&EmptyPayload, FEmptyPayload::StaticStruct());
-
-		ModifiedEvent.Broadcast(NotifyType, this, Payload);
-
-		if (ModifiedEventDynamic.IsBound())
-		{
-			ModifiedEventDynamic.Broadcast(NotifyType, this, Payload);
-		}
-
-		// Only regenerate transient data when not in a bracket, or at the end of one
-		{
-			if (NotifyType == EAnimDataModelNotifyType::BracketOpened)
-			{
-				++BracketCounter;
-			}
-			if (NotifyType == EAnimDataModelNotifyType::BracketClosed)
-			{
-				--BracketCounter;
-			}
-
-			check(BracketCounter >= 0);
-		}
-	}
-
 private:
-
 	UPROPERTY(Transient)
 	int32 BracketCounter = 0;
-private:
+
 	/** Dynamic delegate event allows scripting to register to any broadcasted notify. */
 	UPROPERTY(BlueprintAssignable, Transient, Category = AnimationDataModel, meta = (ScriptName = "ModifiedEvent", AllowPrivateAccess = "true"))
 	FAnimDataModelModifiedDynamicEvent ModifiedEventDynamic;
@@ -369,6 +98,7 @@ private:
 	TArray<FBoneAnimationTrack> BoneAnimationTracks;
 
 	/** Total playable length of the contained animation data */
+	UE_DEPRECATED(5.1, "PlayLength is deprecated use GetPlayLength instead, as it is now calculated with Number of Frames * FrameRate instead of stored as a value")
 	UPROPERTY(VisibleAnywhere, BlueprintReadOnly, Category = "Animation Data Model", meta = (AllowPrivateAccess = "true"))
 	float PlayLength;
 	
@@ -390,7 +120,12 @@ private:
 	
 	/** Container with all animated (bone) attribute data */
 	UPROPERTY(VisibleAnywhere, BlueprintReadOnly, Category = "Animation Data Model", meta = (AllowPrivateAccess = "true"))
-	TArray<FAnimatedBoneAttribute> AnimatedBoneAttributes;	
+	TArray<FAnimatedBoneAttribute> AnimatedBoneAttributes;
+
+	UPROPERTY()
+	bool bPopulated = false;
+
+	TUniquePtr<IAnimationDataModel::FModelNotifier> Notifier;
 
 	friend class UAnimDataController;
 	friend class FAnimDataControllerTestBase;
diff --git a/Engine/Source/Runtime/Engine/Classes/Animation/AnimData/AnimDataNotifications.h b/Engine/Source/Runtime/Engine/Classes/Animation/AnimData/AnimDataNotifications.h
--- a/Engine/Source/Runtime/Engine/Classes/Animation/AnimData/AnimDataNotifications.h
+++ b/Engine/Source/Runtime/Engine/Classes/Animation/AnimData/AnimDataNotifications.h
@@ -13,7 +13,7 @@
 
 #include "AnimDataNotifications.generated.h"
 
-class UAnimDataModel;
+class IAnimationDataModel;
 
 UENUM(BlueprintType)
 enum class EAnimDataModelNotifyType : uint8
@@ -122,17 +122,40 @@ struct FSequenceLengthChangedPayload : public FEmptyPayload
 {
 	GENERATED_BODY()
 
+	PRAGMA_DISABLE_DEPRECATION_WARNINGS
+	FSequenceLengthChangedPayload() = default;
+	FSequenceLengthChangedPayload(const FSequenceLengthChangedPayload&) = default;
+	FSequenceLengthChangedPayload(FSequenceLengthChangedPayload&&) = default;
+	FSequenceLengthChangedPayload& operator=(const FSequenceLengthChangedPayload&) = default;
+	FSequenceLengthChangedPayload& operator=(FSequenceLengthChangedPayload&&) = default;
+	PRAGMA_ENABLE_DEPRECATION_WARNINGS
+
 	/** Previous playable length for the Model */
+	UE_DEPRECATED(5.0, "Time in seconds is deprecated use FFrameNumber members (PreviousNumberOfFrames) instead")
 	UPROPERTY(BlueprintReadOnly, Category = Payload)
 	float PreviousLength = 0.f;
 
 	/** Time at which the change in length has been made */
+	UE_DEPRECATED(5.0, "Time in seconds is deprecated use FFrameNumber members (Frame0) instead")
 	UPROPERTY(BlueprintReadOnly, Category = Payload)
 	float T0 = 0.f;
 
 	/** Length of time which is inserted or removed starting at T0 */
+	UE_DEPRECATED(5.0, "Time in seconds is deprecated use FFrameNumber members (Frame1) instead")
 	UPROPERTY(BlueprintReadOnly, Category = Payload)
 	float T1 = 0.f;
+
+	/** Previous playable number of frames for the Model */
+    UPROPERTY(BlueprintReadOnly, Category = Payload)
+    FFrameNumber PreviousNumberOfFrames;
+
+    /** Frame number at which the change in frames has been made */
+    UPROPERTY(BlueprintReadOnly, Category = Payload)
+    FFrameNumber Frame0;
+
+    /** Amount of frames which is inserted or removed starting at Frame0 */
+    UPROPERTY(BlueprintReadOnly, Category = Payload)
+	FFrameNumber Frame1;
 };
 
 USTRUCT(BlueprintType)
@@ -257,5 +280,5 @@ public:
 #endif // WITH_EDITOR
 };
 
-DECLARE_MULTICAST_DELEGATE_ThreeParams(FAnimDataModelModifiedEvent, const EAnimDataModelNotifyType& /* type */, UAnimDataModel* /* Model */, const FAnimDataModelNotifPayload& /* payload */);
-DECLARE_DYNAMIC_MULTICAST_DELEGATE_ThreeParams(FAnimDataModelModifiedDynamicEvent, EAnimDataModelNotifyType, NotifType, UAnimDataModel*, Model, const FAnimDataModelNotifPayload&, Payload);
\ No newline at end of file
+DECLARE_MULTICAST_DELEGATE_ThreeParams(FAnimDataModelModifiedEvent, const EAnimDataModelNotifyType& /* type */, IAnimationDataModel* /* Model */, const FAnimDataModelNotifPayload& /* payload */);
+DECLARE_DYNAMIC_MULTICAST_DELEGATE_ThreeParams(FAnimDataModelModifiedDynamicEvent, EAnimDataModelNotifyType, NotifType, TScriptInterface<IAnimationDataModel>, Model, const FAnimDataModelNotifPayload&, Payload);
\ No newline at end of file
diff --git a/Engine/Source/Runtime/Engine/Classes/Animation/AnimData/CurveIdentifier.h b/Engine/Source/Runtime/Engine/Classes/Animation/AnimData/CurveIdentifier.h
--- a/Engine/Source/Runtime/Engine/Classes/Animation/AnimData/CurveIdentifier.h
+++ b/Engine/Source/Runtime/Engine/Classes/Animation/AnimData/CurveIdentifier.h
@@ -69,10 +69,22 @@ struct ENGINE_API FAnimationCurveIdentifier
 
 		return *this;
 	}
+	
+	friend uint32 GetTypeHash(const FAnimationCurveIdentifier& CurveId)
+	{
+		return HashCombine(HashCombine(HashCombine((uint32)CurveId.Channel, (uint32)CurveId.Axis), (uint32)CurveId.CurveType), GetTypeHash(CurveId.InternalName.DisplayName));
+	}
 
+	UPROPERTY()
 	FSmartName InternalName;
+	
+	UPROPERTY()
 	ERawCurveTrackTypes CurveType;
+
+	UPROPERTY()
 	ETransformCurveChannel Channel;
+	
+	UPROPERTY()
 	EVectorCurveChannel Axis;
 };
 
diff --git a/Engine/Source/Runtime/Engine/Classes/Animation/AnimData/IAnimationDataController.h b/Engine/Source/Runtime/Engine/Classes/Animation/AnimData/IAnimationDataController.h
--- a/Engine/Source/Runtime/Engine/Classes/Animation/AnimData/IAnimationDataController.h
+++ b/Engine/Source/Runtime/Engine/Classes/Animation/AnimData/IAnimationDataController.h
@@ -7,12 +7,25 @@
 #include "CurveIdentifier.h"
 #include "Misc/FrameRate.h"
 #include "AttributeIdentifier.h"
+#include "IAnimationDataModel.h"
+#include "Logging/LogVerbosity.h"
 #include "Algo/Transform.h"
+#include "Animation/AnimCurveTypes.h"
+
+#if WITH_EDITOR
+#include "ChangeTransactor.h"
+#endif // WITH_EDITOR
 
 #include "IAnimationDataController.generated.h"
 
+#define LOCTEXT_NAMESPACE "IAnimationDataController"
+
+class FText;
 class UAssetUserData;
-class UAnimDataModel;
+class IAnimationDataModel;
+
+struct FAnimationCurveIdentifier;
+struct FAnimationAttributeIdentifier;
 
 namespace UE {
 namespace Anim {
@@ -22,7 +35,7 @@ namespace Anim {
 
 /**
  * The Controller is the sole authority to perform changes on the Animation Data Model. Any mutation to the model made will
- * cause a subsequent notify (EAnimDataModelNotifyType) to be broadcasted from the Model's ModifiedEvent. Alongside of it is a 
+ * cause a subsequent notify (EAnimDataModelNotifyType) to be broadcast-ed from the Model's ModifiedEvent. Alongside of it is a 
  * payload containing information relevant to the mutation. These notifies should be relied upon to update any dependent views 
  * or generated (derived) data.
  */
@@ -53,7 +66,7 @@ public:
 			Controller.OpenBracket(InDescription, bShouldTransact);
 		}
 		
-		FScopedBracket(TScriptInterface<IAnimationDataController>& InController, const FText& InDescription, bool bInShouldTransact=true)
+		FScopedBracket(const TScriptInterface<IAnimationDataController>& InController, const FText& InDescription, bool bInShouldTransact=true)
             : Controller(*InController), bShouldTransact(bInShouldTransact)
 		{
 			Controller.OpenBracket(InDescription, bShouldTransact);
@@ -72,25 +85,25 @@ public:
 	/**
 	* Sets the AnimDataModel instance this controller is supposed to be targeting
 	*
-	* @param	InModel		UAnimDataModel instance to target
+	* @param	InModel		IAnimationDataModel instance to target
 	*/
 	UFUNCTION(BlueprintCallable, Category = AnimationData)
-	virtual void SetModel(UAnimDataModel* InModel) = 0;
+	virtual void SetModel(TScriptInterface<IAnimationDataModel> InModel) = 0;
 
 	/**
-	* @return		The AnimDataModel instance this controller is currently targeting
+	* @return		The IAnimationDataModel instance this controller is currently targeting
 	*/
 	UFUNCTION(BlueprintCallable, Category = AnimationData)
-	virtual UAnimDataModel* GetModel() = 0;
+	virtual TScriptInterface<IAnimationDataModel> GetModelInterface() const = 0;
 
 	/**
-	* @return		The AnimDataModel instance this controller is currently targeting
+	* @return		The IAnimationDataModel instance this controller is currently targeting
 	*/
-	virtual const UAnimDataModel* const GetModel() const = 0;
+	virtual const IAnimationDataModel* const GetModel() const = 0;
 
 	/**
 	* Opens an interaction bracket, used for combining a set of controller actions. Broadcasts a EAnimDataModelNotifyType::BracketOpened notify,
-	* this can be used by any Views or dependendent systems to halt any unnecessary or invalid operations until the (last) bracket is closed.
+	* this can be used by any Views or dependent systems to halt any unnecessary or invalid operations until the (last) bracket is closed.
 	*
 	* @param	InTitle				Description of the bracket, e.g. "Generating Curve Data"
 	* @param	bShouldTransact		Whether or not any undo-redo changes should be generated
@@ -101,7 +114,6 @@ public:
 	/**
 	* Closes a previously opened interaction bracket, used for combining a set of controller actions. Broadcasts a EAnimDataModelNotifyType::BracketClosed notify.
 	*
-	* @param	InTitle				Description of the bracket, e.g. "Generating Curve Data"
 	* @param	bShouldTransact		Whether or not any undo-redo changes should be generated
 	*/
 	UFUNCTION(BlueprintCallable, Category = AnimationData)
@@ -111,22 +123,30 @@ public:
 	* Sets the total play-able length in seconds. Broadcasts a EAnimDataModelNotifyType::SequenceLengthChanged notify if successful.
 	* The number of frames and keys for the provided length is recalculated according to the current value of UAnimDataModel::FrameRate.
 	*
-	* @param	Length				New play-able length value, has to be positive and non-zero
+	* @param	NewLengthInFrames	Total new play-able number of frames value (according to frame rate), has to be positive and non-zero
 	* @param	bShouldTransact		Whether or not any undo-redo changes should be generated
 	*/
 	UFUNCTION(BlueprintCallable, Category = AnimationData)
+	virtual void SetNumberOfFrames(FFrameNumber NewLengthInFrames, bool bShouldTransact = true) = 0;
+
+	UE_DEPRECATED(5.1, "SetPlayLength with length in seconds in deprecated use SetNumberOfFrames using FFrameNumber instead")
+	UFUNCTION(BlueprintCallable, Category = AnimationData, meta=(DeprecatedFunction, DeprecationMessage="SetPlayLength is deprecated use SetNumberOfFrames instead."))
 	virtual void SetPlayLength(float Length, bool bShouldTransact = true) = 0;
 
 	/*** Sets the total play-able length in seconds. Broadcasts a EAnimDataModelNotifyType::SequenceLengthChanged notify if successful.
 	* T0 and T1 are expected to represent the window of time that was either added or removed. E.g. for insertion T0 indicates the time
 	* at which additional time starts and T1 were it ends. For removal T0 indicates the time at which time should be started to remove, and T1 indicates the end. Giving a total of T1 - T0 added or removed length.
 	* The number of frames and keys for the provided length is recalculated according to the current value of UAnimDataModel::FrameRate.
-	* @param	Length				Total new play-able length value, has to be positive and non-zero
-	* @param	T0					Point between 0 and Length at which the change in time starts
-	* @param	T1					Point between 0 and Length at which the change in time ends
+	* @param	NewLengthInFrames	Total new play-able number of frames value (according to frame rate), has to be positive and non-zero
+	* @param	T0					Point between 0 and NewLengthInFrames at which the change in length starts
+	* @param	T1					Point between 0 and NewLengthInFrames at which the change in length ends
 	* @param	bShouldTransact		Whether or not any undo-redo changes should be generated
 	*/
 	UFUNCTION(BlueprintCallable, Category = AnimationData)
+    virtual void ResizeNumberOfFrames(FFrameNumber NewLengthInFrames, FFrameNumber T0, FFrameNumber T1, bool bShouldTransact = true) = 0;
+
+	UE_DEPRECATED(5.1, "ResizePlayLength with length in seconds in deprecated use ResizeNumberOfFrames using FFrameNumber instead")
+	UFUNCTION(BlueprintCallable, Category = AnimationData, meta=(DeprecatedFunction, DeprecationMessage="ResizePlayLength is deprecated use ResizeNumberOfFrames instead."))
 	virtual void ResizePlayLength(float NewLength, float T0, float T1, bool bShouldTransact = true) = 0;
 
 	/**
@@ -136,14 +156,18 @@ public:
 	* at which additional time starts and T1 were it ends. For removal T0 indicates the time at which time should be started to remove, and T1 indicates the end. Giving a total of T1 - T0 added or removed length.
 	* The number of frames and keys for the provided length is recalculated according to the current value of UAnimDataModel::FrameRate.
 	*
-	* @param	Length				Total new play-able length value, has to be positive and non-zero
-	* @param	T0					Point between 0 and Length at which the change in time starts
-	* @param	T1					Point between 0 and Length at which the change in time ends
+	* @param	NewLengthInFrames	Total new play-able number of frames value (according to frame rate), has to be positive and non-zero
+	* @param	T0					Point between 0 and NewLengthInFrames at which the change in length starts
+	* @param	T1					Point between 0 and NewLengthInFrames at which the change in length ends
 	* @param	bShouldTransact		Whether or not any undo-redo changes should be generated
 	*/
 	UFUNCTION(BlueprintCallable, Category = AnimationData)
-	virtual void Resize(float Length, float T0, float T1, bool bShouldTransact = true) = 0;
+	virtual void ResizeInFrames(FFrameNumber NewLengthInFrames, FFrameNumber T0, FFrameNumber T1, bool bShouldTransact = true) = 0;
 
+	UE_DEPRECATED(5.1, "Resize with length in seconds in deprecated use ResizeInFrames using FFrameNumber instead")
+	UFUNCTION(BlueprintCallable, Category = AnimationData, meta=(DeprecatedFunction, DeprecationMessage="Resize is deprecated use ResizeInFrames instead."))
+	virtual void Resize(float Length, float T0, float T1, bool bShouldTransact = true) = 0;
+	
 	/**
 	* Sets the frame rate according to which the bone animation is expected to be sampled. Broadcasts a EAnimDataModelNotifyType::FrameRateChanged notify if successful.
 	* The number of frames and keys for the provided frame rate is recalculated according to the current value of UAnimDataModel::PlayLength.
@@ -175,7 +199,8 @@ public:
 	*
 	* @return	The index at which the bone track was inserted, INDEX_NONE if the insertion failed
 	*/
-	UFUNCTION(BlueprintCallable, Category = AnimationData)
+	UE_DEPRECATED(5.1, "InsertBoneTrack has been deprecated, use AddBoneTrack instead")
+	UFUNCTION(BlueprintCallable, Category = AnimationData, meta=(DeprecatedFunction, DeprecationMessage="InsertBoneTrack is deprecated use AddBoneTrack instead."))
 	virtual int32 InsertBoneTrack(FName BoneName, int32 DesiredIndex, bool bShouldTransact = true) = 0;
 
 	/**
@@ -231,7 +256,7 @@ public:
 	* The provided number of keys provided is expected to match for each component, be between FrameLowerBound and FrameUpperBound (inclusive), and be non-zero.
 	*
 	* @param	BoneName			Bone name of the track for which the keys should be set
-	* @param	FrameRange			Range of frames to set keys for
+	* @param	KeyRangeToSet			Range of frames to set keys for
 	* @param	PositionalKeys		Array of keys for the translation component
 	* @param	RotationalKeys		Array of keys for the rotation component
 	* @param	ScalingKeys			Array of keys for the scale component
@@ -246,7 +271,7 @@ public:
 	* The provided number of keys provided is expected to match for each component, be between FrameLowerBound and FrameUpperBound (inclusive), and be non-zero.
 	*
 	* @param	BoneName			Bone name of the track for which the keys should be set
-	* @param	FrameRange			Range of frames to set keys for
+	* @param	KeyRangeToSet			Range of frames to set keys for
 	* @param	PositionalKeys		Array of keys for the translation component
 	* @param	RotationalKeys		Array of keys for the rotation component
 	* @param	ScalingKeys			Array of keys for the scale component
@@ -284,7 +309,7 @@ public:
 	/**
 	* Remove the curve with provided identifier. Broadcasts a EAnimDataModelNotifyType::CurveRemoved notify if successful.
 	*
-	* @param	CopyCurveId			Identifier for the to-be-removed curve
+	* @param	CurveId				Identifier for the to-be-removed curve
 	* @param	bShouldTransact		Whether or not any undo-redo changes should be generated
 	*
 	* @return	Whether or not the curve was successfully removed
@@ -318,7 +343,7 @@ public:
 	* Replace the flags for the curve with provided identifier. Broadcasts a EAnimDataModelNotifyType::CurveFlagsChanged notify if successful.
 	*
 	* @param	CurveId			    Identifier for the curve for which the flag state is to be set
-	* @param	Flags				Flag mask with which the existings flags are to be replaced
+	* @param	Flags				Flag mask with which the existing flags are to be replaced
 	* @param	bShouldTransact		Whether or not any undo-redo changes should be generated
 	*
 	* @return	Whether or not the flag mask was successfully set
@@ -451,7 +476,7 @@ public:
 	virtual void UpdateCurveNamesFromSkeleton(const USkeleton* Skeleton, ERawCurveTrackTypes SupportedCurveType, bool bShouldTransact = true) = 0;
 
 	/**
-	* Updates the curve names with the provided skeleton, if a display name is not found it will be added thus modifying the skeleton. Broadcasts a EAnimDataModelNotifyType::CurveRenamed for each curve name for which the UID was different or if it was added as a new smartname, wrapped within BracketOpened/BracketClosed notifies.
+	* Updates the curve names with the provided skeleton, if a display name is not found it will be added thus modifying the skeleton. Broadcasts a EAnimDataModelNotifyType::CurveRenamed for each curve name for which the UID was different or if it was added as a new smart-name, wrapped within BracketOpened/BracketClosed notifies.
 	*
 	* @param	Skeleton			Skeleton to retrieve the display name values from
 	* @param	SupportedCurveType	Curve type for which the names should be updated
@@ -548,9 +573,9 @@ public:
 	* @return	Whether or not the key was successfully set
 	*/
 	template<typename AttributeType>
-	bool SetTypedAttributeKey(const FAnimationAttributeIdentifier& AttributeIdentifier, float Time, AttributeType& KeyValue, bool bShouldTransact = true)
+	bool SetTypedAttributeKey(const FAnimationAttributeIdentifier& AttributeIdentifier, float Time, const AttributeType& KeyValue, bool bShouldTransact = true)
 	{
-		return SetAttributeKey(AttributeIdentifier, Time, (const void*)&KeyValue, AttributeType::StaticStruct(), bShouldTransact);
+		return SetAttributeKey(AttributeIdentifier, Time, static_cast<const void*>(&KeyValue), AttributeType::StaticStruct(), bShouldTransact);
 	}
 
 	/**
@@ -560,6 +585,7 @@ public:
 	* @param	AttributeIdentifier		Identifier for the attribute for which the key is to be set
 	* @param	Time					Time of the to-be-set key
 	* @param	KeyValue				Value of the to-be-set key
+	* @param	TypeStruct				UScriptStruct describing the type of KeyValue
 	* @param	bShouldTransact			Whether or not any undo-redo changes should be generated
 	*
 	* @return	Whether or not the key was successfully set
@@ -572,6 +598,7 @@ public:
 	* @param	AttributeIdentifier		Identifier for the attribute for which the keys are to be replaced
 	* @param	Times					Times with which the existing key timings are to be replaced
 	* @param	KeyValues				Values with which the existing key values are to be replaced
+	* @param	TypeStruct				UScriptStruct describing the type of KeyValues
 	* @param	bShouldTransact			Whether or not any undo-redo changes should be generated
 	*
 	* @return	Whether or not replacing the attribute keys was successful
@@ -589,12 +616,12 @@ public:
 	* @return	Whether or not replacing the attribute keys was successful
 	*/
 	template<typename AttributeType>
-	bool SetTypedAttributeKeys(const FAnimationAttributeIdentifier& AttributeIdentifier, TArrayView<const float> Times, TArrayView<AttributeType> KeyValues, bool bShouldTransact = true)
+	bool SetTypedAttributeKeys(const FAnimationAttributeIdentifier& AttributeIdentifier, TArrayView<const float> Times, TArrayView<const AttributeType> KeyValues, bool bShouldTransact = true)
 	{
 		TArray<const void*> KeyValuePtrs;
 		Algo::Transform(KeyValues, KeyValuePtrs, [](const AttributeType& Value)
 		{
-			return (const void*)&Value;
+			return static_cast<const void*>(&Value);
 		});
 
 		return SetAttributeKeys(AttributeIdentifier, Times, MakeArrayView(KeyValuePtrs), AttributeType::StaticStruct(), bShouldTransact);
@@ -623,13 +650,192 @@ public:
 	*/
 	UFUNCTION(BlueprintCallable, Category = AttributeData)
 	virtual bool DuplicateAttribute(const FAnimationAttributeIdentifier& AttributeIdentifier, const FAnimationAttributeIdentifier& NewAttributeIdentifier, bool bShouldTransact = true) = 0;
+
+	/** Updates/removes/remaps contained animation data according to the newly assigned skeleton
+	 *
+	* @param	TargetSkeleton		Skeleton to which the contained animation data is to be mapped
+	* @param	bShouldTransact		Whether or not any undo-redo changes should be generated
+	* 
+	 */
+	virtual void UpdateWithSkeleton(USkeleton* TargetSkeleton, bool bShouldTransact = true) = 0;
+
+	/** Copies any animation relevant data from an already existing IAnimationDataModel object */
+	virtual void PopulateWithExistingModel(TScriptInterface<IAnimationDataModel> InModel) = 0;
+
+	/** Initializes model data structures */
+	virtual void InitializeModel() = 0;
+
+	/** Returns the final frame number calculating according to the Model its frame-rate, additionally outputs log information for invalid/loss of precision */
+	FFrameNumber ConvertSecondsToFrameNumber(double Seconds) const
+	{
+		ValidateModel();
+		
+		const FFrameRate& ModelFrameRate = GetModel()->GetFrameRate();
+		
+		static const FNumberFormattingOptions DurationFormatOptions = FNumberFormattingOptions()
+			.SetMinimumFractionalDigits(8)
+			.SetMaximumFractionalDigits(8);
+	
+		const FFrameTime FrameTime = ModelFrameRate.AsFrameTime(Seconds);
+		// Check for either small sub-frame value or near zero seconds representation (of sub-frame only)
+		if (!FMath::IsNearlyZero(FrameTime.GetSubFrame(), UE_KINDA_SMALL_NUMBER) &&
+			!FMath::IsNearlyZero(ModelFrameRate.AsSeconds(FFrameTime(0, FrameTime.GetSubFrame())), UE_DOUBLE_KINDA_SMALL_NUMBER))
+		{
+			ReportWarningf(LOCTEXT("SecondsToFrameNumberPrecisionWarning", "Insufficient precision while converting seconds to frames: {0} seconds {1} frames using {2} (sub-frame in seconds {3})"), FText::AsNumber(Seconds), FText::AsNumber(FrameTime.AsDecimal(), &DurationFormatOptions), ModelFrameRate.ToPrettyText(), FText::AsNumber(ModelFrameRate.AsSeconds(FFrameTime(0, FrameTime.GetSubFrame())), &DurationFormatOptions));
+		}
+
+		return FrameTime.GetFrame();
+	}	
 protected:
 	/** Functionality used by FOpenBracketAction and FCloseBracketAction to broadcast their equivalent notifies without actually opening a bracket. */
 	virtual void NotifyBracketOpen() = 0;
 	virtual void NotifyBracketClosed() = 0;
 
+	template <typename FmtType, typename... Types>
+    void ReportWarningf(const FmtType& Fmt, Types... Args) const
+	{
+		ReportWarning(FText::Format(Fmt, Args...));
+	}
+
+	template <typename FmtType, typename... Types>
+    void ReportErrorf(const FmtType& Fmt, Types... Args) const
+	{
+		ReportError(FText::Format(Fmt, Args...));
+	}
+
+	void Report(ELogVerbosity::Type Verbosity, const FText& Message) const
+	{
+		ReportMessage(GetModelInterface().GetObject(), Message, Verbosity);
+	}
+
+	template <typename FmtType, typename... Types>
+    void Reportf(ELogVerbosity::Type Verbosity, const FmtType& Fmt, Types... Args) const
+    {
+		ReportMessage(GetModelInterface().GetObject(), FText::Format(Fmt, Args...), Verbosity);
+    }
+
+	void ReportWarning(const FText& Message) const
+	{
+		ReportMessage(GetModelInterface().GetObject(), Message, ELogVerbosity::Warning);
+	}
+	
+	void ReportError(const FText& Message) const
+    {
+    	ReportMessage(GetModelInterface().GetObject(), Message, ELogVerbosity::Error);
+    }
+
+	/** Returns whether or not the supplied curve type is supported by the controller functionality */
+	static bool IsSupportedCurveType(ERawCurveTrackTypes CurveType)
+	{
+		const TArray<ERawCurveTrackTypes> SupportedTypes = { ERawCurveTrackTypes::RCT_Float, ERawCurveTrackTypes::RCT_Transform };
+		return SupportedTypes.Contains(CurveType);
+	}
+
+	/** Ensures that a valid model is currently targeted */
+	void ValidateModel() const
+	{
+		checkf(GetModel() != nullptr, TEXT("Invalid Model"));
+	}
+		
+	/** Verifies whether or not the Model's outer object is (or is derived from) the specified UClass */
+	bool CheckOuterClass(UClass* InClass) const
+	{
+		ValidateModel();
+	
+		if (const UObject* ModelOuter = GetModelInterface().GetObject()->GetOuter())
+		{
+			if (const UClass* OuterClass = ModelOuter->GetClass())
+			{
+				if (OuterClass == InClass || OuterClass->IsChildOf(InClass))
+				{
+					return true;
+				}
+				else
+				{
+					ReportErrorf(NSLOCTEXT("IAnimationDataController", "NoValidOuterClassError", "Incorrect outer object class found for Animation Data Model {0}, expected {1} actual {2}"), FText::FromString(GetModelInterface().GetObject()->GetName()), FText::FromString(InClass->GetName()), FText::FromString(OuterClass->GetName()));
+				}
+			}
+		}
+		else
+		{
+			ReportErrorf(NSLOCTEXT("IAnimationDataController", "NoValidOuterObjectFoundError", "No valid outer object found for Animation Data Model {0}"), FText::FromString(GetModelInterface().GetObject()->GetName()));
+		}
+
+		return false;
+	}
+	
+	/** Returns the string representation of the provided curve enum type value */
+	static FString GetCurveTypeValueName(ERawCurveTrackTypes InType)
+	{
+		FString ValueString;
+		if (const UEnum* Enum = FindObject<UEnum>(nullptr, TEXT("/Script/Engine.ERawCurveTrackTypes")))
+		{
+			ValueString = Enum->GetNameStringByValue(static_cast<int64>(InType));
+		}
+
+		return ValueString;
+	}
+	
+	static constexpr int32 DefaultCurveFlags = EAnimAssetCurveFlags::AACF_Editable;
+
+#if WITH_EDITOR
+	typedef TUniquePtr<UE::FScopedCompoundTransaction> FTransaction;
+	FTransaction ConditionalTransaction(const FText& Description, bool bCondition)
+	{
+		if (UE::FChangeTransactor::CanTransactChanges() && bCondition)
+		{ 
+			return MakeUnique<UE::FScopedCompoundTransaction>(ChangeTransactor, Description);
+		}
+
+		return nullptr;
+	}
+
+	typedef TUniquePtr<IAnimationDataController::FScopedBracket> FBracket; 
+	FBracket ConditionalBracket(const FText& Description, bool bCondition)
+	{	
+		return MakeUnique<IAnimationDataController::FScopedBracket>(this, Description, UE::FChangeTransactor::CanTransactChanges() && bCondition);	
+	}
+
+	template<typename ActionClass, class... ActionArguments>
+	void ConditionalAction(bool bCondition, ActionArguments... Arguments)
+	{
+		if (UE::FChangeTransactor::CanTransactChanges() && bCondition) \
+		{ 
+			ChangeTransactor.AddTransactionChange<ActionClass>(Arguments...); \
+		}
+	}
+protected:
+	UE::FChangeTransactor ChangeTransactor;
+#endif // WITH_EDITOR
+
+public:
+	template <typename FmtType, typename... Types>
+	static void ReportObjectWarningf(const UObject* ErrorObject, const FmtType& Fmt, Types... Args)
+	{
+		ReportMessage(ErrorObject, FText::Format(Fmt, Args...), ELogVerbosity::Warning);
+	}
+
+	template <typename FmtType, typename... Types>
+	static void ReportObjectErrorf(const UObject* ErrorObject, const FmtType& Fmt, Types... Args)
+	{
+		ReportMessage(ErrorObject, FText::Format(Fmt, Args...), ELogVerbosity::Error);
+	}
+	
+	template <typename FmtType, typename... Types>
+	static void Reportf(ELogVerbosity::Type LogVerbosity, const UObject* ErrorObject, const FmtType& Fmt, Types... Args)
+	{
+		ReportMessage(ErrorObject, FText::Format(Fmt, Args...), ELogVerbosity::Error);
+	}
+
+	static void ReportMessage(const UObject* ErrorObject, const FText& InMessage, ELogVerbosity::Type LogVerbosity);	
 private:
 	friend class FAnimDataControllerTestBase;
 	friend UE::Anim::FOpenBracketAction;
 	friend UE::Anim::FCloseBracketAction;
 };
+
+#define CONDITIONAL_TRANSACTION(Text) FTransaction Transaction = ConditionalTransaction(Test, bShouldTransact);
+#define CONDITIONAL_BRACKET(Text) FBracket Bracket = ConditionalBracket(Text, bShouldTransact);
+#define CONDITIONAL_ACTION(ActionClass, ...) ConditionalAction<ActionClass>(bShouldTransact, __VA_ARGS__);
+
+#undef LOCTEXT_NAMESPACE // "IAnimationDataController"
\ No newline at end of file
diff --git a/Engine/Source/Runtime/Engine/Classes/Animation/AnimMontage.h b/Engine/Source/Runtime/Engine/Classes/Animation/AnimMontage.h
--- a/Engine/Source/Runtime/Engine/Classes/Animation/AnimMontage.h
+++ b/Engine/Source/Runtime/Engine/Classes/Animation/AnimMontage.h
@@ -501,7 +501,7 @@ public:
 	bool SetNextSectionName(FName const & SectionName, FName const & NewNextSectionName);
 	bool SetNextSectionID(int32 const & SectionID, int32 const & NewNextSectionID);
 
-	bool IsValid() const { return (Montage!=NULL); }
+	bool IsValid() const { return (Montage!=nullptr); }
 	bool IsPlaying() const { return IsValid() && bPlaying; }
 	void SetPlaying(bool bInPlaying) { bPlaying = bInPlaying; }
 	bool IsStopped() const { return Blend.GetDesiredValue() == 0.f; }
@@ -929,6 +929,9 @@ public:
 	virtual void InvalidateRecursiveAsset() override;
 	virtual bool ContainRecursive(TArray<UAnimCompositeBase*>& CurrentAccumulatedList) override;
 	virtual void SetCompositeLength(float InLength) override;
+#if WITH_EDITOR
+	virtual void PopulateWithExistingModel(TScriptInterface<IAnimationDataModel> ExistingDataModel) override;
+#endif // WITH_EDITOR
 	//~End UAnimCompositeBase Interface
 
 	/** Utility function to create dynamic montage from AnimSequence */
diff --git a/Engine/Source/Runtime/Engine/Classes/Animation/AnimSequence.h b/Engine/Source/Runtime/Engine/Classes/Animation/AnimSequence.h
--- a/Engine/Source/Runtime/Engine/Classes/Animation/AnimSequence.h
+++ b/Engine/Source/Runtime/Engine/Classes/Animation/AnimSequence.h
@@ -23,7 +23,7 @@
 #include "Animation/AttributeCurve.h"
 
 #if WITH_EDITOR
-#include "AnimData/AnimDataModel.h"
+#include "AnimData/IAnimationDataModel.h"
 #endif // WITH_EDITOR
 
 #include "AnimSequence.generated.h"
@@ -233,21 +233,11 @@ protected:
 	UPROPERTY()
 	FFrameRate SamplingFrameRate;
 
-	/**
-	 * In the future, maybe keeping RawAnimSequenceTrack + TrackMap as one would be good idea to avoid inconsistent array size
-	 * TrackToSkeletonMapTable(i) should contains  track mapping data for RawAnimationData(i). 
-	 */
-	UE_DEPRECATED(5.0, "TrackToSkeletonMapTable has been deprecated see FBoneAnimationTrack::BoneTreeIndex")
-	UPROPERTY()
-	TArray<struct FTrackToSkeletonMap> TrackToSkeletonMapTable;
-
-	/**
-	 * Raw uncompressed keyframe data. 
-	 */
 	UE_DEPRECATED(5.0, "RawAnimationData has been deprecated see FBoneAnimationTrack::InternalTrackData")
 	TArray<struct FRawAnimSequenceTrack> RawAnimationData;
 
-	// Update this if the contents of RawAnimationData changes;
+	// Update this if the contents of RawAnimationData changes
+	UE_DEPRECATED(5.1, "RawDataGuid has been deprecated see GenerateGuidFromModel instead")
 	UPROPERTY()
 	FGuid RawDataGuid;
 
@@ -262,7 +252,6 @@ protected:
 	 * Source RawAnimationData. Only can be overridden by when transform curves are added first time OR imported
 	 */
 	TArray<struct FRawAnimSequenceTrack> SourceRawAnimationData_DEPRECATED;
-
 public:
 
 	/**
@@ -370,10 +359,6 @@ public:
 	UPROPERTY()
 	FString SourceFileTimestamp_DEPRECATED;
 
-	UE_DEPRECATED(5.0, "bNeedsRebake has been deprecated, transform curves are now baked during compression")
-	UPROPERTY(transient)
-	bool bNeedsRebake;
-
 	// Track whether we have updated markers so cached data can be updated
 	int32 MarkerDataUpdateCounter;
 #endif // WITH_EDITORONLY_DATA
@@ -394,6 +379,9 @@ public:
 #if WITH_EDITOR
 	virtual void PostEditChangeProperty(FPropertyChangedEvent& PropertyChangedEvent) override;
 	virtual bool IsCachedCookedPlatformDataLoaded(const ITargetPlatform* TargetPlatform) override;
+	virtual void WillNeverCacheCookedPlatformDataAgain() override;
+	virtual void ClearCachedCookedPlatformData(const ITargetPlatform* TargetPlatform) override;
+	virtual void ClearAllCachedCookedPlatformData() override;
 #endif // WITH_EDITOR
 	virtual void BeginDestroy() override;
 	virtual void GetAssetRegistryTags(TArray<FAssetRegistryTag>& OutTags) const override;
@@ -415,24 +403,10 @@ public:
 	virtual void RefreshCacheData() override;
 	virtual EAdditiveAnimationType GetAdditiveAnimType() const override { return AdditiveAnimType; }
 	virtual int32 GetNumberOfSampledKeys() const override;
-	virtual const FFrameRate& GetSamplingFrameRate() const override { return TargetFrameRate; }
+	virtual FFrameRate GetSamplingFrameRate() const override { return TargetFrameRate; }
 	virtual void EvaluateCurveData(FBlendedCurve& OutCurve, float CurrentTime, bool bForceUseRawData = false) const override;
 	virtual float EvaluateCurveData(SmartName::UID_Type CurveUID, float CurrentTime, bool bForceUseRawData = false) const override;
 	virtual bool HasCurveData(SmartName::UID_Type CurveUID, bool bForceUseRawData) const override;
-
-#if WITH_EDITOR
-	UE_DEPRECATED(5.0, "MarkRawDataAsModified has been deprecated, any (Raw Data) modification should be applied using the UAnimDataController API instead. This will handle updating the GUID instead.")
-	PRAGMA_DISABLE_DEPRECATION_WARNINGS
-	virtual void MarkRawDataAsModified(bool bForceNewRawDatGuid = true) override
-	{
-		Super::MarkRawDataAsModified();
-		bUseRawDataOnly = true;
-		RawDataGuid = bForceNewRawDatGuid ? FGuid::NewGuid() : GenerateGuidFromRawData();
-		FlagDependentAnimationsAsRawDataOnly();
-		UpdateDependentStreamingAnimations();
-	}
-	PRAGMA_ENABLE_DEPRECATION_WARNINGS
-#endif
 	//~ End UAnimSequenceBase Interface
 
 	// Extract Root Motion transform from the animation
@@ -473,9 +447,7 @@ public:
 #if WITH_EDITORONLY_DATA
 protected:
 	void UpdateCompressedCurveName(SmartName::UID_Type CurveUID, const struct FSmartName& NewCurveName);
-#endif
 private:
-#if WITH_EDITORONLY_DATA
 	void UpdateRetargetSourceAsset();
 
 	/** Updates the stored sampling frame-rate using the sequence length and number of sampling keys */
@@ -527,7 +499,9 @@ public:
 	 * @param	Time			Time on track to interpolate to.
 	 * @param	bUseRawData		If true, use raw animation data instead of compressed data.
 	 */
+	UE_DEPRECATED(5.1, "Use other GetBoneTransform signature")
 	void GetBoneTransform(FTransform& OutAtom, int32 TrackIndex, float Time, bool bUseRawData) const;
+	void GetBoneTransform(FTransform& OutAtom, int32 TrackIndex, double Time, bool bUseRawData) const;
 
 	/**
 	 * Get Bone Transform of the Time given, relative to Parent for the Track Given
@@ -594,8 +568,8 @@ protected:
     void SerializeCompressedData(FArchive& Ar, bool bDDCData);
 #if WITH_EDITOR
 	void SetCompressionComplete();
+	virtual void OnAnimModelLoaded() override;
 #endif
-
 public:
 	bool IsCompressedDataValid() const;
 	bool IsCurveCompressedDataValid() const;
@@ -705,6 +679,7 @@ public:
 	void UpdateDependentStreamingAnimations() const;
 
 	// Generate a GUID from a hash of our own raw data
+	UE_DEPRECATED(5.1, "GenerateGuidFromRawData has been deprecated use IAnimationDataModel::GenerateGuid instead")
 	FGuid GenerateGuidFromRawData() const;
 
 	// Should we be always using our raw data (i.e is our compressed data stale)
@@ -712,9 +687,12 @@ public:
 	void SetUseRawDataOnly(bool bInUseRawDataOnly) { bUseRawDataOnly = bInUseRawDataOnly; }
 
 	// Return this animations guid for the raw data
+	UE_DEPRECATED(5.1, "GetRawDataGuid has been deprecated use GenerateGuidFromModel instead")
 	FGuid GetRawDataGuid() const
 	{ 
+		PRAGMA_DISABLE_DEPRECATION_WARNINGS
 		return RawDataGuid;
+		PRAGMA_ENABLE_DEPRECATION_WARNINGS
 	}
 
 	/** Resets Bone Animation, Curve data and Notify tracks **/
@@ -744,18 +722,13 @@ private:
 	 */ 
 	virtual void RemapTracksToNewSkeleton( USkeleton* NewSkeleton, bool bConvertSpaces ) override;
 
-	/**
-	 * Remap NaN tracks from the RawAnimation data and recompress
-	 */	
-	void RemoveNaNTracks();
-
 	/** Retargeting functions */
 	bool ConvertAnimationDataToRiggingData(FAnimSequenceTrackContainer & RiggingAnimationData);
 	bool ConvertRiggingDataToAnimationData(FAnimSequenceTrackContainer & RiggingAnimationData);
 	int32 GetSpaceBasedAnimationData(TArray< TArray<FTransform> > & AnimationDataInComponentSpace, FAnimSequenceTrackContainer * RiggingAnimationData) const;
+#endif
 
 public:
-#endif
 	/** Refresh sync marker data*/
 	void RefreshSyncMarkerDataFromAuthored();
 
@@ -816,13 +789,15 @@ protected:
 protected:
 #if WITH_EDITOR
 	// Begin UAnimSequenceBase virtual overrides
-	virtual void OnModelModified(const EAnimDataModelNotifyType& NotifyType, UAnimDataModel* Model, const FAnimDataModelNotifPayload& Payload) override;
+	virtual void OnModelModified(const EAnimDataModelNotifyType& NotifyType, IAnimationDataModel* Model, const FAnimDataModelNotifPayload& Payload) override;
 	virtual void PopulateModel() override;
 	// End UAnimSequenceBase virtual overrides
 
 	void EnsureValidRawDataGuid();
 	void RecompressAnimationData();
 	void ResampleAnimationTrackData();
+	void CalculateNumberOfSampledKeys();
+	void ClearResampledAnimationTrackData();
 
 	void DeleteBoneAnimationData();
 	void DeleteDeprecatedRawAnimationData();
@@ -867,4 +842,5 @@ public:
 	friend class FCustomAttributeCustomization;
 	friend class FAnimSequenceTestBase;
 	friend struct UE::Anim::Compression::FScopedCompressionGuard;
+	friend class FAnimDataControllerTestBase;
 };
diff --git a/Engine/Source/Runtime/Engine/Classes/Animation/AnimSequenceBase.h b/Engine/Source/Runtime/Engine/Classes/Animation/AnimSequenceBase.h
--- a/Engine/Source/Runtime/Engine/Classes/Animation/AnimSequenceBase.h
+++ b/Engine/Source/Runtime/Engine/Classes/Animation/AnimSequenceBase.h
@@ -29,6 +29,7 @@ enum ETypeAdvanceAnim
 struct FAnimationPoseData;
 struct FAnimDataModelNotifPayload;
 class UAnimDataModel;
+class IAnimationDataModel;
 enum class EAnimDataModelNotifyType : uint8;
 
 UCLASS(abstract, BlueprintType)
@@ -41,14 +42,20 @@ public:
 	UPROPERTY()
 	TArray<struct FAnimNotifyEvent> Notifies;
 
+protected:
 	/** Length (in seconds) of this AnimSequence if played back with a speed of 1.0. */
 	UE_DEPRECATED(5.0, "Public access to SequenceLength is deprecated, use GetPlayLength or UAnimDataController::SetPlayLength instead")
 	UPROPERTY(Category=Length, AssetRegistrySearchable, VisibleAnywhere, BlueprintReadOnly)
 	float SequenceLength;
 
-	UE_DEPRECATED(5.0, "SetSequenceLength is deprecated use UAnimDataController::SetPlayLength instead")
-	virtual void SetSequenceLength(float NewLength);
-
+	/**
+	 * Raw uncompressed float curve data
+	 */
+	UE_DEPRECATED(5.0, "Public access to RawCurveData is deprecated, see UAnimDataModel for source data or use GetCurveData for runtime instead")
+	UPROPERTY()
+	struct FRawCurveTracks RawCurveData;
+	
+public:
 	/** Number for tweaking playback rate of this animation globally. */
 	UPROPERTY(EditAnywhere, Category=Animation)
 	float RateScale;
@@ -59,14 +66,6 @@ public:
 	 */
 	UPROPERTY(EditAnywhere, Category=Animation)
 	bool bLoop;
-
-	/**
-	 * Raw uncompressed float curve data 
-	 */
-	UE_DEPRECATED(5.0, "Public access to RawCurveData is deprecated, see UAnimDataModel for source data or use GetCurveData for runtime instead")
-	UPROPERTY()
-	struct FRawCurveTracks RawCurveData;
-
 #if WITH_EDITORONLY_DATA
 	// if you change Notifies array, this will need to be rebuilt
 	UPROPERTY()
@@ -145,7 +144,7 @@ public:
 	virtual int32 GetNumberOfSampledKeys() const;
 
 	/** Return rate at which the animation is sampled **/
-	virtual const FFrameRate& GetSamplingFrameRate() const;
+	virtual FFrameRate GetSamplingFrameRate() const;
 
 #if WITH_EDITOR
 	/** Get the frame number for the provided time */
@@ -238,10 +237,6 @@ public:
 	virtual void EnableRootMotionSettingFromMontage(bool bInEnableRootMotion, const ERootMotionRootLock::Type InRootMotionRootLock) {};
 
 #if WITH_EDITOR
-	// Store that our raw data has changed so that we can get correct compressed data later on
-	UE_DEPRECATED(5.0, "MarkRawDataAsModified has been deprecated, any (Raw Data) modification should be applied using the UAnimDataController API instead. This will handle updating the GUID instead.")
-	virtual void MarkRawDataAsModified(bool bForceNewRawDatGuid = true) {}
-
 private:
 	DECLARE_MULTICAST_DELEGATE( FOnNotifyChangedMulticaster );
 	FOnNotifyChangedMulticaster OnNotifyChanged;
@@ -269,18 +264,24 @@ protected:
 		}
 	}
 #if WITH_EDITOR
+	void OnEndLoadPackage(const FEndLoadPackageContext& Context);
+	virtual void OnAnimModelLoaded();
 public:
-	/** Returns the UAnimDataModel object embedded in this UAnimSequenceBase */
-	UAnimDataModel* GetDataModel() const;
+	/** Returns the IAnimationDataModel object embedded in this UAnimSequenceBase */
+	IAnimationDataModel* GetDataModel() const;
+
+	/** Returns the IAnimationDataModel as a script-interface, provides access to UObject and Interface */
+	TScriptInterface<IAnimationDataModel> GetDataModelInterface() const;
 
 	/** Returns the transient UAnimDataController set to operate on DataModel */
 	IAnimationDataController& GetController();
 protected:
 	/** Populates the UAnimDataModel object according to any pre-existing data. (overrides expect to populate the model according to their data) */
 	virtual void PopulateModel();
+	virtual void PopulateWithExistingModel(TScriptInterface<IAnimationDataModel> ExistingDataModel);
 
 	/** Callback registered to UAnimDatModel::GetModifiedEvent for the embedded object */
-	virtual void OnModelModified(const EAnimDataModelNotifyType& NotifyType, UAnimDataModel* Model, const FAnimDataModelNotifPayload& Payload);
+	virtual void OnModelModified(const EAnimDataModelNotifyType& NotifyType, IAnimationDataModel* Model, const FAnimDataModelNotifPayload& Payload);
 
 	/** Validates that DataModel contains a valid UAnimDataModel object */
 	void ValidateModel() const;
@@ -289,7 +290,7 @@ protected:
 	void BindToModelModificationEvent();
 
 	/** Replaces the current DataModel, if any, with the provided one */
-	void CopyDataModel(const UAnimDataModel* ModelToDuplicate);
+	void CopyDataModel(const TScriptInterface<IAnimationDataModel>& ModelToDuplicate);
 private:
 	/** Creates a new UAnimDataModel instance and sets DataModel accordingly */
 	void CreateModel();
@@ -301,7 +302,7 @@ public:
 		if(ShouldDataModelBeValid())
 		{
 			ValidateModel();
-			return DataModel != nullptr;
+			return DataModelInterface != nullptr;
 		}
 
 		return false;
@@ -310,10 +311,14 @@ public:
 
 #if WITH_EDITORONLY_DATA
 protected:
-	/** UAnimDataModel instance containing source animation data */
+	UE_DEPRECATED(5.1, "DataModel has been converted to DataModelInterface")
 	UPROPERTY(VisibleAnywhere, BlueprintReadOnly, Category = "Animation Model")
 	TObjectPtr<UAnimDataModel> DataModel;
 
+	/** IAnimationDataModel instance containing (source) animation data */
+	UPROPERTY(VisibleAnywhere, BlueprintReadOnly, Category = "Animation Model")
+	TScriptInterface<IAnimationDataModel> DataModelInterface;
+
 	/** Flag set whenever the data-model is initially populated (during upgrade path) */
 	bool bPopulatingDataModel;
 
