Reintroducing AnimationDataModel refactor with additional memory and cooktime optimizations.

diff --git a/Engine/Source/Runtime/Engine/Classes/Animation/AnimSequenceHelpers.h b/Engine/Source/Runtime/Engine/Classes/Animation/AnimSequenceHelpers.h
--- a/Engine/Source/Runtime/Engine/Classes/Animation/AnimSequenceHelpers.h
+++ b/Engine/Source/Runtime/Engine/Classes/Animation/AnimSequenceHelpers.h
@@ -8,7 +8,7 @@
 #include "Delegates/IDelegateInstance.h"
 
 class UAnimSequence;
-class UAnimDataModel;
+class IAnimationDataModel;
 class UAnimDataController;
 class IAnimationDataController;
 class FName;
@@ -17,6 +17,7 @@ struct FRawCurveTracks;
 struct FRawAnimSequenceTrack;
 struct FBoneAnimationTrack;
 struct FRawAnimSequenceTrack;
+struct FAnimationPoseData;
 class IAnimationDataController;
 
 enum class EAnimInterpolationType : uint8;
@@ -27,7 +28,7 @@ namespace Anim {
 
 #if WITH_EDITOR
 	/**
-	* Populates a FCompactPose according to the data stored within a UAnimDataModel provided a time to sample.
+	* Populates a FCompactPose according to the data stored within a IAnimationDataModel provided a time to sample.
 	*
 	* @param	Model				Model to sample from
 	* @param	OutPose				[out] Pose to be populated using the model
@@ -36,29 +37,42 @@ namespace Anim {
 	* @param	RetargetSource		Name of the retarget source to retrieve from the skeleton
 	* @param	RetargetTransforms	Base transforms to use when retargetting the pose
 	*/
-	ENGINE_API void BuildPoseFromModel(const UAnimDataModel* Model, FCompactPose& OutPose, const float Time, const EAnimInterpolationType& InterpolationType, const FName& RetargetSource, const TArray<FTransform>& RetargetTransforms);
+	UE_DEPRECATED(5.1, "Use BuildPoseFromModel with different signature instead")
+	ENGINE_API void BuildPoseFromModel(const IAnimationDataModel* Model, FCompactPose& OutPose, const float Time, const EAnimInterpolationType& InterpolationType, const FName& RetargetSource, const TArray<FTransform>& RetargetTransforms);
 
 	/**
-	* Populates float curves according to the data stored within a UAnimDataModel provided a time to sample.
+	* Populates a FAnimationPoseData according to the data stored within a IAnimationDataModel provided a time to sample.
+	*
+	* @param	Model				Model to sample from
+	* @param	OutPoseData			[out] Pose data to be populated using the model
+	* @param	Time				Time value to sample at
+	* @param	InterpolationType	Method to be used to interpolate between different sample points
+	* @param	RetargetSource		Name of the retarget source to retrieve from the skeleton
+	* @param	RetargetTransforms	Base transforms to use when retargetting the pose
+	*/
+	ENGINE_API void BuildPoseFromModel(const IAnimationDataModel* Model, FAnimationPoseData& OutPoseData, const double Time, const EAnimInterpolationType& InterpolationType, const FName& RetargetSource, const TArray<FTransform>& RetargetTransforms);
+
+	/**
+	* Populates float curves according to the data stored within a IAnimationDataModel provided a time to sample.
 	*
 	* @param	Model				Model to sample from
 	* @param	OutCurves			[out] Curves to be populated using the model
 	* @param	Time				Time value to sample at
 	*/
-	ENGINE_API void EvaluateFloatCurvesFromModel(const UAnimDataModel* Model, FBlendedCurve& OutCurves, float Time);
+	ENGINE_API void EvaluateFloatCurvesFromModel(const IAnimationDataModel* Model, FBlendedCurve& OutCurves, double Time);
 
 	/**
-	* Populates transform curves according to the data stored within a UAnimDataModel provided a time to sample.
+	* Populates transform curves according to the data stored within a IAnimationDataModel provided a time to sample.
 	*
 	* @param	Model				Model to sample from
 	* @param	OutCurves			[out] Curves which were sampled from the model, key-ed with their equivalent name
 	* @param	Time				Time value to sample at
 	* @param	BlendWeight			Blending weight to apply when sampling
 	*/
-	ENGINE_API void EvaluateTransformCurvesFromModel(const UAnimDataModel* Model, TMap<FName, FTransform>& OutCurves, float Time, float BlendWeight);
+	ENGINE_API void EvaluateTransformCurvesFromModel(const IAnimationDataModel* Model, TMap<FName, FTransform>& OutCurves, double Time, float BlendWeight);
 
 	/**
-	* Retrieves a single bone transform according to the data stored within a UAnimDataModel provided a time and track to sample.
+	* Retrieves a single bone transform according to the data stored within a IAnimationDataModel provided a time and track to sample.
 	*
 	* @param	Model				Model to sample from
 	* @param	OutTransform		[out] Transform which was sampled from the model
@@ -66,20 +80,20 @@ namespace Anim {
 	* @param	Time				Time value to sample at
 	* @param	InterpolationType	Method to be used to interpolate between different sample points
 	*/
-	ENGINE_API void GetBoneTransformFromModel(const UAnimDataModel* Model, FTransform& OutTransform, int32 TrackIndex, float Time, const EAnimInterpolationType& Interpolation);
+	ENGINE_API void GetBoneTransformFromModel(const IAnimationDataModel* Model, FTransform& OutTransform, int32 TrackIndex, double Time, const EAnimInterpolationType& Interpolation);
 	
 	/**
-	* Retrieves a single bone transform according to the data stored within a UAnimDataModel provided a key and track to sample.
+	* Retrieves a single bone transform according to the data stored within a IAnimationDataModel provided a key and track to sample.
 	*
 	* @param	Model				Model to sample from
 	* @param	OutTransform		[out] Transform which was sampled from the model
 	* @param	TrackIndex			Track / bone index which should be sampld
 	* @param	KeyIndex			Specific key index to sample
 	*/
-	ENGINE_API void GetBoneTransformFromModel(const UAnimDataModel* Model, FTransform& OutTransform, int32 TrackIndex, int32 KeyIndex);
+	ENGINE_API void GetBoneTransformFromModel(const IAnimationDataModel* Model, FTransform& OutTransform, int32 TrackIndex, int32 KeyIndex);
 
 	/**
-	* Copies over any individual curve from FRawCurveTracks to a UAnimDataModel instance targeted by the provided controller.
+	* Copies over any individual curve from FRawCurveTracks to a IAnimationDataModel instance targeted by the provided controller.
 	*
 	* @param	CurveData			Container of the curves to be copied over
 	* @param	Skeleton			Skeleton to use for verifying and or add the curve names
@@ -88,7 +102,7 @@ namespace Anim {
 	ENGINE_API void CopyCurveDataToModel(const FRawCurveTracks& CurveData, const USkeleton* Skeleton, IAnimationDataController& Controller);
 
 	/**
-	* Copies over any individual curve from FRawCurveTracks to a UAnimDataModel instance targeted by the provided controller.
+	* Copies over any individual curve from FRawCurveTracks to a IAnimationDataModel instance targeted by the provided controller.
 	*
 	* @param	SourceAnimSeq				Animation Sequence to copy the notifies from
 	* @param	DestAnimSeq					Animation Sequence to copy the notifies to
@@ -212,8 +226,30 @@ namespace Anim {
 #endif // WITH_EDITOR
 	}
 
+	namespace DataModel { struct FEvaluationContext; }
 	namespace Retargeting
 	{
+		struct ENGINE_API FRetargetTracking
+		{
+			const FCompactPoseBoneIndex PoseBoneIndex;
+			const int32 SkeletonBoneIndex;
+
+			FRetargetTracking(const FCompactPoseBoneIndex InPoseBoneIndex, const int32 InSkeletonBoneIndex) : PoseBoneIndex(InPoseBoneIndex), SkeletonBoneIndex(InSkeletonBoneIndex) {}
+		};
+
+		struct ENGINE_API FRetargetingScope
+		{
+			FRetargetingScope(FCompactPose& ToRetargetPose, const DataModel::FEvaluationContext& InEvaluationContext);
+			~FRetargetingScope();
+			
+			void AddTrackedBone(FCompactPoseBoneIndex CompactBoneIndex, int32 SkeletonBoneIndex) const;
+		private:
+			FCompactPose& RetargetPose;
+			const DataModel::FEvaluationContext& EvaluationContext;
+			TArray<FRetargetTracking>& RetargetTracking;
+			bool bShouldRetarget;
+		};
+		
 		/**
 		* Retargeting the provided pose using the retarget source or transforms
 		*
@@ -221,9 +257,16 @@ namespace Anim {
 		* @param	RetargetSource		Name of the retarget source to retrieve data for
 		* @param	RetargetTransforms	Set of transforms to use as the ref skeleton bone transforms
 		*/
+		UE_DEPRECATED(5.1, "UE::Anim::Retargeting::RetargetPose has been deprecated as retargeting while evaluating pose data from a DataModel is now handled/performed implicitly")
 		ENGINE_API void RetargetPose(FCompactPose& InOutPose, const FName& RetargetSource, const TArray<FTransform>& RetargetTransforms);
 	}
 
+	struct ENGINE_API FBuildRawPoseScratchArea : public TThreadSingleton<FBuildRawPoseScratchArea>
+	{
+		TArray<Retargeting::FRetargetTracking> RetargetTracking;
+		TArray<FVirtualBoneCompactPoseData> VirtualBoneCompactPoseData;
+	};
+
 } // namespace Anim	
 
 } // namespace UE
\ No newline at end of file
diff --git a/Engine/Source/Runtime/Engine/Classes/Animation/AnimStreamable.h b/Engine/Source/Runtime/Engine/Classes/Animation/AnimStreamable.h
--- a/Engine/Source/Runtime/Engine/Classes/Animation/AnimStreamable.h
+++ b/Engine/Source/Runtime/Engine/Classes/Animation/AnimStreamable.h
@@ -38,6 +38,7 @@ public:
 	float StartTime;
 
 	float SequenceLength;
+	int32 NumFrames;
 
 	// Compressed Data for this chunk (if nullptr then data needs to be loaded via BulkData)
 	FCompressedAnimSequence* CompressedAnimSequence;
@@ -103,6 +104,9 @@ public:
 	UPROPERTY(EditAnywhere, AssetRegistrySearchable, Category = Animation)
 	FName RetargetSource;
 
+	UPROPERTY()
+	FFrameRate SamplingFrameRate;
+
 #if WITH_EDITORONLY_DATA
 
 	// Sequence the streamable was created from (used for reflecting changes to the source in editor)
diff --git a/Engine/Source/Runtime/Engine/Classes/Animation/AnimationAsset.h b/Engine/Source/Runtime/Engine/Classes/Animation/AnimationAsset.h
--- a/Engine/Source/Runtime/Engine/Classes/Animation/AnimationAsset.h
+++ b/Engine/Source/Runtime/Engine/Classes/Animation/AnimationAsset.h
@@ -250,7 +250,7 @@ struct FPoseCurve
 struct FAnimExtractContext
 {
 	/** Position in animation to extract pose from */
-	float CurrentTime;
+	double CurrentTime;
 
 	/** Is root motion being extracted? */
 	bool bExtractRootMotion;
@@ -273,8 +273,19 @@ struct FAnimExtractContext
 	 * by several animation nodes to optimize evaluation time.
 	 */
 	TArray<bool> BonesRequired;
+	
+	UE_DEPRECATED(5.1, "FAnimExtractContext construct with float-based time value is deprecated, use other signature")
+	FAnimExtractContext(float InCurrentTime, bool InbExtractRootMotion = false, FDeltaTimeRecord InDeltaTimeRecord = {}, bool InbLooping = false)
+		: CurrentTime((double)InCurrentTime)
+		, bExtractRootMotion(InbExtractRootMotion)
+		, DeltaTimeRecord(InDeltaTimeRecord)
+		, bLooping(InbLooping)
+		, PoseCurves()
+		, BonesRequired()
+	{
+	}
 
-	FAnimExtractContext(float InCurrentTime = 0.f, bool InbExtractRootMotion = false, FDeltaTimeRecord InDeltaTimeRecord = {}, bool InbLooping = false)
+	FAnimExtractContext(double InCurrentTime = 0.0, bool InbExtractRootMotion = false, FDeltaTimeRecord InDeltaTimeRecord = {}, bool InbLooping = false)
 		: CurrentTime(InCurrentTime)
 		, bExtractRootMotion(InbExtractRootMotion)
 		, DeltaTimeRecord(InDeltaTimeRecord)
@@ -1120,7 +1131,7 @@ protected:
 
 public:
 	/** Return a list of unique marker names for blending compatibility */
-	virtual TArray<FName>* GetUniqueMarkerNames() { return NULL; }
+	virtual TArray<FName>* GetUniqueMarkerNames() { return nullptr; }
 
 	//~ Begin IInterface_AssetUserData Interface
 	virtual void AddAssetUserData(UAssetUserData* InUserData) override;
diff --git a/Engine/Source/Runtime/Engine/Classes/Animation/AnimationSettings.h b/Engine/Source/Runtime/Engine/Classes/Animation/AnimationSettings.h
--- a/Engine/Source/Runtime/Engine/Classes/Animation/AnimationSettings.h
+++ b/Engine/Source/Runtime/Engine/Classes/Animation/AnimationSettings.h
@@ -100,9 +100,15 @@ class ENGINE_API UAnimationSettings : public UDeveloperSettings
 	UPROPERTY(config, EditAnywhere, Category = Mirroring)
 	TArray<FMirrorFindReplaceExpression> MirrorFindReplaceExpressions;
 
+	/** Project specific default frame-rate used when (re)initializing any animation based data */
+	UPROPERTY(config, EditAnywhere, Category = AnimationData)
+	FFrameRate DefaultFrameRate;
 public:
 	static UAnimationSettings * Get() { return CastChecked<UAnimationSettings>(UAnimationSettings::StaticClass()->GetDefaultObject()); }
 
+	/** Returns the project specific default frame-rate */
+	const FFrameRate& GetDefaultFrameRate() const;
+
 #if WITH_EDITOR
 	virtual void PostEditChangeProperty(struct FPropertyChangedEvent& PropertyChangedEvent) override;
 #endif // WITH_EDITOR
diff --git a/Engine/Source/Runtime/Engine/Private/Animation/AnimComposite.cpp b/Engine/Source/Runtime/Engine/Private/Animation/AnimComposite.cpp
--- a/Engine/Source/Runtime/Engine/Private/Animation/AnimComposite.cpp
+++ b/Engine/Source/Runtime/Engine/Private/Animation/AnimComposite.cpp
@@ -144,10 +144,10 @@ bool UAnimComposite::ContainRecursive(TArray<UAnimCompositeBase*>& CurrentAccumu
 void UAnimComposite::SetCompositeLength(float InLength)
 {
 #if WITH_EDITOR		
-	Controller->SetPlayLength(InLength);
+	Controller->SetNumberOfFrames(DataModelInterface->GetFrameRate().AsFrameNumber(InLength));
 #else
 	PRAGMA_DISABLE_DEPRECATION_WARNINGS
-	SetSequenceLength(InLength);
+	SequenceLength = InLength;
 	PRAGMA_ENABLE_DEPRECATION_WARNINGS
 #endif	
 }
diff --git a/Engine/Source/Runtime/Engine/Private/Animation/AnimCompositeBase.cpp b/Engine/Source/Runtime/Engine/Private/Animation/AnimCompositeBase.cpp
--- a/Engine/Source/Runtime/Engine/Private/Animation/AnimCompositeBase.cpp
+++ b/Engine/Source/Runtime/Engine/Private/Animation/AnimCompositeBase.cpp
@@ -16,7 +16,7 @@
 // FAnimSegment
 ///////////////////////////////////////////////////////
 
-UAnimSequenceBase * FAnimSegment::GetAnimationData(float PositionInTrack, float& PositionInAnim) const
+UAnimSequenceBase* FAnimSegment::GetAnimationData(float PositionInTrack, float& PositionInAnim) const
 {
 	if( bValid && IsInRange(PositionInTrack) )
 	{
@@ -569,9 +569,11 @@ void FAnimTrack::GetAnimationPose(FAnimationPoseData& OutAnimationPoseData, cons
 		if (AnimSegment->bValid)
 		{
 			// Copy passed in Extraction Context, but override position and root motion parameters.
-			FAnimExtractContext SequenceExtractionContext(ExtractionContext);
-			if (const UAnimSequenceBase* const AnimRef = AnimSegment->GetAnimationData(ClampedTime, SequenceExtractionContext.CurrentTime))
+			float PositionInAnim = 0.f;
+			if (const UAnimSequenceBase* const AnimRef = AnimSegment->GetAnimationData(ClampedTime, PositionInAnim))
 			{
+				FAnimExtractContext SequenceExtractionContext(ExtractionContext);
+				SequenceExtractionContext.CurrentTime = static_cast<double>(PositionInAnim);
 				SequenceExtractionContext.DeltaTimeRecord.SetPrevious(
 					SequenceExtractionContext.CurrentTime - SequenceExtractionContext.DeltaTimeRecord.Delta);
 				SequenceExtractionContext.bExtractRootMotion &= AnimRef->HasRootMotion();
@@ -766,6 +768,13 @@ void UAnimCompositeBase::ExtractRootMotionFromTrack(const FAnimTrack &SlotAnimTr
 	}
 }
 
+FFrameRate UAnimCompositeBase::GetSamplingFrameRate() const
+{
+	// Allowing for 0.00001s precision in composite/montage length
+	static const FFrameRate CompositeFrameRate(100000, 1);
+	return CompositeFrameRate;
+}
+
 void UAnimCompositeBase::PostLoad()
 {
 	Super::PostLoad();
@@ -817,11 +826,17 @@ void FAnimSegment::UpdateCachedPlayLength()
 {
 	CachedPlayLength = 0.f;
 	PRAGMA_DISABLE_DEPRECATION_WARNINGS
-	const UAnimDataModel* DataModel = AnimReference ? AnimReference->GetDataModel() : nullptr;	
+	const IAnimationDataModel* DataModel = AnimReference ? AnimReference->GetDataModel() : nullptr;	
 	PRAGMA_ENABLE_DEPRECATION_WARNINGS
 	if(DataModel)
 	{
 		CachedPlayLength = DataModel->GetPlayLength();
 	}
 }
+
+void UAnimCompositeBase::PopulateWithExistingModel(TScriptInterface<IAnimationDataModel> ExistingDataModel)
+{
+	Super::PopulateWithExistingModel(ExistingDataModel);
+	Controller->SetFrameRate(GetSamplingFrameRate());
+}
 #endif // WITH_EDITOR
\ No newline at end of file
diff --git a/Engine/Source/Runtime/Engine/Private/Animation/AnimCompress.cpp b/Engine/Source/Runtime/Engine/Private/Animation/AnimCompress.cpp
--- a/Engine/Source/Runtime/Engine/Private/Animation/AnimCompress.cpp
+++ b/Engine/Source/Runtime/Engine/Private/Animation/AnimCompress.cpp
@@ -321,8 +321,7 @@ void UAnimCompress::BitwiseCompressAnimationTracks(
 			UE_LOG(LogAnimationCompression, Warning, TEXT("When compressing %s: no key-reduced data"), *CompressibleAnimData.Name);
 		}
 
-		AnimData.CompressedTrackOffsets.Empty(NumTracks * 4);
-		AnimData.CompressedTrackOffsets.AddUninitialized(NumTracks * 4);
+		AnimData.CompressedTrackOffsets.SetNumUninitialized(NumTracks * 4);
 
 		// just empty it since there is chance this can be 0
 		AnimData.CompressedScaleOffsets.Empty();
@@ -333,7 +332,8 @@ void UAnimCompress::BitwiseCompressAnimationTracks(
 			AnimData.CompressedScaleOffsets.AddUninitialized(NumTracks);
 		}
 
-		AnimData.CompressedByteStream.Empty();
+		const int32 MaxSize = CompressibleAnimData.RawAnimationData.Num() * CompressibleAnimData.NumberOfKeys * sizeof(FVector3f) + sizeof(FQuat4f) + sizeof(FVector3f);
+		AnimData.CompressedByteStream.Reset(MaxSize);
 
 		for (int32 TrackIndex = 0; TrackIndex < NumTracks; ++TrackIndex)
 		{
@@ -991,12 +991,12 @@ void UAnimCompress::SeparateRawDataIntoTracks(
 {
 	const int32 NumTracks = RawAnimData.Num();
 
-	OutTranslationData.Empty( NumTracks );
-	OutRotationData.Empty( NumTracks );
-	OutScaleData.Empty( NumTracks );
-	OutTranslationData.AddZeroed( NumTracks );
-	OutRotationData.AddZeroed( NumTracks );
-	OutScaleData.AddZeroed( NumTracks );
+	OutTranslationData.Reset( NumTracks );
+	OutRotationData.Reset( NumTracks );
+	OutScaleData.Reset( NumTracks );
+	OutTranslationData.SetNumZeroed( NumTracks );
+	OutRotationData.SetNumZeroed( NumTracks );
+	OutScaleData.SetNumZeroed( NumTracks );
 
 	// only compress scale if it has valid scale keys
 	bool bCompressScaleKeys = false;
@@ -1020,65 +1020,78 @@ void UAnimCompress::SeparateRawDataIntoTracks(
 		}
 
 		// Copy over position keys.
-		for ( int32 PosIndex = 0; PosIndex < RawTrack.PosKeys.Num(); ++PosIndex )
-		{
-			TranslationTrack.PosKeys.Add( RawTrack.PosKeys[PosIndex] );
-		}
+		TranslationTrack.PosKeys = RawTrack.PosKeys;
 
 		// Copy over rotation keys.
-		for ( int32 RotIndex = 0; RotIndex < RawTrack.RotKeys.Num(); ++RotIndex )
-		{
-			RotationTrack.RotKeys.Add( RawTrack.RotKeys[RotIndex] );
-		}
+		RotationTrack.RotKeys = RawTrack.RotKeys;
 
 		// Set times for the translation track.
-		if ( TranslationTrack.PosKeys.Num() > 1 )
+		const int32 NumPosKeys = TranslationTrack.PosKeys.Num();
+		TranslationTrack.Times.SetNumUninitialized(NumPosKeys);
+		if ( NumPosKeys > 1 )
 		{
-			const float PosFrameInterval = SequenceLength / static_cast<float>(TranslationTrack.PosKeys.Num()-1);
-			for ( int32 PosIndex = 0; PosIndex < TranslationTrack.PosKeys.Num(); ++PosIndex )
+			const float PosFrameInterval = SequenceLength / static_cast<float>(NumPosKeys-1);
+			for ( int32 PosIndex = 0; PosIndex < NumPosKeys; ++PosIndex )
 			{
-				TranslationTrack.Times.Add( PosIndex * PosFrameInterval );
+				TranslationTrack.Times[PosIndex] = PosIndex * PosFrameInterval;
 			}
 		}
 		else
 		{
-			TranslationTrack.Times.Add( 0.f );
+			TranslationTrack.Times[0] = 0.f;
 		}
 
 		// Set times for the rotation track.
-		if ( RotationTrack.RotKeys.Num() > 1 )
+		const int32 NumRotKeys = RotationTrack.RotKeys.Num();
+		RotationTrack.Times.SetNumUninitialized(NumRotKeys);
+		if ( NumRotKeys > 1 )
 		{
-			const float RotFrameInterval = SequenceLength / static_cast<float>(RotationTrack.RotKeys.Num()-1);
-			for ( int32 RotIndex = 0; RotIndex < RotationTrack.RotKeys.Num(); ++RotIndex )
+			// If # of keys match between translation and rotation, re-use the timing values
+			if (NumRotKeys == NumPosKeys)
 			{
-				RotationTrack.Times.Add( RotIndex * RotFrameInterval );
+				RotationTrack.Times = TranslationTrack.Times;
 			}
+			else
+			{
+				const float RotFrameInterval = SequenceLength / static_cast<float>(NumRotKeys-1);
+				for ( int32 RotIndex = 0; RotIndex < NumRotKeys; ++RotIndex )
+				{
+					RotationTrack.Times[RotIndex] = RotIndex * RotFrameInterval;
+				}
+			}
+		
 		}
 		else
 		{
-			RotationTrack.Times.Add( 0.f );
+			RotationTrack.Times[0] = 0.f;
 		}
 
 		if (bHasScale)
 		{
 			// Copy over scalekeys.
-			for ( int32 ScaleIndex = 0; ScaleIndex < RawTrack.ScaleKeys.Num(); ++ScaleIndex )
-			{
-				ScaleTrack.ScaleKeys.Add( RawTrack.ScaleKeys[ScaleIndex] );
-			}
-					
+			ScaleTrack.ScaleKeys = RawTrack.ScaleKeys;
 			// Set times for the rotation track.
-			if ( ScaleTrack.ScaleKeys.Num() > 1 )
+			const int32 NumScaleKeys = ScaleTrack.ScaleKeys.Num();
+			ScaleTrack.Times.SetNumUninitialized(NumScaleKeys);
+			if ( NumScaleKeys > 1 )
 			{
-				const float ScaleFrameInterval = SequenceLength / static_cast<float>(ScaleTrack.ScaleKeys.Num()-1);
-				for ( int32 ScaleIndex = 0; ScaleIndex < ScaleTrack.ScaleKeys.Num(); ++ScaleIndex )
+				// If # of keys match between translation and scale, re-use the timing values
+				if (NumScaleKeys == NumPosKeys)
+				{
+					ScaleTrack.Times = TranslationTrack.Times;
+				}
+				else
 				{
-					ScaleTrack.Times.Add( ScaleIndex * ScaleFrameInterval );
+					const float ScaleFrameInterval = SequenceLength / static_cast<float>(NumScaleKeys-1);
+					for ( int32 ScaleIndex = 0; ScaleIndex < NumScaleKeys; ++ScaleIndex )
+					{
+						ScaleTrack.Times[ScaleIndex] = ScaleIndex * ScaleFrameInterval;
+					}
 				}
 			}
 			else
 			{
-				ScaleTrack.Times.Add( 0.f );
+				ScaleTrack.Times[0] = 0.f;
 			}
 		}
 
diff --git a/Engine/Source/Runtime/Engine/Private/Animation/AnimCompress_PerTrackCompression.cpp b/Engine/Source/Runtime/Engine/Private/Animation/AnimCompress_PerTrackCompression.cpp
--- a/Engine/Source/Runtime/Engine/Private/Animation/AnimCompress_PerTrackCompression.cpp
+++ b/Engine/Source/Runtime/Engine/Private/Animation/AnimCompress_PerTrackCompression.cpp
@@ -65,7 +65,7 @@ protected:
 		bReallyNeedsFrameTable = false;
 		ActualCompressionMode = ACF_None;
 		ActualKeyFlags = 0;
-		CompressedBytes.Empty();
+		CompressedBytes.Reset();
 	}
 
 	/**
@@ -827,6 +827,9 @@ public:
 		Reset();
 		bReallyNeedsFrameTable = Params.bIncludeKeyTable && (TranslationData.PosKeys.Num() > 1) && (TranslationData.PosKeys.Num() < Params.NumberOfKeys);
 
+		const int32 MaxSize = TranslationData.Times.Num() * sizeof(FVector3f);
+		CompressedBytes.Reset(MaxSize);
+
 		switch (InCompressionType)
 		{
 		case ACF_Identity:
@@ -861,6 +864,9 @@ public:
 		Reset();
 		bReallyNeedsFrameTable = Params.bIncludeKeyTable && (RotationData.RotKeys.Num() > 1) && (RotationData.RotKeys.Num() < Params.NumberOfKeys);
 
+		const int32 MaxSize = RotationData.Times.Num() * sizeof(FQuat4f);
+		CompressedBytes.Reset(MaxSize);
+
 		switch (InCompressionType)
 		{
 		case ACF_Identity:
diff --git a/Engine/Source/Runtime/Engine/Private/Animation/AnimCompress_RemoveLinearKeys.cpp b/Engine/Source/Runtime/Engine/Private/Animation/AnimCompress_RemoveLinearKeys.cpp
--- a/Engine/Source/Runtime/Engine/Private/Animation/AnimCompress_RemoveLinearKeys.cpp
+++ b/Engine/Source/Runtime/Engine/Private/Animation/AnimCompress_RemoveLinearKeys.cpp
@@ -113,8 +113,8 @@ void FilterLinearKeysTemplate(
 	// generate new arrays we will fill with the final keys
 	TArray<KeyType> NewKeys;
 	TArray<float> NewTimes;
-	NewKeys.Empty(KeyCount);
-	NewTimes.Empty(KeyCount);
+	NewKeys.Reset(KeyCount);
+	NewTimes.Reset(KeyCount);
 
 	// Only bother doing anything if we have some keys!
 	if(KeyCount > 0)
@@ -122,8 +122,8 @@ void FilterLinearKeysTemplate(
 		int32 LowKey = 0;
 		int32 HighKey = KeyCount-1;
 
-		TArray<uint32> KnownParentTimes;
-		KnownParentTimes.Empty(KeyCount);
+		TArray<bool> KnownParentTimes;
+		KnownParentTimes.SetNumUninitialized(KeyCount);
 		const int32 ParentKeyCount = ParentTimes ? ParentTimes->Num() : 0;
 		for (int32 TimeIndex = 0, ParentTimeIndex = 0; TimeIndex < KeyCount; TimeIndex++)
 		{
@@ -132,15 +132,15 @@ void FilterLinearKeysTemplate(
 				ParentTimeIndex++;
 			}
 
-			KnownParentTimes.Add((ParentTimeIndex < ParentKeyCount) && (Times[TimeIndex] == (*ParentTimes)[ParentTimeIndex]));
+			KnownParentTimes[TimeIndex] = (ParentTimeIndex < ParentKeyCount) && (Times[TimeIndex] == (*ParentTimes)[ParentTimeIndex]);
 		}
 
 		TArray<FTransform> CachedInvRawBases;
-		CachedInvRawBases.Empty(KeyCount);
+		CachedInvRawBases.SetNumUninitialized(KeyCount);
 		for (int32 FrameIndex = 0; FrameIndex < KeyCount; ++FrameIndex)
 		{
 			const FTransform& RawBase = RawWorldBones[(BoneIndex*NumFrames) + FrameIndex];
-			CachedInvRawBases.Add(RawBase.Inverse());
+			CachedInvRawBases[FrameIndex] = RawBase.Inverse();
 		}
 		
 		// copy the low key (this one is a given)
@@ -304,20 +304,19 @@ void UAnimCompress_RemoveLinearKeys::UpdateWorldBoneTransformTable(
 {
 	const FBoneData& Bone		= CompressibleAnimData.BoneData[BoneIndex];
 	const int32 NumKeys		= CompressibleAnimData.NumberOfKeys;
-	const float SequenceLength	= CompressibleAnimData.SequenceLength;
 	const int32 FrameStart		= (BoneIndex*NumKeys);
 	const int32 TrackIndex = FAnimationUtils::GetAnimTrackIndexForSkeletonBone(BoneIndex, CompressibleAnimData.TrackToSkeletonMapTable);
 	
 	check(OutputWorldBones.Num() >= (FrameStart+NumKeys));
 
-	const float TimePerFrame = SequenceLength / (float)(NumKeys-1);
+	const FFrameRate& SamplingRate = CompressibleAnimData.SampledFrameRate;
 
 	if( TrackIndex != INDEX_NONE )
 	{
 		// get the local-space bone transforms using the animation solver
 		for ( int32 KeyIndex = 0; KeyIndex < NumKeys; ++KeyIndex )
 		{
-			float Time = (float)KeyIndex * TimePerFrame;
+			const double Time = SamplingRate.AsSeconds(KeyIndex);
 			FTransform LocalAtom;
 
 			FAnimationUtils::ExtractTransformFromCompressionData(CompressibleAnimData, OutCompressedData, Time, TrackIndex, UseRaw, LocalAtom);
@@ -417,9 +416,10 @@ void UAnimCompress_RemoveLinearKeys::UpdateBoneAtomList(
 	TArray<FTransform>& BoneAtoms)
 {
 	BoneAtoms.Reset(NumFrames);
+	const FFrameRate& SamplingRate = CompressibleAnimData.SampledFrameRate;
 	for ( int32 FrameIndex = 0; FrameIndex < NumFrames; ++FrameIndex )
 	{
-		float Time = (float)FrameIndex * TimePerFrame;
+		const double Time = SamplingRate.AsSeconds(FrameIndex);
 		FTransform LocalAtom;
 		FAnimationUtils::ExtractTransformFromCompressionData(CompressibleAnimData, OutCompressedData, Time, TrackIndex, false, LocalAtom);
 
@@ -465,11 +465,10 @@ void UAnimCompress_RemoveLinearKeys::ConvertFromRelativeSpace(FCompressibleAnimD
 		// make sure scale key exists
 		if (RawTrack.ScaleKeys.Num() > 0)
 		{
-			const FVector3f DefaultScale(1.f);
-			const FVector3f& RefBoneScale = (BasePoseTrack.ScaleKeys.Num() > 0)? BasePoseTrack.ScaleKeys[0] : DefaultScale;
+			const FVector3f& RefBoneScale = (BasePoseTrack.ScaleKeys.Num() > 0)? BasePoseTrack.ScaleKeys[0] : FVector3f::OneVector;
 			for (int32 ScaleIndex = 0; ScaleIndex < RawTrack.ScaleKeys.Num(); ++ScaleIndex)
 			{
-				RawTrack.ScaleKeys[ScaleIndex] = RefBoneScale * (DefaultScale + RawTrack.ScaleKeys[ScaleIndex]);
+				RawTrack.ScaleKeys[ScaleIndex] = RefBoneScale * (FVector3f::OneVector + RawTrack.ScaleKeys[ScaleIndex]);
 			}
 		}
 	}
@@ -515,7 +514,8 @@ void UAnimCompress_RemoveLinearKeys::ConvertToRelativeSpace(FCompressibleAnimDat
 		// scale key
 		if (RawTrack.ScaleKeys.Num() > 0)
 		{
-			const FVector3f InvRefBoneScale = (FVector3f)FTransform::GetSafeScaleReciprocal((FVector)BasePoseTrack.ScaleKeys[0]);
+        	const FVector3f& RefBoneScale = (BasePoseTrack.ScaleKeys.Num() > 0)? BasePoseTrack.ScaleKeys[0] : FVector3f::OneVector;
+			const FVector3f InvRefBoneScale = (FVector3f)FTransform::GetSafeScaleReciprocal((FVector)RefBoneScale);
 
 			// transform scale keys.
 			for (int32 ScaleIndex = 0; ScaleIndex < RawTrack.ScaleKeys.Num(); ++ScaleIndex)
@@ -563,7 +563,8 @@ void UAnimCompress_RemoveLinearKeys::ConvertToRelativeSpace(
 		// scale key
 		if (ScaleData.Num() > 0)
 		{
-			const FVector3f InvRefBoneScale = (FVector3f)FTransform::GetSafeScaleReciprocal((FVector)BasePoseTrack.ScaleKeys[0]);
+        	const FVector3f& RefBoneScale = (BasePoseTrack.ScaleKeys.Num() > 0)? BasePoseTrack.ScaleKeys[0] : FVector3f::OneVector;
+			const FVector3f InvRefBoneScale = (FVector3f)FTransform::GetSafeScaleReciprocal((FVector)RefBoneScale);
 
 			// convert the new scale tracks to additive space
 			FScaleTrack& ScaleTrack = ScaleData[TrackIndex];
@@ -879,7 +880,7 @@ void UAnimCompress_RemoveLinearKeys::ProcessAnimationTracks(
 			}
 
 			// Determine if a guidance track should be used to aid in choosing keys to retain
-			TArray<float>* GuidanceTrack = NULL;
+			TArray<float>* GuidanceTrack = nullptr;
 			float GuidanceScale = 1.0f;
 			if (GuideTrackIndex != INDEX_NONE)
 			{
diff --git a/Engine/Source/Runtime/Engine/Private/Animation/AnimCompressionDerivedData.cpp b/Engine/Source/Runtime/Engine/Private/Animation/AnimCompressionDerivedData.cpp
--- a/Engine/Source/Runtime/Engine/Private/Animation/AnimCompressionDerivedData.cpp
+++ b/Engine/Source/Runtime/Engine/Private/Animation/AnimCompressionDerivedData.cpp
@@ -37,7 +37,7 @@ const TCHAR* FDerivedDataAnimationCompression::GetVersionString() const
 	// This is a version string that mimics the old versioning scheme. If you
 	// want to bump this version, generate a new guid using VS->Tools->Create GUID and
 	// return it here. Ex.
-	return TEXT("3A3F9FBB7A4047278CABE35820CC44D7");
+	return TEXT("E29090F49C2348068642D5BEE76253B3");
 }
 
 FString FDerivedDataAnimationCompression::GetDebugContextString() const
diff --git a/Engine/Source/Runtime/Engine/Private/Animation/AnimCompressionTypes.cpp b/Engine/Source/Runtime/Engine/Private/Animation/AnimCompressionTypes.cpp
--- a/Engine/Source/Runtime/Engine/Private/Animation/AnimCompressionTypes.cpp
+++ b/Engine/Source/Runtime/Engine/Private/Animation/AnimCompressionTypes.cpp
@@ -16,11 +16,6 @@
 #include "Misc/FileHelper.h"
 #include "Serialization/JsonWriter.h"
 
-CSV_DECLARE_CATEGORY_MODULE_EXTERN(ENGINE_API, Animation);
-
-DECLARE_CYCLE_STAT(TEXT("Build Anim Track Pairs"), STAT_BuildAnimTrackPairs, STATGROUP_Anim);
-DECLARE_CYCLE_STAT(TEXT("Extract Pose From Anim Data"), STAT_ExtractPoseFromAnimData, STATGROUP_Anim);
-
 static FString GCompressionJsonOutput;
 static FAutoConsoleVariableRef CVarCompressionJsonOutput(
 	TEXT("a.Compression.CompressibleDataOutput"),
@@ -337,6 +332,8 @@ FCompressibleAnimData::FCompressibleAnimData(class UAnimSequence* InSeq, const b
 	/* Always get the resampled data to start off with */
 	const TArray<FBoneAnimationTrack>& ResampledTrackData = InSeq->GetResampledTrackData();
 
+	SampledFrameRate = InSeq->GetSamplingFrameRate();
+
 	TArray<FName> OriginalTrackNames;
 	RawAnimationData.Empty(ResampledTrackData.Num());
 	TrackToSkeletonMapTable.Empty(ResampledTrackData.Num());
@@ -350,8 +347,8 @@ FCompressibleAnimData::FCompressibleAnimData(class UAnimSequence* InSeq, const b
 		OriginalTrackNames.Add(AnimTrack.Name);
 	}
 
-	const bool bBakeAdditive = InSeq->CanBakeAdditive();
-	if (bBakeAdditive)
+	const bool bIsAdditiveAnimation = InSeq->CanBakeAdditive();
+	if (bIsAdditiveAnimation)
 	{
 		InSeq->BakeOutAdditiveIntoRawData(RawAnimationData, OriginalTrackNames, TrackToSkeletonMapTable, RawFloatCurves, AdditiveBaseAnimationData);
 	}
@@ -384,7 +381,7 @@ FCompressibleAnimData::FCompressibleAnimData(class UAnimSequence* InSeq, const b
 		return Keys.Num() == 0 || (Keys.Num() == 1 && Keys[0].Equals(IdentityValue));
 	};
 
-	auto IsRawTrackValidForRemoval = [IsKeyArrayValidForRemoval](const FRawAnimSequenceTrack& Track) -> bool
+	auto IsRawTrackZeroAdditive = [IsKeyArrayValidForRemoval](const FRawAnimSequenceTrack& Track) -> bool
 	{
 		return IsKeyArrayValidForRemoval(Track.PosKeys, FVector3f::ZeroVector) &&
 			IsKeyArrayValidForRemoval(Track.RotKeys, FQuat4f::Identity) &&
@@ -395,21 +392,59 @@ FCompressibleAnimData::FCompressibleAnimData(class UAnimSequence* InSeq, const b
     // And for additive animations remove any track deemed not to add any additive animation (identity rotation and zero-vector translation and scale)
 	// Note on (TrackIndex > 0) below : deliberately stop before track 0, compression code doesn't like getting a completely empty animation
 	TArray<FName> FinalTrackNames;
-	FinalTrackNames.Reserve(ResampledTrackData.Num());	
-	for (int32 TrackIndex = RawAnimationData.Num() - 1; TrackIndex > 0; --TrackIndex)
+
+	// Ensure we have any tracks at all	to begin with
+	if (OriginalTrackNames.Num())
 	{
-		const FRawAnimSequenceTrack& Track = RawAnimationData[TrackIndex];
-		// Try find correct bone index
-		const int32 BoneIndex = RefSkeleton.FindBoneIndex(OriginalTrackNames[TrackIndex]);
-		if ((bBakeAdditive && IsRawTrackValidForRemoval(Track)) || BoneIndex == INDEX_NONE)
+		TArray<FRawAnimSequenceTrack> TempRawAnimationData;
+		TArray<FRawAnimSequenceTrack> TempAdditiveBaseAnimationData;
+		TArray<FTrackToSkeletonMap> TempTrackToSkeletonMapTable;
+		TempTrackToSkeletonMapTable.Reserve(OriginalTrackNames.Num());
+		TempRawAnimationData.Reserve(OriginalTrackNames.Num());
+		FinalTrackNames.Reserve(ResampledTrackData.Num());
+		TempAdditiveBaseAnimationData.Reserve(AdditiveBaseAnimationData.Num() ? AdditiveBaseAnimationData.Num() : 0);
+
+		// Include root bone track
+		FinalTrackNames.Add(OriginalTrackNames[0]);
+		TempTrackToSkeletonMapTable.Add(TrackToSkeletonMapTable[0]);
+		TempRawAnimationData.Add(RawAnimationData[0]);
+		if (AdditiveBaseAnimationData.Num())
 		{
-			RawAnimationData.RemoveAtSwap(TrackIndex, 1, false);
-			TrackToSkeletonMapTable.RemoveAtSwap(TrackIndex, 1, false);
+			TempAdditiveBaseAnimationData.Add(AdditiveBaseAnimationData[0]);
+		}
+
+		const int32 NumTracks = RawAnimationData.Num();
+		for (int32 TrackIndex = 1; TrackIndex < NumTracks; ++TrackIndex)
+		{
+			const FRawAnimSequenceTrack& Track = RawAnimationData[TrackIndex];
+			// Try find correct bone index
+			const int32 BoneIndex = RefSkeleton.FindBoneIndex(OriginalTrackNames[TrackIndex]);
+
+			const bool bValidBoneIndex = BoneIndex != INDEX_NONE;
+			const bool bValidAdditiveTrack = !IsRawTrackZeroAdditive(Track);
+
+			// Only include track if it contains valid (additive) data and its name corresponds to a bone on the skeleton
+			if ((!bIsAdditiveAnimation || bValidAdditiveTrack) && bValidBoneIndex)
+			{
+				FinalTrackNames.Add(OriginalTrackNames[TrackIndex]);
+				TempTrackToSkeletonMapTable.Add(TrackToSkeletonMapTable[TrackIndex]);
+				TempRawAnimationData.Add(RawAnimationData[TrackIndex]);
+
+				if (AdditiveBaseAnimationData.Num())
+				{
+					TempAdditiveBaseAnimationData.Add(AdditiveBaseAnimationData[TrackIndex]);
+				}
+			}
+		}
+
+		// Swap out maintained track data
+		Swap(RawAnimationData, TempRawAnimationData);
+		Swap(TrackToSkeletonMapTable, TempTrackToSkeletonMapTable);
+
+		if (AdditiveBaseAnimationData.Num())
+		{
+			Swap(AdditiveBaseAnimationData, TempAdditiveBaseAnimationData);
 		}
-		else
-        {
-			FinalTrackNames.Add(OriginalTrackNames[TrackIndex]);
-        }
 	}
 
 	// Find or add curve names on skeleton
@@ -459,12 +494,6 @@ FCompressibleAnimData::FCompressibleAnimData(class UAnimSequence* InSeq, const b
 
 	if (GCompressionJsonOutput.Len())
 	{
-		// Insert root-bone name
-		if (RawAnimationData.Num())
-		{			
-			FinalTrackNames.Insert(OriginalTrackNames[0], 0);
-		}
-		
 		WriteCompressionDataToJSON(MakeArrayView(OriginalTrackNames), MakeArrayView(RawAnimationData), MakeArrayView(FinalTrackNames));	
 	}
 #endif
@@ -910,7 +939,7 @@ void FCompressedAnimSequence::SerializeCompressedData(FArchive& Ar, bool bDDCDat
 			FMemoryWriter MemoryWriter(TempSerialized, true);
 			MemoryWriter.SetByteSwapping(Ar.ForceByteSwapping());
 
-			check(RotationCodec != NULL);
+			check(RotationCodec != nullptr);
 
 			FMemoryReader MemoryReader(TempSerialized, true);
 			MemoryReader.SetByteSwapping(Ar.ForceByteSwapping());
@@ -1004,269 +1033,21 @@ void FCompressedAnimSequence::ClearCompressedCurveData()
 	CurveCompressionCodec = nullptr;
 }
 
-struct FGetBonePoseScratchArea : public TThreadSingleton<FGetBonePoseScratchArea>
-{
-	BoneTrackArray RotationScalePairs;
-	BoneTrackArray TranslationPairs;
-	BoneTrackArray AnimScaleRetargetingPairs;
-	BoneTrackArray AnimRelativeRetargetingPairs;
-	BoneTrackArray OrientAndScaleRetargetingPairs;
-};
-
 void DecompressPose(FCompactPose& OutPose, const FCompressedAnimSequence& CompressedData, const FAnimExtractContext& ExtractionContext, USkeleton* SourceSkeleton, float SequenceLength, EAnimInterpolationType Interpolation, bool bIsBakedAdditive, FName RetargetSource, FName SourceName, const FRootMotionReset& RootMotionReset)
 {
 	const TArray<FTransform>& RetargetTransforms = SourceSkeleton->GetRefLocalPoses(RetargetSource);
-	DecompressPose(OutPose, CompressedData, ExtractionContext, SourceSkeleton, SequenceLength, Interpolation, bIsBakedAdditive, RetargetTransforms, SourceName, RootMotionReset);
+	PRAGMA_DISABLE_DEPRECATION_WARNINGS
+	FAnimSequenceDecompressionContext DecompressionContext(SequenceLength, Interpolation, SourceName, *CompressedData.CompressedDataStructure.Get(), SourceSkeleton, bIsBakedAdditive);
+	PRAGMA_ENABLE_DEPRECATION_WARNINGS
+	UE::Anim::Decompression::DecompressPose(OutPose, CompressedData, ExtractionContext, DecompressionContext, RetargetTransforms, RootMotionReset);
 }
 
 void DecompressPose(FCompactPose& OutPose, const FCompressedAnimSequence& CompressedData, const FAnimExtractContext& ExtractionContext, USkeleton* SourceSkeleton, float SequenceLength, EAnimInterpolationType Interpolation, bool bIsBakedAdditive, const TArray<FTransform>& RetargetTransforms, FName SourceName, const FRootMotionReset& RootMotionReset)
 {
-	const FBoneContainer& RequiredBones = OutPose.GetBoneContainer();
-	const int32 NumTracks = CompressedData.CompressedTrackToSkeletonMapTable.Num();
-
-	const USkeleton* TargetSkeleton = RequiredBones.GetSkeletonAsset();
-	const FSkeletonRemapping* SkeletonRemapping = TargetSkeleton->GetSkeletonRemapping(SourceSkeleton);
-
-	BoneTrackArray& RotationScalePairs = FGetBonePoseScratchArea::Get().RotationScalePairs;
-	BoneTrackArray& TranslationPairs = FGetBonePoseScratchArea::Get().TranslationPairs;
-	BoneTrackArray& AnimScaleRetargetingPairs = FGetBonePoseScratchArea::Get().AnimScaleRetargetingPairs;
-	BoneTrackArray& AnimRelativeRetargetingPairs = FGetBonePoseScratchArea::Get().AnimRelativeRetargetingPairs;
-	BoneTrackArray& OrientAndScaleRetargetingPairs = FGetBonePoseScratchArea::Get().OrientAndScaleRetargetingPairs;
-
-	// build a list of desired bones
-	RotationScalePairs.Reset();
-	TranslationPairs.Reset();
-	AnimScaleRetargetingPairs.Reset();
-	AnimRelativeRetargetingPairs.Reset();
-	OrientAndScaleRetargetingPairs.Reset();
-
-	// Optimization: assuming first index is root bone. That should always be the case in Skeletons.
-	checkSlow((RequiredBones.GetMeshPoseIndexFromSkeletonPoseIndex(FSkeletonPoseBoneIndex(0)) == FMeshPoseBoneIndex(0)));
-	// this is not guaranteed for AnimSequences though... If Root is not animated, Track will not exist.
-	const bool bFirstTrackIsRootBone = (CompressedData.GetSkeletonIndexFromTrackIndex(0) == 0);
-
-	{
-		SCOPE_CYCLE_COUNTER(STAT_BuildAnimTrackPairs);
-
-		// Handle root bone separately if it is track 0. so we start w/ Index 1.
-		for (int32 TrackIndex = (bFirstTrackIsRootBone ? 1 : 0); TrackIndex < NumTracks; TrackIndex++)
-		{
-			const int32 SourceSkeletonBoneIndex = CompressedData.GetSkeletonIndexFromTrackIndex(TrackIndex);
-			const int32 TargetSkeletonBoneIndex = (SkeletonRemapping) ? SkeletonRemapping->GetTargetSkeletonBoneIndex(SourceSkeletonBoneIndex) : SourceSkeletonBoneIndex;
-
-			if (TargetSkeletonBoneIndex != INDEX_NONE)
-			{
-				const FCompactPoseBoneIndex BoneIndex = RequiredBones.GetCompactPoseIndexFromSkeletonIndex(TargetSkeletonBoneIndex);
-				//Nasty, we break our type safety, code in the lower levels should be adjusted for this
-				const int32 CompactPoseBoneIndex = BoneIndex.GetInt();
-				if (CompactPoseBoneIndex != INDEX_NONE)
-				{
-					RotationScalePairs.Add(BoneTrackPair(CompactPoseBoneIndex, TrackIndex));
-
-					// Skip extracting translation component for EBoneTranslationRetargetingMode::Skeleton.
-					switch (TargetSkeleton->GetBoneTranslationRetargetingMode(TargetSkeletonBoneIndex))
-					{
-					case EBoneTranslationRetargetingMode::Animation:
-						TranslationPairs.Add(BoneTrackPair(CompactPoseBoneIndex, TrackIndex));
-						break;
-					case EBoneTranslationRetargetingMode::AnimationScaled:
-						TranslationPairs.Add(BoneTrackPair(CompactPoseBoneIndex, TrackIndex));
-						AnimScaleRetargetingPairs.Add(BoneTrackPair(CompactPoseBoneIndex, SourceSkeletonBoneIndex));
-						break;
-					case EBoneTranslationRetargetingMode::AnimationRelative:
-						TranslationPairs.Add(BoneTrackPair(CompactPoseBoneIndex, TrackIndex));
-
-						// With baked additives, we can skip 'AnimationRelative' tracks, as the relative transform gets canceled out.
-						// (A1 + Rel) - (A2 + Rel) = A1 - A2.
-						if (!bIsBakedAdditive)
-						{
-							AnimRelativeRetargetingPairs.Add(BoneTrackPair(CompactPoseBoneIndex, SourceSkeletonBoneIndex));
-						}
-						break;
-					case EBoneTranslationRetargetingMode::OrientAndScale:
-						TranslationPairs.Add(BoneTrackPair(CompactPoseBoneIndex, TrackIndex));
-
-						// Additives remain additives, they're not retargeted.
-						if (!bIsBakedAdditive)
-						{
-							OrientAndScaleRetargetingPairs.Add(BoneTrackPair(CompactPoseBoneIndex, SourceSkeletonBoneIndex));
-						}
-						break;
-					}
-				}
-			}
-		}
-	}
-
-	{
-		SCOPE_CYCLE_COUNTER(STAT_ExtractPoseFromAnimData);
-		CSV_SCOPED_TIMING_STAT(Animation, ExtractPoseFromAnimData);
-		CSV_CUSTOM_STAT(Animation, NumberOfExtractedAnimations, 1, ECsvCustomStatOp::Accumulate);
-
-		FAnimSequenceDecompressionContext EvalDecompContext(SequenceLength, Interpolation, SourceName, *CompressedData.CompressedDataStructure);
-		EvalDecompContext.Seek(ExtractionContext.CurrentTime);
-
-		// Handle Root Bone separately
-		if (bFirstTrackIsRootBone)
-		{
-			const int32 TrackIndex = 0;
-			FCompactPoseBoneIndex RootBone(0);
-			FTransform& RootAtom = OutPose[RootBone];
-
-			CompressedData.BoneCompressionCodec->DecompressBone(EvalDecompContext, TrackIndex, RootAtom);
-
-			// Retarget the root onto the target skeleton (correcting for differences in rest poses)
-			if (SkeletonRemapping)
-			{
-				const int32 TargetSkeletonBoneIndex = 0;
-
-				if (bIsBakedAdditive)
-				{
-					RootAtom.SetRotation(SkeletonRemapping->RetargetAdditiveRotationToTargetSkeleton(TargetSkeletonBoneIndex, RootAtom.GetRotation()));
-					if (TargetSkeleton->GetBoneTranslationRetargetingMode(TargetSkeletonBoneIndex) != EBoneTranslationRetargetingMode::Skeleton)
-					{
-						RootAtom.SetTranslation(SkeletonRemapping->RetargetAdditiveTranslationToTargetSkeleton(TargetSkeletonBoneIndex, RootAtom.GetTranslation()));
-					}
-				}
-				else
-				{
-					RootAtom.SetRotation(SkeletonRemapping->RetargetBoneRotationToTargetSkeleton(TargetSkeletonBoneIndex, RootAtom.GetRotation()));
-					if (TargetSkeleton->GetBoneTranslationRetargetingMode(TargetSkeletonBoneIndex) != EBoneTranslationRetargetingMode::Skeleton)
-					{
-						RootAtom.SetTranslation(SkeletonRemapping->RetargetBoneTranslationToTargetSkeleton(TargetSkeletonBoneIndex, RootAtom.GetTranslation()));
-					}
-				}
-			}
-
-			// @laurent - we should look into splitting rotation and translation tracks, so we don't have to process translation twice.
-			FAnimationRuntime::RetargetBoneTransform(SourceSkeleton, SourceName, RetargetTransforms, RootAtom, 0, RootBone, RequiredBones, bIsBakedAdditive);
-		}
-
-		if (RotationScalePairs.Num() > 0)
-		{
-			// get the remaining bone atoms
-			TArrayView<FTransform> OutPoseBones = OutPose.GetMutableBones();
-			CompressedData.BoneCompressionCodec->DecompressPose(EvalDecompContext, RotationScalePairs, TranslationPairs, RotationScalePairs, OutPoseBones);
-		}
-	}
-
-	// Retarget the pose onto the target skeleton (correcting for differences in rest poses)
-	if (SkeletonRemapping)
-	{
-		if (bIsBakedAdditive)
-		{
-			for (FCompactPoseBoneIndex BoneIndex(bFirstTrackIsRootBone ? 1 : 0); BoneIndex < OutPose.GetNumBones(); ++BoneIndex)
-			{
-				const int32 TargetSkeletonBoneIndex = RequiredBones.GetSkeletonIndex(BoneIndex);
-				OutPose[BoneIndex].SetRotation(SkeletonRemapping->RetargetAdditiveRotationToTargetSkeleton(TargetSkeletonBoneIndex, OutPose[BoneIndex].GetRotation()));
-				if (TargetSkeleton->GetBoneTranslationRetargetingMode(TargetSkeletonBoneIndex) != EBoneTranslationRetargetingMode::Skeleton)
-				{
-					OutPose[BoneIndex].SetTranslation(SkeletonRemapping->RetargetAdditiveTranslationToTargetSkeleton(TargetSkeletonBoneIndex, OutPose[BoneIndex].GetTranslation()));
-				}
-			}
-		}
-		else
-		{
-			for (FCompactPoseBoneIndex BoneIndex(bFirstTrackIsRootBone ? 1 : 0); BoneIndex < OutPose.GetNumBones(); ++BoneIndex)
-			{
-				const int32 TargetSkeletonBoneIndex = RequiredBones.GetSkeletonIndex(BoneIndex);
-				OutPose[BoneIndex].SetRotation(SkeletonRemapping->RetargetBoneRotationToTargetSkeleton(TargetSkeletonBoneIndex, OutPose[BoneIndex].GetRotation()));
-				if (TargetSkeleton->GetBoneTranslationRetargetingMode(TargetSkeletonBoneIndex) != EBoneTranslationRetargetingMode::Skeleton)
-				{
-					OutPose[BoneIndex].SetTranslation(SkeletonRemapping->RetargetBoneTranslationToTargetSkeleton(TargetSkeletonBoneIndex, OutPose[BoneIndex].GetTranslation()));
-				}
-			}
-		}
-	}
-
-	// Once pose has been extracted, snap root bone back to first frame if we are extracting root motion.
-	if ((ExtractionContext.bExtractRootMotion && RootMotionReset.bEnableRootMotion) || RootMotionReset.bForceRootLock)
-	{
-		RootMotionReset.ResetRootBoneForRootMotion(OutPose[FCompactPoseBoneIndex(0)], RequiredBones);
-	}
-
-	// Anim Scale Retargeting
-	int32 const NumBonesToScaleRetarget = AnimScaleRetargetingPairs.Num();
-	if (NumBonesToScaleRetarget > 0)
-	{
-		TArray<FTransform> const& AuthoredOnRefSkeleton = RetargetTransforms;
-
-		for (const BoneTrackPair& BonePair : AnimScaleRetargetingPairs)
-		{
-			const FCompactPoseBoneIndex BoneIndex(BonePair.AtomIndex); //Nasty, we break our type safety, code in the lower levels should be adjusted for this
-			const int32 SourceSkeletonBoneIndex = BonePair.TrackIndex;
-
-			// @todo - precache that in FBoneContainer when we have SkeletonIndex->TrackIndex mapping. So we can just apply scale right away.
-			float const SourceTranslationLength = AuthoredOnRefSkeleton[SourceSkeletonBoneIndex].GetTranslation().Size();
-			if (SourceTranslationLength > UE_KINDA_SMALL_NUMBER)
-			{
-				float const TargetTranslationLength = RequiredBones.GetRefPoseTransform(BoneIndex).GetTranslation().Size();
-				OutPose[BoneIndex].ScaleTranslation(TargetTranslationLength / SourceTranslationLength);
-			}
-		}
-	}
-
-	// Anim Relative Retargeting
-	int32 const NumBonesToRelativeRetarget = AnimRelativeRetargetingPairs.Num();
-	if (NumBonesToRelativeRetarget > 0)
-	{
-		TArray<FTransform> const& AuthoredOnRefSkeleton = RetargetTransforms;
-
-		for (const BoneTrackPair& BonePair : AnimRelativeRetargetingPairs)
-		{
-			const FCompactPoseBoneIndex BoneIndex(BonePair.AtomIndex); //Nasty, we break our type safety, code in the lower levels should be adjusted for this
-			const int32 SourceSkeletonBoneIndex = BonePair.TrackIndex;
-
-			const FTransform& RefPoseTransform = RequiredBones.GetRefPoseTransform(BoneIndex);
-
-			// Remap the base pose onto the target skeleton so that we are working entirely in target space
-			FTransform BaseTransform = AuthoredOnRefSkeleton[SourceSkeletonBoneIndex];
-			if (SkeletonRemapping)
-			{
-				const int32 TargetSkeletonBoneIndex = SkeletonRemapping->GetTargetSkeletonBoneIndex(SourceSkeletonBoneIndex);
-				BaseTransform = SkeletonRemapping->RetargetBoneTransformToTargetSkeleton(TargetSkeletonBoneIndex, BaseTransform);
-			}
-
-			// Apply the retargeting as if it were an additive difference between the current skeleton and the retarget skeleton. 
-			OutPose[BoneIndex].SetRotation(OutPose[BoneIndex].GetRotation() * BaseTransform.GetRotation().Inverse() * RefPoseTransform.GetRotation());
-			OutPose[BoneIndex].SetTranslation(OutPose[BoneIndex].GetTranslation() + (RefPoseTransform.GetTranslation() - BaseTransform.GetTranslation()));
-			OutPose[BoneIndex].SetScale3D(OutPose[BoneIndex].GetScale3D() * (RefPoseTransform.GetScale3D() * BaseTransform.GetSafeScaleReciprocal(BaseTransform.GetScale3D())));
-			OutPose[BoneIndex].NormalizeRotation();
-		}
-	}
-
-	// Translation 'Orient and Scale' Translation Retargeting
-	const int32 NumBonesToOrientAndScaleRetarget = OrientAndScaleRetargetingPairs.Num();
-	if (NumBonesToOrientAndScaleRetarget > 0)
-	{
-		const FRetargetSourceCachedData& RetargetSourceCachedData = RequiredBones.GetRetargetSourceCachedData(SourceName, RetargetTransforms);
-		const TArray<FOrientAndScaleRetargetingCachedData>& OrientAndScaleDataArray = RetargetSourceCachedData.OrientAndScaleData;
-		const TArray<int32>& CompactPoseIndexToOrientAndScaleIndex = RetargetSourceCachedData.CompactPoseIndexToOrientAndScaleIndex;
-
-		// If we have any cached retargeting data.
-		if ((OrientAndScaleDataArray.Num() > 0) && (CompactPoseIndexToOrientAndScaleIndex.Num() == RequiredBones.GetCompactPoseNumBones()))
-		{
-			for (int32 Index = 0; Index < NumBonesToOrientAndScaleRetarget; Index++)
-			{
-				const BoneTrackPair& BonePair = OrientAndScaleRetargetingPairs[Index];
-				const FCompactPoseBoneIndex CompactPoseBoneIndex(BonePair.AtomIndex);
-				const int32 OrientAndScaleIndex = CompactPoseIndexToOrientAndScaleIndex[CompactPoseBoneIndex.GetInt()];
-				if (OrientAndScaleIndex != INDEX_NONE)
-				{
-					const FOrientAndScaleRetargetingCachedData& OrientAndScaleData = OrientAndScaleDataArray[OrientAndScaleIndex];
-					FTransform& BoneTransform = OutPose[CompactPoseBoneIndex];
-					const FVector AnimatedTranslation = BoneTransform.GetTranslation();
-
-					// If Translation is not animated, we can just copy the TargetTranslation. No retargeting needs to be done.
-					const FVector NewTranslation = (AnimatedTranslation - OrientAndScaleData.SourceTranslation).IsNearlyZero(BONE_TRANS_RT_ORIENT_AND_SCALE_PRECISION) ?
-						OrientAndScaleData.TargetTranslation :
-						OrientAndScaleData.TranslationDeltaOrient.RotateVector(AnimatedTranslation) * OrientAndScaleData.TranslationScale;
-
-					BoneTransform.SetTranslation(NewTranslation);
-				}
-			}
-		}
-	}
+	PRAGMA_DISABLE_DEPRECATION_WARNINGS
+	FAnimSequenceDecompressionContext DecompressionContext(SequenceLength, Interpolation, SourceName, *CompressedData.CompressedDataStructure.Get(), SourceSkeleton, bIsBakedAdditive);
+	PRAGMA_ENABLE_DEPRECATION_WARNINGS
+	UE::Anim::Decompression::DecompressPose(OutPose, CompressedData, ExtractionContext, DecompressionContext, RetargetTransforms, RootMotionReset);
 }
 
 FArchive& operator<<(FArchive& Ar, FCompressedOffsetData& D)
diff --git a/Engine/Source/Runtime/Engine/Private/Animation/AnimCurveCompressionCodec_CompressedRichCurve.cpp b/Engine/Source/Runtime/Engine/Private/Animation/AnimCurveCompressionCodec_CompressedRichCurve.cpp
--- a/Engine/Source/Runtime/Engine/Private/Animation/AnimCurveCompressionCodec_CompressedRichCurve.cpp
+++ b/Engine/Source/Runtime/Engine/Private/Animation/AnimCurveCompressionCodec_CompressedRichCurve.cpp
@@ -45,8 +45,8 @@ bool UAnimCurveCompressionCodec_CompressedRichCurve::Compress(const FCompressibl
 	int32 KeyDataOffset = 0;
 	KeyDataOffset += sizeof(FCurveDesc) * NumCurves;
 
-	const FAnimKeyHelper Helper(AnimSeq.SequenceLength, AnimSeq.NumberOfKeys);
-	const float SampleRate = UseAnimSequenceSampleRate ? Helper.KeysPerSecond() : ErrorSampleRate;
+	const FFrameRate& SamplingFrameRate = AnimSeq.SampledFrameRate;
+	const double SampleRate = AnimSeq.SampledFrameRate.AsDecimal();
 
 	TArray<uint8> KeyData;
 
@@ -78,15 +78,15 @@ bool UAnimCurveCompressionCodec_CompressedRichCurve::Compress(const FCompressibl
 			for (int32 KeyIndex = 0; KeyIndex < SamplesToTake; ++KeyIndex)
 			{
 				// Evaluated RawCurve and Compressed Curve
-				const float EvalTime = KeyIndex * Helper.TimePerKey();
-				const float RawValue = RawCurve.Eval(EvalTime);
+				const double EvalTime = SamplingFrameRate.AsSeconds(KeyIndex);
+				const float RawValue = Curve.FloatCurve.Eval(EvalTime);
 				const float CompressedValue = CompressedCurve.Eval(EvalTime);
 
 				const float AbsDelta = FMath::Abs(CompressedValue - RawValue);
-				if (AbsDelta > MaxCurveError)
+				if (!FMath::IsNearlyZero(AbsDelta, MaxCurveError + UE_KINDA_SMALL_NUMBER))
 				{
 					// Delta larger than tolerated error value
-					UE_LOG(LogAnimationCompression, Warning, TEXT("Curve %s: delta too large %f at %f"), *Curve.Name.DisplayName.ToString(), AbsDelta, EvalTime);
+					UE_LOG(LogAnimationCompression, Warning, TEXT("Curve %s: delta too large %f, between %f and %f, at %f"), *Curve.Name.DisplayName.ToString(), AbsDelta, RawValue, CompressedValue, EvalTime);
 				}
 			}
 		}
diff --git a/Engine/Source/Runtime/Engine/Private/Animation/AnimCurveCompressionCodec_UniformIndexable.cpp b/Engine/Source/Runtime/Engine/Private/Animation/AnimCurveCompressionCodec_UniformIndexable.cpp
--- a/Engine/Source/Runtime/Engine/Private/Animation/AnimCurveCompressionCodec_UniformIndexable.cpp
+++ b/Engine/Source/Runtime/Engine/Private/Animation/AnimCurveCompressionCodec_UniformIndexable.cpp
@@ -89,11 +89,7 @@ UAnimCurveCompressionCodec_UniformIndexable::UAnimCurveCompressionCodec_UniformI
 bool UAnimCurveCompressionCodec_UniformIndexable::Compress(const FCompressibleAnimData& AnimSeq, FAnimCurveCompressionResult& OutResult)
 {
 	const int32 NumCurves = AnimSeq.RawFloatCurves.Num();
-	const float Duration = AnimSeq.SequenceLength;
-
-	const FAnimKeyHelper Helper(AnimSeq.SequenceLength, AnimSeq.NumberOfKeys);
-	const float SampleRate = Helper.KeysPerSecond();
-	const int32 NumSamples = FMath::RoundToInt(Duration * SampleRate) + 1;
+	const int32 NumSamples = AnimSeq.NumberOfKeys;
 
 	int32 BufferSize = 0;
 	BufferSize += sizeof(int32);	// NumSamples
@@ -113,7 +109,7 @@ bool UAnimCurveCompressionCodec_UniformIndexable::Compress(const FCompressibleAn
 	BufferOffset += sizeof(float);
 
 	*NumSamplesPtr = NumSamples;
-	*SampleRatePtr = SampleRate;
+	*SampleRatePtr = AnimSeq.SampledFrameRate.AsDecimal();
 
 	if (NumCurves > 0 && NumSamples > 0)
 	{
@@ -123,12 +119,12 @@ bool UAnimCurveCompressionCodec_UniformIndexable::Compress(const FCompressibleAn
 		// Write out samples by curve so that we can more easily index curve keys
 
 		// Curve 0 Key 0, Curve 0 Key 1, Curve 0 Key N, Curve 1 Key 0, Curve 1 Key 1, Curve 1 Key N, Curve M Key 0, ...
-		const float InvSampleRate = 1.0f / SampleRate;
 		for (const FFloatCurve& Curve : AnimSeq.RawFloatCurves)
 		{
 			for (int32 SampleIndex = 0; SampleIndex < NumSamples; ++SampleIndex)
 			{
-				const float SampleTime = FMath::Clamp(SampleIndex * InvSampleRate, 0.0f, Duration);
+				const double FrameSampleTime = AnimSeq.SampledFrameRate.AsSeconds(SampleIndex);
+				const float SampleTime = FMath::Clamp(FrameSampleTime, 0.0f, AnimSeq.SequenceLength);
 				const float SampleValue = Curve.FloatCurve.Eval(SampleTime);
 
 				(*AnimatedSamplesPtr) = SampleValue;
diff --git a/Engine/Source/Runtime/Engine/Private/Animation/AnimCurveCompressionCodec_UniformlySampled.cpp b/Engine/Source/Runtime/Engine/Private/Animation/AnimCurveCompressionCodec_UniformlySampled.cpp
--- a/Engine/Source/Runtime/Engine/Private/Animation/AnimCurveCompressionCodec_UniformlySampled.cpp
+++ b/Engine/Source/Runtime/Engine/Private/Animation/AnimCurveCompressionCodec_UniformlySampled.cpp
@@ -23,9 +23,8 @@ bool UAnimCurveCompressionCodec_UniformlySampled::Compress(const FCompressibleAn
 	float SampleRate_;
 	if (UseAnimSequenceSampleRate)
 	{
-		const FAnimKeyHelper Helper(AnimSeq.SequenceLength, AnimSeq.NumberOfKeys);
-		SampleRate_ = Helper.KeysPerSecond();
-		NumSamples = FMath::RoundToInt(Duration * SampleRate_) + 1;
+		NumSamples = AnimSeq.NumberOfKeys;
+		SampleRate_ = AnimSeq.SampledFrameRate.AsDecimal();
 	}
 	else
 	{
diff --git a/Engine/Source/Runtime/Engine/Private/Animation/AnimDataModel.cpp b/Engine/Source/Runtime/Engine/Private/Animation/AnimDataModel.cpp
--- a/Engine/Source/Runtime/Engine/Private/Animation/AnimDataModel.cpp
+++ b/Engine/Source/Runtime/Engine/Private/Animation/AnimDataModel.cpp
@@ -1,44 +1,50 @@
 // Copyright Epic Games, Inc. All Rights Reserved.
 
 #include "Animation/AnimData/AnimDataModel.h"
+
+#include "AnimationRuntime.h"
+#include "AnimEncoding.h"
 #include "UObject/NameTypes.h"
 
 #include "Animation/AnimSequence.h"
 
 #include "Algo/Transform.h"
 #include "Algo/Accumulate.h"
+#include "Animation/AnimSequenceHelpers.h"
 #include "Animation/SmartName.h"
+#include "Animation/AttributesRuntime.h"
 #include "UObject/UE5MainStreamObjectVersion.h"
 #include "UObject/FortniteMainBranchObjectVersion.h"
 
+#if WITH_EDITOR
+#include "IAnimationDataControllerModule.h"
+#include "Modules/ModuleManager.h"
+#endif // WITH_EDITOR
+
+#define LOCTEXT_NAMESPACE "AnimDataModel"
+
 void UAnimDataModel::PostLoad()
 {
 	UObject::PostLoad();
 
-	if (GetLinkerCustomVersion(FUE5MainStreamObjectVersion::GUID) >= FUE5MainStreamObjectVersion::AnimationDataModelInterface_BackedOut &&
-		GetLinkerCustomVersion(FUE5MainStreamObjectVersion::GUID) < FUE5MainStreamObjectVersion::BackoutAnimationDataModelInterface)
-	{
-		UE_LOG(LogAnimation, Fatal, TEXT("This package was saved with a version that had to be backed out and is no longer able to be loaded."));
-	}
-
-
 	if (GetLinkerCustomVersion(FUE5MainStreamObjectVersion::GUID) < FUE5MainStreamObjectVersion::ForceUpdateAnimationAssetCurveTangents)
 	{
 		// Forcefully AutoSetTangents to fix-up any imported sequences pre the fix for flattening first/last key leave/arrive tangents
-		Notify(EAnimDataModelNotifyType::BracketOpened);
+		GetNotifier().Notify(EAnimDataModelNotifyType::BracketOpened);
 		for (FFloatCurve& FloatCurve : CurveData.FloatCurves)
 		{
 			FloatCurve.FloatCurve.AutoSetTangents();
 			FCurvePayload Payload;
 			Payload.Identifier = FAnimationCurveIdentifier(FloatCurve.Name.UID, ERawCurveTrackTypes::RCT_Float);
-			Notify(EAnimDataModelNotifyType::CurveChanged, Payload);
+			GetNotifier().Notify(EAnimDataModelNotifyType::CurveChanged, Payload);
 		}
-		Notify(EAnimDataModelNotifyType::BracketClosed);
+		GetNotifier().Notify(EAnimDataModelNotifyType::BracketClosed);
 	}
 
+	const bool bHasBoneTracks = BoneAnimationTracks.Num() > 0;
+
 	if (GetLinkerCustomVersion(FFortniteMainBranchObjectVersion::GUID) < FFortniteMainBranchObjectVersion::SingleFrameAndKeyAnimModel)
 	{
-		const bool bHasBoneTracks = BoneAnimationTracks.Num() > 0;
 		if (bHasBoneTracks)
 		{
 			// Number of keys was used directly rather than Max(Value,2), as a single _frame_ animation should always have two _keys_ 
@@ -59,17 +65,47 @@ void UAnimDataModel::PostLoad()
 					AddKey(BoneTrack.InternalTrackData.ScaleKeys);
 				}
 
-				Notify(EAnimDataModelNotifyType::TrackChanged);
+				GetNotifier().Notify(EAnimDataModelNotifyType::TrackChanged);
 			}
 		}
-	}	
+	}
+	
+	if (GetLinkerCustomVersion(FUE5MainStreamObjectVersion::GUID) < FUE5MainStreamObjectVersion::ReintroduceAnimationDataModelInterface)
+	{		
+		const double NumFramesAsSeconds = FrameRate.AsSeconds(NumberOfFrames);
+		PRAGMA_DISABLE_DEPRECATION_WARNINGS
+		const float CachedPlayLength = PlayLength;
+		if (static_cast<float>(NumFramesAsSeconds) != PlayLength)
+		{
+			if (bHasBoneTracks)
+			{
+				// Number of keys was rounded up from sequence length, so set play length from stored number of bone keys 
+				const int32 ActualNumKeys = bHasBoneTracks ? FMath::Min(BoneAnimationTracks[0].InternalTrackData.PosKeys.Num(), NumberOfKeys) : NumberOfKeys;
+				if (ActualNumKeys != NumberOfKeys)
+				{
+					PlayLength = FrameRate.AsSeconds(ActualNumKeys - 1);
+					NumberOfFrames = ActualNumKeys - 1;
+					NumberOfKeys = ActualNumKeys;
+
+					IAnimationDataController::ReportObjectWarningf(this, LOCTEXT("AnimDataModelPlayLengthAlteredToBones", "Playable length {0} has been updated according to number of frames {1} and frame-rate {2}, new length: {3}"), FText::AsNumber(CachedPlayLength), FText::AsNumber(GetNumberOfFrames()), FrameRate.ToPrettyText(), FText::AsNumber(PlayLength));
+				}
+			}
+			else
+			{
+				// No way to determine 'correct' length so set current PlayLength value to # of frames * frame-rate
+				PlayLength = FrameRate.AsSeconds(NumberOfFrames);
+			}
+			GetNotifier().Notify(EAnimDataModelNotifyType::Populated);
+		}
+		PRAGMA_ENABLE_DEPRECATION_WARNINGS
+	}
 }
 
 void UAnimDataModel::PostDuplicate(bool bDuplicateForPIE)
 {
 	UObject::PostDuplicate(bDuplicateForPIE);
 
-	Notify(EAnimDataModelNotifyType::Populated);
+	GetNotifier().Notify(EAnimDataModelNotifyType::Populated);
 }
 
 const TArray<FBoneAnimationTrack>& UAnimDataModel::GetBoneAnimationTracks() const
@@ -85,7 +121,7 @@ const FBoneAnimationTrack& UAnimDataModel::GetBoneTrackByIndex(int32 TrackIndex)
 
 const FBoneAnimationTrack& UAnimDataModel::GetBoneTrackByName(FName TrackName) const
 {
-	const FBoneAnimationTrack* TrackPtr = BoneAnimationTracks.FindByPredicate([TrackName](FBoneAnimationTrack& Track)
+	const FBoneAnimationTrack* TrackPtr = BoneAnimationTracks.FindByPredicate([TrackName](const FBoneAnimationTrack& Track)
 	{
 		return Track.Name == TrackName;
 	});
@@ -97,7 +133,7 @@ const FBoneAnimationTrack& UAnimDataModel::GetBoneTrackByName(FName TrackName) c
 
 const FBoneAnimationTrack* UAnimDataModel::FindBoneTrackByIndex(int32 BoneIndex) const
 {
-	const FBoneAnimationTrack* TrackPtr = BoneAnimationTracks.FindByPredicate([BoneIndex](FBoneAnimationTrack& Track)
+	const FBoneAnimationTrack* TrackPtr = BoneAnimationTracks.FindByPredicate([BoneIndex](const FBoneAnimationTrack& Track)
 	{
 		return Track.BoneTreeIndex == BoneIndex;
 	});
@@ -107,7 +143,7 @@ const FBoneAnimationTrack* UAnimDataModel::FindBoneTrackByIndex(int32 BoneIndex)
 
 int32 UAnimDataModel::GetBoneTrackIndex(const FBoneAnimationTrack& Track) const
 {
-	return BoneAnimationTracks.IndexOfByPredicate([&Track](const FBoneAnimationTrack SearchTrack)
+	return BoneAnimationTracks.IndexOfByPredicate([&Track](const FBoneAnimationTrack& SearchTrack)
 	{
 		return SearchTrack.Name == Track.Name;
 	});
@@ -123,9 +159,9 @@ int32 UAnimDataModel::GetBoneTrackIndexByName(FName TrackName) const
 	return INDEX_NONE;
 }
 
-float UAnimDataModel::GetPlayLength() const
+double UAnimDataModel::GetPlayLength() const
 {
-	return PlayLength;
+	return FrameRate.AsSeconds(NumberOfFrames);
 }
 
 int32 UAnimDataModel::GetNumberOfFrames() const
@@ -138,7 +174,7 @@ int32 UAnimDataModel::GetNumberOfKeys() const
 	return NumberOfKeys;
 }
 
-const FFrameRate& UAnimDataModel::GetFrameRate() const
+FFrameRate UAnimDataModel::GetFrameRate() const
 {
 	return FrameRate;
 }
@@ -148,19 +184,20 @@ bool UAnimDataModel::IsValidBoneTrackIndex(int32 TrackIndex) const
 	return BoneAnimationTracks.IsValidIndex(TrackIndex);
 }
 
-const int32 UAnimDataModel::GetNumBoneTracks() const
+int32 UAnimDataModel::GetNumBoneTracks() const
 {
 	return BoneAnimationTracks.Num();
 }
 
 UAnimSequence* UAnimDataModel::GetAnimationSequence() const
 {
-	return Cast<UAnimSequence>(GetOuter());
+	UObject* AnimObject = static_cast<UObject*>(FindObjectWithOuter(GetOuter(), UAnimSequence::StaticClass()));
+	return CastChecked<UAnimSequence>(AnimObject ? AnimObject : GetOuter());
 }
 
 void UAnimDataModel::GetBoneTrackNames(TArray<FName>& OutNames) const
 {
-	Algo::Transform(BoneAnimationTracks, OutNames, [](FBoneAnimationTrack Track)
+	Algo::Transform(BoneAnimationTracks, OutNames, [](const FBoneAnimationTrack& Track)
 	{
 		return Track.Name; 
 	});
@@ -171,6 +208,16 @@ const FAnimationCurveData& UAnimDataModel::GetCurveData() const
 	return CurveData;
 }
 
+IAnimationDataModel::FModelNotifier& UAnimDataModel::GetNotifier()
+{
+	if (!Notifier)
+	{
+		Notifier.Reset(new IAnimationDataModel::FModelNotifier(this));
+	}
+	
+	return *Notifier.Get();
+}
+
 int32 UAnimDataModel::GetNumberOfTransformCurves() const
 {
 	return CurveData.TransformCurves.Num();
@@ -306,7 +353,6 @@ int32 UAnimDataModel::GetNumberOfAttributesForBoneIndex(const int32 BoneIndex) c
 
 void UAnimDataModel::GetAttributesForBone(const FName& BoneName, TArray<const FAnimatedBoneAttribute*>& OutBoneAttributes) const
 {
-	// Sum up total number of attributes with provided bone name
 	Algo::TransformIf(AnimatedBoneAttributes, OutBoneAttributes, [BoneName](const FAnimatedBoneAttribute& Attribute) -> bool
 	{
 		return Attribute.Identifier.GetBoneName() == BoneName;
@@ -385,7 +431,7 @@ FGuid UAnimDataModel::GenerateGuid() const
     {
     	UpdateWithData(Curve.Name.UID);
 
-    	auto UpdateWithComponent = [&Sha, &UpdateWithFloatCurve](const FVectorCurve& VectorCurve)
+    	auto UpdateWithComponent = [&UpdateWithFloatCurve](const FVectorCurve& VectorCurve)
     	{
     		for (int32 ChannelIndex = 0; ChannelIndex < 3; ++ChannelIndex)
     		{
@@ -413,6 +459,252 @@ FGuid UAnimDataModel::GenerateGuid() const
 	return Guid;
 }
 
+TScriptInterface<IAnimationDataController> UAnimDataModel::GetController()
+{
+	TScriptInterface<IAnimationDataController> Controller = nullptr;
+#if WITH_EDITOR
+	IAnimationDataControllerModule& ControllerModule = FModuleManager::Get().GetModuleChecked<IAnimationDataControllerModule>("AnimationDataController");
+	Controller = ControllerModule.GetController();
+	Controller->SetModel(this);
+#endif // WITH_EDITOR
+
+	return Controller;
+}
+
+#if WITH_EDITOR
+FTransform ExtractTransformForKeyIndex(int32 Key, const FRawAnimSequenceTrack& TrackToExtract)
+{
+	static const FVector DefaultScale3D = FVector(1.f);
+	const bool bHasScaleKey = TrackToExtract.ScaleKeys.Num() > 0;
+
+	const int32 PosKeyIndex = FMath::Min(Key, TrackToExtract.PosKeys.Num() - 1);
+	const int32 RotKeyIndex = FMath::Min(Key, TrackToExtract.RotKeys.Num() - 1);
+	if (bHasScaleKey)
+	{
+		const int32 ScaleKeyIndex = FMath::Min(Key, TrackToExtract.ScaleKeys.Num() - 1);
+		return FTransform(FQuat(TrackToExtract.RotKeys[RotKeyIndex]), FVector(TrackToExtract.PosKeys[PosKeyIndex]), FVector(TrackToExtract.ScaleKeys[ScaleKeyIndex]));
+	}
+	else
+	{
+		return FTransform(FQuat(TrackToExtract.RotKeys[RotKeyIndex]), FVector(TrackToExtract.PosKeys[PosKeyIndex]), DefaultScale3D);
+	}
+}
+
+template<bool bInterpolateT>
+void ExtractPose(const TArray<FBoneAnimationTrack>& BoneAnimationTracks, const TMap<FName, const FTransformCurve*>& OverrideBoneTransforms, FCompactPose& InOutPose, const int32 KeyIndex1, const int32 KeyIndex2, float Alpha, float TimePerKey, UE::Anim::Retargeting::FRetargetingScope& RetargetingScope)
+{
+	const int32 NumAnimationTracks = BoneAnimationTracks.Num();
+	const FBoneContainer& RequiredBones = InOutPose.GetBoneContainer();
+
+	TArray<FVirtualBoneCompactPoseData>& VBCompactPoseData = UE::Anim::FBuildRawPoseScratchArea::Get().VirtualBoneCompactPoseData;
+	VBCompactPoseData = RequiredBones.GetVirtualBoneCompactPoseData();
+
+	FCompactPose Key2Pose;
+	if (bInterpolateT)
+	{
+		Key2Pose.CopyBonesFrom(InOutPose);
+	}
+
+	for (const FBoneAnimationTrack& AnimationTrack : BoneAnimationTracks)
+	{
+		const int32 SkeletonBoneIndex = AnimationTrack.BoneTreeIndex;
+		// not sure it's safe to assume that SkeletonBoneIndex can never be INDEX_NONE
+		if ((SkeletonBoneIndex != INDEX_NONE) && (SkeletonBoneIndex < MAX_BONES))
+		{
+			const FCompactPoseBoneIndex PoseBoneIndex = RequiredBones.GetCompactPoseIndexFromSkeletonIndex(SkeletonBoneIndex);
+			if (PoseBoneIndex != INDEX_NONE)
+			{
+				for (int32 Idx = 0; Idx < VBCompactPoseData.Num(); ++Idx)
+				{
+					FVirtualBoneCompactPoseData& VB = VBCompactPoseData[Idx];
+					if (PoseBoneIndex == VB.VBIndex)
+					{
+						// Remove this bone as we have written data for it (false so we dont resize allocation)
+						VBCompactPoseData.RemoveAtSwap(Idx, 1, false);
+						break; //Modified TArray so must break here
+					}
+				}
+				
+				// extract animation
+				const FRawAnimSequenceTrack& TrackToExtract = AnimationTrack.InternalTrackData;
+				auto OverrideTransform = OverrideBoneTransforms.Find(AnimationTrack.Name);
+				{
+					// Bail out (with rather wacky data) if data is empty for some reason.
+					if (TrackToExtract.PosKeys.Num() == 0 || TrackToExtract.RotKeys.Num() == 0)
+					{
+						InOutPose[PoseBoneIndex].SetIdentity();
+
+						if (bInterpolateT)
+						{
+							Key2Pose[PoseBoneIndex].SetIdentity();
+						}
+					}
+					else
+					{
+						InOutPose[PoseBoneIndex] = ExtractTransformForKeyIndex(KeyIndex1, TrackToExtract);
+
+						if (bInterpolateT)
+						{
+							Key2Pose[PoseBoneIndex] = ExtractTransformForKeyIndex(KeyIndex2, TrackToExtract);
+						}
+					}
+				}
+
+				if (OverrideTransform)
+				{
+					{
+						const FTransform PoseOneAdditive = (*OverrideTransform)->Evaluate(KeyIndex1 * TimePerKey, 1.f);
+						const FTransform PoseOneLocalTransform = InOutPose[PoseBoneIndex];
+						InOutPose[PoseBoneIndex].SetRotation(PoseOneLocalTransform.GetRotation() * PoseOneAdditive.GetRotation());
+						InOutPose[PoseBoneIndex].SetTranslation(PoseOneLocalTransform.TransformPosition(PoseOneAdditive.GetTranslation()));
+						InOutPose[PoseBoneIndex].SetScale3D(PoseOneLocalTransform.GetScale3D() * PoseOneAdditive.GetScale3D());
+					}
+					
+
+					if (bInterpolateT)
+					{
+						const FTransform PoseTwoAdditive = (*OverrideTransform)->Evaluate(KeyIndex2 * TimePerKey, 1.f);
+						const FTransform PoseTwoLocalTransform = Key2Pose[PoseBoneIndex];						
+						Key2Pose[PoseBoneIndex].SetRotation(PoseTwoLocalTransform.GetRotation() * PoseTwoAdditive.GetRotation());
+						Key2Pose[PoseBoneIndex].SetTranslation(PoseTwoLocalTransform.TransformPosition(PoseTwoAdditive.GetTranslation()));
+						Key2Pose[PoseBoneIndex].SetScale3D(PoseTwoLocalTransform.GetScale3D() * PoseTwoAdditive.GetScale3D());
+					}
+				}
+				RetargetingScope.AddTrackedBone(PoseBoneIndex, SkeletonBoneIndex);
+			}
+		}
+	}
+
+	//Build Virtual Bones
+	if (VBCompactPoseData.Num() > 0)
+	{
+		FCSPose<FCompactPose> CSPose1;
+		CSPose1.InitPose(InOutPose);
+
+		FCSPose<FCompactPose> CSPose2;
+		if (bInterpolateT)
+		{
+			CSPose2.InitPose(Key2Pose);
+		}
+
+		for (FVirtualBoneCompactPoseData& VB : VBCompactPoseData)
+		{
+			FTransform Source = CSPose1.GetComponentSpaceTransform(VB.SourceIndex);
+			FTransform Target = CSPose1.GetComponentSpaceTransform(VB.TargetIndex);
+			InOutPose[VB.VBIndex] = Target.GetRelativeTransform(Source);
+
+			if (bInterpolateT)
+			{
+				FTransform Source2 = CSPose2.GetComponentSpaceTransform(VB.SourceIndex);
+				FTransform Target2 = CSPose2.GetComponentSpaceTransform(VB.TargetIndex);
+				Key2Pose[VB.VBIndex] = Target2.GetRelativeTransform(Source2);
+			}
+		}
+	}
+
+	if (bInterpolateT)
+	{
+		for (FCompactPoseBoneIndex BoneIndex : InOutPose.ForEachBoneIndex())
+		{
+			InOutPose[BoneIndex].Blend(InOutPose[BoneIndex], Key2Pose[BoneIndex], Alpha);
+		}
+	}
+}
+
+void UAnimDataModel::Evaluate(FAnimationPoseData& InOutPoseData, const UE::Anim::DataModel::FEvaluationContext& EvaluationContext) const
+{
+	const float SequenceLength = GetPlayLength();
+
+	const double Time = EvaluationContext.SampleFrameRate.AsSeconds(EvaluationContext.SampleTime);
+	const EAnimInterpolationType& InterpolationType = EvaluationContext.InterpolationType;
+	FCompactPose& OutPose = InOutPoseData.GetPose();
+	
+	// Generate keys to interpolate between
+	int32 KeyIndex1, KeyIndex2;
+	float Alpha;
+	FAnimationRuntime::GetKeyIndicesFromTime(KeyIndex1, KeyIndex2, Alpha, Time, FrameRate, NumberOfKeys);
+
+	if (InterpolationType == EAnimInterpolationType::Step)
+	{
+		// Force stepping between keys
+		Alpha = 0.f;
+	}
+
+	bool bShouldInterpolate = true;
+
+	if (Alpha < UE_KINDA_SMALL_NUMBER)
+	{
+		Alpha = 0.f;
+		bShouldInterpolate = false;
+	}
+	else if (Alpha > 1.f - UE_KINDA_SMALL_NUMBER)
+	{
+		bShouldInterpolate = false;
+		KeyIndex1 = KeyIndex2;
+	}
+
+	// Evaluate animation float curve data
+	FBlendedCurve& Curves = InOutPoseData.GetCurve();
+	if (Curves.NumValidCurveCount > 0)
+	{
+		// evaluate the curve data at the CurrentTime and add to Instance
+		for (const FFloatCurve& Curve : CurveData.FloatCurves)
+		{
+			if (Curves.IsEnabled(Curve.Name.UID))
+			{
+				float Value = Curve.Evaluate(Time);
+				Curves.Set(Curve.Name.UID, Value);
+			}
+		}
+	}
+
+	const double TimePerFrame = FrameRate.AsInterval();
+	TMap<FName, const FTransformCurve*> ActiveCurves;
+
+	if (!OutPose.GetBoneContainer().ShouldUseSourceData())
+	{
+		for (const FTransformCurve& Curve : GetTransformCurves())
+		{
+			// if disabled, do not handle
+			if (Curve.GetCurveTypeFlag(AACF_Disabled))
+			{
+				continue;
+			}
+
+			// Add or retrieve curve
+			const FName& CurveName = Curve.Name.DisplayName;
+
+			// note we're not checking Curve.GetCurveTypeFlags() yet
+			ActiveCurves.FindOrAdd(CurveName, &Curve);
+		}
+	}
+
+	{		
+		UE::Anim::Retargeting::FRetargetingScope RetargetingScope(OutPose, EvaluationContext);
+		if (bShouldInterpolate)
+		{		
+			ExtractPose<true>(GetBoneAnimationTracks(), ActiveCurves, OutPose, KeyIndex1, KeyIndex2, Alpha, TimePerFrame, RetargetingScope);
+		}
+		else
+		{
+			ExtractPose<false>(GetBoneAnimationTracks(), ActiveCurves, OutPose, KeyIndex1, KeyIndex2, Alpha, TimePerFrame, RetargetingScope);
+		}
+	}
+	
+
+	const FBoneContainer& RequiredBones = InOutPoseData.GetPose().GetBoneContainer();
+	for (const FAnimatedBoneAttribute& Attribute : AnimatedBoneAttributes)
+	{
+		const FCompactPoseBoneIndex PoseBoneIndex = RequiredBones.GetCompactPoseIndexFromSkeletonIndex(Attribute.Identifier.GetBoneIndex());
+		// Only add attribute if the bone its tied to exists in the currently evaluated set of bones
+		if(PoseBoneIndex.IsValid())
+		{
+			UE::Anim::Attributes::GetAttributeValue(InOutPoseData.GetAttributes(), PoseBoneIndex, Attribute, EvaluationContext.SampleFrameRate.AsSeconds(EvaluationContext.SampleTime));
+		}
+	}	
+}
+#endif // WITH_EDITOR
+
 FRichCurve* UAnimDataModel::GetMutableRichCurve(const FAnimationCurveIdentifier& CurveIdentifier)
 {
 	FRichCurve* RichCurve = nullptr;
@@ -451,7 +743,7 @@ FRichCurve* UAnimDataModel::GetMutableRichCurve(const FAnimationCurveIdentifier&
 
 FBoneAnimationTrack* UAnimDataModel::FindMutableBoneTrackByName(FName Name)
 {
-	return BoneAnimationTracks.FindByPredicate([Name](FBoneAnimationTrack& Track)
+	return BoneAnimationTracks.FindByPredicate([Name](const FBoneAnimationTrack& Track)
 	{
 		return Track.Name == Name;
 	});
@@ -459,7 +751,7 @@ FBoneAnimationTrack* UAnimDataModel::FindMutableBoneTrackByName(FName Name)
 
 const FBoneAnimationTrack* UAnimDataModel::FindBoneTrackByName(FName Name) const
 {
-	return BoneAnimationTracks.FindByPredicate([Name](FBoneAnimationTrack& Track)
+	return BoneAnimationTracks.FindByPredicate([Name](const FBoneAnimationTrack& Track)
 	{
 		return Track.Name == Name;
 	});
@@ -467,7 +759,7 @@ const FBoneAnimationTrack* UAnimDataModel::FindBoneTrackByName(FName Name) const
 
 FBoneAnimationTrack& UAnimDataModel::GetMutableBoneTrackByName(FName Name)
 {
-	FBoneAnimationTrack* TrackPtr = BoneAnimationTracks.FindByPredicate([Name](FBoneAnimationTrack& Track)
+	FBoneAnimationTrack* TrackPtr = BoneAnimationTracks.FindByPredicate([Name](const FBoneAnimationTrack& Track)
 	{
 		return Track.Name == Name;
 	});
@@ -516,3 +808,4 @@ FAnimCurveBase* UAnimDataModel::FindMutableCurveById(const FAnimationCurveIdenti
 	return nullptr;
 }
 
+#undef LOCTEXT_NAMESPACE // "AnimDataModel"
diff --git a/Engine/Source/Runtime/Engine/Private/Animation/AnimInstanceProxy.cpp b/Engine/Source/Runtime/Engine/Private/Animation/AnimInstanceProxy.cpp
--- a/Engine/Source/Runtime/Engine/Private/Animation/AnimInstanceProxy.cpp
+++ b/Engine/Source/Runtime/Engine/Private/Animation/AnimInstanceProxy.cpp
@@ -1365,7 +1365,7 @@ void FAnimInstanceProxy::EvaluateAnimationNode(FPoseContext& Output)
 
 void FAnimInstanceProxy::EvaluateAnimationNode_WithRoot(FPoseContext& Output, FAnimNode_Base* InRootNode)
 {
-	if (InRootNode != NULL)
+	if (InRootNode != nullptr)
 	{
 		ANIM_MT_SCOPE_CYCLE_COUNTER(EvaluateAnimGraph, !IsInGameThread());
 		if(InRootNode == RootNode)
@@ -1486,7 +1486,7 @@ void FAnimInstanceProxy::SlotEvaluatePoseWithBlendProfiles(const FName& SlotNode
 			NewPose.Curve.InitFrom(RequiredBones);
 
 			// Extract pose from Track.
-			FAnimExtractContext ExtractionContext(EvalState.MontagePosition, Montage->HasRootMotion() && RootMotionMode != ERootMotionMode::NoRootMotionExtraction, EvalState.DeltaTimeRecord);
+			FAnimExtractContext ExtractionContext(static_cast<double>(EvalState.MontagePosition), Montage->HasRootMotion() && RootMotionMode != ERootMotionMode::NoRootMotionExtraction, EvalState.DeltaTimeRecord);
 			FAnimationPoseData NewAnimationPoseData(NewPose);
 			AnimTrack->GetAnimationPose(NewAnimationPoseData, ExtractionContext);
 
@@ -1825,7 +1825,7 @@ void FAnimInstanceProxy::SlotEvaluatePose(const FName& SlotNodeName, const FAnim
 			NewPose.Curve.InitFrom(RequiredBones);
 
 			// Extract pose from Track
-			FAnimExtractContext ExtractionContext(EvalState.MontagePosition, Montage->HasRootMotion() && RootMotionMode != ERootMotionMode::NoRootMotionExtraction, EvalState.DeltaTimeRecord);
+			FAnimExtractContext ExtractionContext(static_cast<double>(EvalState.MontagePosition), Montage->HasRootMotion() && RootMotionMode != ERootMotionMode::NoRootMotionExtraction, EvalState.DeltaTimeRecord);
 
 			FAnimationPoseData NewAnimationPoseData(NewPose);
 			AnimTrack->GetAnimationPose(NewAnimationPoseData, ExtractionContext);
diff --git a/Engine/Source/Runtime/Engine/Private/Animation/AnimMontage.cpp b/Engine/Source/Runtime/Engine/Private/Animation/AnimMontage.cpp
--- a/Engine/Source/Runtime/Engine/Private/Animation/AnimMontage.cpp
+++ b/Engine/Source/Runtime/Engine/Private/Animation/AnimMontage.cpp
@@ -93,7 +93,7 @@ const FAnimTrack* UAnimMontage::GetAnimationData(FName InSlotName) const
 		}
 	}
 
-	return NULL;
+	return nullptr;
 }
 
 bool UAnimMontage::IsWithinPos(int32 FirstIndex, int32 SecondIndex, float CurrentTime) const
@@ -370,7 +370,7 @@ void UAnimMontage::PostLoad()
 		Track.ValidateSegmentTimes();
 
 		const float CurrentCalculatedLength = CalculateSequenceLength();
-		if(CurrentCalculatedLength != GetPlayLength())		
+		if(!FMath::IsNearlyEqual(CurrentCalculatedLength, GetPlayLength(), UE_KINDA_SMALL_NUMBER))		
 		{
 			UE_LOG(LogAnimMontage, Display, TEXT("UAnimMontage::PostLoad: The actual sequence length for %s does not match the length stored in the asset, please resave the asset."), *GetFullName());
 			SetCompositeLength(CurrentCalculatedLength);
@@ -419,7 +419,7 @@ void UAnimMontage::PostLoad()
 	}
 	// find preview base pose if it can
 #if WITH_EDITORONLY_DATA
-	if ( IsValidAdditive() && PreviewBasePose == NULL )
+	if ( IsValidAdditive() && PreviewBasePose == nullptr )
 	{
 		for (int32 I=0; I<SlotAnimTracks.Num(); ++I)
 		{
@@ -556,8 +556,8 @@ void UAnimMontage::ConvertBranchingPointsToAnimNotifies()
 #if WITH_EDITORONLY_DATA
 			NewEvent.TrackIndex = TrackIndex;
 #endif
-			NewEvent.Notify = NULL;
-			NewEvent.NotifyStateClass = NULL;
+			NewEvent.Notify = nullptr;
+			NewEvent.NotifyStateClass = nullptr;
 			NewEvent.bConvertedFromBranchingPoint = true;
 			NewEvent.MontageTickType = EMontageNotifyTickType::BranchingPoint;
 
@@ -1311,11 +1311,14 @@ bool UAnimMontage::ContainRecursive(TArray<UAnimCompositeBase*>& CurrentAccumula
 
 void UAnimMontage::SetCompositeLength(float InLength)
 {
-#if WITH_EDITOR	
-	Controller->SetPlayLength(InLength);
+#if WITH_EDITOR
+	const FFrameTime LengthInFrameTime = DataModelInterface->GetFrameRate().AsFrameTime(InLength);
+	PRAGMA_DISABLE_DEPRECATION_WARNINGS
+	Controller->SetNumberOfFrames(LengthInFrameTime.RoundToFrame());
+	PRAGMA_ENABLE_DEPRECATION_WARNINGS
 #else
 	PRAGMA_DISABLE_DEPRECATION_WARNINGS
-	SetSequenceLength(InLength);
+	SequenceLength = InLength;
 	PRAGMA_ENABLE_DEPRECATION_WARNINGS
 #endif	
 }
@@ -1325,11 +1328,11 @@ void UAnimMontage::SetCompositeLength(float InLength)
 /////////////////////////////////////////////////////////////////////////////////////////////
 
 FAnimMontageInstance::FAnimMontageInstance()
-	: Montage(NULL)
+	: Montage(nullptr)
 	, bPlaying(false)
 	, DefaultBlendTimeMultiplier(1.0f)
 	, bDidUseMarkerSyncThisTick(false)
-	, AnimInstance(NULL)
+	, AnimInstance(nullptr)
 	, InstanceID(INDEX_NONE)
 	, Position(0.f)
 	, PlayRate(1.f)
@@ -1341,13 +1344,13 @@ FAnimMontageInstance::FAnimMontageInstance()
 	, ActiveBlendProfile(nullptr)
 	, ActiveBlendProfileMode(EBlendProfileMode::TimeFactor)
 	, DisableRootMotionCount(0)
-	, MontageSyncLeader(NULL)
+	, MontageSyncLeader(nullptr)
 	, MontageSyncUpdateFrameCounter(INDEX_NONE)
 {
 }
 
 FAnimMontageInstance::FAnimMontageInstance(UAnimInstance * InAnimInstance)
-	: Montage(NULL)
+	: Montage(nullptr)
 	, bPlaying(false)
 	, DefaultBlendTimeMultiplier(1.0f)
 	, bDidUseMarkerSyncThisTick(false)
@@ -1364,7 +1367,7 @@ FAnimMontageInstance::FAnimMontageInstance(UAnimInstance * InAnimInstance)
 	, ActiveBlendProfile(nullptr)
 	, ActiveBlendProfileMode(EBlendProfileMode::TimeFactor)
 	, DisableRootMotionCount(0)
-	, MontageSyncLeader(NULL)
+	, MontageSyncLeader(nullptr)
 	, MontageSyncUpdateFrameCounter(INDEX_NONE)
 {
 }
@@ -1590,7 +1593,7 @@ void FAnimMontageInstance::Terminate()
 {
 	SCOPE_CYCLE_COUNTER(STAT_AnimMontageInstance_Terminate);
 
-	if (Montage == NULL)
+	if (Montage == nullptr)
 	{
 		return;
 	}
@@ -1629,7 +1632,7 @@ void FAnimMontageInstance::Terminate()
 	}
 
 	// clear Blend curve
-	Blend.SetCustomCurve(NULL);
+	Blend.SetCustomCurve(nullptr);
 	Blend.SetBlendOption(EAlphaBlendOption::Linear);
 
 	ActiveBlendProfile = nullptr;
@@ -1775,7 +1778,7 @@ void FAnimMontageInstance::MontageSync_StopLeading()
 		if (MontageSyncFollower)
 		{
 			ensure(MontageSyncFollower->MontageSyncLeader == this);
-			MontageSyncFollower->MontageSyncLeader = NULL;
+			MontageSyncFollower->MontageSyncLeader = nullptr;
 		}
 	}
 	MontageSyncFollowers.Empty();
@@ -1786,7 +1789,7 @@ void FAnimMontageInstance::MontageSync_StopFollowing()
 	if (MontageSyncLeader)
 	{
 		MontageSyncLeader->MontageSyncFollowers.RemoveSingleSwap(this);
-		MontageSyncLeader = NULL;
+		MontageSyncLeader = nullptr;
 	}
 }
 
@@ -2659,7 +2662,7 @@ void FAnimMontageInstance::BranchingPointEventHandler(const FBranchingPointMarke
 		// Must grab a reference on the stack in case "this" is deleted during iteration
 		TWeakObjectPtr<UAnimInstance> AnimInstanceLocal = AnimInstance;
 
-		FAnimNotifyEvent* NotifyEvent = (BranchingPointMarker->NotifyIndex < Montage->Notifies.Num()) ? &Montage->Notifies[BranchingPointMarker->NotifyIndex] : NULL;
+		FAnimNotifyEvent* NotifyEvent = (BranchingPointMarker->NotifyIndex < Montage->Notifies.Num()) ? &Montage->Notifies[BranchingPointMarker->NotifyIndex] : nullptr;
 		if (NotifyEvent)
 		{
 			// Handle backwards compatibility with older BranchingPoints.
@@ -2671,7 +2674,7 @@ void FAnimMontageInstance::BranchingPointEventHandler(const FBranchingPointMarke
 				UFunction* Function = AnimInstance.Get()->FindFunction(FuncFName);
 				if (Function)
 				{
-					AnimInstance.Get()->ProcessEvent(Function, NULL);
+					AnimInstance.Get()->ProcessEvent(Function, nullptr);
 				}
 				// In case older BranchingPoint has been re-implemented as a new Custom Notify, this is if BranchingPoint function hasn't been found.
 				else
@@ -2679,7 +2682,7 @@ void FAnimMontageInstance::BranchingPointEventHandler(const FBranchingPointMarke
 					AnimInstance.Get()->TriggerSingleAnimNotify(NotifyEvent);
 				}
 			}
-			else if (NotifyEvent->NotifyStateClass != NULL)
+			else if (NotifyEvent->NotifyStateClass != nullptr)
 			{
 				if (BranchingPointMarker->NotifyEventType == EAnimNotifyEventType::Begin)
 				{
@@ -2808,7 +2811,7 @@ UAnimMontage* FAnimMontageInstance::PreviewSequencerMontagePosition(FName SlotNa
 			// we have to get it again in case if this is new
 			MontageInstanceToUpdate = AnimInst->GetMontageInstanceForID(InOutInstanceId);
 			// since we don't advance montage in the tick, we manually have to handle notifies
-			MontageInstanceToUpdate->HandleEvents(InFromPosition, InToPosition, NULL);
+			MontageInstanceToUpdate->HandleEvents(InFromPosition, InToPosition, nullptr);
 			if (!bFireNotifies)
 			{
 				AnimInst->NotifyQueue.Reset(SkeletalMeshComponent);
@@ -2927,6 +2930,15 @@ void UAnimMontage::BakeTimeStretchCurve()
 	TimeStretchCurve.BakeFromFloatCurve(*TimeStretchFloatCurve, SequenceLength);
 	PRAGMA_ENABLE_DEPRECATION_WARNINGS
 }
+
+void UAnimMontage::PopulateWithExistingModel(TScriptInterface<IAnimationDataModel> ExistingDataModel)
+{
+	Super::PopulateWithExistingModel(ExistingDataModel);
+	
+	// Set composite length while model is being populated
+	const float CurrentCalculatedLength = CalculateSequenceLength();
+	SetCompositeLength(CurrentCalculatedLength);
+}
 #endif // WITH_EDITOR
 
 FMontageBlendSettings::FMontageBlendSettings()
diff --git a/Engine/Source/Runtime/Engine/Private/Animation/AnimNode_SequencePlayer.cpp b/Engine/Source/Runtime/Engine/Private/Animation/AnimNode_SequencePlayer.cpp
--- a/Engine/Source/Runtime/Engine/Private/Animation/AnimNode_SequencePlayer.cpp
+++ b/Engine/Source/Runtime/Engine/Private/Animation/AnimNode_SequencePlayer.cpp
@@ -128,7 +128,7 @@ void FAnimNode_SequencePlayerBase::Evaluate_AnyThread(FPoseContext& Output)
 		}
 
 		FAnimationPoseData AnimationPoseData(Output);
-		CurrentSequence->GetAnimationPose(AnimationPoseData, FAnimExtractContext(InternalTimeAccumulator, Output.AnimInstanceProxy->ShouldExtractRootMotion(), DeltaTimeRecord, GetLoopAnimation()));
+		CurrentSequence->GetAnimationPose(AnimationPoseData, FAnimExtractContext(static_cast<double>(InternalTimeAccumulator), Output.AnimInstanceProxy->ShouldExtractRootMotion(), DeltaTimeRecord, GetLoopAnimation()));
 	}
 	else
 	{
