Reintroducing AnimationDataModel refactor with additional memory and cooktime optimizations.

diff --git a/Engine/Source/Runtime/Engine/Private/Animation/AnimSequence.cpp b/Engine/Source/Runtime/Engine/Private/Animation/AnimSequence.cpp
--- a/Engine/Source/Runtime/Engine/Private/Animation/AnimSequence.cpp
+++ b/Engine/Source/Runtime/Engine/Private/Animation/AnimSequence.cpp
@@ -52,8 +52,9 @@
 #include "Animation/AttributesRuntime.h"
 #include "Stats/StatsHierarchical.h"
 #include "Animation/AnimationPoseData.h"
-#include "ITimeManagementModule.h"
 #include "CommonFrameRates.h"
+#include "ITimeManagementModule.h"
+#include "Algo/Accumulate.h"
 #include "HAL/LowLevelMemTracker.h"
 
 LLM_DEFINE_TAG(SequenceData);
@@ -61,7 +62,7 @@ LLM_DEFINE_TAG(SequenceData);
 #include "HAL/FileManager.h"
 
 #if WITH_EDITOR
-#include "Animation/AnimData/AnimDataModel.h"
+#include "Animation/AnimData/IAnimationDataModel.h"
 #include "Animation/BuiltInAttributeTypes.h"
 #include "AssetRegistry/AssetRegistryModule.h"
 #endif // WITH_EDITOR
@@ -293,7 +294,7 @@ FString GetAnimSequenceSpecificCacheKeySuffix(const UAnimSequence& Seq, bool bPe
 	FString Ret = FString::Printf(TEXT("%i_%s%s%s_%c%c%i_%s_%s_%i"),
 		Seq.CompressCommandletVersion,
 		PRAGMA_DISABLE_DEPRECATION_WARNINGS
-		*Seq.GetRawDataGuid().ToString(),
+		*Seq.GetDataModel()->GenerateGuid().ToString(),
 		PRAGMA_ENABLE_DEPRECATION_WARNINGS
 		*Seq.GetSkeleton()->GetGuid().ToString(),
 		*Seq.GetSkeleton()->GetVirtualBoneGuid().ToString(),
@@ -301,7 +302,7 @@ FString GetAnimSequenceSpecificCacheKeySuffix(const UAnimSequence& Seq, bool bPe
 		RefType,
 		Seq.RefFrameIndex,
 		PRAGMA_DISABLE_DEPRECATION_WARNINGS
-		(bIsValidAdditive && Seq.RefPoseSeq) ? *Seq.RefPoseSeq->GetRawDataGuid().ToString() : TEXT("NoAdditiveGuid"),
+		(bIsValidAdditive && Seq.RefPoseSeq) ? *Seq.RefPoseSeq->GetDataModel()->GenerateGuid().ToString() : TEXT("NoAdditiveGuid"),
 		PRAGMA_ENABLE_DEPRECATION_WARNINGS
 		*ArcToHexString.MakeString(),
 		Seq.GetSamplingFrameRate().Numerator / Seq.GetSamplingFrameRate().Denominator
@@ -431,15 +432,6 @@ bool FCurveTrack::CompressCurveWeights()
 	return false;
 }
 
-/////////////////////////////////////////////////////////////
-
-// since we want this change for hot fix, I can't change header file, 
-// next time move this to the header
-float GetIntervalPerKey(int32 NumFrames, float SequenceLength) 
-{
-	return (NumFrames > 1) ? (SequenceLength / (NumFrames-1)) : MINIMUM_ANIMATION_LENGTH;
-}
-
 /////////////////////////////////////////////////////
 // UAnimSequence
 
@@ -466,7 +458,7 @@ UAnimSequence::UAnimSequence(const FObjectInitializer& ObjectInitializer)
 
 	if (!HasAnyFlags(EObjectFlags::RF_ClassDefaultObject| EObjectFlags::RF_NeedLoad))
 	{
-		check(IsDataModelValid());
+		check(DataModelInterface.GetObject());
 	}
 #endif
 }
@@ -491,12 +483,12 @@ void UAnimSequence::GetAssetRegistryTags(TArray<FAssetRegistryTag>& OutTags) con
 		OutTags.Add( FAssetRegistryTag(SourceFileTagName(), AssetImportData->GetSourceData().ToJson(), FAssetRegistryTag::TT_Hidden) );
 	}
 
-	OutTags.Add(FAssetRegistryTag(TEXT("Compression Ratio"), FString::Printf(TEXT("%.03f"), (float)GetApproxCompressedSize() / (float)GetUncompressedRawSize()), FAssetRegistryTag::TT_Numerical));
-	if (DataModel)
+	if (DataModelInterface.GetObject() && DataModelInterface->HasBeenPopulated())
 	{
-		OutTags.Add(FAssetRegistryTag(TEXT("Source Frame Rate"), FString::Printf(TEXT("%.2f"), GetDataModel()->GetFrameRate().AsDecimal()), FAssetRegistryTag::TT_Numerical));
-		OutTags.Add(FAssetRegistryTag(TEXT("Number of Frames"), FString::Printf(TEXT("%.i"), GetDataModel()->GetNumberOfFrames()), FAssetRegistryTag::TT_Numerical));
-		OutTags.Add(FAssetRegistryTag(TEXT("Number of Keys"), FString::Printf(TEXT("%i"), GetDataModel()->GetNumberOfKeys()), FAssetRegistryTag::TT_Numerical));
+		OutTags.Add(FAssetRegistryTag(TEXT("Compression Ratio"), FString::Printf(TEXT("%.03f"), (float)GetApproxCompressedSize() / (float)GetUncompressedRawSize()), FAssetRegistryTag::TT_Numerical));
+		OutTags.Add(FAssetRegistryTag(TEXT("Source Frame Rate"), FString::Printf(TEXT("%.2f"), DataModelInterface->GetFrameRate().AsDecimal()), FAssetRegistryTag::TT_Numerical));
+		OutTags.Add(FAssetRegistryTag(TEXT("Number of Frames"), FString::Printf(TEXT("%.i"), DataModelInterface->GetNumberOfFrames()), FAssetRegistryTag::TT_Numerical));
+		OutTags.Add(FAssetRegistryTag(TEXT("Number of Keys"), FString::Printf(TEXT("%i"), DataModelInterface->GetNumberOfKeys()), FAssetRegistryTag::TT_Numerical));
 	}
 	else
 	{
@@ -521,14 +513,39 @@ void UAnimSequence::AddReferencedObjects(UObject* This, FReferenceCollector& Col
 }
 
 #if WITH_EDITOR
+void UAnimSequence::WillNeverCacheCookedPlatformDataAgain()
+{
+	Super::WillNeverCacheCookedPlatformDataAgain();
+	ClearResampledAnimationTrackData();
+	ClearCompressedCurveData();
+	ClearCompressedBoneData();
+	bUseRawDataOnly = true;
+}
+
+void UAnimSequence::ClearCachedCookedPlatformData(const ITargetPlatform* TargetPlatform)
+{
+	Super::ClearCachedCookedPlatformData(TargetPlatform);
+
+	// For now just clear resampled data (as this will always be resampled when not found for compression)
+	ClearResampledAnimationTrackData();
+}
+
+void UAnimSequence::ClearAllCachedCookedPlatformData()
+{
+	Super::ClearAllCachedCookedPlatformData();
+	
+	// For now just clear resampled data (as this will always be resampled when not found for compression)
+	ClearResampledAnimationTrackData();
+}
+
 int32 UAnimSequence::GetUncompressedRawSize() const
 {
-	if (DataModel)
+	if (DataModelInterface.GetObject())
 	{
-		const int32 BoneRawSize = FRawAnimSequenceTrack::SingleKeySize * DataModel->GetNumBoneTracks() * DataModel->GetNumberOfKeys();
+		const int32 BoneRawSize = FRawAnimSequenceTrack::SingleKeySize * DataModelInterface->GetNumBoneTracks() * DataModelInterface->GetNumberOfKeys();
 		int32 CurveRawSize = 0;
 
-		for (const FFloatCurve& Curve : DataModel->GetCurveData().FloatCurves)
+		for (const FFloatCurve& Curve : DataModelInterface->GetCurveData().FloatCurves)
 		{
 			CurveRawSize += sizeof(FFloatCurve);
 			CurveRawSize += sizeof(FRichCurveKey) * Curve.FloatCurve.Keys.Num();
@@ -540,17 +557,18 @@ int32 UAnimSequence::GetUncompressedRawSize() const
 
 int32 UAnimSequence::GetApproxRawSize() const
 {
-	if (DataModel)
+	if (ShouldDataModelBeValid())
 	{
-		int32 Total = sizeof(FRawAnimSequenceTrack) * DataModel->GetNumBoneTracks();
-		for (const FBoneAnimationTrack& AnimationTrack : DataModel->GetBoneAnimationTracks())
+		ValidateModel();
+		int32 Total = sizeof(FRawAnimSequenceTrack) * DataModelInterface->GetNumBoneTracks();
+		for (const FBoneAnimationTrack& AnimationTrack : DataModelInterface->GetBoneAnimationTracks())
 		{
-				Total += AnimationTrack.InternalTrackData.PosKeys.Num() * AnimationTrack.InternalTrackData.PosKeys.GetTypeSize();
-				Total += AnimationTrack.InternalTrackData.RotKeys.Num() * AnimationTrack.InternalTrackData.RotKeys.GetTypeSize();
-				Total += AnimationTrack.InternalTrackData.ScaleKeys.Num() * AnimationTrack.InternalTrackData.ScaleKeys.GetTypeSize();
+			Total += AnimationTrack.InternalTrackData.PosKeys.Num() * AnimationTrack.InternalTrackData.PosKeys.GetTypeSize();
+			Total += AnimationTrack.InternalTrackData.RotKeys.Num() * AnimationTrack.InternalTrackData.RotKeys.GetTypeSize();
+			Total += AnimationTrack.InternalTrackData.ScaleKeys.Num() * AnimationTrack.InternalTrackData.ScaleKeys.GetTypeSize();
 		}
-	
-			for (const FFloatCurve& Curve : DataModel->GetCurveData().FloatCurves)
+
+		for (const FFloatCurve& Curve : DataModelInterface->GetCurveData().FloatCurves)
 		{
 			Total += sizeof(FFloatCurve);
 			Total += sizeof(FRichCurveKey) * Curve.FloatCurve.Keys.Num();
@@ -668,7 +686,6 @@ void UAnimSequence::Serialize(FArchive& Ar)
 						// Set RawAnimationData to Source
 						RawAnimationData = SourceRawAnimationData_DEPRECATED;
 						SourceRawAnimationData_DEPRECATED.Empty();
-						MarkRawDataAsModified();
 					}
 					PRAGMA_ENABLE_DEPRECATION_WARNINGS
 				}
@@ -901,20 +918,6 @@ void UAnimSequence::PostLoad()
 	{
 		RefPoseSeq->ConditionalPostLoad();
 	}
-
-	EnsureValidRawDataGuid();
-
-	// I have to do this first thing in here
-	// so that remove all NaNs before even being read
-	if(GetLinkerUEVersion() < VER_UE4_ANIMATION_REMOVE_NANS)
-	{
-		RemoveNaNTracks();
-	}
-
-	if (ShouldDataModelBeValid())
-	{
-	Controller->RemoveBoneTracksMissingFromSkeleton(GetSkeleton());
-	}
 #endif // WITH_EDITOR
 
 	// if valid additive, but if base additive isn't 
@@ -922,7 +925,7 @@ void UAnimSequence::PostLoad()
 	// but this causes issue since once this happens this is unrecoverable until you delete from outside of editor
 	if (IsValidAdditive())
 	{
-		if (RefPoseSeq && !GetSkeleton()->IsCompatible(RefPoseSeq->GetSkeleton()))
+		if (RefPoseSeq && GetSkeleton() && !GetSkeleton()->IsCompatible(RefPoseSeq->GetSkeleton()))
 		{
 			// if this happens, there was a issue with retargeting
 			UE_LOG(LogAnimation, Warning, TEXT("Animation %s - Invalid additive animation base animation (%s)"), *GetName(), *RefPoseSeq->GetName());
@@ -930,27 +933,6 @@ void UAnimSequence::PostLoad()
 		}
 	}
 
-#if WITH_EDITOR
-	// Generate transient animation data, used to fallback on whenever the animation modified and not yet compressed (what previously used to be the raw data)
-	ResampleAnimationTrackData();
-
-	static bool ForcedRecompressionSetting = FAnimationUtils::GetForcedRecompressionSetting();
-
-	if (ForcedRecompressionSetting)
-	{
-		//Force recompression
-		PRAGMA_DISABLE_DEPRECATION_WARNINGS
-		RawDataGuid = FGuid::NewGuid();
-		PRAGMA_ENABLE_DEPRECATION_WARNINGS
-		bUseRawDataOnly = true;
-	}
-
-	if (bUseRawDataOnly)
-	{
-		RequestAnimCompression(FRequestAnimCompressionParams(true, false, false));
-	}
-#endif
-
 	// Ensure notifies are sorted.
 	SortNotifies();
 
@@ -960,48 +942,15 @@ void UAnimSequence::PostLoad()
 	if (GetOutermost()->HasAnyPackageFlags(PKG_Cooked | PKG_FilterEditorOnly))
 	{
 		// Update the virtual bone GUID according to the skeleton
-		if (GetSkeleton())
+		if (const USkeleton* CurrentSkeleton = GetSkeleton())
 		{
-			SetSkeletonVirtualBoneGuid(GetSkeleton()->GetVirtualBoneGuid());
-		}
-	}
-	else
-	{
-#if WITH_EDITOR
-		if (IsDataModelValid() && DataModel->GetNumberOfKeys() == 0 && DataModel->GetNumberOfFloatCurves() == 0 )
-		{
-			UE_LOG(LogAnimation, Warning, TEXT("No animation data exists for sequence %s (%s)"), *GetName(), (GetOuter() ? *GetOuter()->GetFullName() : *GetFullName()) );
+			SetSkeletonVirtualBoneGuid(CurrentSkeleton->GetVirtualBoneGuid());
 
-			if (!IsRunningGame())
+			if(!IsCompressedDataValid())
 			{
-				static FName NAME_LoadErrors("LoadErrors");
-				FMessageLog LoadErrors(NAME_LoadErrors);
-
-				TSharedRef<FTokenizedMessage> Message = LoadErrors.Warning();
-				Message->AddToken(FTextToken::Create(LOCTEXT("EmptyAnimationData1", "The Animation ")));
-				Message->AddToken(FAssetNameToken::Create(GetPathName(), FText::FromString(GetName())));
-				Message->AddToken(FTextToken::Create(LOCTEXT("EmptyAnimationData2", " has no animation data. Recommend to remove.")));
-				LoadErrors.Notify();
+				UE_LOG(LogAnimation, Fatal, TEXT("No animation compression exists for sequence %s (%s)"), *GetName(), (GetOuter() ? *GetOuter()->GetFullName() : *GetFullName()) );
 			}
 		}
-		// @remove temp hack for fixing length
-		// @todo need to fix importer/editing feature
-		else if (IsDataModelValid() && GetPlayLength() == 0.f)
-		{
-			ensure(DataModel->GetNumberOfKeys() == 1);
-			Controller->SetPlayLength(MINIMUM_ANIMATION_LENGTH);
-		}
-		// Raw data exists, but missing compress animation data
-		else
-#endif
-		if( GetSkeleton() && !IsCompressedDataValid()
-#if WITH_EDITOR
-			&& !bCompressionInProgress
-#endif
-			)
-		{
-			UE_LOG(LogAnimation, Fatal, TEXT("No animation compression exists for sequence %s (%s)"), *GetName(), (GetOuter() ? *GetOuter()->GetFullName() : *GetFullName()) );
-		}
 	}
 
 	for(FAnimNotifyEvent& Notify : Notifies)
@@ -1029,26 +978,12 @@ void UAnimSequence::PostLoad()
 			CurrentSkeleton->VerifySmartName(USkeleton::AnimCurveMappingName, CurveName);
 		}
 #if WITH_EDITOR
-		if (IsDataModelValid())
-		{
-		static const bool DoNotTransactAction = false;
-		Controller->FindOrAddCurveNamesOnSkeleton(CurrentSkeleton, ERawCurveTrackTypes::RCT_Transform, DoNotTransactAction);
-		}
-
 		for (const FAnimSyncMarker& SyncMarker : AuthoredSyncMarkers)
 		{
 			CurrentSkeleton->RegisterMarkerName(SyncMarker.MarkerName);
 		}
 #endif
 	}
-
-#if WITH_EDITOR
-	if (IsDataModelValid() && GetLinkerCustomVersion(FUE5MainStreamObjectVersion::GUID) < FUE5MainStreamObjectVersion::MoveCustomAttributesToDataModel)
-	{
-		ValidateModel();
-		MoveAttributesToModel();
-	}
-#endif // WITH_EDITOR
 }
 
 #if WITH_EDITOR
@@ -1107,11 +1042,11 @@ void UAnimSequence::PostEditChangeProperty(FPropertyChangedEvent& PropertyChange
 		}
 	}
 
-	if (RefPoseSeq != NULL)
+	if (RefPoseSeq != nullptr)
 	{
 		if (RefPoseSeq->GetSkeleton() != GetSkeleton()) // @todo this may require to be changed when hierarchy of skeletons is introduced
 		{
-			RefPoseSeq = NULL;
+			RefPoseSeq = nullptr;
 		}
 	}
 
@@ -1157,19 +1092,16 @@ void UAnimSequence::PostEditChangeProperty(FPropertyChangedEvent& PropertyChange
 
 	if (bShouldResample)
 	{
-		ResampleAnimationTrackData();
+		ClearResampledAnimationTrackData();
+		CalculateNumberOfSampledKeys();
 	}
 
 	const bool bNeedPostProcess = !IsCompressedDataValid() || bAdditiveSettingsChanged || bCompressionAffectingSettingsChanged || bShouldResample;
 
 	// @Todo fix me: This is temporary fix to make sure they always have compressed data
 	ValidateModel();
-	if (IsDataModelValid() && DataModel->GetNumBoneTracks() && bNeedPostProcess)
+	if (IsDataModelValid() && DataModelInterface->GetNumBoneTracks() && bNeedPostProcess)
 	{
-		// QQ need new function here to queue async compression
-		PRAGMA_DISABLE_DEPRECATION_WARNINGS
-		MarkRawDataAsModified(false);
-		PRAGMA_ENABLE_DEPRECATION_WARNINGS
 		RecompressAnimationData();
 	}
 }
@@ -1202,15 +1134,33 @@ static int32 FindKeyIndex(float Time, const TimeArray& Times)
 	return FoundIndex;
 }
 
+#if WITH_EDITOR
+int32 GForceRawData = 0;
+static FAutoConsoleVariableRef CVarForceRawData(
+	TEXT("a.ForceEvalRawData"),
+	GForceRawData,
+	TEXT("Values: 0/1\n")
+	TEXT("Controls whether or not to forcefully sample non-compressed anim data."),
+	ECVF_Cheat);
+#endif // WITH_EDITOR
+
 void UAnimSequence::GetBoneTransform(FTransform& OutAtom, int32 TrackIndex, float Time, bool bUseRawData) const
+{
+	GetBoneTransform(OutAtom, TrackIndex, (double)Time, bUseRawData);
+}
+
+void UAnimSequence::GetBoneTransform(FTransform& OutAtom, int32 TrackIndex, double Time, bool bUseRawData) const
 {
 	// If the caller didn't request that raw animation data be used . . .
-	const bool bEvaluateCompressedData = (!bUseRawData || !CanEvaluateRawAnimationData());
+	const bool bEvaluateCompressedData =
+#if WITH_EDITOR 
+		!GForceRawData && 
+#endif 
+		(!bUseRawData || !CanEvaluateRawAnimationData());
+
 	if ( bEvaluateCompressedData && IsCompressedDataValid())
 	{
-		PRAGMA_DISABLE_DEPRECATION_WARNINGS
-		FAnimSequenceDecompressionContext DecompContext(SequenceLength, Interpolation, GetFName(), *CompressedData.CompressedDataStructure);
-		PRAGMA_ENABLE_DEPRECATION_WARNINGS
+		FAnimSequenceDecompressionContext DecompContext(TargetFrameRate, TargetFrameRate.AsFrameNumber(GetPlayLength()).Value, Interpolation, GetRetargetTransformsSourceName(), *CompressedData.CompressedDataStructure, GetSkeleton(), IsValidAdditive());
 		DecompContext.Seek(Time);
 		if (CompressedData.BoneCompressionCodec)
 		{
@@ -1222,7 +1172,7 @@ void UAnimSequence::GetBoneTransform(FTransform& OutAtom, int32 TrackIndex, floa
 	{
 #if WITH_EDITOR
 		ValidateModel();
-		UE::Anim::GetBoneTransformFromModel(DataModel, OutAtom, TrackIndex, Time, Interpolation);
+		UE::Anim::GetBoneTransformFromModel(DataModelInterface.GetInterface(), OutAtom, TrackIndex, Time, Interpolation);
 #endif
 	}
 }
@@ -1243,7 +1193,7 @@ void UAnimSequence::GetBoneTransform(FTransform& OutAtom, int32 TrackIndex, FAni
 	{
 #if WITH_EDITOR
 		ValidateModel();
-		UE::Anim::GetBoneTransformFromModel(DataModel, OutAtom, TrackIndex, DecompContext.GetEvaluationTime(), Interpolation);
+		UE::Anim::GetBoneTransformFromModel(DataModelInterface.GetInterface(), OutAtom, TrackIndex, DecompContext.GetEvaluationTime(), Interpolation);
 #endif
 	}
 }
@@ -1267,9 +1217,9 @@ FTransform UAnimSequence::ExtractRootTrackTransform(float Pos, const FBoneContai
 		if (CanEvaluateRawAnimationData() && bUseRawDataOnly)
 		{
 			ValidateModel();
-			if (const FBoneAnimationTrack* RootTrack = DataModel->FindBoneTrackByIndex(RootBoneIndex))
+			if (const FBoneAnimationTrack* RootTrack = DataModelInterface->FindBoneTrackByIndex(RootBoneIndex))
 			{
-				return DataModel->GetBoneTrackIndex(*RootTrack);
+				return DataModelInterface->GetBoneTrackIndex(*RootTrack);
 			}
 		}
 		else
@@ -1290,7 +1240,7 @@ FTransform UAnimSequence::ExtractRootTrackTransform(float Pos, const FBoneContai
 	{
 		// if we do have root data, then return root data
 		FTransform RootTransform;
-		GetBoneTransform(RootTransform, TrackIndex, Pos, bUseRawDataOnly);
+		GetBoneTransform(RootTransform, TrackIndex, static_cast<double>(Pos), bUseRawDataOnly);
 		return RootTransform;
 	}
 
@@ -1345,7 +1295,7 @@ FTransform UAnimSequence::ExtractRootMotion(float StartTime, float DeltaTime, bo
 			// If we've hit the end of the animation, and we're allowed to loop, keep going.
 			if ((AdvanceType == ETAA_Finished) && bAllowLooping)
 			{
-				const float ActualDeltaMove = (CurrentPosition - PreviousPosition);
+				const double ActualDeltaMove = (CurrentPosition - PreviousPosition);
 				DesiredDeltaMove -= ActualDeltaMove;
 
 				PreviousPosition = bPlayingBackwards ? GetPlayLength() : 0.f;
@@ -1375,8 +1325,8 @@ FTransform UAnimSequence::ExtractRootMotionFromRange(float StartTrackPosition, f
 		}
 	}
 
-	FTransform StartTransform = ExtractRootTrackTransform(StartTrackPosition, NULL);
-	FTransform EndTransform = ExtractRootTrackTransform(EndTrackPosition, NULL);
+	FTransform StartTransform = ExtractRootTrackTransform(StartTrackPosition, nullptr);
+	FTransform EndTransform = ExtractRootTrackTransform(EndTrackPosition, nullptr);
 
 	// Use old calculation if needed.
 	if (bUseNormalizedRootMotionScale)
@@ -1510,7 +1460,11 @@ void UAnimSequence::GetBonePose(FAnimationPoseData& OutAnimationPoseData, const
 
 	FBoneContainer& RequiredBones = OutPose.GetBoneContainer();
 	check(!bForceUseRawData || CanEvaluateRawAnimationData());
-	const bool bUseRawDataForPoseExtraction = (CanEvaluateRawAnimationData() && bForceUseRawData) || UseRawDataForPoseExtraction(RequiredBones);
+	const bool bUseRawDataForPoseExtraction = (CanEvaluateRawAnimationData() && bForceUseRawData) ||
+#if WITH_EDITOR
+ 		GForceRawData == 1 || 
+#endif // WITH_EDITOR
+		UseRawDataForPoseExtraction(RequiredBones);
 
 	const bool bIsBakedAdditive = !bUseRawDataForPoseExtraction && IsValidAdditive();
 
@@ -1552,7 +1506,7 @@ void UAnimSequence::GetBonePose(FAnimationPoseData& OutAnimationPoseData, const
 				// Virtual bones are part of the retarget transform pose, so if the pose has not been updated (recently) there might be a mismatch
 				if (SkeletonBoneIndex < NumRawBones || AuthoredOnRefSkeleton.IsValidIndex(SkeletonBoneIndex.GetInt()))
 				{
-					OutPose[PoseBoneIndex] = AuthoredOnRefSkeleton[SkeletonBoneIndex.GetInt()];					
+					OutPose[PoseBoneIndex] = AuthoredOnRefSkeleton[SkeletonBoneIndex.GetInt()];
 				}
 			}
 		}
@@ -1562,35 +1516,22 @@ void UAnimSequence::GetBonePose(FAnimationPoseData& OutAnimationPoseData, const
 		}
 	}
 
-	// Evaluate the (remapped if needed) curves.
-	// Scoped so that the RemappedCurve destructs and isn't kept alive longer than needed.
-	{
-		FSkeletonRemappingCurve RemappedCurve(OutAnimationPoseData.GetCurve(), RequiredBones, GetSkeleton());
-		EvaluateCurveData(RemappedCurve.GetCurve(), ExtractionContext.CurrentTime, bUseRawDataForPoseExtraction);
-	}
-
 #if WITH_EDITOR
-	const int32 NumTracks = bUseRawDataForPoseExtraction ? DataModel->GetNumBoneTracks() : CompressedData.CompressedTrackToSkeletonMapTable.Num();
+	const int32 NumTracks = bUseRawDataForPoseExtraction ? DataModelInterface->GetNumBoneTracks() : CompressedData.CompressedTrackToSkeletonMapTable.Num();
 #else
 	const int32 NumTracks = CompressedData.CompressedTrackToSkeletonMapTable.Num();
 #endif 
-	if (NumTracks == 0)
-	{
-		return;
-	}
-
 	const bool bTreatAnimAsAdditive = (IsValidAdditive() && !bUseRawDataForPoseExtraction); // Raw data is never additive
-	FRootMotionReset RootMotionReset(bEnableRootMotion, RootMotionRootLock, bForceRootLock, ExtractRootTrackTransform(0.f, &RequiredBones), bTreatAnimAsAdditive);
+	const FRootMotionReset RootMotionReset(bEnableRootMotion, RootMotionRootLock, bForceRootLock, ExtractRootTrackTransform(0.0f, &RequiredBones), bTreatAnimAsAdditive);
 
 #if WITH_EDITOR
-	// this happens only with editor data
-	// Slower path for disable retargeting, that's only used in editor and for debugging.
+	// Evaluate raw (source) curve and bone data
 	if (bUseRawDataForPoseExtraction)
 	{
 		// Warning if we have invalid data
 		for (int32 TrackIndex = 0; TrackIndex < NumTracks; TrackIndex++)
 		{
-			const FBoneAnimationTrack& BoneAnimationTrack = DataModel->GetBoneTrackByIndex(TrackIndex);
+			const FBoneAnimationTrack& BoneAnimationTrack = DataModelInterface->GetBoneTrackByIndex(TrackIndex);
 			const FRawAnimSequenceTrack& TrackToExtract = BoneAnimationTrack.InternalTrackData;
 
 			// Bail out (with rather wacky data) if data is empty for some reason.
@@ -1601,38 +1542,47 @@ void UAnimSequence::GetBonePose(FAnimationPoseData& OutAnimationPoseData, const
 			}
 		}
 
-		UE::Anim::BuildPoseFromModel(DataModel, OutPose, ExtractionContext.CurrentTime, Interpolation, GetRetargetTransformsSourceName(), GetRetargetTransforms());
+		{
+			FSkeletonRemappingCurve RemappedCurve(OutAnimationPoseData.GetCurve(), RequiredBones, GetSkeleton());
+			FAnimationPoseData RemappedPoseData(OutAnimationPoseData.GetPose(), RemappedCurve.GetCurve(), OutAnimationPoseData.GetAttributes());
+
+			const UE::Anim::DataModel::FEvaluationContext EvaluationContext(ExtractionContext.CurrentTime, DataModelInterface->GetFrameRate(), GetRetargetTransformsSourceName(), GetRetargetTransforms(), Interpolation);
+			DataModelInterface->Evaluate(RemappedPoseData, EvaluationContext);
+		}
 
 		if ((ExtractionContext.bExtractRootMotion && RootMotionReset.bEnableRootMotion) || RootMotionReset.bForceRootLock)
 		{
 			RootMotionReset.ResetRootBoneForRootMotion(OutPose[FCompactPoseBoneIndex(0)], RequiredBones);
 		}
 
-		EvaluateAttributes(OutAnimationPoseData, ExtractionContext, true);
-
 		return;
 	}
+	else
 #endif // WITH_EDITOR
+	// Only try and evaluate compressed bone data if the animation contains any bone tracks
+	if (NumTracks != 0)
+	{
+		// Evaluate compressed bone data
+		FAnimSequenceDecompressionContext DecompContext(TargetFrameRate, TargetFrameRate.AsFrameTime(GetPlayLength()).RoundToFrame().Value, Interpolation, GetRetargetTransformsSourceName(), *CompressedData.CompressedDataStructure, GetSkeleton(), IsValidAdditive());
+		UE::Anim::Decompression::DecompressPose(OutPose, CompressedData, ExtractionContext, DecompContext, GetRetargetTransforms(), RootMotionReset);
+	}
 
-	DecompressPose(OutPose, CompressedData, ExtractionContext, GetSkeleton(), GetPlayLength(), Interpolation, bIsBakedAdditive, GetRetargetTransforms(), GetRetargetTransformsSourceName(), RootMotionReset);
-
-	EvaluateAttributes(OutAnimationPoseData, ExtractionContext, false);
-}
-
-
-#if WITH_EDITOR
-void UAnimSequence::UpdateCompressedCurveName(SmartName::UID_Type CurveUID, const struct FSmartName& NewCurveName)
-{
-	for (FSmartName& CurveName : CompressedData.CompressedCurveNames)
+	// (Always) evaluate compressed curve data
 	{
-		if (CurveName.UID == CurveUID)
+		// Scoped so that the RemappedCurve destructs and isn't kept alive longer than needed.
+		FSkeletonRemappingCurve RemappedCurve(OutAnimationPoseData.GetCurve(), RequiredBones, GetSkeleton());
+#if WITH_EDITOR
+		// When evaluating from raw animation data, UE::Anim::BuildPoseFromModel will populate the curve data
+		if (!bUseRawDataForPoseExtraction)
+#endif // WITH_EDITOR
 		{
-			CurveName = NewCurveName;
-			break;
+			EvaluateCurveData(RemappedCurve.GetCurve(), ExtractionContext.CurrentTime, bUseRawDataForPoseExtraction);
 		}
 	}
+
+	// Evaluate animation attributes (no compressed format yet)
+	EvaluateAttributes(OutAnimationPoseData, ExtractionContext, false);
 }
-#endif
 
 void UAnimSequence::GetBonePose_Additive(FCompactPose& OutPose, FBlendedCurve& OutCurve, const FAnimExtractContext& ExtractionContext) const
 {
@@ -1684,10 +1634,10 @@ void UAnimSequence::GetAdditiveBasePose(FCompactPose& OutPose, FBlendedCurve& Ou
 	GetAdditiveBasePose(OutAnimationPoseData, ExtractionContext);
 }
 
-static void GetSequencePose(FAnimationPoseData& OutAnimationPoseData, const FAnimExtractContext& ExtractionContext, const UAnimSequence &Seq, int FrameIndex)
+static void GetSequencePose(FAnimationPoseData& OutAnimationPoseData, const FAnimExtractContext& ExtractionContext, const UAnimSequence &Seq, int32 FrameIndex)
 {
-	const float Fraction = (Seq.GetNumberOfSampledKeys() > 0) ? FMath::Clamp<float>((float)FrameIndex / (float)Seq.GetNumberOfSampledKeys(), 0.f, 1.f) : 0.f;
-	const float BasePoseTime = Seq.GetPlayLength() * Fraction;
+	const double Fraction = (Seq.GetNumberOfSampledKeys() > 0) ? FMath::Clamp<double>((double)FrameIndex / (double)Seq.GetNumberOfSampledKeys(), 0.0, 1.0) : 0.0;
+	const double BasePoseTime = Seq.GetPlayLength() * Fraction;
 	FAnimExtractContext BasePoseExtractionContext(ExtractionContext);
 	BasePoseExtractionContext.CurrentTime = BasePoseTime;
 	Seq.GetBonePose(OutAnimationPoseData, BasePoseExtractionContext, true);
@@ -1702,8 +1652,8 @@ void UAnimSequence::GetAdditiveBasePose(FAnimationPoseData& OutAnimationPoseData
 		{
 			// normalize time to fit base seq
 			PRAGMA_DISABLE_DEPRECATION_WARNINGS
-			const float Fraction = (SequenceLength > 0.f)? FMath::Clamp<float>(ExtractionContext.CurrentTime / SequenceLength, 0.f, 1.f) : 0.f;
-			const float BasePoseTime = RefPoseSeq->SequenceLength * Fraction;
+			const double Fraction = (GetPlayLength() > 0.0)? FMath::Clamp<double>(ExtractionContext.CurrentTime / (double)GetPlayLength(), 0.0, 1.0) : 0.0;
+			const double BasePoseTime = (double)RefPoseSeq->GetPlayLength() * Fraction;
 			PRAGMA_ENABLE_DEPRECATION_WARNINGS
 
 			FAnimExtractContext BasePoseExtractionContext(ExtractionContext);
@@ -1787,6 +1737,20 @@ bool UAnimSequence::CanEvaluateRawAnimationData() const
 	return false;
 }
 
+#if WITH_EDITOR
+void UAnimSequence::UpdateCompressedCurveName(SmartName::UID_Type CurveUID, const struct FSmartName& NewCurveName)
+{
+	for (FSmartName& CurveName : CompressedData.CompressedCurveNames)
+	{
+		if (CurveName.UID == CurveUID)
+		{
+			CurveName = NewCurveName;
+			break;
+		}
+	}
+}
+#endif // WITH_EDITOR
+
 #if WITH_EDITORONLY_DATA
 void UAnimSequence::UpdateRetargetSourceAsset()
 {
@@ -1884,6 +1848,14 @@ void UAnimSequence::SetCompressionComplete()
 		AssetRegistryModule.Get().AssetTagsFinalized(*this);
 	}
 }
+#endif
+
+#if WITH_EDITOR
+int32 GSkipDDC = 0;
+static FAutoConsoleVariableRef CVarSkipDDC(
+	TEXT("a.SkipDDC"),
+	GSkipDDC,
+	TEXT("1 = Skip DDC during compression. 0 = Include DDC results during compresion "));
 
 void UAnimSequence::RequestAnimCompression(FRequestAnimCompressionParams Params)
 {
@@ -1921,6 +1893,11 @@ void UAnimSequence::RequestAnimCompression(FRequestAnimCompressionParams Params)
 
 	EnsureValidRawDataGuid();
 
+	if(ResampledAnimationTrackData.Num() == 0)
+	{		
+		ResampleAnimationTrackData();
+	}
+
 	bUseRawDataOnly = true;
 
 	check(!bCompressionInProgress);
@@ -1942,7 +1919,7 @@ void UAnimSequence::RequestAnimCompression(FRequestAnimCompressionParams Params)
 		FDerivedDataAnimationCompression* AnimCompressor = new FDerivedDataAnimationCompression(TEXT("AnimSeq"), AssetDDCKey, Params.CompressContext);
 
 		// For debugging DDC/Compression issues
-		const bool bSkipDDC = false;
+		const bool bSkipDDC = GSkipDDC == 1;
 
 		// Data does not exist, need to build it.
 		FCompressibleAnimRef CompressibleData = MakeShared<FCompressibleAnimData, ESPMode::ThreadSafe>(this, bPerformStripping);
@@ -1994,7 +1971,6 @@ void UAnimSequence::RequestAnimCompression(FRequestAnimCompressionParams Params)
 			FAnimationUtils::BuildSkeletonMetaData(GetSkeleton(), BoneData);
 			Params.CompressContext->GatherPostCompressionStats(CompressedData, BoneData, GetFName(), CompressionTime, false);
 		}
-
 	}
 }
 
@@ -2014,8 +1990,6 @@ void UAnimSequence::ApplyCompressedData(const FString& DataCacheKeySuffix, const
 void UAnimSequence::ApplyCompressedData(const TArray<uint8>& Data)
 {
 #if WITH_EDITOR
-	SetCompressionComplete();
-	
 	SynchronousAnimatedBoneAttributesCompression();
 #endif
 	if(Data.Num() > 0)
@@ -2028,6 +2002,10 @@ void UAnimSequence::ApplyCompressedData(const TArray<uint8>& Data)
 		SetSkeletonVirtualBoneGuid(GetSkeleton()->GetVirtualBoneGuid());
 		bUseRawDataOnly = false;
 	}
+
+#if WITH_EDITOR
+	SetCompressionComplete();
+#endif
 }
 
 void UAnimSequence::SerializeCompressedData(FArchive& Ar, bool bDDCData)
@@ -2037,11 +2015,49 @@ void UAnimSequence::SerializeCompressedData(FArchive& Ar, bool bDDCData)
 	if (!HasAnyFlags(RF_ClassDefaultObject))
 	{
 		CompressedData.SerializeCompressedData(Ar, bDDCData, this, this->GetSkeleton(), BoneCompressionSettings, CurveCompressionSettings);
-		}
 	}
-
+}
 
 #if WITH_EDITOR
+void UAnimSequence::OnAnimModelLoaded()
+{
+	Super::OnAnimModelLoaded();
+
+	if (ShouldDataModelBeValid())
+	{
+		CalculateNumberOfSampledKeys();
+		Controller->RemoveBoneTracksMissingFromSkeleton(GetSkeleton());
+
+		if (FAnimationUtils::GetForcedRecompressionSetting())
+		{
+			RequestAnimCompression(FRequestAnimCompressionParams(true, false, false));
+		}
+
+		if ( DataModelInterface->GetNumberOfKeys() == 0 && DataModelInterface->GetNumberOfFloatCurves() == 0 )
+		{
+			UE_LOG(LogAnimation, Warning, TEXT("No animation data exists for sequence %s (%s)"), *GetName(), (GetOuter() ? *GetOuter()->GetFullName() : *GetFullName()) );
+
+			if (!IsRunningGame())
+			{
+				static FName NAME_LoadErrors("LoadErrors");
+				FMessageLog LoadErrors(NAME_LoadErrors);
+
+				TSharedRef<FTokenizedMessage> Message = LoadErrors.Warning();
+				Message->AddToken(FTextToken::Create(LOCTEXT("EmptyAnimationData1", "The Animation ")));
+				Message->AddToken(FAssetNameToken::Create(GetPathName(), FText::FromString(GetName())));
+				Message->AddToken(FTextToken::Create(LOCTEXT("EmptyAnimationData2", " has no animation data. Recommend to remove.")));
+				LoadErrors.Notify();
+			}
+		}
+
+		if (GetLinkerCustomVersion(FUE5MainStreamObjectVersion::GUID) < FUE5MainStreamObjectVersion::MoveCustomAttributesToDataModel)
+		{
+			ValidateModel();
+			MoveAttributesToModel();
+		}
+	}
+}
+
 bool UAnimSequence::CanBakeAdditive() const
 {
 	return	(NumberOfSampledKeys > 0) &&
@@ -2053,7 +2069,7 @@ bool UAnimSequence::DoesSequenceContainZeroScale() const
 {
 	ValidateModel();
 
-	const TArray<FBoneAnimationTrack>& BoneAnimationTracks = DataModel->GetBoneAnimationTracks();
+	const TArray<FBoneAnimationTrack>& BoneAnimationTracks = DataModelInterface->GetBoneAnimationTracks();
 	for (const FBoneAnimationTrack& Track : BoneAnimationTracks)
 	{
 		for (const FVector3f& ScaleKey : Track.InternalTrackData.ScaleKeys)
@@ -2071,7 +2087,7 @@ bool UAnimSequence::DoesSequenceContainZeroScale() const
 FGuid UAnimSequence::GenerateGuidFromRawData() const
 {
 	ValidateModel();
-	return DataModel->GenerateGuid();
+	return DataModelInterface->GenerateGuid();
 }
 
 void CopyTransformToRawAnimationData(const FTransform& BoneTransform, FRawAnimSequenceTrack& Track, int32 Frame)
@@ -2087,17 +2103,13 @@ struct FByFramePoseEvalContext
 public:
 	FBoneContainer RequiredBones;
 
-	// Length of one frame.
-	const float IntervalTime;
-
 	TArray<FBoneIndexType> RequiredBoneIndexArray;
 
 	FByFramePoseEvalContext(const UAnimSequence* InAnimToEval)
-		: FByFramePoseEvalContext(InAnimToEval->GetPlayLength(), InAnimToEval->GetNumberOfSampledKeys(), InAnimToEval->GetSkeleton())
+		: FByFramePoseEvalContext(InAnimToEval->GetSkeleton())
 	{}
 		
-	FByFramePoseEvalContext(float InSequenceLength, int32 InRawNumOfFrames, USkeleton* InSkeleton)
-		: IntervalTime(InSequenceLength / ((float)FMath::Max(InRawNumOfFrames - 1, 1)))
+	FByFramePoseEvalContext(USkeleton* InSkeleton)
 	{
 		// Initialize RequiredBones for pose evaluation
 		RequiredBones.SetUseRAWData(true);
@@ -2226,15 +2238,14 @@ void UAnimSequence::BakeOutVirtualBoneTracks(TArray<FRawAnimSequenceTrack>& NewR
 
 	FAnimExtractContext ExtractContext;
 	const TArray<FVirtualBoneRefData>& VBRefData = GetSkeleton()->GetReferenceSkeleton().GetVirtualBoneRefData();
-	for (int32 Frame = 0; Frame < NumTrackKeys; ++Frame)
+	for (int32 KeyIndex = 0; KeyIndex < NumTrackKeys; ++KeyIndex)
 	{
 		// Initialise curve data from Skeleton
 		FBlendedCurve Curve;
 		Curve.InitFrom(EvalContext.RequiredBones);
 
 		//Grab pose for this frame
-		const float CurrentFrameTime = Frame * EvalContext.IntervalTime;
-		ExtractContext.CurrentTime = CurrentFrameTime;
+		ExtractContext.CurrentTime = TargetFrameRate.AsSeconds(KeyIndex);
 
 		UE::Anim::FStackAttributeContainer TempAttributes;
 		FAnimationPoseData AnimPoseData(Pose, Curve, TempAttributes);
@@ -2245,7 +2256,7 @@ void UAnimSequence::BakeOutVirtualBoneTracks(TArray<FRawAnimSequenceTrack>& NewR
 		{
 			const int32 VBIndex = VirtualBonesToAdd[VBToAddIndex];
 			const FVirtualBoneRefData& VB = VBRefData[VBIndex];
-			CopyTransformToRawAnimationData(Pose[FCompactPoseBoneIndex(VB.VBRefSkelIndex)], NewRawTracks[VBToAddIndex + VirtualBoneStartIndex], Frame);
+			CopyTransformToRawAnimationData(Pose[FCompactPoseBoneIndex(VB.VBRefSkelIndex)], NewRawTracks[VBToAddIndex + VirtualBoneStartIndex], KeyIndex);
 			UE::Anim::Compression::SanitizeRawAnimSequenceTrack(NewRawTracks[VBToAddIndex + VirtualBoneStartIndex]);
 		}
 	}
@@ -2295,15 +2306,14 @@ void UAnimSequence::TestEvalauteAnimation() const
 
 	FAnimExtractContext ExtractContext;
 
-	for (int32 Frame = 0; Frame < NumberOfSampledKeys; ++Frame)
+	for (int32 KeyIndex = 0; KeyIndex < NumberOfSampledKeys; ++KeyIndex)
 	{
 		// Initialise curve data from Skeleton
 		FBlendedCurve Curve;
 		Curve.InitFrom(EvalContext.RequiredBones);
 
 		//Grab pose for this frame
-		const float CurrentFrameTime = Frame * EvalContext.IntervalTime;
-		ExtractContext.CurrentTime = CurrentFrameTime;
+		ExtractContext.CurrentTime = TargetFrameRate.AsSeconds(KeyIndex);
 
 		UE::Anim::FStackAttributeContainer TempAttributes;
 		FAnimationPoseData AnimPoseData(Pose, Curve, TempAttributes);
@@ -2363,7 +2373,7 @@ void UAnimSequence::BakeOutAdditiveIntoRawData(TArray<FRawAnimSequenceTrack>& Ne
 	const FSmartNameMapping* CurveMapping = MySkeleton->GetSmartNameContainer(USkeleton::AnimCurveMappingName);
 	TArray<FName> SkeletonCurveNames;
 	CurveMapping->FillNameArray(SkeletonCurveNames);
-	for (const FFloatCurve& FloatCurve : DataModel->GetFloatCurves())
+	for (const FFloatCurve& FloatCurve : DataModelInterface->GetFloatCurves())
 	{
 		ensure(SkeletonCurveNames.Contains(FloatCurve.Name.DisplayName));
 	}
@@ -2405,19 +2415,21 @@ void UAnimSequence::BakeOutAdditiveIntoRawData(TArray<FRawAnimSequenceTrack>& Ne
 
 	FAnimExtractContext ExtractContext;
 
-	for (int32 Frame = 0; Frame < NumberOfSampledKeys; ++Frame)
+	FBlendedCurve Curve;
+	FBlendedCurve DummyBaseCurve;
+	
+	for (int32 KeyIndex = 0; KeyIndex < NumberOfSampledKeys; ++KeyIndex)
 	{
 		// Initialise curve data from Skeleton
-		FBlendedCurve Curve;
 		Curve.InitFrom(EvalContext.RequiredBones);
-
-		FBlendedCurve DummyBaseCurve;
 		DummyBaseCurve.InitFrom(EvalContext.RequiredBones);
 
 		//Grab pose for this frame
-		const float PreviousFrameTime = (Frame - 1) * EvalContext.IntervalTime;
-		const float CurrentFrameTime = Frame * EvalContext.IntervalTime;
-		ExtractContext.CurrentTime = CurrentFrameTime;
+		const double PreviousKeyTime = TargetFrameRate.AsSeconds(KeyIndex - 1);
+		const double CurrentKeyTime = TargetFrameRate.AsSeconds(KeyIndex);
+		FFrameTime CurrentFrameTime = KeyIndex;
+
+		ExtractContext.CurrentTime = CurrentKeyTime;
 
 		UE::Anim::FStackAttributeContainer BaseAttributes;
 		FAnimationPoseData AnimPoseData(Pose, Curve, BaseAttributes);
@@ -2430,8 +2442,8 @@ void UAnimSequence::BakeOutAdditiveIntoRawData(TArray<FRawAnimSequenceTrack>& Ne
 		//Write out every track for this frame
 		for (FCompactPoseBoneIndex TrackIndex(0); TrackIndex < NewRawTracks.Num(); ++TrackIndex)
 		{
-			CopyTransformToRawAnimationData(Pose[TrackIndex], NewRawTracks[TrackIndex.GetInt()], Frame);
-			CopyTransformToRawAnimationData(BasePose[TrackIndex], AdditiveBaseAnimationData[TrackIndex.GetInt()], Frame);
+			CopyTransformToRawAnimationData(Pose[TrackIndex], NewRawTracks[TrackIndex.GetInt()], KeyIndex);
+			CopyTransformToRawAnimationData(BasePose[TrackIndex], AdditiveBaseAnimationData[TrackIndex.GetInt()], KeyIndex);
 		}
 
 		//Write out curve data for this frame
@@ -2458,21 +2470,22 @@ void UAnimSequence::BakeOutAdditiveIntoRawData(TArray<FRawAnimSequenceTrack>& Ne
 					if (!bHasKeys)
 					{
 						//Add pre key of 0
-						if (Frame > 0)
+						if (KeyIndex > 0)
 						{
-							RawCurve->UpdateOrAddKey(0.f, PreviousFrameTime);
+							RawCurve->UpdateOrAddKey(0.f, PreviousKeyTime);
 						}
-
 					}
 
 					if (!bHasKeys || IsNewKeyDifferent(RawCurve->FloatCurve.GetLastKey(), CurveWeight))
 					{
-						RawCurve->UpdateOrAddKey(CurveWeight, CurrentFrameTime);
+						RawCurve->UpdateOrAddKey(CurveWeight, CurrentKeyTime);
 						TArray<FRichCurveKey>& CurveKeys = RawCurve->FloatCurve.Keys;
 						if (CurveKeys.Num() > 1)
 						{
 							FRichCurveKey& PrevKey = CurveKeys.Last(1);
-							if (PrevKey.Time < (PreviousFrameTime - UE_SMALL_NUMBER)) // Did we skip a frame, if so need to make previous key const
+							// Round to frame here as it would have been added at a specific frame boundary (though float->double conversion might mean the value is off)
+							const FFrameTime PrevKeyTime = TargetFrameRate.AsFrameTime(PrevKey.Time).RoundToFrame();
+							if (PrevKeyTime < CurrentFrameTime - 1) // Did we skip a frame, if so need to make previous key const
 							{
 								PrevKey.InterpMode = RCIM_Constant;
 							}
@@ -2510,13 +2523,14 @@ void UAnimSequence::UpdateDependentStreamingAnimations() const
 
 #endif
 
-
 #if WITH_EDITOR
 void UAnimSequence::ResetAnimation()
 {
 	// clear everything. Making new animation, so need to reset all the things that belong here
 	ValidateModel();
 
+	TargetFrameRate = UAnimationSettings::Get()->GetDefaultFrameRate();
+
 	Controller->OpenBracket(LOCTEXT("ResetAnimation_Bracket", "Resetting Animation Sequence"));
 	{
 		Controller->ResetModel();
@@ -2528,15 +2542,12 @@ void UAnimSequence::ResetAnimation()
 	AuthoredSyncMarkers.Empty();
 	UniqueMarkerNames.Empty();
 	AnimNotifyTracks.Empty();
-	RawDataGuid.Invalidate();
 }
 
 void UAnimSequence::DeleteDeprecatedRawAnimationData()
 {
 	PRAGMA_DISABLE_DEPRECATION_WARNINGS
 	RawAnimationData.Empty();
-	AnimationTrackNames.Empty();
-	TrackToSkeletonMapTable.Empty();
 	SourceRawAnimationData_DEPRECATED.Empty();
 	PRAGMA_ENABLE_DEPRECATION_WARNINGS
 }
@@ -2647,7 +2658,7 @@ void FillUpTransformBasedOnRig(USkeleton* Skeleton, TArray<FTransform>& NodeSpac
 
 					const FTransformBase* TransformBase = Rig->GetTransformBaseByNodeName(NodeName);
 
-					if (TransformBase != NULL)
+					if (TransformBase != nullptr)
 					{
 						// orientation constraint			
 						const auto& RotConstraint = TransformBase->Constraints[EControlConstraint::Type::Orientation];
@@ -2717,7 +2728,7 @@ void UAnimSequence::RemapTracksToNewSkeleton( USkeleton* NewSkeleton, bool bConv
 		USkeleton* OldSkeleton = GetSkeleton();
 
 		// first check if both has same rig, if so, we'll retarget using it
-		if (OldSkeleton && OldSkeleton->GetRig() != NULL && NewSkeleton->GetRig() == OldSkeleton->GetRig() && OldSkeleton->GetPreviewMesh() && NewSkeleton->GetPreviewMesh())
+		if (OldSkeleton && OldSkeleton->GetRig() != nullptr && NewSkeleton->GetRig() == OldSkeleton->GetRig() && OldSkeleton->GetPreviewMesh() && NewSkeleton->GetPreviewMesh())
 		{
 			const URig* Rig = OldSkeleton->GetRig();
 
@@ -2802,8 +2813,8 @@ void UAnimSequence::RemapTracksToNewSkeleton( USkeleton* NewSkeleton, bool bConv
 			ConvertedSpaceAnimations.AddZeroed(SrcNumTracks);
 			ConvertedLocalSpaceAnimations.AddZeroed(SrcNumTracks);
 
-			const int32 NumKeys = DataModel->GetNumberOfKeys();
-			float Interval = DataModel->GetFrameRate().AsInterval();
+			const int32 NumKeys = DataModelInterface->GetNumberOfKeys();
+			float Interval = DataModelInterface->GetFrameRate().AsInterval();
 
 			// allocate arrays
 			for(int32 SrcTrackIndex=0; SrcTrackIndex<SrcNumTracks; ++SrcTrackIndex)
@@ -2827,7 +2838,7 @@ void UAnimSequence::RemapTracksToNewSkeleton( USkeleton* NewSkeleton, bool bConv
 				{
 					FTransform AnimatedLocalKey;
 					//ExtractBoneTransform(RiggingAnimationData.AnimationTracks[SrcTrackIndex], AnimatedLocalKey, Interval * Key);
-					FAnimationUtils::ExtractTransformFromTrack(Interval* Key, NumberOfSampledKeys, GetPlayLength(), RiggingAnimationData.AnimationTracks[SrcTrackIndex], Interpolation, AnimatedLocalKey);
+					FAnimationUtils::ExtractTransformFromTrack(RiggingAnimationData.AnimationTracks[SrcTrackIndex], Interval* Key, NumberOfSampledKeys, GetPlayLength(), Interpolation, AnimatedLocalKey);
 
 					AnimatedLocalKey.ScaleTranslation(OldToNewTranslationRatio[NodeIndex]);
 
@@ -2930,7 +2941,7 @@ void UAnimSequence::RemapTracksToNewSkeleton( USkeleton* NewSkeleton, bool bConv
 		{
 			// Validate animation tracks against the new skeleton, any tracks linked to bones that do not exist in the new hierarchy are removed
 			Controller->RemoveBoneTracksMissingFromSkeleton(NewSkeleton);
-			ensure(DataModel->GetNumBoneTracks() != 0);
+			ensure(DataModelInterface->GetNumBoneTracks() != 0);
 
 			// make sure you do update reference pose before coming here
 			
@@ -2993,8 +3004,8 @@ void UAnimSequence::RemapTracksToNewSkeleton( USkeleton* NewSkeleton, bool bConv
 			ConvertedLocalSpaces.AddZeroed(NumBones);
 			ConvertedSpaceBases.AddZeroed(NumBones);
 
-			const int32 NumKeys = DataModel->GetNumberOfKeys();
-			const float Interval = DataModel->GetFrameRate().AsInterval();
+			const int32 NumKeys = DataModelInterface->GetNumberOfKeys();
+			const float Interval = DataModelInterface->GetFrameRate().AsInterval();
 
 			// allocate arrays
 			for(int32 BoneIndex=0; BoneIndex<NumBones; ++BoneIndex)
@@ -3010,7 +3021,7 @@ void UAnimSequence::RemapTracksToNewSkeleton( USkeleton* NewSkeleton, bool bConv
 			{
 				FName BoneName = NewSkeleton->GetReferenceSkeleton().GetBoneName(BoneIndex);
 				const int32 OldBoneIndex = OldSkeleton->GetReferenceSkeleton().FindBoneIndex(BoneName);
-				const int32 TrackIndex = DataModel->GetBoneTrackIndexByName(BoneName);
+				const int32 TrackIndex = DataModelInterface->GetBoneTrackIndexByName(BoneName);
 				const int32 ParentBoneIndex = NewSkeleton->GetReferenceSkeleton().GetParentIndex(BoneIndex);
 
 				if(TrackIndex != INDEX_NONE)
@@ -3019,7 +3030,7 @@ void UAnimSequence::RemapTracksToNewSkeleton( USkeleton* NewSkeleton, bool bConv
 					for(int32 Key=0; Key<NumKeys; ++Key)
 					{
 						FTransform AnimatedLocalKey;
-						UE::Anim::GetBoneTransformFromModel(DataModel, AnimatedLocalKey, TrackIndex, Interval * Key, Interpolation);
+						UE::Anim::GetBoneTransformFromModel(DataModelInterface.GetInterface(), AnimatedLocalKey, TrackIndex, Interval * Key, Interpolation);
 
 						// note that we apply scale in the animated space
 						// at this point, you should have scaled version of animated skeleton
@@ -3062,7 +3073,7 @@ void UAnimSequence::RemapTracksToNewSkeleton( USkeleton* NewSkeleton, bool bConv
 			for(int32 BoneIndex=0; BoneIndex<NumBones; ++BoneIndex)
 			{
 				const FName BoneName = NewSkeleton->GetReferenceSkeleton().GetBoneName(BoneIndex);
-				const int32 TrackIndex = DataModel->GetBoneTrackIndexByName(BoneName);
+				const int32 TrackIndex = DataModelInterface->GetBoneTrackIndexByName(BoneName);
 				int32 ParentBoneIndex = NewSkeleton->GetReferenceSkeleton().GetParentIndex(BoneIndex);
 
 				for(int32 Key=0; Key<NumKeys; ++Key)
@@ -3114,7 +3125,7 @@ void UAnimSequence::RemapTracksToNewSkeleton( USkeleton* NewSkeleton, bool bConv
 		}
 		else
 		{
-			Controller->RemoveBoneTracksMissingFromSkeleton(NewSkeleton);
+			Controller->UpdateWithSkeleton(NewSkeleton);
 		}
 
 		// I have to set this here in order for compression
@@ -3123,62 +3134,13 @@ void UAnimSequence::RemapTracksToNewSkeleton( USkeleton* NewSkeleton, bool bConv
 	}
 	else
 	{
-		Controller->RemoveBoneTracksMissingFromSkeleton(NewSkeleton);
+		SetSkeleton(NewSkeleton);
+		Controller->UpdateWithSkeleton(NewSkeleton, false);
 	}
 
 	Super::RemapTracksToNewSkeleton(NewSkeleton, bConvertSpaces);
 }
 
-
-void UAnimSequence::RemoveNaNTracks()
-{
-	ValidateModel();
-
-	TArray<FName> TracksToRemove;
-
-	const TArray<FBoneAnimationTrack>& BoneAnimationTracks = DataModel->GetBoneAnimationTracks();		
-	for (const FBoneAnimationTrack& AnimationTrack : BoneAnimationTracks)
-	{
-		const FRawAnimSequenceTrack& RawTrack = AnimationTrack.InternalTrackData;
-
-		bool bContainsNaN = false;
-		for ( auto Key : RawTrack.PosKeys )
-		{
-			bContainsNaN |= Key.ContainsNaN();
-		}
-
-		if (!bContainsNaN)
-		{
-			for(auto Key : RawTrack.RotKeys)
-			{
-				bContainsNaN |= Key.ContainsNaN();
-			}
-		}
-
-		if (!bContainsNaN)
-		{
-			for(auto Key : RawTrack.ScaleKeys)
-			{
-				bContainsNaN |= Key.ContainsNaN();
-			}
-		}
-
-		if (bContainsNaN)
-		{
-			UE_LOG(LogAnimation, Warning, TEXT("Animation raw data contains NaNs - Removing the following track [%s Track (%s)]"), (GetOuter() ? *GetOuter()->GetFullName() : *GetFullName()), *AnimationTrack.Name.ToString());
-
-			TracksToRemove.Add(AnimationTrack.Name);
-		}
-	}
-
-	IAnimationDataController::FScopedBracket ScopedBracket(Controller, LOCTEXT("RemoveNaNTracks_Bracket", "Removing track(s) containing NaN key data"));	
-	for (const FName& TrackName : TracksToRemove)
-	{
-		Controller->RemoveBoneTrack(TrackName);
-	}	
-}
-
-
 int32 FindFirstChildTrack(const USkeleton* MySkeleton, const FReferenceSkeleton& RefSkeleton, const TArray<FName>& AnimationTrackNames, FName BoneName)
 {
 	const int32 BoneIndex = RefSkeleton.FindBoneIndex(BoneName);
@@ -3280,8 +3242,8 @@ int32 UAnimSequence::GetSpaceBasedAnimationData(TArray< TArray<FTransform> >& An
 	ValidateModel();
 
 	// 2d array of animated time [boneindex][time key]
-	const int32 NumKeys = DataModel->GetNumberOfKeys();
-	const float Interval = DataModel->GetFrameRate().AsInterval();
+	const int32 NumKeys = DataModelInterface->GetNumberOfKeys();
+	const float Interval = DataModelInterface->GetFrameRate().AsInterval();
 
 	// allocate arrays
 	for (int32 BoneIndex = 0; BoneIndex < NumBones; ++BoneIndex)
@@ -3447,7 +3409,7 @@ int32 UAnimSequence::GetSpaceBasedAnimationData(TArray< TArray<FTransform> >& An
 		for (int32 BoneIndex = 0; BoneIndex < NumBones; ++BoneIndex)
 		{
 			const FName BoneName = MySkeleton->GetReferenceSkeleton().GetBoneName(BoneIndex);
-			const int32 TrackIndex = DataModel->GetBoneTrackIndexByName(BoneName);
+			const int32 TrackIndex = DataModelInterface->GetBoneTrackIndexByName(BoneName);
 			const int32 ParentBoneIndex = MySkeleton->GetReferenceSkeleton().GetParentIndex(BoneIndex);
 
 			if (TrackIndex != INDEX_NONE)
@@ -3456,7 +3418,7 @@ int32 UAnimSequence::GetSpaceBasedAnimationData(TArray< TArray<FTransform> >& An
 				for (int32 Key = 0; Key < NumKeys; ++Key)
 				{
 					FTransform AnimatedLocalKey;
-					UE::Anim::GetBoneTransformFromModel(DataModel, AnimatedLocalKey, TrackIndex, Interval * Key, Interpolation);
+					UE::Anim::GetBoneTransformFromModel(DataModelInterface.GetInterface(), AnimatedLocalKey, TrackIndex, Interval * Key, Interpolation);
 
 					if (ParentBoneIndex != INDEX_NONE)
 					{
@@ -3501,10 +3463,10 @@ bool UAnimSequence::ConvertAnimationDataToRiggingData(FAnimSequenceTrackContaine
 		TArray<FName> ValidNodeNames;
 		int32 NumNodes = MySkeleton->GetMappedValidNodes(ValidNodeNames);
 		TArray< TArray<FTransform> > AnimationDataInComponentSpace;
-		int32 NumBones = GetSpaceBasedAnimationData(AnimationDataInComponentSpace, NULL);
+		int32 NumBones = GetSpaceBasedAnimationData(AnimationDataInComponentSpace, nullptr);
 
 		ValidateModel();
-		const int32 NumModelKeys = DataModel->GetNumberOfKeys();
+		const int32 NumModelKeys = DataModelInterface->GetNumberOfKeys();
 
 		if (NumBones > 0)
 		{
@@ -3661,7 +3623,7 @@ bool UAnimSequence::ConvertRiggingDataToAnimationData(FAnimSequenceTrackContaine
 		TArray<FQuat4f> RotKeys;
 		TArray<FVector3f> ScaleKeys;
 
-		const int32 NumModelKeys = DataModel->GetNumberOfKeys();
+		const int32 NumModelKeys = DataModelInterface->GetNumberOfKeys();
 		const FReferenceSkeleton& RefSkeleton = MySkeleton->GetReferenceSkeleton();
 		const URig* Rig = MySkeleton->GetRig();
 		for (int32 NodeIndex = 0; NodeIndex < ValidNumNodes; ++NodeIndex)
@@ -3718,7 +3680,7 @@ void UAnimSequence::BakeTrackCurvesToRawAnimationTracks(TArray<FRawAnimSequenceT
 {
 	ValidateModel();
 
-	if (DataModel->GetNumberOfTransformCurves() > 0)
+	if (DataModelInterface->GetNumberOfTransformCurves() > 0)
 	{
 		USkeleton* CurSkeleton = GetSkeleton();
 		check(CurSkeleton);
@@ -3766,15 +3728,14 @@ void UAnimSequence::BakeTrackCurvesToRawAnimationTracks(TArray<FRawAnimSequenceT
 
 
 		// NumFrames can't be zero (filtered earlier)
-		const float Interval = GetIntervalPerKey(NumberOfSampledKeys, GetPlayLength());
 
 		const FSmartNameMapping* NameMapping = CurSkeleton->GetSmartNameContainer(USkeleton::AnimTrackCurveMappingName);
-		const FAnimationCurveData& AnimationCurveData = DataModel->GetCurveData();
+		const FAnimationCurveData& AnimationCurveData = DataModelInterface->GetCurveData();
 		for (const FTransformCurve& Curve : AnimationCurveData.TransformCurves)
 		{
 			// find curves first, and then see what is index of this curve
 			FName BoneName;
-
+			NameMapping->GetName(Curve.Name.UID, BoneName);
 			if (Curve.GetCurveTypeFlag(AACF_Disabled) == false &&
 				ensureAlways(NameMapping->GetName(Curve.Name.UID, BoneName)))
 			{
@@ -3835,11 +3796,13 @@ void UAnimSequence::BakeTrackCurvesToRawAnimationTracks(TArray<FRawAnimSequenceT
 				for (int32 KeyIndex = 0; KeyIndex < NumberOfSampledKeys; ++KeyIndex)
 				{
 					// now evaluate
-					FTransform AdditiveTransform = Curve.Evaluate(KeyIndex * Interval, 1.0);
-						FTransform LocalTransform(FQuat(RawTrack.RotKeys[KeyIndex]), FVector(RawTrack.PosKeys[KeyIndex]), FVector(RawTrack.ScaleKeys[KeyIndex]));
-						RawTrack.RotKeys[KeyIndex] = FQuat4f(LocalTransform.GetRotation() * AdditiveTransform.GetRotation());
-						RawTrack.PosKeys[KeyIndex] = FVector3f(LocalTransform.TransformPosition(AdditiveTransform.GetTranslation()));
-						RawTrack.ScaleKeys[KeyIndex] = FVector3f(LocalTransform.GetScale3D() * AdditiveTransform.GetScale3D());
+					const double EvaluationTime = DataModelInterface->GetFrameRate().AsSeconds(KeyIndex);
+					FTransform AdditiveTransform = Curve.Evaluate(EvaluationTime, 1.0);
+					FTransform LocalTransform(FQuat(RawTrack.RotKeys[KeyIndex]), FVector(RawTrack.PosKeys[KeyIndex]), FVector(RawTrack.ScaleKeys[KeyIndex]));
+
+					RawTrack.RotKeys[KeyIndex] = FQuat4f(LocalTransform.GetRotation() * AdditiveTransform.GetRotation());
+					RawTrack.PosKeys[KeyIndex] = FVector3f(LocalTransform.TransformPosition(AdditiveTransform.GetTranslation()));
+					RawTrack.ScaleKeys[KeyIndex] = FVector3f(LocalTransform.GetScale3D() * AdditiveTransform.GetScale3D());
 				}
 
 				// Apply PostProcess behaviour in-place
@@ -3865,9 +3828,7 @@ void UAnimSequence::AddKeyToSequence(float Time, const FName& BoneName, const FT
 	FAnimationCurveIdentifier TransformCurveId(NewCurveName, ERawCurveTrackTypes::RCT_Transform);
 	Controller->AddCurve(TransformCurveId, AACF_DriveTrack | AACF_Editable);
 
-	Controller->RemoveBoneTracksMissingFromSkeleton(CurrentSkeleton);
-
-	const FTransformCurve* TransformCurve = DataModel->FindTransformCurve(TransformCurveId);
+	const FTransformCurve* TransformCurve = DataModelInterface->FindTransformCurve(TransformCurveId);
 	check(TransformCurve);
 
 	Controller->SetTransformCurveKey(TransformCurveId, Time, AdditiveTransform);
@@ -3875,13 +3836,13 @@ void UAnimSequence::AddKeyToSequence(float Time, const FName& BoneName, const FT
 
 uint8* UAnimSequence::FindSyncMarkerPropertyData(int32 SyncMarkerIndex, FArrayProperty*& ArrayProperty)
 {
-	ArrayProperty = NULL;
+	ArrayProperty = nullptr;
 
 	if (AuthoredSyncMarkers.IsValidIndex(SyncMarkerIndex))
 	{
 		return FindArrayProperty(TEXT("AuthoredSyncMarkers"), ArrayProperty, SyncMarkerIndex);
 	}
-	return NULL;
+	return nullptr;
 }
 
 bool UAnimSequence::CreateAnimation(USkeletalMesh* Mesh)
@@ -3895,8 +3856,8 @@ bool UAnimSequence::CreateAnimation(USkeletalMesh* Mesh)
 		IAnimationDataController::FScopedBracket ScopedBracket(Controller, LOCTEXT("CreateAnimationSkeletalMesh_Bracket", "Creating Animation Sequence based up Skeletal Mesh"));
 		ResetAnimation();
 
-		Controller->SetPlayLength(MINIMUM_ANIMATION_LENGTH);
-		Controller->SetFrameRate(FFrameRate(1.f / MINIMUM_ANIMATION_LENGTH, 1));
+		Controller->SetFrameRate(UAnimationSettings::Get()->GetDefaultFrameRate());
+		Controller->SetNumberOfFrames(1);
 
 		const int32 NumBones = RefSkeleton.GetRawBoneNum();
 
@@ -3932,8 +3893,8 @@ bool UAnimSequence::CreateAnimation(USkeletalMeshComponent* MeshComponent)
 		IAnimationDataController::FScopedBracket ScopedBracket(Controller, LOCTEXT("CreateAnimationSkeletalMeshComponent_Bracket", "Creating Animation Sequence based up Skeletal Mesh Component"));
 		ResetAnimation();
 
-		Controller->SetPlayLength(MINIMUM_ANIMATION_LENGTH);
-		Controller->SetFrameRate(FFrameRate(1.f / MINIMUM_ANIMATION_LENGTH, 1));
+		Controller->SetFrameRate(UAnimationSettings::Get()->GetDefaultFrameRate());
+		Controller->SetNumberOfFrames(1);
 
 		const int32 NumBones = RefSkeleton.GetRawBoneNum();
 
@@ -3961,8 +3922,8 @@ bool UAnimSequence::CreateAnimation(UAnimSequence* Sequence)
 	{
 		ResetAnimation();
 
-		CopyDataModel(Sequence->GetDataModel());
-		Controller->SetModel(DataModel);
+		CopyDataModel(Sequence->GetDataModelInterface());
+		Controller->SetModel(DataModelInterface);
 		Notifies = Sequence->Notifies;
 		AnimNotifyTracks = Sequence->AnimNotifyTracks;
 
@@ -4029,7 +3990,7 @@ void UAnimSequence::EvaluateCurveData(FBlendedCurve& OutCurve, float CurrentTime
 		// Evaluate float curves from the AnimationData Model
 		if (OutCurve.NumValidCurveCount > 0)
 		{
-			for (auto CurveIter = DataModel->GetFloatCurves().CreateConstIterator(); CurveIter; ++CurveIter)
+			for (auto CurveIter = DataModelInterface->GetFloatCurves().CreateConstIterator(); CurveIter; ++CurveIter)
 			{
 				const FFloatCurve& Curve = *CurveIter;
 				if (OutCurve.IsEnabled(Curve.Name.UID))
@@ -4384,7 +4345,11 @@ void UAnimSequence::ValidateCurrentPosition(const FMarkerSyncAnimPosition& Posit
 
 bool UAnimSequence::UseRawDataForPoseExtraction(const FBoneContainer& RequiredBones) const
 {
-	return CanEvaluateRawAnimationData() && (bUseRawDataOnly || (GetSkeletonVirtualBoneGuid() != GetSkeleton()->GetVirtualBoneGuid()) || RequiredBones.GetDisableRetargeting() || RequiredBones.ShouldUseRawData() || RequiredBones.ShouldUseSourceData());
+	return CanEvaluateRawAnimationData() && (bUseRawDataOnly || (GetSkeletonVirtualBoneGuid() != GetSkeleton()->GetVirtualBoneGuid()) || RequiredBones.GetDisableRetargeting() || RequiredBones.ShouldUseRawData() ||
+ #if WITH_EDITOR
+  	GForceRawData == 1 ||
+ #endif // WITH_EDITOR
+	RequiredBones.ShouldUseSourceData());
 }
 
 #if WITH_EDITOR
@@ -4415,13 +4380,13 @@ void UAnimSequence::EvaluateAttributes(FAnimationPoseData& OutAnimationPoseData,
 	if (bUseRawData)
 	{
 		ValidateModel();
-		for (const FAnimatedBoneAttribute& Attribute : DataModel->GetAttributes())
+		for (const FAnimatedBoneAttribute& Attribute : DataModelInterface->GetAttributes())
 		{
 			const FCompactPoseBoneIndex PoseBoneIndex = RequiredBones.GetCompactPoseIndexFromSkeletonIndex(Attribute.Identifier.GetBoneIndex());
 			// Only add attribute if the bone its tied to exists in the currently evaluated set of bones
 			if(PoseBoneIndex.IsValid())
 			{
-				UE::Anim::Attributes::GetAttributeValue(OutAttributes, PoseBoneIndex, Attribute, ExtractionContext);
+				UE::Anim::Attributes::GetAttributeValue(OutAttributes, PoseBoneIndex, Attribute, ExtractionContext.CurrentTime);
 			}
 		}
 	}
@@ -4449,7 +4414,7 @@ void UAnimSequence::RemoveCustomAttribute(const FName& BoneName, const FName& At
 
 	TArray<FAnimationAttributeIdentifier> IdentifiersToRemove;
 
-	Algo::TransformIf(DataModel->GetAttributes(), IdentifiersToRemove, 
+	Algo::TransformIf(DataModelInterface->GetAttributes(), IdentifiersToRemove, 
 		[BoneName, AttributeName](const FAnimatedBoneAttribute& Attribute)
 		{
 			return Attribute.Identifier.GetBoneName() == BoneName && Attribute.Identifier.GetName() == AttributeName;
@@ -4522,7 +4487,7 @@ void UAnimSequence::SynchronousAnimatedBoneAttributesCompression()
 		// Helper struct to match sample timings with regular additive baking
 		FByFramePoseEvalContext EvalContext(this);
 
-		for (const FAnimatedBoneAttribute& AdditiveAttribute : DataModel->GetAttributes())
+		for (const FAnimatedBoneAttribute& AdditiveAttribute : DataModelInterface->GetAttributes())
 		{
 			ensure(!AttributeCurves.Contains(AdditiveAttribute.Identifier));
 			FAttributeCurve& AttributeCurve = AttributeCurves.Add(AdditiveAttribute.Identifier);
@@ -4531,10 +4496,9 @@ void UAnimSequence::SynchronousAnimatedBoneAttributesCompression()
 			{
 				AttributeCurve.SetScriptStruct(AdditiveAttribute.Identifier.GetType());
 				const FAnimatedBoneAttribute& RefAttribute = *RefPoseAttributePtr;
-				for (int32 Frame = 0; Frame < NumberOfSampledKeys; ++Frame)
+				for (int32 KeyIndex = 0; KeyIndex < NumberOfSampledKeys; ++KeyIndex)
 				{
-					const float CurrentFrameTime = Frame * EvalContext.IntervalTime;
-
+					const double CurrentFrameTime = TargetFrameRate.AsSeconds(KeyIndex);
 					UE::Anim::FStackAttributeContainer AdditiveAttributes;
 					uint8* AdditivePtr = AdditiveAttributes.FindOrAdd(AdditiveAttribute.Identifier.GetType(), UE::Anim::FAttributeId(NAME_None, 0,NAME_None));
 					AdditiveAttribute.Curve.EvaluateToPtr(AdditiveAttribute.Identifier.GetType(), CurrentFrameTime, AdditivePtr);
@@ -4556,7 +4520,7 @@ void UAnimSequence::SynchronousAnimatedBoneAttributesCompression()
 	}
 	else
 	{
-		for (const FAnimatedBoneAttribute& Attribute : DataModel->GetAttributes())
+		for (const FAnimatedBoneAttribute& Attribute : DataModelInterface->GetAttributes())
 		{
 			// Do something with the attributes
 			ensure(!AttributeCurves.Contains(Attribute.Identifier));
@@ -5150,7 +5114,7 @@ bool UAnimSequence::IsCompressedDataValid() const
 	// if we have no raw data. This can happen with sequences that only has curves.
 #if WITH_EDITOR
 	PRAGMA_DISABLE_DEPRECATION_WARNINGS
-	if (IsDataModelValid() && DataModel->GetNumBoneTracks() == 0 && RawAnimationData.Num() == 0)
+	if (IsDataModelValid() && DataModelInterface->GetNumBoneTracks() == 0 && RawAnimationData.Num() == 0)
 	PRAGMA_ENABLE_DEPRECATION_WARNINGS
 	{
 		return true;
@@ -5167,7 +5131,7 @@ bool UAnimSequence::IsCurveCompressedDataValid() const
 	// if we have no raw data. This can happen with sequences that only has bones.
 
 	const bool bNoCompressedCurveData = CompressedData.CompressedCurveByteStream.Num() == 0;
-	const bool bHasRawCurveData = IsDataModelValid() && DataModel->GetNumberOfFloatCurves() != 0;
+	const bool bHasRawCurveData = IsDataModelValid() && DataModelInterface->GetNumberOfFloatCurves() != 0;
 
 	if ( bNoCompressedCurveData && bHasRawCurveData )
 	{
@@ -5201,7 +5165,7 @@ void UAnimSequence::EnsureValidRawDataGuid()
 PRAGMA_DISABLE_DEPRECATION_WARNINGS
 	if (IsDataModelValid() && !RawDataGuid.IsValid())
 	{
-		RawDataGuid = GenerateGuidFromRawData();
+		RawDataGuid = DataModelInterface->GenerateGuid();
 	}
 PRAGMA_ENABLE_DEPRECATION_WARNINGS
 }
@@ -5226,7 +5190,7 @@ void UAnimSequence::PopulateModel()
 	const bool bValidNumberOfKeys = NumberOfKeys >= 1 || bCurveOnlySequence;
 	const float ExpectedPlayLength = (float)SamplingFrameRate.AsSeconds(FMath::Max(NumberOfKeys - 1, 1));
 	const bool bValidSequenceLength = FMath::IsNearlyEqual(ExpectedPlayLength, SequenceLength, UE_KINDA_SMALL_NUMBER) || bCurveOnlySequence;
-	const bool bValidNumberOfTrackArrayEntries = AnimationTrackNames.Num() == RawAnimationData.Num() && RawAnimationData.Num() == TrackToSkeletonMapTable.Num();
+	const bool bValidNumberOfTrackArrayEntries = AnimationTrackNames.Num() == RawAnimationData.Num();
 	const bool bValidAnimationData = SourceRawAnimationData_DEPRECATED.Num() == 0 || SourceRawAnimationData_DEPRECATED.Num() == RawAnimationData.Num();
 	const bool bValidFrameRate = SamplingFrameRate.IsValid() && (SamplingFrameRate.AsInterval() >= 0.f);
 	const bool bValidCurveData = RawCurveData.FloatCurves.Num() >= 1;
@@ -5244,68 +5208,83 @@ void UAnimSequence::PopulateModel()
 
 	Controller->OpenBracket(LOCTEXT("UAnimSequence::PopulateModel_Bracket", "Generating Animation Model Data from Animation Sequence"));
 
-	Controller->SetPlayLength(bCurveOnlySequence ? PlayLength : ExpectedPlayLength);
 	Controller->SetFrameRate(FrameRate);
 
-	USkeleton* TargetSkeleton = GetSkeleton();
+	const FFrameTime FrameTime = FrameRate.AsFrameTime(bCurveOnlySequence ? PlayLength : ExpectedPlayLength);
+	const FFrameNumber RoundedFrameNumber = FrameTime.RoundToFrame();
+	ensureAlways(RoundedFrameNumber.Value == FMath::Max((NumKeys - 1), 1));
+	Controller->SetNumberOfFrames(RoundedFrameNumber);	
+
+	const USkeleton* TargetSkeleton = GetSkeleton();
 	UE::Anim::CopyCurveDataToModel(CurveData, TargetSkeleton, *Controller);
 	
 	const int32 NumTracks = SequenceTracks.Num();
 	for (int32 TrackIndex = 0; TrackIndex < NumTracks; ++TrackIndex)
 	{
-		Controller->AddBoneTrack(TempAnimationTrackNames[TrackIndex]);
-
-		auto GenerateUniformKeys = [NumKeys](auto& Keys, auto IdentityKey)
+		if (SequenceTracks[TrackIndex].ContainsNaN())
 		{
-			// Convert track keys to be uniform			
-			if (Keys.Num() == 0)
+			continue;
+		}
+		
+		const int32 BoneTrackIndex = Controller->AddBoneTrack(TempAnimationTrackNames[TrackIndex]);
+		if (BoneTrackIndex != INDEX_NONE)
+		{
+			auto GenerateUniformKeys = [NumKeys](auto& Keys, auto IdentityKey)
 			{
-				// set all to identity
-				for (int32 Index = 0; Index < NumKeys; ++Index)
+				// Convert track keys to be uniform
+				if (Keys.Num() == 0)
 				{
-					Keys.Add(IdentityKey);
+					// set all to identity
+					for (int32 Index = 0; Index < NumKeys; ++Index)
+					{
+						Keys.Add(IdentityKey);
+					}
 				}
-			}
-			else if (Keys.Num() == 1)
-			{
-				// set all to single key
-				const auto KeyZero = Keys[0];
-				for (int32 Index = 0; Index < (NumKeys - 1); ++Index)
+				else if (Keys.Num() == 1)
 				{
-					Keys.Add(KeyZero);
+					// set all to single key
+					const auto KeyZero = Keys[0];
+					for (int32 Index = 0; Index < (NumKeys - 1); ++Index)
+					{
+						Keys.Add(KeyZero);
+					}
 				}
-			}
 
-			ensure(Keys.Num() == NumKeys);
-		};
+				ensure(Keys.Num() == NumKeys);
+			};
 
-		TArray<FVector3f> PosKeys = SequenceTracks[TrackIndex].PosKeys;
-		GenerateUniformKeys(PosKeys, FVector3f::ZeroVector);
+			TArray<FVector3f> PosKeys = SequenceTracks[TrackIndex].PosKeys;
+			GenerateUniformKeys(PosKeys, FVector3f::ZeroVector);
 
-		TArray<FQuat4f> RotKeys = SequenceTracks[TrackIndex].RotKeys;
-		GenerateUniformKeys(RotKeys, FQuat4f::Identity);
+			TArray<FQuat4f> RotKeys = SequenceTracks[TrackIndex].RotKeys;
+			GenerateUniformKeys(RotKeys, FQuat4f::Identity);
 
-		TArray<FVector3f> ScaleKeys = SequenceTracks[TrackIndex].ScaleKeys;
-		GenerateUniformKeys(ScaleKeys, FVector3f::OneVector);
+			TArray<FVector3f> ScaleKeys = SequenceTracks[TrackIndex].ScaleKeys;
+			GenerateUniformKeys(ScaleKeys, FVector3f::OneVector);
 
-		Controller->SetBoneTrackKeys(TempAnimationTrackNames[TrackIndex], PosKeys, RotKeys, ScaleKeys);
+			Controller->SetBoneTrackKeys(TempAnimationTrackNames[TrackIndex], PosKeys, RotKeys, ScaleKeys);
+		}
 	}
 
 	DeleteDeprecatedRawAnimationData();
 	
+	Controller->NotifyPopulated();
+	
 	Controller->CloseBracket();
 }
 
-void UAnimSequence::OnModelModified(const EAnimDataModelNotifyType& NotifyType, UAnimDataModel* Model, const FAnimDataModelNotifPayload& Payload)
+void UAnimSequence::OnModelModified(const EAnimDataModelNotifyType& NotifyType, IAnimationDataModel* Model, const FAnimDataModelNotifPayload& Payload)
 {
-	UAnimSequenceBase::OnModelModified(NotifyType, Model, Payload);	
+	UAnimSequenceBase::OnModelModified(NotifyType, Model, Payload);
 
 	static const bool GenerateNewGUID = true;
 	static const bool RegenerateGUID = false;
 
-	auto UpdateRawDataGuid = [this](bool bForceNewRawDataGuid)
+	auto UpdateRawDataGuid = [this, Model](bool bForceNewRawDataGuid)
 	{
-		RawDataGuid = bForceNewRawDataGuid ? FGuid::NewGuid() : GenerateGuidFromRawData();
+		PRAGMA_DISABLE_DEPRECATION_WARNINGS
+		RawDataGuid = bForceNewRawDataGuid ? FGuid::NewGuid() : Model->GenerateGuid();
+		PRAGMA_ENABLE_DEPRECATION_WARNINGS
 		bUseRawDataOnly = true;
 
 		FlagDependentAnimationsAsRawDataOnly();
@@ -5316,20 +5295,26 @@ void UAnimSequence::OnModelModified(const EAnimDataModelNotifyType& NotifyType,
 	{
 		if (NotifyCollector.IsNotWithinBracket())
 		{
-			ResampleAnimationTrackData();
+			ClearResampledAnimationTrackData();
+			CalculateNumberOfSampledKeys();
 			UpdateRawDataGuid(RegenerateGUID);
 			RecompressAnimationData();
 		}
 	};
 
-	auto HandleLengthChanged = [this](float NewLength, float OldLength, float T0, float T1)
+	auto HandleNumberOfFramesChanged = [this](FFrameNumber NewNumberOfFrames, FFrameNumber OldNumberOfFrames, FFrameNumber Frame0, FFrameNumber Frame1)
 	{
 		if (bPopulatingDataModel)
 		{
 			return;
 		}
+		const FFrameRate& ModelFrameRate = DataModelInterface->GetFrameRate();
+
+		const float NewLength = ModelFrameRate.AsSeconds(NewNumberOfFrames);
+		const float T0 = ModelFrameRate.AsSeconds(Frame0);
+		const float T1 = ModelFrameRate.AsSeconds(Frame1);
 
-		if (NewLength > OldLength)
+		if (NewNumberOfFrames > OldNumberOfFrames)
 		{
 			const float InsertTime = T0;
 			const float Duration = T1 - T0;
@@ -5348,7 +5333,7 @@ void UAnimSequence::OnModelModified(const EAnimDataModelNotifyType& NotifyType,
 				Marker.Time = FMath::Clamp(CurrentTime, 0.f, NewLength);
 			}
 		}
-		else if (NewLength < OldLength)
+		else if (NewNumberOfFrames < OldNumberOfFrames)
 		{
 			const float StartRemoveTime = T0;
 			const float EndRemoveTime = T1;
@@ -5371,24 +5356,21 @@ void UAnimSequence::OnModelModified(const EAnimDataModelNotifyType& NotifyType,
 
 				Marker.Time = FMath::Clamp(CurrentTime, 0.f, NewLength);
 			}
-
 		}
 	};
 
 	bool bShouldMarkPackageDirty = !FUObjectThreadContext::Get().IsRoutingPostLoad &&
 		NotifyType != EAnimDataModelNotifyType::BracketOpened && NotifyType != EAnimDataModelNotifyType::BracketClosed;
 
-	PRAGMA_DISABLE_DEPRECATION_WARNINGS
 	switch (NotifyType)
 	{
-		case EAnimDataModelNotifyType::SequenceLengthChanged:		
+		case EAnimDataModelNotifyType::SequenceLengthChanged:
 		{
 			const FSequenceLengthChangedPayload& TypedPayload = Payload.GetPayload<FSequenceLengthChangedPayload>();
 
-			const float OldSequenceLength = TypedPayload.PreviousLength;
-			const float CurrentSequenceLength = Model->GetPlayLength();
-			
-			HandleLengthChanged(CurrentSequenceLength, OldSequenceLength, TypedPayload.T0, TypedPayload.T1);
+			const FFrameNumber OldNumberOfFrames = TypedPayload.PreviousNumberOfFrames;
+			const FFrameNumber CurrentNumberOfFrames = Model->GetNumberOfFrames();
+			HandleNumberOfFramesChanged(CurrentNumberOfFrames, OldNumberOfFrames, TypedPayload.Frame0, TypedPayload.Frame1);
 						
 			if (NotifyCollector.IsNotWithinBracket())
 			{
@@ -5432,7 +5414,8 @@ void UAnimSequence::OnModelModified(const EAnimDataModelNotifyType& NotifyType,
 				
 				if (NotifyCollector.Contains(LengthChangingNotifies) || NotifyCollector.Contains(ResamplingNotifies))
 				{
-					ResampleAnimationTrackData();
+					ClearResampledAnimationTrackData();
+					CalculateNumberOfSampledKeys();
 				}
 
 				if (NotifyCollector.Contains(LengthChangingNotifies) || NotifyCollector.Contains(ResamplingNotifies) || NotifyCollector.Contains(RecompressNotifies))
@@ -5456,7 +5439,8 @@ void UAnimSequence::OnModelModified(const EAnimDataModelNotifyType& NotifyType,
 		{
 			if (NotifyCollector.IsNotWithinBracket())
 			{
-				ResampleAnimationTrackData();
+				ClearResampledAnimationTrackData();
+				CalculateNumberOfSampledKeys();
 				UpdateRawDataGuid(GenerateNewGUID);
 				RecompressAnimationData();
 			}
@@ -5525,60 +5509,99 @@ void UAnimSequence::OnModelModified(const EAnimDataModelNotifyType& NotifyType,
 			break;
 		}
 	}
-	PRAGMA_ENABLE_DEPRECATION_WARNINGS
 
-	if (bShouldMarkPackageDirty)
+	if (NotifyCollector.IsNotWithinBracket() && NotifyType != EAnimDataModelNotifyType::BracketClosed && bShouldMarkPackageDirty)
 	{
 		MarkPackageDirty();
 	}
 }
 
-
-void UAnimSequence::ResampleAnimationTrackData()
+void UAnimSequence::CalculateNumberOfSampledKeys()
 {
-	if (ShouldDataModelBeValid())
-	{
-	ValidateModel();
-
-	// Make a copy, deals with bone name and index
-	ResampledAnimationTrackData = DataModel->GetBoneAnimationTracks();
-
-	if (TargetFrameRate == DataModel->GetFrameRate())
+	if (TargetFrameRate == DataModelInterface->GetFrameRate())
 	{
-		// Can directly copy
-		NumberOfSampledKeys = DataModel->GetNumberOfKeys();
-		NumberOfSampledFrames = DataModel->GetNumberOfFrames();
+		NumberOfSampledKeys = DataModelInterface->GetNumberOfKeys();
+		NumberOfSampledFrames = DataModelInterface->GetNumberOfFrames();
 	}
 	else
 	{
-		// We actually need to resample bone transforms
-		NumberOfSampledFrames = TargetFrameRate.AsFrameNumber(DataModel->GetPlayLength()).Value;
+		const FFrameNumber ModelNumberOfFrames = DataModelInterface->GetNumberOfFrames();
+		const FFrameTime ResampledFrameTime = FFrameRate::TransformTime(ModelNumberOfFrames, DataModelInterface->GetFrameRate(), TargetFrameRate);
+		ensureMsgf(FMath::IsNearlyZero(ResampledFrameTime.GetSubFrame()), TEXT("Incompatible resampling frame rate for animation sequence %s, frame remainder of %1.8f"), *GetName(), ResampledFrameTime.GetSubFrame());
+
+		NumberOfSampledFrames = ResampledFrameTime.FloorToFrame().Value;
 		NumberOfSampledKeys = NumberOfSampledFrames + 1;
-				
-		for (int32 TrackIndex = 0; TrackIndex < ResampledAnimationTrackData.Num(); ++TrackIndex)	
-		{
-			FBoneAnimationTrack& TrackData = ResampledAnimationTrackData[TrackIndex];
+	}
+}
+
+
+void UAnimSequence::ClearResampledAnimationTrackData()
+{
+	ResampledAnimationTrackData.Empty();
+}
 
-			TrackData.InternalTrackData.PosKeys.SetNumUninitialized(NumberOfSampledKeys);
-			TrackData.InternalTrackData.RotKeys.SetNumUninitialized(NumberOfSampledKeys);
-			TrackData.InternalTrackData.ScaleKeys.SetNumUninitialized(NumberOfSampledKeys);
+void UAnimSequence::ResampleAnimationTrackData()
+{
+	QUICK_SCOPE_CYCLE_COUNTER(STAT_ResampleAnimationTrackData);
+	
+	if (ShouldDataModelBeValid())
+	{
+		ValidateModel();
+
+		FAnimationUtils::EnsureAnimSequenceLoaded(*this);
 
-			for (int32 FrameIndex = 0; FrameIndex < NumberOfSampledKeys; ++FrameIndex)
+		// Make a copy, deals with bone name and index
+		ResampledAnimationTrackData = DataModelInterface->GetBoneAnimationTracks();
+		CalculateNumberOfSampledKeys();
+	
+		ensure(DataModelInterface != nullptr);
+		{
+			// Can directly copy (for now, as we still maintain GetBoneAnimationTracks)
+			if (TargetFrameRate != DataModelInterface->GetFrameRate())
 			{
-				FTransform Transform;
+				FMemMark Mark(FMemStack::Get());
+				FByFramePoseEvalContext EvalContext(this);
+				EvalContext.RequiredBones.SetDisableRetargeting(true);
+				EvalContext.RequiredBones.SetUseRAWData(true);
+				EvalContext.RequiredBones.SetUseSourceData(true);
+
+				// We actually need to resample bone transforms
+				for (int32 TrackIndex = 0; TrackIndex < ResampledAnimationTrackData.Num(); ++TrackIndex)
+				{
+					FBoneAnimationTrack& TrackData = ResampledAnimationTrackData[TrackIndex];
+					TrackData.InternalTrackData.PosKeys.SetNumUninitialized(NumberOfSampledKeys);
+					TrackData.InternalTrackData.RotKeys.SetNumUninitialized(NumberOfSampledKeys);
+					TrackData.InternalTrackData.ScaleKeys.SetNumUninitialized(NumberOfSampledKeys);
+				}
+				
+				FBlendedCurve Curve;
+				Curve.InitFrom(EvalContext.RequiredBones);
+				UE::Anim::FStackAttributeContainer AttributeContainer;
+				
+				for (int32 FrameIndex = 0; FrameIndex < NumberOfSampledKeys; ++FrameIndex)
+				{
+					UE::Anim::DataModel::FEvaluationContext EvaluationContext(FFrameTime(FrameIndex), TargetFrameRate, GetRetargetTransformsSourceName(), GetRetargetTransforms());
+					
+					FCompactPose Pose;
+					Pose.SetBoneContainer(&EvalContext.RequiredBones);
 
-				const double SampleTime = TargetFrameRate.AsSeconds(FrameIndex);
+					FAnimationPoseData PoseData(Pose, Curve, AttributeContainer);
+					DataModelInterface->Evaluate(PoseData, EvaluationContext);
 
-				UE::Anim::GetBoneTransformFromModel(DataModel, Transform, TrackIndex, SampleTime, EAnimInterpolationType::Linear);
+					for (int32 TrackIndex = 0; TrackIndex < ResampledAnimationTrackData.Num(); ++TrackIndex)
+					{
+						FBoneAnimationTrack& TrackData = ResampledAnimationTrackData[TrackIndex];
+						const FTransform& Transform = Pose[EvalContext.RequiredBones.GetCompactPoseIndexFromSkeletonIndex(TrackData.BoneTreeIndex)];
 
-				TrackData.InternalTrackData.PosKeys[FrameIndex] = FVector3f(Transform.GetLocation());
-				TrackData.InternalTrackData.RotKeys[FrameIndex] = FQuat4f(Transform.GetRotation());
-				TrackData.InternalTrackData.ScaleKeys[FrameIndex] = FVector3f(Transform.GetScale3D());
+						TrackData.InternalTrackData.PosKeys[FrameIndex] = FVector3f(Transform.GetLocation());
+						TrackData.InternalTrackData.RotKeys[FrameIndex] = FQuat4f(Transform.GetRotation());
+						TrackData.InternalTrackData.ScaleKeys[FrameIndex] = FVector3f(Transform.GetScale3D());
+					}
+				}
 			}
 		}
 	}
 }
-}
 
 void UAnimSequence::RecompressAnimationData()
 {
