Reintroducing AnimationDataModel refactor with additional memory and cooktime optimizations.

diff --git a/Engine/Source/Runtime/Engine/Private/Animation/AnimSequenceBase.cpp b/Engine/Source/Runtime/Engine/Private/Animation/AnimSequenceBase.cpp
--- a/Engine/Source/Runtime/Engine/Private/Animation/AnimSequenceBase.cpp
+++ b/Engine/Source/Runtime/Engine/Private/Animation/AnimSequenceBase.cpp
@@ -14,6 +14,7 @@
 #include "Animation/AnimationPoseData.h"
 #include "Animation/AnimSequenceHelpers.h"
 #include "Animation/MirrorDataTable.h"
+#include "Animation/AnimData/IAnimationDataModel.h"
 #include "Animation/AnimData/AnimDataModel.h"
 #include "Modules/ModuleManager.h"
 #include "MathUtil.h"
@@ -23,6 +24,7 @@
 #include "Modules/ModuleManager.h"
 #endif // WITH_EDITOR
 
+#include "Animation/AnimationSettings.h"
 #include "UObject/UE5MainStreamObjectVersion.h"
 
 DEFINE_LOG_CATEGORY(LogAnimMarkerSync);
@@ -50,19 +52,6 @@ UAnimSequenceBase::UAnimSequenceBase(const FObjectInitializer& ObjectInitializer
 #endif // WITH_EDITOR
 }
 
-void UAnimSequenceBase::SetSequenceLength(float NewLength)
-{
-#if WITH_EDITOR
-	// Editor time we update the source data
-	Controller->SetPlayLength(NewLength);
-#else 
-	// In case this is called during runtime, update the sequence length directly. Current use-case is runtime created Montages.
-	PRAGMA_DISABLE_DEPRECATION_WARNINGS
-	SequenceLength = NewLength;
-	PRAGMA_ENABLE_DEPRECATION_WARNINGS
-#endif
-}
-
 bool UAnimSequenceBase::IsPostLoadThreadSafe() const
 {
 	return false;	// PostLoad is not thread safe because of the call to VerifyCurveNames() (calling USkeleton::VerifySmartName) that can mutate a shared map in the skeleton.
@@ -73,6 +62,8 @@ void UAnimSequenceBase::PostLoad()
 #if WITH_EDITORONLY_DATA
 	if (!HasAnyFlags(EObjectFlags::RF_ClassDefaultObject))
 	{
+		LLM_SCOPE(ELLMTag::Animation);
+		
 		auto PreloadSkeletonAndVerifyCurves = [this]()
 		{
 			if (USkeleton* MySkeleton = GetSkeleton())
@@ -89,31 +80,83 @@ void UAnimSequenceBase::PostLoad()
 				PRAGMA_ENABLE_DEPRECATION_WARNINGS
 			}
 		};
-		
+
 		if(ShouldDataModelBeValid())
 		{
-			const bool bRequiresModelCreation = DataModel == nullptr;
-			const bool bRequiresModelPopulation = GetLinkerCustomVersion(FUE5MainStreamObjectVersion::GUID) < FUE5MainStreamObjectVersion::IntroducingAnimationDataModel;
-			checkf(bRequiresModelPopulation || DataModel != nullptr, TEXT("Invalid Animation Sequence base state, no data model found past upgrade object version"));
-
-			// Construct a new UAnimDataModel instance
-			if(bRequiresModelCreation)
-			{
-				CreateModel();
-			}
-
-			ValidateModel();
-			GetController();
-			BindToModelModificationEvent();
-
-			PreloadSkeletonAndVerifyCurves();
-
-			GetController();
-			if (bRequiresModelPopulation)
-			{
-				bPopulatingDataModel = true;
-				PopulateModel();
-				bPopulatingDataModel = false;
+		    const UClass* TargetDataModelClass = UE::Anim::DataModel::IAnimationDataModels::FindClassForAnimationAsset(this);
+		    const bool bRequiresModelCreation = DataModelInterface == nullptr || DataModelInterface.GetObject()->GetClass() != TargetDataModelClass ||  GetLinkerCustomVersion(FUE5MainStreamObjectVersion::GUID) < FUE5MainStreamObjectVersion::ReintroduceAnimationDataModelInterface;
+
+		    TScriptInterface<IAnimationDataModel> CachedDataModelInterface = DataModelInterface;
+		    
+		    const bool bRequiresModelPopulation = GetLinkerCustomVersion(FUE5MainStreamObjectVersion::GUID) < FUE5MainStreamObjectVersion::IntroducingAnimationDataModel;
+		    PRAGMA_DISABLE_DEPRECATION_WARNINGS
+		    checkf(bRequiresModelPopulation || DataModel != nullptr || DataModelInterface != nullptr, TEXT("Invalid Animation Sequence base state, no data model found past upgrade object version"));
+		    PRAGMA_ENABLE_DEPRECATION_WARNINGS
+    
+		    // Construct a new IAnimationDataModel instance
+		    if(bRequiresModelCreation)
+		    {
+			    CreateModel();
+		    }
+
+		    ValidateModel();
+		    GetController();
+		    BindToModelModificationEvent();
+
+		    PreloadSkeletonAndVerifyCurves();
+
+		    if (bRequiresModelPopulation || bRequiresModelCreation)
+		    {
+			    const bool bDoNotTransactAction = false;
+    
+		    	Controller->OpenBracket(LOCTEXT("UAnimSequenceBase::PostLoad_PopulatingModelInterface","Populating Animation Data Model Interface"), bDoNotTransactAction);
+			    
+		    	bPopulatingDataModel = true;
+		    	
+		    	Controller->InitializeModel();
+    
+		    	PRAGMA_DISABLE_DEPRECATION_WARNINGS
+				// In case a data model has already been created populate the new one with its data
+				if (DataModel != nullptr)
+				{
+					if (FLinkerLoad* DataModelLinker = DataModel->GetLinker())
+					{
+						DataModelLinker->Preload(DataModel);
+					}					
+					DataModel->PostLoad();
+					DataModel->ConditionalPostLoadSubobjects();
+					PopulateWithExistingModel(DataModel.Get());
+
+					TScriptInterface<IAnimationDataController> DataModelController = DataModel->GetController();
+					DataModelController->ResetModel(false);
+					
+					DataModel = nullptr;
+					Controller->NotifyPopulated();
+				}
+		    	PRAGMA_ENABLE_DEPRECATION_WARNINGS
+				// If switching to a different DataModelInterface implementation, copy the data from the existing one
+				else if (CachedDataModelInterface != nullptr)
+				{
+					PopulateWithExistingModel(CachedDataModelInterface);
+					Controller->NotifyPopulated();
+				}
+		    	// Otherwise upgrade this animation asset to be model-based
+		    	else
+		    	{
+					PopulateModel();
+		    		Controller->NotifyPopulated();
+		    	}
+		    	bPopulatingDataModel = false;
+		    	Controller->CloseBracket();
+		    }
+		    else
+		    {
+		    	// Fix-up to ensure correct curves are used for compression
+		    	PRAGMA_DISABLE_DEPRECATION_WARNINGS
+				RawCurveData.FloatCurves = DataModelInterface->GetCurveData().FloatCurves;
+		    	RawCurveData.TransformCurves = DataModelInterface->GetCurveData().TransformCurves;
+		    	RawCurveData.RemoveRedundantKeys(0.f);
+		    	PRAGMA_ENABLE_DEPRECATION_WARNINGS
 			}
 		}
 		else
@@ -146,18 +189,28 @@ void UAnimSequenceBase::PostLoad()
 #endif
 	RefreshCacheData();
 
+#if WITH_EDITOR
+	if (!GetPackage()->GetHasBeenEndLoaded())
+	{
+		FCoreUObjectDelegates::OnEndLoadPackage.AddUObject(this, &UAnimSequenceBase::OnEndLoadPackage);
+	}
+	else
+	{
+		OnAnimModelLoaded();
+	}	
+#endif
+
 	if(USkeleton* MySkeleton = GetSkeleton())
 	{
-		const bool bDoNotTransactAction = false;
 #if WITH_EDITOR
-		if (IsDataModelValid())
+		if (ensure(IsDataModelValid()))
 		{
+			const bool bDoNotTransactAction = false;
 			if (GetLinkerCustomVersion(FFortniteMainBranchObjectVersion::GUID) < FFortniteMainBranchObjectVersion::FixUpNoneNameAnimationCurves)
 			{
-				
 				Controller->OpenBracket(LOCTEXT("FFortniteMainBranchObjectVersion::FixUpNoneNameAnimationCurves_Bracket","FFortniteMainBranchObjectVersion::FixUpNoneNameAnimationCurves"), bDoNotTransactAction);
 				{
-					const TArray<FFloatCurve>& FloatCurves = DataModel->GetFloatCurves();
+					const TArray<FFloatCurve>& FloatCurves = DataModelInterface->GetFloatCurves();
 					for (int32 Index = 0; Index < FloatCurves.Num(); ++Index)
 					{
 						const FFloatCurve& Curve = FloatCurves[Index];
@@ -176,23 +229,7 @@ void UAnimSequenceBase::PostLoad()
 				}
 				Controller->CloseBracket(bDoNotTransactAction);
 			}
-			else
-			{
-				const TArray<FFloatCurve>& FloatCurves = DataModel->GetFloatCurves();
-				for (int32 Index = 0; Index < FloatCurves.Num(); ++Index)
-				{
-					const FFloatCurve& Curve = FloatCurves[Index];
-					ensureMsgf(Curve.Name.DisplayName != NAME_None, TEXT("[AnimSequencer %s] has invalid curve name."), *GetFullName());
-				}
-			}
-
-			Controller->FindOrAddCurveNamesOnSkeleton(MySkeleton, ERawCurveTrackTypes::RCT_Float, bDoNotTransactAction);
-			Controller->FindOrAddCurveNamesOnSkeleton(MySkeleton, ERawCurveTrackTypes::RCT_Transform, bDoNotTransactAction);
 		}
-#else
-		PRAGMA_DISABLE_DEPRECATION_WARNINGS
-		VerifyCurveNames<FFloatCurve>(*MySkeleton, USkeleton::AnimCurveMappingName, RawCurveData.FloatCurves);
-		PRAGMA_ENABLE_DEPRECATION_WARNINGS
 #endif
 
 		// this should continue to add if skeleton hasn't been saved either 
@@ -200,9 +237,9 @@ void UAnimSequenceBase::PostLoad()
 		if (GetLinkerCustomVersion(FFrameworkObjectVersion::GUID) < FFrameworkObjectVersion::MoveCurveTypesToSkeleton
 			|| MySkeleton->GetLinkerCustomVersion(FFrameworkObjectVersion::GUID) < FFrameworkObjectVersion::MoveCurveTypesToSkeleton)
 		{
-
 #if WITH_EDITOR
-			const TArray<FFloatCurve>& FloatCurves = DataModel->GetFloatCurves();
+			// This is safe as the data model will have been created during this PostLoad call
+			const TArray<FFloatCurve>& FloatCurves = DataModelInterface->GetFloatCurves();
 #else
 			PRAGMA_DISABLE_DEPRECATION_WARNINGS
 			const TArray<FFloatCurve>& FloatCurves = RawCurveData.FloatCurves;
@@ -229,10 +266,8 @@ void UAnimSequenceBase::PostDuplicate(EDuplicateMode::Type DuplicateMode)
 	Super::PostDuplicate(DuplicateMode);
 
 #if WITH_EDITOR
-	if (IsDataModelValid())
-	{
-		BindToModelModificationEvent();
-	}
+	checkf(DataModelInterface.GetObject()->GetOuter() == this, TEXT("Animation Data Model interface has incorrect outer, expected %s - found %f"), *this->GetName(), *DataModelInterface.GetObject()->GetOuter()->GetName());
+	BindToModelModificationEvent();
 #endif // WITH_EDITOR
 }
 
@@ -718,9 +753,9 @@ int32 UAnimSequenceBase::GetNumberOfSampledKeys() const
 	return GetSamplingFrameRate().AsFrameTime(GetPlayLength()).RoundToFrame().Value;
 }
 
-const FFrameRate& UAnimSequenceBase::GetSamplingFrameRate() const
+FFrameRate UAnimSequenceBase::GetSamplingFrameRate() const
 {
-	static const FFrameRate DefaultFrameRate = FFrameRate(DEFAULT_SAMPLERATE, 1);
+	static const FFrameRate DefaultFrameRate = UAnimationSettings::Get()->GetDefaultFrameRate();
 	return DefaultFrameRate;
 }
 
@@ -807,9 +842,9 @@ void UAnimSequenceBase::GetAssetRegistryTags(TArray<FAssetRegistryTag>& OutTags)
 	// as possible.
 	FString CurveNameList;
 
-	if (DataModel)
+	if (DataModelInterface.GetObject())
 	{
-		for(const FFloatCurve& Curve : DataModel->GetFloatCurves())
+		for(const FFloatCurve& Curve : DataModelInterface->GetFloatCurves())
 		{
 			CurveNameList += FString::Printf(TEXT("%s%s"), *Curve.Name.DisplayName.ToString(), *USkeleton::CurveTagDelimiter);
 		}
@@ -819,14 +854,14 @@ void UAnimSequenceBase::GetAssetRegistryTags(TArray<FAssetRegistryTag>& OutTags)
 
 uint8* UAnimSequenceBase::FindNotifyPropertyData(int32 NotifyIndex, FArrayProperty*& ArrayProperty)
 {
-	// initialize to NULL
-	ArrayProperty = NULL;
+	// initialize to nullptr
+	ArrayProperty = nullptr;
 
 	if(Notifies.IsValidIndex(NotifyIndex))
 	{
 		return FindArrayProperty(TEXT("Notifies"), ArrayProperty, NotifyIndex);
 	}
-	return NULL;
+	return nullptr;
 }
 
 uint8* UAnimSequenceBase::FindArrayProperty(const TCHAR* PropName, FArrayProperty*& ArrayProperty, int32 ArrayIndex)
@@ -850,7 +885,7 @@ uint8* UAnimSequenceBase::FindArrayProperty(const TCHAR* PropName, FArrayPropert
 			return ArrayHelper.GetRawPtr(ArrayIndex);
 		}
 	}
-	return NULL;
+	return nullptr;
 }
 
 void UAnimSequenceBase::RefreshParentAssetData()
@@ -866,17 +901,17 @@ void UAnimSequenceBase::RefreshParentAssetData()
 	ValidateModel();
 	Controller->OpenBracket(LOCTEXT("RefreshParentAssetData_Bracket", "Refreshing Parent Asset Data"));
 	{
-		const UAnimDataModel* ParentDataModel = ParentSeqBase->GetDataModel();
+		const IAnimationDataModel* ParentDataModel = ParentSeqBase->GetDataModel();
 
-		if (!FMath::IsNearlyEqual(DataModel->GetPlayLength(), ParentDataModel->GetPlayLength()))
+		if (DataModelInterface->GetNumberOfFrames() != ParentDataModel->GetNumberOfFrames())
 		{
-			Controller->SetPlayLength(ParentDataModel->GetPlayLength());
+			Controller->SetNumberOfFrames(ParentDataModel->GetNumberOfFrames());
 		}
 		
 		Controller->RemoveAllCurvesOfType(ERawCurveTrackTypes::RCT_Float);
 		for (const FFloatCurve& FloatCurve : ParentDataModel->GetFloatCurves())
 		{
-			const FAnimationCurveIdentifier CurveId(FloatCurve.Name, ERawCurveTrackTypes::RCT_Float);
+			const FAnimationCurveIdentifier CurveId = UAnimationCurveIdentifierExtensions::FindCurveIdentifier(GetSkeleton(), FloatCurve.Name.DisplayName, ERawCurveTrackTypes::RCT_Float);
 			Controller->AddCurve(CurveId, FloatCurve.GetCurveTypeFlags());
 			Controller->SetCurveKeys(CurveId, FloatCurve.FloatCurve.GetConstRefOfKeys());
 		}
@@ -887,11 +922,19 @@ void UAnimSequenceBase::RefreshParentAssetData()
 			const FAnimationCurveIdentifier CurveId(TransformCurve.Name, ERawCurveTrackTypes::RCT_Transform);
 			Controller->AddCurve(CurveId, TransformCurve.GetCurveTypeFlags());
 
-			TArray<FTransform> Values;
-			TArray<float> Times;
-			TransformCurve.GetKeys(Times, Values);
-
-			Controller->SetTransformCurveKeys(CurveId, Values, Times);
+			// Set each individual channel rich curve keys, to account for any custom tangents etc.
+			for (int32 SubCurveIndex = 0; SubCurveIndex < 3; ++SubCurveIndex)
+			{
+				const ETransformCurveChannel Channel = static_cast<ETransformCurveChannel>(SubCurveIndex);
+				const FVectorCurve* VectorCurve = TransformCurve.GetVectorCurveByIndex(SubCurveIndex);
+				for (int32 ChannelIndex = 0; ChannelIndex < 3; ++ChannelIndex)
+				{
+					const EVectorCurveChannel Axis = static_cast<EVectorCurveChannel>(ChannelIndex);
+					FAnimationCurveIdentifier TargetCurveIdentifier = CurveId;
+					UAnimationCurveIdentifierExtensions::GetTransformChildCurveIdentifier(TargetCurveIdentifier, Channel, Axis);
+					Controller->SetCurveKeys(TargetCurveIdentifier, VectorCurve->FloatCurves[ChannelIndex].GetConstRefOfKeys());
+				}
+			}
 		}
 	}
 	Controller->CloseBracket();
@@ -972,7 +1015,7 @@ bool UAnimSequenceBase::HasCurveData(SmartName::UID_Type CurveUID, bool bForceUs
 #if WITH_EDITOR
 	if (IsDataModelValid())
 	{
-		return DataModel->FindFloatCurve(FAnimationCurveIdentifier(CurveUID, ERawCurveTrackTypes::RCT_Float)) != nullptr;
+		return DataModelInterface->FindFloatCurve(FAnimationCurveIdentifier(CurveUID, ERawCurveTrackTypes::RCT_Float)) != nullptr;
 	}
 #endif
 	
@@ -1017,7 +1060,7 @@ void UAnimSequenceBase::HandleAssetPlayerTickedInternal(FAnimAssetTickContext &C
 }
 
 #if WITH_EDITOR
-void UAnimSequenceBase::OnModelModified(const EAnimDataModelNotifyType& NotifyType, UAnimDataModel* Model, const FAnimDataModelNotifPayload& Payload)
+void UAnimSequenceBase::OnModelModified(const EAnimDataModelNotifyType& NotifyType, IAnimationDataModel* Model, const FAnimDataModelNotifPayload& Payload)
 {
 	NotifyCollector.Handle(NotifyType);
 
@@ -1025,20 +1068,25 @@ void UAnimSequenceBase::OnModelModified(const EAnimDataModelNotifyType& NotifyTy
 	InitializeNotifyTrack();
 	PRAGMA_ENABLE_DEPRECATION_WARNINGS
 
-	auto HandleLengthChange = [&](float NewLength, float OldLength, float T0, float T1)
+	auto HandleNumberOfFramesChange = [&](int32 NewNumFrames, int32 OldNumFrames, int32 Frame0, int32 Frame1)
 	{
 		if (bPopulatingDataModel)
 		{
 			return;
 		}
+
+		const FFrameRate& ModelFrameRate = DataModelInterface->GetFrameRate();
+		const float T0 = ModelFrameRate.AsSeconds(Frame0);
+		const float T1 = ModelFrameRate.AsSeconds(Frame1);
+		const float NewLength = ModelFrameRate.AsSeconds(NewNumFrames);
 		
-		const float StartRemoveTime = T0;
-		const float EndRemoveTime = T1;
+		const float StartRemoveTime = ModelFrameRate.AsSeconds(Frame0);
+		const float EndRemoveTime = ModelFrameRate.AsSeconds(Frame1);
 
 		// Total time value for frames that were removed
 		const float Duration = T1 - T0;
 
-		if (NewLength > OldLength)
+		if (NewNumFrames > OldNumFrames)
 		{
 			const float InsertTime = T0;
 			for (FAnimNotifyEvent& Notify : Notifies)
@@ -1055,8 +1103,8 @@ void UAnimSequenceBase::OnModelModified(const EAnimDataModelNotifyType& NotifyTy
 				// if state, make sure to adjust end time
 				if (Notify.NotifyStateClass)
 				{
-					float NotifyDuration = Notify.GetDuration();
-					float NotifyEnd = CurrentTime + NotifyDuration;
+					const float NotifyDuration = Notify.GetDuration();
+					const float NotifyEnd = CurrentTime + NotifyDuration;
 					if (NotifyEnd >= InsertTime)
 					{
 						NewDuration = NotifyDuration + Duration;
@@ -1067,7 +1115,7 @@ void UAnimSequenceBase::OnModelModified(const EAnimDataModelNotifyType& NotifyTy
 					}
 				}
 
-				// Shift out notify by the time alloted for the inserted keys
+				// Shift out notify by the time allotted for the inserted keys
 				if (CurrentTime >= InsertTime)
 				{
 					CurrentTime += Duration;
@@ -1090,7 +1138,7 @@ void UAnimSequenceBase::OnModelModified(const EAnimDataModelNotifyType& NotifyTy
 				}
 			}
 		}
-		else if (NewLength < OldLength)
+		else if (NewNumFrames < OldNumFrames)
 		{
 			// re-locate notifies
 			for (FAnimNotifyEvent& Notify : Notifies)
@@ -1107,15 +1155,15 @@ void UAnimSequenceBase::OnModelModified(const EAnimDataModelNotifyType& NotifyTy
 				// if state, make sure to adjust end time
 				if (Notify.NotifyStateClass)
 				{
-					float NotifyDuration = Notify.GetDuration();
-					float NotifyEnd = CurrentTime + NotifyDuration;
+					const float NotifyDuration = Notify.GetDuration();
+					const float NotifyEnd = CurrentTime + NotifyDuration;
 					NewDuration = NotifyDuration;
 
 					// If Notify is inside of the trimmed time frame, zero out the duration
 					if (CurrentTime >= StartRemoveTime && NotifyEnd <= EndRemoveTime)
 					{
 						// small number @todo see if there is define for this
-						NewDuration = DataModel->GetFrameRate().AsInterval();
+						NewDuration = DataModelInterface->GetFrameRate().AsInterval();
 					}
 					// If Notify overlaps trimmed time frame, remove trimmed duration
 					else if (CurrentTime < EndRemoveTime && NotifyEnd > EndRemoveTime)
@@ -1123,7 +1171,7 @@ void UAnimSequenceBase::OnModelModified(const EAnimDataModelNotifyType& NotifyTy
 						NewDuration = NotifyEnd - Duration - CurrentTime;
 					}
 
-					NewDuration = FMath::Max(NewDuration, (float)DataModel->GetFrameRate().AsInterval());
+					NewDuration = FMath::Max(NewDuration, (float)DataModelInterface->GetFrameRate().AsInterval());
 				}
 
 				if (CurrentTime >= StartRemoveTime && CurrentTime <= EndRemoveTime)
@@ -1149,16 +1197,15 @@ void UAnimSequenceBase::OnModelModified(const EAnimDataModelNotifyType& NotifyTy
 					Notify.TriggerTimeOffset = GetTriggerTimeOffsetForType(EAnimEventTriggerOffsets::OffsetBefore);
 				}
 			}
-		}	
+		}
 	};
 
-	// Source animation curve data can be coarse, so copy the data down and remove any redundant keys
+	// Copy any float/transform curves from the model into RawCurveData, as it is used at runtime for AnimMontage/Composite(s)
 	auto CopyCurvesFromModel = [this]()
 	{
 		PRAGMA_DISABLE_DEPRECATION_WARNINGS
-		RawCurveData.FloatCurves = DataModel->GetFloatCurves();
-		RawCurveData.TransformCurves = DataModel->GetTransformCurves();
-		RawCurveData.RemoveRedundantKeys();
+		RawCurveData.FloatCurves = DataModelInterface->GetCurveData().FloatCurves;
+		RawCurveData.TransformCurves = DataModelInterface->GetCurveData().TransformCurves;
 		PRAGMA_ENABLE_DEPRECATION_WARNINGS
 	};
 
@@ -1167,14 +1214,14 @@ void UAnimSequenceBase::OnModelModified(const EAnimDataModelNotifyType& NotifyTy
 		case EAnimDataModelNotifyType::SequenceLengthChanged:
 		{
 			const FSequenceLengthChangedPayload& TypedPayload = Payload.GetPayload<FSequenceLengthChangedPayload>();
-			const float PreviousSequenceLength = TypedPayload.PreviousLength;
-			const float CurrentSequenceLength = Model->GetPlayLength();
+			const FFrameNumber PreviousNumberOfFrames = TypedPayload.PreviousNumberOfFrames;
+			const int32 CurrentNumberOfFrames = Model->GetNumberOfFrames();
 
 			PRAGMA_DISABLE_DEPRECATION_WARNINGS
-			SequenceLength = CurrentSequenceLength;
+			SequenceLength = Model->GetFrameRate().AsSeconds(CurrentNumberOfFrames);
 			PRAGMA_ENABLE_DEPRECATION_WARNINGS
 
-			HandleLengthChange(CurrentSequenceLength, PreviousSequenceLength, TypedPayload.T0, TypedPayload.T1);
+			HandleNumberOfFramesChange(CurrentNumberOfFrames, PreviousNumberOfFrames.Value, TypedPayload.Frame0.Value, TypedPayload.Frame1.Value);
 
 			// Ensure that we only clamp the notifies at the end of a bracket
 			if (NotifyCollector.IsNotWithinBracket())
@@ -1273,9 +1320,14 @@ void UAnimSequenceBase::OnModelModified(const EAnimDataModelNotifyType& NotifyTy
 	}
 }
 
-UAnimDataModel* UAnimSequenceBase::GetDataModel() const
+IAnimationDataModel* UAnimSequenceBase::GetDataModel() const
+{
+	return DataModelInterface.GetInterface();
+}
+
+TScriptInterface<IAnimationDataModel> UAnimSequenceBase::GetDataModelInterface() const
 {
-	return DataModel;
+	return DataModelInterface;
 }
 
 IAnimationDataController& UAnimSequenceBase::GetController()
@@ -1284,12 +1336,12 @@ IAnimationDataController& UAnimSequenceBase::GetController()
 
 	if(Controller == nullptr)
 	{
-		IAnimationDataControllerModule& ControllerModule = FModuleManager::Get().GetModuleChecked<IAnimationDataControllerModule>("AnimationDataController");
-		Controller = ControllerModule.GetController();
-		Controller->SetModel(DataModel);
+		Controller = DataModelInterface->GetController();
+		checkf(Controller, TEXT("Failed to create AnimationDataController"));
+		Controller->SetModel(DataModelInterface);
 	}
 
-	ensureAlways(Controller->GetModel() == DataModel);
+	ensureAlways(Controller->GetModelInterface() == DataModelInterface);
 
 	return *Controller;
 }
@@ -1300,51 +1352,68 @@ void UAnimSequenceBase::PopulateModel()
 	
 	// Make a copy of the current data, to mitigate any changes due to notify callbacks
 	const FFrameRate FrameRate = GetSamplingFrameRate();
-	const float PlayLength = GetPlayLength();
 	PRAGMA_DISABLE_DEPRECATION_WARNINGS
+	const float PlayLength = SequenceLength;
 	const FRawCurveTracks CurveData = RawCurveData;
 	PRAGMA_ENABLE_DEPRECATION_WARNINGS
 		
 	IAnimationDataController::FScopedBracket ScopedBracket(Controller, LOCTEXT("UAnimSequenceBase::PopulateModel_Bracket", "Generating Animation Model Data for Animation Sequence Base"));
-	Controller->SetPlayLength(PlayLength);
 	Controller->SetFrameRate(FrameRate);
+	Controller->SetNumberOfFrames(Controller->ConvertSecondsToFrameNumber(PlayLength));
 
 	USkeleton* TargetSkeleton = GetSkeleton();	
 	UE::Anim::CopyCurveDataToModel(CurveData, TargetSkeleton,  *Controller);	
 }
 
+void UAnimSequenceBase::PopulateWithExistingModel(TScriptInterface<IAnimationDataModel> ExistingDataModel)
+{
+	Controller->PopulateWithExistingModel(ExistingDataModel);
+}
+
 void UAnimSequenceBase::BindToModelModificationEvent()
 {
 	ValidateModel();
-	DataModel->GetModifiedEvent().RemoveAll(this);
-	DataModel->GetModifiedEvent().AddUObject(this, &UAnimSequenceBase::OnModelModified);
+	DataModelInterface->GetModifiedEvent().RemoveAll(this);
+	DataModelInterface->GetModifiedEvent().AddUObject(this, &UAnimSequenceBase::OnModelModified);
 }
 
-void UAnimSequenceBase::CopyDataModel(const UAnimDataModel* ModelToDuplicate)
+void UAnimSequenceBase::CopyDataModel(const TScriptInterface<IAnimationDataModel>& ModelToDuplicate)
 {
 	checkf(ModelToDuplicate != nullptr, TEXT("Invalidate data model %s"), *GetFullName());
-	if (DataModel)
+	if (ModelToDuplicate)
 	{
-		DataModel->GetModifiedEvent().RemoveAll(this);
+		ModelToDuplicate->GetModifiedEvent().RemoveAll(this);
 	}
 
-	DataModel = DuplicateObject(ModelToDuplicate, this);
-	Controller->SetModel(DataModel);
+	DataModelInterface = DuplicateObject(ModelToDuplicate.GetObject(), this);
+
+	if(Controller)
+	{
+		Controller->SetModel(DataModelInterface);
+	}
+	else
+	{
+		GetController();
+	}
 
 	BindToModelModificationEvent();
 }
 
 void UAnimSequenceBase::CreateModel()
 {
-	checkf(DataModel == nullptr, TEXT("Invalid attempt to override the existing data model %s"), *GetFullName());
-	DataModel = NewObject<UAnimDataModel>(this, FName(TEXT("AnimationDataModel")));
-		
+	const UClass* TargetClass = UE::Anim::DataModel::IAnimationDataModels::FindClassForAnimationAsset(this);
+	checkf(TargetClass != nullptr, TEXT("Unable to find valid AnimationDataModel class"));
+
+	checkf(!DataModelInterface || DataModelInterface.GetObject()->GetClass() != TargetClass, TEXT("Invalid attempt to override the existing data model %s"), *GetFullName());
+	UObject* ClassDataModel = NewObject<UObject>(this, TargetClass, TargetClass->GetFName());
+	DataModelInterface = ClassDataModel;
+
 	BindToModelModificationEvent();
 }
 
 void UAnimSequenceBase::ValidateModel() const
 {
-	checkf(DataModel != nullptr, TEXT("Invalid AnimSequenceBase state (%s), no data model found"), *GetPathName());
+	checkf(DataModelInterface != nullptr, TEXT("Invalid AnimSequenceBase state (%s), no data model found"), *GetPathName());
 }
 
 bool UAnimSequenceBase::ShouldDataModelBeValid() const
@@ -1360,3 +1429,37 @@ bool UAnimSequenceBase::ShouldDataModelBeValid() const
 
 #undef LOCTEXT_NAMESPACE 
 
+#if WITH_EDITOR
+void UAnimSequenceBase::OnEndLoadPackage(const FEndLoadPackageContext& Context)
+{
+	if (!GetPackage()->GetHasBeenEndLoaded())
+	{
+		return;
+	}
+
+	FCoreUObjectDelegates::OnEndLoadPackage.RemoveAll(this);
+
+	OnAnimModelLoaded();
+}
+
+void UAnimSequenceBase::OnAnimModelLoaded()
+{
+	QUICK_SCOPE_CYCLE_COUNTER(STAT_OnAnimModelLoaded);
+	if(USkeleton* MySkeleton = GetSkeleton())
+	{		
+		if (IsDataModelValid())
+		{
+			const TArray<FFloatCurve>& FloatCurves = DataModelInterface->GetFloatCurves();
+			for (int32 Index = 0; Index < FloatCurves.Num(); ++Index)
+			{
+				const FFloatCurve& Curve = FloatCurves[Index];
+				ensureMsgf(Curve.Name.DisplayName != NAME_None, TEXT("[AnimSequencer %s] has invalid curve name."), *GetFullName());
+			}
+
+			constexpr bool bDoNotTransactAction = false;
+			Controller->FindOrAddCurveNamesOnSkeleton(MySkeleton, ERawCurveTrackTypes::RCT_Float, bDoNotTransactAction);
+			Controller->FindOrAddCurveNamesOnSkeleton(MySkeleton, ERawCurveTrackTypes::RCT_Transform, bDoNotTransactAction);
+		}
+	}
+}
+#endif // WITH_EDITOR
\ No newline at end of file
diff --git a/Engine/Source/Runtime/Engine/Private/Animation/AnimSequenceHelpers.cpp b/Engine/Source/Runtime/Engine/Private/Animation/AnimSequenceHelpers.cpp
--- a/Engine/Source/Runtime/Engine/Private/Animation/AnimSequenceHelpers.cpp
+++ b/Engine/Source/Runtime/Engine/Private/Animation/AnimSequenceHelpers.cpp
@@ -1,7 +1,7 @@
 // Copyright Epic Games, Inc. All Rights Reserved.
 
 #include "Animation/AnimSequenceHelpers.h"
-#include "Animation/AnimData/AnimDataModel.h"
+#include "Animation/AnimData/IAnimationDataModel.h"
 #include "Animation/AnimData/IAnimationDataController.h"
 
 #include "UObject/NameTypes.h"
@@ -13,11 +13,11 @@
 #include "AnimationRuntime.h"
 #include "AnimEncoding.h"
 #include "AnimationUtils.h"
-#include "Animation/AnimTypes.h"
+#include "Animation/AnimationPoseData.h"
 #include "Animation/AnimNotifies/AnimNotify.h"
 #include "Animation/AnimNotifies/AnimNotifyState.h"
 #include "Misc/MessageDialog.h"
-#include "Animation/AnimSequenceHelpers.h"
+#include "Animation/AttributesRuntime.h"
 
 #define LOCTEXT_NAMESPACE "AnimSequenceHelpers"
 
@@ -27,222 +27,32 @@ namespace UE
 namespace Anim
 {
 
-struct FRetargetTracking
-{
-	const FCompactPoseBoneIndex PoseBoneIndex;
-	const int32 SkeletonBoneIndex;
-
-	FRetargetTracking(const FCompactPoseBoneIndex InPoseBoneIndex, const int32 InSkeletonBoneIndex)
-		: PoseBoneIndex(InPoseBoneIndex), SkeletonBoneIndex(InSkeletonBoneIndex)
-	{
-	}
-};
-
-struct FBuildRawPoseScratchArea : public TThreadSingleton<FBuildRawPoseScratchArea>
-{
-	TArray<FRetargetTracking> RetargetTracking;
-	TArray<FVirtualBoneCompactPoseData> VirtualBoneCompactPoseData;
-};
-
-
 #if WITH_EDITOR
-FTransform ExtractTransformForKey(int32 Key, const FRawAnimSequenceTrack& TrackToExtract)
+void BuildPoseFromModel(const IAnimationDataModel* Model, FCompactPose& OutPose, const double Time, const EAnimInterpolationType& InterpolationType, const FName& RetargetSource, const TArray<FTransform>& RetargetTransforms)
 {
-	static const FVector DefaultScale3D = FVector(1.f);
-	const bool bHasScaleKey = TrackToExtract.ScaleKeys.Num() > 0;
+	FBlendedCurve TempCurve;
+	UE::Anim::FStackAttributeContainer TempAttributes;
+	TempCurve.InitFrom(OutPose.GetBoneContainer());
 
-	const int32 PosKeyIndex = FMath::Min(Key, TrackToExtract.PosKeys.Num() - 1);
-	const int32 RotKeyIndex = FMath::Min(Key, TrackToExtract.RotKeys.Num() - 1);
-	if (bHasScaleKey)
-	{
-		const int32 ScaleKeyIndex = FMath::Min(Key, TrackToExtract.ScaleKeys.Num() - 1);
-		return FTransform(FQuat(TrackToExtract.RotKeys[RotKeyIndex]), FVector(TrackToExtract.PosKeys[PosKeyIndex]), FVector(TrackToExtract.ScaleKeys[ScaleKeyIndex]));
-	}
-	else
-	{
-		return FTransform(FQuat(TrackToExtract.RotKeys[RotKeyIndex]), FVector(TrackToExtract.PosKeys[PosKeyIndex]), FVector(DefaultScale3D));
-	}
+	FAnimationPoseData PoseData(OutPose, TempCurve, TempAttributes);
+	BuildPoseFromModel(Model, PoseData, Time, InterpolationType, RetargetSource, RetargetTransforms);
 }
-
-template<bool bInterpolateT>
-void ExtractPoseFromModel(const TArray<FBoneAnimationTrack>& BoneAnimationTracks, const TMap<FName, FTransform>& OverrideBoneTransforms, FCompactPose& InOutPose, const int32 KeyIndex1, const int32 KeyIndex2, float Alpha, const USkeleton* SourceSkeleton)
-{
-	const int32 NumAnimationTracks = BoneAnimationTracks.Num();
-	const FBoneContainer& RequiredBones = InOutPose.GetBoneContainer();
-
-	TArray<UE::Anim::FRetargetTracking>& RetargetTracking = UE::Anim::FBuildRawPoseScratchArea::Get().RetargetTracking;
-	RetargetTracking.Reset(NumAnimationTracks);
-
-	TArray<FVirtualBoneCompactPoseData>& VBCompactPoseData = UE::Anim::FBuildRawPoseScratchArea::Get().VirtualBoneCompactPoseData;
-	VBCompactPoseData = RequiredBones.GetVirtualBoneCompactPoseData();
-
-	FCompactPose Key2Pose;
-	if (bInterpolateT)
-	{
-		Key2Pose.CopyBonesFrom(InOutPose);
-	}
-
-	const USkeleton* TargetSkeleton = RequiredBones.GetSkeletonAsset();
-	const FSkeletonRemapping* SkeletonRemapping = TargetSkeleton->GetSkeletonRemapping(SourceSkeleton);
-
-	for (const FBoneAnimationTrack& AnimationTrack : BoneAnimationTracks)
-	{
-		const int32 SourceSkeletonBoneIndex = AnimationTrack.BoneTreeIndex;
-		const int32 SkeletonBoneIndex = SkeletonRemapping ? SkeletonRemapping->GetTargetSkeletonBoneIndex(SourceSkeletonBoneIndex) : SourceSkeletonBoneIndex;
-
-		// not sure it's safe to assume that SkeletonBoneIndex can never be INDEX_NONE
-		if ((SkeletonBoneIndex != INDEX_NONE) && (SkeletonBoneIndex < MAX_BONES))
-		{
-			const FCompactPoseBoneIndex PoseBoneIndex = RequiredBones.GetCompactPoseIndexFromSkeletonIndex(SkeletonBoneIndex);
-			if (PoseBoneIndex != INDEX_NONE)
-			{
-				for (int32 Idx = 0; Idx < VBCompactPoseData.Num(); ++Idx)
-				{
-					FVirtualBoneCompactPoseData& VB = VBCompactPoseData[Idx];
-					if (PoseBoneIndex == VB.VBIndex)
-					{
-						// Remove this bone as we have written data for it (false so we dont resize allocation)
-						VBCompactPoseData.RemoveAtSwap(Idx, 1, false);
-						break; //Modified TArray so must break here
-					}
-				}
 				
-				// extract animation
-				const FRawAnimSequenceTrack& TrackToExtract = AnimationTrack.InternalTrackData;
-				const FTransform* OverrideTransform = OverrideBoneTransforms.Find(AnimationTrack.Name);
-				{
-					// Bail out (with rather wacky data) if data is empty for some reason.
-					if (TrackToExtract.PosKeys.Num() == 0 || TrackToExtract.RotKeys.Num() == 0)
-					{
-						InOutPose[PoseBoneIndex].SetIdentity();
-
-						if (bInterpolateT)
-						{
-							Key2Pose[PoseBoneIndex].SetIdentity();
-						}
-					}
-					else
-					{
-						InOutPose[PoseBoneIndex] = UE::Anim::ExtractTransformForKey(KeyIndex1, TrackToExtract);
-
-						if (bInterpolateT)
-						{
-							Key2Pose[PoseBoneIndex] = UE::Anim::ExtractTransformForKey(KeyIndex2, TrackToExtract);
-						}
-					}
-				}
-
-				if (OverrideTransform)
-				{
-					InOutPose[PoseBoneIndex].SetRotation(InOutPose[PoseBoneIndex].GetRotation() * OverrideTransform->GetRotation());
-					InOutPose[PoseBoneIndex].TransformPosition(OverrideTransform->GetTranslation());
-					InOutPose[PoseBoneIndex].SetScale3D(InOutPose[PoseBoneIndex].GetScale3D() * OverrideTransform->GetScale3D());
-
-					if (bInterpolateT)
-					{
-						Key2Pose[PoseBoneIndex].SetRotation(Key2Pose[PoseBoneIndex].GetRotation() * OverrideTransform->GetRotation());
-						Key2Pose[PoseBoneIndex].TransformPosition(OverrideTransform->GetTranslation());
-						Key2Pose[PoseBoneIndex].SetScale3D(Key2Pose[PoseBoneIndex].GetScale3D() * OverrideTransform->GetScale3D());
-					}
-				}
-
-				RetargetTracking.Add(UE::Anim::FRetargetTracking(PoseBoneIndex, SkeletonBoneIndex));
-			}
-		}
-	}
-
-	//Build Virtual Bones
-	if (VBCompactPoseData.Num() > 0)
-	{
-		FCSPose<FCompactPose> CSPose1;
-		CSPose1.InitPose(InOutPose);
-
-		FCSPose<FCompactPose> CSPose2;
-		if (bInterpolateT)
-		{
-			CSPose2.InitPose(Key2Pose);
-		}
-
-		for (FVirtualBoneCompactPoseData& VB : VBCompactPoseData)
-		{
-			FTransform Source = CSPose1.GetComponentSpaceTransform(VB.SourceIndex);
-			FTransform Target = CSPose1.GetComponentSpaceTransform(VB.TargetIndex);
-			InOutPose[VB.VBIndex] = Target.GetRelativeTransform(Source);
-
-			if (bInterpolateT)
-			{
-				FTransform Source2 = CSPose2.GetComponentSpaceTransform(VB.SourceIndex);
-				FTransform Target2 = CSPose2.GetComponentSpaceTransform(VB.TargetIndex);
-				Key2Pose[VB.VBIndex] = Target2.GetRelativeTransform(Source2);
-			}
-		}
-	}
-
-	if (bInterpolateT)
-	{
-		for (FCompactPoseBoneIndex BoneIndex : InOutPose.ForEachBoneIndex())
-		{
-			InOutPose[BoneIndex].Blend(InOutPose[BoneIndex], Key2Pose[BoneIndex], Alpha);
-		}
-	}
-}
-
-void BuildPoseFromModel(const UAnimDataModel* Model, FCompactPose& OutPose, const float Time, const EAnimInterpolationType& InterpolationType, const FName& RetargetSource, const TArray<FTransform>& RetargetTransforms)
+void BuildPoseFromModel(const IAnimationDataModel* Model, FAnimationPoseData& OutPoseData, const double Time, const EAnimInterpolationType& InterpolationType, const FName& RetargetSource, const TArray<FTransform>& RetargetTransforms)
 {
+	FCompactPose& OutPose = OutPoseData.GetPose();
 	check(Model);
 
-	const int32 NumberOfKeys = Model->GetNumberOfKeys();
-	const float SequenceLength = Model->GetPlayLength();
-	const double FramesPerSecond = Model->GetFrameRate().AsDecimal();
-
-	// Generate keys to interpolate between
-	int32 KeyIndex1, KeyIndex2;
-	float Alpha;
-	FAnimationRuntime::GetKeyIndicesFromTime(KeyIndex1, KeyIndex2, Alpha, Time, NumberOfKeys, SequenceLength, FramesPerSecond);
-
-	if (InterpolationType == EAnimInterpolationType::Step)
-	{
-		// Force stepping between keys
-		Alpha = 0.f;
-	}
-
-	bool bShouldInterpolate = true;
-
-	if (Alpha < UE_KINDA_SMALL_NUMBER)
-	{
-		Alpha = 0.f;
-		bShouldInterpolate = false;
-	}
-	else if (Alpha > 1.f - UE_KINDA_SMALL_NUMBER)
-	{
-		bShouldInterpolate = false;
-		KeyIndex1 = KeyIndex2;
-	}
-
-	TMap<FName, FTransform> ActiveCurves;
-	if (!OutPose.GetBoneContainer().ShouldUseSourceData())
-	{
-		EvaluateTransformCurvesFromModel(Model, ActiveCurves, Time, 1.f);
-	}
-
-	const USkeleton* SourceSkeleton = Model->GetAnimationSequence()->GetSkeleton();
-	if (bShouldInterpolate)
-	{
-		ExtractPoseFromModel<true>(Model->GetBoneAnimationTracks(), ActiveCurves, OutPose, KeyIndex1, KeyIndex2, Alpha, SourceSkeleton);
-	}
-	else
-	{
-		ExtractPoseFromModel<false>(Model->GetBoneAnimationTracks(), ActiveCurves, OutPose, KeyIndex1, KeyIndex2, Alpha, SourceSkeleton);
-	}
+	OutPose.ResetToRefPose();
 
-	Retargeting::RetargetPose(OutPose, RetargetSource, RetargetTransforms);
+	const UE::Anim::DataModel::FEvaluationContext EvaluationContext(Time, Model->GetFrameRate(), RetargetSource, RetargetTransforms, InterpolationType);
+	Model->Evaluate(OutPoseData, EvaluationContext);
 }
 
-void EvaluateFloatCurvesFromModel(const UAnimDataModel* Model, FBlendedCurve& OutCurves, float Time)
+void EvaluateFloatCurvesFromModel(const IAnimationDataModel* Model, FBlendedCurve& OutCurves, double Time)
 {
 	check(Model);
 
-	//SCOPE_CYCLE_COUNTER(STAT_EvalRawCurveData);
 	if (OutCurves.NumValidCurveCount > 0)
 	{
 		// evaluate the curve data at the Time and add to Instance
@@ -257,7 +67,7 @@ void EvaluateFloatCurvesFromModel(const UAnimDataModel* Model, FBlendedCurve& Ou
 	}
 }
 
-void EvaluateTransformCurvesFromModel(const UAnimDataModel* Model, TMap<FName, FTransform>& OutCurves, float Time, float BlendWeight)
+void EvaluateTransformCurvesFromModel(const IAnimationDataModel* Model, TMap<FName, FTransform>& OutCurves, double Time, float BlendWeight)
 {
 	if (Model)
 	{
@@ -279,15 +89,15 @@ void EvaluateTransformCurvesFromModel(const UAnimDataModel* Model, TMap<FName, F
 	}
 }
 
-void GetBoneTransformFromModel(const UAnimDataModel* Model, FTransform& OutTransform, int32 TrackIndex, float Time, const EAnimInterpolationType& Interpolation)
+void GetBoneTransformFromModel(const IAnimationDataModel* Model, FTransform& OutTransform, int32 TrackIndex, double Time, const EAnimInterpolationType& Interpolation)
 {
 	const FBoneAnimationTrack& TrackData = Model->GetBoneTrackByIndex(TrackIndex);
 	const FRawAnimSequenceTrack& RawTrack = TrackData.InternalTrackData;
 	
-	FAnimationUtils::ExtractTransformFromTrack(Time, Model->GetNumberOfKeys(), Model->GetPlayLength(), RawTrack, Interpolation, OutTransform);
+	FAnimationUtils::ExtractTransformFromTrack(RawTrack, Time, Model->GetNumberOfKeys(), Model->GetPlayLength(), Interpolation, OutTransform);
 }
 
-void GetBoneTransformFromModel(const UAnimDataModel* Model, FTransform& OutTransform, int32 TrackIndex, int32 KeyIndex)
+void GetBoneTransformFromModel(const IAnimationDataModel* Model, FTransform& OutTransform, int32 TrackIndex, int32 KeyIndex)
 {
 	const FBoneAnimationTrack& TrackData = Model->GetBoneTrackByIndex(TrackIndex);
 	const FRawAnimSequenceTrack& RawTrack = TrackData.InternalTrackData;
@@ -395,6 +205,36 @@ void CopyCurveDataToModel(const FRawCurveTracks& CurveData, const USkeleton* Ske
 	else
 	{
 		OutTransform.SetScale3D(DefaultScale3D);
+	}	
+}
+	
+Retargeting::FRetargetingScope::FRetargetingScope(FCompactPose& ToRetargetPose, const DataModel::FEvaluationContext& InEvaluationContext)
+	: RetargetPose(ToRetargetPose),
+	EvaluationContext(InEvaluationContext),
+	RetargetTracking(FBuildRawPoseScratchArea::Get().RetargetTracking),
+	bShouldRetarget(!ToRetargetPose.GetBoneContainer().GetDisableRetargeting() && EvaluationContext.RetargetTransforms.Num())
+{
+	RetargetTracking.Reset();
+}
+
+void Retargeting::FRetargetingScope::AddTrackedBone(FCompactPoseBoneIndex CompactBoneIndex, int32 SkeletonBoneIndex) const
+{
+	if (bShouldRetarget)
+	{
+		RetargetTracking.Add({CompactBoneIndex, SkeletonBoneIndex});
+	}
+}
+
+Retargeting::FRetargetingScope::~FRetargetingScope()
+{
+	if (bShouldRetarget)
+	{
+		const FBoneContainer& RequiredBones = RetargetPose.GetBoneContainer();
+		const USkeleton* Skeleton = RequiredBones.GetSkeletonAsset();
+		for (const FRetargetTracking& RT : RetargetTracking)
+		{
+			FAnimationRuntime::RetargetBoneTransform(Skeleton, EvaluationContext.RetargetSource, EvaluationContext.RetargetTransforms, RetargetPose[RT.PoseBoneIndex], RT.SkeletonBoneIndex, RT.PoseBoneIndex, RequiredBones, false);
+		}
 	}
 }
 
@@ -405,11 +245,11 @@ void Retargeting::RetargetPose(FCompactPose& InOutPose, const FName& RetargetSou
 
 	if (!bDisableRetargeting && RetargetTransforms.Num())
 	{
-		const TArray<UE::Anim::FRetargetTracking>& RetargetTracking = UE::Anim::FBuildRawPoseScratchArea::Get().RetargetTracking;
+		const TArray<UE::Anim::Retargeting::FRetargetTracking>& RetargetTracking = UE::Anim::FBuildRawPoseScratchArea::Get().RetargetTracking;
 
 		USkeleton* Skeleton = RequiredBones.GetSkeletonAsset();
 
-		for (const UE::Anim::FRetargetTracking& RT : RetargetTracking)
+		for (const UE::Anim::Retargeting::FRetargetTracking& RT : RetargetTracking)
 		{
 			FAnimationRuntime::RetargetBoneTransform(Skeleton, RetargetSource, RetargetTransforms, InOutPose[RT.PoseBoneIndex], RT.SkeletonBoneIndex, RT.PoseBoneIndex, RequiredBones, false);
 		}
@@ -702,9 +542,8 @@ void Compression::SanitizeRawAnimSequenceTrack(FRawAnimSequenceTrack& RawTrack)
 	// if scale is too small, zero it out. Cause it hard to retarget when compress
 	// inverse scale is applied to translation, and causing translation to be huge to retarget, but
 	// compression can't handle that much precision. 
-	for (auto ScaleIter = RawTrack.ScaleKeys.CreateIterator(); ScaleIter; ++ScaleIter)
+	for (auto& Scale3D : RawTrack.ScaleKeys)
 	{
-		FVector3f& Scale3D = *ScaleIter;
 		if (FMath::IsNearlyZero(Scale3D.X))
 		{
 			Scale3D.X = 0.f;
@@ -718,15 +557,43 @@ void Compression::SanitizeRawAnimSequenceTrack(FRawAnimSequenceTrack& RawTrack)
 			Scale3D.Z = 0.f;
 		}
 	}
+	
+	for (auto& Position : RawTrack.PosKeys)
+	{		
+		if (FMath::IsNearlyZero(Position.X))
+		{
+			Position.X = 0.f;
+		}
+		if (FMath::IsNearlyZero(Position.Y))
+		{
+			Position.Y = 0.f;
+		}
+		if (FMath::IsNearlyZero(Position.Z))
+		{
+			Position.Z = 0.f;
+		}
+	}
 
 	// make sure Rotation part is normalized before compress
-	for (auto RotIter = RawTrack.RotKeys.CreateIterator(); RotIter; ++RotIter)
+	for (auto& Rotation : RawTrack.RotKeys)
 	{
-		FQuat4f& Rotation = *RotIter;
 		if (!Rotation.IsNormalized())
 		{
 			Rotation.Normalize();
 		}
+
+		if (FMath::IsNearlyZero(Rotation.X))
+		{
+			Rotation.X = 0.f;
+		}
+		if (FMath::IsNearlyZero(Rotation.Y))
+		{
+			Rotation.Y = 0.f;
+		}
+		if (FMath::IsNearlyZero(Rotation.Z))
+		{
+			Rotation.Z = 0.f;
+		}
 	}
 }
 
@@ -749,7 +616,7 @@ Compression::FScopedCompressionGuard::~FScopedCompressionGuard()
 
 bool AnimationData::AddLoopingInterpolation(UAnimSequence* InSequence)
 {
-	const UAnimDataModel* DataModel = InSequence->GetDataModel();
+	const IAnimationDataModel* DataModel = InSequence->GetDataModel();
 	IAnimationDataController& Controller = InSequence->GetController();
 
 	const int32 NumTracks = DataModel->GetNumBoneTracks();
@@ -785,8 +652,8 @@ bool AnimationData::AddLoopingInterpolation(UAnimSequence* InSequence)
 			Controller.SetBoneTrackKeys(AnimationTrack.Name, PositionalKeys, RotationalKeys, ScaleKeys);
 		}
 
-		const float Interval = DataModel->GetFrameRate().AsInterval();
-		Controller.SetPlayLength(Interval + DataModel->GetPlayLength());
+		// New number of frames is equal to current number of keys, as we'll be adding one frame (Number of Frames + 1 == Number of Keys)
+		Controller.SetNumberOfFrames(NumKeys);
 
 		return true;
 	}
@@ -796,7 +663,7 @@ bool AnimationData::AddLoopingInterpolation(UAnimSequence* InSequence)
 
 bool AnimationData::Trim(UAnimSequence* InSequence, float TrimStart, float TrimEnd)
 {
-	const UAnimDataModel* DataModel = InSequence->GetDataModel();
+	const IAnimationDataModel* DataModel = InSequence->GetDataModel();
 	IAnimationDataController& Controller = InSequence->GetController();
 
 	const int32 NumTracks = DataModel->GetNumBoneTracks();
@@ -842,7 +709,7 @@ bool AnimationData::Trim(UAnimSequence* InSequence, float TrimStart, float TrimE
 
 void AnimationData::DuplicateKeys(UAnimSequence* InSequence, int32 StartKeyIndex, int32 NumDuplicates, int32 SourceKeyIndex /*= INDEX_NONE */)
 {
-	const UAnimDataModel* Model = InSequence->GetDataModel();
+	const IAnimationDataModel* Model = InSequence->GetDataModel();
 	IAnimationDataController& Controller = InSequence->GetController();
 
 	const int32 NumberOfKeys = Model->GetNumberOfKeys();
@@ -894,14 +761,14 @@ void AnimationData::DuplicateKeys(UAnimSequence* InSequence, int32 StartKeyIndex
 			const float InsertedTime = FrameRate.AsInterval() * NumDuplicates;
 
 			// Notify will happen with time slice that was inserted
-			Controller.Resize(NewSequenceLength, StartTime, StartTime + InsertedTime);
+			Controller.ResizeInFrames(NewNumFrames, StartKeyIndex, StartKeyIndex + NumDuplicates);
 		}
 	}
 }
 
 void AnimationData::RemoveKeys(UAnimSequence* InSequence, int32 StartKeyIndex, int32 NumKeysToRemove)
 {
-	const UAnimDataModel* Model = InSequence->GetDataModel();
+	const IAnimationDataModel* Model = InSequence->GetDataModel();
 	IAnimationDataController& Controller = InSequence->GetController();
 
 	const int32 NumberOfKeys = Model->GetNumberOfKeys();
@@ -937,20 +804,17 @@ void AnimationData::RemoveKeys(UAnimSequence* InSequence, int32 StartKeyIndex, i
 		}
 
 		const int32 NewNumberOfFrames = FMath::Max(NewNumberOfKeys - 1, 0);
-		const float NewSequenceLength = FrameRate.AsSeconds(NewNumberOfFrames);
 
-		const float StartFrameTime = FrameRate.AsSeconds(FMath::Max(StartKeyIndex, 0));
-		const float RemovedTime = NumKeysToRemove * FrameRate.AsInterval();
-		const float EndTime = StartFrameTime + RemovedTime;
+		const int32 StartKey = FMath::Max(StartKeyIndex, 0);
 
 		// Notify will happen with time slice that was removed
-		Controller.Resize(NewSequenceLength, StartFrameTime, EndTime);
+		Controller.ResizeInFrames(NewNumberOfFrames, StartKey, StartKey + NumberOfKeys);
 	}
 }
 
 int32 AnimationData::FindFirstChildTrackIndex(const UAnimSequence* InSequence, const USkeleton* Skeleton, const FName& BoneName)
 {
-	const UAnimDataModel* DataModel = InSequence->GetDataModel();
+	const IAnimationDataModel* DataModel = InSequence->GetDataModel();
 	const FReferenceSkeleton& RefSkeleton = Skeleton->GetReferenceSkeleton();
 
 	const int32 BoneIndex = RefSkeleton.FindBoneIndex(BoneName);
diff --git a/Engine/Source/Runtime/Engine/Private/Animation/AnimSingleNodeInstanceProxy.cpp b/Engine/Source/Runtime/Engine/Private/Animation/AnimSingleNodeInstanceProxy.cpp
--- a/Engine/Source/Runtime/Engine/Private/Animation/AnimSingleNodeInstanceProxy.cpp
+++ b/Engine/Source/Runtime/Engine/Private/Animation/AnimSingleNodeInstanceProxy.cpp
@@ -174,7 +174,7 @@ void FAnimSingleNodeInstanceProxy::InternalBlendSpaceEvaluatePose(class UBlendSp
 {
 	FAnimationPoseData AnimationPoseData = { OutContext.Pose, OutContext.Curve, OutContext.CustomAttributes };
 
-	FAnimExtractContext ExtractionContext(CurrentTime, ShouldExtractRootMotion(), DeltaTimeRecord, bLooping);
+	FAnimExtractContext ExtractionContext(static_cast<double>(CurrentTime), ShouldExtractRootMotion(), DeltaTimeRecord, bLooping);
 
 	if (BlendSpace->IsValidAdditive())
 	{
@@ -188,7 +188,7 @@ void FAnimSingleNodeInstanceProxy::InternalBlendSpaceEvaluatePose(class UBlendSp
 #if WITH_EDITORONLY_DATA
 		if (BlendSpace->PreviewBasePose)
 		{
-			BlendSpace->PreviewBasePose->GetBonePose(AnimationPoseData, FAnimExtractContext(PreviewPoseCurrentTime));
+			BlendSpace->PreviewBasePose->GetBonePose(AnimationPoseData, FAnimExtractContext(static_cast<double>(PreviewPoseCurrentTime)));
 		}
 		else
 #endif // WITH_EDITORONLY_DATA
@@ -309,7 +309,7 @@ void FAnimNode_SingleNode::Evaluate_AnyThread(FPoseContext& Output)
 		= false;
 #endif
 
-	if (Proxy->CurrentAsset != NULL && !Proxy->CurrentAsset->HasAnyFlags(RF_BeginDestroyed))
+	if (Proxy->CurrentAsset != nullptr && !Proxy->CurrentAsset->HasAnyFlags(RF_BeginDestroyed))
 	{
 		FAnimationPoseData OutputAnimationPoseData(Output);
 
@@ -320,7 +320,7 @@ void FAnimNode_SingleNode::Evaluate_AnyThread(FPoseContext& Output)
 		}
 		else if (UAnimSequence* Sequence = Cast<UAnimSequence>(Proxy->CurrentAsset))
 		{
-			FAnimExtractContext ExtractionContext(Proxy->CurrentTime, Sequence->bEnableRootMotion, Proxy->DeltaTimeRecord, Proxy->bLooping);
+			FAnimExtractContext ExtractionContext(static_cast<double>(Proxy->CurrentTime), Sequence->bEnableRootMotion, Proxy->DeltaTimeRecord, Proxy->bLooping);
 
 			if (Sequence->IsValidAdditive())
 			{
@@ -379,12 +379,12 @@ void FAnimNode_SingleNode::Evaluate_AnyThread(FPoseContext& Output)
 			else*/
 			{
 				// if SkeletalMesh isn't there, we'll need to use skeleton
-				Streamable->GetAnimationPose(OutputAnimationPoseData, FAnimExtractContext(Proxy->CurrentTime, Streamable->bEnableRootMotion, Proxy->DeltaTimeRecord, Proxy->bLooping));
+				Streamable->GetAnimationPose(OutputAnimationPoseData, FAnimExtractContext(static_cast<double>(Proxy->CurrentTime), Streamable->bEnableRootMotion, Proxy->DeltaTimeRecord, Proxy->bLooping));
 			}
 		}
 		else if (UAnimComposite* Composite = Cast<UAnimComposite>(Proxy->CurrentAsset))
 		{
-			FAnimExtractContext ExtractionContext(Proxy->CurrentTime, Proxy->ShouldExtractRootMotion(), Proxy->DeltaTimeRecord, Proxy->bLooping);
+			FAnimExtractContext ExtractionContext(static_cast<double>(Proxy->CurrentTime), Proxy->ShouldExtractRootMotion(), Proxy->DeltaTimeRecord, Proxy->bLooping);
 			const FAnimTrack& AnimTrack = Composite->AnimationTrack;
 
 			// find out if this is additive animation
@@ -442,7 +442,7 @@ void FAnimNode_SingleNode::Evaluate_AnyThread(FPoseContext& Output)
 					if (bCanProcessAdditiveAnimationsLocal && Montage->PreviewBasePose && Montage->GetPlayLength() > 0.f)
 					{
 						FAnimationPoseData LocalAnimationPoseData = { LocalSourcePose, LocalSourceCurve, LocalSourceAttributes };
-						Montage->PreviewBasePose->GetBonePose(LocalAnimationPoseData, FAnimExtractContext(Proxy->CurrentTime));
+						Montage->PreviewBasePose->GetBonePose(LocalAnimationPoseData, FAnimExtractContext(static_cast<double>(Proxy->CurrentTime)));
 					}
 					else
 #endif // WITH_EDITORONLY_DATA
@@ -562,7 +562,7 @@ void FAnimNode_SingleNode::Evaluate_AnyThread(FPoseContext& Output)
 void FAnimNode_SingleNode::Update_AnyThread(const FAnimationUpdateContext& Context)
 {
 	float NewPlayRate = Proxy->PlayRate;
-	UAnimSequence* PreviewBasePose = NULL;
+	UAnimSequence* PreviewBasePose = nullptr;
 
 	if (Proxy->bPlaying == false)
 	{
@@ -570,7 +570,7 @@ void FAnimNode_SingleNode::Update_AnyThread(const FAnimationUpdateContext& Conte
 		NewPlayRate = 0.f;
 	}
 
-	if(Proxy->CurrentAsset != NULL)
+	if(Proxy->CurrentAsset != nullptr)
 	{
 		UE::Anim::FAnimSyncGroupScope& SyncScope = Context.GetMessageChecked<UE::Anim::FAnimSyncGroupScope>();
 
diff --git a/Engine/Source/Runtime/Engine/Private/Animation/AnimStreamable.cpp b/Engine/Source/Runtime/Engine/Private/Animation/AnimStreamable.cpp
--- a/Engine/Source/Runtime/Engine/Private/Animation/AnimStreamable.cpp
+++ b/Engine/Source/Runtime/Engine/Private/Animation/AnimStreamable.cpp
@@ -19,14 +19,17 @@
 #include "Animation/AnimBoneCompressionCodec.h"
 #include "Animation/AnimCurveCompressionCodec.h"
 #include "BonePose.h"
+#include "CommonFrameRates.h"
 #include "ContentStreaming.h"
+#include "ITimeManagementModule.h"
 #include "ProfilingDebugging/CookStats.h"
 #include "Animation/AnimationPoseData.h"
 #include "Animation/AttributesRuntime.h"
 #include "UObject/UE5MainStreamObjectVersion.h"
 
-#include "Animation/AnimData/AnimDataModel.h"
 #include "Animation/AnimSequenceHelpers.h"
+#include "Animation/AnimSequence.h"
+#include "Animation/AnimData/IAnimationDataModel.h"
 
 CSV_DECLARE_CATEGORY_MODULE_EXTERN(ENGINE_API, Animation);
 
@@ -160,13 +163,6 @@ void UAnimStreamable::Serialize(FArchive& Ar)
 	
 	Ar.UsingCustomVersion(FUE5MainStreamObjectVersion::GUID);
 
-	if (Ar.IsLoading() && 
-	    Ar.CustomVer(FUE5MainStreamObjectVersion::GUID) >= FUE5MainStreamObjectVersion::AnimationDataModelInterface_BackedOut &&
-		Ar.CustomVer(FUE5MainStreamObjectVersion::GUID) < FUE5MainStreamObjectVersion::BackoutAnimationDataModelInterface)
-	{
-		UE_LOG(LogAnimation, Fatal, TEXT("This package was saved with a version that had to be backed out and is no longer able to be loaded."));
-	}
-
 	bool bCooked = Ar.IsCooking();
 	Ar << bCooked;
 
@@ -175,6 +171,41 @@ void UAnimStreamable::Serialize(FArchive& Ar)
 		UE_LOG(LogAnimation, Fatal, TEXT("This platform requires cooked packages, and animation data was not cooked into %s."), *GetFullName());
 	}
 
+#if WITH_EDITORONLY_DATA
+	if (Ar.IsLoading() && Ar.CustomVer(FUE5MainStreamObjectVersion::GUID) < FUE5MainStreamObjectVersion::ReintroduceAnimationDataModelInterface)
+	{
+		// Figure out correct SamplingFrameRate value 
+		if(SourceSequence && !SourceSequence->HasAnyFlags(RF_NeedPostLoad))
+		{
+			SamplingFrameRate = SourceSequence->GetSamplingFrameRate();
+		}
+		else
+		{
+			const int32 NumberOfFrames = FMath::Max(NumberOfKeys - 1, 1);
+
+			// Generate the frame-rate according to the number of frames and sequence length
+			const double DecimalFrameRate = (double)NumberOfFrames / ((double)GetPlayLength() > 0.0 ? (double)GetPlayLength() : 1.0);
+
+			// Account for non-whole number frame rates using large denominator
+			const double Denominator = 1000000.0;
+			SamplingFrameRate = FFrameRate(DecimalFrameRate * Denominator, Denominator);	 
+
+			// Try to simplifiy the frame rate, in case it is a multiple of the commonly used frame rates e.g. 10000/300000 -> 1/30
+			TArrayView<const FCommonFrameRateInfo> CommonFrameRates = FModuleManager::LoadModulePtr<ITimeManagementModule>("TimeManagement")->GetAllCommonFrameRates();
+			for (const FCommonFrameRateInfo& Info : CommonFrameRates)
+			{
+				const bool bDoesNotAlreadyMatch = Info.FrameRate.Denominator != SamplingFrameRate.Denominator && Info.FrameRate.Numerator != SamplingFrameRate.Numerator;
+		
+				if (bDoesNotAlreadyMatch && FMath::IsNearlyEqual(SamplingFrameRate.AsInterval(), Info.FrameRate.AsInterval()))
+				{
+					SamplingFrameRate = Info.FrameRate;
+					break;
+				}
+			}
+		}
+	}
+#endif // WITH_EDITORONLY_DATA
+
 	if (bCooked)
 	{
 		if (Ar.IsLoading())
@@ -198,19 +229,6 @@ void UAnimStreamable::Serialize(FArchive& Ar)
 #endif // WITH_EDITORONLY_DATA
 }
 
-#if WITH_EDITOR
-/*bool UAnimComposite::GetAllAnimationSequencesReferred(TArray<UAnimationAsset*>& AnimationAssets, bool bRecursive) 
-{
-	return AnimationTrack.GetAllAnimationSequencesReferred(AnimationAssets, bRecursive);
-}
-
-void UAnimComposite::ReplaceReferredAnimations(const TMap<UAnimationAsset*, UAnimationAsset*>& ReplacementMap)
-{
-	AnimationTrack.ReplaceReferredAnimations(ReplacementMap);
-}*/
-#endif
-
-
 void UAnimStreamable::HandleAssetPlayerTickedInternal(FAnimAssetTickContext &Context, const float PreviousTime, const float MoveDelta, const FAnimTickRecord &Instance, struct FAnimNotifyQueue& NotifyQueue) const
 {
 	Super::HandleAssetPlayerTickedInternal(Context, PreviousTime, MoveDelta, Instance, NotifyQueue);
@@ -297,10 +315,8 @@ void UAnimStreamable::GetAnimationPose(FAnimationPoseData& OutAnimationPoseData,
 		//Need to evaluate raw data
 		ValidateModel();
 
-		UE::Anim::EvaluateFloatCurvesFromModel(DataModel, OutCurve, ExtractionContext.CurrentTime);
-
 		// Warning if we have invalid data
-		for (const FBoneAnimationTrack& AnimationTrack : DataModel->GetBoneAnimationTracks())
+		for (const FBoneAnimationTrack& AnimationTrack : DataModelInterface->GetBoneAnimationTracks())
 		{
 			const FRawAnimSequenceTrack& TrackToExtract = AnimationTrack.InternalTrackData;
 			// Bail out (with rather wacky data) if data is empty for some reason.
@@ -310,7 +326,8 @@ void UAnimStreamable::GetAnimationPose(FAnimationPoseData& OutAnimationPoseData,
 			}
 		}
 
-		UE::Anim::BuildPoseFromModel(DataModel, OutPose, ExtractionContext.CurrentTime, Interpolation, RetargetSource, MySkeleton->GetRefLocalPoses(RetargetSource));
+		const UE::Anim::DataModel::FEvaluationContext EvaluationContext(ExtractionContext.CurrentTime, DataModelInterface->GetFrameRate(), RetargetSource, MySkeleton->GetRefLocalPoses(RetargetSource), Interpolation);
+		DataModelInterface->Evaluate(OutAnimationPoseData, EvaluationContext);
 
 		if ((ExtractionContext.bExtractRootMotion && RootMotionReset.bEnableRootMotion) || RootMotionReset.bForceRootLock)
 		{
@@ -320,19 +337,22 @@ void UAnimStreamable::GetAnimationPose(FAnimationPoseData& OutAnimationPoseData,
 	}
 #endif
 
-	const int32 ChunkIndex = GetChunkIndexForTime(GetRunningPlatformData().Chunks, ExtractionContext.CurrentTime);
+	const FStreamableAnimPlatformData& StreamableAnimData = GetRunningPlatformData();
+
+	const int32 ChunkIndex = GetChunkIndexForTime(StreamableAnimData.Chunks, ExtractionContext.CurrentTime);
 	if(ensureMsgf(ChunkIndex != INDEX_NONE, TEXT("Could not get valid chunk with Time %.2f for Streaming Anim %s"), ExtractionContext.CurrentTime, *GetFullName()))
 	{
-		IAnimationStreamingManager& StreamingManager = IStreamingManager::Get().GetAnimationStreamingManager();
+		const FAnimStreamableChunk& CurrentChunk = StreamableAnimData.Chunks[ChunkIndex];
+		const IAnimationStreamingManager& StreamingManager = IStreamingManager::Get().GetAnimationStreamingManager();
 
-		bool bUsingFirstChunk = (ChunkIndex == 0);
+		const bool bUsingFirstChunk = (ChunkIndex == 0);
 
 		const FCompressedAnimSequence* CurveCompressedDataChunk = StreamingManager.GetLoadedChunk(this, 0, bUsingFirstChunk); //Curve Data stored in chunk 0 till it is properly cropped
 
 		if (!CurveCompressedDataChunk)
 		{
 #if WITH_EDITOR
-			CurveCompressedDataChunk = GetRunningPlatformData().Chunks[0].CompressedAnimSequence;
+			CurveCompressedDataChunk = StreamableAnimData.Chunks[0].CompressedAnimSequence;
 #else
 			UE_LOG(LogAnimation, Warning, TEXT("Failed to get streamed compressed data Time: %.2f, ChunkIndex:%i, Anim: %s"), ExtractionContext.CurrentTime, 0, *GetFullName());
 			return;
@@ -343,14 +363,13 @@ void UAnimStreamable::GetAnimationPose(FAnimationPoseData& OutAnimationPoseData,
 
 		const FCompressedAnimSequence* CompressedData = bUsingFirstChunk ? CurveCompressedDataChunk : StreamingManager.GetLoadedChunk(this, ChunkIndex, true);
 
-		float ChunkCurrentTime = ExtractionContext.CurrentTime - GetRunningPlatformData().Chunks[ChunkIndex].StartTime;
-
+		double ChunkCurrentTime = ExtractionContext.CurrentTime - CurrentChunk.StartTime;
 		if (!CompressedData)
 		{
 #if WITH_EDITOR
-			CompressedData = GetRunningPlatformData().Chunks[ChunkIndex].CompressedAnimSequence;
+			CompressedData = CurrentChunk.CompressedAnimSequence;
 #else
-			const int32 NumChunks = GetRunningPlatformData().Chunks.Num();
+			const int32 NumChunks = StreamableAnimData.Chunks.Num();
 
 			int32 FallbackChunkIndex = ChunkIndex;
 			while (!CompressedData)
@@ -363,7 +382,7 @@ void UAnimStreamable::GetAnimationPose(FAnimationPoseData& OutAnimationPoseData,
 					return;
 				}
 				CompressedData = StreamingManager.GetLoadedChunk(this, FallbackChunkIndex, false);
-				ChunkCurrentTime = GetRunningPlatformData().Chunks[FallbackChunkIndex].SequenceLength;
+				ChunkCurrentTime = StreamableAnimData.Chunks[FallbackChunkIndex].SequenceLength;
 			}
 
 			UE_LOG(LogAnimation, Warning, TEXT("Failed to get streamed compressed data Time: %.2f, ChunkIndex:%i - Using Chunk %i Anim: %s"), ExtractionContext.CurrentTime, ChunkIndex, FallbackChunkIndex, *GetFullName());
@@ -381,7 +400,8 @@ void UAnimStreamable::GetAnimationPose(FAnimationPoseData& OutAnimationPoseData,
 		ChunkExtractionContext.PoseCurves = ExtractionContext.PoseCurves;
 
 		//FPlatformMisc::LowLevelOutputDebugStringf(TEXT("Playing Streaming Anim %s Time: %.2f Chunk:%i\n"), *GetName(), ExtractionContext.CurrentTime, ChunkIndex);
-		DecompressPose(OutPose, *CompressedData, ChunkExtractionContext, GetSkeleton(), GetRunningPlatformData().Chunks[ChunkIndex].SequenceLength, Interpolation, bIsBakedAdditive, RetargetSource, GetFName(), RootMotionReset);
+		FAnimSequenceDecompressionContext Context(SamplingFrameRate, CurrentChunk.NumFrames, Interpolation, GetFName(), *CompressedData->CompressedDataStructure.Get(), GetSkeleton(), bIsBakedAdditive);
+		UE::Anim::Decompression::DecompressPose(OutPose, *CompressedData, ChunkExtractionContext, Context, RetargetSource, RootMotionReset);
 	}
 }
 
@@ -469,16 +489,22 @@ void UAnimStreamable::InitFrom(const UAnimSequence* InSourceSequence)
 
 	BoneCompressionSettings = InSourceSequence->BoneCompressionSettings;
 	CurveCompressionSettings = InSourceSequence->CurveCompressionSettings;
+	
+	DataModelInterface = StaticDuplicateObject(InSourceSequence->GetDataModelInterface().GetObject(), this);
 
-	DataModel = DuplicateObject(InSourceSequence->GetDataModel(), this);
-	Controller->SetModel(DataModel);
+	// Far from ideal (retrieving controller to ensure it matches the DataModelInterface type)
+	Controller = DataModelInterface->GetController();
+	Controller->SetModel(DataModelInterface);
 
 	Notifies = InSourceSequence->Notifies;
-	Controller->SetPlayLength(InSourceSequence->GetDataModel()->GetPlayLength());
+	Controller->SetNumberOfFrames(InSourceSequence->GetDataModel()->GetNumberOfFrames());
 	RateScale = InSourceSequence->RateScale;
 	Interpolation = InSourceSequence->Interpolation;
 	RetargetSource = InSourceSequence->RetargetSource;
-	NumberOfKeys = DataModel->GetNumberOfKeys();
+	NumberOfKeys = DataModelInterface->GetNumberOfKeys();
+	SamplingFrameRate = SourceSequence->GetDataModel()->GetFrameRate();
+
+	Controller->NotifyPopulated();
 
 	bEnableRootMotion = InSourceSequence->bEnableRootMotion;
 	RootMotionRootLock = InSourceSequence->RootMotionRootLock;
@@ -568,8 +594,8 @@ void UAnimStreamable::RequestCompressedData(const ITargetPlatform* Platform)
 		NumChunks = FMath::Max(InitialNumChunks, 1);
 	}
 
-	int32 NumFramesToChunk = NumberOfKeys - 1;
-	int32 FramesPerChunk = NumberOfKeys / NumChunks;
+	const int32 NumFramesToChunk = NumberOfKeys - 1;
+	const int32 FramesPerChunk = NumFramesToChunk / NumChunks;
 
 	PlatformData.Chunks.AddDefaulted(NumChunks);
 
@@ -615,7 +641,7 @@ float UAnimStreamable::GetChunkSizeSeconds(const ITargetPlatform* Platform) cons
 }
 
 template<typename KeyType>
-void MakeKeyChunk(const TArray<KeyType>& SrcKeys, TArray<KeyType>& DestKeys, int32 NumFrames, const uint32 FrameStart, const uint32 FrameEnd)
+void MakeKeyChunk(const TArray<KeyType>& SrcKeys, TArray<KeyType>& DestKeys, int32 NumberOfKeys, const uint32 FrameStart, const uint32 FrameEnd)
 {
 	if (SrcKeys.Num() == 1)
 	{
@@ -623,7 +649,7 @@ void MakeKeyChunk(const TArray<KeyType>& SrcKeys, TArray<KeyType>& DestKeys, int
 	}
 	else
 	{
-		check(SrcKeys.Num() == NumFrames); // Invalid data otherwise
+		check(SrcKeys.Num() == NumberOfKeys); // Invalid data otherwise
 
 		DestKeys.Reset((FrameEnd - FrameStart) + 1);
 		for (uint32 FrameIndex = FrameStart; FrameIndex <= FrameEnd; ++FrameIndex)
@@ -650,10 +676,9 @@ void UAnimStreamable::RequestCompressedDataForChunk(const FString& ChunkDDCKey,
 		const bool bSkipDDC = false;
 
 		const int32 ChunkNumFrames = FrameEnd - FrameStart;
-		const int32 NumberOfFrames = NumberOfKeys - 1;
-		const float FrameLength = GetPlayLength() / (float)NumberOfFrames;
-		Chunk.StartTime = FrameStart * FrameLength;
-		Chunk.SequenceLength = ChunkNumFrames * FrameLength;
+		Chunk.StartTime = SamplingFrameRate.AsSeconds(static_cast<int32>(FrameStart));
+		Chunk.SequenceLength = SamplingFrameRate.AsSeconds(ChunkNumFrames);
+		Chunk.NumFrames = ChunkNumFrames;
 
 		if (!bSkipDDC && GetDerivedDataCacheRef().GetSynchronous(*FinalDDCKey, OutData, GetPathName()))
 		{
@@ -663,7 +688,7 @@ void UAnimStreamable::RequestCompressedDataForChunk(const FString& ChunkDDCKey,
 		{
 			FCompressibleAnimRef CompressibleData = MakeShared<FCompressibleAnimData, ESPMode::ThreadSafe>(BoneCompressionSettings, CurveCompressionSettings, GetSkeleton(), Interpolation, Chunk.SequenceLength, ChunkNumFrames+1);
 
-			const TArray<FBoneAnimationTrack>& BoneAnimationTracks = DataModel->GetBoneAnimationTracks();
+			const TArray<FBoneAnimationTrack>& BoneAnimationTracks = DataModelInterface->GetBoneAnimationTracks();
 			CompressibleData->TrackToSkeletonMapTable.Empty();
 
 			for (const FBoneAnimationTrack& AnimTrack : BoneAnimationTracks)
@@ -684,22 +709,9 @@ void UAnimStreamable::RequestCompressedDataForChunk(const FString& ChunkDDCKey,
 			if (FrameStart == 0)
 			{
 				//Crop curve logic broken, for the moment store curve data in always loaded chunk 0
-				CompressibleData->RawFloatCurves = DataModel->GetFloatCurves();
+				CompressibleData->RawFloatCurves = DataModelInterface->GetFloatCurves();
 			}
 
-			/*if (SourceSequence && ChunkIndex == 7)
-			{
-				UAnimSequence* Seq = const_cast<UAnimSequence*>(SourceSequence); //Stop update loop
-				SourceSequence = nullptr;
-				Seq->RawAnimationData = CompressibleData->RawAnimationData;
-				Seq->SequenceLength = CompressibleData->SequenceLength;
-				Seq->NumFrames = CompressibleData->NumFrames;
-				Seq->MarkRawDataAsModified(false);
-				Seq->OnRawDataChanged();
-
-				SourceSequence = Seq;
-			}*/
-
 			AnimCompressor->SetCompressibleData(CompressibleData);
 
 			if (bSkipDDC)
@@ -743,7 +755,7 @@ void UAnimStreamable::UpdateRawData()
 {
 	if (IsDataModelValid())
 	{
-		RawDataGuid = DataModel->GenerateGuid();
+		RawDataGuid = DataModelInterface->GenerateGuid();
 		RequestCompressedData();
 	}
 }
@@ -790,29 +802,4 @@ FStreamableAnimPlatformData& UAnimStreamable::GetStreamingAnimPlatformData(const
 #else
 	return RunningAnimPlatformData;
 #endif
-}
-
-/*EAdditiveAnimationType UAnimComposite::GetAdditiveAnimType() const
-{
-	int32 AdditiveType = AnimationTrack.GetTrackAdditiveType();
-
-	if (AdditiveType != -1)
-	{
-		return (EAdditiveAnimationType)AdditiveType;
-	}
-
-	return AAT_None;
-}
-
-bool UAnimComposite::HasRootMotion() const
-{
-	return AnimationTrack.HasRootMotion();
-}*/
-
-#if WITH_EDITOR
-/*class UAnimSequence* UAnimComposite::GetAdditiveBasePose() const
-{
-	// @todo : for now it just picks up the first sequence
-	return AnimationTrack.GetAdditiveBasePose();
-}*/
-#endif 
+}
\ No newline at end of file
diff --git a/Engine/Source/Runtime/Engine/Private/Animation/AnimationDecompression.cpp b/Engine/Source/Runtime/Engine/Private/Animation/AnimationDecompression.cpp
new file mode 100644
--- /dev/null
+++ b/Engine/Source/Runtime/Engine/Private/Animation/AnimationDecompression.cpp
@@ -0,0 +1,291 @@
+// Copyright Epic Games, Inc. All Rights Reserved.
+
+#include "Animation/AnimationDecompression.h"
+#include "Animation/AnimSequenceDecompressionContext.h"
+#include "Animation/AnimCompressionTypes.h"
+#include "Animation/AnimBoneCompressionCodec.h"
+
+#include "AnimationRuntime.h"
+#include "AnimEncoding.h"
+
+
+CSV_DECLARE_CATEGORY_MODULE_EXTERN(ENGINE_API, Animation);
+DECLARE_CYCLE_STAT(TEXT("Build Anim Track Pairs"), STAT_BuildAnimTrackPairs, STATGROUP_Anim);
+DECLARE_CYCLE_STAT(TEXT("Extract Pose From Anim Data"), STAT_ExtractPoseFromAnimData, STATGROUP_Anim);
+
+namespace UE::Anim::Decompression {
+	
+struct FGetBonePoseScratchArea : public TThreadSingleton<FGetBonePoseScratchArea>
+{
+	BoneTrackArray RotationScalePairs;
+	BoneTrackArray TranslationPairs;
+	BoneTrackArray AnimScaleRetargetingPairs;
+	BoneTrackArray AnimRelativeRetargetingPairs;
+	BoneTrackArray OrientAndScaleRetargetingPairs;
+};
+
+void DecompressPose(FCompactPose& OutPose,
+								const FCompressedAnimSequence& CompressedData,
+								const FAnimExtractContext& ExtractionContext,
+								FAnimSequenceDecompressionContext& DecompressionContext,
+								FName RetargetSource,
+								const FRootMotionReset& RootMotionReset)
+{
+	const TArray<FTransform>& RetargetTransforms = DecompressionContext.GetSourceSkeleton()->GetRefLocalPoses(RetargetSource);
+	DecompressPose(OutPose, CompressedData, ExtractionContext, DecompressionContext, RetargetTransforms, RootMotionReset);
+}
+
+void DecompressPose(FCompactPose& OutPose,
+	const FCompressedAnimSequence& CompressedData,
+	const FAnimExtractContext& ExtractionContext,
+	FAnimSequenceDecompressionContext& DecompressionContext,
+	const TArray<FTransform>& RetargetTransforms,
+	const FRootMotionReset& RootMotionReset)
+{
+	const FBoneContainer& RequiredBones = OutPose.GetBoneContainer();
+	const int32 NumTracks = CompressedData.CompressedTrackToSkeletonMapTable.Num();
+
+	const USkeleton* TargetSkeleton = RequiredBones.GetSkeletonAsset();
+	const FSkeletonRemapping* SkeletonRemapping = TargetSkeleton->GetSkeletonRemapping(DecompressionContext.GetSourceSkeleton());
+
+	BoneTrackArray& RotationScalePairs = FGetBonePoseScratchArea::Get().RotationScalePairs;
+	BoneTrackArray& TranslationPairs = FGetBonePoseScratchArea::Get().TranslationPairs;
+	BoneTrackArray& AnimScaleRetargetingPairs = FGetBonePoseScratchArea::Get().AnimScaleRetargetingPairs;
+	BoneTrackArray& AnimRelativeRetargetingPairs = FGetBonePoseScratchArea::Get().AnimRelativeRetargetingPairs;
+	BoneTrackArray& OrientAndScaleRetargetingPairs = FGetBonePoseScratchArea::Get().OrientAndScaleRetargetingPairs;
+
+	// build a list of desired bones
+	RotationScalePairs.Reset();
+	TranslationPairs.Reset();
+	AnimScaleRetargetingPairs.Reset();
+	AnimRelativeRetargetingPairs.Reset();
+	OrientAndScaleRetargetingPairs.Reset();
+
+	// Optimization: assuming first index is root bone. That should always be the case in Skeletons.
+	checkSlow((RequiredBones.GetMeshPoseIndexFromSkeletonPoseIndex(FSkeletonPoseBoneIndex(0)) == FMeshPoseBoneIndex(0)));
+	// this is not guaranteed for AnimSequences though... If Root is not animated, Track will not exist.
+	const bool bFirstTrackIsRootBone = (CompressedData.GetSkeletonIndexFromTrackIndex(0) == 0);
+
+	{
+		SCOPE_CYCLE_COUNTER(STAT_BuildAnimTrackPairs);
+
+		// Handle root bone separately if it is track 0. so we start w/ Index 1.
+		for (int32 TrackIndex = (bFirstTrackIsRootBone ? 1 : 0); TrackIndex < NumTracks; TrackIndex++)
+		{
+			const int32 SourceSkeletonBoneIndex = CompressedData.GetSkeletonIndexFromTrackIndex(TrackIndex);
+			const int32 TargetSkeletonBoneIndex = (SkeletonRemapping) ? SkeletonRemapping->GetTargetSkeletonBoneIndex(SourceSkeletonBoneIndex) : SourceSkeletonBoneIndex;
+
+			if (TargetSkeletonBoneIndex != INDEX_NONE)
+			{
+				const FCompactPoseBoneIndex BoneIndex = RequiredBones.GetCompactPoseIndexFromSkeletonIndex(TargetSkeletonBoneIndex);
+				//Nasty, we break our type safety, code in the lower levels should be adjusted for this
+				const int32 CompactPoseBoneIndex = BoneIndex.GetInt();
+				if (CompactPoseBoneIndex != INDEX_NONE)
+				{
+					RotationScalePairs.Add(BoneTrackPair(CompactPoseBoneIndex, TrackIndex));
+
+					// Skip extracting translation component for EBoneTranslationRetargetingMode::Skeleton.
+					switch (TargetSkeleton->GetBoneTranslationRetargetingMode(TargetSkeletonBoneIndex))
+					{
+					case EBoneTranslationRetargetingMode::Animation:
+						TranslationPairs.Add(BoneTrackPair(CompactPoseBoneIndex, TrackIndex));
+						break;
+					case EBoneTranslationRetargetingMode::AnimationScaled:
+						TranslationPairs.Add(BoneTrackPair(CompactPoseBoneIndex, TrackIndex));
+						AnimScaleRetargetingPairs.Add(BoneTrackPair(CompactPoseBoneIndex, SourceSkeletonBoneIndex));
+						break;
+					case EBoneTranslationRetargetingMode::AnimationRelative:
+						TranslationPairs.Add(BoneTrackPair(CompactPoseBoneIndex, TrackIndex));
+
+						// With baked additives, we can skip 'AnimationRelative' tracks, as the relative transform gets canceled out.
+						// (A1 + Rel) - (A2 + Rel) = A1 - A2.
+						if (!DecompressionContext.IsAdditiveAnimation())
+						{
+							AnimRelativeRetargetingPairs.Add(BoneTrackPair(CompactPoseBoneIndex, SourceSkeletonBoneIndex));
+						}
+						break;
+					case EBoneTranslationRetargetingMode::OrientAndScale:
+						TranslationPairs.Add(BoneTrackPair(CompactPoseBoneIndex, TrackIndex));
+
+						// Additives remain additives, they're not retargeted.
+						if (!DecompressionContext.IsAdditiveAnimation())
+						{
+							OrientAndScaleRetargetingPairs.Add(BoneTrackPair(CompactPoseBoneIndex, SourceSkeletonBoneIndex));
+						}
+						break;
+					}
+				}
+			}
+		}
+	}
+
+	{
+		SCOPE_CYCLE_COUNTER(STAT_ExtractPoseFromAnimData);
+		CSV_SCOPED_TIMING_STAT(Animation, ExtractPoseFromAnimData);
+		CSV_CUSTOM_STAT(Animation, NumberOfExtractedAnimations, 1, ECsvCustomStatOp::Accumulate);
+
+		DecompressionContext.Seek(ExtractionContext.CurrentTime);
+
+		// Handle Root Bone separately
+		if (bFirstTrackIsRootBone)
+		{
+			const int32 TrackIndex = 0;
+			FCompactPoseBoneIndex RootBone(0);
+			FTransform& RootAtom = OutPose[RootBone];
+
+			CompressedData.BoneCompressionCodec->DecompressBone(DecompressionContext, TrackIndex, RootAtom);
+
+			// Retarget the root onto the target skeleton (correcting for differences in rest poses)
+			if (SkeletonRemapping)
+			{
+				const int32 TargetSkeletonBoneIndex = 0;
+
+				if (DecompressionContext.IsAdditiveAnimation())
+				{
+					RootAtom.SetRotation(SkeletonRemapping->RetargetAdditiveRotationToTargetSkeleton(TargetSkeletonBoneIndex, RootAtom.GetRotation()));
+					if (TargetSkeleton->GetBoneTranslationRetargetingMode(TargetSkeletonBoneIndex) != EBoneTranslationRetargetingMode::Skeleton)
+					{
+						RootAtom.SetTranslation(SkeletonRemapping->RetargetAdditiveTranslationToTargetSkeleton(TargetSkeletonBoneIndex, RootAtom.GetTranslation()));
+					}
+				}
+				else
+				{
+					RootAtom.SetRotation(SkeletonRemapping->RetargetBoneRotationToTargetSkeleton(TargetSkeletonBoneIndex, RootAtom.GetRotation()));
+					if (TargetSkeleton->GetBoneTranslationRetargetingMode(TargetSkeletonBoneIndex) != EBoneTranslationRetargetingMode::Skeleton)
+					{
+						RootAtom.SetTranslation(SkeletonRemapping->RetargetBoneTranslationToTargetSkeleton(TargetSkeletonBoneIndex, RootAtom.GetTranslation()));
+					}
+				}
+			}
+
+			// @laurent - we should look into splitting rotation and translation tracks, so we don't have to process translation twice.
+			FAnimationRuntime::RetargetBoneTransform(DecompressionContext.GetSourceSkeleton(), DecompressionContext.AnimName, RetargetTransforms, RootAtom, 0, RootBone, RequiredBones, DecompressionContext.IsAdditiveAnimation());
+		}
+
+		if (RotationScalePairs.Num() > 0)
+		{
+			// get the remaining bone atoms
+			TArrayView<FTransform> OutPoseBones = OutPose.GetMutableBones();
+			CompressedData.BoneCompressionCodec->DecompressPose(DecompressionContext, RotationScalePairs, TranslationPairs, RotationScalePairs, OutPoseBones);
+		}
+	}
+
+	// Retarget the pose onto the target skeleton (correcting for differences in rest poses)
+	if (SkeletonRemapping)
+	{
+		if (DecompressionContext.IsAdditiveAnimation())
+		{
+			for (FCompactPoseBoneIndex BoneIndex(bFirstTrackIsRootBone ? 1 : 0); BoneIndex < OutPose.GetNumBones(); ++BoneIndex)
+			{
+				const int32 TargetSkeletonBoneIndex = RequiredBones.GetSkeletonIndex(BoneIndex);
+				OutPose[BoneIndex].SetRotation(SkeletonRemapping->RetargetAdditiveRotationToTargetSkeleton(TargetSkeletonBoneIndex, OutPose[BoneIndex].GetRotation()));
+				if (TargetSkeleton->GetBoneTranslationRetargetingMode(TargetSkeletonBoneIndex) != EBoneTranslationRetargetingMode::Skeleton)
+				{
+					OutPose[BoneIndex].SetTranslation(SkeletonRemapping->RetargetAdditiveTranslationToTargetSkeleton(TargetSkeletonBoneIndex, OutPose[BoneIndex].GetTranslation()));
+				}
+			}
+		}
+		else
+		{
+			for (FCompactPoseBoneIndex BoneIndex(bFirstTrackIsRootBone ? 1 : 0); BoneIndex < OutPose.GetNumBones(); ++BoneIndex)
+			{
+				const int32 TargetSkeletonBoneIndex = RequiredBones.GetSkeletonIndex(BoneIndex);
+				OutPose[BoneIndex].SetRotation(SkeletonRemapping->RetargetBoneRotationToTargetSkeleton(TargetSkeletonBoneIndex, OutPose[BoneIndex].GetRotation()));
+				if (TargetSkeleton->GetBoneTranslationRetargetingMode(TargetSkeletonBoneIndex) != EBoneTranslationRetargetingMode::Skeleton)
+				{
+					OutPose[BoneIndex].SetTranslation(SkeletonRemapping->RetargetBoneTranslationToTargetSkeleton(TargetSkeletonBoneIndex, OutPose[BoneIndex].GetTranslation()));
+				}
+			}
+		}
+	}
+
+	// Once pose has been extracted, snap root bone back to first frame if we are extracting root motion.
+	if ((ExtractionContext.bExtractRootMotion && RootMotionReset.bEnableRootMotion) || RootMotionReset.bForceRootLock)
+	{
+		RootMotionReset.ResetRootBoneForRootMotion(OutPose[FCompactPoseBoneIndex(0)], RequiredBones);
+	}
+
+	// Anim Scale Retargeting
+	int32 const NumBonesToScaleRetarget = AnimScaleRetargetingPairs.Num();
+	if (NumBonesToScaleRetarget > 0)
+	{
+		TArray<FTransform> const& AuthoredOnRefSkeleton = RetargetTransforms;
+
+		for (const BoneTrackPair& BonePair : AnimScaleRetargetingPairs)
+		{
+			const FCompactPoseBoneIndex BoneIndex(BonePair.AtomIndex); //Nasty, we break our type safety, code in the lower levels should be adjusted for this
+			const int32 SourceSkeletonBoneIndex = BonePair.TrackIndex;
+
+			// @todo - precache that in FBoneContainer when we have SkeletonIndex->TrackIndex mapping. So we can just apply scale right away.
+			float const SourceTranslationLength = AuthoredOnRefSkeleton[SourceSkeletonBoneIndex].GetTranslation().Size();
+			if (SourceTranslationLength > UE_KINDA_SMALL_NUMBER)
+			{
+				float const TargetTranslationLength = RequiredBones.GetRefPoseTransform(BoneIndex).GetTranslation().Size();
+				OutPose[BoneIndex].ScaleTranslation(TargetTranslationLength / SourceTranslationLength);
+			}
+		}
+	}
+
+	// Anim Relative Retargeting
+	int32 const NumBonesToRelativeRetarget = AnimRelativeRetargetingPairs.Num();
+	if (NumBonesToRelativeRetarget > 0)
+	{
+		TArray<FTransform> const& AuthoredOnRefSkeleton = RetargetTransforms;
+
+		for (const BoneTrackPair& BonePair : AnimRelativeRetargetingPairs)
+		{
+			const FCompactPoseBoneIndex BoneIndex(BonePair.AtomIndex); //Nasty, we break our type safety, code in the lower levels should be adjusted for this
+			const int32 SourceSkeletonBoneIndex = BonePair.TrackIndex;
+
+			const FTransform& RefPoseTransform = RequiredBones.GetRefPoseTransform(BoneIndex);
+
+			// Remap the base pose onto the target skeleton so that we are working entirely in target space
+			FTransform BaseTransform = AuthoredOnRefSkeleton[SourceSkeletonBoneIndex];
+			if (SkeletonRemapping)
+			{
+				const int32 TargetSkeletonBoneIndex = SkeletonRemapping->GetTargetSkeletonBoneIndex(SourceSkeletonBoneIndex);
+				BaseTransform = SkeletonRemapping->RetargetBoneTransformToTargetSkeleton(TargetSkeletonBoneIndex, BaseTransform);
+			}
+
+			// Apply the retargeting as if it were an additive difference between the current skeleton and the retarget skeleton. 
+			OutPose[BoneIndex].SetRotation(OutPose[BoneIndex].GetRotation() * BaseTransform.GetRotation().Inverse() * RefPoseTransform.GetRotation());
+			OutPose[BoneIndex].SetTranslation(OutPose[BoneIndex].GetTranslation() + (RefPoseTransform.GetTranslation() - BaseTransform.GetTranslation()));
+			OutPose[BoneIndex].SetScale3D(OutPose[BoneIndex].GetScale3D() * (RefPoseTransform.GetScale3D() * BaseTransform.GetSafeScaleReciprocal(BaseTransform.GetScale3D())));
+			OutPose[BoneIndex].NormalizeRotation();
+		}
+	}
+
+	// Translation 'Orient and Scale' Translation Retargeting
+	const int32 NumBonesToOrientAndScaleRetarget = OrientAndScaleRetargetingPairs.Num();
+	if (NumBonesToOrientAndScaleRetarget > 0)
+	{
+		const FRetargetSourceCachedData& RetargetSourceCachedData = RequiredBones.GetRetargetSourceCachedData(DecompressionContext.AnimName, RetargetTransforms);
+		const TArray<FOrientAndScaleRetargetingCachedData>& OrientAndScaleDataArray = RetargetSourceCachedData.OrientAndScaleData;
+		const TArray<int32>& CompactPoseIndexToOrientAndScaleIndex = RetargetSourceCachedData.CompactPoseIndexToOrientAndScaleIndex;
+
+		// If we have any cached retargeting data.
+		if ((OrientAndScaleDataArray.Num() > 0) && (CompactPoseIndexToOrientAndScaleIndex.Num() == RequiredBones.GetCompactPoseNumBones()))
+		{
+			for (int32 Index = 0; Index < NumBonesToOrientAndScaleRetarget; Index++)
+			{
+				const BoneTrackPair& BonePair = OrientAndScaleRetargetingPairs[Index];
+				const FCompactPoseBoneIndex CompactPoseBoneIndex(BonePair.AtomIndex);
+				const int32 OrientAndScaleIndex = CompactPoseIndexToOrientAndScaleIndex[CompactPoseBoneIndex.GetInt()];
+				if (OrientAndScaleIndex != INDEX_NONE)
+				{
+					const FOrientAndScaleRetargetingCachedData& OrientAndScaleData = OrientAndScaleDataArray[OrientAndScaleIndex];
+					FTransform& BoneTransform = OutPose[CompactPoseBoneIndex];
+					const FVector AnimatedTranslation = BoneTransform.GetTranslation();
+
+					// If Translation is not animated, we can just copy the TargetTranslation. No retargeting needs to be done.
+					const FVector NewTranslation = (AnimatedTranslation - OrientAndScaleData.SourceTranslation).IsNearlyZero(BONE_TRANS_RT_ORIENT_AND_SCALE_PRECISION) ?
+						OrientAndScaleData.TargetTranslation :
+						OrientAndScaleData.TranslationDeltaOrient.RotateVector(AnimatedTranslation) * OrientAndScaleData.TranslationScale;
+
+					BoneTransform.SetTranslation(NewTranslation);
+				}
+			}
+		}
+	}
+}
+}
diff --git a/Engine/Source/Runtime/Engine/Private/Animation/AnimationRuntime.cpp b/Engine/Source/Runtime/Engine/Private/Animation/AnimationRuntime.cpp
--- a/Engine/Source/Runtime/Engine/Private/Animation/AnimationRuntime.cpp
+++ b/Engine/Source/Runtime/Engine/Private/Animation/AnimationRuntime.cpp
@@ -1492,6 +1492,46 @@ void FAnimationRuntime::GetKeyIndicesFromTime(int32& OutKeyIndex1, int32& OutKey
 	OutAlpha = (float)Alpha;
 }
 
+void FAnimationRuntime::GetKeyIndicesFromTime(int32& OutKeyIndex1, int32& OutKeyIndex2, float& OutAlpha, const double Time, const FFrameRate& FrameRate, const int32 NumberOfKeys)
+{
+	// Check for 1-frame, before-first-frame and after-last-frame cases.
+	if (Time <= 0.0 || NumberOfKeys == 1)
+	{
+		OutKeyIndex1 = 0;
+		OutKeyIndex2 = 0;
+		OutAlpha = 0.0f;
+		return;
+	}
+
+	const FFrameTime FrameTime = FrameRate.AsFrameTime(Time);
+	const FFrameTime LastFrameTimeIndex = FFrameTime(NumberOfKeys - 1);
+	if (FrameTime >= LastFrameTimeIndex)
+	{
+		OutKeyIndex1 = LastFrameTimeIndex.FrameNumber.Value;
+		OutKeyIndex2 = 0;
+		OutAlpha = 0.0f;
+		return;
+	}
+
+	// Find the integer part (ensuring within range) and that gives us the 'starting' key index.
+	const int32 KeyIndex1 = FMath::Clamp<int32>(FrameTime.GetFrame().Value, 0, NumberOfKeys - 1); 
+
+	// The alpha (fractional part) is then just the remainder.
+	const float Alpha = FrameTime.GetSubFrame();
+
+	int32 KeyIndex2 = KeyIndex1 + 1;
+
+	// If we have gone over the end, do different things in case of looping
+	if (KeyIndex2 == NumberOfKeys)
+	{
+		KeyIndex2 = KeyIndex1;
+	}
+
+	OutKeyIndex1 = KeyIndex1;
+	OutKeyIndex2 = KeyIndex2;
+	OutAlpha = Alpha;
+}
+
 FTransform FAnimationRuntime::GetComponentSpaceRefPose(const FCompactPoseBoneIndex& CompactPoseBoneIndex, const FBoneContainer& BoneContainer)
 {
 	FCompactPoseBoneIndex CurrentIndex = CompactPoseBoneIndex;
diff --git a/Engine/Source/Runtime/Engine/Private/Animation/AnimationSettings.cpp b/Engine/Source/Runtime/Engine/Private/Animation/AnimationSettings.cpp
--- a/Engine/Source/Runtime/Engine/Private/Animation/AnimationSettings.cpp
+++ b/Engine/Source/Runtime/Engine/Private/Animation/AnimationSettings.cpp
@@ -1,7 +1,6 @@
 // Copyright Epic Games, Inc. All Rights Reserved.
 
 #include "Animation/AnimationSettings.h"
-#include "Animation/AnimCompress_BitwiseCompressOnly.h"
 #include "Animation/AttributeTypes.h"
 
 
@@ -38,6 +37,8 @@ UAnimationSettings::UAnimationSettings(const FObjectInitializer& ObjectInitializ
 		FMirrorFindReplaceExpression("([^}]*)_Left_([^}]*)", "$1_Right_$2", EMirrorFindReplaceMethod::RegularExpression), FMirrorFindReplaceExpression("([^}]*)_Right_([^}]*)", "$1_Left_$2", EMirrorFindReplaceMethod::RegularExpression),
 		FMirrorFindReplaceExpression("((?:^[sS]pine|^[rR]oot|^[pP]elvis|^[nN]eck|^[hH]ead|^ik_hand_gun).*)", "$1", EMirrorFindReplaceMethod::RegularExpression)
 	};
+	
+	DefaultFrameRate = FFrameRate(30,1);
 }
 
 TArray<FString> UAnimationSettings::GetBoneCustomAttributeNamesToImport() const
@@ -60,6 +61,11 @@ TArray<FString> UAnimationSettings::GetBoneCustomAttributeNamesToImport() const
 	return AttributeNames;
 }
 
+const FFrameRate& UAnimationSettings::GetDefaultFrameRate() const
+{
+	checkf(DefaultFrameRate.IsValid(), TEXT("Invalid default frame-rate set in Project Settings"));
+	return DefaultFrameRate;
+}
 
 #if WITH_EDITOR
 void UAnimationSettings::PostEditChangeProperty(struct FPropertyChangedEvent& PropertyChangedEvent)
diff --git a/Engine/Source/Runtime/Engine/Private/Animation/AnimationUtils.cpp b/Engine/Source/Runtime/Engine/Private/Animation/AnimationUtils.cpp
--- a/Engine/Source/Runtime/Engine/Private/Animation/AnimationUtils.cpp
+++ b/Engine/Source/Runtime/Engine/Private/Animation/AnimationUtils.cpp
@@ -278,16 +278,14 @@ void FAnimationUtils::ComputeCompressionError(const FCompressibleAnimData& Compr
 
 		const FTransform EndEffectorDummyBoneSocket(FQuat::Identity, FVector(END_EFFECTOR_DUMMY_BONE_LENGTH_SOCKET));
 		const FTransform EndEffectorDummyBone(FQuat::Identity, FVector(END_EFFECTOR_DUMMY_BONE_LENGTH));
-		const FAnimKeyHelper Helper(CompressibleAnimData.SequenceLength, CompressedData.AnimData->CompressedNumberOfKeys);
-		const float KeyLength = Helper.TimePerKey() + UE_SMALL_NUMBER;
 
-		FAnimSequenceDecompressionContext DecompContext(CompressibleAnimData.SequenceLength, CompressibleAnimData.Interpolation, CompressibleAnimData.AnimFName, *CompressedData.AnimData);
+		FAnimSequenceDecompressionContext DecompContext(CompressibleAnimData.SampledFrameRate, CompressibleAnimData.GetNumberOfFrames(), CompressibleAnimData.Interpolation, CompressibleAnimData.AnimFName, *CompressedData.AnimData, nullptr, CompressibleAnimData.bIsValidAdditive);
 
 		const TArray<FBoneData>& BoneData = CompressibleAnimData.BoneData;
 
 		for (int32 KeyIndex = 0; KeyIndex< CompressedData.AnimData->CompressedNumberOfKeys; KeyIndex++)
 		{
-			const float Time = (float)KeyIndex * KeyLength;
+			const double Time = CompressibleAnimData.SampledFrameRate.AsSeconds(KeyIndex);
 			DecompContext.Seek(Time);
 
 			// get the raw and compressed atom for each bone
@@ -315,7 +313,7 @@ void FAnimationUtils::ComputeCompressionError(const FCompressibleAnimData& Compr
 
 						FTransform AdditiveRawTransform;
 						FTransform AdditiveNewTransform;
-						FAnimationUtils::ExtractTransformFromTrack(Time, CompressibleAnimData.NumberOfKeys, CompressibleAnimData.SequenceLength, CompressibleAnimData.RawAnimationData[BoneIndexData.TrackIndex], CompressibleAnimData.Interpolation, AdditiveRawTransform);
+						FAnimationUtils::ExtractTransformFromTrack(CompressibleAnimData.RawAnimationData[BoneIndexData.TrackIndex], Time, CompressibleAnimData.NumberOfKeys, CompressibleAnimData.SequenceLength, CompressibleAnimData.Interpolation, AdditiveRawTransform);
 						CompressedData.Codec->DecompressBone(DecompContext, BoneIndexData.TrackIndex, AdditiveNewTransform);
 
 						const ScalarRegister VBlendWeight(1.f);
@@ -324,7 +322,7 @@ void FAnimationUtils::ComputeCompressionError(const FCompressibleAnimData& Compr
 					}
 					else
 					{
-						FAnimationUtils::ExtractTransformFromTrack(Time, CompressibleAnimData.NumberOfKeys, CompressibleAnimData.SequenceLength, CompressibleAnimData.RawAnimationData[BoneIndexData.TrackIndex], CompressibleAnimData.Interpolation, RawTransforms[BoneIndex]);
+						FAnimationUtils::ExtractTransformFromTrack(CompressibleAnimData.RawAnimationData[BoneIndexData.TrackIndex], Time, CompressibleAnimData.NumberOfKeys, CompressibleAnimData.SequenceLength,  CompressibleAnimData.Interpolation, RawTransforms[BoneIndex]);
 						CompressedData.Codec->DecompressBone(DecompContext, BoneIndexData.TrackIndex, NewTransforms[BoneIndex]);
 					}
 				}
@@ -447,7 +445,7 @@ UObject* FAnimationUtils::GetDefaultAnimSequenceOuter(UAnimSet* InAnimSet, bool
 		return NewPackage;
 	}
 
-	return NULL;
+	return nullptr;
 }
 
 
@@ -950,6 +948,7 @@ void FAnimationUtils::EnsureAnimSequenceLoaded(UAnimSequence& AnimSeq)
 	EnsureDependenciesAreLoaded(AnimSeq.GetSkeleton());
 	EnsureDependenciesAreLoaded(AnimSeq.BoneCompressionSettings);
 	EnsureDependenciesAreLoaded(AnimSeq.CurveCompressionSettings);
+	EnsureDependenciesAreLoaded(AnimSeq.RefPoseSeq);
 }
 
 void FAnimationUtils::ExtractTransformForFrameFromTrackSafe(const FRawAnimSequenceTrack& RawTrack, int32 Frame, FTransform& OutAtom)
@@ -983,6 +982,11 @@ void FAnimationUtils::ExtractTransformForFrameFromTrack(const FRawAnimSequenceTr
 }
 
 void FAnimationUtils::ExtractTransformFromTrack(float Time, int32 NumFrames, float SequenceLength, const FRawAnimSequenceTrack& RawTrack, EAnimInterpolationType Interpolation, FTransform &OutAtom)
+{
+	ExtractTransformFromTrack(RawTrack, static_cast<double>(Time), NumFrames, static_cast<double>(SequenceLength), Interpolation, OutAtom);
+}
+
+void FAnimationUtils::ExtractTransformFromTrack(const FRawAnimSequenceTrack& RawTrack, double Time, int32 NumFrames, double SequenceLength, EAnimInterpolationType Interpolation, FTransform &OutAtom)
 {
 	// Bail out (with rather wacky data) if data is empty for some reason.
 	if (RawTrack.PosKeys.Num() == 0 || RawTrack.RotKeys.Num() == 0)
@@ -1050,6 +1054,11 @@ void FAnimationUtils::ExtractTransformFromTrack(float Time, int32 NumFrames, flo
 
 #if WITH_EDITOR
 void FAnimationUtils::ExtractTransformFromCompressionData(const FCompressibleAnimData& CompressibleAnimData, FCompressibleAnimDataResult& CompressedAnimData, float Time, int32 TrackIndex, bool bUseRawData, FTransform& OutBoneTransform)
+{
+	ExtractTransformFromCompressionData(CompressibleAnimData, CompressedAnimData, (double)Time, TrackIndex, bUseRawData, OutBoneTransform);
+}
+
+void FAnimationUtils::ExtractTransformFromCompressionData(const FCompressibleAnimData& CompressibleAnimData, FCompressibleAnimDataResult& CompressedAnimData, double Time, int32 TrackIndex, bool bUseRawData, FTransform& OutBoneTransform)
 {
 	FUECompressedAnimDataMutable& AnimDataMutable = static_cast<FUECompressedAnimDataMutable&>(*CompressedAnimData.AnimData);
 
@@ -1059,13 +1068,13 @@ void FAnimationUtils::ExtractTransformFromCompressionData(const FCompressibleAni
 		// Build our read-only version from the mutable source
 		FUECompressedAnimData AnimData(AnimDataMutable);
 
-		FAnimSequenceDecompressionContext DecompContext(CompressibleAnimData.SequenceLength, CompressibleAnimData.Interpolation, CompressibleAnimData.AnimFName, AnimData);
+		FAnimSequenceDecompressionContext DecompContext(CompressibleAnimData.SampledFrameRate, CompressibleAnimData.GetNumberOfFrames(), CompressibleAnimData.Interpolation, CompressibleAnimData.AnimFName, AnimData, nullptr, CompressibleAnimData.bIsValidAdditive);
 		DecompContext.Seek(Time);
 		CompressedAnimData.Codec->DecompressBone(DecompContext, TrackIndex, OutBoneTransform);
 		return;
 	}
 
-	FAnimationUtils::ExtractTransformFromTrack(Time, CompressibleAnimData.NumberOfKeys, CompressibleAnimData.SequenceLength, CompressibleAnimData.RawAnimationData[TrackIndex], CompressibleAnimData.Interpolation, OutBoneTransform);
+	FAnimationUtils::ExtractTransformFromTrack(CompressibleAnimData.RawAnimationData[TrackIndex], Time, CompressibleAnimData.NumberOfKeys, CompressibleAnimData.SequenceLength, CompressibleAnimData.Interpolation, OutBoneTransform);
 }
 
 bool FAnimationUtils::CompressAnimBones(FCompressibleAnimData& AnimSeq, FCompressibleAnimDataResult& Target)
diff --git a/Engine/Source/Runtime/Engine/Private/Animation/AttributesRuntime.cpp b/Engine/Source/Runtime/Engine/Private/Animation/AttributesRuntime.cpp
--- a/Engine/Source/Runtime/Engine/Private/Animation/AttributesRuntime.cpp
+++ b/Engine/Source/Runtime/Engine/Private/Animation/AttributesRuntime.cpp
@@ -27,13 +27,18 @@ namespace UE { namespace Anim {
 
 #if WITH_EDITOR
 void Attributes::GetAttributeValue(FStackAttributeContainer& OutAttributes, const FCompactPoseBoneIndex& PoseBoneIndex, const FAnimatedBoneAttribute& Attribute, const FAnimExtractContext& ExtractionContext)
+{
+	GetAttributeValue(OutAttributes, PoseBoneIndex, Attribute, ExtractionContext.CurrentTime);
+}
+
+void Attributes::GetAttributeValue(FStackAttributeContainer& OutAttributes, const FCompactPoseBoneIndex& PoseBoneIndex, const FAnimatedBoneAttribute& Attribute, double CurrentTime)
 {
 	// Evaluating a single attribute into a stack-based container
 	if (Attribute.Identifier.IsValid())
 	{
 		uint8* AttributeDataPtr = OutAttributes.FindOrAdd(Attribute.Identifier.GetType(), FAttributeId(Attribute.Identifier.GetName(), PoseBoneIndex));
 		check(AttributeDataPtr)
-		Attribute.Curve.EvaluateToPtr(Attribute.Identifier.GetType(), ExtractionContext.CurrentTime, AttributeDataPtr);
+		Attribute.Curve.EvaluateToPtr(Attribute.Identifier.GetType(), CurrentTime, AttributeDataPtr);
 	}
 }
 #endif // WITH_EDITOR
diff --git a/Engine/Source/Runtime/Engine/Private/Animation/BlendSpace.cpp b/Engine/Source/Runtime/Engine/Private/Animation/BlendSpace.cpp
--- a/Engine/Source/Runtime/Engine/Private/Animation/BlendSpace.cpp
+++ b/Engine/Source/Runtime/Engine/Private/Animation/BlendSpace.cpp
@@ -794,7 +794,7 @@ void UBlendSpace::ResetToRefPose(FCompactPose& OutPose) const
 	}
 }
 
-static const FAnimExtractContext DefaultBlendSpaceExtractionContext = { 0.f, true, {}, false };
+static const FAnimExtractContext DefaultBlendSpaceExtractionContext = { 0.0, true, {}, false };
 
 void UBlendSpace::GetAnimationPose(TArray<FBlendSampleData>& BlendSampleDataCache, /*out*/ FCompactPose& OutPose, /*out*/ FBlendedCurve& OutCurve) const
 {
@@ -896,7 +896,7 @@ void UBlendSpace::GetAnimationPose_Internal(TArray<FBlendSampleData>& BlendSampl
 
 					FAnimationPoseData ChildAnimationPoseData = { Pose, ChildrenCurves[I], ChildrenAttributes[I] };
 					// first one always fills up the source one
-					Sample.Animation->GetAnimationPose(ChildAnimationPoseData, FAnimExtractContext(Time, ExtractionContext.bExtractRootMotion, BlendSampleDataCache[I].DeltaTimeRecord, ExtractionContext.bLooping));
+					Sample.Animation->GetAnimationPose(ChildAnimationPoseData, FAnimExtractContext(static_cast<double>(Time), ExtractionContext.bExtractRootMotion, BlendSampleDataCache[I].DeltaTimeRecord, ExtractionContext.bLooping));
 				}
 				else
 				{
@@ -1622,7 +1622,7 @@ FVector UBlendSpace::GetNormalizedBlendInput(const FVector& BlendInput) const
 
 const FEditorElement* UBlendSpace::GetGridSampleInternal(int32 Index) const
 {
-	return GridSamples.IsValidIndex(Index) ? &GridSamples[Index] : NULL;
+	return GridSamples.IsValidIndex(Index) ? &GridSamples[Index] : nullptr;
 }
 
 // When using CriticallyDampedSmoothing, how to go from the interpolation speed to the smooth
diff --git a/Engine/Source/Runtime/Engine/Private/Animation/BonePose.cpp b/Engine/Source/Runtime/Engine/Private/Animation/BonePose.cpp
--- a/Engine/Source/Runtime/Engine/Private/Animation/BonePose.cpp
+++ b/Engine/Source/Runtime/Engine/Private/Animation/BonePose.cpp
@@ -4,6 +4,7 @@
 #include "AnimationRuntime.h"
 #include "AnimEncoding.h"
 #include "HAL/ThreadSingleton.h"
+#include "Animation/AnimSequenceHelpers.h"
 #if INTEL_ISPC
 #include "BonePose.ispc.generated.h"
 
@@ -136,17 +137,6 @@ bool FMeshPose::IsNormalized() const
 	return true;
 }
 
-struct FRetargetTracking
-{
-	const FCompactPoseBoneIndex PoseBoneIndex;
-	const int32 SkeletonBoneIndex;
-
-	FRetargetTracking(const FCompactPoseBoneIndex InPoseBoneIndex, const int32 InSkeletonBoneIndex)
-		: PoseBoneIndex(InPoseBoneIndex), SkeletonBoneIndex(InSkeletonBoneIndex)
-	{
-	}
-};
-
 FTransform ExtractTransformForKey(int32 Key, const FRawAnimSequenceTrack &TrackToExtract)
 {
 	static const FVector DefaultScale3D = FVector(1.f);
@@ -165,23 +155,16 @@ FTransform ExtractTransformForKey(int32 Key, const FRawAnimSequenceTrack &TrackT
 	}
 }
 
-struct FBuildRawPoseScratchArea : public TThreadSingleton<FBuildRawPoseScratchArea>
-{
-	TArray<FRetargetTracking> RetargetTracking;
-	TArray<FVirtualBoneCompactPoseData> VirtualBoneCompactPoseData;
-};
-
-
 template<bool bInterpolateT>
 void BuildPoseFromRawDataInternal(const TArray<FRawAnimSequenceTrack>& InAnimationData, const TArray<struct FTrackToSkeletonMap>& TrackToSkeletonMapTable, FCompactPose& InOutPose, int32 KeyIndex1, int32 KeyIndex2, float Alpha, float TimePerKey, const TMap<int32, const FTransformCurve*>* AdditiveBoneTransformCurves)
 {
 	const int32 NumTracks = InAnimationData.Num();
 	const FBoneContainer& RequiredBones = InOutPose.GetBoneContainer();
 
-	TArray<FRetargetTracking>& RetargetTracking = FBuildRawPoseScratchArea::Get().RetargetTracking;
+	TArray<UE::Anim::Retargeting::FRetargetTracking>& RetargetTracking = UE::Anim::FBuildRawPoseScratchArea::Get().RetargetTracking;
 	RetargetTracking.Reset(NumTracks);
 
-	TArray<FVirtualBoneCompactPoseData>& VBCompactPoseData = FBuildRawPoseScratchArea::Get().VirtualBoneCompactPoseData;
+	TArray<FVirtualBoneCompactPoseData>& VBCompactPoseData =  UE::Anim::FBuildRawPoseScratchArea::Get().VirtualBoneCompactPoseData;
 	VBCompactPoseData = RequiredBones.GetVirtualBoneCompactPoseData();
 
 	FCompactPose Key2Pose;
@@ -254,7 +237,7 @@ void BuildPoseFromRawDataInternal(const TArray<FRawAnimSequenceTrack>& InAnimati
 					}
 				}
 
-				RetargetTracking.Add(FRetargetTracking(PoseBoneIndex, SkeletonBoneIndex));
+				RetargetTracking.Add(UE::Anim::Retargeting::FRetargetTracking(PoseBoneIndex, SkeletonBoneIndex));
 			}
 		}
 	}
@@ -332,7 +315,9 @@ void BuildPoseFromRawData(
 {
 	int32 KeyIndex1, KeyIndex2;
 	float Alpha;
+	PRAGMA_DISABLE_DEPRECATION_WARNINGS
 	FAnimationRuntime::GetKeyIndicesFromTime(KeyIndex1, KeyIndex2, Alpha, InTime, NumFrames, SequenceLength);
+	PRAGMA_ENABLE_DEPRECATION_WARNINGS
 
 	if (Interpolation == EAnimInterpolationType::Step)
 	{
@@ -368,11 +353,10 @@ void BuildPoseFromRawData(
 
 	if (!bDisableRetargeting)
 	{
-		const TArray<FRetargetTracking>& RetargetTracking = FBuildRawPoseScratchArea::Get().RetargetTracking;
-
-		USkeleton* Skeleton = RequiredBones.GetSkeletonAsset();
+		const TArray<UE::Anim::Retargeting::FRetargetTracking>& RetargetTracking = UE::Anim::FBuildRawPoseScratchArea::Get().RetargetTracking;
 
-		for (const FRetargetTracking& RT : RetargetTracking)
+		const USkeleton* Skeleton = RequiredBones.GetSkeletonAsset();
+		for (const UE::Anim::Retargeting::FRetargetTracking& RT : RetargetTracking)
 		{
 			FAnimationRuntime::RetargetBoneTransform(Skeleton, SourceName, RetargetTransforms, InOutPose[RT.PoseBoneIndex], RT.SkeletonBoneIndex, RT.PoseBoneIndex, RequiredBones, false);
 		}
