Reintroducing AnimationDataModel refactor with additional memory and cooktime optimizations.

diff --git a/Engine/Source/Runtime/Engine/Private/Animation/PoseAsset.cpp b/Engine/Source/Runtime/Engine/Private/Animation/PoseAsset.cpp
--- a/Engine/Source/Runtime/Engine/Private/Animation/PoseAsset.cpp
+++ b/Engine/Source/Runtime/Engine/Private/Animation/PoseAsset.cpp
@@ -9,7 +9,7 @@
 #include "Animation/AnimSequence.h"
 #include "Animation/AnimSequenceBase.h"
 #include "Animation/AnimationPoseData.h"
-#include "Animation/AnimData/AnimDataModel.h"
+#include "Animation/AnimData/IAnimationDataModel.h"
 #include "Animation/AnimSequenceHelpers.h"
 #include "Framework/Notifications/NotificationManager.h"
 #include "UObject/ObjectSaveContext.h"
@@ -925,12 +925,16 @@ void UPoseAsset::PostLoad()
 			}
 			SourceAnimation->ConditionalPostLoad();
 			
-			if (!SourceAnimationRawDataGUID.IsValid() || SourceAnimationRawDataGUID != SourceAnimation->GetRawDataGuid())
+			if (!SourceAnimationRawDataGUID.IsValid() || SourceAnimationRawDataGUID != SourceAnimation->GetDataModel()->GenerateGuid())
 			{
 				FFormatNamedArguments Args;
 				Args.Add(TEXT("AssetName"), FText::FromString(GetPathName()));
 				Args.Add(TEXT("SourceAsset"), FText::FromString(SourceAnimation->GetPathName()));
-				const FText ResultText = FText::Format(LOCTEXT("PoseAssetSourceOutOfDate", "PoseAsset {AssetName} is out-of-date with its source animation {SourceAsset}"), Args);
+
+				Args.Add(TEXT("Stored"), FText::FromString(SourceAnimationRawDataGUID.ToString()));
+				Args.Add(TEXT("Found"), FText::FromString(SourceAnimation->GetDataModel()->GenerateGuid().ToString()));
+				
+				const FText ResultText = FText::Format(LOCTEXT("PoseAssetSourceOutOfDate", "PoseAsset {AssetName} is out-of-date with its source animation {SourceAsset} {Stored} vs {Found}"), Args);
 				UE_LOG(LogAnimation, Warning,TEXT("%s"), *ResultText.ToString());
 			}
 		}	
@@ -1380,7 +1384,7 @@ void UPoseAsset::CreatePoseFromAnimation(class UAnimSequence* AnimSequence, cons
 		{
 			SetSkeleton(TargetSkeleton);
 			SourceAnimation = AnimSequence;
-			SourceAnimationRawDataGUID = AnimSequence->GetRawDataGuid();
+			SourceAnimationRawDataGUID = AnimSequence->GetDataModel()->GenerateGuid();
 
 			// reinitialize, now we're making new pose from this animation
 			Reinitialize();
@@ -1394,7 +1398,7 @@ void UPoseAsset::CreatePoseFromAnimation(class UAnimSequence* AnimSequence, cons
 				FMemMark Mark(FMemStack::Get());
 
 				// set up track data - @todo: add revaliation code when checked
-				UAnimDataModel* DataModel = AnimSequence->GetDataModel();
+				IAnimationDataModel* DataModel = AnimSequence->GetDataModel();
 
 				TArray<FName> TrackNames;
 				DataModel->GetBoneTrackNames(TrackNames);
diff --git a/Engine/Source/Runtime/Engine/Private/Curves/KeyHandle.cpp b/Engine/Source/Runtime/Engine/Private/Curves/KeyHandle.cpp
--- a/Engine/Source/Runtime/Engine/Private/Curves/KeyHandle.cpp
+++ b/Engine/Source/Runtime/Engine/Private/Curves/KeyHandle.cpp
@@ -8,7 +8,18 @@ FKeyHandle::FKeyHandle()
 	static uint32 LastKeyHandleIndex = 1;
 	Index = ++LastKeyHandleIndex;
 
-	check(LastKeyHandleIndex != 0); // check in the unlikely event that this overflows
+	if (LastKeyHandleIndex == 0)
+	{
+		// If we are cooking, allow wrap-around
+		if (IsRunningCookCommandlet())
+		{
+			Index = LastKeyHandleIndex = 1;
+		}
+		else
+		{
+			check(LastKeyHandleIndex != 0); // check in the unlikely event that this overflows
+		}
+	}
 }
 
 FKeyHandle::FKeyHandle(uint32 SpecificIndex)
diff --git a/Engine/Source/Runtime/Engine/Private/Curves/RichCurve.cpp b/Engine/Source/Runtime/Engine/Private/Curves/RichCurve.cpp
--- a/Engine/Source/Runtime/Engine/Private/Curves/RichCurve.cpp
+++ b/Engine/Source/Runtime/Engine/Private/Curves/RichCurve.cpp
@@ -964,7 +964,9 @@ void FRichCurve::RemoveRedundantKeysInternal(float Tolerance, int32 InStartKeepK
 	TArray<FKeyHandle> AllHandlesByIndex;
 	TArray<FKeyHandle> KeepHandles;
 
-	if (KeyHandlesToIndices.Num() != 0)
+	// If we have no handles to start with do not bother generating any
+	const bool bHasKeyHandles = KeyHandlesToIndices.Num() != 0;
+	if (bHasKeyHandles)
 	{
 		check(KeyHandlesToIndices.Num() == Keys.Num());
 		AllHandlesByIndex.AddZeroed(Keys.Num());
@@ -975,11 +977,6 @@ void FRichCurve::RemoveRedundantKeysInternal(float Tolerance, int32 InStartKeepK
 			AllHandlesByIndex[HandleIndexPair.Value] = HandleIndexPair.Key;
 		}
 	}
-	else
-	{
-		AllHandlesByIndex.AddDefaulted(Keys.Num());
-	}
-	
 
 	{
 		TArray<FRichCurveKey> NewKeys;
@@ -989,7 +986,11 @@ void FRichCurve::RemoveRedundantKeysInternal(float Tolerance, int32 InStartKeepK
 		for(int32 StartKeepIndex = 0; StartKeepIndex <= ActualStartKeepKey; ++StartKeepIndex)
 		{
 			NewKeys.Add(Keys[StartKeepIndex]);
-			KeepHandles.Add(AllHandlesByIndex[StartKeepIndex]);
+			
+			if (bHasKeyHandles)
+			{
+				KeepHandles.Add(AllHandlesByIndex[StartKeepIndex]);
+			}
 		}
 
 		//Add keys up to the first end keep key if they are not redundant
@@ -1050,7 +1051,11 @@ void FRichCurve::RemoveRedundantKeysInternal(float Tolerance, int32 InStartKeepK
 			{
 				MostRecentKeepKeyIndex = TestIndex;
 				NewKeys.Add(Keys[TestIndex]);
-				KeepHandles.Add(AllHandlesByIndex[TestIndex]);
+
+				if(bHasKeyHandles)
+				{
+					KeepHandles.Add(AllHandlesByIndex[TestIndex]);
+				}
 			}
 		}
 
@@ -1058,16 +1063,22 @@ void FRichCurve::RemoveRedundantKeysInternal(float Tolerance, int32 InStartKeepK
 		for (int32 EndKeepIndex = ActualEndKeepKey; EndKeepIndex < Keys.Num(); ++EndKeepIndex)
 		{
 			NewKeys.Add(Keys[EndKeepIndex]);
-			KeepHandles.Add(AllHandlesByIndex[EndKeepIndex]);
+			if (bHasKeyHandles)
+			{
+				KeepHandles.Add(AllHandlesByIndex[EndKeepIndex]);
+			}
 		}
 		Keys = MoveTemp(NewKeys); //Do this at the end of scope, guaranteed that NewKeys is going away
 	}
 
 	AutoSetTangents();
 
-	// Rebuild KeyHandlesToIndices
-	check(Keys.Num() == KeepHandles.Num());
-	KeyHandlesToIndices.Initialize(KeepHandles);
+	if (bHasKeyHandles)
+	{
+		// Rebuild KeyHandlesToIndices
+		check(Keys.Num() == KeepHandles.Num());
+		KeyHandlesToIndices.Initialize(KeepHandles);
+	}
 }
 
 void FRichCurve::RemapTimeValue(float& InTime, float& CycleValueOffset) const
@@ -1144,7 +1155,7 @@ float FRichCurve::Eval(float InTime, float InDefaultValue) const
 	{
 		// If no keys in curve, return the Default value.
 	} 
-	else if (NumKeys < 2 || (InTime <= Keys[0].Time))
+	else if (NumKeys < 2 || (InTime < Keys[0].Time))
 	{
 		if (PreInfinityExtrap == RCCE_Linear && NumKeys > 1)
 		{
@@ -1168,7 +1179,7 @@ float FRichCurve::Eval(float InTime, float InDefaultValue) const
 			InterpVal = Keys[0].Value;
 		}
 	}
-	else if (InTime < Keys[NumKeys - 1].Time)
+	else if (InTime <= Keys[NumKeys - 1].Time)
 	{
 		// perform a lower bound to get the second of the interpolation nodes
 		int32 first = 1;
@@ -1177,8 +1188,8 @@ float FRichCurve::Eval(float InTime, float InDefaultValue) const
 
 		while (count > 0)
 		{
-			int32 step = count / 2;
-			int32 middle = first + step;
+			const int32 step = count / 2;
+			const int32 middle = first + step;
 
 			if (InTime >= Keys[middle].Time)
 			{
@@ -1264,13 +1275,10 @@ static ERichCurveCompressionFormat FindRichCurveCompressionFormat(const FRichCur
 	{
 		if (Key.InterpMode == RCIM_Cubic)
 		{
-			if (Key.TangentMode != RCTM_Auto && Key.TangentMode != RCTM_None)
+			if (Key.TangentWeightMode != RCTWM_WeightedNone)
 			{
-				if (Key.TangentWeightMode != RCTWM_WeightedNone)
-				{
-					bIsWeighted = true;
-					break;
-				}
+				bIsWeighted = true;
+				break;
 			}
 		}
 		if (Key.InterpMode != RefKey.InterpMode)
@@ -1448,10 +1456,10 @@ static ERichCurveKeyTimeCompressionFormat FindRichCurveKeyFormat(const FRichCurv
 
 void FRichCurve::CompressCurve(FCompressedRichCurve& OutCurve, float ErrorThreshold, float SampleRate) const
 {
-	ERichCurveCompressionFormat CompressionFormat = FindRichCurveCompressionFormat(*this);
+	const ERichCurveCompressionFormat CompressionFormat = FindRichCurveCompressionFormat(*this);
 	OutCurve.CompressionFormat = CompressionFormat;
 
-	ERichCurveKeyTimeCompressionFormat KeyFormat = FindRichCurveKeyFormat(*this, ErrorThreshold, SampleRate, CompressionFormat);
+	const ERichCurveKeyTimeCompressionFormat KeyFormat = FindRichCurveKeyFormat(*this, ErrorThreshold, SampleRate, CompressionFormat);
 	OutCurve.KeyTimeCompressionFormat = KeyFormat;
 
 	OutCurve.PreInfinityExtrap = PreInfinityExtrap;
@@ -1977,7 +1985,7 @@ FORCEINLINE_DEBUGGABLE static float InterpEval(float InTime, const KeyTimeAdapte
 	float CycleValueOffset = 0.0;
 
 	const float FirstKeyTime = KeyTimeAdapter.GetTime(0);
-	if (InTime <= FirstKeyTime)
+	if (InTime < FirstKeyTime)
 	{
 		if (PreInfinityExtrap != RCCE_Linear && PreInfinityExtrap != RCCE_Constant)
 		{
@@ -1990,7 +1998,7 @@ FORCEINLINE_DEBUGGABLE static float InterpEval(float InTime, const KeyTimeAdapte
 	}
 
 	const float LastKeyTime = KeyTimeAdapter.GetTime(NumKeys - 1);
-	if (InTime >= LastKeyTime)
+	if (InTime > LastKeyTime)
 	{
 		if (PostInfinityExtrap != RCCE_Linear && PostInfinityExtrap != RCCE_Constant)
 		{
@@ -2032,6 +2040,14 @@ FORCEINLINE_DEBUGGABLE static float InterpEval(float InTime, const KeyTimeAdapte
 	const KeyDataHandle KeyValueHandle0 = KeyDataAdapter.GetKeyDataHandle(First - 1);
 	const float KeyValue0 = KeyDataAdapter.GetKeyValue(KeyValueHandle0);
 
+	auto IsItNotWeighted = [KeyDataAdapter](KeyDataHandle Handle0, KeyDataHandle Handle1) -> bool
+	{
+		const ERichCurveTangentWeightMode Mode0 = KeyDataAdapter.GetKeyTangentWeightMode(Handle0);
+		const ERichCurveTangentWeightMode Mode1 = KeyDataAdapter.GetKeyTangentWeightMode(Handle1);
+		
+		return (Mode0 == RCTWM_WeightedNone || Mode0 == RCTWM_WeightedArrive) && (Mode1 == RCTWM_WeightedNone || Mode1 == RCTWM_WeightedLeave);
+	};
+
 	// const value here allows the code to be stripped statically if the data is uniform
 	// which it is most of the time
 	const ERichCurveCompressionFormat KeyInterpMode0 = KeyDataAdapter.GetKeyInterpMode(First - 1);
@@ -2049,7 +2065,7 @@ FORCEINLINE_DEBUGGABLE static float InterpEval(float InTime, const KeyTimeAdapte
 		{
 			InterpolatedValue = FMath::Lerp(P0, P3, Alpha);
 		}
-		else if (KeyInterpMode0 != RCCF_Weighted || KeyDataAdapter.GetKeyInterpMode(First) != RCCF_Weighted)
+		else if (IsItNotWeighted(KeyValueHandle0, KeyValueHandle1))
 		{
 			const float OneThird = 1.0f / 3.0f;
 			const float ScaledDiff = Diff * OneThird;
diff --git a/Engine/Source/Runtime/Engine/Public/Animation/AnimCompressionTypes.h b/Engine/Source/Runtime/Engine/Public/Animation/AnimCompressionTypes.h
--- a/Engine/Source/Runtime/Engine/Public/Animation/AnimCompressionTypes.h
+++ b/Engine/Source/Runtime/Engine/Public/Animation/AnimCompressionTypes.h
@@ -21,6 +21,7 @@
 #include "Serialization/MemoryWriter.h"
 
 #include "ProfilingDebugging/CsvProfiler.h"
+#include "Animation/AnimationDecompression.h"
 
 #include "AnimCompressionTypes.generated.h"
 
@@ -222,10 +223,12 @@ public:
 
 	/** Number of keys within the (non-uniform) RawAnimationData tracks */
 	
-	UE_DEPRECATED(5.0, "NumberOfFrames has been replaced with NumberOfKeys")
+	UE_DEPRECATED(5.0, "NumberOfFrames has been replaced with NumberOfKeys and GetNumberOfFrames")
 	int32 NumberOfFrames;
 	int32 NumberOfKeys;
 
+	int32 GetNumberOfFrames() const { return FMath::Max(NumberOfKeys - 1, 1); }
+
 	bool bIsValidAdditive;
 
 	float ErrorThresholdScale;
@@ -235,6 +238,8 @@ public:
 	FString FullName;
 	FName   AnimFName;
 
+	FFrameRate SampledFrameRate;
+
 	FCancelCompressionSignal IsCancelledSignal;
 
 	static int32 GetApproxRawDataArraySize(const TArray<FRawAnimSequenceTrack>& AnimData)
@@ -806,6 +811,7 @@ struct FRootMotionReset
 	}
 };
 
+UE_DEPRECATED(5.1, "Signature of DecompressPose has been deprecated, use UE::Anim::Decompression::DecompressPose instead")
 extern void DecompressPose(FCompactPose& OutPose,
 							const FCompressedAnimSequence& CompressedData,
 							const FAnimExtractContext& ExtractionContext,
@@ -817,6 +823,7 @@ extern void DecompressPose(FCompactPose& OutPose,
 							FName SourceName,
 							const FRootMotionReset& RootMotionReset);
 
+UE_DEPRECATED(5.1, "Signature of DecompressPose has been deprecated, use UE::Anim::Decompression::DecompressPose instead")
 extern void DecompressPose(	FCompactPose& OutPose,
 							const FCompressedAnimSequence& CompressedData,
 							const FAnimExtractContext& ExtractionContext,
diff --git a/Engine/Source/Runtime/Engine/Public/Animation/AnimSequenceDecompressionContext.h b/Engine/Source/Runtime/Engine/Public/Animation/AnimSequenceDecompressionContext.h
--- a/Engine/Source/Runtime/Engine/Public/Animation/AnimSequenceDecompressionContext.h
+++ b/Engine/Source/Runtime/Engine/Public/Animation/AnimSequenceDecompressionContext.h
@@ -4,24 +4,45 @@
 
 #include "CoreMinimal.h"
 #include "Animation/AnimTypes.h"
+#include "Misc/FrameRate.h"
 
 struct ICompressedAnimData;
+class USkeleton;
 
 /* Encapsulates decompression related data used by bone compression codecs. */
 struct FAnimSequenceDecompressionContext
 {
-	FAnimSequenceDecompressionContext(float SequenceLength_, EAnimInterpolationType Interpolation_, const FName& AnimName_, const ICompressedAnimData& CompressedAnimData_)
-			:
+	UE_DEPRECATED(5.1, "FAnimSequenceDecompressionContext signature has been deprecated, use different one")
+	FAnimSequenceDecompressionContext(float SequenceLength_, EAnimInterpolationType Interpolation_, const FName& AnimName_, const ICompressedAnimData& CompressedAnimData_,  const USkeleton* InSourceSkeleton, bool bIsBakedAdditive)
+		:
 		PRAGMA_DISABLE_DEPRECATION_WARNINGS
 		SequenceLength(SequenceLength_),
 		PRAGMA_ENABLE_DEPRECATION_WARNINGS
 		Interpolation(Interpolation_), AnimName(AnimName_), CompressedAnimData(CompressedAnimData_),
 		PRAGMA_DISABLE_DEPRECATION_WARNINGS
-		Time(0.f), RelativePos(0.f)
+		Time(0.f), RelativePos(0.f),
 		PRAGMA_ENABLE_DEPRECATION_WARNINGS
+		SourceSkeleton(InSourceSkeleton),
+		bAdditiveAnimation(bIsBakedAdditive)
 	{
 	}
 
+	FAnimSequenceDecompressionContext(const FFrameRate& InSamplingRate, const int32 InNumberOfFrames, EAnimInterpolationType Interpolation_, const FName& AnimName_, const ICompressedAnimData& CompressedAnimData_, const USkeleton* InSourceSkeleton, bool bIsBakedAdditive)
+		:
+		PRAGMA_DISABLE_DEPRECATION_WARNINGS
+		SequenceLength(InSamplingRate.AsSeconds(InNumberOfFrames)),
+		PRAGMA_ENABLE_DEPRECATION_WARNINGS
+		Interpolation(Interpolation_), AnimName(AnimName_), CompressedAnimData(CompressedAnimData_),
+		PRAGMA_DISABLE_DEPRECATION_WARNINGS
+		Time(0.f), RelativePos(0.f),
+		PRAGMA_ENABLE_DEPRECATION_WARNINGS
+		SamplingRate(InSamplingRate),
+		SamplingTime(0),
+		NumberOfFrames(InNumberOfFrames),
+		SourceSkeleton(InSourceSkeleton),
+		bAdditiveAnimation(bIsBakedAdditive)
+	{}
+
 	// Anim info
 	UE_DEPRECATED(5.1, "Direct access to SequenceLength has been deprecated use GetPlayableLength instead")
 	float SequenceLength;
@@ -36,33 +57,49 @@ struct FAnimSequenceDecompressionContext
 	UE_DEPRECATED(5.1, "Direct access to RelativePos has been deprecated use GetRelativePosition instead")
 	float RelativePos;
 
-	float GetPlayableLength() const
+	double GetPlayableLength() const
 	{
-		PRAGMA_DISABLE_DEPRECATION_WARNINGS
-		return SequenceLength;
-		PRAGMA_ENABLE_DEPRECATION_WARNINGS
+		return SamplingRate.AsSeconds(NumberOfFrames);
 	}
-
-	float GetEvaluationTime() const
+	double GetEvaluationTime() const
 	{
-		PRAGMA_DISABLE_DEPRECATION_WARNINGS
-		return Time;
-		PRAGMA_ENABLE_DEPRECATION_WARNINGS
+		return SamplingRate.AsSeconds(SamplingTime);
+	}
+	double GetRelativePosition() const
+	{
+		return FMath::Min(SamplingTime.AsDecimal() / static_cast<double>(NumberOfFrames), 1.0);
 	}
 
-	float GetRelativePosition() const
+	double GetInterpolatedEvaluationTime() const
 	{
-		PRAGMA_DISABLE_DEPRECATION_WARNINGS
-		return RelativePos;
-		PRAGMA_ENABLE_DEPRECATION_WARNINGS
+		const FFrameTime InterpolationTime = Interpolation == EAnimInterpolationType::Step ? SamplingTime.FloorToFrame() : SamplingTime;
+		return SamplingRate.AsSeconds(InterpolationTime);
 	}
 
+	UE_DEPRECATED(5.1, "Seek with float precision sampling time has been deprecated use double signature instead")
 	void Seek(float SampleAtTime)
+	{
+		Seek(static_cast<double>(SampleAtTime));
+	}
+
+	void Seek(double SampleAtTime)
 	{
 		PRAGMA_DISABLE_DEPRECATION_WARNINGS
 		Time = SampleAtTime;
 		RelativePos = SequenceLength != 0.f ? SampleAtTime / SequenceLength : 0.f;
 		PRAGMA_ENABLE_DEPRECATION_WARNINGS
+		
+		SamplingTime = SamplingRate.AsFrameTime(SampleAtTime);
 	}
+
+	const USkeleton* GetSourceSkeleton() const { return SourceSkeleton; }
+	bool IsAdditiveAnimation() const { return bAdditiveAnimation; }	
+
+protected:
+	FFrameRate SamplingRate;
+	FFrameTime SamplingTime;
+	int32 NumberOfFrames;
+	const USkeleton* SourceSkeleton;
+	bool bAdditiveAnimation;
 };
 
diff --git a/Engine/Source/Runtime/Engine/Public/Animation/AnimTypes.h b/Engine/Source/Runtime/Engine/Public/Animation/AnimTypes.h
--- a/Engine/Source/Runtime/Engine/Public/Animation/AnimTypes.h
+++ b/Engine/Source/Runtime/Engine/Public/Animation/AnimTypes.h
@@ -374,8 +374,8 @@ public:
 		, TriggerTimeOffset(0)
 		, EndTriggerTimeOffset(0)
 		, TriggerWeightThreshold(ZERO_ANIMWEIGHT_THRESH)
-		, Notify(NULL)
-		, NotifyStateClass(NULL)
+		, Notify(nullptr)
+		, NotifyStateClass(nullptr)
 		, Duration(0)
 		, bConvertedFromBranchingPoint(false)
 		, MontageTickType(EMontageNotifyTickType::Queued)
@@ -424,7 +424,7 @@ public:
 	/** Returns true if this is blueprint derived notifies **/
 	bool IsBlueprintNotify() const
 	{
-		return Notify != NULL || NotifyStateClass != NULL;
+		return Notify != nullptr || NotifyStateClass != nullptr;
 	}
 
 	bool operator ==(const FAnimNotifyEvent& Other) const
@@ -880,6 +880,33 @@ struct ENGINE_API FRawAnimSequenceTrack
 		return Ar;
 	}
 
+	bool ContainsNaN() const
+	{
+		bool bContainsNaN = false;
+
+		auto CheckForNan = [&bContainsNaN](const auto& Keys) -> bool
+		{
+			if (!bContainsNaN)
+			{
+				for (const auto& Key : Keys)
+				{
+					if (Key.ContainsNaN())
+						return true;
+				}
+
+				return false;
+			}
+		
+			return true;
+		};
+
+		bContainsNaN = CheckForNan(PosKeys);
+		bContainsNaN = CheckForNan(RotKeys);
+		bContainsNaN = CheckForNan(ScaleKeys);
+
+		return bContainsNaN;
+	}	
+
 	static const uint32 SingleKeySize = sizeof(FVector3f) + sizeof(FQuat4f) + sizeof(FVector3f);
 };
 
diff --git a/Engine/Source/Runtime/Engine/Public/Animation/AttributesRuntime.h b/Engine/Source/Runtime/Engine/Public/Animation/AttributesRuntime.h
--- a/Engine/Source/Runtime/Engine/Public/Animation/AttributesRuntime.h
+++ b/Engine/Source/Runtime/Engine/Public/Animation/AttributesRuntime.h
@@ -56,7 +56,10 @@ namespace UE
 		struct ENGINE_API Attributes
 		{
 #if WITH_EDITOR
-			static void GetAttributeValue(FStackAttributeContainer& OutAttributes, const FCompactPoseBoneIndex& PoseBoneIndex, const  FAnimatedBoneAttribute& Attribute, const FAnimExtractContext& ExtractionContext);
+			UE_DEPRECATED(5.1, "GetAttributeValue with signature using FAnimExtractContext is deprecated use other version instead")
+			static void GetAttributeValue(FStackAttributeContainer& OutAttributes, const FCompactPoseBoneIndex& PoseBoneIndex, const FAnimatedBoneAttribute& Attribute, const FAnimExtractContext& ExtractionContext);
+
+			static void GetAttributeValue(FStackAttributeContainer& OutAttributes, const FCompactPoseBoneIndex& PoseBoneIndex, const FAnimatedBoneAttribute& Attribute, double CurrentTime);
 #endif // WITH_EDITOR
 
 			/** Blend custom attribute values from N set of inputs */
diff --git a/Engine/Source/Runtime/Engine/Public/AnimationRuntime.h b/Engine/Source/Runtime/Engine/Public/AnimationRuntime.h
--- a/Engine/Source/Runtime/Engine/Public/AnimationRuntime.h
+++ b/Engine/Source/Runtime/Engine/Public/AnimationRuntime.h
@@ -545,6 +545,9 @@ public:
 	 * The reason why you can provide a FramesPerSecond value is because this can be slightly more accurate than calculating it, in case super high precision is needed.
 	 **/
 	static void GetKeyIndicesFromTime(int32& OutKeyIndex1, int32& OutKeyIndex2, float& OutAlpha, const double Time, const int32 NumKeys, const double SequenceLength, double FramesPerSecond=-1.0);
+	
+	/** Get KeyIndices using FFrameRate::AsFrameTime to calculate the keys and alpha value **/
+	static void GetKeyIndicesFromTime(int32& OutKeyIndex1, int32& OutKeyIndex2, float& OutAlpha, const double Time, const FFrameRate& FrameRate, const int32 NumberOfKeys);
 
 	/** 
 	 *	Utility for taking an array of bone indices and ensuring that all parents are present 
diff --git a/Engine/Source/Runtime/Engine/Public/AnimationUtils.h b/Engine/Source/Runtime/Engine/Public/AnimationUtils.h
--- a/Engine/Source/Runtime/Engine/Public/AnimationUtils.h
+++ b/Engine/Source/Runtime/Engine/Public/AnimationUtils.h
@@ -168,10 +168,16 @@ public:
 	// Extract specific frame from raw track and place in OutAtom
 	ENGINE_API static void ExtractTransformForFrameFromTrack(const FRawAnimSequenceTrack& RawTrack, int32 Frame, FTransform& OutAtom);
 
+	UE_DEPRECATED(5.1, "ExtractTransformFromTrack has been deprecated, use different signature")
 	ENGINE_API static void ExtractTransformFromTrack(float Time, int32 NumFrames, float SequenceLength, const struct FRawAnimSequenceTrack& RawTrack, EAnimInterpolationType Interpolation, FTransform &OutAtom);
 
+	ENGINE_API static void ExtractTransformFromTrack(const struct FRawAnimSequenceTrack& RawTrack, double Time, int32 NumFrames, double SequenceLength, EAnimInterpolationType Interpolation, FTransform &OutAtom);
+
 #if WITH_EDITOR
+	UE_DEPRECATED(5.1, "ExtractTransformFromCompressionData has been deprecated, use different signature")
 	static void ExtractTransformFromCompressionData(const FCompressibleAnimData& CompressibleAnimData, FCompressibleAnimDataResult& CompressedAnimData, float Time, int32 TrackIndex, bool bUseRawData, FTransform& OutBoneTransform);
+	
+	static void ExtractTransformFromCompressionData(const FCompressibleAnimData& CompressibleAnimData, FCompressibleAnimDataResult& CompressedAnimData, double Time, int32 TrackIndex, bool bUseRawData, FTransform& OutBoneTransform);
 
 	/**
 	 * Compresses the animation bones within a sequence with the chosen settings.
diff --git a/Engine/Source/Runtime/MovieScene/Private/Channels/MovieSceneCurveChannelImpl.h b/Engine/Source/Runtime/MovieScene/Private/Channels/MovieSceneCurveChannelImpl.h
--- a/Engine/Source/Runtime/MovieScene/Private/Channels/MovieSceneCurveChannelImpl.h
+++ b/Engine/Source/Runtime/MovieScene/Private/Channels/MovieSceneCurveChannelImpl.h
@@ -27,6 +27,18 @@ struct MOVIESCENE_API TMovieSceneCurveChannelImpl
 
 	/** Evaluate this channel with the frame resolution */
 	static bool Evaluate(const ChannelType* InChannel, FFrameTime InTime, CurveValueType& OutValue);
+
+	/** Structure used to store the result of UE::MovieScene::EvaluateTime for a given channel/key distribution*/
+	struct FTimeEvaluationCache
+	{
+		float InterpValue = 0.f;
+		int32 Index1 = INDEX_NONE, Index2 = INDEX_NONE;
+		int32 CachedNumFrames = INDEX_NONE;
+		FFrameTime CacheFrameTime;
+	};
+	/** Evaluate this channel at provided FrameTime, using or populating cached time to frame-number(s) calculation */
+	static bool EvaluateWithCache(const ChannelType* InChannel, FTimeEvaluationCache* InOutEvaluationCache, FFrameTime InTime, CurveValueType& OutValue);
+	
 	 /*
 	  * Populate the specified array with times and values that represent the smooth interpolation of 
 	  * the given channel across the specified range
diff --git a/Engine/Source/Runtime/MovieScene/Private/Channels/MovieSceneDoubleChannel.cpp b/Engine/Source/Runtime/MovieScene/Private/Channels/MovieSceneDoubleChannel.cpp
--- a/Engine/Source/Runtime/MovieScene/Private/Channels/MovieSceneDoubleChannel.cpp
+++ b/Engine/Source/Runtime/MovieScene/Private/Channels/MovieSceneDoubleChannel.cpp
@@ -63,6 +63,16 @@ void FMovieSceneDoubleChannel::Set(TArray<FFrameNumber> InTimes, TArray<FMovieSc
 	FMovieSceneDoubleChannelImpl::Set(this, InTimes, InValues);
 }
 
+void FMovieSceneDoubleChannel::SetKeysOnly(TArrayView<FFrameNumber> InTimes, TArrayView<FMovieSceneDoubleValue> InValues)
+{
+	check(InTimes.Num() == InValues.Num());
+
+	Times = MoveTemp(InTimes);
+	Values = MoveTemp(InValues);
+
+	KeyHandles.Reset();
+}
+
 void FMovieSceneDoubleChannel::AutoSetTangents(float Tension)
 {
 	FMovieSceneDoubleChannelImpl::AutoSetTangents(this, Tension);
diff --git a/Engine/Source/Runtime/MovieScene/Private/Channels/MovieSceneFloatChannel.cpp b/Engine/Source/Runtime/MovieScene/Private/Channels/MovieSceneFloatChannel.cpp
--- a/Engine/Source/Runtime/MovieScene/Private/Channels/MovieSceneFloatChannel.cpp
+++ b/Engine/Source/Runtime/MovieScene/Private/Channels/MovieSceneFloatChannel.cpp
@@ -53,6 +53,16 @@ void FMovieSceneFloatChannel::Set(TArray<FFrameNumber> InTimes, TArray<FMovieSce
 	FMovieSceneFloatChannelImpl::Set(this, InTimes, InValues);
 }
 
+void FMovieSceneFloatChannel::SetKeysOnly(TArrayView<FFrameNumber> InTimes, TArrayView<FMovieSceneFloatValue> InValues)
+{
+	check(InTimes.Num() == InValues.Num());
+
+	Times = MoveTemp(InTimes);
+	Values = MoveTemp(InValues);
+
+	KeyHandles.Reset();
+}
+
 void FMovieSceneFloatChannel::AutoSetTangents(float Tension)
 {
 	FMovieSceneFloatChannelImpl::AutoSetTangents(this, Tension);
diff --git a/Engine/Source/Runtime/MovieScene/Public/Channels/MovieSceneDoubleChannel.h b/Engine/Source/Runtime/MovieScene/Public/Channels/MovieSceneDoubleChannel.h
--- a/Engine/Source/Runtime/MovieScene/Public/Channels/MovieSceneDoubleChannel.h
+++ b/Engine/Source/Runtime/MovieScene/Public/Channels/MovieSceneDoubleChannel.h
@@ -174,6 +174,10 @@ struct MOVIESCENE_API FMovieSceneDoubleChannel : public FMovieSceneChannel
 	 */
 	void Set(TArray<FFrameNumber> InTimes, TArray<FMovieSceneDoubleValue> InValues);
 
+	/**
+	 * Set the channel's times and values to the requested values, but does not allocate key handles
+	 */
+	void SetKeysOnly(TArrayView<FFrameNumber> InTimes, TArrayView<FMovieSceneDoubleValue> InValues);
 public:
 
 	// ~ FMovieSceneChannel Interface
diff --git a/Engine/Source/Runtime/MovieScene/Public/Channels/MovieSceneFloatChannel.h b/Engine/Source/Runtime/MovieScene/Public/Channels/MovieSceneFloatChannel.h
--- a/Engine/Source/Runtime/MovieScene/Public/Channels/MovieSceneFloatChannel.h
+++ b/Engine/Source/Runtime/MovieScene/Public/Channels/MovieSceneFloatChannel.h
@@ -169,6 +169,10 @@ struct MOVIESCENE_API FMovieSceneFloatChannel : public FMovieSceneChannel
 	 */
 	void Set(TArray<FFrameNumber> InTimes, TArray<FMovieSceneFloatValue> InValues);
 
+	/**
+	 * Set the channel's times and values to the requested values, but does not allocate key handles
+	 */
+	void SetKeysOnly(TArrayView<FFrameNumber> InTimes, TArrayView<FMovieSceneFloatValue> InValues);
 public:
 
 	// ~ FMovieSceneChannel Interface
diff --git a/Engine/Source/Runtime/MovieSceneTracks/Private/Sections/MovieSceneSkeletalAnimationSection.cpp b/Engine/Source/Runtime/MovieSceneTracks/Private/Sections/MovieSceneSkeletalAnimationSection.cpp
--- a/Engine/Source/Runtime/MovieSceneTracks/Private/Sections/MovieSceneSkeletalAnimationSection.cpp
+++ b/Engine/Source/Runtime/MovieSceneTracks/Private/Sections/MovieSceneSkeletalAnimationSection.cpp
@@ -465,7 +465,7 @@ int32 UMovieSceneSkeletalAnimationSection::SetBoneIndexForRootMotionCalculations
 		//but if not first find first
 		int32 RootIndex = INDEX_NONE;
 #if WITH_EDITOR
-		const UAnimDataModel* DataModel = AnimSequence->GetDataModel();
+		const IAnimationDataModel* DataModel = AnimSequence->GetDataModel();
 		const TArray<FBoneAnimationTrack>& BoneAnimationTracks = DataModel->GetBoneAnimationTracks();
 		const int32 NumTracks = BoneAnimationTracks.Num();
 		for (int32 TrackIndex = 0; TrackIndex < NumTracks; ++TrackIndex)
@@ -525,7 +525,7 @@ bool UMovieSceneSkeletalAnimationSection::GetRootMotionTransform(FFrameTime Curr
 		Params.Weight.Evaluate(CurrentTime, ManualWeight);
 		OutWeight = ManualWeight * EvaluateEasing(CurrentTime);
 		bIsAdditive = false;
-		float CurrentTimeSeconds = static_cast<float>(MapTimeToAnimation(CurrentTime, FrameRate));
+		const double CurrentTimeSeconds = MapTimeToAnimation(CurrentTime, FrameRate);
 		bIsAdditive = AnimSequence->GetAdditiveAnimType() != EAdditiveAnimationType::AAT_None;
 
 		if (TempRootBoneIndex.IsSet())
diff --git a/Engine/Source/Runtime/MovieSceneTracks/Private/Tracks/MovieSceneSkeletalAnimationTrack.cpp b/Engine/Source/Runtime/MovieSceneTracks/Private/Tracks/MovieSceneSkeletalAnimationTrack.cpp
--- a/Engine/Source/Runtime/MovieSceneTracks/Private/Tracks/MovieSceneSkeletalAnimationTrack.cpp
+++ b/Engine/Source/Runtime/MovieSceneTracks/Private/Tracks/MovieSceneSkeletalAnimationTrack.cpp
@@ -383,7 +383,7 @@ static void CalculateDistanceMap(USkeletalMeshComponent* SkelMeshComp, UAnimSequ
 		FloatArray.SetNum(SecondAnimNumFrames);
 		float FirstAnimTime = FirstAnimIndex * FrameRateDiff + StartFirstAnimTime;
 		FirstAnimIndex += 1.0f;
-		FAnimExtractContext FirstExtractionContext(FirstAnimTime, false);
+		FAnimExtractContext FirstExtractionContext(static_cast<double>(FirstAnimTime), false);
 		FirstAnimSeq->GetAnimationPose(FirstAnimPoseData, FirstExtractionContext);
 		FirstMeshPoses.InitPose(FirstAnimPoseData.GetPose());
 		float SecondAnimIndex = 0.0f;
@@ -392,7 +392,7 @@ static void CalculateDistanceMap(USkeletalMeshComponent* SkelMeshComp, UAnimSequ
 			DistVal = 0.0f;
 			float SecondAnimTime = SecondAnimIndex * FrameRateDiff;
 			SecondAnimIndex += 1.0f;
-			FAnimExtractContext SecondExtractionContext(SecondAnimTime, false);
+			FAnimExtractContext SecondExtractionContext(static_cast<double>(SecondAnimTime), false);
 			SecondAnimSeq->GetAnimationPose(SecondAnimPoseData, SecondExtractionContext);
 			SecondMeshPoses.InitPose(SecondAnimPoseData.GetPose());
 
@@ -641,7 +641,7 @@ void UMovieSceneSkeletalAnimationTrack::SetUpRootMotions(bool bForce)
 
 	}
 }
-static FTransform GetWorldTransformForBone(UAnimSequence* AnimSequence, USkeletalMeshComponent* MeshComponent,const FName& InBoneName, float Seconds)
+static FTransform GetWorldTransformForBone(UAnimSequence* AnimSequence, USkeletalMeshComponent* MeshComponent,const FName& InBoneName, double Seconds)
 {
 	FName BoneName = InBoneName;
 	FTransform  WorldTransform = FTransform::Identity;
@@ -653,7 +653,7 @@ static FTransform GetWorldTransformForBone(UAnimSequence* AnimSequence, USkeleta
 		int32 TrackIndex = INDEX_NONE;
 
 #if WITH_EDITOR
-		const UAnimDataModel* Model = AnimSequence->GetDataModel();
+		const IAnimationDataModel* Model = AnimSequence->GetDataModel();
 		if (const FBoneAnimationTrack* TrackData = Model->FindBoneTrackByName(BoneName))
 		{
 			TrackIndex = TrackData->BoneTreeIndex;
@@ -936,9 +936,9 @@ void UMovieSceneSkeletalAnimationTrack::MatchSectionByBoneTransform(bool bMatchW
 
 		if (FirstAnimSequence && SecondAnimSequence)
 		{
-			float FirstSectionTime = static_cast<float>(FirstSection->MapTimeToAnimation(CurrentFrame, FrameRate));
+			double FirstSectionTime = FirstSection->MapTimeToAnimation(CurrentFrame, FrameRate);
 			FTransform  FirstTransform = GetWorldTransformForBone(FirstAnimSequence, SkelMeshComp, BoneName, FirstSectionTime);
-			float SecondSectionTime = static_cast<float>(CurrentSection->MapTimeToAnimation(CurrentFrame, FrameRate));
+			double SecondSectionTime = CurrentSection->MapTimeToAnimation(CurrentFrame, FrameRate);
 			FTransform  SecondTransform = GetWorldTransformForBone(SecondAnimSequence, SkelMeshComp, BoneName, SecondSectionTime);
 
 			//Need to match the translations and rotations here 
