Added cloud show flag.

diff --git a/Engine/Source/Editor/DetailCustomizations/Private/SceneCaptureDetails.cpp b/Engine/Source/Editor/DetailCustomizations/Private/SceneCaptureDetails.cpp
--- a/Engine/Source/Editor/DetailCustomizations/Private/SceneCaptureDetails.cpp
+++ b/Engine/Source/Editor/DetailCustomizations/Private/SceneCaptureDetails.cpp
@@ -82,6 +82,7 @@ void FSceneCaptureDetails::CustomizeDetails( IDetailLayoutBuilder& DetailLayout
 	TArray<FEngineShowFlags::EShowFlag> ShowFlagsToAllowForCaptures;
 
 	ShowFlagsToAllowForCaptures.Add(FEngineShowFlags::EShowFlag::SF_Atmosphere);
+	ShowFlagsToAllowForCaptures.Add(FEngineShowFlags::EShowFlag::SF_Cloud);
 	ShowFlagsToAllowForCaptures.Add(FEngineShowFlags::EShowFlag::SF_BSP);
 	ShowFlagsToAllowForCaptures.Add(FEngineShowFlags::EShowFlag::SF_Decals);
 	ShowFlagsToAllowForCaptures.Add(FEngineShowFlags::EShowFlag::SF_Fog);
diff --git a/Engine/Source/Editor/EditorStyle/Private/SlateEditorStyle.cpp b/Engine/Source/Editor/EditorStyle/Private/SlateEditorStyle.cpp
--- a/Engine/Source/Editor/EditorStyle/Private/SlateEditorStyle.cpp
+++ b/Engine/Source/Editor/EditorStyle/Private/SlateEditorStyle.cpp
@@ -5149,6 +5149,7 @@ void FSlateEditorStyle::FStyle::SetupLevelEditorStyle()
 	{
 		Set( "ShowFlagsMenu.AntiAliasing", new IMAGE_BRUSH( "Icons/icon_ShowAnti-aliasing_16x", Icon16x16 ) );
 		Set( "ShowFlagsMenu.Atmosphere", new IMAGE_BRUSH( "Icons/icon_ShowAtmosphere_16x", Icon16x16 ) );
+		Set( "ShowFlagsMenu.Cloud", new IMAGE_BRUSH( "Icons/AssetIcons/VolumetricCloud_16x", Icon16x16 ) );
 		Set( "ShowFlagsMenu.BSP", new IMAGE_BRUSH( "Icons/icon_ShowBSP_16x", Icon16x16 ) );
 		Set( "ShowFlagsMenu.Collision", new IMAGE_BRUSH( "Icons/icon_ShowCollision_16x", Icon16x16 ) );
 		Set( "ShowFlagsMenu.Decals", new IMAGE_BRUSH( "Icons/icon_ShowDecals_16x", Icon16x16 ) );
diff --git a/Engine/Source/Editor/EditorStyle/Private/StarshipStyle.cpp b/Engine/Source/Editor/EditorStyle/Private/StarshipStyle.cpp
--- a/Engine/Source/Editor/EditorStyle/Private/StarshipStyle.cpp
+++ b/Engine/Source/Editor/EditorStyle/Private/StarshipStyle.cpp
@@ -4512,6 +4512,7 @@ void FStarshipEditorStyle::FStyle::SetupLevelEditorStyle()
 	{
 		Set( "ShowFlagsMenu.AntiAliasing", new IMAGE_BRUSH_SVG( "Starship/Common/AntiAliasing", Icon16x16 ) );
 		Set( "ShowFlagsMenu.Atmosphere", new IMAGE_BRUSH_SVG( "Starship/Common/Atmosphere", Icon16x16 ) );
+		Set( "ShowFlagsMenu.Cloud", new IMAGE_BRUSH_SVG( "Starship/AssetIcons/VolumetricCloud_16", Icon16x16 ) );
 		Set( "ShowFlagsMenu.BSP", new IMAGE_BRUSH_SVG( "Starship/Common/BSP", Icon16x16 ) );
 		Set( "ShowFlagsMenu.Collision", new IMAGE_BRUSH_SVG( "Starship/Common/Collision", Icon16x16 ) );
 		Set( "ShowFlagsMenu.Decals", new IMAGE_BRUSH_SVG( "Starship/Common/Decals", Icon16x16 ) );
diff --git a/Engine/Source/Runtime/Engine/Public/ShowFlagsValues.inl b/Engine/Source/Runtime/Engine/Public/ShowFlagsValues.inl
--- a/Engine/Source/Runtime/Engine/Public/ShowFlagsValues.inl
+++ b/Engine/Source/Runtime/Engine/Public/ShowFlagsValues.inl
@@ -281,6 +281,8 @@ SHOWFLAG_ALWAYS_ACCESSIBLE(LightShafts, SFG_LightingFeatures, NSLOCTEXT("UnrealE
 SHOWFLAG_FIXED_IN_SHIPPING(1, PostProcessMaterial, SFG_PostProcess, NSLOCTEXT("UnrealEd", "PostProcessMaterialSF", "Post Process Material"))
 /** Render Sky and Atmospheric lighting, for now SHOWFLAG_ALWAYS_ACCESSIBLE because it's exposed in SceneCapture */
 SHOWFLAG_ALWAYS_ACCESSIBLE(Atmosphere, SFG_Normal, NSLOCTEXT("UnrealEd", "AtmosphereSF", "Atmosphere"))
+/** Render Cloud */
+SHOWFLAG_FIXED_IN_SHIPPING(1, Cloud, SFG_Normal, NSLOCTEXT("UnrealEd", "CloudSF", "Cloud"))
 /** Render safe frames bars*/
 SHOWFLAG_FIXED_IN_SHIPPING(0, CameraAspectRatioBars, SFG_Advanced, NSLOCTEXT("UnrealEd", "CameraAspectRatioBarsSF", "Camera Aspect Ratio Bars"))
 /** Render safe frames */
diff --git a/Engine/Source/Runtime/Renderer/Private/VolumetricCloudRendering.cpp b/Engine/Source/Runtime/Renderer/Private/VolumetricCloudRendering.cpp
--- a/Engine/Source/Runtime/Renderer/Private/VolumetricCloudRendering.cpp
+++ b/Engine/Source/Runtime/Renderer/Private/VolumetricCloudRendering.cpp
@@ -230,7 +230,7 @@ static bool ShouldUseComputeForCloudTracing(const FStaticFeatureLevel InFeatureL
 
 bool ShouldRenderVolumetricCloud(const FScene* Scene, const FEngineShowFlags& EngineShowFlags)
 {
-	if (Scene && Scene->HasVolumetricCloud() && EngineShowFlags.Atmosphere)
+	if (Scene && Scene->HasVolumetricCloud() && EngineShowFlags.Atmosphere && EngineShowFlags.Cloud)
 	{
 		const FVolumetricCloudRenderSceneInfo* VolumetricCloud = Scene->GetVolumetricCloudSceneInfo();
 		check(VolumetricCloud);
