[Engine] fix for crash in AssetManager when using lazy load. TObjectPtr may not be loaded and a `reinterpret_cast` to an UObject can crash. Using the `FObjectProperty::GetObjectPropertyValue` solves the problem

#p4v-preflight-copy 21921183

diff --git a/Engine/Source/Runtime/Engine/Private/AssetManager.cpp b/Engine/Source/Runtime/Engine/Private/AssetManager.cpp
--- a/Engine/Source/Runtime/Engine/Private/AssetManager.cpp
+++ b/Engine/Source/Runtime/Engine/Private/AssetManager.cpp
@@ -4538,19 +4538,13 @@ void UAssetManager::InitializeAssetBundlesFromMetadata_Recursive(const UStruct*
 		FSoftObjectPath FoundRef;
 		if (const FSoftClassProperty* AssetClassProp = CastField<FSoftClassProperty>(Property))
 		{
-			const TSoftClassPtr<UObject>* AssetClassPtr = reinterpret_cast<const TSoftClassPtr<UObject>*>(PropertyValue);
-			if (AssetClassPtr)
-			{
-				FoundRef = AssetClassPtr->ToSoftObjectPath();
-			}
+			const FSoftObjectPtr& AssetClassPtr = AssetClassProp->GetPropertyValue(PropertyValue);
+			FoundRef = AssetClassPtr.ToSoftObjectPath();
 		}
 		else if (const FSoftObjectProperty* AssetProp = CastField<FSoftObjectProperty>(Property))
 		{
-			const TSoftObjectPtr<UObject>* AssetPtr = reinterpret_cast<const TSoftObjectPtr<UObject>*>(PropertyValue);
-			if (AssetPtr)
-			{
-				FoundRef = AssetPtr->ToSoftObjectPath();
-			}
+			const FSoftObjectPtr& AssetClassPtr = AssetProp->GetPropertyValue(PropertyValue);
+			FoundRef = AssetClassPtr.ToSoftObjectPath();
 		}
 		else if (const FStructProperty* StructProperty = CastField<FStructProperty>(Property))
 		{
@@ -4570,10 +4564,9 @@ void UAssetManager::InitializeAssetBundlesFromMetadata_Recursive(const UStruct*
 		{
 			if (ObjectProperty->PropertyFlags & CPF_InstancedReference || ObjectProperty->GetOwnerProperty()->HasMetaData(IncludeAssetBundlesName))
 			{
-				UObject* const* ObjectPtr = reinterpret_cast<UObject* const*>(PropertyValue);
-				if (ObjectPtr && *ObjectPtr)
+				const UObject* Object = ObjectProperty->GetObjectPropertyValue(PropertyValue);
+				if (Object != nullptr)
 				{
-					const UObject* Object = *ObjectPtr;
 					InitializeAssetBundlesFromMetadata_Recursive(Object->GetClass(), Object, AssetBundle, Object->GetFName(), AllVisitedStructValues);
 				}
 			}
