Preview worlds in the sub-level can create condition where the component needs re-registered but the world is no longer available on a reload of that sublevel.

diff --git a/Engine/Source/Runtime/Engine/Public/ComponentReregisterContext.h b/Engine/Source/Runtime/Engine/Public/ComponentReregisterContext.h
--- a/Engine/Source/Runtime/Engine/Public/ComponentReregisterContext.h
+++ b/Engine/Source/Runtime/Engine/Public/ComponentReregisterContext.h
@@ -58,7 +58,7 @@ public:
 
 	~FComponentReregisterContext()
 	{
-		if( Component.IsValid() )
+		if( Component.IsValid() && World.IsValid() )
 		{
 			ReRegister(Component.Get(), World.Get());
 		}
