Fix modulators not versioning to set properly for procedural sounds
#rnx

diff --git a/Engine/Source/Runtime/Engine/Private/SoundWaveProcedural.cpp b/Engine/Source/Runtime/Engine/Private/SoundWaveProcedural.cpp
--- a/Engine/Source/Runtime/Engine/Private/SoundWaveProcedural.cpp
+++ b/Engine/Source/Runtime/Engine/Private/SoundWaveProcedural.cpp
@@ -185,6 +185,16 @@ void USoundWaveProcedural::Serialize(FArchive& Ar)
 {
 	// Do not call the USoundWave version of serialize
 	USoundBase::Serialize(Ar);
+
+#if WITH_EDITORONLY_DATA
+	// Due to the above "skipping" USoundWave version in the
+	// heirarchy, this is required to version modulators properly
+	// here as well.
+	if (Ar.IsLoading())
+	{
+		ModulationSettings.VersionModulators();
+	}
+#endif // WITH_EDITORONLY_DATA
 }
 
 void USoundWaveProcedural::InitAudioResource(FByteBulkData& CompressedData)
