i before e, except when I don't sleep (minor general comment clean-up)
#rnx

diff --git a/Engine/Source/Runtime/Engine/Private/SoundWaveProcedural.cpp b/Engine/Source/Runtime/Engine/Private/SoundWaveProcedural.cpp
--- a/Engine/Source/Runtime/Engine/Private/SoundWaveProcedural.cpp
+++ b/Engine/Source/Runtime/Engine/Private/SoundWaveProcedural.cpp
@@ -187,9 +187,8 @@ void USoundWaveProcedural::Serialize(FArchive& Ar)
 	USoundBase::Serialize(Ar);
 
 #if WITH_EDITORONLY_DATA
-	// Due to the above "skipping" USoundWave version in the
-	// heirarchy, this is required to version modulators properly
-	// here as well.
+	// Due to "skipping" USoundWave::Serialize above, modulation
+	// versioning is required to be called explicitly here.
 	if (Ar.IsLoading())
 	{
 		ModulationSettings.VersionModulators();
