PR #9309: Made ShouldTriggerAnimNotifyState public (Contributed by [REDACTED_PERSON])

diff --git a/Engine/Source/Runtime/Engine/Classes/Animation/AnimInstance.h b/Engine/Source/Runtime/Engine/Classes/Animation/AnimInstance.h
--- a/Engine/Source/Runtime/Engine/Classes/Animation/AnimInstance.h
+++ b/Engine/Source/Runtime/Engine/Classes/Animation/AnimInstance.h
@@ -1622,6 +1622,7 @@ protected:
 	friend struct FAnimNode_LinkedAnimLayer;
 	friend struct FAnimInstanceProxy;
 	
+public:
 	/** Return whether this AnimNotifyState should be triggered */
 	virtual bool ShouldTriggerAnimNotifyState(const UAnimNotifyState* AnimNotifyState) const;
 
