PR #9032: HideBoneByName should not stop BoneAnim when TickOption is AlwaysTickPoseAndRefreshBones (Contributed by [REDACTED_PERSON])

diff --git a/Engine/Source/Runtime/Engine/Private/Components/SkeletalMeshComponent.cpp b/Engine/Source/Runtime/Engine/Private/Components/SkeletalMeshComponent.cpp
--- a/Engine/Source/Runtime/Engine/Private/Components/SkeletalMeshComponent.cpp
+++ b/Engine/Source/Runtime/Engine/Private/Components/SkeletalMeshComponent.cpp
@@ -1776,7 +1776,8 @@ void USkeletalMeshComponent::ComputeRequiredBones(TArray<FBoneIndexType>& OutReq
 				FBoneIndexType CurBoneIndex = OutRequiredBones[i];
 
 				// Current bone visible?
-				if (EditableBoneVisibilityStates[CurBoneIndex] == BVS_Visible)
+				if (EditableBoneVisibilityStates[CurBoneIndex] == BVS_Visible || 
+					VisibilityBasedAnimTickOption == EVisibilityBasedAnimTickOption::AlwaysTickPoseAndRefreshBones)
 				{
 					OutRequiredBones[VisibleBoneWriteIndex++] = CurBoneIndex;
 				}
