Check status of underlying operation on FSoundWavePCMWriter::IsDone() to properly sync said operation
potential fix for #jira UE-161403

diff --git a/Engine/Source/Runtime/Engine/Private/SampleBufferIO.cpp b/Engine/Source/Runtime/Engine/Private/SampleBufferIO.cpp
--- a/Engine/Source/Runtime/Engine/Private/SampleBufferIO.cpp
+++ b/Engine/Source/Runtime/Engine/Private/SampleBufferIO.cpp
@@ -205,7 +205,7 @@ namespace Audio
 		{
 			CurrentBuffer.MixBufferToChannels(2);
 		}
-
+	
 		CurrentOperation.Reset(new FAsyncSoundWavePCMWriterTask(this, ESoundWavePCMWriteTaskType::GenerateAndWriteSoundWave, OnSuccess));
 		CurrentOperation->StartBackgroundTask();
 
@@ -362,10 +362,16 @@ namespace Audio
 
 	bool FSoundWavePCMWriter::IsDone()
 	{
-		return (CurrentState == ESoundWavePCMWriterState::Suceeded
+		bool bCurrentOperationDone = true;
+		if (CurrentOperation.IsValid())
+		{
+			bCurrentOperationDone = CurrentOperation->IsDone();
+		}
+
+		return ((CurrentState == ESoundWavePCMWriterState::Suceeded
 			|| CurrentState == ESoundWavePCMWriterState::Failed
 			|| CurrentState == ESoundWavePCMWriterState::Cancelled
-			|| CurrentState == ESoundWavePCMWriterState::Idle);
+			|| CurrentState == ESoundWavePCMWriterState::Idle) && bCurrentOperationDone);
 	}
 
 	void FSoundWavePCMWriter::Reset()
@@ -684,4 +690,3 @@ namespace Audio
 
 }
 
-
