When placing partition actors by location using wp.Runtime.RuntimeSpatialHashPlacePartitionActorsUsingLocation, choose best grid level based on actor bounds instead of always using lowest level.
This allows having streaming grid cell size and partition actors size to be very different while respecting runtime streaming loading range.

diff --git a/Engine/Source/Runtime/Engine/Private/WorldPartition/RuntimeSpatialHash/RuntimeSpatialHashGridHelper.cpp b/Engine/Source/Runtime/Engine/Private/WorldPartition/RuntimeSpatialHash/RuntimeSpatialHashGridHelper.cpp
--- a/Engine/Source/Runtime/Engine/Private/WorldPartition/RuntimeSpatialHash/RuntimeSpatialHashGridHelper.cpp
+++ b/Engine/Source/Runtime/Engine/Private/WorldPartition/RuntimeSpatialHash/RuntimeSpatialHashGridHelper.cpp
@@ -164,22 +164,24 @@ FSquare2DGridHelper GetPartitionedActors(const FBox& WorldBounds, const FSpatial
 		}
 	}
 
-	const float CellArea = PartitionedActors.GetLowestLevel().CellSize * PartitionedActors.GetLowestLevel().CellSize;
+	const int32 GridLevelCount = PartitionedActors.Levels.Num();
+	const int64 CellSize = PartitionedActors.GetLowestLevel().CellSize;
+	const float CellArea = CellSize * CellSize;
 
-	auto ShouldActorUseLocationPlacement = [&CellArea](const IStreamingGenerationContext::FActorSetInstance* ActorSetInstance)
+	auto ShouldActorUseLocationPlacement = [CellArea, CellSize, GridLevelCount](const IStreamingGenerationContext::FActorSetInstance* InActorSetInstance, const FBox2D& InActorSetInstanceBounds, int32& OutGridLevel)
 	{
+		OutGridLevel = 0;
 		if (GRuntimeSpatialHashPlaceSmallActorsUsingLocation)
 		{
-			const FBox2D ClusterBounds(FVector2D(ActorSetInstance->Bounds.Min), FVector2D(ActorSetInstance->Bounds.Max));
-			return ClusterBounds.GetArea() <= CellArea;
+			return InActorSetInstanceBounds.GetArea() <= CellArea;
 		}
 
 		if (GRuntimeSpatialHashPlacePartitionActorsUsingLocation)
 		{
 			bool bUseLocation = true;
-			for (const FGuid& ActorGuid : ActorSetInstance->ActorSet->Actors)
+			for (const FGuid& ActorGuid : InActorSetInstance->ActorSet->Actors)
 			{
-				const FWorldPartitionActorDescView& ActorDescView = ActorSetInstance->ContainerInstance->ActorDescViewMap->FindByGuidChecked(ActorGuid);
+				const FWorldPartitionActorDescView& ActorDescView = InActorSetInstance->ContainerInstance->ActorDescViewMap->FindByGuidChecked(ActorGuid);
 
 				if (!ActorDescView.GetActorNativeClass()->IsChildOf<APartitionActor>())
 				{
@@ -187,6 +189,12 @@ FSquare2DGridHelper GetPartitionedActors(const FBox& WorldBounds, const FSpatial
 					break;
 				}
 			}
+			if (bUseLocation)
+			{
+				// Find grid level that best matches actor set bounds
+				const float MaxLength = InActorSetInstanceBounds.GetExtent().GetMax() * 2.0;
+				OutGridLevel = FMath::Min(FMath::CeilToInt(FMath::Max<float>(FMath::Log2(MaxLength / CellSize), 0)), GridLevelCount - 1);
+			}
 			return bUseLocation;
 		}
 
@@ -195,19 +203,20 @@ FSquare2DGridHelper GetPartitionedActors(const FBox& WorldBounds, const FSpatial
 
 	for (const IStreamingGenerationContext::FActorSetInstance* ActorSetInstance : ActorSetInstances)
 	{
-		const FBox2D ActorSetInstanceBounds(FVector2D(ActorSetInstance->Bounds.Min), FVector2D(ActorSetInstance->Bounds.Max));
 		check(ActorSetInstance->ActorSet->Actors.Num() > 0);
 		FSquare2DGridHelper::FGridLevel::FGridCell* GridCell = nullptr;
 
 		if (ActorSetInstance->bIsSpatiallyLoaded)
 		{
-			if (ShouldActorUseLocationPlacement(ActorSetInstance))
+			const FBox2D ActorSetInstanceBounds(FVector2D(ActorSetInstance->Bounds.Min), FVector2D(ActorSetInstance->Bounds.Max));
+			int32 LocationPlacementGridLevel = 0;
+			if (ShouldActorUseLocationPlacement(ActorSetInstance, ActorSetInstanceBounds, LocationPlacementGridLevel))
 			{
 				// Find grid level cell that contains the actor cluster pivot and put actors in it.
 				FGridCellCoord2 CellCoords;
-				if (PartitionedActors.GetLowestLevel().GetCellCoords(ActorSetInstanceBounds.GetCenter(), CellCoords))
+				if (PartitionedActors.Levels[LocationPlacementGridLevel].GetCellCoords(ActorSetInstanceBounds.GetCenter(), CellCoords))
 				{
-					GridCell = &PartitionedActors.GetLowestLevel().GetCell(CellCoords);
+					GridCell = &PartitionedActors.Levels[LocationPlacementGridLevel].GetCell(CellCoords);
 				}
 			}
 			else
