Allow to edit LandscapeSplineMeshesGridSize in WorldSettings for partitioned worlds.

diff --git a/Engine/Source/Runtime/Engine/Private/WorldSettings.cpp b/Engine/Source/Runtime/Engine/Private/WorldSettings.cpp
--- a/Engine/Source/Runtime/Engine/Private/WorldSettings.cpp
+++ b/Engine/Source/Runtime/Engine/Private/WorldSettings.cpp
@@ -681,11 +681,14 @@ bool AWorldSettings::CanEditChange(const FProperty* InProperty) const
 		}
 		else if (PropertyName == GET_MEMBER_NAME_STRING_CHECKED(AWorldSettings, bEnableWorldComposition ) ||
 				 PropertyName == GET_MEMBER_NAME_STRING_CHECKED(AWorldSettings, bForceNoPrecomputedLighting ) ||
-				 PropertyName == GET_MEMBER_NAME_STRING_CHECKED(AWorldSettings, bPrecomputeVisibility) ||
-				 PropertyName == GET_MEMBER_NAME_STRING_CHECKED(AWorldSettings, LandscapeSplineMeshesGridSize))
+				 PropertyName == GET_MEMBER_NAME_STRING_CHECKED(AWorldSettings, bPrecomputeVisibility))
 		{
 			return !IsPartitionedWorld();
 		}
+		else if (PropertyName == GET_MEMBER_NAME_STRING_CHECKED(AWorldSettings, LandscapeSplineMeshesGridSize))
+		{
+			return IsPartitionedWorld();
+		}
 		if (PropertyName == GET_MEMBER_NAME_STRING_CHECKED(AWorldSettings, InstancedFoliageGridSize))
 		{
 			return false;
