Content Bundle system core classes
Content bundles stores actors descriptors at a different path than the WorldPartition typical external actors.
When a content bundle actor editor context is activated, all external actors are then saved at that root
Content Bundle can inject/'remove content dynamically.
Content Bundle can also add data layers dynamically.
Content Bundles stores a streaming object in game world. The streaming object allows a RuntimeHash to inject/remove cells dynamically.
New ContentBundleManager exists only is UWorld with worldpartition.

#test asteria_terrain, qagame, cooked game, pie, sie , hlod generation

diff --git a/Engine/Source/Editor/WorldPartitionEditor/Private/WorldPartition/ContentBundle/ContentBundleEditorSubsystem.cpp b/Engine/Source/Editor/WorldPartitionEditor/Private/WorldPartition/ContentBundle/ContentBundleEditorSubsystem.cpp
new file mode 100644
--- /dev/null
+++ b/Engine/Source/Editor/WorldPartitionEditor/Private/WorldPartition/ContentBundle/ContentBundleEditorSubsystem.cpp
@@ -0,0 +1,265 @@
+// Copyright Epic Games, Inc. All Rights Reserved.
+#include "WorldPartition/ContentBundle/ContentBundleEditorSubsystem.h"
+
+#include "WorldPartition/ContentBundle/ContentBundle.h"
+#include "WorldPartition/ContentBundle/ContentBundleDescriptor.h"
+#include "WorldPartition/ContentBundle/ContentBundle.h"
+#include "WorldPartition/ContentBundle/ContentBundleEditor.h"
+#include "WorldPartition/WorldPartition.h"
+#include "Subsystems/ActorEditorContextSubsystem.h"
+#include "Engine/Selection.h"
+#include "WorldPartition/ContentBundle/ContentBundleWorldSubsystem.h"
+#include "WorldPartition/ContentBundle/ContentBundleLog.h"
+#include "Styling/AppStyle.h"
+#include "Widgets/SBoxPanel.h"
+#include "Editor.h"
+
+void UContentBundleEditionSubmodule::DoInitialize()
+{
+	UActorEditorContextSubsystem::Get()->RegisterClient(this);
+}
+
+void UContentBundleEditionSubmodule::DoDenitialize()
+{
+	UActorEditorContextSubsystem::Get()->UnregisterClient(this);
+
+	EditingContentBundle = nullptr;
+	ActiveContentBundlesStack.Empty();
+}
+
+void UContentBundleEditionSubmodule::OnExecuteActorEditorContextAction(UWorld* InWorld, const EActorEditorContextAction& InType, AActor* InActor /* = nullptr */)
+{
+	TSharedPtr<FContentBundleEditor> ContentBundleEditorPin = EditingContentBundle.Pin();
+
+	switch (InType)
+	{
+	case EActorEditorContextAction::ApplyContext:
+		if (ContentBundleEditorPin != nullptr)
+		{
+			ContentBundleEditorPin->AddActor(InActor);
+		}
+		break;
+	case EActorEditorContextAction::ResetContext:
+		break;
+	case EActorEditorContextAction::PushContext:
+		PushContentBundleEditing();
+		break;
+	case EActorEditorContextAction::PopContext:
+		PopContentBundleEditing();
+		break;
+	}
+}
+
+bool UContentBundleEditionSubmodule::GetActorEditorContextDisplayInfo(UWorld* InWorld, FActorEditorContextClientDisplayInfo& OutDiplayInfo) const
+{
+	if (EditingContentBundle != nullptr)
+	{
+		constexpr TCHAR PLACEHOLDER_ContentBundle_Brush[] = TEXT("DataLayer.Editor"); // todo_ow: create placeholder brush for content bundle
+		OutDiplayInfo.Title = TEXT("Content Bundle");
+		OutDiplayInfo.Brush = FAppStyle::GetBrush(PLACEHOLDER_ContentBundle_Brush);
+		return true;
+	}
+
+	return false;
+}
+
+bool UContentBundleEditionSubmodule::CanResetContext(UWorld* InWorld) const
+{
+	return false;
+}
+
+TSharedRef<SWidget> UContentBundleEditionSubmodule::GetActorEditorContextWidget(UWorld* World) const
+{
+	constexpr TCHAR PLACEHOLDER_ContentBundle_ColorIcon[] = TEXT("DataLayer.ColorIcon"); // todo_ow: create placeholder color icon for content bundle
+
+	TSharedRef<SVerticalBox> OutWidget = SNew(SVerticalBox);
+
+	if (EditingContentBundle != nullptr)
+	{
+		OutWidget->AddSlot().AutoHeight()
+		[
+			SNew(SHorizontalBox)
+			+ SHorizontalBox::Slot()
+			.AutoWidth()
+			.Padding(0.0f, 1.0f, 1.0f, 1.0f)
+			.VAlign(VAlign_Center)
+			[
+				SNew(SImage)
+				.Image(FAppStyle::Get().GetBrush(PLACEHOLDER_ContentBundle_ColorIcon))
+				.DesiredSizeOverride(FVector2D(8, 8))
+			]
+		+ SHorizontalBox::Slot()
+			.AutoWidth()
+			.Padding(4.0f, 1.0f, 1.0f, 1.0f)
+			[
+				SNew(STextBlock)
+				.Text_Lambda([this]()
+				{
+					TSharedPtr<FContentBundleEditor> ContentBundleEditorPin = EditingContentBundle.Pin();
+					if (ContentBundleEditorPin != nullptr)
+					{
+						return FText::FromString(ContentBundleEditorPin->GetDisplayName());
+					}
+
+					return FText::FromString(TEXT("Unset Display Name"));
+				})
+					
+			]
+		];
+	}
+
+
+	return OutWidget;
+}
+
+void UContentBundleEditionSubmodule::ActivateContentBundleEditing(TSharedPtr<FContentBundleEditor>& ContentBundleEditor)
+{
+	EContentBundleStatus ContentBundleStatus = ContentBundleEditor->GetStatus();
+	bool bCanActivateEditing = ContentBundleStatus == EContentBundleStatus::ContentInjected || ContentBundleStatus == EContentBundleStatus::ReadyToInject;
+	if (bCanActivateEditing && !ContentBundleEditor->IsBeingEdited())
+	{
+		if (IsEditingContentBundle())
+		{
+			check(!IsEditingContentBundle(ContentBundleEditor.Get()));
+			DeactivateCurrentContentBundleEditing();
+		}
+
+		UActorEditorContextSubsystem::Get()->PushContext();
+		EditingContentBundle = ContentBundleEditor;
+
+		ContentBundleEditor->StartEditing();
+		if (ContentBundleStatus == EContentBundleStatus::ReadyToInject)
+		{
+			ContentBundleEditor->InjectBaseContent();
+		}
+
+		ActorEditorContextClientChanged.Broadcast(this);
+
+		UE_LOG(LogContentBundle, Log, TEXT("[CB: %s] Content Bundle is being edited"), *ContentBundleEditor->GetDescriptor()->GetDisplayName());
+
+		GetSubsystem()->NotifyContentBundleChanged(ContentBundleEditor.Get());
+	}
+}
+
+void UContentBundleEditionSubmodule::DeactivateCurrentContentBundleEditing()
+{
+	if (TSharedPtr<FContentBundleEditor> ContentBundleEditorPin = EditingContentBundle.Pin())
+	{
+		ContentBundleEditorPin->StopEditing();
+		EditingContentBundle = nullptr;
+		ActorEditorContextClientChanged.Broadcast(this);
+		UActorEditorContextSubsystem::Get()->PopContext();
+
+		GetSubsystem()->NotifyContentBundleChanged(ContentBundleEditorPin.Get());
+
+		UE_LOG(LogContentBundle, Log, TEXT("[CB: %s] Content Bundle is no longer being edited"), *ContentBundleEditorPin->GetDescriptor()->GetDisplayName());
+	}
+}
+
+void UContentBundleEditionSubmodule::PushContentBundleEditing()
+{
+	ActiveContentBundlesStack.Add(EditingContentBundle);
+	DeactivateCurrentContentBundleEditing();
+}
+
+void UContentBundleEditionSubmodule::PopContentBundleEditing()
+{
+	DeactivateCurrentContentBundleEditing();
+
+	TSharedPtr<FContentBundleEditor> ContentBundleToActivate = ActiveContentBundlesStack.Pop().Pin();
+	if (ContentBundleToActivate != nullptr)
+	{
+		ActivateContentBundleEditing(ContentBundleToActivate);
+	}
+}
+
+UContentBundleEditorSubsystem::UContentBundleEditorSubsystem()
+{
+	
+}
+
+void UContentBundleEditorSubsystem::Initialize(FSubsystemCollectionBase& Collection)
+{
+	Collection.InitializeDependency<UActorEditorContextSubsystem>();
+
+	ContentBundleEditionSubModule = NewObject<UContentBundleEditionSubmodule>(this);
+	ContentBundleEditionSubModule->Initialize();
+
+	IContentBundleEditorSubsystemInterface::SetInstance(this);
+}
+
+void UContentBundleEditorSubsystem::Deinitialize()
+{
+	IContentBundleEditorSubsystemInterface::SetInstance(nullptr);
+
+	ContentBundleEditionSubModule->Deinitialize();
+	ContentBundleEditionSubModule = nullptr;
+}
+
+UWorld* UContentBundleEditorSubsystem::GetWorld() const
+{
+	return GEditor->GetEditorWorldContext().World();
+}
+
+TSharedPtr<FContentBundleEditor> UContentBundleEditorSubsystem::GetEditorContentBundleForActor(const AActor* Actor)
+{
+	for (TSharedPtr<FContentBundleEditor>& ContentBundle : GetEditorContentBundles())
+	{
+		if (ContentBundle->ContainsActor(Actor))
+		{
+			return ContentBundle;
+		}
+	}
+
+	return nullptr;
+}
+
+TArray<TSharedPtr<FContentBundleEditor>> UContentBundleEditorSubsystem::GetEditorContentBundles()
+{
+	TArray<TSharedPtr<FContentBundleEditor>> EditorContentBundles;
+	if (UContentBundleManager* ContentBundleManager = GetWorld()->ContentBundleManager)
+	{
+		ContentBundleManager->GetEditorContentBundle(EditorContentBundles);
+	}
+
+	return EditorContentBundles;
+}
+
+void UContentBundleEditorSubsystem::SelectActors(FContentBundleEditor& EditorContentBundle)
+{
+	SelectActorsInternal(EditorContentBundle, true);
+}
+
+void UContentBundleEditorSubsystem::DeselectActors(FContentBundleEditor& EditorContentBundle)
+{
+	SelectActorsInternal(EditorContentBundle, false);
+}
+
+void UContentBundleEditorSubsystem::ReferenceAllActors(FContentBundleEditor& EditorContentBundle)
+{
+	EditorContentBundle.ReferenceAllActors();
+}
+
+void UContentBundleEditorSubsystem::UnreferenceAllActors(FContentBundleEditor& EditorContentBundle)
+{
+	EditorContentBundle.UnreferenceAllActors();
+}
+
+void UContentBundleEditorSubsystem::SelectActorsInternal(FContentBundleEditor& EditorContentBundle, bool bSelect)
+{
+	GEditor->GetSelectedActors()->BeginBatchSelectOperation();
+	{
+		TArray<AActor*> ContentBundleActors;
+		EditorContentBundle.GetActors(ContentBundleActors);
+
+		for (AActor* Actor : ContentBundleActors)
+		{
+			bool bNotifyForActor = false;
+			GEditor->GetSelectedActors()->Modify();
+			GEditor->SelectActor(Actor, bSelect, true, false);
+		}
+	}
+	GEditor->GetSelectedActors()->EndBatchSelectOperation();
+
+	GEditor->NoteSelectionChange();
+}
\ No newline at end of file
diff --git a/Engine/Source/Editor/WorldPartitionEditor/Public/WorldPartition/ContentBundle/ContentBundleEditorSubsystem.h b/Engine/Source/Editor/WorldPartitionEditor/Public/WorldPartition/ContentBundle/ContentBundleEditorSubsystem.h
new file mode 100644
--- /dev/null
+++ b/Engine/Source/Editor/WorldPartitionEditor/Public/WorldPartition/ContentBundle/ContentBundleEditorSubsystem.h
@@ -0,0 +1,131 @@
+// Copyright Epic Games, Inc. All Rights Reserved.
+#pragma once
+
+#include "CoreMinimal.h"
+#include "UObject/ObjectMacros.h"
+#include "UObject/Object.h"
+#include "EditorSubsystem.h"
+#include "WorldPartition/ContentBundle/ContentBundleEditorSubsystemInterface.h"
+#include "IActorEditorContextClient.h"
+
+#include "ContentBundleEditorSubsystem.generated.h"
+
+class FContentBundleEditor;
+class FContentBundleEditor;
+class UContentBundleEditorSubsystem;
+
+UCLASS(Within = ContentBundleEditorSubsystem)
+class UContentBundleEditorSubsystemModule : public UObject
+{
+	GENERATED_BODY()
+
+public:
+	//~ Begin UObject interface
+	virtual void BeginDestroy() override { check(!IsInitialize()); }
+	//~ End UObject interface
+
+	void Initialize() { DoInitialize(); bIsInitialized = true; }
+	void Deinitialize() { bIsInitialized = false; DoDenitialize(); }
+
+	bool IsInitialize() const { return bIsInitialized; }
+
+	UContentBundleEditorSubsystem* GetSubsystem() const { return GetOuterUContentBundleEditorSubsystem(); }
+
+protected:
+	virtual void DoInitialize() {};
+	virtual void DoDenitialize() {};
+
+private:
+	bool bIsInitialized = false;
+};
+
+UCLASS()
+class UContentBundleEditionSubmodule : public UContentBundleEditorSubsystemModule, public IActorEditorContextClient
+{
+	GENERATED_BODY()
+
+public:
+	//~ Begin IActorEditorContextClient interface
+	virtual void OnExecuteActorEditorContextAction(UWorld* InWorld, const EActorEditorContextAction& InType, AActor* InActor = nullptr) override;
+	virtual bool GetActorEditorContextDisplayInfo(UWorld* InWorld, FActorEditorContextClientDisplayInfo& OutDiplayInfo) const override;
+	virtual bool CanResetContext(UWorld* InWorld) const override;
+	virtual TSharedRef<SWidget> GetActorEditorContextWidget(UWorld* InWorld) const override;
+	virtual FOnActorEditorContextClientChanged& GetOnActorEditorContextClientChanged() override { return ActorEditorContextClientChanged; }
+	//~ End IActorEditorContextClient interface
+
+	void ActivateContentBundleEditing(TSharedPtr<FContentBundleEditor>& ContentBundleEditor);
+	void DeactivateCurrentContentBundleEditing();
+	bool IsEditingContentBundle() const { return EditingContentBundle.IsValid(); }
+	bool IsEditingContentBundle(const FContentBundleEditor* ContentBundleEditor) const { return EditingContentBundle.Pin().Get() == ContentBundleEditor; }
+
+protected:
+	//~ Begin UContentBundleEditorSubsystemModule interface
+	virtual void DoInitialize() override;
+	virtual void DoDenitialize() override;
+	//~ End UContentBundleEditorSubsystemModule interface
+
+private:
+	void PushContentBundleEditing();
+	void PopContentBundleEditing();
+
+	TArray<TWeakPtr<FContentBundleEditor>> ActiveContentBundlesStack;
+
+	TWeakPtr<FContentBundleEditor> EditingContentBundle;
+
+	FOnActorEditorContextClientChanged ActorEditorContextClientChanged;
+};
+
+UCLASS()
+class UContentBundleEditorSubsystem : public UEditorSubsystem, public IContentBundleEditorSubsystemInterface
+{
+	GENERATED_BODY()
+
+public:
+	static UContentBundleEditorSubsystem* Get() { return StaticCast<UContentBundleEditorSubsystem*>(IContentBundleEditorSubsystemInterface::Get()); }
+
+	UContentBundleEditorSubsystem();
+
+	//~ Begin UEditorSubsystem interface
+	virtual void Initialize(FSubsystemCollectionBase& Collection) override;
+	virtual void Deinitialize() override;
+	//~ End UEditorSubsystem interface
+
+	//~ Begin IContentBundleEditorSubsystemInterface interface
+	virtual void NotifyContentBundleAdded(const FContentBundleEditor* ContentBundle) override { OnContentBundleAdded().Broadcast(ContentBundle); }
+	virtual void NotifyContentBundleRemoved(const FContentBundleEditor* ContentBundle) override { OnContentBundleRemoved().Broadcast(ContentBundle); }
+	virtual void NotifyContentBundleChanged(const FContentBundleEditor* ContentBundle) override { OnContentBundleChanged().Broadcast(ContentBundle); }
+	//~ End IContentBundleEditorSubsystemInterface interface
+
+	UWorld* GetWorld() const;
+
+	UContentBundleEditionSubmodule* GetEditionSubmodule() { return ContentBundleEditionSubModule; }
+
+	TSharedPtr<FContentBundleEditor> GetEditorContentBundleForActor(const AActor* Actor);
+
+	TArray<TSharedPtr<FContentBundleEditor>> GetEditorContentBundles();
+
+	void SelectActors(FContentBundleEditor& EditorContentBundle);
+	void DeselectActors(FContentBundleEditor& EditorContentBundle);
+
+	void ReferenceAllActors(FContentBundleEditor& EditorContentBundle);
+	void UnreferenceAllActors(FContentBundleEditor& EditorContentBundle);
+
+	DECLARE_EVENT_OneParam(UContentBundleEditorSubsystem, FOnContentBundleChanged, const FContentBundleEditor*);
+	FOnContentBundleChanged& OnContentBundleChanged() { return ContentBundleChanged; }
+
+	DECLARE_EVENT_OneParam(UContentBundleEditorSubsystem, FOnContentBundleAdded, const FContentBundleEditor*);
+	FOnContentBundleAdded& OnContentBundleAdded() { return ContentBundleAdded; }
+
+	DECLARE_EVENT_OneParam(UContentBundleEditorSubsystem, FOnContentBundleRemoved, const FContentBundleEditor*);
+	FOnContentBundleRemoved& OnContentBundleRemoved() { return ContentBundleRemoved; }
+
+private:
+	void SelectActorsInternal(FContentBundleEditor& EditorContentBundle, bool bSelect);
+
+	UPROPERTY()
+	TObjectPtr<UContentBundleEditionSubmodule> ContentBundleEditionSubModule;
+
+	FOnContentBundleChanged ContentBundleChanged;
+	FOnContentBundleAdded ContentBundleAdded;
+	FOnContentBundleRemoved ContentBundleRemoved;
+};
\ No newline at end of file
diff --git a/Engine/Source/Editor/WorldPartitionEditor/WorldPartitionEditor.Build.cs b/Engine/Source/Editor/WorldPartitionEditor/WorldPartitionEditor.Build.cs
--- a/Engine/Source/Editor/WorldPartitionEditor/WorldPartitionEditor.Build.cs
+++ b/Engine/Source/Editor/WorldPartitionEditor/WorldPartitionEditor.Build.cs
@@ -11,6 +11,7 @@ public class WorldPartitionEditor : ModuleRules
 				"Core",
 				"CoreUObject",
 				"EditorFramework",
+				"EditorSubsystem",
 				"Engine",
 				"InputCore",
 				"Slate",
diff --git a/Engine/Source/Runtime/Engine/Classes/Engine/World.h b/Engine/Source/Runtime/Engine/Classes/Engine/World.h
--- a/Engine/Source/Runtime/Engine/Classes/Engine/World.h
+++ b/Engine/Source/Runtime/Engine/Classes/Engine/World.h
@@ -1845,6 +1845,9 @@ public:
 	/** All levels information from which our world is composed */
 	UPROPERTY()
 	TObjectPtr<class UWorldComposition> WorldComposition;
+
+	UPROPERTY()
+	TObjectPtr<class UContentBundleManager> ContentBundleManager;
 	
 	/** Whether we flushing level streaming state */ 
 	EFlushLevelStreamingType FlushLevelStreamingType;
diff --git a/Engine/Source/Runtime/Engine/Private/WorldPartition/ContentBundle/ContentBundle.cpp b/Engine/Source/Runtime/Engine/Private/WorldPartition/ContentBundle/ContentBundle.cpp
new file mode 100644
--- /dev/null
+++ b/Engine/Source/Runtime/Engine/Private/WorldPartition/ContentBundle/ContentBundle.cpp
@@ -0,0 +1,116 @@
+// Copyright Epic Games, Inc. All Rights Reserved.
+#include "WorldPartition/ContentBundle/ContentBundle.h"
+
+#include "WorldPartition/WorldPartition.h"
+#include "WorldPartition/ContentBundle/ContentBundleDescriptor.h"
+#include "WorldPartition/WorldPartitionRuntimeHash.h"
+#include "WorldPartition/ContentBundle/ContentBundleLog.h"
+
+#if WITH_EDITOR
+#include "WorldPartition/ContentBundle/ContentBundleEditor.h"
+#include "WorldPartition/ContentBundle/ContentBundleWorldSubsystem.h"
+#include "Editor.h"
+#endif
+
+FContentBundle::FContentBundle(TSharedPtr<FContentBundleClient>& InClient, UWorld* InWorld)
+	: FContentBundleBase(InClient, InWorld),
+	ExternalStreamingObjectPackage(nullptr),
+	ExternalStreamingObject(nullptr)
+{
+
+}
+
+void FContentBundle::DoInitialize()
+{
+#if WITH_EDITOR
+	InitializeForPIE();
+#else
+	FString StreamingObjectPathName = FPaths::RemoveDuplicateSlashes(GetDescriptor()->GetPackageRoot() + TEXT("/_GENERATED_/") + TEXT("StreamingObject"));
+	if (ExternalStreamingObjectPackage = LoadPackage(nullptr, *StreamingObjectPathName, LOAD_None))
+	{
+		if (UObject* Object = StaticFindObjectFast(URuntimeHashExternalStreamingObjectBase::StaticClass(), ExternalStreamingObjectPackage, TEXT("StreamingObject")))
+		{
+			ExternalStreamingObject = CastChecked<URuntimeHashExternalStreamingObjectBase>(Object);
+		}
+		else
+		{
+			UE_LOG(LogContentBundle, Error, TEXT("[CB: %s] No streaming object found in package %s."), *GetDescriptor()->GetDisplayName(), *StreamingObjectPathName);
+		}
+	}
+	else
+	{
+		UE_LOG(LogContentBundle, Log, TEXT("[CB: %s] No streaming object found. No content will be injected."), *GetDescriptor()->GetDisplayName());
+	}
+#endif
+
+	SetStatus(EContentBundleStatus::Registered);
+}
+
+void FContentBundle::DoUninitialize()
+{
+	SetStatus(EContentBundleStatus::Unknown);
+
+	ExternalStreamingObject = nullptr;
+	ExternalStreamingObjectPackage = nullptr;
+}
+
+void FContentBundle::DoInjectContent()
+{
+	if (ExternalStreamingObject != nullptr)
+	{
+		if (GetInjectedWorld()->GetWorldPartition()->RuntimeHash->InjectExternalStreamingObject(ExternalStreamingObject))
+		{
+			UE_LOG(LogContentBundle, Log, TEXT("[CB: %s] Streaming Object Injected."), *GetDescriptor()->GetDisplayName());
+			SetStatus(EContentBundleStatus::ContentInjected);
+		}
+		else
+		{
+			UE_LOG(LogContentBundle, Error, TEXT("[CB: %s] Failed to inject streaming object."), *GetDescriptor()->GetDisplayName());
+			SetStatus(EContentBundleStatus::FailedToInject);
+		}
+	}
+	else
+	{
+		UE_LOG(LogContentBundle, Log, TEXT("[CB: %s] No streaming object to inject."), *GetDescriptor()->GetDisplayName());
+		SetStatus(EContentBundleStatus::ContentInjected);
+	}
+}
+
+void FContentBundle::DoRemoveContent()
+{
+	if (ExternalStreamingObject != nullptr)
+	{
+		if (!GetInjectedWorld()->GetWorldPartition()->RuntimeHash->RemoveExternalStreamingObject(ExternalStreamingObject))
+		{
+			UE_LOG(LogContentBundle, Error, TEXT("[CB: %s] Error while removing streaming object."), *GetDescriptor()->GetDisplayName());
+		}
+	}
+	
+	SetStatus(EContentBundleStatus::Registered);
+}
+
+void FContentBundle::AddReferencedObjects(FReferenceCollector& Collector)
+{
+	FContentBundleBase::AddReferencedObjects(Collector);
+
+	Collector.AddReferencedObject(ExternalStreamingObjectPackage);
+	Collector.AddReferencedObject(ExternalStreamingObject);
+}
+
+bool FContentBundle::IsValid() const
+{
+	return GetDescriptor()->IsValid();
+}
+
+#if WITH_EDITOR
+void FContentBundle::InitializeForPIE()
+{
+	UContentBundleManager* ContentBundleManager = GetInjectedWorld()->ContentBundleManager;
+	UContentBundleDuplicateForPIEHelper* PIEHelper = ContentBundleManager->GetPIEDuplicateHelper();
+	ExternalStreamingObject = PIEHelper->RetrieveContentBundleStreamingObject(*this);
+	if (ExternalStreamingObject == nullptr)
+	{
+		UE_LOG(LogContentBundle, Log, TEXT("[CB: %s] No streaming object found. There are %u existing streaming objects."), *GetDescriptor()->GetDisplayName(), PIEHelper->GetStreamingObjectCount());
+	}
+}
+#endif
\ No newline at end of file
diff --git a/Engine/Source/Runtime/Engine/Private/WorldPartition/ContentBundle/ContentBundleBase.cpp b/Engine/Source/Runtime/Engine/Private/WorldPartition/ContentBundle/ContentBundleBase.cpp
new file mode 100644
--- /dev/null
+++ b/Engine/Source/Runtime/Engine/Private/WorldPartition/ContentBundle/ContentBundleBase.cpp
@@ -0,0 +1,99 @@
+// Copyright Epic Games, Inc. All Rights Reserved.
+#include "WorldPartition/ContentBundle/ContentBundleBase.h"
+
+#include "WorldPartition/WorldPartition.h"
+#include "WorldPartition/ContentBundle/ContentBundleDescriptor.h"
+#include "WorldPartition/ContentBundle/ContentBundleClient.h"
+#include "WorldPartition/ContentBundle/ContentBundleLog.h"
+#include "Engine/World.h"
+
+FContentBundleBase::FContentBundleBase(TSharedPtr<FContentBundleClient>& InClient, UWorld* InWorld)
+	: Client(InClient),
+	InjectedWorld(InWorld),
+	Descriptor(InClient->GetDescriptor()),
+	Status(EContentBundleStatus::Unknown)
+{
+
+}
+
+FContentBundleBase::~FContentBundleBase()
+{
+	check(GetStatus() == EContentBundleStatus::Unknown);
+}
+
+void FContentBundleBase::Initialize()
+{
+	check(GetStatus() == EContentBundleStatus::Unknown);
+
+	DoInitialize();
+
+	check(GetStatus() == EContentBundleStatus::Registered);
+}
+
+void FContentBundleBase::Uninitialize()
+{
+	check(GetStatus() != EContentBundleStatus::Unknown);
+
+	if (GetStatus() == EContentBundleStatus::ReadyToInject || GetStatus() == EContentBundleStatus::ContentInjected)
+	{
+		RemoveContent();
+	}
+
+	DoUninitialize();
+
+	check(GetStatus() == EContentBundleStatus::Unknown);
+}
+
+void FContentBundleBase::InjectContent()
+{
+	check(GetStatus() == EContentBundleStatus::Registered);
+
+	DoInjectContent();
+
+	check(GetStatus() == EContentBundleStatus::ReadyToInject || GetStatus() == EContentBundleStatus::ContentInjected || GetStatus() == EContentBundleStatus::FailedToInject);
+}
+
+void FContentBundleBase::RemoveContent()
+{
+	check(GetStatus() == EContentBundleStatus::ReadyToInject || GetStatus() == EContentBundleStatus::ContentInjected || GetStatus() == EContentBundleStatus::FailedToInject);
+
+	if (GetStatus() == EContentBundleStatus::ContentInjected || GetStatus() == EContentBundleStatus::ReadyToInject)
+	{
+		DoRemoveContent();
+	}
+
+	check(GetStatus() == EContentBundleStatus::Registered);
+}
+
+const TWeakPtr<FContentBundleClient>& FContentBundleBase::GetClient() const
+{
+	return Client;
+}
+
+const UContentBundleDescriptor* FContentBundleBase::GetDescriptor() const
+{
+	return Descriptor;
+}
+
+UWorld* FContentBundleBase::GetInjectedWorld() const
+{
+	return InjectedWorld.Get();
+}
+
+const FString& FContentBundleBase::GetDisplayName() const
+{
+	return GetDescriptor()->GetDisplayName();
+}
+
+void FContentBundleBase::AddReferencedObjects(FReferenceCollector& Collector)
+{
+	Collector.AddReferencedObject(Descriptor);
+}
+
+void FContentBundleBase::SetStatus(EContentBundleStatus NewStatus)
+{
+	check(NewStatus != Status);
+
+	UE_LOG(LogContentBundle, Log, TEXT("[CB: %s] State changing from %s to %s"), *GetDescriptor()->GetDisplayName(), *UEnum::GetDisplayValueAsText(Status).ToString(), *UEnum::GetDisplayValueAsText(NewStatus).ToString());
+	Status = NewStatus;
+}
\ No newline at end of file
diff --git a/Engine/Source/Runtime/Engine/Private/WorldPartition/ContentBundle/ContentBundleClient.cpp b/Engine/Source/Runtime/Engine/Private/WorldPartition/ContentBundle/ContentBundleClient.cpp
new file mode 100644
--- /dev/null
+++ b/Engine/Source/Runtime/Engine/Private/WorldPartition/ContentBundle/ContentBundleClient.cpp
@@ -0,0 +1,39 @@
+// Copyright Epic Games, Inc. All Rights Reserved.
+#include "WorldPartition/ContentBundle/ContentBundleClient.h"
+
+#include "WorldPartition/ContentBundle/ContentBundleEngineSubsystem.h"
+#include "WorldPartition/ContentBundle/ContentBundleDescriptor.h"
+#include "Engine/Engine.h"
+
+TSharedPtr<FContentBundleClient> FContentBundleClient::CreateClient(const UContentBundleDescriptor* InContentBundleDescriptor, FString const& InDisplayName)
+{
+	return GEngine->GetEngineSubsystem<UContentBundleEngineSubsystem>()->RegisterContentBundle(InContentBundleDescriptor, InDisplayName);
+}
+
+FContentBundleClient::FContentBundleClient(const UContentBundleDescriptor* InContentBundleDescriptor, FString const& InDisplayName)
+	:ContentBundleDescriptor(InContentBundleDescriptor),
+#if WITH_EDITOR
+	DisplayName(InDisplayName),
+#endif
+	ContentInjectionState(EContentBundleClientState::Unregistered)
+{
+
+}
+
+void FContentBundleClient::RequestContentInjection()
+{
+	GEngine->GetEngineSubsystem<UContentBundleEngineSubsystem>()->RequestContentInjection(*this);
+	ContentInjectionState = EContentBundleClientState::ContentInjectionRequested;
+}
+
+void FContentBundleClient::RequestRemoveContent()
+{
+	ContentInjectionState = EContentBundleClientState::ContentRemovalRequested;
+	GEngine->GetEngineSubsystem<UContentBundleEngineSubsystem>()->RequestContentRemoval(*this);
+}
+
+void FContentBundleClient::RequestUnregister()
+{
+	ContentInjectionState = EContentBundleClientState::ContentRemovalRequested;
+	GEngine->GetEngineSubsystem<UContentBundleEngineSubsystem>()->UnregisterContentBundle(*this);
+}
\ No newline at end of file
diff --git a/Engine/Source/Runtime/Engine/Private/WorldPartition/ContentBundle/ContentBundleContainer.cpp b/Engine/Source/Runtime/Engine/Private/WorldPartition/ContentBundle/ContentBundleContainer.cpp
new file mode 100644
--- /dev/null
+++ b/Engine/Source/Runtime/Engine/Private/WorldPartition/ContentBundle/ContentBundleContainer.cpp
@@ -0,0 +1,366 @@
+// Copyright Epic Games, Inc. All Rights Reserved.
+#include "WorldPartition/ContentBundle/ContentBundleContainer.h"
+
+#include "WorldPartition/ContentBundle/ContentBundle.h"
+#include "WorldPartition/ContentBundle/ContentBundleDescriptor.h"
+#include "WorldPartition/ContentBundle/ContentBundleClient.h"
+#include "WorldPartition/ContentBundle/ContentBundleEngineSubsystem.h"
+#include "WorldPartition/WorldPartition.h"
+#include "WorldPartition/WorldPartitionSubsystem.h"
+#include "WorldPartition/ContentBundle/ContentBundleEditor.h"
+#include "WorldPartition/ContentBundle/ContentBundleWorldSubsystem.h"
+#include "WorldPartition/ContentBundle/ContentBundleLog.h"
+#include "UObject/UObjectGlobals.h"
+
+FContentBundleContainer::FContentBundleContainer(UWorld* WorldToInjectIn)
+	:InjectedWorld(WorldToInjectIn)
+{
+
+}
+
+FContentBundleContainer::~FContentBundleContainer()
+{
+#if WITH_EDITOR
+	if (UseEditorContentBundle())
+	{
+		check(GetEditorContentBundles().IsEmpty());
+		return;
+	}
+#endif
+
+	check(GetGameContentBundles().IsEmpty());
+}
+
+UWorld* FContentBundleContainer::GetInjectedWorld() const
+{
+	return InjectedWorld.Get();
+}
+
+void FContentBundleContainer::Initialize()
+{
+	UE_LOG(LogContentBundle, Log, TEXT("[Container: %s] Creating new contrainer."), *GetInjectedWorld()->GetName());
+
+#if WITH_EDITOR
+	if (UseEditorContentBundle())
+	{
+		ContentBundlesVariant.Emplace<ContentBundleEditorArray>(ContentBundleEditorArray());
+	}
+	else
+	{
+		ContentBundlesVariant.Emplace<ContentBundleGameArray>(ContentBundleGameArray());
+	}	
+#endif
+
+	InitializeContentBundlesForegisteredClients();
+	RegisterContentBundleClientEvents();
+
+#if WITH_EDITOR
+	UWorldPartition* WorldPartition = GetInjectedWorld()->GetWorldPartition();
+	//WorldPartition->OnPreGenerateStreaming.AddRaw(this, &FContentBundleContainer::OnPreGenerateStreaming); // PHILTODO
+#endif
+}
+
+void FContentBundleContainer::Deinitialize()
+{
+	UE_LOG(LogContentBundle, Log, TEXT("[Container: %s] Deleting container."), *GetInjectedWorld()->GetName());
+
+#if WITH_EDITOR
+	UWorldPartition* WorldPartition = GetInjectedWorld()->GetWorldPartition();
+	//WorldPartition->OnPreGenerateStreaming.RemoveAll(this);
+#endif
+
+	UnregisterContentBundleClientEvents();
+	DeinitializeContentBundles();
+}
+
+void FContentBundleContainer::AddReferencedObjects(FReferenceCollector& Collector)
+{
+	ForEachContentBundle([&Collector](FContentBundleBase* ContentBundle) { ContentBundle->AddReferencedObjects(Collector); });
+}
+
+uint32 FContentBundleContainer::GetNumContentBundles() const
+{
+#if WITH_EDITOR
+	if(UseEditorContentBundle())
+	{
+		GetEditorContentBundles().Num();
+	}
+#endif
+
+	return GetGameContentBundles().Num();
+}
+
+const TArray<TUniquePtr<FContentBundle>>& FContentBundleContainer::GetGameContentBundles() const
+{
+#if WITH_EDITOR
+	check(!UseEditorContentBundle());
+	return ContentBundlesVariant.Get<ContentBundleGameArray>();
+#else
+	return ContentBundles;
+#endif
+}
+
+TArray<TUniquePtr<FContentBundle>>& FContentBundleContainer::GetGameContentBundles()
+{
+	return const_cast<TArray<TUniquePtr<FContentBundle>>&>(const_cast<const FContentBundleContainer*>(this)->GetGameContentBundles());
+}
+
+#if WITH_EDITOR
+
+TSharedPtr<FContentBundleEditor> FContentBundleContainer::GetEditorContentBundle(const UContentBundleDescriptor* Descriptor) const
+{
+	for (const TSharedPtr<FContentBundleEditor>& ContentBundle : GetEditorContentBundles())
+	{
+		if (ContentBundle->GetDescriptor() == Descriptor)
+		{
+			return ContentBundle;
+		}
+	}
+
+	return nullptr;
+}
+
+const TArray<TSharedPtr<FContentBundleEditor>>& FContentBundleContainer::GetEditorContentBundles() const
+{
+	check(UseEditorContentBundle());
+	return ContentBundlesVariant.Get<ContentBundleEditorArray>();
+}
+
+TArray<TSharedPtr<FContentBundleEditor>>& FContentBundleContainer::GetEditorContentBundles()
+{
+	return const_cast<TArray<TSharedPtr<FContentBundleEditor>>&>(const_cast<const FContentBundleContainer*>(this)->GetEditorContentBundles());
+}
+
+bool FContentBundleContainer::UseEditorContentBundle() const
+{
+	bool bIsEditorEditWorld = !GetInjectedWorld()->IsGameWorld() && GetInjectedWorld()->IsEditorWorld();
+	return bIsEditorEditWorld || IsRunningCookCommandlet();
+}
+
+#endif
+
+void FContentBundleContainer::ForEachContentBundle(TFunctionRef<void(FContentBundleBase*)> Func) const
+{
+#if WITH_EDITOR
+	if (UseEditorContentBundle())
+	{
+		for (const TSharedPtr<FContentBundleEditor>& ContentBundle : GetEditorContentBundles())
+		{
+			Func(ContentBundle.Get());
+		}
+		return;
+	}
+#endif
+
+	for (const TUniquePtr<FContentBundle>& ContentBundle : GetGameContentBundles())
+	{
+		Func(ContentBundle.Get());
+	}
+}
+
+void FContentBundleContainer::ForEachContentBundleBreakable(TFunctionRef<bool(FContentBundleBase*)> Func) const
+{
+#if WITH_EDITOR
+	if (UseEditorContentBundle())
+	{
+		for (const TSharedPtr<FContentBundleEditor>& ContentBundle : GetEditorContentBundles())
+		{
+			if (!Func(ContentBundle.Get()))
+			{
+				return;
+			}
+		}
+		return;
+	}
+#endif
+
+	for (const TUniquePtr<FContentBundle>& ContentBundle : GetGameContentBundles())
+	{
+		if (!Func(ContentBundle.Get()))
+		{
+			return;
+		}
+	}
+}
+
+FContentBundleBase* FContentBundleContainer::GetContentBundle(const FContentBundleClient& ContentBundleClient)
+{
+	FContentBundleBase* ContentBundlePtr = nullptr;
+	ForEachContentBundleBreakable([&ContentBundleClient, &ContentBundlePtr] (FContentBundleBase* ContentBundle) 
+	{
+		if (ContentBundle->GetDescriptor() == ContentBundleClient.GetDescriptor())
+		{
+			ContentBundlePtr = ContentBundle;
+			return false;
+		}
+		return true;
+	});
+	
+	if (ContentBundlePtr != nullptr)
+	{
+		return ContentBundlePtr;
+	}
+
+	return nullptr;
+}
+
+FContentBundleBase& FContentBundleContainer::InitializeContentBundle(TSharedPtr<FContentBundleClient>& ContentBundleClient)
+{
+	check(InjectedWorld->GetWorldPartition()->IsInitialized());
+	check(GetContentBundle(*ContentBundleClient) == nullptr);
+
+	FContentBundleBase* ContentBundle = nullptr;
+
+	UE_LOG(LogContentBundle, Log, TEXT("[Container: %s] Creating new ContentBundle %s from client %s with state %s."), 
+		*GetInjectedWorld()->GetName(), *ContentBundleClient->GetDescriptor()->GetDisplayName(), *ContentBundleClient->GetDisplayName(), *UEnum::GetDisplayValueAsText(ContentBundleClient->GetState()).ToString());
+
+#if WITH_EDITOR
+	if (UseEditorContentBundle())
+	{
+		ContentBundle = GetEditorContentBundles().Emplace_GetRef(MakeShared<FContentBundleEditor>(ContentBundleClient, InjectedWorld.Get())).Get();
+	}
+	else
+	{
+		ContentBundle = GetGameContentBundles().Emplace_GetRef(MakeUnique<FContentBundle>(ContentBundleClient, InjectedWorld.Get())).Get();
+	}
+#else
+	 ContentBundle = GetGameContentBundles().Emplace_GetRef(MakeUnique<FContentBundle>(ContentBundleClient, InjectedWorld.Get())).Get();
+#endif
+
+	 ContentBundle->Initialize();
+	 return *ContentBundle;
+}
+
+void FContentBundleContainer::DeinitializeContentBundle(FContentBundleBase& ContentBundle)
+{
+	ContentBundle.Uninitialize();
+
+#if WITH_EDITOR
+	if (UseEditorContentBundle())
+	{
+		uint32 Index = GetEditorContentBundles().IndexOfByPredicate([&ContentBundle](const TSharedPtr<FContentBundleEditor>& Other) { return Other.Get() == &ContentBundle; });
+		check(Index != INDEX_NONE);
+		GetEditorContentBundles().RemoveAtSwap(Index);
+		return;
+	}
+#endif
+
+	uint32 Index = GetGameContentBundles().IndexOfByPredicate([&ContentBundle](const TUniquePtr<FContentBundle>& Other) { return Other.Get() == &ContentBundle; });
+	GetGameContentBundles().RemoveAtSwap(Index);
+}
+
+void FContentBundleContainer::InjectContentBundle(FContentBundleBase& ContentBundle)
+{
+	ContentBundle.InjectContent();
+}
+
+void FContentBundleContainer::RemoveContentBundle(FContentBundleBase& ContentBundle)
+{
+	ContentBundle.RemoveContent();
+}
+
+void FContentBundleContainer::OnContentBundleClientRegistered(TSharedPtr<FContentBundleClient>& ContentBundleClient)
+{
+	InitializeContentBundle(ContentBundleClient);
+}
+
+void FContentBundleContainer::OnContentBundleClientUnregistered(FContentBundleClient& ContentBundleClient)
+{
+	if (FContentBundleBase* ContentBundle = GetContentBundle(ContentBundleClient))
+	{
+		DeinitializeContentBundle(*ContentBundle);
+	}
+}
+
+void FContentBundleContainer::OnContentBundleClientContentInjectionRequested(FContentBundleClient& ContentBundleClient)
+{
+	if (FContentBundleBase* ContentBundle = GetContentBundle(ContentBundleClient))
+	{
+		InjectContentBundle(*ContentBundle);
+	}
+}
+
+void FContentBundleContainer::OnContentBundleClientContentRemovalRequested(FContentBundleClient& ContentBundleClient)
+{
+	if (FContentBundleBase* ContentBundle = GetContentBundle(ContentBundleClient))
+	{
+		RemoveContentBundle(*ContentBundle);
+	}
+}
+
+void FContentBundleContainer::RegisterContentBundleClientEvents()
+{
+	UContentBundleEngineSubsystem* ContentBundleEngineSubsystem = UContentBundleEngineSubsystem::Get();
+	ContentBundleEngineSubsystem->OnContentBundleClientRegistered.AddRaw(this, &FContentBundleContainer::OnContentBundleClientRegistered);
+	ContentBundleEngineSubsystem->OnContentBundleClientUnregistered.AddRaw(this, &FContentBundleContainer::OnContentBundleClientUnregistered);
+	ContentBundleEngineSubsystem->OnContentBundleClientRequestedContentInjection.AddRaw(this, &FContentBundleContainer::OnContentBundleClientContentInjectionRequested);
+	ContentBundleEngineSubsystem->OnContentBundleClientRequestedContentRemoval.AddRaw(this, &FContentBundleContainer::OnContentBundleClientContentRemovalRequested);
+}
+
+void FContentBundleContainer::UnregisterContentBundleClientEvents()
+{
+	UContentBundleEngineSubsystem* ContentBundleEngineSubsystem = UContentBundleEngineSubsystem::Get();
+	ContentBundleEngineSubsystem->OnContentBundleClientRegistered.RemoveAll(this);
+	ContentBundleEngineSubsystem->OnContentBundleClientUnregistered.RemoveAll(this);
+	ContentBundleEngineSubsystem->OnContentBundleClientRequestedContentInjection.RemoveAll(this);
+	ContentBundleEngineSubsystem->OnContentBundleClientRequestedContentRemoval.RemoveAll(this);
+}
+
+bool FContentBundleContainer::ShouldInjectClientContent(const TSharedPtr<FContentBundleClient>& ContentBundleClient) const
+{
+	return ContentBundleClient->GetState() == EContentBundleClientState::ContentInjected || ContentBundleClient->GetState() == EContentBundleClientState::ContentInjectionRequested;
+}
+
+void FContentBundleContainer::InitializeContentBundlesForegisteredClients()
+{
+	UContentBundleEngineSubsystem* ContentBundleEngineSubsystem = UContentBundleEngineSubsystem::Get();
+	TArrayView<TSharedPtr<FContentBundleClient>> ContentBundleClients = ContentBundleEngineSubsystem->GetContentBundleClients();
+
+	if (!ContentBundleClients.IsEmpty())
+	{
+		UE_LOG(LogContentBundle, Log, TEXT("[Container: %s] Begin initializing ContentBundles from %u registered clients."), *GetInjectedWorld()->GetName(), ContentBundleClients.Num());
+
+		for (TSharedPtr<FContentBundleClient>& ContentBundleClient : ContentBundleClients)
+		{
+			FContentBundleBase& ContentBundle = InitializeContentBundle(ContentBundleClient);
+
+			if (ShouldInjectClientContent(ContentBundleClient))
+			{
+				InjectContentBundle(ContentBundle);
+			}
+		}
+
+		UE_LOG(LogContentBundle, Log, TEXT("[Container: %s] End initializing ContentBundles."), *GetInjectedWorld()->GetName());
+	}
+}
+
+void FContentBundleContainer::DeinitializeContentBundles()
+{
+	ForEachContentBundle([](FContentBundleBase* ContentBundle) { ContentBundle->Uninitialize(); });
+
+#if WITH_EDITOR
+	if (UseEditorContentBundle())
+	{
+		GetEditorContentBundles().Empty();
+		return;
+	}
+#endif
+
+	GetGameContentBundles().Empty();
+}
+
+#if WITH_EDITOR
+
+void FContentBundleContainer::OnPreGenerateStreaming()
+{
+	UE_LOG(LogContentBundle, Log, TEXT("[Container: %s] Generating Streaming for %u Content Bundles."), *GetInjectedWorld()->GetName(), GetEditorContentBundles().Num());
+
+	GetInjectedWorld()->ContentBundleManager->GetPIEDuplicateHelper()->Clear();
+
+	for (TSharedPtr<FContentBundleEditor>& ContentBundle : GetEditorContentBundles())
+	{
+		ContentBundle->GenerateStreaming();
+	}
+}
+
+#endif
\ No newline at end of file
diff --git a/Engine/Source/Runtime/Engine/Private/WorldPartition/ContentBundle/ContentBundleDescriptor.cpp b/Engine/Source/Runtime/Engine/Private/WorldPartition/ContentBundle/ContentBundleDescriptor.cpp
new file mode 100644
--- /dev/null
+++ b/Engine/Source/Runtime/Engine/Private/WorldPartition/ContentBundle/ContentBundleDescriptor.cpp
@@ -0,0 +1,20 @@
+// Copyright Epic Games, Inc. All Rights Reserved.
+#include "WorldPartition/ContentBundle/ContentBundleDescriptor.h"
+
+UContentBundleDescriptor::UContentBundleDescriptor(const FObjectInitializer& ObjectInitializer)
+	:Guid(FGuid::NewGuid())
+{
+
+}
+
+const FString& UContentBundleDescriptor::GetDisplayName() const
+{
+	return DisplayName;
+}
+
+bool UContentBundleDescriptor::IsValid() const
+{
+	return Guid.IsValid()
+		&& !DisplayName.IsEmpty()
+		&& !PackageRoot.IsEmpty();
+}
\ No newline at end of file
diff --git a/Engine/Source/Runtime/Engine/Private/WorldPartition/ContentBundle/ContentBundleEditor.cpp b/Engine/Source/Runtime/Engine/Private/WorldPartition/ContentBundle/ContentBundleEditor.cpp
new file mode 100644
--- /dev/null
+++ b/Engine/Source/Runtime/Engine/Private/WorldPartition/ContentBundle/ContentBundleEditor.cpp
@@ -0,0 +1,479 @@
+// Copyright Epic Games, Inc. All Rights Reserved.
+#include "WorldPartition/ContentBundle/ContentBundleEditor.h"
+
+#if WITH_EDITOR
+
+#include "WorldPartition/ContentBundle/ContentBundle.h"
+#include "WorldPartition/WorldPartition.h"
+#include "WorldPartition/DataLayer/WorldDataLayers.h"
+#include "WorldPartition/WorldPartitionRuntimeCell.h"
+#include "Engine/World.h"
+#include "WorldPartition/WorldPartitionRuntimeHash.h"
+#include "WorldPartition/ContentBundle/ContentBundleDescriptor.h"
+#include "WorldPartition/ContentBundle/ContentBundleEditorSubsystemInterface.h"
+#include "WorldPartition/ContentBundle/ContentBundleLog.h"
+#include "PackageTools.h"
+#include "ObjectTools.h"
+#include "Editor.h"
+#include "WorldPartition/ContentBundle/ContentBundleWorldSubsystem.h"
+
+FContentBundleEditor::FContentBundleEditor(TSharedPtr<FContentBundleClient>& InClient, UWorld* InWorld)
+	:FContentBundleBase(InClient, InWorld),
+	UnsavedActorMonitor(nullptr),
+	ExternalStreamingObject(nullptr),
+	Guid(FGuid::NewGuid()),
+	bIsBeingEdited(false)
+{
+
+}
+
+
+void FContentBundleEditor::DoInitialize()
+{
+	SetStatus(EContentBundleStatus::Registered);
+
+	if (IContentBundleEditorSubsystemInterface* EditorSubsystem = IContentBundleEditorSubsystemInterface::Get())
+	{
+		EditorSubsystem->NotifyContentBundleAdded(this);
+	}
+}
+
+void FContentBundleEditor::DoUninitialize()
+{
+	if (IContentBundleEditorSubsystemInterface* EditorSubsystem = IContentBundleEditorSubsystemInterface::Get())
+	{
+		EditorSubsystem->NotifyContentBundleRemoved(this);
+	}
+
+	SetStatus(EContentBundleStatus::Unknown);
+}
+
+void FContentBundleEditor::DoInjectContent()
+{
+	FString ActorDescContainerPackage;
+	if (BuildContentBundleContainerPackagePath(ActorDescContainerPackage))
+	{
+		UnsavedActorMonitor = NewObject<UContentBundleUnsavedActorMonitor>(GetTransientPackage(), NAME_None, RF_Transactional);
+		UnsavedActorMonitor->Initialize(*this);
+
+		UWorldPartition* WorldPartition = GetInjectedWorld()->GetWorldPartition();
+		ActorDescContainer = WorldPartition->RegisterActorDescContainer(FName(*ActorDescContainerPackage));
+		if (ActorDescContainer.IsValid())
+		{
+			UE_LOG(LogContentBundle, Log, TEXT("[CB: %s] ExternalActors in %s found. %u actors were injected"), *GetDescriptor()->GetDisplayName(), *ActorDescContainer->GetExternalActorPath(), ActorDescContainer->GetActorDescCount());
+
+			if (!ActorDescContainer->IsEmpty())
+			{
+				WorldDataLayersActorReference = FWorldDataLayersReference(ActorDescContainer.Get(), BuildWorlDataLayersName());
+				SetStatus(EContentBundleStatus::ContentInjected);
+			}
+			else
+			{
+				SetStatus(EContentBundleStatus::ReadyToInject);
+			}
+
+			RegisterDelegates();
+		}
+		else
+		{
+			UE_LOG(LogContentBundle, Log, TEXT("[CB: %s] Failed to register actor desc container with %s"), *GetDescriptor()->GetDisplayName(), *ActorDescContainerPackage);
+			SetStatus(EContentBundleStatus::FailedToInject);
+		}
+	}
+	else
+	{
+		SetStatus(EContentBundleStatus::FailedToInject);
+	}
+
+
+	BroadcastChanged();
+}
+
+void FContentBundleEditor::DoRemoveContent()
+{
+	UnreferenceAllActors();
+
+	WorldDataLayersActorReference.Reset();
+
+	UnsavedActorMonitor->Uninitialize();
+
+	if (ActorDescContainer.IsValid())
+	{
+		UnregisterDelegates();
+
+		GetInjectedWorld()->GetWorldPartition()->UnregisterActorDescContainer(ActorDescContainer.Get());
+		ActorDescContainer = nullptr;
+	}
+
+	SetStatus(EContentBundleStatus::Registered);
+
+	BroadcastChanged();
+}
+
+bool FContentBundleEditor::IsValid() const
+{
+	bool bIsValid = true;
+
+	return bIsValid;
+}
+
+bool FContentBundleEditor::AddActor(AActor* InActor)
+{
+	if (InActor->GetWorld() != ActorDescContainer->GetWorld() || InActor->HasAllFlags(RF_Transient) || !InActor->IsMainPackageActor())
+	{
+		return false;
+	}
+
+	if (InActor->IsA<AWorldDataLayers>())
+	{
+		return false;
+	}
+
+	check(GetStatus() == EContentBundleStatus::ContentInjected);
+
+	FName ActorPackageNameInContentBundle(*ULevel::GetActorPackageName(ActorDescContainer->GetExternalActorPath(), EActorPackagingScheme::Reduced, InActor->GetName()));
+	FName ActorPackageName = InActor->GetPackage()->GetFName();
+	if (ActorPackageName != ActorPackageNameInContentBundle)
+	{
+		InActor->GetPackage()->Rename(*ActorPackageNameInContentBundle.ToString());
+	}
+
+	UnsavedActorMonitor->MonitorActor(InActor);
+
+	UE_LOG(LogContentBundle, Verbose, TEXT("[CB: %s] Added new actor %s, ActorCount:  %u. Package %s."), *GetDescriptor()->GetDisplayName(), *InActor->GetActorNameOrLabel(), GetActorCount(), *InActor->GetPackage()->GetName());
+
+	return true;
+}
+
+bool FContentBundleEditor::ContainsActor(const AActor* InActor) const
+{
+	if (InActor != nullptr)
+	{
+		return ActorDescContainer->GetActorDesc(InActor) != nullptr || UnsavedActorMonitor->IsMonitoring(InActor);
+	}
+
+	return false;
+}
+
+bool FContentBundleEditor::GetActors(TArray<AActor*>& Actors)
+{
+	Actors.Reserve(GetActorCount());
+
+	for (FActorDescList::TIterator<> It(ActorDescContainer.Get()); It; ++It)
+	{
+		if (AActor* Actor = It->GetActor())
+		{
+			if (Actor != WorldDataLayersActorReference.Get())
+			{
+				Actors.Add(Actor);
+			}
+		}
+	}
+
+	for (auto UnsavedActor : UnsavedActorMonitor->GetUnsavedActors())
+	{
+		if (UnsavedActor.IsValid())
+		{
+			Actors.Add(UnsavedActor.Get());
+		}
+	}
+
+	return !Actors.IsEmpty();
+}
+
+bool FContentBundleEditor::HasUserPlacedActors() const
+{
+	// If there is only one actor in the container its the WorldDataLayer automatically created when injecting base content.
+	bool bActorDescContHasUserPlacedActors = ActorDescContainer.IsValid() && ActorDescContainer->GetActorDescCount() > 1;
+	return (bActorDescContHasUserPlacedActors || UnsavedActorMonitor->IsMonitoringActors());
+}
+
+uint32 FContentBundleEditor::GetActorCount() const
+{
+	if (GetStatus() == EContentBundleStatus::ContentInjected)
+	{
+		uint32 UnsavedWorldDataLayerCount = WorldDataLayersActorReference.IsValid() && ActorDescContainer.IsValid() && ActorDescContainer->IsEmpty() ? 1 : 0;
+		return ActorDescContainer->GetActorDescCount() + UnsavedActorMonitor->GetActorCount() + UnsavedWorldDataLayerCount;
+	}
+
+	return 0;
+}
+
+uint32 FContentBundleEditor::GetUnsavedActorAcount() const
+{
+	if (GetStatus() == EContentBundleStatus::ContentInjected)
+	{
+		return UnsavedActorMonitor->GetActorCount();
+	}
+
+	return 0;
+}
+
+void FContentBundleEditor::ReferenceAllActors()
+{
+	if (ActorDescContainer.IsValid())
+	{
+		ActorDescContainer->LoadAllActors(ForceLoadedActors);
+	}
+}
+
+void FContentBundleEditor::UnreferenceAllActors()
+{
+	ForceLoadedActors.Empty();
+}
+
+void FContentBundleEditor::StartEditing()
+{
+	check(GetStatus() == EContentBundleStatus::ReadyToInject || GetStatus() == EContentBundleStatus::ContentInjected);
+
+	UnsavedActorMonitor->StartListenOnActorEvents();
+
+	bIsBeingEdited = true;
+}
+
+void FContentBundleEditor::StopEditing()
+{
+	check(GetStatus() == EContentBundleStatus::ReadyToInject || GetStatus() == EContentBundleStatus::ContentInjected);
+
+	UnsavedActorMonitor->StopListeningOnActorEvents();
+
+	bIsBeingEdited = false;
+}
+
+void FContentBundleEditor::InjectBaseContent()
+{
+	check(GetStatus() == EContentBundleStatus::ReadyToInject);
+	UE_LOG(LogContentBundle, Log, TEXT("[CB: %s] Injecting Base Content"), *GetDescriptor()->GetDisplayName());
+
+	FActorSpawnParameters SpawnParameters;
+	SpawnParameters.Name = BuildWorlDataLayersName();
+	SpawnParameters.OverrideLevel = GetInjectedWorld()->PersistentLevel;
+	SpawnParameters.OverridePackage = CreateActorPackage(SpawnParameters.Name);
+	SpawnParameters.bCreateActorPackage = false;
+
+	WorldDataLayersActorReference = FWorldDataLayersReference(SpawnParameters);
+
+	WorldDataLayersActorReference->SetActorLabel(GetDisplayName());
+
+	SetStatus(EContentBundleStatus::ContentInjected);
+
+	BroadcastChanged();
+}
+
+void FContentBundleEditor::RemoveBaseContent()
+{
+	check(!HasUserPlacedActors());
+	check(GetStatus() == EContentBundleStatus::ContentInjected);
+	UE_LOG(LogContentBundle, Log, TEXT("[CB: %s] Removing Base Content"), *GetDescriptor()->GetDisplayName());
+
+	GetInjectedWorld()->DestroyActor(WorldDataLayersActorReference.Get());
+	WorldDataLayersActorReference.Reset();
+
+	SetStatus(EContentBundleStatus::ReadyToInject);
+
+	BroadcastChanged();
+}
+
+void FContentBundleEditor::AddReferencedObjects(FReferenceCollector& Collector)
+{
+	FContentBundleBase::AddReferencedObjects(Collector);
+
+	Collector.AddReferencedObjects(ContentBundleCells);
+	Collector.AddReferencedObject(ExternalStreamingObject);
+	Collector.AddReferencedObject(UnsavedActorMonitor);
+}
+
+void FContentBundleEditor::GenerateStreaming()
+{
+	if (GetStatus() != EContentBundleStatus::ContentInjected)
+	{
+		UE_LOG(LogContentBundle, Log, TEXT("[CB: %s] Skipping streaming generation. It's status is: %s."), *GetDescriptor()->GetDisplayName(), *UEnum::GetDisplayValueAsText(GetStatus()).ToString());
+		return;
+	}
+
+	UWorldPartition* WorldPartition = GetInjectedWorld()->GetWorldPartition();
+	WorldPartition->GenerateContainerStreaming(ActorDescContainer.Get());
+
+	WorldPartition->RuntimeHash->GetAllStreamingCells(ContentBundleCells, true);
+
+	FString ExternalStreamingObjectName = ObjectTools::SanitizeInvalidChars(GetDisplayName() + TEXT("_ExternalStreamingObject"), INVALID_LONGPACKAGE_CHARACTERS) ;
+	ExternalStreamingObject = WorldPartition->RuntimeHash->StoreToExternalStreamingObject(GetInjectedWorld()->ContentBundleManager, *ExternalStreamingObjectName);
+
+	UE_LOG(LogContentBundle, Log, TEXT("[CB: %s] Generated streaming cells. %u cells were generated"), *GetDescriptor()->GetDisplayName(), ContentBundleCells.Num());
+
+	if (!IsRunningCookCommandlet())
+	{
+		UContentBundleManager* ContentBundleManager = GetInjectedWorld()->ContentBundleManager;
+		UContentBundleDuplicateForPIEHelper* DuplicateForPIEHelper = ContentBundleManager->GetPIEDuplicateHelper();
+		if (!DuplicateForPIEHelper->StoreContentBundleStreamingObect(*this, ExternalStreamingObject))
+		{
+			UE_LOG(LogContentBundle, Error, TEXT("[CB: %s] Failed to store streaming object for %s. PIE duplication will not work."), *GetDescriptor()->GetDisplayName());
+		}
+	}
+
+	WorldPartition->FlushStreaming();
+}
+
+void FContentBundleEditor::BroadcastChanged()
+{
+	if (IContentBundleEditorSubsystemInterface* EditorSubsystem = IContentBundleEditorSubsystemInterface::Get())
+	{
+		EditorSubsystem->NotifyContentBundleChanged(this);
+	}
+}
+
+bool FContentBundleEditor::BuildContentBundleContainerPackagePath(FString& ContainerPackagePath) const
+{
+	FString PackageRoot, PackagePath, PackageName;
+	FString LongPackageName = GetInjectedWorld()->GetPackage()->GetName();
+	if (FPackageName::SplitLongPackageName(LongPackageName, PackageRoot, PackagePath, PackageName))
+	{
+		TStringBuilderWithBuffer<TCHAR, NAME_SIZE> PluginLeveldPackagePath;
+		PluginLeveldPackagePath += TEXT("/");
+		PluginLeveldPackagePath += GetDescriptor()->GetPackageRoot();
+		PluginLeveldPackagePath += TEXT("/ContentBundle/");
+		PluginLeveldPackagePath += GetDescriptor()->GetGuid().ToString();
+		PluginLeveldPackagePath += TEXT("/");
+		PluginLeveldPackagePath += PackagePath;
+		PluginLeveldPackagePath += PackageName;
+
+		ContainerPackagePath = UPackageTools::SanitizePackageName(*PluginLeveldPackagePath);
+		return true;
+	}
+
+	UE_LOG(LogContentBundle, Error, TEXT("[CB: %s] Failed to build Container Package Path using %s"), *GetDescriptor()->GetDisplayName(), *LongPackageName);
+	return false;
+}
+
+UPackage* FContentBundleEditor::CreateActorPackage(const FName& ActorName) const
+{
+	FString ActorPackagePath = ULevel::GetActorPackageName(ActorDescContainer->GetExternalActorPath(), EActorPackagingScheme::Reduced, ActorName.ToString());
+	UPackage* ActorPackage = CreatePackage(*ActorPackagePath);
+
+	ActorPackage->SetDirtyFlag(true);
+
+	return ActorPackage;
+}
+
+FName FContentBundleEditor::BuildWorlDataLayersName() const
+{
+	return *GetDescriptor()->GetGuid().ToString();
+}
+
+void FContentBundleEditor::RegisterDelegates()
+{
+	ActorDescContainer->OnActorDescAddedEvent.AddRaw(this, &FContentBundleEditor::OnActorDescAdded);
+	ActorDescContainer->OnActorDescRemovedEvent.AddRaw(this, &FContentBundleEditor::OnActorDescRemoved);
+}
+
+void FContentBundleEditor::UnregisterDelegates()
+{
+	ActorDescContainer->OnActorDescAddedEvent.RemoveAll(this);
+	ActorDescContainer->OnActorDescRemovedEvent.RemoveAll(this);
+}
+
+void FContentBundleEditor::OnActorDescAdded(FWorldPartitionActorDesc* ActorDesc)
+{
+	UE_LOG(LogContentBundle, Verbose, TEXT("[CB: %s] Added actor %s to container, ActorCount: %u. Package %s."), 
+		*GetDescriptor()->GetDisplayName(), *ActorDesc->GetActorLabelOrName().ToString(), GetActorCount(), *ActorDesc->GetActorPackage().ToString());
+
+	AActor* Actor = ActorDesc->GetActor();
+	UnsavedActorMonitor->StopMonitoringActor(Actor);
+}
+
+void FContentBundleEditor::OnActorDescRemoved(FWorldPartitionActorDesc* ActorDesc)
+{
+	UE_LOG(LogContentBundle, Verbose, TEXT("[CB: %s] Removed actor %s from container, ActorCount:  %u. Package %s."), 
+		*GetDescriptor()->GetDisplayName(), *ActorDesc->GetActorLabelOrName().ToString(), GetActorCount(), *ActorDesc->GetActorPackage().ToString());
+
+	if (!HasUserPlacedActors())
+	{
+		if (GetStatus() == EContentBundleStatus::ContentInjected)
+		{
+			RemoveBaseContent();
+		}
+		else
+		{
+			// If content is not injected, then RemoveContent was already called and we are deleting the WorldDataLayers
+			check(ActorDesc->GetActorNativeClass() == AWorldDataLayers::StaticClass() && GetStatus() == EContentBundleStatus::ReadyToInject);
+		}
+	}
+
+	AActor* ActorInWorld = ActorDesc->GetActor(false, false);
+	check(ActorInWorld == nullptr || !UnsavedActorMonitor->IsMonitoring(ActorInWorld)); // ActorDesc existed is being deleted. Make sure the actor is not present in the unsaved list as it should have been saved for the desc to exist.
+}
+
+void FContentBundleEditor::OnUnsavedActorDeleted(AActor* Actor)
+{
+	UE_LOG(LogContentBundle, Verbose, TEXT("[CB: %s] Removed unsaved actor %s, ActorCount: %u. Package %s."),
+		*GetDescriptor()->GetDisplayName(), *Actor->GetActorNameOrLabel(), GetActorCount(), *Actor->GetPackage()->GetName());
+
+	if (!HasUserPlacedActors())
+	{
+		RemoveBaseContent();
+	}
+}
+
+UContentBundleUnsavedActorMonitor::~UContentBundleUnsavedActorMonitor()
+{
+	check(UnsavedActors.IsEmpty());
+	check(ContentBundle == nullptr);
+}
+
+void UContentBundleUnsavedActorMonitor::Initialize(FContentBundleEditor& InContentBundleEditor)
+{
+	ContentBundle = &InContentBundleEditor;
+}
+
+void UContentBundleUnsavedActorMonitor::StartListenOnActorEvents()
+{
+	GEngine->OnLevelActorDeleted().AddUObject(this, &UContentBundleUnsavedActorMonitor::OnActorDeleted);
+}
+
+void UContentBundleUnsavedActorMonitor::StopListeningOnActorEvents()
+{
+	GEngine->OnLevelActorDeleted().RemoveAll(this);
+}
+
+void UContentBundleUnsavedActorMonitor::Uninitialize()
+{
+	StopListeningOnActorEvents();
+
+	for (TWeakObjectPtr<AActor>& Actor : UnsavedActors)
+	{
+		ContentBundle->GetInjectedWorld()->DestroyActor(Actor.Get());
+	}
+	UnsavedActors.Empty();
+
+	ContentBundle = nullptr;
+}
+
+void UContentBundleUnsavedActorMonitor::MonitorActor(AActor* InActor)
+{
+	Modify();
+	UnsavedActors.Add(InActor);
+}
+
+bool UContentBundleUnsavedActorMonitor::StopMonitoringActor(AActor* InActor)
+{
+	uint32 Idx = UnsavedActors.IndexOfByKey(InActor);
+	if (Idx != INDEX_NONE)
+	{
+		Modify();
+		UnsavedActors.Remove(InActor);
+		check(UnsavedActors.IndexOfByKey(InActor) == INDEX_NONE);
+		return true;
+	}
+
+	return false;
+}
+
+void UContentBundleUnsavedActorMonitor::OnActorDeleted(AActor* InActor)
+{
+	if (StopMonitoringActor(InActor))
+	{
+		ContentBundle->OnUnsavedActorDeleted(InActor);
+	}
+}
+
+#endif
\ No newline at end of file
diff --git a/Engine/Source/Runtime/Engine/Private/WorldPartition/ContentBundle/ContentBundleEditorSubsystemInterface.cpp b/Engine/Source/Runtime/Engine/Private/WorldPartition/ContentBundle/ContentBundleEditorSubsystemInterface.cpp
new file mode 100644
--- /dev/null
+++ b/Engine/Source/Runtime/Engine/Private/WorldPartition/ContentBundle/ContentBundleEditorSubsystemInterface.cpp
@@ -0,0 +1,14 @@
+// Copyright Epic Games, Inc. All Rights Reserved.
+#include "WorldPartition/ContentBundle/ContentBundleEditorSubsystemInterface.h"
+
+#if WITH_EDITOR
+
+IContentBundleEditorSubsystemInterface* IContentBundleEditorSubsystemInterface::Instance = nullptr;
+
+void IContentBundleEditorSubsystemInterface::SetInstance(IContentBundleEditorSubsystemInterface* InInstance)
+{ 
+	check(InInstance == nullptr || Instance == nullptr); 
+	Instance = InInstance; 
+}
+
+#endif
\ No newline at end of file
diff --git a/Engine/Source/Runtime/Engine/Private/WorldPartition/ContentBundle/ContentBundleEngineSubsystem.cpp b/Engine/Source/Runtime/Engine/Private/WorldPartition/ContentBundle/ContentBundleEngineSubsystem.cpp
new file mode 100644
--- /dev/null
+++ b/Engine/Source/Runtime/Engine/Private/WorldPartition/ContentBundle/ContentBundleEngineSubsystem.cpp
@@ -0,0 +1,147 @@
+// Copyright Epic Games, Inc. All Rights Reserved.
+#include "WorldPartition/ContentBundle/ContentBundleEngineSubsystem.h"
+
+#include "WorldPartition/ContentBundle/ContentBundleClient.h"
+#include "WorldPartition/ContentBundle/ContentBundleDescriptor.h"
+#include "WorldPartition/ContentBundle/ContentBundleLog.h"
+#include "WorldPartition/ContentBundle/ContentBundleWorldSubsystem.h"
+#include "Engine/Engine.h"
+
+void UContentBundleEngineSubsystem::Initialize(FSubsystemCollectionBase& Collection)
+{
+	Super::Initialize(Collection);
+
+	FWorldDelegates::OnPreWorldInitialization.AddUObject(this, &UContentBundleEngineSubsystem::OnPreWorldInitialization);
+	FWorldDelegates::OnPostWorldCleanup.AddUObject(this, &UContentBundleEngineSubsystem::OnWorldPostCleanup);
+}
+
+void UContentBundleEngineSubsystem::Deinitialize()
+{
+	ContentBundleClients.Empty();
+
+	FWorldDelegates::OnPreWorldInitialization.RemoveAll(this);
+	FWorldDelegates::OnPostWorldCleanup.RemoveAll(this);
+}
+
+TSharedPtr<FContentBundleClient> UContentBundleEngineSubsystem::RegisterContentBundle(const UContentBundleDescriptor* Descriptor, const FString& ClientDisplayName)
+{
+	if (ContentBundleClients.IsEmpty())
+	{
+		UE_SET_LOG_VERBOSITY(LogContentBundle, Verbose);
+	}
+
+	if (Descriptor != nullptr)
+	{
+		TSharedPtr<FContentBundleClient>* RegisteredClient = FindRegisteredClient(Descriptor);
+		if (RegisteredClient == nullptr)
+		{
+			check(FindRegisteredClient(Descriptor) == nullptr);;
+
+			TSharedPtr<FContentBundleClient>& ContentBundleClient = ContentBundleClients.Add_GetRef(MakeShared<FContentBundleClient>(Descriptor, ClientDisplayName));
+
+			UE_LOG(LogContentBundle, Log, TEXT("[Client] New client registered. Client: %s, Descriptor: %s"), *ContentBundleClient->GetDisplayName(), *ContentBundleClient->GetDescriptor()->GetDisplayName());
+
+			OnContentBundleClientRegistered.Broadcast(ContentBundleClient);
+
+			return ContentBundleClient;
+		}
+		else
+		{
+			UE_LOG(LogContentBundle, Error, TEXT("[Client] Content bundle descriptor %s is already registered by client %s"), *Descriptor->GetDisplayName(), *(*RegisteredClient)->GetDisplayName());
+		}
+	}
+
+	return nullptr;
+}	
+
+void UContentBundleEngineSubsystem::UnregisterContentBundle(FContentBundleClient& InClient)
+{
+	uint32 Index = ContentBundleClients.IndexOfByPredicate([&InClient](const TSharedPtr<FContentBundleClient>& Client) {return Client.Get() == &InClient;});
+	if (Index != INDEX_NONE)
+	{
+		const UContentBundleDescriptor* DescriptorBackup = InClient.GetDescriptor();
+		check(FindRegisteredClient(DescriptorBackup) != nullptr);
+
+		UE_LOG(LogContentBundle, Log, TEXT("[Client] Unregistered Client. Client: %s, Descriptor: %s"), *InClient.GetDisplayName(), *InClient.GetDescriptor()->GetDisplayName());
+
+		OnContentBundleClientUnregistered.Broadcast(*ContentBundleClients[Index]);
+		ContentBundleClients.RemoveAtSwap(Index);
+
+		check(FindRegisteredClient(DescriptorBackup) == nullptr);
+	}
+	else 
+	{
+		UE_LOG(LogContentBundle, Error, TEXT("[Client] Cannot unregister client %s. It was not registered"), *InClient.GetDisplayName());
+	}	
+}
+
+void UContentBundleEngineSubsystem::RequestContentInjection(FContentBundleClient& InClient)
+{
+	if (TSharedPtr<FContentBundleClient>* Client = FindRegisteredClient(InClient.GetDescriptor()))
+	{
+		UE_LOG(LogContentBundle, Log, TEXT("[Client] Client requested content injection. Client: %s, Descriptor: %s"), *InClient.GetDisplayName(), *InClient.GetDescriptor()->GetDisplayName());
+
+		OnContentBundleClientRequestedContentInjection.Broadcast(**Client);
+	}
+	else
+	{
+		UE_LOG(LogContentBundle, Error, TEXT("[Client] Cannot inject client %s\'s content. It was not registered"), *InClient.GetDisplayName());
+	}
+}
+
+void UContentBundleEngineSubsystem::RequestContentRemoval(FContentBundleClient& InClient)
+{
+	if (TSharedPtr<FContentBundleClient>* Client = FindRegisteredClient(InClient.GetDescriptor()))
+	{
+		UE_LOG(LogContentBundle, Log, TEXT("[Client] Client requested content removal. Client: %s, Descriptor: %s"), *InClient.GetDisplayName(), *InClient.GetDescriptor()->GetDisplayName());
+
+		OnContentBundleClientRequestedContentRemoval.Broadcast(**Client);
+	}
+	else
+	{
+		UE_LOG(LogContentBundle, Error, TEXT("[Client] Cannot inject client %s\'s content. It was not registered"), *InClient.GetDisplayName());
+	}
+}
+
+UContentBundleEngineSubsystem* UContentBundleEngineSubsystem::Get()
+{
+	if (GEngine)
+	{
+		return GEngine->GetEngineSubsystem<UContentBundleEngineSubsystem>();
+	}
+	check(0); // Retrieving the ContentBundleEngineSubsystem before it was created.
+	return nullptr;
+}
+
+TSharedPtr<FContentBundleClient>* UContentBundleEngineSubsystem::FindRegisteredClient(const UContentBundleDescriptor* Descriptor)
+{
+	return ContentBundleClients.FindByPredicate([Descriptor](const TSharedPtr<FContentBundleClient>& Client) { return Client->GetDescriptor() == Descriptor; });
+}
+
+TSharedPtr<FContentBundleClient>* UContentBundleEngineSubsystem::FindRegisteredClient(FContentBundleClient& InClient)
+{
+	TSharedPtr<FContentBundleClient>* RegisteredClient = FindRegisteredClient(InClient.GetDescriptor());
+	check(RegisteredClient == nullptr || RegisteredClient->Get() == &InClient);
+	return RegisteredClient;
+}
+
+void UContentBundleEngineSubsystem::OnPreWorldInitialization(UWorld* World, const UWorld::InitializationValues IVS)
+{
+	if (World->GetWorldPartition() && World->ContentBundleManager == nullptr)
+	{
+		World->ContentBundleManager = NewObject<UContentBundleManager>(World);
+	}
+
+	if (World->ContentBundleManager != nullptr)
+	{
+		World->ContentBundleManager->Initialize();
+	}	
+}
+
+void UContentBundleEngineSubsystem::OnWorldPostCleanup(UWorld* World, bool bSessionEnded, bool bCleanupResources)
+{
+	if (World->ContentBundleManager != nullptr)
+	{
+		World->ContentBundleManager->Deinitialize();
+	}
+}
\ No newline at end of file
diff --git a/Engine/Source/Runtime/Engine/Private/WorldPartition/ContentBundle/ContentBundleLog.cpp b/Engine/Source/Runtime/Engine/Private/WorldPartition/ContentBundle/ContentBundleLog.cpp
new file mode 100644
--- /dev/null
+++ b/Engine/Source/Runtime/Engine/Private/WorldPartition/ContentBundle/ContentBundleLog.cpp
@@ -0,0 +1,5 @@
+// Copyright Epic Games, Inc. All Rights Reserved.
+
+#include "WorldPartition/ContentBundle/ContentBundleLog.h"
+
+DEFINE_LOG_CATEGORY(LogContentBundle);
\ No newline at end of file
diff --git a/Engine/Source/Runtime/Engine/Private/WorldPartition/ContentBundle/ContentBundleWorldSubsystem.cpp b/Engine/Source/Runtime/Engine/Private/WorldPartition/ContentBundle/ContentBundleWorldSubsystem.cpp
new file mode 100644
--- /dev/null
+++ b/Engine/Source/Runtime/Engine/Private/WorldPartition/ContentBundle/ContentBundleWorldSubsystem.cpp
@@ -0,0 +1,150 @@
+// Copyright Epic Games, Inc. All Rights Reserved.
+#include "WorldPartition/ContentBundle/ContentBundleWorldSubsystem.h"
+
+#include "WorldPartition/WorldPartition.h"
+#include "WorldPartition/ContentBundle/ContentBundle.h"
+#include "WorldPartition/ContentBundle/ContentBundleDescriptor.h"
+#include "WorldPartition/ContentBundle/ContentBundleEngineSubsystem.h"
+#include "WorldPartition/ContentBundle/ContentBundleClient.h"
+#include "WorldPartition/WorldPartition.h"
+#include "WorldPartition/WorldPartitionSubsystem.h"
+#include "Engine/World.h"
+
+
+#if WITH_EDITOR
+#include "WorldPartition/ContentBundle/ContentBundleEditor.h"
+#endif
+
+UContentBundleManager::UContentBundleManager()
+{
+
+}
+
+void UContentBundleManager::Initialize()
+{
+#if WITH_EDITOR
+	if (!GetWorld()->IsGameWorld())
+	{
+		PIEDuplicateHelper = NewObject<UContentBundleDuplicateForPIEHelper>(this);
+	}
+#endif
+
+	GetWorld()->OnWorldPartitionInitialized().AddUObject(this, &UContentBundleManager::OnWorldPartitionInitialized);
+	GetWorld()->OnWorldPartitionUninitialized().AddUObject(this, &UContentBundleManager::OnWorldPartitionUninitialized);
+}
+
+void UContentBundleManager::Deinitialize()
+{
+	GetWorld()->OnWorldPartitionInitialized().RemoveAll(this);
+	GetWorld()->OnWorldPartitionUninitialized().RemoveAll(this);
+
+	for (TUniquePtr<FContentBundleContainer>& ContentBundleContainer : ContentBundleContainers)
+	{
+		ContentBundleContainer->Deinitialize();
+	}
+
+#if WITH_EDITOR
+	PIEDuplicateHelper = nullptr;
+#endif
+}
+
+void UContentBundleManager::AddReferencedObjects(UObject* InThis, FReferenceCollector& Collector)
+{
+	UContentBundleManager* Subsystem = CastChecked<UContentBundleManager>(InThis);
+	for (TUniquePtr<FContentBundleContainer>& ContentBundleContainer : Subsystem->ContentBundleContainers)
+	{
+		ContentBundleContainer->AddReferencedObjects(Collector);
+	}
+}
+
+#if WITH_EDITOR
+
+bool UContentBundleManager::GetEditorContentBundle(TArray<TSharedPtr<FContentBundleEditor>>& OutContentBundles)
+{
+	for (TUniquePtr<FContentBundleContainer>& ContentBundleContainer : ContentBundleContainers)
+	{
+		OutContentBundles.Append(ContentBundleContainer->GetEditorContentBundles());
+	}
+
+	return !OutContentBundles.IsEmpty();
+}
+
+TSharedPtr<FContentBundleEditor> UContentBundleManager::GetEditorContentBundle(const UContentBundleDescriptor* Descriptor, const UWorld* ContentBundleWorld) const
+{
+	for (const TUniquePtr<FContentBundleContainer>& ContentBundleContainer : ContentBundleContainers)
+	{
+		if (ContentBundleContainer->GetInjectedWorld() == ContentBundleWorld)
+		{
+			return ContentBundleContainer->GetEditorContentBundle(Descriptor);
+		}
+	}
+
+	return nullptr;
+}
+
+#endif
+
+uint32 UContentBundleManager::GetContentBundleContainerIndex(const UWorld* InjectedWorld)
+{
+	return ContentBundleContainers.IndexOfByPredicate([InjectedWorld](const TUniquePtr<FContentBundleContainer>& ContentBundleContainer)
+		{
+			return ContentBundleContainer->GetInjectedWorld() == InjectedWorld;
+		});
+}
+
+TUniquePtr<FContentBundleContainer>* UContentBundleManager::GetContentBundleContainer(const UWorld* InjectedWorld)
+{
+	uint32 ContainerIndex = GetContentBundleContainerIndex(InjectedWorld);
+	if (ContainerIndex != INDEX_NONE)
+	{
+		return &ContentBundleContainers[ContainerIndex];
+	}
+
+	return nullptr;
+}
+
+void UContentBundleManager::OnWorldPartitionInitialized(UWorldPartition* InWorldPartition)
+{
+	if (InWorldPartition->GetWorld() != InWorldPartition->GetTypedOuter<UWorld>())
+	{
+		return;
+	}
+
+	check(GetContentBundleContainer(InWorldPartition->GetTypedOuter<UWorld>()) == nullptr);
+	TUniquePtr<FContentBundleContainer>& ContentBundleContainer = ContentBundleContainers.Emplace_GetRef(MakeUnique<FContentBundleContainer>(InWorldPartition->GetTypedOuter<UWorld>()));
+	ContentBundleContainer->Initialize();
+}
+
+void UContentBundleManager::OnWorldPartitionUninitialized(UWorldPartition* InWorldPartition)
+{
+	if (InWorldPartition->GetWorld() != InWorldPartition->GetTypedOuter<UWorld>())
+	{
+		return;
+	}
+
+	uint32 ContainerIndex = GetContentBundleContainerIndex(InWorldPartition->GetTypedOuter<UWorld>());
+	check(ContainerIndex != INDEX_NONE);
+	ContentBundleContainers[ContainerIndex]->Deinitialize();
+	ContentBundleContainers.RemoveAtSwap(ContainerIndex);
+}
+
+#if WITH_EDITOR
+
+bool UContentBundleDuplicateForPIEHelper::StoreContentBundleStreamingObect(const FContentBundleEditor& ContentBundleEditor, URuntimeHashExternalStreamingObjectBase* StreamingObject)
+{
+	FGuid ContentBundleUid = ContentBundleEditor.GetDescriptor()->GetGuid();
+	if (!StreamingObjects.Contains(ContentBundleUid))
+	{
+		StreamingObjects.Add(ContentBundleUid, StreamingObject);
+	}
+	
+	return true;
+}
+
+URuntimeHashExternalStreamingObjectBase* UContentBundleDuplicateForPIEHelper::RetrieveContentBundleStreamingObject(const FContentBundle& ContentBundle) const
+{
+	const TObjectPtr<URuntimeHashExternalStreamingObjectBase>* ContentBundleStreamingObject = StreamingObjects.Find(ContentBundle.GetDescriptor()->GetGuid());
+	return ContentBundleStreamingObject != nullptr ? *ContentBundleStreamingObject : nullptr;
+}
+
+#endif
\ No newline at end of file
diff --git a/Engine/Source/Runtime/Engine/Private/WorldPartition/ContentBundle/WorldDataLayerReference.cpp b/Engine/Source/Runtime/Engine/Private/WorldPartition/ContentBundle/WorldDataLayerReference.cpp
new file mode 100644
--- /dev/null
+++ b/Engine/Source/Runtime/Engine/Private/WorldPartition/ContentBundle/WorldDataLayerReference.cpp
@@ -0,0 +1,123 @@
+// Copyright Epic Games, Inc. All Rights Reserved.
+#include "WorldPartition/ContentBundle/WorldDataLayerReference.h"
+
+#if WITH_EDITOR
+
+#include "Engine/World.h"
+#include "WorldPartition/ActorDescContainer.h"
+#include "WorldPartition/DataLayer/WorldDataLayers.h"
+#include "WorldPartition/WorldPartition.h"
+
+FWorldDataLayersReference::FWorldDataLayersReference()
+{
+	WorldDataLayersVariant.Set<AWorldDataLayers*>(nullptr);
+}
+
+FWorldDataLayersReference::~FWorldDataLayersReference()
+{
+	Reset();
+}
+
+FWorldDataLayersReference::FWorldDataLayersReference(UActorDescContainer* Container, FName WorldDataLayerName)
+{
+	TrySetReference(Container, WorldDataLayerName);
+}
+
+FWorldDataLayersReference::FWorldDataLayersReference(const FActorSpawnParameters& SpawnParameters)
+{
+	AWorldDataLayers* WorldDataLayers = AWorldDataLayers::Create(SpawnParameters);
+	WorldDataLayers->SetActorLabel(SpawnParameters.Name.ToString());
+	WorldDataLayersVariant.Set<AWorldDataLayers*>(WorldDataLayers);
+}
+
+FWorldDataLayersReference::FWorldDataLayersReference(FWorldDataLayersReference&& Other)
+	:WorldDataLayersVariant(MoveTemp(Other.WorldDataLayersVariant))
+{
+	if (Other.WorldDataLayersVariant.IsType<AWorldDataLayers*>())
+	{
+		Other.WorldDataLayersVariant.Set<AWorldDataLayers*>(nullptr);
+	}
+}
+
+FWorldDataLayersReference& FWorldDataLayersReference::operator=(FWorldDataLayersReference&& Other)
+{
+	if (this != &Other)
+	{
+		WorldDataLayersVariant = MoveTemp(Other.WorldDataLayersVariant);
+
+		if (Other.WorldDataLayersVariant.IsType<AWorldDataLayers*>())
+		{
+			Other.WorldDataLayersVariant.Set<AWorldDataLayers*>(nullptr);
+		}
+	}
+	
+	return *this;
+}
+
+const AWorldDataLayers* FWorldDataLayersReference::Get() const
+{
+	if (WorldDataLayersVariant.IsType<FWorldPartitionReference>())
+	{
+		const FWorldPartitionReference& WorldDataLayerReference = WorldDataLayersVariant.Get<FWorldPartitionReference>();
+		if (WorldDataLayerReference.IsValid())
+		{
+			return CastChecked<AWorldDataLayers>(WorldDataLayerReference->GetActor());
+		}
+		return nullptr;
+	}
+
+	return WorldDataLayersVariant.Get<AWorldDataLayers*>();
+}
+
+bool FWorldDataLayersReference::TrySetReference(UActorDescContainer* Container, FName WorldDataLayerName)
+{
+	for (UActorDescContainer::TIterator<> ActorDescIterator(Container); ActorDescIterator; ++ActorDescIterator)
+	{
+		if (ActorDescIterator->GetActorNativeClass()->IsChildOf<AWorldDataLayers>())
+		{
+			if (ActorDescIterator->GetActorName() == WorldDataLayerName)
+			{
+				WorldDataLayersVariant.Emplace<FWorldPartitionReference>(Container, ActorDescIterator->GetGuid());
+				return true;
+			}
+		}
+	}
+
+	return false;
+}
+
+void FWorldDataLayersReference::Reset()
+{
+	// An Actor Descriptor might have been created since we created the reference.
+	// Try resolving the reference. If it still does note resolve, delete the actor.
+	if (WorldDataLayersVariant.IsType<AWorldDataLayers*>())
+	{
+		AWorldDataLayers* WorldDataLayers = WorldDataLayersVariant.Get<AWorldDataLayers*>();
+		if (WorldDataLayers != nullptr)
+		{
+			if (UWorld* World = WorldDataLayers->GetWorld()) // May be null during Garbage collection
+			{
+				UWorldPartition* WorldPartition = World->GetWorldPartition();
+				FWorldPartitionReference ResolvedReference(WorldPartition, WorldDataLayers->GetActorGuid());
+				if (ResolvedReference.IsValid())
+				{
+					check(WorldDataLayers == ResolvedReference->GetActor());
+					WorldDataLayersVariant.Emplace<FWorldPartitionReference>(ResolvedReference);
+				}
+				else
+				{
+					World->DestroyActor(WorldDataLayers);
+					WorldDataLayersVariant.Set<AWorldDataLayers*>(nullptr);
+				}
+			}
+			
+		}
+	}
+
+	if (WorldDataLayersVariant.IsType<FWorldPartitionReference>())
+	{
+		WorldDataLayersVariant.Get<FWorldPartitionReference>() = nullptr;
+	}
+}
+
+#endif // WITH_EDITOR
\ No newline at end of file
diff --git a/Engine/Source/Runtime/Engine/Public/WorldPartition/ContentBundle/ContentBundle.h b/Engine/Source/Runtime/Engine/Public/WorldPartition/ContentBundle/ContentBundle.h
new file mode 100644
--- /dev/null
+++ b/Engine/Source/Runtime/Engine/Public/WorldPartition/ContentBundle/ContentBundle.h
@@ -0,0 +1,41 @@
+// Copyright Epic Games, Inc. All Rights Reserved.
+#pragma once
+
+#include "CoreMinimal.h"
+#include "UObject/ObjectMacros.h"
+#include "UObject/Object.h"
+#include "WorldPartition/ContentBundle/ContentBundleStatus.h"
+#include "WorldPartition/ContentBundle/ContentBundleBase.h"
+
+class UContentBundleDescriptor;
+class FContentBundleClient;
+class IWorldPartitionCookPackageContext;
+class URuntimeHashExternalStreamingObjectBase;
+
+class ENGINE_API FContentBundle : public FContentBundleBase
+{
+public:
+	FContentBundle(TSharedPtr<FContentBundleClient>& InClient, UWorld* InWorld);
+
+	//~ Begin IContentBundle Interface
+	virtual void AddReferencedObjects(FReferenceCollector& Collector) override;
+
+	virtual bool IsValid() const override;
+	//~ End IContentBundle Interface
+
+protected:
+	//~ Begin IContentBundle Interface
+	virtual void DoInitialize() override;
+	virtual void DoUninitialize() override;
+	virtual void DoInjectContent() override;
+	virtual void DoRemoveContent() override;
+	//~ End IContentBundle Interface
+	
+private:
+#if WITH_EDITOR
+	void InitializeForPIE();
+#endif
+
+	UPackage* ExternalStreamingObjectPackage;
+	URuntimeHashExternalStreamingObjectBase* ExternalStreamingObject;
+};
\ No newline at end of file
diff --git a/Engine/Source/Runtime/Engine/Public/WorldPartition/ContentBundle/ContentBundleBase.h b/Engine/Source/Runtime/Engine/Public/WorldPartition/ContentBundle/ContentBundleBase.h
new file mode 100644
--- /dev/null
+++ b/Engine/Source/Runtime/Engine/Public/WorldPartition/ContentBundle/ContentBundleBase.h
@@ -0,0 +1,46 @@
+// Copyright Epic Games, Inc. All Rights Reserved.
+#pragma once
+
+#include "CoreMinimal.h"
+#include "WorldPartition/ContentBundle/ContentBundleStatus.h"
+
+class UContentBundleDescriptor;
+class FContentBundleClient;
+class IWorldPartitionCookPackageContext;
+
+class ENGINE_API FContentBundleBase
+{
+public:
+	FContentBundleBase(TSharedPtr<FContentBundleClient>& InClient, UWorld* InWorld);
+	virtual ~FContentBundleBase();
+
+	void Initialize();
+	void Uninitialize();
+
+	void InjectContent();
+	void RemoveContent();
+
+	virtual void AddReferencedObjects(FReferenceCollector& Collector);
+
+	virtual bool IsValid() const = 0;
+
+	UWorld* GetInjectedWorld() const;
+	const FString& GetDisplayName() const;
+	const TWeakPtr<FContentBundleClient>& GetClient() const;
+	const UContentBundleDescriptor* GetDescriptor() const;
+	EContentBundleStatus GetStatus() const { return Status; }
+
+protected:
+	void SetStatus(EContentBundleStatus NewStatus);
+
+	virtual void DoInitialize() = 0;
+	virtual void DoUninitialize() = 0;
+	virtual void DoInjectContent() = 0;
+	virtual void DoRemoveContent() = 0;
+
+private:
+	TWeakPtr<FContentBundleClient> Client;
+	TWeakObjectPtr<UWorld> InjectedWorld;
+	const UContentBundleDescriptor* Descriptor;
+	EContentBundleStatus Status;
+};
\ No newline at end of file
diff --git a/Engine/Source/Runtime/Engine/Public/WorldPartition/ContentBundle/ContentBundleClient.h b/Engine/Source/Runtime/Engine/Public/WorldPartition/ContentBundle/ContentBundleClient.h
new file mode 100644
--- /dev/null
+++ b/Engine/Source/Runtime/Engine/Public/WorldPartition/ContentBundle/ContentBundleClient.h
@@ -0,0 +1,51 @@
+// Copyright Epic Games, Inc. All Rights Reserved.
+#pragma once
+
+#include "CoreMinimal.h"
+#include "UObject/ObjectMacros.h"
+#include "UObject/WeakObjectPtr.h"
+
+class UContentBundleDescriptor;
+
+UENUM()
+enum class EContentBundleClientState
+{
+	Unregistered,
+	Registered,
+	ContentInjectionRequested,
+	ContentInjected,
+	ContentRemovalRequested,
+	
+	// Failed state
+	RegistrationFailed,
+	ContentInjectionFailed,
+};
+
+class ENGINE_API FContentBundleClient
+{
+	friend class UContentBundleEngineSubsystem;
+
+public:
+	static TSharedPtr<FContentBundleClient> CreateClient(const UContentBundleDescriptor* InContentBundleDescriptor, FString const& InDisplayName);
+
+	FContentBundleClient(const UContentBundleDescriptor* InContentBundleDescriptor, FString const& InDisplayName);
+	virtual ~FContentBundleClient() = default;
+
+	const UContentBundleDescriptor* GetDescriptor() const { return ContentBundleDescriptor.Get(); }
+
+	void RequestContentInjection();
+	void RequestRemoveContent();
+	
+	void RequestUnregister();
+
+	EContentBundleClientState GetState() const { return ContentInjectionState; }
+
+	FString const& GetDisplayName() const { return DisplayName; }
+
+private:
+	TWeakObjectPtr<const UContentBundleDescriptor> ContentBundleDescriptor;
+
+	FString DisplayName;
+
+	EContentBundleClientState ContentInjectionState;
+};
\ No newline at end of file
diff --git a/Engine/Source/Runtime/Engine/Public/WorldPartition/ContentBundle/ContentBundleContainer.h b/Engine/Source/Runtime/Engine/Public/WorldPartition/ContentBundle/ContentBundleContainer.h
new file mode 100644
--- /dev/null
+++ b/Engine/Source/Runtime/Engine/Public/WorldPartition/ContentBundle/ContentBundleContainer.h
@@ -0,0 +1,98 @@
+// Copyright Epic Games, Inc. All Rights Reserved.
+#pragma once
+
+#include "CoreMinimal.h"
+#include "UObject/WeakObjectPtr.h"
+
+#if WITH_EDITOR
+#include "Misc/TVariant.h"
+#endif
+
+class FContentBundleBase;
+class FContentBundle;
+class FContentBundleEditor;
+class FContentBundleClient;
+class IWorldPartitionCookPackageContext;
+class UContentBundleDescriptor;
+class FReferenceCollector;
+
+/*
+* Stores the content bundles for a given world.
+* In Editor builds, the content bundles can be store as 2 different types:
+*	1. Editor and Non-Game Worlds:
+*		- The content bundles are stored as FContentBundleEditor, which support editing of content bundles.
+*		- They are stored in TSharedPtr to allow the UI to keep TWeakPtr.
+*	2. Editor and Game Worlds (PIE)
+*		- The content bundles are stored as FContentBundles, which support injections of streaming objects and cells in a runtime hash.
+*		- They are stored in TUniquePtr.
+* 
+* A TVariant is use to express that behavior in editor builds. In Editor and Non-Game Worlds the variant will contain FContentBundleEditors. In Game-Worlds it will contain FContentBundles.
+* 
+* In Non-Editor builds, the content bundles are always stored as FContentBundles.
+*/
+class FContentBundleContainer
+{
+public:
+	FContentBundleContainer(UWorld* WorldToInjectIn);
+	~FContentBundleContainer();
+
+	UWorld* GetInjectedWorld() const;
+
+	void Initialize();
+	void Deinitialize();
+
+	void AddReferencedObjects(FReferenceCollector& Collector);
+
+	uint32 GetNumContentBundles() const;
+
+	const TArray<TUniquePtr<FContentBundle>>& GetGameContentBundles() const;
+	TArray<TUniquePtr<FContentBundle>>& GetGameContentBundles();
+
+#if WITH_EDITOR
+	TSharedPtr<FContentBundleEditor> GetEditorContentBundle(const UContentBundleDescriptor* Descriptor) const;
+
+	const TArray<TSharedPtr<FContentBundleEditor>>& GetEditorContentBundles() const;
+	TArray<TSharedPtr<FContentBundleEditor>>& GetEditorContentBundles();
+
+	bool UseEditorContentBundle() const;
+#endif
+
+	void ForEachContentBundle(TFunctionRef<void(FContentBundleBase*)> Func) const;
+	void ForEachContentBundleBreakable(TFunctionRef<bool(FContentBundleBase*)> Func) const;
+
+private:
+	bool ShouldInjectClientContent(const TSharedPtr<FContentBundleClient>& ContentBundleClient) const;
+
+	FContentBundleBase* GetContentBundle(const FContentBundleClient& ContentBundleClient);
+
+	void RegisterContentBundleClientEvents();
+	void UnregisterContentBundleClientEvents();
+
+	void OnContentBundleClientRegistered(TSharedPtr<FContentBundleClient>& ContentBundleClient);
+	void OnContentBundleClientUnregistered(FContentBundleClient& ContentBundleClient);
+	void OnContentBundleClientContentInjectionRequested(FContentBundleClient& ContentBundleClient);
+	void OnContentBundleClientContentRemovalRequested(FContentBundleClient& ContentBundleClient);
+
+	FContentBundleBase& InitializeContentBundle(TSharedPtr<FContentBundleClient>& ContentBundleClient);
+	void DeinitializeContentBundle(FContentBundleBase& ContentBundle);
+
+	void InitializeContentBundlesForegisteredClients();
+	void DeinitializeContentBundles();
+
+	void InjectContentBundle(FContentBundleBase& ContentBundle);
+	void RemoveContentBundle(FContentBundleBase& ContentBundle);
+
+#if WITH_EDITOR
+	void OnPreGenerateStreaming();
+#endif
+
+	using ContentBundleGameArray = TArray<TUniquePtr<FContentBundle>>;
+#if WITH_EDITOR
+	using ContentBundleEditorArray = TArray<TSharedPtr<FContentBundleEditor>>;
+	TVariant<ContentBundleEditorArray, ContentBundleGameArray> ContentBundlesVariant;
+#else
+	ContentBundleGameArray ContentBundles;
+#endif
+
+	TWeakObjectPtr<UWorld> InjectedWorld;
+};
\ No newline at end of file
diff --git a/Engine/Source/Runtime/Engine/Public/WorldPartition/ContentBundle/ContentBundleDescriptor.h b/Engine/Source/Runtime/Engine/Public/WorldPartition/ContentBundle/ContentBundleDescriptor.h
new file mode 100644
--- /dev/null
+++ b/Engine/Source/Runtime/Engine/Public/WorldPartition/ContentBundle/ContentBundleDescriptor.h
@@ -0,0 +1,39 @@
+// Copyright Epic Games, Inc. All Rights Reserved.
+#pragma once
+
+#include "CoreMinimal.h"
+#include "UObject/ObjectMacros.h"
+#include "UObject/Object.h"
+
+#include "ContentBundleDescriptor.generated.h"
+
+class UActorDescContainer;
+class AWorldDataLayers;
+
+UCLASS()
+class ENGINE_API UContentBundleDescriptor : public UObject
+{
+	GENERATED_UCLASS_BODY()
+
+public:
+	const FString& GetDisplayName() const;
+	const FString& GetPackageRoot() const { return PackageRoot; }
+	const FGuid& GetGuid() const { return Guid; }
+
+#if WITH_EDITOR
+	void SetDisplayName(const FString& InName) { DisplayName = InName; }
+	void SetPackageRoot(const FString& InPackageRoot) { PackageRoot = InPackageRoot; }
+#endif
+
+	bool IsValid() const;
+
+private:
+	UPROPERTY(EditAnywhere, Category = BaseInformation)
+	FString DisplayName;
+
+	UPROPERTY(VisibleAnywhere, Category = BaseInformation, AdvancedDisplay)
+	FGuid Guid;
+
+	UPROPERTY(VisibleAnywhere, Category = BaseInformation, AdvancedDisplay)
+	FString PackageRoot;
+};
\ No newline at end of file
diff --git a/Engine/Source/Runtime/Engine/Public/WorldPartition/ContentBundle/ContentBundleEditor.h b/Engine/Source/Runtime/Engine/Public/WorldPartition/ContentBundle/ContentBundleEditor.h
new file mode 100644
--- /dev/null
+++ b/Engine/Source/Runtime/Engine/Public/WorldPartition/ContentBundle/ContentBundleEditor.h
@@ -0,0 +1,124 @@
+// Copyright Epic Games, Inc. All Rights Reserved.
+#pragma once
+
+#if WITH_EDITOR
+
+#include "CoreMinimal.h"
+#include "WorldPartition/ContentBundle/WorldDataLayerReference.h"
+#include "WorldPartition/ContentBundle/ContentBundleBase.h"
+
+#include "ContentBundleEditor.generated.h"
+
+class URuntimeHashExternalStreamingObjectBase;
+class UWorldPartitionRuntimeCell;
+
+class ENGINE_API FContentBundleEditor : public FContentBundleBase
+{
+	friend class UContentBundleUnsavedActorMonitor;
+
+public:
+	FContentBundleEditor(TSharedPtr<FContentBundleClient>& InClient, UWorld* InWorld);
+
+	//~ Begin IContentBundle Interface
+	virtual void AddReferencedObjects(FReferenceCollector& Collector) override;
+
+	virtual bool IsValid() const override;
+	//~ End IContentBundle Interface
+
+	const FGuid& GetGuid() const { return Guid; }
+
+	void StartEditing();
+	void StopEditing();
+	bool IsBeingEdited() const { return bIsBeingEdited; }
+
+	void InjectBaseContent();
+	void RemoveBaseContent();
+
+	bool AddActor(AActor* InActor);
+	bool ContainsActor(const AActor* InActor) const;
+	bool GetActors(TArray<AActor*>& Actors);
+	bool HasUserPlacedActors() const;
+	uint32 GetActorCount() const;
+	uint32 GetUnsavedActorAcount() const;
+
+	void ReferenceAllActors();
+	void UnreferenceAllActors();
+
+	void GenerateStreaming();
+
+	URuntimeHashExternalStreamingObjectBase* GetStreamingObject() const { return ExternalStreamingObject; }
+
+protected:
+	//~ Begin IContentBundle Interface
+	virtual void DoInitialize() override;
+	virtual void DoUninitialize() override;
+	virtual void DoInjectContent() override;
+	virtual void DoRemoveContent() override;
+	//~ End IContentBundle Interface
+
+private:
+	void OnUnsavedActorDeleted(AActor* Actor);
+
+	void BroadcastChanged();
+	bool BuildContentBundleContainerPackagePath(FString& ContainerPackagePath) const;
+	UPackage* CreateActorPackage(const FName& ActorName) const;
+	FName BuildWorlDataLayersName() const;
+
+	void RegisterDelegates();
+	void UnregisterDelegates();
+	
+	void OnActorDescAdded(FWorldPartitionActorDesc* ActorDesc);
+	void OnActorDescRemoved(FWorldPartitionActorDesc* ActorDesc);
+
+	class UContentBundleUnsavedActorMonitor* UnsavedActorMonitor;
+
+	TWeakObjectPtr<UActorDescContainer> ActorDescContainer;
+	FWorldDataLayersReference WorldDataLayersActorReference;
+
+	TArray<FWorldPartitionReference> ForceLoadedActors;
+
+	TSet<const UWorldPartitionRuntimeCell*> ContentBundleCells;
+	TMap<uint32, const UWorldPartitionRuntimeCell*> CookPackageIdsToCell;
+	URuntimeHashExternalStreamingObjectBase* ExternalStreamingObject;
+
+	FGuid Guid;
+
+	bool bIsBeingEdited;
+};
+
+UCLASS()
+class UContentBundleUnsavedActorMonitor : public UObject
+{
+	GENERATED_BODY()
+
+public:
+	~UContentBundleUnsavedActorMonitor();
+
+	void Initialize(FContentBundleEditor& InContentBundleEditor);
+	void Uninitialize();
+
+	void StartListenOnActorEvents();
+	void StopListeningOnActorEvents();
+
+	void MonitorActor(AActor* InActor);
+	bool StopMonitoringActor(AActor* InActor);
+
+	uint32 GetActorCount() const { return UnsavedActors.Num(); }
+
+	bool IsMonitoringActors() const { return !UnsavedActors.IsEmpty(); }
+	bool IsMonitoring(const AActor* Actor) const { return UnsavedActors.Contains(Actor); }
+
+	const TArray<TWeakObjectPtr<AActor>>& GetUnsavedActors() const { return UnsavedActors; }
+
+private:
+	void OnActorDeleted(AActor* InActor);
+
+#if WITH_EDITORONLY_DATA
+	UPROPERTY()
+	TArray<TWeakObjectPtr<AActor>> UnsavedActors;
+#endif
+
+	FContentBundleEditor* ContentBundle;
+};
+
+#endif
\ No newline at end of file
diff --git a/Engine/Source/Runtime/Engine/Public/WorldPartition/ContentBundle/ContentBundleEditorSubsystemInterface.h b/Engine/Source/Runtime/Engine/Public/WorldPartition/ContentBundle/ContentBundleEditorSubsystemInterface.h
new file mode 100644
--- /dev/null
+++ b/Engine/Source/Runtime/Engine/Public/WorldPartition/ContentBundle/ContentBundleEditorSubsystemInterface.h
@@ -0,0 +1,22 @@
+// Copyright Epic Games, Inc. All Rights Reserved.
+#pragma once
+
+#include "CoreMinimal.h"
+
+class FContentBundleEditor;
+
+#if WITH_EDITOR
+class ENGINE_API IContentBundleEditorSubsystemInterface
+{
+public:
+	static IContentBundleEditorSubsystemInterface* Get() { return Instance; }
+
+	virtual void NotifyContentBundleAdded(const FContentBundleEditor* ContentBundle) = 0;
+	virtual void NotifyContentBundleRemoved(const FContentBundleEditor* ContentBundle) = 0;
+	virtual void NotifyContentBundleChanged(const FContentBundleEditor* ContentBundle) = 0;
+
+protected:
+	static void SetInstance(IContentBundleEditorSubsystemInterface* InInstance);
+	static IContentBundleEditorSubsystemInterface* Instance;
+};
+#endif
\ No newline at end of file
diff --git a/Engine/Source/Runtime/Engine/Public/WorldPartition/ContentBundle/ContentBundleEngineSubsystem.h b/Engine/Source/Runtime/Engine/Public/WorldPartition/ContentBundle/ContentBundleEngineSubsystem.h
new file mode 100644
--- /dev/null
+++ b/Engine/Source/Runtime/Engine/Public/WorldPartition/ContentBundle/ContentBundleEngineSubsystem.h
@@ -0,0 +1,51 @@
+// Copyright Epic Games, Inc. All Rights Reserved.
+#pragma once
+
+#include "CoreMinimal.h"
+#include "UObject/ObjectMacros.h"
+#include "UObject/Object.h"
+#include "Engine/World.h"
+#include "Subsystems/EngineSubsystem.h"
+
+#include "ContentBundleEngineSubsystem.generated.h"
+
+class FContentBundleClient;
+class UContentBundleDescriptor;
+class URuntimeHashExternalStreamingObjectBase;
+
+UCLASS()
+class ENGINE_API UContentBundleEngineSubsystem : public UEngineSubsystem
+{
+	GENERATED_BODY()
+
+public:
+	virtual void Initialize(FSubsystemCollectionBase& Collection) override;
+	virtual void Deinitialize() override;
+
+	TSharedPtr<FContentBundleClient> RegisterContentBundle(const UContentBundleDescriptor* Descriptor, const FString& ClientDisplayName);
+	void UnregisterContentBundle(FContentBundleClient& Client);
+
+	void RequestContentInjection(FContentBundleClient& Client);
+	void RequestContentRemoval(FContentBundleClient& Client);
+
+	DECLARE_MULTICAST_DELEGATE_OneParam(FContentBundleRegisterDelegate, TSharedPtr<FContentBundleClient>&);
+	FContentBundleRegisterDelegate OnContentBundleClientRegistered;
+
+	DECLARE_MULTICAST_DELEGATE_OneParam(FContentBundleEventDelegate, FContentBundleClient&);
+	FContentBundleEventDelegate OnContentBundleClientUnregistered;
+	FContentBundleEventDelegate OnContentBundleClientRequestedContentInjection;
+	FContentBundleEventDelegate OnContentBundleClientRequestedContentRemoval;
+
+	TArrayView<TSharedPtr<FContentBundleClient>> GetContentBundleClients() { return MakeArrayView(ContentBundleClients); }
+
+	static UContentBundleEngineSubsystem* Get();
+
+private:
+	void OnPreWorldInitialization(UWorld* World, const UWorld::InitializationValues IVS);
+	void OnWorldPostCleanup(UWorld* World, bool bSessionEnded, bool bCleanupResources);
+
+	TSharedPtr<FContentBundleClient>* FindRegisteredClient(const UContentBundleDescriptor* Descriptor);
+	TSharedPtr<FContentBundleClient>* FindRegisteredClient(FContentBundleClient& InClient);
+
+	TArray<TSharedPtr<FContentBundleClient>> ContentBundleClients;
+};
\ No newline at end of file
diff --git a/Engine/Source/Runtime/Engine/Public/WorldPartition/ContentBundle/ContentBundleLog.h b/Engine/Source/Runtime/Engine/Public/WorldPartition/ContentBundle/ContentBundleLog.h
new file mode 100644
--- /dev/null
+++ b/Engine/Source/Runtime/Engine/Public/WorldPartition/ContentBundle/ContentBundleLog.h
@@ -0,0 +1,6 @@
+// Copyright Epic Games, Inc. All Rights Reserved.
+#pragma once
+
+#include "CoreMinimal.h"
+
+ENGINE_API DECLARE_LOG_CATEGORY_EXTERN(LogContentBundle, Log, All);
\ No newline at end of file
diff --git a/Engine/Source/Runtime/Engine/Public/WorldPartition/ContentBundle/ContentBundleStatus.h b/Engine/Source/Runtime/Engine/Public/WorldPartition/ContentBundle/ContentBundleStatus.h
new file mode 100644
--- /dev/null
+++ b/Engine/Source/Runtime/Engine/Public/WorldPartition/ContentBundle/ContentBundleStatus.h
@@ -0,0 +1,17 @@
+// Copyright Epic Games, Inc. All Rights Reserved.
+#pragma once
+
+#include "CoreMinimal.h"
+#include "UObject/ObjectMacros.h"
+
+#include "ContentBundleStatus.generated.h"
+
+UENUM()
+enum class EContentBundleStatus
+{
+	Registered,
+	ReadyToInject,
+	FailedToInject,
+	ContentInjected,
+	Unknown = -1
+};
\ No newline at end of file
diff --git a/Engine/Source/Runtime/Engine/Public/WorldPartition/ContentBundle/ContentBundleWorldSubsystem.h b/Engine/Source/Runtime/Engine/Public/WorldPartition/ContentBundle/ContentBundleWorldSubsystem.h
new file mode 100644
--- /dev/null
+++ b/Engine/Source/Runtime/Engine/Public/WorldPartition/ContentBundle/ContentBundleWorldSubsystem.h
@@ -0,0 +1,80 @@
+// Copyright Epic Games, Inc. All Rights Reserved.
+#pragma once
+
+#include "CoreMinimal.h"
+#include "UObject/ObjectMacros.h"
+#include "UObject/Object.h"
+
+#include "WorldPartition/ContentBundle/ContentBundleContainer.h"
+
+#include "ContentBundleWorldSubsystem.generated.h"
+
+class FContentBundle;
+class UContentBundleDescriptor;
+class FContentBundleClient;
+class FContentBundleEditor;
+class UWorldPartition;
+
+#if WITH_EDITOR
+class URuntimeHashExternalStreamingObjectBase;
+class UContentBundleDuplicateForPIEHelper;
+#endif
+
+UCLASS()
+class ENGINE_API UContentBundleManager : public UObject
+{
+	GENERATED_BODY()
+
+public:
+	UContentBundleManager();
+
+	void Initialize();
+	void Deinitialize();
+
+	static void AddReferencedObjects(UObject* InThis, FReferenceCollector& Collector);
+
+#if WITH_EDITOR
+	bool GetEditorContentBundle(TArray<TSharedPtr<FContentBundleEditor>>& OutContentBundles);
+	TSharedPtr<FContentBundleEditor> GetEditorContentBundle(const UContentBundleDescriptor* Descriptor, const UWorld* ContentBundleWorld) const;
+
+	UContentBundleDuplicateForPIEHelper* GetPIEDuplicateHelper() const { return PIEDuplicateHelper; }
+#endif
+
+private:
+	uint32 GetContentBundleContainerIndex(const UWorld* InjectedWorld);
+	TUniquePtr<FContentBundleContainer>* GetContentBundleContainer(const UWorld* InjectedWorld);
+
+	void OnWorldPartitionInitialized(UWorldPartition* WorldPartition);
+	void OnWorldPartitionUninitialized(UWorldPartition* WorldPartition);
+
+	TArray<TUniquePtr<FContentBundleContainer>> ContentBundleContainers;
+
+#if WITH_EDITORONLY_DATA
+	UPROPERTY()
+	TObjectPtr<UContentBundleDuplicateForPIEHelper> PIEDuplicateHelper;
+#endif
+};
+
+
+
+UCLASS()
+class UContentBundleDuplicateForPIEHelper : public UObject
+{
+	GENERATED_BODY()
+
+public:
+#if WITH_EDITOR
+	bool StoreContentBundleStreamingObect(const FContentBundleEditor& ContentBundleEditor, URuntimeHashExternalStreamingObjectBase* StreamingObject);
+	URuntimeHashExternalStreamingObjectBase* RetrieveContentBundleStreamingObject(const FContentBundle& ContentBundle) const;
+
+	uint32 GetStreamingObjectCount() { return StreamingObjects.Num(); }
+	void Clear() { StreamingObjects.Empty(); }
+#endif
+
+private:
+#if WITH_EDITORONLY_DATA
+	UPROPERTY()
+	TMap<FGuid, TObjectPtr<URuntimeHashExternalStreamingObjectBase>> StreamingObjects;
+#endif
+
+};
\ No newline at end of file
diff --git a/Engine/Source/Runtime/Engine/Public/WorldPartition/ContentBundle/WorldDataLayerReference.h b/Engine/Source/Runtime/Engine/Public/WorldPartition/ContentBundle/WorldDataLayerReference.h
new file mode 100644
--- /dev/null
+++ b/Engine/Source/Runtime/Engine/Public/WorldPartition/ContentBundle/WorldDataLayerReference.h
@@ -0,0 +1,40 @@
+// Copyright Epic Games, Inc. All Rights Reserved.
+#pragma once
+
+#if WITH_EDITOR
+#include "WorldPartition/WorldPartitionHandle.h"
+#include "Misc/TVariant.h"
+
+class UActorDescContainer;
+class AWorldDataLayers;
+struct FActorSpawnParameters;
+
+class ENGINE_API FWorldDataLayersReference
+{
+public:
+	FWorldDataLayersReference();
+	FWorldDataLayersReference(UActorDescContainer* Container, FName WorldDataLayerName); // Retrieve Actor
+	FWorldDataLayersReference(const FActorSpawnParameters& SpawnParameters); // Retrieve and Create Actor if Not Found, FActorSpawnParameters::Name must be set
+	FWorldDataLayersReference(FWorldDataLayersReference&& Other);
+
+	~FWorldDataLayersReference();
+
+	bool IsValid() const { return Get() != nullptr; }
+
+	AWorldDataLayers* operator->() { return const_cast<AWorldDataLayers*>(const_cast<const FWorldDataLayersReference*>(this)->operator->()); }
+	const AWorldDataLayers* operator->() const { return Get(); }
+
+	const AWorldDataLayers* Get() const;
+	AWorldDataLayers* Get() { return const_cast<AWorldDataLayers*>(const_cast<const FWorldDataLayersReference*>(this)->Get()); }
+
+	void Reset();
+
+	FWorldDataLayersReference& operator=(FWorldDataLayersReference&& Other);
+
+private:
+	bool TrySetReference(UActorDescContainer* Container, FName WorldDataLayerName);
+
+	TVariant<AWorldDataLayers*, FWorldPartitionReference> WorldDataLayersVariant;
+};
+
+#endif // WITH_EDITOR
\ No newline at end of file
