Fix Error in Linux & Shipping build

diff --git a/Engine/Source/Runtime/Engine/Private/WorldPartition/ContentBundle/ContentBundle.cpp b/Engine/Source/Runtime/Engine/Private/WorldPartition/ContentBundle/ContentBundle.cpp
--- a/Engine/Source/Runtime/Engine/Private/WorldPartition/ContentBundle/ContentBundle.cpp
+++ b/Engine/Source/Runtime/Engine/Private/WorldPartition/ContentBundle/ContentBundle.cpp
@@ -26,7 +26,8 @@ void FContentBundle::DoInitialize()
 	InitializeForPIE();
 #else
 	FString StreamingObjectPathName = FPaths::RemoveDuplicateSlashes(GetDescriptor()->GetPackageRoot() + TEXT("/_GENERATED_/") + TEXT("StreamingObject"));
-	if (ExternalStreamingObjectPackage = LoadPackage(nullptr, *StreamingObjectPathName, LOAD_None))
+	ExternalStreamingObjectPackage = LoadPackage(nullptr, *StreamingObjectPathName, LOAD_None);
+	if (ExternalStreamingObjectPackage != nullptr)
 	{
 		if (UObject* Object = StaticFindObjectFast(URuntimeHashExternalStreamingObjectBase::StaticClass(), ExternalStreamingObjectPackage, TEXT("StreamingObject")))
 		{
diff --git a/Engine/Source/Runtime/Engine/Private/WorldPartition/ContentBundle/ContentBundleEditor.cpp b/Engine/Source/Runtime/Engine/Private/WorldPartition/ContentBundle/ContentBundleEditor.cpp
--- a/Engine/Source/Runtime/Engine/Private/WorldPartition/ContentBundle/ContentBundleEditor.cpp
+++ b/Engine/Source/Runtime/Engine/Private/WorldPartition/ContentBundle/ContentBundleEditor.cpp
@@ -414,12 +414,18 @@ void FContentBundleEditor::OnUnsavedActorDeleted(AActor* Actor)
 	}
 }
 
+#endif
+
 UContentBundleUnsavedActorMonitor::~UContentBundleUnsavedActorMonitor()
 {
+#if WITH_EDITOR
 	check(UnsavedActors.IsEmpty());
 	check(ContentBundle == nullptr);
+#endif
 }
 
+#if WITH_EDITOR
+
 void UContentBundleUnsavedActorMonitor::Initialize(FContentBundleEditor& InContentBundleEditor)
 {
 	ContentBundle = &InContentBundleEditor;
diff --git a/Engine/Source/Runtime/Engine/Public/WorldPartition/ContentBundle/ContentBundleEditor.h b/Engine/Source/Runtime/Engine/Public/WorldPartition/ContentBundle/ContentBundleEditor.h
--- a/Engine/Source/Runtime/Engine/Public/WorldPartition/ContentBundle/ContentBundleEditor.h
+++ b/Engine/Source/Runtime/Engine/Public/WorldPartition/ContentBundle/ContentBundleEditor.h
@@ -1,14 +1,19 @@
 // Copyright Epic Games, Inc. All Rights Reserved.
 #pragma once
 
-#if WITH_EDITOR
-
 #include "CoreMinimal.h"
+#include "UObject/ObjectMacros.h"
+#include "UObject/Object.h"
+
+#if WITH_EDITOR
 #include "WorldPartition/ContentBundle/WorldDataLayerReference.h"
 #include "WorldPartition/ContentBundle/ContentBundleBase.h"
+#endif
 
 #include "ContentBundleEditor.generated.h"
 
+#if WITH_EDITOR
+
 class URuntimeHashExternalStreamingObjectBase;
 class UWorldPartitionRuntimeCell;
 
@@ -86,6 +91,8 @@ private:
 	bool bIsBeingEdited;
 };
 
+#endif
+
 UCLASS()
 class UContentBundleUnsavedActorMonitor : public UObject
 {
@@ -94,6 +101,7 @@ class UContentBundleUnsavedActorMonitor : public UObject
 public:
 	~UContentBundleUnsavedActorMonitor();
 
+#if WITH_EDITOR
 	void Initialize(FContentBundleEditor& InContentBundleEditor);
 	void Uninitialize();
 
@@ -109,16 +117,19 @@ public:
 	bool IsMonitoring(const AActor* Actor) const { return UnsavedActors.Contains(Actor); }
 
 	const TArray<TWeakObjectPtr<AActor>>& GetUnsavedActors() const { return UnsavedActors; }
+#endif
 
 private:
+#if WITH_EDITOR
 	void OnActorDeleted(AActor* InActor);
+#endif
 
 #if WITH_EDITORONLY_DATA
 	UPROPERTY()
-	TArray<TWeakObjectPtr<AActor>> UnsavedActors;
+		TArray<TWeakObjectPtr<AActor>> UnsavedActors;
 #endif
 
+#if WITH_EDITOR
 	FContentBundleEditor* ContentBundle;
-};
-
-#endif
\ No newline at end of file
+#endif
+};
\ No newline at end of file
