Fix non-unity build

diff --git a/Engine/Source/Runtime/Engine/Private/WorldPartition/ContentBundle/ContentBundleEditor.cpp b/Engine/Source/Runtime/Engine/Private/WorldPartition/ContentBundle/ContentBundleEditor.cpp
--- a/Engine/Source/Runtime/Engine/Private/WorldPartition/ContentBundle/ContentBundleEditor.cpp
+++ b/Engine/Source/Runtime/Engine/Private/WorldPartition/ContentBundle/ContentBundleEditor.cpp
@@ -474,6 +474,11 @@ bool UContentBundleUnsavedActorMonitor::StopMonitoringActor(AActor* InActor)
 	return false;
 }
 
+bool UContentBundleUnsavedActorMonitor::IsMonitoring(const AActor* Actor) const
+{ 
+	return UnsavedActors.Contains(Actor); 
+}
+
 void UContentBundleUnsavedActorMonitor::OnActorDeleted(AActor* InActor)
 {
 	if (StopMonitoringActor(InActor))
diff --git a/Engine/Source/Runtime/Engine/Public/WorldPartition/ContentBundle/ContentBundleEditor.h b/Engine/Source/Runtime/Engine/Public/WorldPartition/ContentBundle/ContentBundleEditor.h
--- a/Engine/Source/Runtime/Engine/Public/WorldPartition/ContentBundle/ContentBundleEditor.h
+++ b/Engine/Source/Runtime/Engine/Public/WorldPartition/ContentBundle/ContentBundleEditor.h
@@ -4,6 +4,7 @@
 #include "CoreMinimal.h"
 #include "UObject/ObjectMacros.h"
 #include "UObject/Object.h"
+#include "UObject/WeakObjectPtr.h"
 
 #if WITH_EDITOR
 #include "WorldPartition/ContentBundle/WorldDataLayerReference.h"
@@ -12,6 +13,8 @@
 
 #include "ContentBundleEditor.generated.h"
 
+class AActor;
+
 #if WITH_EDITOR
 
 class URuntimeHashExternalStreamingObjectBase;
@@ -114,7 +117,7 @@ public:
 	uint32 GetActorCount() const { return UnsavedActors.Num(); }
 
 	bool IsMonitoringActors() const { return !UnsavedActors.IsEmpty(); }
-	bool IsMonitoring(const AActor* Actor) const { return UnsavedActors.Contains(Actor); }
+	bool IsMonitoring(const AActor* Actor) const;
 
 	const TArray<TWeakObjectPtr<AActor>>& GetUnsavedActors() const { return UnsavedActors; }
 #endif
@@ -126,7 +129,7 @@ private:
 
 #if WITH_EDITORONLY_DATA
 	UPROPERTY()
-		TArray<TWeakObjectPtr<AActor>> UnsavedActors;
+	TArray<TWeakObjectPtr<AActor>> UnsavedActors;
 #endif
 
 #if WITH_EDITOR
