Fix bug where we were not blocking on shader compilation before executing a material update context.

- Depending on timing could cause issues with missing shaders where we would update the component before its shaders were complete.
- By default just make the function that submits jobs always block.
- Renable some tests, which this change should fix.

diff --git a/Engine/Plugins/Experimental/AutomationUtils/Source/AutomationUtils/Private/AutomationUtilsBlueprintLibrary.cpp b/Engine/Plugins/Experimental/AutomationUtils/Source/AutomationUtils/Private/AutomationUtilsBlueprintLibrary.cpp
--- a/Engine/Plugins/Experimental/AutomationUtils/Source/AutomationUtils/Private/AutomationUtilsBlueprintLibrary.cpp
+++ b/Engine/Plugins/Experimental/AutomationUtils/Source/AutomationUtils/Private/AutomationUtilsBlueprintLibrary.cpp
@@ -363,8 +363,6 @@ void UAutomationUtilsBlueprintLibrary::TakeGameplayAutomationScreenshot(const FS
 	if (!FPlatformProperties::RequiresCookedData())
 	{
 		UMaterialInterface::SubmitRemainingJobsForWorld(CurrentWorld);
-		//Finish Compiling all shaders and other async assets
-		FAssetCompilingManager::Get().FinishAllCompilation();
 	}
 
 	//Stream in everything
diff --git a/Engine/Plugins/MovieScene/MovieRenderPipeline/Source/MovieRenderPipelineCore/Private/MoviePipelineRendering.cpp b/Engine/Plugins/MovieScene/MovieRenderPipeline/Source/MovieRenderPipelineCore/Private/MoviePipelineRendering.cpp
--- a/Engine/Plugins/MovieScene/MovieRenderPipeline/Source/MovieRenderPipelineCore/Private/MoviePipelineRendering.cpp
+++ b/Engine/Plugins/MovieScene/MovieRenderPipeline/Source/MovieRenderPipelineCore/Private/MoviePipelineRendering.cpp
@@ -540,10 +540,6 @@ void UMoviePipeline::FlushAsyncEngineSystems()
 	// This way we will never render with the default material.
 	UMaterialInterface::SubmitRemainingJobsForWorld(GetWorld());
 
-	// Flush all assets still being compiled asynchronously.
-	// A progressbar is already in place so the user can get feedback while waiting for everything to settle.
-	FAssetCompilingManager::Get().FinishAllCompilation();
-
 	// Flush streaming managers
 	{
 		UMoviePipelineGameOverrideSetting* GameOverrideSettings = FindOrAddSettingForShot<UMoviePipelineGameOverrideSetting>(ActiveShotList[CurrentShotIndex]);
diff --git a/Engine/Source/Developer/FunctionalTesting/Private/AutomationBlueprintFunctionLibrary.cpp b/Engine/Source/Developer/FunctionalTesting/Private/AutomationBlueprintFunctionLibrary.cpp
--- a/Engine/Source/Developer/FunctionalTesting/Private/AutomationBlueprintFunctionLibrary.cpp
+++ b/Engine/Source/Developer/FunctionalTesting/Private/AutomationBlueprintFunctionLibrary.cpp
@@ -710,7 +710,6 @@ void UAutomationBlueprintFunctionLibrary::FinishLoadingBeforeScreenshot()
 	if (!FPlatformProperties::RequiresCookedData())
 	{
 		UMaterialInterface::SubmitRemainingJobsForWorld(CurrentWorld);
-		FAssetCompilingManager::Get().FinishAllCompilation();
 		FModuleManager::GetModuleChecked<IAutomationControllerModule>("AutomationController").GetAutomationController()->ResetAutomationTestTimeout(TEXT("shader compilation"));
 	}
 
diff --git a/Engine/Source/Editor/Blutility/Private/AsyncCaptureScene.cpp b/Engine/Source/Editor/Blutility/Private/AsyncCaptureScene.cpp
--- a/Engine/Source/Editor/Blutility/Private/AsyncCaptureScene.cpp
+++ b/Engine/Source/Editor/Blutility/Private/AsyncCaptureScene.cpp
@@ -134,7 +134,6 @@ void UAsyncCaptureScene::FinishLoadingBeforeScreenshot()
 	if (!FPlatformProperties::RequiresCookedData())
 	{
 		UMaterialInterface::SubmitRemainingJobsForWorld(CurrentWorld);
-		FAssetCompilingManager::Get().FinishAllCompilation();
 		IAutomationControllerModule& AutomationController = IAutomationControllerModule::Get();
 		AutomationController.GetAutomationController()->ResetAutomationTestTimeout(TEXT("shader compilation"));
 	}
diff --git a/Engine/Source/Runtime/Engine/Classes/Materials/MaterialInterface.h b/Engine/Source/Runtime/Engine/Classes/Materials/MaterialInterface.h
--- a/Engine/Source/Runtime/Engine/Classes/Materials/MaterialInterface.h
+++ b/Engine/Source/Runtime/Engine/Classes/Materials/MaterialInterface.h
@@ -842,24 +842,16 @@ public:
 	/**
 	 * @brief Submits shaders to be compiled for all the materials in the world.
 	 *
-	 * This function will submit any remaining shaders to be compiled for all the materials in the passed in world.  By default
-	 * these shader compilation jobs will be compiled in the background so if you need the results immediately you can call
-	 * FinishAllCompilation() to block on the results.
+	 * This function will submit any remaining shaders to be compiled for all the materials in the passed in world.
 	 *
 	 * If the world is nullptr this will submit remaining shaders to be compiled for all the loaded materials.
 	 *
-	 * @code
-	 * GShaderCompilingManager->SubmitRemainingJobsForWorld(World);
-	 * GShaderCompilingManager->FinishAllCompilation();
-	 * @endcode
-	 *
 	 * @param World Only compile shaders for materials that are used on primitives in this world.
-	 * @param CompileMode Controls whether or not we block on the shader compile results.
 	 *
 	 * @note This will only submit shader compile jobs for missing shaders on each material.  Calling this multiple times on the same world
 	 * will result in a no-op.
 	 */
-	ENGINE_API static void SubmitRemainingJobsForWorld(UWorld* World, EMaterialShaderPrecompileMode CompileMode = EMaterialShaderPrecompileMode::Default);
+	ENGINE_API static void SubmitRemainingJobsForWorld(UWorld* World);
 
 	/**
 	 * Re-caches uniform expressions for this material interface                   
diff --git a/Engine/Source/Runtime/Engine/Private/Materials/MaterialInterface.cpp b/Engine/Source/Runtime/Engine/Private/Materials/MaterialInterface.cpp
--- a/Engine/Source/Runtime/Engine/Private/Materials/MaterialInterface.cpp
+++ b/Engine/Source/Runtime/Engine/Private/Materials/MaterialInterface.cpp
@@ -536,7 +536,7 @@ void UMaterialInterface::RecacheAllMaterialUniformExpressions(bool bRecreateUnif
 	}
 }
 
-void UMaterialInterface::SubmitRemainingJobsForWorld(UWorld* World, EMaterialShaderPrecompileMode CompileMode)
+void UMaterialInterface::SubmitRemainingJobsForWorld(UWorld* World)
 {
 	TRACE_CPUPROFILER_EVENT_SCOPE(UMaterialInterface::SubmitRemainingJobsForWorld);
 
@@ -583,8 +583,11 @@ void UMaterialInterface::SubmitRemainingJobsForWorld(UWorld* World, EMaterialSha
 			// This is needed because CacheShaders blindly recreates uniform buffers
 			// which can only be done if the draw command is going to be re-cached.
 			UpdateContext.AddMaterialInterface(Material);
-			Material->CacheShaders(CompileMode);
+			Material->CacheShaders();
 		}
+
+		// Block on all compilation.
+		FAssetCompilingManager::Get().FinishAllCompilation();
 	}
 }
 
