World Partition: don't load last loaded regions when streaming is disabled.

#rnx

diff --git a/Engine/Source/Runtime/Engine/Private/WorldPartition/WorldPartition.cpp b/Engine/Source/Runtime/Engine/Private/WorldPartition/WorldPartition.cpp
--- a/Engine/Source/Runtime/Engine/Private/WorldPartition/WorldPartition.cpp
+++ b/Engine/Source/Runtime/Engine/Private/WorldPartition/WorldPartition.cpp
@@ -550,7 +550,7 @@ void UWorldPartition::Initialize(UWorld* InWorld, const FTransform& InTransform)
 
 		// Load more cells depending on the user's settings
 		// Skipped when running from a commandlet and for subpartitions
-		if (IsMainWorldPartition() && !IsRunningCommandlet())
+		if (IsMainWorldPartition() && IsStreamingEnabled() && !IsRunningCommandlet())
 		{
 			// Load last loaded regions
 			if (GetMutableDefault<UWorldPartitionEditorPerProjectUserSettings>()->GetEnableLoadingOfLastLoadedRegions())
