Reduce morph target vertex factory shader permutations by merging into GPUSkin vertex factory

CitySample PC cook stats (recompile all shaders in both runs):
Baseline: 75229 shaders compiled, Cook time 1h35m46s
WithChange: 69939 shaders compiled, Cook time 1h31m51s
Shader reduction 7.1%, Cook time reduction 4mins

diff --git a/Engine/Shaders/Private/GpuSkinVertexFactory.ush b/Engine/Shaders/Private/GpuSkinVertexFactory.ush
--- a/Engine/Shaders/Private/GpuSkinVertexFactory.ush
+++ b/Engine/Shaders/Private/GpuSkinVertexFactory.ush
@@ -8,6 +8,7 @@
 #include "VertexFactoryCommon.ush"
 #include "GpuSkinCommon.ush"
 
+// Whether the material supports morph targets
 #ifndef GPUSKIN_MORPH_BLEND
 #define GPUSKIN_MORPH_BLEND 0
 #endif
@@ -39,6 +40,21 @@
 
 #define FIXED_VERTEX_INDEX 0xFFFF
 
+// Whether this is a morph target
+#if GPUSKIN_MORPH_BLEND
+uint bIsMorphTarget;
+#endif
+
+bool IsMorphTarget()
+{
+#if GPUSKIN_MORPH_BLEND
+	return bIsMorphTarget;
+#else
+	// When Morph targets are not supported on the material, always return false so the associated shader code can be compiled out.
+	return false;
+#endif
+}
+
 // if the per bone motion blur is enabled for this draw call
 uint PerBoneMotionBlur;
 // number of influences for this draw call, 4 or 8
@@ -118,7 +134,7 @@ struct FVertexFactoryInput
 	VF_INSTANCED_STEREO_DECLARE_INPUT_BLOCK()
 };
 
-#if !GPUSKIN_MORPH_BLEND && GPUSKIN_APEX_CLOTH
+#if GPUSKIN_APEX_CLOTH
 // .x = Draw Index Buffer offset, .y = Offset into Cloth Vertex Buffer
 uint2 GPUSkinApexClothStartIndexOffset;
 
@@ -621,11 +637,17 @@ FBoneMatrix CalcPreviousBoneMatrix( FVertexFactoryInput Input )
 /** transform position by weighted sum of skinning matrices */
 float3 SkinPosition( FVertexFactoryInput Input, FVertexFactoryIntermediates Intermediates )
 {
-#if GPUSKIN_MORPH_BLEND
-	float3 Position = MorphPosition(Input, Intermediates);
-#else
-	float3 Position = Intermediates.UnpackedPosition;
-#endif
+	float3 Position;
+	if (IsMorphTarget())
+	{
+	#if GPUSKIN_MORPH_BLEND
+		Position = MorphPosition(Input, Intermediates);
+	#endif
+	}
+	else
+	{
+		Position = Intermediates.UnpackedPosition;
+	}
 
 	// Note the use of mul(Matrix,Vector), bone matrices are stored transposed
 	// for tighter packing.
@@ -637,11 +659,17 @@ float3 SkinPosition( FVertexFactoryInput Input, FVertexFactoryIntermediates Inte
 /** transform position by weighted sum of skinning matrices */
 float3 SkinPreviousPosition( FVertexFactoryInput Input, FVertexFactoryIntermediates Intermediates )
 {
-#if GPUSKIN_MORPH_BLEND
-	float3 Position = MorphPreviousPosition(Input, Intermediates);
-#else
-	float3 Position = Intermediates.UnpackedPosition;
-#endif
+	float3 Position;
+	if (IsMorphTarget())
+	{
+	#if GPUSKIN_MORPH_BLEND
+		Position = MorphPreviousPosition(Input, Intermediates);
+	#endif
+	}
+	else
+	{
+		Position = Intermediates.UnpackedPosition;
+	}
 
 	FBoneMatrix BlendMatrix = Intermediates.BlendMatrix;
 	
@@ -669,15 +697,7 @@ float3x3 SkinTangents( FVertexFactoryInput Input, FVertexFactoryIntermediates In
 	// -1..1 .xyz:normal, .w:contains sign of tangent basis determinant (left or right handed)
 	half4 LocalTangentZ = Input.TangentZ;
 
-#if GPUSKIN_MORPH_BLEND
-	// calc new normal by offseting it with the delta
-	LocalTangentZ.xyz = normalize(LocalTangentZ.xyz + Input.DeltaTangentZ);
-	// derive the new tangent by orthonormalizing the new normal against
-	// the base tangent vector (assuming these are normalized)
-	LocalTangentX = normalize(LocalTangentX - dot(LocalTangentX, LocalTangentZ.xyz) * LocalTangentZ.xyz);
-#else
-	#if GPUSKIN_APEX_CLOTH
-
+#if GPUSKIN_APEX_CLOTH
 	float3x3 ClothTangent = float3x3(1,0,0, 0,1,0, 0,0,1);
 	float TempClothBlendWeight = 0.0f;
 
@@ -763,10 +783,19 @@ float3x3 SkinTangents( FVertexFactoryInput Input, FVertexFactoryIntermediates In
 
 		TempClothBlendWeight *= SimulWeight;
 	}
-	
-	#endif // GPUSKIN_APEX_CLOTH
+#else
+	if (IsMorphTarget())
+	{
+#if GPUSKIN_MORPH_BLEND
+		// calc new normal by offseting it with the delta
+		LocalTangentZ.xyz = normalize(LocalTangentZ.xyz + Input.DeltaTangentZ);
+		// derive the new tangent by orthonormalizing the new normal against
+		// the base tangent vector (assuming these are normalized)
+		LocalTangentX = normalize(LocalTangentX - dot(LocalTangentX, LocalTangentZ.xyz) * LocalTangentZ.xyz);
+#endif
+	}
+#endif // GPUSKIN_APEX_CLOTH
 
-#endif // GPUSKIN_MORPH_BLEND
 
 	// Note the use of mul(Matrix,Vector), bone matrices are stored transposed
 	// for tighter packing.
diff --git a/Engine/Source/Runtime/Engine/Private/GPUSkinVertexFactory.cpp b/Engine/Source/Runtime/Engine/Private/GPUSkinVertexFactory.cpp
--- a/Engine/Source/Runtime/Engine/Private/GPUSkinVertexFactory.cpp
+++ b/Engine/Source/Runtime/Engine/Private/GPUSkinVertexFactory.cpp
@@ -507,7 +507,8 @@ bool TGPUSkinVertexFactory<BoneInfluenceType>::ShouldCompilePermutation(const FV
 {
 	bool bUnlimitedBoneInfluences = (BoneInfluenceType == UnlimitedBoneInfluence && GCVarUnlimitedBoneInfluences);
 	return ShouldWeCompileGPUSkinVFShaders(Parameters.Platform, Parameters.MaterialParameters.FeatureLevel) &&
-		  ((Parameters.MaterialParameters.bIsUsedWithSkeletalMesh && (BoneInfluenceType != UnlimitedBoneInfluence || bUnlimitedBoneInfluences)) || Parameters.MaterialParameters.bIsSpecialEngineMaterial);
+		  (((Parameters.MaterialParameters.bIsUsedWithSkeletalMesh || Parameters.MaterialParameters.bIsUsedWithMorphTargets) && (BoneInfluenceType != UnlimitedBoneInfluence || bUnlimitedBoneInfluences)) 
+			  || Parameters.MaterialParameters.bIsSpecialEngineMaterial);
 }
 
 template <GPUSkinBoneInfluenceType BoneInfluenceType>
@@ -530,6 +531,13 @@ void TGPUSkinVertexFactory<BoneInfluenceType>::ModifyCompilationEnvironment(cons
 	OutEnvironment.SetDefine(TEXT("GPU_SKINNED_MESH_FACTORY"), 1);
 
 	OutEnvironment.SetDefine(TEXT("VF_SUPPORTS_PRIMITIVE_SCENE_DATA"), bSupportsPrimitiveIdStream && bUseGPUScene);
+
+	// Mobile doesn't support motion blur, don't use previous frame morph delta for mobile.
+	const bool bIsMobile = IsMobilePlatform(Parameters.Platform);
+	OutEnvironment.SetDefine(TEXT("GPUSKIN_MORPH_USE_PREVIOUS"), !bIsMobile);
+
+	// Whether the material supports morph targets
+	OutEnvironment.SetDefine(TEXT("GPUSKIN_MORPH_BLEND"), Parameters.MaterialParameters.bIsUsedWithMorphTargets || Parameters.MaterialParameters.bIsSpecialEngineMaterial);
 }
 
 /**
@@ -575,6 +583,9 @@ void TGPUSkinVertexFactory<BoneInfluenceType>::GetPSOPrecacheVertexFetchElements
 	// Attribute ID
 	Elements.Add(FVertexElement(BaseStreamIndex++, 0, VET_UInt, 16, 0, true));
 
+	// Morph blend data
+	Elements.Add(FVertexElement(Elements.Num(), 0, VET_Float3, 9, 0, false));
+	Elements.Add(FVertexElement(Elements.Num(), 0, VET_Float3, 10, 0, false));
 }
 
 /**
@@ -662,6 +673,14 @@ void TGPUSkinVertexFactory<BoneInfluenceType>::AddVertexElements(FVertexDeclarat
 
 	// Primitive Id
 	AddPrimitiveIdStreamElement(EVertexInputStreamType::Default, OutElements, 16, 0xff);
+
+	// If the mesh is not a morph target, bind null component to morph delta stream.
+	FVertexStreamComponent NullComponent(&GNullVertexBuffer, 0, 0, VET_Float3);
+	FVertexElement DeltaPositionElement = AccessStreamComponent(Data->bMorphTarget ? Data->DeltaPositionComponent : NullComponent, 9);
+	// Cache delta stream index (position & tangentZ share the same stream)
+	MorphDeltaStreamIndex = DeltaPositionElement.StreamIndex;
+	OutElements.Add(DeltaPositionElement);
+	OutElements.Add(FVertexFactory::AccessStreamComponent(Data->bMorphTarget ? Data->DeltaTangentZComponent : NullComponent, 10));
 }
 
 /**
@@ -693,6 +712,22 @@ void TGPUSkinVertexFactory<BoneInfluenceType>::ReleaseDynamicRHI()
 	ShaderData.ReleaseBoneData();
 }
 
+template <GPUSkinBoneInfluenceType BoneInfluenceType>
+void TGPUSkinVertexFactory<BoneInfluenceType>::UpdateMorphVertexStream(const FMorphVertexBuffer* MorphVertexBuffer)
+{
+	if (MorphVertexBuffer && this->Streams.IsValidIndex(MorphDeltaStreamIndex))
+	{
+		this->Streams[MorphDeltaStreamIndex].VertexBuffer = MorphVertexBuffer;
+	}
+}
+
+template <GPUSkinBoneInfluenceType BoneInfluenceType>
+const FMorphVertexBuffer* TGPUSkinVertexFactory<BoneInfluenceType>::GetMorphVertexBuffer(bool bPrevious, uint32 FrameNumber) const
+{
+	check(Data.IsValid());
+	return Data->MorphVertexBufferPool ? &Data->MorphVertexBufferPool->GetMorphVertexBufferForReading(bPrevious, FrameNumber) : nullptr;
+}
+
 /*-----------------------------------------------------------------------------
 TGPUSkinAPEXClothVertexFactory
 -----------------------------------------------------------------------------*/
@@ -731,6 +766,8 @@ public:
 		InputWeightIndexSize.Bind(ParameterMap, TEXT("InputWeightIndexSize"));
 		InputWeightStream.Bind(ParameterMap, TEXT("InputWeightStream"));
 		NumBoneInfluencesParam.Bind(ParameterMap, TEXT("NumBoneInfluencesParam"));
+		IsMorphTarget.Bind(ParameterMap, TEXT("bIsMorphTarget"));
+		PreviousMorphBufferParameter.Bind(ParameterMap, TEXT("PreviousMorphBuffer"));
 	}
 
 	void GetElementShaderBindings(
@@ -789,6 +826,22 @@ public:
 			uint32 NumInfluences = ((const FGPUBaseSkinVertexFactory*)VertexFactory)->GetNumBoneInfluences();
 			ShaderBindings.Add(NumBoneInfluencesParam, NumInfluences);
 		}
+
+		ShaderBindings.Add(IsMorphTarget, (uint32)(((const FGPUBaseSkinVertexFactory*)VertexFactory)->IsMorphTarget() ? 1 : 0));
+
+		// Mobile doesn't support motion blur, don't use previous frame morph delta for mobile.
+		const EShaderPlatform ShaderPlatform = GetFeatureLevelShaderPlatform(FeatureLevel);
+		const bool bIsMobile = IsMobilePlatform(ShaderPlatform);
+		if (!bIsMobile)
+		{
+			const FMorphVertexBuffer* MorphVertexBuffer = nullptr;
+			const auto* GPUSkinVertexFactory = (const FGPUBaseSkinVertexFactory*)VertexFactory;
+			MorphVertexBuffer = GPUSkinVertexFactory->GetMorphVertexBuffer(!View->Family->bWorldIsPaused, View->Family->FrameNumber);
+			if (MorphVertexBuffer)
+			{
+				ShaderBindings.Add(PreviousMorphBufferParameter, MorphVertexBuffer->GetSRV());
+			}
+		}
 	}
 
 private:
@@ -798,7 +851,8 @@ private:
 	LAYOUT_FIELD(FShaderParameter, InputWeightIndexSize);
 	LAYOUT_FIELD(FShaderResourceParameter, InputWeightStream);
 	LAYOUT_FIELD(FShaderParameter, NumBoneInfluencesParam);
-
+	LAYOUT_FIELD(FShaderParameter, IsMorphTarget);
+	LAYOUT_FIELD(FShaderResourceParameter, PreviousMorphBufferParameter);
 };
 
 IMPLEMENT_TYPE_LAYOUT(FGPUSkinVertexFactoryShaderParameters);
@@ -975,175 +1029,6 @@ void FGPUSkinPassthroughVertexFactory::InternalUpdateVertexDeclaration(
 }
 
 
-/*-----------------------------------------------------------------------------
-	TGPUSkinMorphVertexFactoryShaderParameters
------------------------------------------------------------------------------*/
-/** Shader parameters for use with TGPUSkinMorphVertexFactory */
-class FGPUSkinMorphVertexFactoryShaderParameters : public FGPUSkinVertexFactoryShaderParameters
-{
-	DECLARE_TYPE_LAYOUT(FGPUSkinMorphVertexFactoryShaderParameters, NonVirtual);
-public:
-
-	/**
-	* Bind shader constants by name
-	* @param	ParameterMap - mapping of named shader constants to indices
-	*/
-	void Bind(const FShaderParameterMap& ParameterMap)
-	{
-		FGPUSkinVertexFactoryShaderParameters::Bind(ParameterMap);
-		PreviousMorphBufferParameter.Bind(ParameterMap, TEXT("PreviousMorphBuffer"));
-	}
-
-	void GetElementShaderBindings(
-		const FSceneInterface* Scene,
-		const FSceneView* View,
-		const FMeshMaterialShader* Shader,
-		const EVertexInputStreamType InputStreamType,
-		ERHIFeatureLevel::Type FeatureLevel,
-		const FVertexFactory* VertexFactory,
-		const FMeshBatchElement& BatchElement,
-		class FMeshDrawSingleShaderBindings& ShaderBindings,
-		FVertexInputStreamArray& VertexStreams) const
-	{
-		FGPUSkinVertexFactoryShaderParameters::GetElementShaderBindings(Scene, View, Shader, InputStreamType, FeatureLevel, VertexFactory, BatchElement, ShaderBindings, VertexStreams);
-
-		// Mobile doesn't support motion blur, don't use previous frame morph delta for mobile.
-		const EShaderPlatform ShaderPlatform = GetFeatureLevelShaderPlatform(FeatureLevel);
-		const bool bIsMobile = IsMobilePlatform(ShaderPlatform);
-		if (!bIsMobile)
-		{
-			const FMorphVertexBuffer* MorphVertexBuffer = nullptr;
-			const auto* GPUSkinVertexFactory = (const FGPUBaseSkinVertexFactory*)VertexFactory;
-			MorphVertexBuffer = GPUSkinVertexFactory->GetMorphVertexBuffer(!View->Family->bWorldIsPaused, View->Family->FrameNumber);
-			if (MorphVertexBuffer)
-			{
-				ShaderBindings.Add(PreviousMorphBufferParameter, MorphVertexBuffer->GetSRV());
-			}
-		}
-	}
-
-protected:
-	LAYOUT_FIELD(FShaderResourceParameter, PreviousMorphBufferParameter);
-};
-
-IMPLEMENT_TYPE_LAYOUT(FGPUSkinMorphVertexFactoryShaderParameters);
-
-
-/*-----------------------------------------------------------------------------
-TGPUSkinMorphVertexFactory
------------------------------------------------------------------------------*/
-
-/**
-* Modify compile environment to enable the morph blend codepath
-* @param OutEnvironment - shader compile environment to modify
-*/
-template <GPUSkinBoneInfluenceType BoneInfluenceType>
-void TGPUSkinMorphVertexFactory<BoneInfluenceType>::ModifyCompilationEnvironment( const FVertexFactoryShaderPermutationParameters& Parameters, FShaderCompilerEnvironment& OutEnvironment )
-{
-	Super::ModifyCompilationEnvironment(Parameters, OutEnvironment);
-	OutEnvironment.SetDefine(TEXT("GPUSKIN_MORPH_BLEND"),TEXT("1"));
-	// Mobile doesn't support motion blur, don't use previous frame morph delta for mobile.
-	const bool bIsMobile = IsMobilePlatform(Parameters.Platform);
-	OutEnvironment.SetDefine(TEXT("GPUSKIN_MORPH_USE_PREVIOUS"), !bIsMobile);
-}
-
-template <GPUSkinBoneInfluenceType BoneInfluenceType>
-bool TGPUSkinMorphVertexFactory<BoneInfluenceType>::ShouldCompilePermutation(const FVertexFactoryShaderPermutationParameters& Parameters)
-{
-	return (Parameters.MaterialParameters.bIsUsedWithMorphTargets || Parameters.MaterialParameters.bIsSpecialEngineMaterial)
-		&& Super::ShouldCompilePermutation(Parameters);
-}
-
-template <GPUSkinBoneInfluenceType BoneInfluenceType>
-void TGPUSkinMorphVertexFactory<BoneInfluenceType>::GetPSOPrecacheVertexFetchElements(EVertexInputStreamType VertexInputStreamType, FVertexDeclarationElementList& Elements)
-{
-	Super::GetPSOPrecacheVertexFetchElements(VertexInputStreamType, Elements);
-
-	// Morph blend data
-	Elements.Add(FVertexElement(Elements.Num(), 0, VET_Float3, 9, 0, false));
-	Elements.Add(FVertexElement(Elements.Num(), 0, VET_Float3, 10, 0, false));
-}
-
-template <GPUSkinBoneInfluenceType BoneInfluenceType>
-void TGPUSkinMorphVertexFactory<BoneInfluenceType>::SetData(const FGPUSkinDataType* InData)
-{
-	const FGPUSkinMorphDataType* InMorphData = (const FGPUSkinMorphDataType*)(InData);
-	check(InMorphData);
-
-	if (!this->Data)
-	{
-		MorphDataPtr = new FGPUSkinMorphDataType();
-		this->Data = TUniquePtr<FGPUSkinDataType>(MorphDataPtr);
-	}
-
-	*MorphDataPtr = *InMorphData;
-	FGPUBaseSkinVertexFactory::UpdateRHI();
-}
-
-/**
-* Add the decl elements for the streams
-* @param InData - type with stream components
-* @param OutElements - vertex decl list to modify
-*/
-template <GPUSkinBoneInfluenceType BoneInfluenceType>
-void TGPUSkinMorphVertexFactory<BoneInfluenceType>::AddVertexElements(FVertexDeclarationElementList& OutElements)
-{
-	// add the base gpu skin elements
-	Super::AddVertexElements(OutElements);
-	// add the morph delta elements
-	check(MorphDataPtr);
-	FVertexElement DeltaPositionElement = FVertexFactory::AccessStreamComponent(MorphDataPtr->DeltaPositionComponent, 9);
-	// Cache delta stream index (position & tangentZ share the same stream)
-	MorphDeltaStreamIndex = DeltaPositionElement.StreamIndex;
-	OutElements.Add(DeltaPositionElement);
-	OutElements.Add(FVertexFactory::AccessStreamComponent(MorphDataPtr->DeltaTangentZComponent, 10));
-}
-
-/**
-* Creates declarations for each of the vertex stream components and
-* initializes the device resource
-*/
-template <GPUSkinBoneInfluenceType BoneInfluenceType>
-void TGPUSkinMorphVertexFactory<BoneInfluenceType>::InitRHI()
-{
-	// list of declaration items
-	FVertexDeclarationElementList Elements;	
-	AddVertexElements(Elements);
-
-	// create the actual device decls
-	FVertexFactory::InitDeclaration(Elements);
-}
-
-/**
-* Update morph delta stream with the updated morph vertex buffer
-*/
-template <GPUSkinBoneInfluenceType BoneInfluenceType>
-void TGPUSkinMorphVertexFactory<BoneInfluenceType>::UpdateMorphVertexStream(const FMorphVertexBuffer* MorphVertexBuffer)
-{
-	if (MorphVertexBuffer && this->Streams.IsValidIndex(MorphDeltaStreamIndex))
-	{
-		this->Streams[MorphDeltaStreamIndex].VertexBuffer = MorphVertexBuffer;
-	}
-}
-
-template <GPUSkinBoneInfluenceType BoneInfluenceType>
-const FMorphVertexBuffer* TGPUSkinMorphVertexFactory<BoneInfluenceType>::GetMorphVertexBuffer(bool bPrevious, uint32 FrameNumber) const
-{
-	FGPUSkinMorphDataType* MorphData = (FGPUSkinMorphDataType*)(this->Data.Get());
-	check(MorphData);
-	return MorphData->MorphVertexBufferPool ? &MorphData->MorphVertexBufferPool->GetMorphVertexBufferForReading(bPrevious, FrameNumber) : nullptr;
-}
-
-IMPLEMENT_GPUSKINNING_VERTEX_FACTORY_PARAMETER_TYPE(TGPUSkinMorphVertexFactory, SF_Vertex, FGPUSkinMorphVertexFactoryShaderParameters);
-
-/** bind morph target gpu skin vertex factory to its shader file and its shader parameters */
-IMPLEMENT_GPUSKINNING_VERTEX_FACTORY_TYPE(TGPUSkinMorphVertexFactory, "/Engine/Private/GpuSkinVertexFactory.ush",
-	  EVertexFactoryFlags::UsedWithMaterials
-	| EVertexFactoryFlags::SupportsDynamicLighting
-	| EVertexFactoryFlags::SupportsPSOPrecaching
-);
-
-
 /*-----------------------------------------------------------------------------
 	FGPUBaseSkinAPEXClothVertexFactory
 -----------------------------------------------------------------------------*/
diff --git a/Engine/Source/Runtime/Engine/Private/SkeletalRenderGPUSkin.cpp b/Engine/Source/Runtime/Engine/Private/SkeletalRenderGPUSkin.cpp
--- a/Engine/Source/Runtime/Engine/Private/SkeletalRenderGPUSkin.cpp
+++ b/Engine/Source/Runtime/Engine/Private/SkeletalRenderGPUSkin.cpp
@@ -1683,9 +1683,10 @@ void InitGPUSkinVertexFactoryComponents(FGPUSkinDataType* VertexFactoryData, con
  * @param VertexBuffers - vertex buffers which contains the data and also stride info
  * @param bUseInstancedVertexWeights - use instanced influence weights instead of default weights
  */
-void InitMorphVertexFactoryComponents(FGPUSkinMorphDataType* VertexFactoryData,
+void InitMorphVertexFactoryComponents(FGPUSkinDataType* VertexFactoryData,
 										const FSkeletalMeshObjectGPUSkin::FVertexFactoryBuffers& VertexBuffers)
 {
+	VertexFactoryData->bMorphTarget = true;
 	VertexFactoryData->MorphVertexBufferPool = VertexBuffers.MorphVertexBufferPool;
 	const FMorphVertexBuffer& MorphVB = VertexBuffers.MorphVertexBufferPool->GetMorphVertexBufferForReading(false, -1);
 
@@ -1828,11 +1829,11 @@ static void CreateVertexFactoryMorph(TArray<TUniquePtr<FGPUBaseSkinVertexFactory
 	GPUSkinBoneInfluenceType BoneInfluenceType = InVertexBuffers.SkinWeightVertexBuffer->GetBoneInfluenceType();
 	if (BoneInfluenceType == GPUSkinBoneInfluenceType::DefaultBoneInfluence)
 	{
-		VertexFactory = new TGPUSkinMorphVertexFactory<GPUSkinBoneInfluenceType::DefaultBoneInfluence>(FeatureLevel, InVertexBuffers.NumVertices);
+		VertexFactory = new TGPUSkinVertexFactory<GPUSkinBoneInfluenceType::DefaultBoneInfluence>(FeatureLevel, InVertexBuffers.NumVertices);
 	}
 	else
 	{
-		VertexFactory = new TGPUSkinMorphVertexFactory<GPUSkinBoneInfluenceType::UnlimitedBoneInfluence>(FeatureLevel, InVertexBuffers.NumVertices);
+		VertexFactory = new TGPUSkinVertexFactory<GPUSkinBoneInfluenceType::UnlimitedBoneInfluence>(FeatureLevel, InVertexBuffers.NumVertices);
 	}
 	VertexFactories.Add(TUniquePtr<FGPUBaseSkinVertexFactory>(VertexFactory));
 
@@ -1843,7 +1844,7 @@ static void CreateVertexFactoryMorph(TArray<TUniquePtr<FGPUBaseSkinVertexFactory
 	ENQUEUE_RENDER_COMMAND(InitGPUSkinVertexFactoryMorph)(
 		[VertexUpdateData](FRHICommandList& RHICmdList)
 		{
-			FGPUSkinMorphDataType Data;
+			FGPUSkinDataType Data;
 			InitGPUSkinVertexFactoryComponents(&Data, VertexUpdateData.VertexBuffers, VertexUpdateData.VertexFactory);
 			InitMorphVertexFactoryComponents(&Data, VertexUpdateData.VertexBuffers);
 			VertexUpdateData.VertexFactory->SetData(&Data);
@@ -1868,7 +1869,7 @@ static void UpdateVertexFactoryMorph(TArray<TUniquePtr<FGPUBaseSkinVertexFactory
 			ENQUEUE_RENDER_COMMAND(InitGPUSkinVertexFactoryMorph)(
 				[VertexUpdateData](FRHICommandList& RHICmdList)
 			{
-				FGPUSkinMorphDataType Data;
+				FGPUSkinDataType Data;
 				InitGPUSkinVertexFactoryComponents(&Data, VertexUpdateData.VertexBuffers, VertexUpdateData.VertexFactory);
 				InitMorphVertexFactoryComponents(&Data, VertexUpdateData.VertexBuffers);
 				VertexUpdateData.VertexFactory->SetData(&Data);
diff --git a/Engine/Source/Runtime/Engine/Public/GPUSkinVertexFactory.h b/Engine/Source/Runtime/Engine/Public/GPUSkinVertexFactory.h
--- a/Engine/Source/Runtime/Engine/Public/GPUSkinVertexFactory.h
+++ b/Engine/Source/Runtime/Engine/Public/GPUSkinVertexFactory.h
@@ -182,6 +182,18 @@ struct FGPUSkinDataType : public FStaticMeshDataType
 
 	/** If the bone indices are 16 or 8-bit format */
 	bool bUse16BitBoneIndex = 0;
+
+	/** If this is a morph target */
+	bool bMorphTarget = false;
+
+	/** Morph target stream which has the position deltas to add to the vertex position */
+	FVertexStreamComponent DeltaPositionComponent;
+
+	/** Morph target stream which has the TangentZ deltas to add to the vertex normals */
+	FVertexStreamComponent DeltaTangentZComponent;
+
+	/** Morph vertex buffer pool double buffering delta data  */
+	class FMorphVertexBufferPool* MorphVertexBufferPool = nullptr;
 };
 
 /** Vertex factory with vertex stream components for GPU skinned vertices */
@@ -380,6 +392,7 @@ public:
 	virtual const FShaderResourceViewRHIRef GetColorComponentsSRV() const		{ return Data.IsValid() ? Data->ColorComponentsSRV : nullptr; }
 	virtual uint32 GetNumTexCoords() const										{ return Data.IsValid() ? Data->NumTexCoords : 0; }
 	virtual const uint32 GetColorIndexMask() const								{ return Data.IsValid() ? Data->ColorIndexMask : 0; }
+	virtual bool IsMorphTarget() const											{ return Data.IsValid() ? Data->bMorphTarget : false; }
 
 	inline const FVertexStreamComponent& GetPositionStreamComponent() const
 	{
@@ -436,7 +449,12 @@ public:
 
 	static void ModifyCompilationEnvironment(const FVertexFactoryShaderPermutationParameters& Parameters, FShaderCompilerEnvironment& OutEnvironment);
 	static bool ShouldCompilePermutation(const FVertexFactoryShaderPermutationParameters& Parameters);
-	static void GetPSOPrecacheVertexFetchElements(EVertexInputStreamType VertexInputStreamType, FVertexDeclarationElementList& Elements);
+	static void GetPSOPrecacheVertexFetchElements(EVertexInputStreamType VertexInputStreamType, FVertexDeclarationElementList& Elements);	
+	
+	/** FGPUBaseSkinVertexFactory overrides */
+	virtual void UpdateMorphVertexStream(const class FMorphVertexBuffer* MorphVertexBuffer) override;
+	virtual const class FMorphVertexBuffer* GetMorphVertexBuffer(bool bPrevious, uint32 FrameNumber) const override;
+
 
 	// FRenderResource interface.
 	virtual void InitRHI() override;
@@ -450,6 +468,9 @@ protected:
 	* @param OutElements - vertex decl list to modify
 	*/
 	virtual void AddVertexElements(FVertexDeclarationElementList& OutElements) override;
+
+private:
+	int32 MorphDeltaStreamIndex = -1;
 };
 
 /** 
@@ -576,72 +597,6 @@ protected:
 
 ENUM_CLASS_FLAGS(FGPUSkinPassthroughVertexFactory::EOverrideFlags)
 
-/** Stream component data bound to morph vertex factory */
-struct FGPUSkinMorphDataType : public FGPUSkinDataType
-{
-	/** stream which has the position deltas to add to the vertex position */
-	FVertexStreamComponent DeltaPositionComponent;
-	/** stream which has the TangentZ deltas to add to the vertex normals */
-	FVertexStreamComponent DeltaTangentZComponent;
-	/** Morph vertex buffer pool double buffering delta data  */
-	class FMorphVertexBufferPool* MorphVertexBufferPool = nullptr;
-};
-
-/** Vertex factory with vertex stream components for GPU-skinned and morph target streams */
-template<GPUSkinBoneInfluenceType BoneInfluenceType>
-class TGPUSkinMorphVertexFactory : public TGPUSkinVertexFactory<BoneInfluenceType>
-{
-	DECLARE_VERTEX_FACTORY_TYPE(TGPUSkinMorphVertexFactory<BoneInfluenceType>);
-
-	typedef TGPUSkinVertexFactory<BoneInfluenceType> Super;
-public:
-	/**
-	 * Constructor presizing bone matrices array to used amount.
-	 *
-	 * @param	InBoneMatrices	Reference to shared bone matrices array.
-	 */
-	TGPUSkinMorphVertexFactory(ERHIFeatureLevel::Type InFeatureLevel, uint32 InNumVertices)
-	: TGPUSkinVertexFactory<BoneInfluenceType>(InFeatureLevel, InNumVertices)
-	{}
-
-	static void ModifyCompilationEnvironment(const FVertexFactoryShaderPermutationParameters& Parameters, FShaderCompilerEnvironment& OutEnvironment);
-	static bool ShouldCompilePermutation(const FVertexFactoryShaderPermutationParameters& Parameters);
-	static void GetPSOPrecacheVertexFetchElements(EVertexInputStreamType VertexInputStreamType, FVertexDeclarationElementList& Elements);
-
-	/**
-	* An implementation of the interface used by TSynchronizedResource to 
-	* update the resource with new data from the game thread.
-	* @param	InData - new stream component data
-	*/
-	virtual void SetData(const FGPUSkinDataType* InData) override;
-
-	/** FGPUBaseSkinVertexFactory overrides */
-	virtual void UpdateMorphVertexStream(const class FMorphVertexBuffer* MorphVertexBuffer) override;
-	virtual const class FMorphVertexBuffer* GetMorphVertexBuffer(bool bPrevious, uint32 FrameNumber) const override;
-
-	// FRenderResource interface.
-
-	/**
-	* Creates declarations for each of the vertex stream components and
-	* initializes the device resource
-	*/
-	virtual void InitRHI() override;
-
-protected:
-	/**
-	* Add the decl elements for the streams
-	* @param InData - type with stream components
-	* @param OutElements - vertex decl list to modify
-	*/
-	virtual void AddVertexElements(FVertexDeclarationElementList& OutElements) override;
-
-	/** Alias pointer to TUniquePtr<FGPUSkinDataType> Data of FGPUBaseSkinVertexFactory. Note memory isn't managed through this pointer. */
-	FGPUSkinMorphDataType* MorphDataPtr = nullptr;
-
-private:
-	int32 MorphDeltaStreamIndex = -1;
-};
-
 /** Vertex factory with vertex stream components for GPU-skinned and morph target streams */
 class FGPUBaseSkinAPEXClothVertexFactory
 {
