Texture.cpp fix large nonpow2 textures failing to build
fix was already there but commented out by mistake

diff --git a/Engine/Source/Runtime/Engine/Private/Texture.cpp b/Engine/Source/Runtime/Engine/Private/Texture.cpp
--- a/Engine/Source/Runtime/Engine/Private/Texture.cpp
+++ b/Engine/Source/Runtime/Engine/Private/Texture.cpp
@@ -422,7 +422,7 @@ void UTexture::ValidateSettingsAfterImportOrEdit(bool * pRequiresNotifyMaterials
 		{
 			if ( GetTextureClass() == ETextureClass::TwoD )
 			{
-				UE_LOG(LogTexture, Display, TEXT("Large Texture %s Dimension=%d changed to VT"), *GetName(),MaxDimension);
+				UE_LOG(LogTexture, Display, TEXT("Large Texture %s Dimension=%d changed to VT; to disable VT set MaxTextureSize first"), *GetName(),MaxDimension);
 				VirtualTextureStreaming = true;
 				bRequiresNotifyMaterials = true;
 			}
@@ -453,7 +453,7 @@ void UTexture::ValidateSettingsAfterImportOrEdit(bool * pRequiresNotifyMaterials
 				// VT nonpow2 will fail to build
 				// force it into a state that will succeed? or just let it fail?
 				// you can either pad to pow2 or set MaxTextureSize and turn off VT
-				//PowerOfTwoMode = ETexturePowerOfTwoSetting::PadToPowerOfTwo;
+				PowerOfTwoMode = ETexturePowerOfTwoSetting::PadToPowerOfTwo;
 			}
 			else
 			{
