Use IsPlayInEditor() instead of GIsPlayInEditorWorld in FDeferredDecalProxy construction.
GIsPlayInEditorWorld is set too late for some async loading scenarios.
Fix request from Reikon Games.

diff --git a/Engine/Source/Runtime/Engine/Private/Components/DecalComponent.cpp b/Engine/Source/Runtime/Engine/Private/Components/DecalComponent.cpp
--- a/Engine/Source/Runtime/Engine/Private/Components/DecalComponent.cpp
+++ b/Engine/Source/Runtime/Engine/Private/Components/DecalComponent.cpp
@@ -47,7 +47,7 @@ FDeferredDecalProxy::FDeferredDecalProxy(const UDecalComponent* InComponent)
 
 #if WITH_EDITOR
 	// We don't want to fade when we're editing, only in Simulate/PIE/Game
-	if (!GIsEditor || GIsPlayInEditorWorld)
+	if (!GIsEditor || (InComponent->GetWorld() && InComponent->GetWorld()->IsPlayInEditor()))
 #endif
 	{
 		InitializeFadingParameters(InComponent->GetWorld()->GetTimeSeconds(), InComponent->GetFadeDuration(), InComponent->GetFadeStartDelay(), InComponent->GetFadeInDuration(), InComponent->GetFadeInStartDelay());
