Remove References to WORLD_MAX from Audio Device

diff --git a/Engine/Source/Runtime/Engine/Private/AudioDevice.cpp b/Engine/Source/Runtime/Engine/Private/AudioDevice.cpp
--- a/Engine/Source/Runtime/Engine/Private/AudioDevice.cpp
+++ b/Engine/Source/Runtime/Engine/Private/AudioDevice.cpp
@@ -5678,7 +5678,7 @@ void FAudioDevice::RemoveActiveSound(FActiveSound* ActiveSound)
 
 bool FAudioDevice::LocationIsAudible(const FVector& Location, const float MaxDistance) const
 {
-	if (MaxDistance >= WORLD_MAX)
+	if (MaxDistance >= MaxWorldDistanceCVar)
 	{
 		return true;
 	}
@@ -5705,7 +5705,7 @@ bool FAudioDevice::LocationIsAudible(const FVector& Location, const FTransform&
 	// This function is deprecated as it assumes listener 0.  
 	// To check if a location is audible by any listener, use FAudioDevice::LocationIsAudible that takes a location and max distance.
 	// To check if a location is audible by a specific listener, use FAudioDevice::LocationIsAudible that additionally takes a listener index. 
-	if (MaxDistance >= WORLD_MAX)
+	if (MaxDistance >= MaxWorldDistanceCVar)
 	{
 		return true;
 	}
@@ -5723,7 +5723,7 @@ bool FAudioDevice::LocationIsAudible(const FVector& Location, const FTransform&
 
 bool FAudioDevice::LocationIsAudible(const FVector& Location, int32 ListenerIndex, float MaxDistance) const
 {
-	if (MaxDistance >= WORLD_MAX)
+	if (MaxDistance >= MaxWorldDistanceCVar)
 	{
 		return true;
 	}
@@ -5747,7 +5747,7 @@ float FAudioDevice::GetDistanceToNearestListener(const FVector& Location) const
 		return FMath::Sqrt(DistSquared);
 	}
 
-	return WORLD_MAX;
+	return MaxWorldDistanceCVar;
 }
 
 float FAudioDevice::GetSquaredDistanceToListener(const FVector& Location, const FTransform& ListenerTransform) const
@@ -5790,7 +5790,7 @@ bool FAudioDevice::GetDistanceSquaredToNearestListener(const FVector& Location,
 	const bool bAllowAttenuationOverrides = true;
 	if (FindClosestListenerIndex(Location, DistSquared, bAllowAttenuationOverrides) == INDEX_NONE)
 	{
-		OutSqDistance = WORLD_MAX;
+		OutSqDistance = MaxWorldDistanceCVar;
 		return false;
 	}
 
