Fix non unity build with missing include file

diff --git a/Engine/Source/Runtime/Engine/Private/Particles/FXSystemPrivate.h b/Engine/Source/Runtime/Engine/Private/Particles/FXSystemPrivate.h
--- a/Engine/Source/Runtime/Engine/Private/Particles/FXSystemPrivate.h
+++ b/Engine/Source/Runtime/Engine/Private/Particles/FXSystemPrivate.h
@@ -11,6 +11,7 @@
 #include "VectorField.h"
 #include "ParticleSortingGPU.h"
 #include "GPUSortManager.h"
+#include "ParticleVertexFactory.h"
 
 class FCanvas;
 class FGlobalDistanceFieldParameterData;
