Improved world partition 3D debug draw (wp.Runtime.ToggleDrawRuntimeHash3D) : Cell Z bounds are now using their internal min/max Z values (necessary for wp.Runtime.ForceRuntimeSpatialHashZCulling).

diff --git a/Engine/Source/Runtime/Engine/Private/WorldPartition/WorldPartitionRuntimeSpatialHash.cpp b/Engine/Source/Runtime/Engine/Private/WorldPartition/WorldPartitionRuntimeSpatialHash.cpp
--- a/Engine/Source/Runtime/Engine/Private/WorldPartition/WorldPartitionRuntimeSpatialHash.cpp
+++ b/Engine/Source/Runtime/Engine/Private/WorldPartition/WorldPartitionRuntimeSpatialHash.cpp
@@ -430,15 +430,6 @@ void FSpatialHashStreamingGrid::Draw3D(UWorld* World, const TArray<FWorldPartiti
 
 	for (const FWorldPartitionStreamingSource& Source : Sources)
 	{
-		FVector StartTrace = Source.Location + FVector(0.f, 0.f, 100.f);
-		FVector EndTrace = StartTrace - FVector(0.f, 0.f, 1000000.f);
-		double Z = Source.Location.Z;
-		FHitResult Hit;
-		if (World->LineTraceSingleByObjectType(Hit, StartTrace, EndTrace, FCollisionObjectQueryParams(ECC_WorldStatic), FCollisionQueryParams(SCENE_QUERY_STAT(DebugWorldPartitionTrace), true)))
-		{
-			Z = Hit.ImpactPoint.Z;
-		}
-
 		FBox Region = Source.CalcBounds(GridLoadingRange, GridName, HLODLayer);
 		Region += FBox(Region.GetCenter() - MinExtent, Region.GetCenter() + MinExtent);
 
@@ -455,9 +446,6 @@ void FSpatialHashStreamingGrid::Draw3D(UWorld* World, const TArray<FWorldPartiti
 				Helper.Levels[GridLevel].GetCellBounds(FGridCellCoord2(Coords.X, Coords.Y), CellWorldBounds);
 				double CellSizeY = CellWorldBounds.GetSize().Y / FilteredCells.Num();
 				CellWorldBounds.Max.Y = CellWorldBounds.Min.Y + CellSizeY;
-				FVector BoundsExtent(CellWorldBounds.GetExtent(), 100.f);
-				FVector BoundsOrigin(CellWorldBounds.GetCenter(), Z);
-				FBox CellBox = FBox::BuildAABB(BoundsOrigin, BoundsExtent);
 				FTranslationMatrix CellOffsetMatrix(FVector(0.f, CellSizeY, 0.f));
 
 				for (const UWorldPartitionRuntimeCell* Cell : FilteredCells)
@@ -469,8 +457,15 @@ void FSpatialHashStreamingGrid::Draw3D(UWorld* World, const TArray<FWorldPartiti
 						continue;
 					}
 
+					// Use cell MinMaxZ to compute effective cell bounds
+					double CellZExtent = 0.5 * (Cell->GetMinMaxZ().Y - Cell->GetMinMaxZ().X);
+					double CellZPosition = Cell->GetMinMaxZ().X + CellZExtent;
+					FVector BoundsExtent(CellWorldBounds.GetExtent(), CellZExtent);
+					FVector BoundsOrigin(CellWorldBounds.GetCenter(), CellZPosition);
+					FBox CellBox = FBox::BuildAABB(BoundsOrigin, BoundsExtent);
+
 					// Draw Cell using its debug color
-					FColor CellColor = Cell->GetDebugColor(VisualizeMode).ToFColor(false).WithAlpha(64);
+					FColor CellColor = Cell->GetDebugColor(VisualizeMode).ToFColor(false).WithAlpha(16);
 					DrawDebugSolidBox(World, CellBox, CellColor, Transform, false, -1.f, 255);
 					FVector CellPos = Transform.TransformPosition(CellBox.GetCenter());
 					DrawDebugBox(World, CellPos, BoundsExtent, Transform.GetRotation(), CellColor.WithAlpha(255), false, -1.f, 255, 10.f);
@@ -495,11 +490,24 @@ void FSpatialHashStreamingGrid::Draw3D(UWorld* World, const TArray<FWorldPartiti
 		}
 
 		// Draw Streaming Source
+		double SourceLocationZ = Source.Location.Z;
+		// In simulation, still use line trace to find best Z location for source debug draw as streaming source follows the camera
+		if (UWorldPartition::IsSimulating())
+		{
+			FVector StartTrace = Source.Location + FVector(0.f, 0.f, 100.f);
+			FVector EndTrace = StartTrace - FVector(0.f, 0.f, 1000000.f);
+			FHitResult Hit;
+			if (World->LineTraceSingleByObjectType(Hit, StartTrace, EndTrace, FCollisionObjectQueryParams(ECC_WorldStatic), FCollisionQueryParams(SCENE_QUERY_STAT(DebugWorldPartitionTrace), true)))
+			{
+				SourceLocationZ = Hit.ImpactPoint.Z;
+			}
+		}
+
 		const FColor Color = Source.GetDebugColor();
-		Source.ForEachShape(GetLoadingRange(), GridName, HLODLayer, /*bProjectIn2D*/ true, [&Color, &Z, &Transform, &World, this](const FSphericalSector& Shape)
+		Source.ForEachShape(GetLoadingRange(), GridName, HLODLayer, /*bProjectIn2D*/ true, [&Color, &SourceLocationZ, &Transform, &World, this](const FSphericalSector& Shape)
 		{
 			FSphericalSector ZOffsettedShape = Shape;
-			ZOffsettedShape.SetCenter(FVector(FVector2D(ZOffsettedShape.GetCenter()), Z));
+			ZOffsettedShape.SetCenter(FVector(FVector2D(ZOffsettedShape.GetCenter()), SourceLocationZ));
 			DrawStreamingSource3D(World, ZOffsettedShape, Transform, Color);
 		});
 	}
