Fixed forward rendering broken static lighting fix. Fixed a but where lightmap build data would get invalidated incorrectly while cooking.

Contributed by [REDACTED_PERSON]
#0227

diff --git a/Engine/Source/Runtime/Engine/Private/World.cpp b/Engine/Source/Runtime/Engine/Private/World.cpp
--- a/Engine/Source/Runtime/Engine/Private/World.cpp
+++ b/Engine/Source/Runtime/Engine/Private/World.cpp
@@ -1924,14 +1924,17 @@ void UWorld::InitWorld(const InitializationValues IVS)
 #if WITH_EDITOR
 	RepairDefaultBrush();
 
-	// invalidate lighting if VT is enabled but no valid VT data is present or VT is disabled and no valid non-VT data is present.
-	for (auto Level : Levels) //Note: PersistentLevel is part of this array
+	if (!IsRunningCookCommandlet())
 	{
-		if (Level && Level->MapBuildData)
+		// invalidate lighting if VT is enabled but no valid VT data is present or VT is disabled and no valid non-VT data is present.
+		for (auto Level : Levels) //Note: PersistentLevel is part of this array
 		{
-			if (Level->MapBuildData->IsLightingValid(FeatureLevel) == false)
+			if (Level && Level->MapBuildData)
 			{
-				Level->MapBuildData->InvalidateSurfaceLightmaps(this);
+				if (Level->MapBuildData->IsLightingValid(FeatureLevel) == false)
+				{
+					Level->MapBuildData->InvalidateSurfaceLightmaps(this);
+				}
 			}
 		}
 	}
