Optimized ConditionalCollectGarbage
* Check the world conditions using the WorldList instead of an ActorIterator
* On dedicated servers only run over the world list once

diff --git a/Engine/Source/Runtime/Engine/Classes/Engine/Engine.h b/Engine/Source/Runtime/Engine/Classes/Engine/Engine.h
--- a/Engine/Source/Runtime/Engine/Classes/Engine/Engine.h
+++ b/Engine/Source/Runtime/Engine/Classes/Engine/Engine.h
@@ -2607,6 +2607,9 @@ private:
 	*/
 	void OnExternalUIChange(bool bInIsOpening);
 
+	/** Returns GetTimeBetweenGarbageCollectionPasses but tweaked if its an idle server or not */
+	float GetTimeBetweenGarbageCollectionPasses(bool bHasPlayersConnected) const;
+
 protected:
 
 	/**
diff --git a/Engine/Source/Runtime/Engine/Private/UnrealEngine.cpp b/Engine/Source/Runtime/Engine/Private/UnrealEngine.cpp
--- a/Engine/Source/Runtime/Engine/Private/UnrealEngine.cpp
+++ b/Engine/Source/Runtime/Engine/Private/UnrealEngine.cpp
@@ -1615,48 +1615,69 @@ void UEngine::PreGarbageCollect()
 	SendWorldEndOfFrameUpdates();
 }
 
-float UEngine::GetTimeBetweenGarbageCollectionPasses() const
+float GetLowMemoryGCTimer(const float DefaultTimeBetweenGC)
 {
-	float TimeBetweenGC = GTimeBetweenPurgingPendingKillObjects;
-
-	if (IsRunningDedicatedServer())
+	float MBFree = float(FPlatformMemory::GetStats().AvailablePhysical / 1024 / 1024);
+#if !UE_BUILD_SHIPPING
+	MBFree -= float(FPlatformMemory::GetExtraDevelopmentMemorySize() / 1024 / 1024);
+#endif
+	if (MBFree <= GLowMemoryMemoryThresholdMB)
 	{
-		bool bAtLeastOnePlayerConnected = false;
-
-		ForEachObjectOfClass(UWorld::StaticClass(),[&bAtLeastOnePlayerConnected](UObject* WorldObj)
+		if (GLowMemoryTimeBetweenPurgingPendingKillObjects < DefaultTimeBetweenGC)
 		{
-			UWorld* World = CastChecked<UWorld>(WorldObj);
-			bAtLeastOnePlayerConnected = bAtLeastOnePlayerConnected || ((World->NetDriver != nullptr) && World->NetDriver->ClientConnections.Num() > 0);
-		});
+			return GLowMemoryTimeBetweenPurgingPendingKillObjects;
+		}
 
-		if (!bAtLeastOnePlayerConnected)
+		// Go even faster if there are levels that need to be unloaded
+		if (GLowMemoryTimeBetweenPurgingPendingLevels < DefaultTimeBetweenGC && FLevelStreamingGCHelper::GetNumLevelsPendingPurge() > 0)
 		{
-			TimeBetweenGC *= GTimeBetweenPurgingPendingKillObjectsOnIdleServerMultiplier;
+			return GLowMemoryTimeBetweenPurgingPendingLevels;
 		}
 	}
 
-	// Do more frequent GCs if we're under the low memory threshold (if enabled)
-	if ( GLowMemoryMemoryThresholdMB > 0.0 )
+	return DefaultTimeBetweenGC;
+}
+
+float UEngine::GetTimeBetweenGarbageCollectionPasses() const
+{
+	bool bAtLeastOnePlayerConnected = false;
+
+	if (IsRunningDedicatedServer())
 	{
-		float MBFree = float(FPlatformMemory::GetStats().AvailablePhysical / 1024 / 1024);
-#if !UE_BUILD_SHIPPING
-		MBFree -= float(FPlatformMemory::GetExtraDevelopmentMemorySize() / 1024 / 1024);
-#endif
-		if (MBFree <= GLowMemoryMemoryThresholdMB)
+		for (const FWorldContext& Context : WorldList)
 		{
-			if (GLowMemoryTimeBetweenPurgingPendingKillObjects < TimeBetweenGC)
+			if (UWorld* World = Context.World())
 			{
-				TimeBetweenGC = GLowMemoryTimeBetweenPurgingPendingKillObjects;
+				if (World->NetDriver && World->NetDriver->ClientConnections.Num() > 0)
+				{
+					bAtLeastOnePlayerConnected = true;
+					break;
+				}
 			}
+		}
+	}
 
-			// Go even faster if there are levels that need to be unloaded
-			if (GLowMemoryTimeBetweenPurgingPendingLevels < TimeBetweenGC && FLevelStreamingGCHelper::GetNumLevelsPendingPurge() > 0)
-			{
-				TimeBetweenGC = GLowMemoryTimeBetweenPurgingPendingLevels;
-			}
+	return GetTimeBetweenGarbageCollectionPasses(bAtLeastOnePlayerConnected);
+}
+
+float UEngine::GetTimeBetweenGarbageCollectionPasses(bool bHasPlayersConnected) const
+{
+	float TimeBetweenGC = GTimeBetweenPurgingPendingKillObjects;
+
+	if (IsRunningDedicatedServer())
+	{
+		if (!bHasPlayersConnected)
+		{
+			TimeBetweenGC *= GTimeBetweenPurgingPendingKillObjectsOnIdleServerMultiplier;
 		}
 	}
 
+	// Do more frequent GCs if we're under the low memory threshold (if enabled)
+	if (GLowMemoryMemoryThresholdMB > 0.0)
+	{
+		TimeBetweenGC = GetLowMemoryGCTimer(TimeBetweenGC);
+	}
+
 	return TimeBetweenGC;
 }
 
@@ -1711,17 +1732,38 @@ void UEngine::ConditionalCollectGarbage()
 			}
 			else
 			{
+				const bool bTestForPlayers = IsRunningDedicatedServer();
 				bool bHasAWorldBegunPlay = false;
-				ForEachObjectOfClass(UWorld::StaticClass(), [&bHasAWorldBegunPlay](UObject* World)
+				bool bHasPlayersConnected = false;
+
+				// Look for conditions in the worlds that would change the GC frequency
+				for (const FWorldContext& Context : WorldList)
 				{
-					bHasAWorldBegunPlay = bHasAWorldBegunPlay || CastChecked<UWorld>(World)->HasBegunPlay();
-				});
+					if (UWorld* World = Context.World())
+					{
+						if (World->HasBegunPlay())
+						{
+							bHasAWorldBegunPlay = true;
+						}
+
+						if (bTestForPlayers && World->NetDriver && World->NetDriver->ClientConnections.Num() > 0)
+						{
+							bHasPlayersConnected = true;
+						}
+
+						// If we found the conditions we wanted, no need to continue iterating
+						if (bHasAWorldBegunPlay && (!bTestForPlayers || bHasPlayersConnected))
+						{
+							break;
+						}
+					}
+				}
 
 				if (bHasAWorldBegunPlay)
 				{
 					TimeSinceLastPendingKillPurge += FApp::GetDeltaTime();
 
-					const float TimeBetweenPurgingPendingKillObjects = GetTimeBetweenGarbageCollectionPasses();
+					const float TimeBetweenPurgingPendingKillObjects = GetTimeBetweenGarbageCollectionPasses(bHasPlayersConnected);
 
 					// See if we should delay garbage collect for this frame
 					if (bShouldDelayGarbageCollect)
