Handle invalid external data sizes

diff --git a/Engine/Source/Runtime/Engine/Private/ReplayHelper.cpp b/Engine/Source/Runtime/Engine/Private/ReplayHelper.cpp
--- a/Engine/Source/Runtime/Engine/Private/ReplayHelper.cpp
+++ b/Engine/Source/Runtime/Engine/Private/ReplayHelper.cpp
@@ -1599,7 +1599,7 @@ void FReplayHelper::LoadExternalData(FArchive& Ar, const float TimeSeconds)
 		// Read net guid this payload belongs to
 		Ar << NetGUID;
 
-		int32 ExternalDataNumBytes = (ExternalDataNumBits + 7) >> 3;
+		const int64 ExternalDataNumBytes = ((int64)ExternalDataNumBits + 7) >> 3;
 
 		FBitReader Reader(nullptr, ExternalDataNumBits);
 
