Handle invalid array sizes in replays

diff --git a/Engine/Source/Runtime/Engine/Private/RepLayout.cpp b/Engine/Source/Runtime/Engine/Private/RepLayout.cpp
--- a/Engine/Source/Runtime/Engine/Private/RepLayout.cpp
+++ b/Engine/Source/Runtime/Engine/Private/RepLayout.cpp
@@ -3956,6 +3956,12 @@ bool FRepLayout::ReceiveProperties_BackwardsCompatible_r(
 				return false;
 			}
 
+			if (ArrayNum > MAX_uint16)
+			{
+				UE_LOG(LogRep, Warning, TEXT("ReceiveProperties_BackwardsCompatible_r: ArrayNum out of valid range [%u]. Owner: %s, Name: %s, NetFieldExportHandle: %i, Checksum: %u"), ArrayNum, *Owner->GetName(), *NetFieldExportGroup->NetFieldExports[NetFieldExportHandle].ExportName.ToString(), NetFieldExportHandle, Checksum);
+				return false;
+			}
+
 			const int32 AbsOffset = (Data.Data - OldData.Data) + Cmd.Offset;
 
 			FScriptArray* DataArray = (FScriptArray*)(Data + Cmd).Data;
