Fixed thread assert in GetAllocatedVirtualTexture.

diff --git a/Engine/Source/Runtime/Engine/Private/VT/RuntimeVirtualTexture.cpp b/Engine/Source/Runtime/Engine/Private/VT/RuntimeVirtualTexture.cpp
--- a/Engine/Source/Runtime/Engine/Private/VT/RuntimeVirtualTexture.cpp
+++ b/Engine/Source/Runtime/Engine/Private/VT/RuntimeVirtualTexture.cpp
@@ -122,7 +122,7 @@ public:
 	/** Getter for the virtual texture allocation. */
 	IAllocatedVirtualTexture* GetAllocatedVirtualTexture() const 
 	{
-		checkSlow(IsInRenderingThread());
+		checkSlow(IsInParallelRenderingThread());
 		return AdaptiveVirtualTexture != nullptr ? AdaptiveVirtualTexture->GetAllocatedVirtualTexture() : AllocatedVirtualTexture;
 	}
 
