fix bad merge

#rnx

diff --git a/Engine/Source/Runtime/Engine/Private/Level.cpp b/Engine/Source/Runtime/Engine/Private/Level.cpp
--- a/Engine/Source/Runtime/Engine/Private/Level.cpp
+++ b/Engine/Source/Runtime/Engine/Private/Level.cpp
@@ -1736,10 +1736,6 @@ bool ULevel::IncrementalRunConstructionScripts(bool bProcessAllActors)
 		}
 	}
 
-	if (OwningWorld)
-	{
-	}
-
 	if (OwningWorld)
 	{
 		OwningWorld->SetAllowDeferredPhysicsStateCreation(false);
