World Partition - HLOD: Fix PostLoad() SourceCellName fixup to work on HLOD from runtime data layers

diff --git a/Engine/Source/Runtime/Engine/Private/WorldPartition/HLOD/HLODActor.cpp b/Engine/Source/Runtime/Engine/Private/WorldPartition/HLOD/HLODActor.cpp
--- a/Engine/Source/Runtime/Engine/Private/WorldPartition/HLOD/HLODActor.cpp
+++ b/Engine/Source/Runtime/Engine/Private/WorldPartition/HLOD/HLODActor.cpp
@@ -101,18 +101,18 @@ void AWorldPartitionHLOD::PostLoad()
 
 	if (GetLinkerCustomVersion(FFortniteMainBranchObjectVersion::GUID) < FFortniteMainBranchObjectVersion::WorldPartitionStreamingCellsNamingShortened)
 	{
-		if (UWorldPartition* WorldPartition = GetWorld()->GetWorldPartition())
+		if (!HLODSubActors.IsEmpty())
 		{
-			if (FWorldPartitionActorDesc* ActorDesc = WorldPartition->GetActorDesc(GetActorGuid()))
+			// As we may be dealing with an unsaved world created from a template map, get
+			// the source package name of this HLOD actor and figure out the world name from there
+			FName ExternalActorsPath = HLODSubActors[0].ContainerPackage;
+			FString WorldName = FPackageName::GetShortName(ExternalActorsPath);
+
+			// Strip "WorldName_" from the cell name
+			FString CellName = SourceCellName.ToString();
+			bool bRemoved = CellName.RemoveFromStart(WorldName + TEXT("_"), ESearchCase::CaseSensitive);
+			if (ensure(bRemoved))
 			{
-				// As we may be dealing with an unsaved world created from a template map, get
-				// the source package name of this HLOD actor and figure out the world name from there
-				FString ExternalActorsPath = ActorDesc->GetContainer()->GetExternalActorPath();
-				FString WorldName = FPackageName::GetShortName(ExternalActorsPath);
-				
-				// Strip "WorldName_" from the cell name
-				FString CellName = SourceCellName.ToString();
-				CellName.RemoveFromStart(WorldName + TEXT("_"), ESearchCase::CaseSensitive);
 				SourceCellName = *CellName;
 			}
 		}
