Protect against potential crash in unweld

#8053

diff --git a/Engine/Source/Runtime/Engine/Private/PrimitiveComponentPhysics.cpp b/Engine/Source/Runtime/Engine/Private/PrimitiveComponentPhysics.cpp
--- a/Engine/Source/Runtime/Engine/Private/PrimitiveComponentPhysics.cpp
+++ b/Engine/Source/Runtime/Engine/Private/PrimitiveComponentPhysics.cpp
@@ -807,9 +807,13 @@ void UPrimitiveComponent::UnWeldChildren()
 {
 	for (USceneComponent* ChildComponent : GetAttachChildren())
 	{
-		if (UPrimitiveComponent* PrimComp = Cast<UPrimitiveComponent>(ChildComponent))
+		// Invalid children should be removed in UPrimitiveComponent::DestroyComponent so to check that you have correctly destroyed attached components.
+		if (ensure(IsValid(ChildComponent)))
 		{
-			PrimComp->UnWeldFromParent();
+			if (UPrimitiveComponent* PrimComp = Cast<UPrimitiveComponent>(ChildComponent))
+			{
+				PrimComp->UnWeldFromParent();
+			}
 		}
 	}
 }
