* Made sure that BP-exposed functions resulting in a material (UMaterialInterface) draw are synchronously waiting for their shaders to be ready before issueing the draw command (similar to 21965139 but expanded to user-facing UCanvas functions)
* Removed double-update of canvas in BP functions

diff --git a/Engine/Source/Runtime/Engine/Classes/Materials/MaterialInterface.h b/Engine/Source/Runtime/Engine/Classes/Materials/MaterialInterface.h
--- a/Engine/Source/Runtime/Engine/Classes/Materials/MaterialInterface.h
+++ b/Engine/Source/Runtime/Engine/Classes/Materials/MaterialInterface.h
@@ -887,6 +887,17 @@ public:
 	*/
 	ENGINE_API virtual bool IsComplete() const { return true; }
 
+	/** @brief Checks to see if this material has all its shaders cached and if not, will perform a synchronous compilation of those.
+	*
+	* Materials are not guaranteed to have all their shaders compiled after loading and using this function before a draw will ensure that the material will not render until it's 
+	* ready to (and use a fallback material instead). This needs to be avoided in the common render path but can be useful in critical tools (e.g. landscape) that rely on the 
+	* material and cannot afford a fallback material. 
+	* In the editor, it will display a toast when waiting for the shaders to compile.
+	*
+	* @see CacheShaders
+	*/
+	ENGINE_API void EnsureIsComplete();
+
 #if WITH_EDITOR
 	/** Clears the shader cache and recompiles the shader for rendering. */
 	ENGINE_API virtual void ForceRecompileForRendering() {}
diff --git a/Engine/Source/Runtime/Engine/Private/KismetRenderingLibrary.cpp b/Engine/Source/Runtime/Engine/Private/KismetRenderingLibrary.cpp
--- a/Engine/Source/Runtime/Engine/Private/KismetRenderingLibrary.cpp
+++ b/Engine/Source/Runtime/Engine/Private/KismetRenderingLibrary.cpp
@@ -3,7 +3,6 @@
 #include "Kismet/KismetRenderingLibrary.h"
 #include "HAL/FileManager.h"
 #include "Misc/Paths.h"
-#include "Misc/ScopedSlowTask.h"
 #include "Serialization/BufferArchive.h"
 #include "EngineGlobals.h"
 #include "RenderingThread.h"
@@ -178,12 +177,8 @@ void UKismetRenderingLibrary::DrawMaterialToRenderTarget(UObject* WorldContextOb
 	}
 	else
 	{
-		if (!Material->IsComplete())
-		{
-			FScopedSlowTask SlowTask(0.0f, NSLOCTEXT("KismetRenderingLibrary", "CacheShaders", "Caching Shaders..."));
-			SlowTask.MakeDialog();
-			Material->CacheShaders(EMaterialShaderPrecompileMode::Synchronous);
-		}
+		// This is a user-facing function, so we'd rather make sure that shaders are ready by the time we render, in order to ensure we don't draw with a fallback material :
+		Material->EnsureIsComplete();
 
 		World->FlushDeferredParameterCollectionInstanceUpdates();
 
@@ -198,7 +193,6 @@ void UKismetRenderingLibrary::DrawMaterialToRenderTarget(UObject* WorldContextOb
 			World->FeatureLevel);
 
 		Canvas->Init(TextureRenderTarget->SizeX, TextureRenderTarget->SizeY, nullptr, &RenderCanvas);
-		Canvas->Update();
 
 		{
 			SCOPED_DRAW_EVENTF_GAMETHREAD(DrawMaterialToRenderTarget, *TextureRenderTarget->GetFName().ToString());
@@ -748,7 +742,6 @@ void UKismetRenderingLibrary::BeginDrawCanvasToRenderTarget(UObject* WorldContex
 			// Draw immediately so that interleaved SetVectorParameter (etc) function calls work as expected
 			FCanvas::CDM_ImmediateDrawing);
 		Canvas->Init(TextureRenderTarget->SizeX, TextureRenderTarget->SizeY, nullptr, NewCanvas);
-		Canvas->Update();
 
 #if  WANTS_DRAW_MESH_EVENTS
 		Context.DrawEvent = new FDrawEvent();
diff --git a/Engine/Source/Runtime/Engine/Private/Materials/MaterialInterface.cpp b/Engine/Source/Runtime/Engine/Private/Materials/MaterialInterface.cpp
--- a/Engine/Source/Runtime/Engine/Private/Materials/MaterialInterface.cpp
+++ b/Engine/Source/Runtime/Engine/Private/Materials/MaterialInterface.cpp
@@ -30,12 +30,15 @@
 #include "ShaderPlatformQualitySettings.h"
 #include "ObjectCacheContext.h"
 #include "MaterialCachedData.h"
+#include "Misc/ScopedSlowTask.h"
 
 #if WITH_EDITOR
 #include "ObjectCacheEventSink.h"
 #include "MaterialCachedHLSLTree.h"
 #endif
 
+#define LOCTEXT_NAMESPACE "MaterialInterface"
+
 /**
  * This is used to deprecate data that has been built with older versions.
  * To regenerate the data, commands like "BUILDMATERIALTEXTURESTREAMINGDATA" can be used in the editor.
@@ -1427,6 +1430,16 @@ void UMaterialInterface::RemoveUserDataOfClass(TSubclassOf<UAssetUserData> InUse
 	}
 }
 
+void UMaterialInterface::EnsureIsComplete()
+{
+	if (!IsComplete())
+	{
+		FScopedSlowTask SlowTask(0.0f, LOCTEXT("CacheShaders", "Caching Shaders..."));
+		SlowTask.MakeDialog();
+		CacheShaders(EMaterialShaderPrecompileMode::Synchronous);
+	}
+}
+
 #if WITH_EDITORONLY_DATA
 
 namespace MaterialInterface
@@ -1513,3 +1526,5 @@ void UMaterialInterfaceEditorOnlyData::Serialize(FArchive& Ar)
 		bLoadedCachedExpressionData = true;
 	}
 }
+
+#undef LOCTEXT_NAMESPACE 
\ No newline at end of file
diff --git a/Engine/Source/Runtime/Engine/Private/UserInterface/Canvas.cpp b/Engine/Source/Runtime/Engine/Private/UserInterface/Canvas.cpp
--- a/Engine/Source/Runtime/Engine/Private/UserInterface/Canvas.cpp
+++ b/Engine/Source/Runtime/Engine/Private/UserInterface/Canvas.cpp
@@ -2059,6 +2059,9 @@ void UCanvas::K2_DrawMaterial(UMaterialInterface* RenderMaterial, FVector2D Scre
 		// Canvas can be NULL if the user tried to draw after EndDrawCanvasToRenderTarget
 		&& Canvas)
 	{
+		// This is a user-facing function, so we'd rather make sure that shaders are ready by the time we render, in order to ensure we don't draw with a fallback material :
+		RenderMaterial->EnsureIsComplete();
+
 		FCanvasTileItem TileItem(ScreenPosition, RenderMaterial->GetRenderProxy(), ScreenSize, CoordinatePosition, CoordinatePosition + CoordinateSize);
 		TileItem.Rotation = FRotator(0, Rotation, 0);
 		TileItem.PivotPoint = PivotPoint;
@@ -2126,6 +2129,9 @@ void UCanvas::K2_DrawMaterialTriangle(UMaterialInterface* RenderMaterial, TArray
 {
 	if (RenderMaterial && Triangles.Num() > 0 && Canvas)
 	{
+		// This is a user-facing function, so we'd rather make sure that shaders are ready by the time we render, in order to ensure we don't draw with a fallback material :
+		RenderMaterial->EnsureIsComplete();
+
 		FCanvasTriangleItem TriangleItem(FVector2D::ZeroVector, FVector2D::ZeroVector, FVector2D::ZeroVector, NULL);
 		TriangleItem.MaterialRenderProxy = RenderMaterial->GetRenderProxy();
 		TriangleItem.TriangleList = MoveTemp(Triangles);
