Re-enable generate stremaing on contentbundes

diff --git a/Engine/Source/Runtime/Engine/Private/WorldPartition/ContentBundle/ContentBundleContainer.cpp b/Engine/Source/Runtime/Engine/Private/WorldPartition/ContentBundle/ContentBundleContainer.cpp
--- a/Engine/Source/Runtime/Engine/Private/WorldPartition/ContentBundle/ContentBundleContainer.cpp
+++ b/Engine/Source/Runtime/Engine/Private/WorldPartition/ContentBundle/ContentBundleContainer.cpp
@@ -56,7 +56,7 @@ void FContentBundleContainer::Initialize()
 
 #if WITH_EDITOR
 	UWorldPartition* WorldPartition = GetInjectedWorld()->GetWorldPartition();
-	//WorldPartition->OnPreGenerateStreaming.AddRaw(this, &FContentBundleContainer::OnPreGenerateStreaming); // PHILTODO
+	WorldPartition->OnPreGenerateStreaming.AddRaw(this, &FContentBundleContainer::OnPreGenerateStreaming);
 #endif
 }
 
@@ -66,7 +66,7 @@ void FContentBundleContainer::Deinitialize()
 
 #if WITH_EDITOR
 	UWorldPartition* WorldPartition = GetInjectedWorld()->GetWorldPartition();
-	//WorldPartition->OnPreGenerateStreaming.RemoveAll(this);
+	WorldPartition->OnPreGenerateStreaming.RemoveAll(this);
 #endif
 
 	UnregisterContentBundleClientEvents();
