Fix crash occuring when opening a map which has a content bundle but it's AWorldDataLayer is deleted
Can happen while by manually deleting an external actor

diff --git a/Engine/Source/Runtime/Engine/Private/WorldPartition/ContentBundle/WorldDataLayerReference.cpp b/Engine/Source/Runtime/Engine/Private/WorldPartition/ContentBundle/WorldDataLayerReference.cpp
--- a/Engine/Source/Runtime/Engine/Private/WorldPartition/ContentBundle/WorldDataLayerReference.cpp
+++ b/Engine/Source/Runtime/Engine/Private/WorldPartition/ContentBundle/WorldDataLayerReference.cpp
@@ -61,7 +61,7 @@ const AWorldDataLayers* FWorldDataLayersReference::Get() const
 		const FWorldPartitionReference& WorldDataLayerReference = WorldDataLayersVariant.Get<FWorldPartitionReference>();
 		if (WorldDataLayerReference.IsValid())
 		{
-			return CastChecked<AWorldDataLayers>(WorldDataLayerReference->GetActor());
+			return Cast<AWorldDataLayers>(WorldDataLayerReference->GetActor());
 		}
 		return nullptr;
 	}
