Disable WorldDataLayer automatic injection while investigating reference issue with ULevel imports.
Remove RemoveBaseContent as WorldDataLayer do not support being destroyed.

diff --git a/Engine/Source/Editor/WorldPartitionEditor/Private/WorldPartition/ContentBundle/ContentBundleEditorSubsystem.cpp b/Engine/Source/Editor/WorldPartitionEditor/Private/WorldPartition/ContentBundle/ContentBundleEditorSubsystem.cpp
--- a/Engine/Source/Editor/WorldPartitionEditor/Private/WorldPartition/ContentBundle/ContentBundleEditorSubsystem.cpp
+++ b/Engine/Source/Editor/WorldPartitionEditor/Private/WorldPartition/ContentBundle/ContentBundleEditorSubsystem.cpp
@@ -128,10 +128,6 @@ bool UContentBundleEditionSubmodule::ActivateContentBundleEditing(TSharedPtr<FCo
 		EditingContentBundle = ContentBundleEditor;
 
 		ContentBundleEditor->StartEditing();
-		if (ContentBundleStatus == EContentBundleStatus::ReadyToInject)
-		{
-			ContentBundleEditor->InjectBaseContent();
-		}
 
 		ActorEditorContextClientChanged.Broadcast(this);
 
diff --git a/Engine/Source/Editor/WorldPartitionEditor/Private/WorldPartition/ContentBundle/Outliner/ContentBundleMode.cpp b/Engine/Source/Editor/WorldPartitionEditor/Private/WorldPartition/ContentBundle/Outliner/ContentBundleMode.cpp
--- a/Engine/Source/Editor/WorldPartitionEditor/Private/WorldPartition/ContentBundle/Outliner/ContentBundleMode.cpp
+++ b/Engine/Source/Editor/WorldPartitionEditor/Private/WorldPartition/ContentBundle/Outliner/ContentBundleMode.cpp
@@ -145,34 +145,9 @@ void FContentBundleMode::RegisterContextMenu()
 								}),
 							FCanExecuteAction::CreateLambda([this, SceneOutliner]
 								{
-									if (TSharedPtr<FContentBundleEditor> SelectedContentBundleEditor = GetSelectedContentBundlePin(SceneOutliner))
-									{
-										return SelectedContentBundleEditor->IsBeingEdited() && SelectedContentBundleEditor->GetStatus() == EContentBundleStatus::ReadyToInject;
-									}
 									return false;
 								})
 							));
-
-					Section.AddMenuEntry("RemoveBaseContentMenuEntry", LOCTEXT("RemoveBaseContent", "Remove Base Content"), FText(), FSlateIcon(),
-						FUIAction(
-							FExecuteAction::CreateLambda([this, SceneOutliner]
-								{
-									if (TSharedPtr<FContentBundleEditor> SelectedContentBundleEditor = GetSelectedContentBundlePin(SceneOutliner))
-									{
-										SelectedContentBundleEditor->RemoveBaseContent();
-									}
-								}),
-							FCanExecuteAction::CreateLambda([this, SceneOutliner]
-								{
-									if (TSharedPtr<FContentBundleEditor> SelectedContentBundleEditor = GetSelectedContentBundlePin(SceneOutliner))
-									{
-										return SelectedContentBundleEditor->IsBeingEdited() && SelectedContentBundleEditor->GetStatus() == EContentBundleStatus::ContentInjected && !SelectedContentBundleEditor->HasUserPlacedActors();
-									}
-									return false;
-								})
-									));
-
-					
 				}
 
 				{
diff --git a/Engine/Source/Runtime/Engine/Private/WorldPartition/ContentBundle/ContentBundleEditor.cpp b/Engine/Source/Runtime/Engine/Private/WorldPartition/ContentBundle/ContentBundleEditor.cpp
--- a/Engine/Source/Runtime/Engine/Private/WorldPartition/ContentBundle/ContentBundleEditor.cpp
+++ b/Engine/Source/Runtime/Engine/Private/WorldPartition/ContentBundle/ContentBundleEditor.cpp
@@ -129,6 +129,12 @@ bool FContentBundleEditor::AddActor(AActor* InActor)
 		return false;
 	}
 
+	if (GetStatus() == EContentBundleStatus::ReadyToInject)
+	{
+		UE_LOG(LogContentBundle, Verbose, TEXT("[CB: %s] Adding first actor to content bundle."), *GetDescriptor()->GetDisplayName());
+		SetStatus(EContentBundleStatus::ContentInjected);
+	}
+
 	check(GetStatus() == EContentBundleStatus::ContentInjected);
 
 	FName ActorPackageNameInContentBundle(*ULevel::GetActorPackageName(ActorDescContainer->GetExternalActorPath(), EActorPackagingScheme::Reduced, InActor->GetName()));
@@ -260,20 +266,6 @@ void FContentBundleEditor::InjectBaseContent()
 	BroadcastChanged();
 }
 
-void FContentBundleEditor::RemoveBaseContent()
-{
-	check(!HasUserPlacedActors());
-	check(GetStatus() == EContentBundleStatus::ContentInjected);
-	UE_LOG(LogContentBundle, Log, TEXT("[CB: %s] Removing Base Content"), *GetDescriptor()->GetDisplayName());
-
-	GetInjectedWorld()->DestroyActor(WorldDataLayersActorReference.Get());
-	WorldDataLayersActorReference.Reset();
-
-	SetStatus(EContentBundleStatus::ReadyToInject);
-
-	BroadcastChanged();
-}
-
 void FContentBundleEditor::AddReferencedObjects(FReferenceCollector& Collector)
 {
 	FContentBundleBase::AddReferencedObjects(Collector);
@@ -390,12 +382,7 @@ void FContentBundleEditor::OnActorDescRemoved(FWorldPartitionActorDesc* ActorDes
 	{
 		if (GetStatus() == EContentBundleStatus::ContentInjected)
 		{
-			RemoveBaseContent();
-		}
-		else
-		{
-			// If content is not injected, then RemoveContent was already called and we are deleting the WorldDataLayers
-			check(ActorDesc->GetActorNativeClass() == AWorldDataLayers::StaticClass() && GetStatus() == EContentBundleStatus::ReadyToInject);
+			SetStatus(EContentBundleStatus::ReadyToInject);
 		}
 	}
 
@@ -410,7 +397,7 @@ void FContentBundleEditor::OnUnsavedActorDeleted(AActor* Actor)
 
 	if (!HasUserPlacedActors())
 	{
-		RemoveBaseContent();
+		SetStatus(EContentBundleStatus::ReadyToInject);
 	}
 }
 
diff --git a/Engine/Source/Runtime/Engine/Public/WorldPartition/ContentBundle/ContentBundleEditor.h b/Engine/Source/Runtime/Engine/Public/WorldPartition/ContentBundle/ContentBundleEditor.h
--- a/Engine/Source/Runtime/Engine/Public/WorldPartition/ContentBundle/ContentBundleEditor.h
+++ b/Engine/Source/Runtime/Engine/Public/WorldPartition/ContentBundle/ContentBundleEditor.h
@@ -40,7 +40,6 @@ public:
 	bool IsBeingEdited() const { return bIsBeingEdited; }
 
 	void InjectBaseContent();
-	void RemoveBaseContent();
 
 	bool AddActor(AActor* InActor);
 	bool ContainsActor(const AActor* InActor) const;
