Fix an issue where a content bundle could be unactived but still marked as editing.

diff --git a/Engine/Source/Editor/WorldPartitionEditor/Private/WorldPartition/ContentBundle/ContentBundleEditorSubsystem.cpp b/Engine/Source/Editor/WorldPartitionEditor/Private/WorldPartition/ContentBundle/ContentBundleEditorSubsystem.cpp
--- a/Engine/Source/Editor/WorldPartitionEditor/Private/WorldPartition/ContentBundle/ContentBundleEditorSubsystem.cpp
+++ b/Engine/Source/Editor/WorldPartitionEditor/Private/WorldPartition/ContentBundle/ContentBundleEditorSubsystem.cpp
@@ -215,6 +215,21 @@ UWorld* UContentBundleEditorSubsystem::GetWorld() const
 	return GEditor->GetEditorWorldContext().World();
 }
 
+void UContentBundleEditorSubsystem::NotifyContentBundleInjectedContent(const FContentBundleEditor* ContentBundle)
+{
+	NotifyContentBundleChanged(ContentBundle);
+}
+
+void UContentBundleEditorSubsystem::NotifyContentBundleRemovedContent(const FContentBundleEditor* ContentBundle)
+{
+	if (ContentBundle->IsBeingEdited())
+	{
+		GetEditionSubmodule()->DeactivateCurrentContentBundleEditing();
+	}
+
+	NotifyContentBundleChanged(ContentBundle);
+}
+
 TSharedPtr<FContentBundleEditor> UContentBundleEditorSubsystem::GetEditorContentBundleForActor(const AActor* Actor)
 {
 	for (TSharedPtr<FContentBundleEditor>& ContentBundle : GetEditorContentBundles())
diff --git a/Engine/Source/Editor/WorldPartitionEditor/Public/WorldPartition/ContentBundle/ContentBundleEditorSubsystem.h b/Engine/Source/Editor/WorldPartitionEditor/Public/WorldPartition/ContentBundle/ContentBundleEditorSubsystem.h
--- a/Engine/Source/Editor/WorldPartitionEditor/Public/WorldPartition/ContentBundle/ContentBundleEditorSubsystem.h
+++ b/Engine/Source/Editor/WorldPartitionEditor/Public/WorldPartition/ContentBundle/ContentBundleEditorSubsystem.h
@@ -94,6 +94,8 @@ public:
 	//~ Begin IContentBundleEditorSubsystemInterface interface
 	virtual void NotifyContentBundleAdded(const FContentBundleEditor* ContentBundle) override { OnContentBundleAdded().Broadcast(ContentBundle); }
 	virtual void NotifyContentBundleRemoved(const FContentBundleEditor* ContentBundle) override { OnContentBundleRemoved().Broadcast(ContentBundle); }
+	virtual void NotifyContentBundleInjectedContent(const FContentBundleEditor* ContentBundle) override;
+	virtual void NotifyContentBundleRemovedContent(const FContentBundleEditor* ContentBundle) override;
 	virtual void NotifyContentBundleChanged(const FContentBundleEditor* ContentBundle) override { OnContentBundleChanged().Broadcast(ContentBundle); }
 	//~ End IContentBundleEditorSubsystemInterface interface
 
diff --git a/Engine/Source/Runtime/Engine/Private/WorldPartition/ContentBundle/ContentBundleEditor.cpp b/Engine/Source/Runtime/Engine/Private/WorldPartition/ContentBundle/ContentBundleEditor.cpp
--- a/Engine/Source/Runtime/Engine/Private/WorldPartition/ContentBundle/ContentBundleEditor.cpp
+++ b/Engine/Source/Runtime/Engine/Private/WorldPartition/ContentBundle/ContentBundleEditor.cpp
@@ -85,14 +85,24 @@ void FContentBundleEditor::DoInjectContent()
 		SetStatus(EContentBundleStatus::FailedToInject);
 	}
 
-
-	BroadcastChanged();
+	if (IContentBundleEditorSubsystemInterface* EditorSubsystem = IContentBundleEditorSubsystemInterface::Get())
+	{
+		EditorSubsystem->NotifyContentBundleInjectedContent(this);
+	}
 }
 
 void FContentBundleEditor::DoRemoveContent()
 {
 	UnreferenceAllActors();
 
+	if (IContentBundleEditorSubsystemInterface* EditorSubsystem = IContentBundleEditorSubsystemInterface::Get())
+	{
+		EditorSubsystem->NotifyContentBundleRemovedContent(this);
+	}
+
+	// Make sure the content bundle is no longer flagged as edited before unloading actors
+	check(!IsBeingEdited());
+
 	WorldDataLayersActorReference.Reset();
 
 	UnsavedActorMonitor->Uninitialize();
@@ -119,6 +129,8 @@ bool FContentBundleEditor::IsValid() const
 
 bool FContentBundleEditor::AddActor(AActor* InActor)
 {
+	check(GetStatus() == EContentBundleStatus::ContentInjected || GetStatus() == EContentBundleStatus::ReadyToInject);
+
 	if (InActor->GetWorld() != ActorDescContainer->GetWorld() || InActor->HasAllFlags(RF_Transient) || !InActor->IsMainPackageActor())
 	{
 		return false;
diff --git a/Engine/Source/Runtime/Engine/Public/WorldPartition/ContentBundle/ContentBundleEditorSubsystemInterface.h b/Engine/Source/Runtime/Engine/Public/WorldPartition/ContentBundle/ContentBundleEditorSubsystemInterface.h
--- a/Engine/Source/Runtime/Engine/Public/WorldPartition/ContentBundle/ContentBundleEditorSubsystemInterface.h
+++ b/Engine/Source/Runtime/Engine/Public/WorldPartition/ContentBundle/ContentBundleEditorSubsystemInterface.h
@@ -13,6 +13,8 @@ public:
 
 	virtual void NotifyContentBundleAdded(const FContentBundleEditor* ContentBundle) = 0;
 	virtual void NotifyContentBundleRemoved(const FContentBundleEditor* ContentBundle) = 0;
+	virtual void NotifyContentBundleInjectedContent(const FContentBundleEditor* ContentBundle) = 0;
+	virtual void NotifyContentBundleRemovedContent(const FContentBundleEditor* ContentBundle) = 0;
 	virtual void NotifyContentBundleChanged(const FContentBundleEditor* ContentBundle) = 0;
 
 protected:
