Prevent calling GenerateStreaming on content bundle when cooking.
It is not the right flow and will be adressed in a subsequent CL

#test cook

diff --git a/Engine/Source/Runtime/Engine/Private/WorldPartition/ContentBundle/ContentBundleContainer.cpp b/Engine/Source/Runtime/Engine/Private/WorldPartition/ContentBundle/ContentBundleContainer.cpp
--- a/Engine/Source/Runtime/Engine/Private/WorldPartition/ContentBundle/ContentBundleContainer.cpp
+++ b/Engine/Source/Runtime/Engine/Private/WorldPartition/ContentBundle/ContentBundleContainer.cpp
@@ -355,6 +355,14 @@ void FContentBundleContainer::OnPreGenerateStreaming()
 {
 	UE_LOG(LogContentBundle, Log, TEXT("[Container: %s] Generating Streaming for %u Content Bundles."), *GetInjectedWorld()->GetName(), GetEditorContentBundles().Num());
 
+	if (IsRunningCookCommandlet())
+	{
+		// Cooking for plugin is not in yet.
+		// Avoid calling OnPreGenerateStreaming from the cook package splitter.
+		// It is not the correct flow & It outputs an error in the build log.
+		return;
+	}
+
 	GetInjectedWorld()->ContentBundleManager->GetPIEDuplicateHelper()->Clear();
 
 	for (TSharedPtr<FContentBundleEditor>& ContentBundle : GetEditorContentBundles())
