ILevelInstanceInterface::OnLevelInstanceLoaded is now called in non editor build.
#rnx

diff --git a/Engine/Source/Runtime/Engine/Private/LevelInstance/LevelInstanceSubsystem.cpp b/Engine/Source/Runtime/Engine/Private/LevelInstance/LevelInstanceSubsystem.cpp
--- a/Engine/Source/Runtime/Engine/Private/LevelInstance/LevelInstanceSubsystem.cpp
+++ b/Engine/Source/Runtime/Engine/Private/LevelInstance/LevelInstanceSubsystem.cpp
@@ -270,7 +270,7 @@ void ULevelInstanceSubsystem::RegisterLoadedLevelStreamingLevelInstance(ULevelSt
 	check(!LoadedLevelInstances.Contains(LevelInstanceID));
 	FLevelInstance& LevelInstanceEntry = LoadedLevelInstances.Add(LevelInstanceID);
 	LevelInstanceEntry.LevelStreaming = LevelStreaming;
-#if WITH_EDITOR
+
 	// LevelInstanceID might not be registered anymore in the case where the level instance 
 	// was unloaded while still being load.
 	if (ILevelInstanceInterface* LevelInstance = LevelStreaming->GetLevelInstance())
@@ -283,7 +283,6 @@ void ULevelInstanceSubsystem::RegisterLoadedLevelStreamingLevelInstance(ULevelSt
 		// Validate that the LevelInstanceID is requested to be unloaded
 		check(LevelInstancesToUnload.Contains(LevelInstanceID));
 	}
-#endif
 }
 
 #if WITH_EDITOR
