- Fixes for Cascade and LWC
- Renderers all pack vertices in LWC tile space and apply the tile inside the VS
- Changed ApplyWorldOffset to reset the system as it fails for various emitter types

diff --git a/Engine/Shaders/Private/MeshParticleVertexFactory.ush b/Engine/Shaders/Private/MeshParticleVertexFactory.ush
--- a/Engine/Shaders/Private/MeshParticleVertexFactory.ush
+++ b/Engine/Shaders/Private/MeshParticleVertexFactory.ush
@@ -176,8 +176,8 @@ struct FVertexFactoryIntermediates
 	/** Relative time. */
 	float RelativeTime;
 
-	float4x4 ParticleToWorld;
-	float4x4 WorldToParticle;
+	FLWCMatrix ParticleToWorld;
+	FLWCInverseMatrix WorldToParticle;
 
 	/** Cached primitive and instance data */
 	FSceneDataIntermediates SceneData;
@@ -188,18 +188,33 @@ FPrimitiveSceneData GetPrimitiveData(FVertexFactoryIntermediates Intermediates)
 	return Intermediates.SceneData.Primitive;
 }
 
-float4x4 GetParticleTransform(FVertexFactoryInput Input)
+FLWCMatrix GetParticleTransform(FVertexFactoryInput Input)
 {
 	//Transform is transposed on CPU and position is packed into W. Transpose again to get the correct transform back.
 #if PARTICLE_MESH_INSTANCED
-	return transpose(float4x4(Input.Transform1, Input.Transform2, Input.Transform3, float4(0, 0, 0, 1)));
+	float4x4 Transform4x4 = float4x4(Input.Transform1, Input.Transform2, Input.Transform3, float4(0, 0, 0, 1));
 #else
-	return transpose(float4x4(Transform1, Transform2, Transform3, float4(0, 0, 0, 1)));
+	float4x4 Transform4x4 = float4x4(Transform1, Transform2, Transform3, float4(0, 0, 0, 1));
 #endif
+	Transform4x4 = transpose(Transform4x4);
+
+	FLWCMatrix FinalMat;
+	if ( MeshParticleVF.bUseLocalSpace )
+	{
+		FinalMat = LWCMultiply(Transform4x4, GetPrimitiveDataFromUniformBuffer().LocalToWorld);
+	}
+	else
+	{
+		FinalMat = LWCPromote(Transform4x4);
+		FLWCVector3 LWCTileOffset = MakeLWCVector3(MeshParticleVF.LWCTile, 0);
+		LWCSetOrigin(FinalMat, LWCAdd(LWCGetOrigin(FinalMat), LWCTileOffset));
+	}
+
+	return FinalMat;
 }
 
 #if USE_PREVIOUS_TRANSFORM
-float4x4 GetParticlePrevTransform(FVertexFactoryInput Input)
+FLWCMatrix GetParticlePrevTransform(FVertexFactoryInput Input)
 {
 #if PARTICLE_MESH_INSTANCED
 	uint TransformBaseIdx = GetInstanceId(Input.InstanceId) * 3;
@@ -219,27 +234,52 @@ float4x4 GetParticlePrevTransform(FVertexFactoryInput Input)
 		);
 #endif
 	//Transform is transposed on CPU and position is packed into W. Transpose again to get the correct transform back.
-	return transpose(PrevTransform4x4);
+	PrevTransform4x4 = transpose(PrevTransform4x4);
+
+	FLWCMatrix FinalMat;
+	if ( MeshParticleVF.bUseLocalSpace )
+	{
+		FinalMat = LWCMultiply(PrevTransform4x4, GetPrimitiveDataFromUniformBuffer().PreviousLocalToWorld);
+	}
+	else
+	{
+		FinalMat = LWCPromote(PrevTransform4x4);
+		FLWCVector3 LWCTileOffset = MakeLWCVector3(MeshParticleVF.LWCTile, 0);
+		LWCSetOrigin(FinalMat, LWCAdd(LWCGetOrigin(FinalMat), LWCTileOffset));
+	}
+	return FinalMat;
 }
 #endif
 
-float4x4 PseudoInverseTransform(float4x4 LocalToWorld)
+FLWCInverseMatrix GetParticleInverseTransform(FVertexFactoryInput Input)
 {
+	//Transform is transposed on CPU and position is packed into W. Transpose again to get the correct transform back.
+#if PARTICLE_MESH_INSTANCED
+	float4x4 Transform = transpose(float4x4(Input.Transform1, Input.Transform2, Input.Transform3, float4(0, 0, 0, 1)));
+#else
+	float4x4 Transform = transpose(float4x4(Transform1, Transform2, Transform3, float4(0, 0, 0, 1)));
+#endif
+
 	float3 Scale2;
-	Scale2.x = length2(LocalToWorld[0].xyz);
-	Scale2.y = length2(LocalToWorld[1].xyz);
-	Scale2.z = length2(LocalToWorld[2].xyz);
+	Scale2.x = length2(Transform[0].xyz);
+	Scale2.y = length2(Transform[1].xyz);
+	Scale2.z = length2(Transform[2].xyz);
 
 	float3 InvNonUniformScale = rsqrt(Scale2);
 
 	float4x4 WorldToLocal = 0;
-	WorldToLocal[0].xyz = LocalToWorld[0].xyz * Pow2(InvNonUniformScale.x);
-	WorldToLocal[1].xyz	= LocalToWorld[1].xyz * Pow2(InvNonUniformScale.y);
-	WorldToLocal[2].xyz	= LocalToWorld[2].xyz * Pow2(InvNonUniformScale.z);
+	WorldToLocal[0].xyz = Transform[0].xyz * Pow2(InvNonUniformScale.x);
+	WorldToLocal[1].xyz	= Transform[1].xyz * Pow2(InvNonUniformScale.y);
+	WorldToLocal[2].xyz	= Transform[2].xyz * Pow2(InvNonUniformScale.z);
 	WorldToLocal[3].w	= 1.0f;
 	WorldToLocal		= transpose(WorldToLocal);
-	WorldToLocal[3].xyz	= mul(float4(-LocalToWorld[3].xyz, 0.0f), WorldToLocal).xyz;
-	return WorldToLocal;
+
+	FLWCVector3 LWCTileOffset = MakeLWCVector3(MeshParticleVF.LWCTile, 0);
+	FLWCVector3 Translation = LWCPromote(mul(float4(-Transform[3].xyz, 0.0f), WorldToLocal).xyz);
+	Translation = LWCAdd(Translation, LWCTileOffset);
+	WorldToLocal[3].xyz = Translation.Offset;
+
+	return MakeLWCInverseMatrix(LWCGetTile(Translation), WorldToLocal);
 }
 
 /** Converts from vertex factory specific interpolants to a FMaterialPixelParameters, which is used by material inputs. */
@@ -304,6 +344,7 @@ FMaterialPixelParameters GetMaterialPixelParameters(FVertexFactoryInterpolantsVS
 #endif
 
 #if USE_PARTICLE_LOCAL_TO_WORLD
+	//-TODO: LWC Precision Loss
 	FLWCMatrix ParticleToWorld = LWCPromote(transpose(float4x4(Interpolants.ParticleToWorld[0], Interpolants.ParticleToWorld[1], Interpolants.ParticleToWorld[2], float4(0.0f, 0.0f, 0.0f, 1.0f))));
 	#if NEEDS_PARTICLE_LOCAL_TO_WORLD
 		Result.Particle.ParticleToWorld = ParticleToWorld;
@@ -314,6 +355,7 @@ FMaterialPixelParameters GetMaterialPixelParameters(FVertexFactoryInterpolantsVS
 #endif
 
 #if USE_PARTICLE_WORLD_TO_LOCAL
+	//-TODO: LWC Precision Loss
 	FLWCInverseMatrix WorldToParticle = LWCPromoteInverse(transpose(float4x4(Interpolants.WorldToParticle[0], Interpolants.WorldToParticle[1], Interpolants.WorldToParticle[2], float4(0.0f, 0.0f, 0.0f, 1.0f))));
 	#if NEEDS_PARTICLE_WORLD_TO_LOCAL
 		Result.Particle.WorldToParticle = WorldToParticle;
@@ -353,20 +395,20 @@ FMaterialVertexParameters GetMaterialVertexParameters(FVertexFactoryInput Input,
 	Result.Particle.DynamicParameter = Intermediates.DynamicParameter;
 #endif	//USE_DYNAMIC_PARAMETERS
 
-	float4x4 ParticleTransform = GetParticleTransform(Input);
-	Result.InstanceLocalToWorld = LWCPromote(ParticleTransform);
-	Result.InstanceWorldToLocal = LWCPromoteInverse(PseudoInverseTransform(ParticleTransform));
+	FLWCMatrix ParticleTransform = GetParticleTransform(Input);
+	Result.InstanceLocalToWorld = ParticleTransform;
+	Result.InstanceWorldToLocal = GetParticleInverseTransform(Input);
 
 	Result.Particle.Color = Intermediates.ParticleColor;
 	Result.Particle.ParticleToWorld = Result.InstanceLocalToWorld;
 	Result.Particle.WorldToParticle = Result.InstanceWorldToLocal;
 
-	Result.TangentToWorld = mul(TangentToLocal, (float3x3)ParticleTransform ); 
+	Result.TangentToWorld = mul(TangentToLocal, LWCToFloat3x3(ParticleTransform)); 
 
 #if USE_PREVIOUS_TRANSFORM
-	Result.PrevFrameLocalToWorld = LWCPromote(GetParticlePrevTransform(Input));
+	Result.PrevFrameLocalToWorld = GetParticlePrevTransform(Input);
 #else
-	Result.PrevFrameLocalToWorld = LWCPromote(ParticleTransform);
+	Result.PrevFrameLocalToWorld = ParticleTransform;
 #endif
 
 #if NUM_MATERIAL_TEXCOORDS_VERTEX
@@ -381,8 +423,8 @@ FMaterialVertexParameters GetMaterialVertexParameters(FVertexFactoryInput Input,
 
 float4 CalcWorldPosition(FVertexFactoryInput Input)
 {
-	float4x4 Transform = GetParticleTransform(Input);
-	float3 WorldPosition = mul(Input.Position, Transform).xyz + LWCHackToFloat(ResolvedView.PreViewTranslation);
+	FLWCMatrix Transform = GetParticleTransform(Input);
+	float3 WorldPosition = LWCToFloat(LWCAdd(LWCMultiply(Input.Position.xyz, Transform), ResolvedView.PreViewTranslation));
 	return float4(WorldPosition, Input.Position.w);
 }
 
@@ -391,8 +433,8 @@ float4 CalcPrevWorldPosition(FVertexFactoryInput Input)
 #if USE_PREVIOUS_TRANSFORM
 	if (MeshParticleVF.PrevTransformAvailable > 0)
 	{
-		float4x4 PrevTransform = GetParticlePrevTransform(Input);
-		float3 PrevWorldPosition = mul(Input.Position, PrevTransform).xyz + LWCHackToFloat(ResolvedView.PrevPreViewTranslation);
+		FLWCMatrix PrevTransform = GetParticlePrevTransform(Input);
+		float3 PrevWorldPosition = LWCToFloat(LWCAdd(LWCMultiply(Input.Position.xyz, PrevTransform), ResolvedView.PrevPreViewTranslation));
 		return float4(PrevWorldPosition, Input.Position.w);
 	}
 	else
@@ -410,14 +452,13 @@ FVertexFactoryIntermediates GetVertexFactoryIntermediates(FVertexFactoryInput In
 	Intermediates.VertexColor = Input.VertexColor FCOLOR_COMPONENT_SWIZZLE;
 
 	// World position and localtoworld transform.
-	float4x4 LocalToWorldMat = GetParticleTransform(Input);
+	FLWCMatrix LocalToWorldMat = GetParticleTransform(Input);
 
 	Intermediates.ParticleToWorld = LocalToWorldMat;
-	Intermediates.WorldToParticle = PseudoInverseTransform(LocalToWorldMat);
-
-	float3 ParticleWorldPosition = LocalToWorldMat[3].xyz;
+	Intermediates.WorldToParticle = GetParticleInverseTransform(Input);
 
-	Intermediates.ParticleTranslatedWorldPosition = ParticleWorldPosition + LWCHackToFloat(ResolvedView.PreViewTranslation);
+	FLWCVector3 ParticleWorldPosition = LWCGetOrigin(LocalToWorldMat);
+	Intermediates.ParticleTranslatedWorldPosition = LWCToFloat(LWCAdd(ParticleWorldPosition, ResolvedView.PreViewTranslation));
 
 #if PARTICLE_MESH_INSTANCED
 	Intermediates.ParticleColor = Input.ParticleColor;
@@ -458,7 +499,7 @@ FVertexFactoryIntermediates GetVertexFactoryIntermediates(FVertexFactoryInput In
 	#endif
 #endif	//USE_DYNAMIC_PARAMETERS
 
-	Intermediates.ParticleTranslatedWorldPosition = ParticleWorldPosition.xyz + LWCHackToFloat(ResolvedView.PreViewTranslation);
+	Intermediates.ParticleTranslatedWorldPosition = LWCToFloat(LWCAdd(ParticleWorldPosition, ResolvedView.PreViewTranslation));
 
 #if PARTICLE_MESH_INSTANCED
 	Intermediates.ParticleVelocity = Input.ParticleVelocity;
@@ -525,8 +566,8 @@ void CalcTangentToWorld(FVertexFactoryInput Input, FVertexFactoryIntermediates I
 {
 	float3x3 TangentToLocal = VertexFactoryGetTangentToLocal(Input, Intermediates);
 
-	float4x4 Transform = GetParticleTransform(Input);
-	float3x3 TangentToWorld = mul(TangentToLocal, (float3x3)Transform);
+	FLWCMatrix Transform = GetParticleTransform(Input);
+	float3x3 TangentToWorld = mul(TangentToLocal, LWCToFloat3x3(Transform));
 
 	// Normalize to remove scaling. Incorrect but faster than computing the inverse-transpose.
 	TangentToWorld0 = normalize(TangentToWorld[0]);
@@ -598,14 +639,16 @@ FVertexFactoryInterpolantsVSToPS VertexFactoryGetInterpolantsVSToPS(FVertexFacto
 #endif
 
 #if USE_PARTICLE_LOCAL_TO_WORLD
-	float4x4 ParticleToWorldT = transpose(Intermediates.ParticleToWorld);
+	//-LWC_TODO: LWC Precision Loss
+	float4x4 ParticleToWorldT = transpose(LWCHackToFloat(Intermediates.ParticleToWorld));
 	Interpolants.ParticleToWorld[0] = ParticleToWorldT[0];
 	Interpolants.ParticleToWorld[1] = ParticleToWorldT[1];
 	Interpolants.ParticleToWorld[2] = ParticleToWorldT[2];
 #endif
 
 #if USE_PARTICLE_WORLD_TO_LOCAL
-	float4x4 WorldToParticleT = transpose(Intermediates.WorldToParticle);
+	//-LWC_TODO: LWC Precision Loss
+	float4x4 WorldToParticleT = transpose(LWCHackToFloat(Intermediates.WorldToParticle));
 	Interpolants.WorldToParticle[0] = WorldToParticleT[0];
 	Interpolants.WorldToParticle[1] = WorldToParticleT[1];
 	Interpolants.WorldToParticle[2] = WorldToParticleT[2];
diff --git a/Engine/Shaders/Private/ParticleBeamTrailVertexFactory.ush b/Engine/Shaders/Private/ParticleBeamTrailVertexFactory.ush
--- a/Engine/Shaders/Private/ParticleBeamTrailVertexFactory.ush
+++ b/Engine/Shaders/Private/ParticleBeamTrailVertexFactory.ush
@@ -130,10 +130,26 @@ float3 SafeNormalize(float3 V)
 	return V / sqrt(max(dot(V,V),0.01));
 }
 
+float4 SimulationToTranslatedWorld(float3 SimPosition)
+{
+	FLWCVector3 WorldPosition;
+	if ( BeamTrailVF.bUseLocalSpace )
+	{
+		WorldPosition = LWCMultiply(SimPosition, GetPrimitiveDataFromUniformBuffer().LocalToWorld);
+	}
+	else
+	{
+		WorldPosition = MakeLWCVector3(BeamTrailVF.LWCTile, SimPosition);
+	}
+
+	FLWCVector3 TranslatedWorldPosition = LWCAdd(WorldPosition, ResolvedView.PreViewTranslation);
+	return float4(LWCToFloat(TranslatedWorldPosition), 1.0f);
+}
+
 void GetTangents(FVertexFactoryInput Input,out float4 Right,out float4 Up)
 {
-	float4	Position			= TransformLocalToTranslatedWorld(Input.Position.xyz),
-			OldPosition			= TransformLocalToTranslatedWorld(Input.OldPosition.xyz);
+	float4	Position			= SimulationToTranslatedWorld(Input.Position.xyz),
+			OldPosition			= SimulationToTranslatedWorld(Input.OldPosition.xyz);
 
 	float3	CameraDirection		= SafeNormalize(ResolvedView.TranslatedWorldCameraOrigin - Position.xyz),
 			ParticleDirection	= SafeNormalize(Position.xyz - OldPosition.xyz);
@@ -159,7 +175,7 @@ void GetTangents(FVertexFactoryInput Input,out float4 Right,out float4 Up)
 
 float4 CalcWorldPosition(FVertexFactoryInput Input)
 {
-	return TransformLocalToTranslatedWorld(Input.Position.xyz);
+	return SimulationToTranslatedWorld(Input.Position.xyz);
 }
 
 /** derive basis vectors */
diff --git a/Engine/Shaders/Private/ParticleSpriteVertexFactory.ush b/Engine/Shaders/Private/ParticleSpriteVertexFactory.ush
--- a/Engine/Shaders/Private/ParticleSpriteVertexFactory.ush
+++ b/Engine/Shaders/Private/ParticleSpriteVertexFactory.ush
@@ -93,7 +93,7 @@ struct FVertexFactoryInterpolantsVSToPS
 struct FVertexFactoryIntermediates
 {
 	/** The position of the particle in non-translated world space. */
-	float3 ParticleWorldPosition;
+	FLWCVector3 ParticleWorldPosition;
 	/** The position of the vertex in translated world space. */
 	float3 VertexWorldPosition;
 	/** Particle translated world space position and size. */
@@ -310,7 +310,7 @@ float3 SafeNormalize(float3 V)
 	return V * rsqrt(max(dot(V,V),0.00000001));
 }
 
-void GetTangents(FVertexFactoryInput Input, float3 WorldPosition, float3 OldWorldPosition, out float3 OutRight, out float3 OutUp)
+void GetTangents(FVertexFactoryInput Input, FLWCVector3 WorldPosition, FLWCVector3 OldWorldPosition, out float3 OutRight, out float3 OutUp)
 {
 	// Select camera up/right vectors.
 	float3 ResolvedViewRight = lerp(ResolvedView.ViewRight, ResolvedView.HMDViewNoRollRight, SpriteVF.RemoveHMDRoll);
@@ -320,7 +320,7 @@ void GetTangents(FVertexFactoryInput Input, float3 WorldPosition, float3 OldWorl
 	float3 CameraUp = lerp(-ResolvedViewUp, SpriteVF.AxisLockUp.xyz, SpriteVF.AxisLockUp.w);
 
 	// Determine the vector from the particle to the camera and the particle's movement direction.
-	float3 CameraDirection	= SafeNormalize(LWCHackToFloat(ResolvedView.WorldCameraOrigin) - WorldPosition.xyz);
+	float3 CameraDirection	= SafeNormalize(LWCToFloat(LWCSubtract(ResolvedView.WorldCameraOrigin, WorldPosition)));
 	float3 RightVector = CameraRight.xyz;
 	float3 UpVector = CameraUp.xyz;
 
@@ -330,7 +330,7 @@ void GetTangents(FVertexFactoryInput Input, float3 WorldPosition, float3 OldWorl
 	if (SpriteVF.CameraFacingBlend.x > 0.f)
 	{
 		// Blend between PSA_FacingCamera and PSA_Square over distance
-		float3 CameraVector = LWCHackToFloat(ResolvedView.WorldCameraOrigin) - WorldPosition.xyz;
+		float3 CameraVector = LWCToFloat(LWCSubtract(ResolvedView.WorldCameraOrigin, WorldPosition));
 		float CameraDistanceSq = dot(CameraVector, CameraVector);
 		float AlignmentMode = saturate(CameraDistanceSq * SpriteVF.CameraFacingBlend.y - SpriteVF.CameraFacingBlend.z);
 
@@ -346,7 +346,7 @@ void GetTangents(FVertexFactoryInput Input, float3 WorldPosition, float3 OldWorl
 		if (LocalTangentSelector.y > 0)
 		{
 			// Tangent vectors for PSA_Velocity.
-			float3 ParticleDirection = SafeNormalize(WorldPosition.xyz - OldWorldPosition.xyz);
+			float3 ParticleDirection = SafeNormalize(LWCToFloat(LWCSubtract(WorldPosition, OldWorldPosition)));
 			RightVector	= SafeNormalize(cross(CameraDirection, ParticleDirection));
 			UpVector = -ParticleDirection;
 		}
@@ -397,14 +397,14 @@ float3x3 CalcTangentBasis(FVertexFactoryIntermediates Intermediates)
 	// ENM_Spherical
 	else if (SpriteVF.NormalsType < 1.5f)
 	{
-		float3 TangentZ = normalize(Intermediates.ParticleWorldPosition.xyz - SpriteVF.NormalsSphereCenter.xyz);
+		float3 TangentZ = normalize(LWCHackToFloat(Intermediates.ParticleWorldPosition) - SpriteVF.NormalsSphereCenter.xyz);
 		Result[2] = TangentZ;	
 	}
 	// ENM_Cylindrical
 	else
 	{
-		float3 ClosestPointOnCylinder = SpriteVF.NormalsSphereCenter.xyz + dot(SpriteVF.NormalsCylinderUnitDirection.xyz, Intermediates.ParticleWorldPosition.xyz - SpriteVF.NormalsSphereCenter.xyz) * SpriteVF.NormalsCylinderUnitDirection.xyz;
-		float3 TangentZ = normalize(Intermediates.ParticleWorldPosition.xyz - ClosestPointOnCylinder);
+		float3 ClosestPointOnCylinder = SpriteVF.NormalsSphereCenter.xyz + dot(SpriteVF.NormalsCylinderUnitDirection.xyz, LWCHackToFloat(Intermediates.ParticleWorldPosition) - SpriteVF.NormalsSphereCenter.xyz) * SpriteVF.NormalsCylinderUnitDirection.xyz;
+		float3 TangentZ = normalize(LWCHackToFloat(Intermediates.ParticleWorldPosition) - ClosestPointOnCylinder);
 		Result[2] = TangentZ;	
 	}
 
@@ -473,13 +473,14 @@ FVertexFactoryIntermediates GetVertexFactoryIntermediates(FVertexFactoryInput In
 	Intermediates.SceneData = VF_GPUSCENE_GET_INTERMEDIATES(Input);
 
 	// World position.
-	float3 ParticleWorldPosition = LWCHackToFloat(LWCMultiply(Input.Position.xyz, GetPrimitiveDataFromUniformBuffer().LocalToWorld));
-	float3 ParticleOldWorldPosition = LWCHackToFloat(LWCMultiply(Input.OldPosition.xyz, GetPrimitiveDataFromUniformBuffer().LocalToWorld));
-	Intermediates.ParticleWorldPosition = ParticleWorldPosition.xyz;
+	FLWCVector3 LWCTileOffset = MakeLWCVector3(SpriteVF.LWCTile, 0);
+	FLWCVector3 ParticleWorldPosition = LWCAdd(LWCMultiply(Input.Position.xyz, GetPrimitiveDataFromUniformBuffer().LocalToWorld), LWCTileOffset);
+	FLWCVector3 ParticleOldWorldPosition = LWCAdd(LWCMultiply(Input.OldPosition.xyz, GetPrimitiveDataFromUniformBuffer().LocalToWorld), LWCTileOffset);
+	Intermediates.ParticleWorldPosition = ParticleWorldPosition;
 
 	// Tangents.
 	float3 Right,Up;
-	GetTangents(Input,ParticleWorldPosition.xyz,ParticleOldWorldPosition.xyz,Right,Up);
+	GetTangents(Input, ParticleWorldPosition, ParticleOldWorldPosition, Right, Up);
 	Intermediates.TangentUp = Up;
 	Intermediates.TangentRight = Right;
 
@@ -489,7 +490,7 @@ FVertexFactoryIntermediates GetVertexFactoryIntermediates(FVertexFactoryInput In
 	ComputeBillboardUVs(Input, UVForPosition, UVForTexturing, UVForTexturingUnflipped);
 
 	// Vertex position.
-	float4 VertexWorldPosition = float4(ParticleWorldPosition.xyz + LWCHackToFloat(ResolvedView.PreViewTranslation), 1);
+	float4 VertexWorldPosition = float4(LWCToFloat(LWCAdd(ParticleWorldPosition, ResolvedView.PreViewTranslation)), 1);
 	float2 Size = abs(Input.SizeRotSubImage.xy);
 	VertexWorldPosition += Size.x * (UVForPosition.x + SpriteVF.PivotOffset.x) * float4(Right,0);
 	VertexWorldPosition += Size.y * (UVForPosition.y + SpriteVF.PivotOffset.y) * float4(Up,0);
@@ -523,7 +524,7 @@ FVertexFactoryIntermediates GetVertexFactoryIntermediates(FVertexFactoryInput In
 
 #if USE_PARTICLE_POSITION
 	float ParticleRadius = .5f * min(Input.SizeRotSubImage.x, Input.SizeRotSubImage.y);
-	float3 ParticleTranslatedWorldPosition = ParticleWorldPosition.xyz + LWCHackToFloat(ResolvedView.PreViewTranslation);
+	float3 ParticleTranslatedWorldPosition = LWCToFloat(LWCAdd(ParticleWorldPosition, ResolvedView.PreViewTranslation));
 	Intermediates.TranslatedWorldPositionAndSize = float4(ParticleTranslatedWorldPosition, ParticleRadius);
 
 #if GENERATE_SPHERICAL_PARTICLE_NORMALS
@@ -533,7 +534,7 @@ FVertexFactoryIntermediates GetVertexFactoryIntermediates(FVertexFactoryInput In
 #endif
 
 #if USE_PARTICLE_VELOCITY
-	float3 ParticleVelocity = (ParticleWorldPosition.xyz - ParticleOldWorldPosition.xyz) * SpriteVF.InvDeltaSeconds;
+	float3 ParticleVelocity = LWCToFloat(LWCSubtract(ParticleWorldPosition, ParticleOldWorldPosition)) * SpriteVF.InvDeltaSeconds;
 	Intermediates.ParticleVelocity.xyz = normalize(ParticleVelocity);
 	Intermediates.ParticleVelocity.w = length(ParticleVelocity);
 #endif
@@ -585,7 +586,7 @@ float3 VertexFactoryGetPositionForVertexLighting(FVertexFactoryInput Input, FVer
 {
 #if SUPPORTS_PARTICLE_CUTOUTS && (TRANSLUCENCY_LIGHTING_VOLUMETRIC_PERVERTEX_DIRECTIONAL || TRANSLUCENCY_LIGHTING_VOLUMETRIC_PERVERTEX_NONDIRECTIONAL)
 	// Do per-vertex lighting with particle position instead of vertex position when we're using SubUV cutouts, since the vertex positions vary per-frame which would cause popping
-	return NumCutoutVerticesPerFrame > 0 ? (Intermediates.ParticleWorldPosition.xyz + LWCHackToFloat(ResolvedView.PreViewTranslation)) : TranslatedWorldPosition;
+	return NumCutoutVerticesPerFrame > 0 ? LWCToFloat(LWCAdd(Intermediates.ParticleWorldPosition, ResolvedView.PreViewTranslation)) : TranslatedWorldPosition;
 #else
 	return TranslatedWorldPosition;
 #endif
diff --git a/Engine/Source/Runtime/Engine/Private/Particles/ParticleComponents.cpp b/Engine/Source/Runtime/Engine/Private/Particles/ParticleComponents.cpp
--- a/Engine/Source/Runtime/Engine/Private/Particles/ParticleComponents.cpp
+++ b/Engine/Source/Runtime/Engine/Private/Particles/ParticleComponents.cpp
@@ -6401,7 +6401,12 @@ void UParticleSystemComponent::ApplyWorldOffset(const FVector& InOffset, bool bW
 {
 	Super::ApplyWorldOffset(InOffset, bWorldShift);
 
-	OldPosition+= InOffset;
+	// Trigger a reset as the offset applying below does not work correctly with all emitter types
+	// Niagara also resets so having Cascade follow the same path makes it consistent also
+	bResetTriggered = true;
+
+#if 0
+	OldPosition += InOffset;
 	
 	for (auto It = EmitterInstances.CreateIterator(); It; ++It)
 	{
@@ -6411,6 +6416,7 @@ void UParticleSystemComponent::ApplyWorldOffset(const FVector& InOffset, bool bW
 			EmitterInstance->ApplyWorldOffset(InOffset, bWorldShift);
 		}
 	}
+#endif
 }
 
 void UParticleSystemComponent::ResetToDefaults()
diff --git a/Engine/Source/Runtime/Engine/Private/Particles/ParticleEmitterInstances.cpp b/Engine/Source/Runtime/Engine/Private/Particles/ParticleEmitterInstances.cpp
--- a/Engine/Source/Runtime/Engine/Private/Particles/ParticleEmitterInstances.cpp
+++ b/Engine/Source/Runtime/Engine/Private/Particles/ParticleEmitterInstances.cpp
@@ -2843,6 +2843,7 @@ bool FParticleEmitterInstance::FillReplayData( FDynamicEmitterReplayDataBase& Ou
 		NewReplayData->RequiredModule = LODLevel->RequiredModule->CreateRendererResource();
 		NewReplayData->MaterialInterface = NULL;	// Must be set by derived implementation
 		NewReplayData->InvDeltaSeconds = (LastDeltaTime > UE_KINDA_SMALL_NUMBER) ? (1.0f / LastDeltaTime) : 0.0f;
+		NewReplayData->LWCTile = LODLevel->RequiredModule->bUseLocalSpace ? FVector3f::Zero() : Component->GetLWCTile();
 
 		NewReplayData->MaxDrawCount =
 			(LODLevel->RequiredModule->bUseMaxDrawCount == true) ? LODLevel->RequiredModule->MaxDrawCount : -1;
diff --git a/Engine/Source/Runtime/Engine/Private/Particles/ParticleSystemRender.cpp b/Engine/Source/Runtime/Engine/Private/Particles/ParticleSystemRender.cpp
--- a/Engine/Source/Runtime/Engine/Private/Particles/ParticleSystemRender.cpp
+++ b/Engine/Source/Runtime/Engine/Private/Particles/ParticleSystemRender.cpp
@@ -440,6 +440,7 @@ bool FDynamicSpriteEmitterData::GetVertexAndIndexData(void* VertexData, void* Dy
 	const uint16* ParticleIndices = Source.DataContainer.ParticleIndices;
 	const FParticleOrder* OrderedIndices = ParticleOrder;
 
+	const FVector LWCTileOffset = FVector(Source.LWCTile) * FLargeWorldRenderScalar::GetTileSize();
 	for (int32 i = 0; i < ParticleCount; i++)
 	{
 		ParticleIndex = OrderedIndices ? OrderedIndices[i].ParticleIndex : i;
@@ -480,9 +481,9 @@ bool FDynamicSpriteEmitterData::GetVertexAndIndexData(void* VertexData, void* Dy
 		for (uint32 Factor = 0; Factor < InstanceFactor; Factor++)
 		{
 			FillVertex = (FParticleSpriteVertex*)TempVert;
-			FillVertex->Position = (FVector3f)ParticlePosition;
+			FillVertex->Position = FVector3f(ParticlePosition - LWCTileOffset);
 			FillVertex->RelativeTime = Particle.RelativeTime;
-			FillVertex->OldPosition = (FVector3f)ParticleOldPosition;
+			FillVertex->OldPosition = FVector3f(ParticleOldPosition - LWCTileOffset);
 			// Create a floating point particle ID from the counter, map into approximately 0-1
 			FillVertex->ParticleId = (Particle.Flags & STATE_CounterMask) / 10000.0f;
 			FillVertex->Size = FVector2f(GetParticleSizeWithUVFlipInSign(Particle, Size));
@@ -548,6 +549,7 @@ bool FDynamicSpriteEmitterData::GetVertexAndIndexDataNonInstanced(void* VertexDa
 	const uint8* ParticleData = Source.DataContainer.ParticleData;
 	const uint16* ParticleIndices = Source.DataContainer.ParticleIndices;
 	const FParticleOrder* OrderedIndices = ParticleOrder;
+	const FVector LWCTileOffset = FVector(Source.LWCTile) * FLargeWorldRenderScalar::GetTileSize();
 
 	for (int32 i = 0; i < ParticleCount; i++)
 	{
@@ -632,9 +634,9 @@ bool FDynamicSpriteEmitterData::GetVertexAndIndexDataNonInstanced(void* VertexDa
 				}
 			}
 
-			FillVertex[VertexIndex].Position	= (FVector3f)ParticlePosition;
+			FillVertex[VertexIndex].Position	= FVector3f(ParticlePosition - LWCTileOffset);
 			FillVertex[VertexIndex].RelativeTime = Particle.RelativeTime;
-			FillVertex[VertexIndex].OldPosition	= (FVector3f)ParticleOldPosition;
+			FillVertex[VertexIndex].OldPosition	= FVector3f(ParticleOldPosition - LWCTileOffset);
 			// Create a floating point particle ID from the counter, map into approximately 0-1
 			FillVertex[VertexIndex].ParticleId = (Particle.Flags & STATE_CounterMask) / 10000.0f;
 			FillVertex[VertexIndex].Size = FVector2f(GetParticleSizeWithUVFlipInSign(Particle, Size));
@@ -1043,6 +1045,7 @@ void FDynamicSpriteEmitterData::UpdateRenderThreadResourcesEmitter(const FPartic
 		UniformParameters.RotationBias = 0.0f;
 		UniformParameters.TangentSelector = FVector4f(0.0f, 0.0f, 0.0f, 0.0f);
 		UniformParameters.InvDeltaSeconds = SourceData->InvDeltaSeconds;
+		UniformParameters.LWCTile = SourceData->LWCTile;
 
 		// Parameters for computing sprite tangents.
 		const FMatrix& LocalToWorld = InOwnerProxy->GetLocalToWorld();
@@ -1381,6 +1384,8 @@ void FDynamicMeshEmitterData::GetDynamicMeshElementsEmitter(const FParticleSyste
 			UniformParameters.TexCoordWeightA = TexCoordWeight;
 			UniformParameters.TexCoordWeightB = 1 - TexCoordWeight;
 			UniformParameters.PrevTransformAvailable = Source.MeshMotionBlurOffset ? 1 : 0;
+			UniformParameters.bUseLocalSpace = Source.bUseLocalSpace;
+			UniformParameters.LWCTile = Source.LWCTile;
 
 			CollectorResources.UniformBuffer = FMeshParticleUniformBufferRef::CreateUniformBufferImmediate(UniformParameters, UniformBuffer_MultiFrame);
 			MeshVertexFactory->SetUniformBuffer(CollectorResources.UniformBuffer);
@@ -1517,7 +1522,7 @@ void FDynamicMeshEmitterData::GetDynamicMeshElementsEmitter(const FParticleSyste
 					Mesh.DepthPriorityGroup = (ESceneDepthPriorityGroup)Proxy->GetDepthPriorityGroup(View);
 
 					FMeshBatchElement& BatchElement = Mesh.Elements[0];
-					BatchElement.PrimitiveUniformBuffer = Proxy->GetWorldSpacePrimitiveUniformBuffer();
+					BatchElement.PrimitiveUniformBuffer = SourceData->bUseLocalSpace ? Proxy->GetUniformBuffer() : Proxy->GetWorldSpacePrimitiveUniformBuffer();
 					BatchElement.FirstIndex = Section.FirstIndex;
 					BatchElement.MinVertexIndex = Section.MinVertexIndex;
 					BatchElement.MaxVertexIndex = Section.MaxVertexIndex;
@@ -1758,10 +1763,11 @@ void FDynamicMeshEmitterData::CalculateParticleTransform(
 	FScaleMatrix kScaleMat(FVector(1.0f));
 	FQuat kLockedAxisQuat = FQuat::Identity;
 
+	const FVector LWCTileOffset = FVector(Source.LWCTile) * FLargeWorldRenderScalar::GetTileSize();
 	FVector ParticlePosition(ParticleLocation + ParticlePayloadCameraOffset);
-	kTransMat.M[3][0] = ParticlePosition.X;
-	kTransMat.M[3][1] = ParticlePosition.Y;
-	kTransMat.M[3][2] = ParticlePosition.Z;
+	kTransMat.M[3][0] = ParticlePosition.X - LWCTileOffset.X;
+	kTransMat.M[3][1] = ParticlePosition.Y - LWCTileOffset.Y;
+	kTransMat.M[3][2] = ParticlePosition.Z - LWCTileOffset.Z;
 
 	FVector3f ScaledSize = ParticleSize * Source.Scale;
 	kScaleMat.M[0][0] = ScaledSize.X;
@@ -2032,11 +2038,6 @@ void FDynamicMeshEmitterData::CalculateParticleTransform(
 
 	FTranslationMatrix OrbitMatrix(OrbitOffset);
 	OutTransformMat *= OrbitMatrix;
-
-	if (Source.bUseLocalSpace)
-	{
-		OutTransformMat *= LocalToWorld;
-	}
 }
 
 void FDynamicMeshEmitterData::GetInstanceData(void* InstanceData, void* DynamicParameterData, void* PrevTransformBuffer, const FParticleSystemSceneProxy* Proxy, const FSceneView* View, uint32 InstanceFactor) const
@@ -2134,31 +2135,31 @@ void FDynamicMeshEmitterData::GetInstanceData(void* InstanceData, void* DynamicP
 
 		// Particle velocity. Calculate on CPU to avoid computing in vertex shader.
 		// Note: It would be preferred if we could check whether the material makes use of the 'Particle Direction' node to avoid this work.
-		FVector3f DeltaPosition(Particle.Location - Particle.OldLocation);
+		FVector DeltaPosition(Particle.Location - Particle.OldLocation);
 
 		int32 CurrentOffset = Source.OrbitModuleOffset;
 		if (CurrentOffset != 0)
 		{
 			FOrbitChainModuleInstancePayload& OrbitPayload = *((FOrbitChainModuleInstancePayload*)((uint8*)&Particle + CurrentOffset));																\
-			DeltaPosition = ((FVector3f)Particle.Location + OrbitPayload.Offset) - ((FVector3f)Particle.OldLocation + OrbitPayload.PreviousOffset);
+			DeltaPosition = (Particle.Location + FVector(OrbitPayload.Offset)) - (Particle.OldLocation + FVector(OrbitPayload.PreviousOffset));
 		}
 
 		if (!DeltaPosition.IsZero())
 		{
 			if (Source.bUseLocalSpace)
 			{
-				DeltaPosition = (FVector4f)Proxy->GetLocalToWorld().TransformVector((FVector)DeltaPosition);
+				DeltaPosition = Proxy->GetLocalToWorld().TransformVector(DeltaPosition);
 			}
-			FVector3f Direction;
+			FVector Direction;
 			float Speed; 
 			DeltaPosition.ToDirectionAndLength(Direction, Speed);
 
 			// Pack direction and speed.
-			CurrentInstanceVertex.Velocity = FVector4f(Direction, Speed);
+			CurrentInstanceVertex.Velocity = FVector4f(FVector3f(Direction), Speed);
 		}
 		else
 		{
-			CurrentInstanceVertex.Velocity = FVector4f();
+			CurrentInstanceVertex.Velocity = FVector4f::Zero();
 		}
 
 		// The particle dynamic value
@@ -2357,6 +2358,8 @@ FParticleBeamTrailUniformBufferRef CreateBeamTrailUniformBuffer(
 
 	// Screen alignment.
 	UniformParameters.ScreenAlignment = FVector4f( (float)SourceData->ScreenAlignment, 0.0f, 0.0f, 0.0f );
+	UniformParameters.bUseLocalSpace = SourceData->bUseLocalSpace;
+	UniformParameters.LWCTile = SourceData->LWCTile;
 
 	return FParticleBeamTrailUniformBufferRef::CreateUniformBufferImmediate( UniformParameters, UniformBuffer_SingleFrame );
 }
@@ -3118,6 +3121,7 @@ int32 FDynamicBeam2EmitterData::FillVertexData_NoNoise(FAsyncBufferFillData& Me)
 
 	int32 PackedCount = 0;
 
+	const FVector LWCTileOffset = FVector(Source.LWCTile) * FLargeWorldRenderScalar::GetTileSize();
 	if (TessFactor <= 1)
 	{
 		for (int32 i = 0; i < Source.ActiveParticleCount; i++)
@@ -3229,8 +3233,8 @@ int32 FDynamicBeam2EmitterData::FillVertexData_NoNoise(FAsyncBufferFillData& Me)
 				Offset.Z		= WorkingUp.Z * Size.X * Taper;
 
 				// 'Lead' edge
-				Vertex->Position	= FVector3f(Location + Offset);
-				Vertex->OldPosition	= (FVector3f)Location;
+				Vertex->Position	= FVector3f(Location + Offset - LWCTileOffset);
+				Vertex->OldPosition	= FVector3f(Location - LWCTileOffset);
 				Vertex->ParticleId	= 0;
 				Vertex->Size		= FVector2f(Size);
 				Vertex->Tex_U		= 0.0f;
@@ -3242,8 +3246,8 @@ int32 FDynamicBeam2EmitterData::FillVertexData_NoNoise(FAsyncBufferFillData& Me)
 				Vertex++;
 				PackedCount++;
 
-				Vertex->Position	= FVector3f(Location - Offset);
-				Vertex->OldPosition	= (FVector3f)Location;
+				Vertex->Position	= FVector3f(Location - Offset - LWCTileOffset);
+				Vertex->OldPosition	= FVector3f(Location - LWCTileOffset);
 				Vertex->ParticleId	= 0;
 				Vertex->Size		= FVector2f(Size);
 				Vertex->Tex_U		= 0.0f;
@@ -3267,8 +3271,8 @@ int32 FDynamicBeam2EmitterData::FillVertexData_NoNoise(FAsyncBufferFillData& Me)
 				Offset.Z		= WorkingUp.Z * Size.X * Taper;
 
 				//
-				Vertex->Position	= FVector3f(EndPoint + Offset);
-				Vertex->OldPosition	= (FVector3f)Particle->OldLocation;
+				Vertex->Position	= FVector3f(EndPoint + Offset - LWCTileOffset);
+				Vertex->OldPosition	= FVector3f(Particle->OldLocation - LWCTileOffset);
 				Vertex->ParticleId	= 0;
 				Vertex->Size		= FVector2f(Size);
 				Vertex->Tex_U		= fUEnd;
@@ -3280,8 +3284,8 @@ int32 FDynamicBeam2EmitterData::FillVertexData_NoNoise(FAsyncBufferFillData& Me)
 				Vertex++;
 				PackedCount++;
 
-				Vertex->Position	= FVector3f(EndPoint - Offset);
-				Vertex->OldPosition	= (FVector3f)Particle->OldLocation;
+				Vertex->Position	= FVector3f(EndPoint - Offset - LWCTileOffset);
+				Vertex->OldPosition	= FVector3f(Particle->OldLocation - LWCTileOffset);
 				Vertex->ParticleId	= 0;
 				Vertex->Size		= FVector2f(Size);
 				Vertex->Tex_U		= fUEnd;
@@ -3409,8 +3413,8 @@ int32 FDynamicBeam2EmitterData::FillVertexData_NoNoise(FAsyncBufferFillData& Me)
 				Offset.Z	= WorkingUp.Z * Size.X * Taper;
 
 				// 'Lead' edge
-				Vertex->Position	= FVector3f(Location + Offset);
-				Vertex->OldPosition	= (FVector3f)Location;
+				Vertex->Position	= FVector3f(Location + Offset - LWCTileOffset);
+				Vertex->OldPosition	= FVector3f(Location - LWCTileOffset);
 				Vertex->ParticleId	= 0;
 				Vertex->Size		= FVector2f(Size);
 				Vertex->Tex_U		= fU;
@@ -3422,8 +3426,8 @@ int32 FDynamicBeam2EmitterData::FillVertexData_NoNoise(FAsyncBufferFillData& Me)
 				Vertex++;
 				PackedCount++;
 
-				Vertex->Position	= FVector3f(Location - Offset);
-				Vertex->OldPosition	= (FVector3f)Location;
+				Vertex->Position	= FVector3f(Location - Offset - LWCTileOffset);
+				Vertex->OldPosition	= FVector3f(Location - LWCTileOffset);
 				Vertex->ParticleId	= 0;
 				Vertex->Size		= FVector2f(Size);
 				Vertex->Tex_U		= fU;
@@ -3471,8 +3475,8 @@ int32 FDynamicBeam2EmitterData::FillVertexData_NoNoise(FAsyncBufferFillData& Me)
 					Offset.Z		= WorkingUp.Z * Size.X * Taper;
 
 					//
-					Vertex->Position	= FVector3f(EndPoint + Offset);
-					Vertex->OldPosition	= (FVector3f)EndPoint;
+					Vertex->Position	= FVector3f(EndPoint + Offset - LWCTileOffset);
+					Vertex->OldPosition	= FVector3f(EndPoint - LWCTileOffset);
 					Vertex->ParticleId	= 0;
 					Vertex->Size		= FVector2f(Size);
 					Vertex->Tex_U		= fU + fTextureIncrement;
@@ -3484,8 +3488,8 @@ int32 FDynamicBeam2EmitterData::FillVertexData_NoNoise(FAsyncBufferFillData& Me)
 					Vertex++;
 				PackedCount++;
 
-					Vertex->Position	= FVector3f(EndPoint - Offset);
-					Vertex->OldPosition	= (FVector3f)EndPoint;
+					Vertex->Position	= FVector3f(EndPoint - Offset - LWCTileOffset);
+					Vertex->OldPosition	= FVector3f(EndPoint - LWCTileOffset);
 					Vertex->ParticleId	= 0;
 					Vertex->Size		= FVector2f(Size);
 					Vertex->Tex_U		= fU + fTextureIncrement;
@@ -3573,6 +3577,8 @@ int32 FDynamicBeam2EmitterData::FillData_Noise(FAsyncBufferFillData& Me) const
 	FMatrix WorldToLocal = Me.WorldToLocal;
 	FMatrix LocalToWorld = Me.LocalToWorld;
 
+	const FVector LWCTileOffset = FVector(Source.LWCTile) * FLargeWorldRenderScalar::GetTileSize();
+
 	// Tessellate the beam along the noise points
 	for (i = 0; i < Source.ActiveParticleCount; i++)
 	{
@@ -3749,8 +3755,8 @@ int32 FDynamicBeam2EmitterData::FillData_Noise(FAsyncBufferFillData& Me) const
 				LastOffset.Z	= WorkingLastUp.Z * Size.X * Taper;
 
 				// 'Lead' edge
-				Vertex->Position	= FVector3f(Location + LastOffset);
-				Vertex->OldPosition	= (FVector3f)Location;
+				Vertex->Position	= FVector3f(Location + LastOffset - LWCTileOffset);
+				Vertex->OldPosition	= FVector3f(Location - LWCTileOffset);
 				Vertex->ParticleId	= 0;
 				Vertex->Size		= FVector2f(Size);
 				Vertex->Tex_U		= fU;
@@ -3760,8 +3766,8 @@ int32 FDynamicBeam2EmitterData::FillData_Noise(FAsyncBufferFillData& Me) const
 				Vertex++;
 				CheckVertexCount++;
 
-				Vertex->Position	= FVector3f(Location - LastOffset);
-				Vertex->OldPosition	= (FVector3f)Location;
+				Vertex->Position	= FVector3f(Location - LastOffset - LWCTileOffset);
+				Vertex->OldPosition	= FVector3f(Location - LWCTileOffset);
 				Vertex->ParticleId	= 0;
 				Vertex->Size		= FVector2f(Size);
 				Vertex->Tex_U		= fU;
@@ -3906,8 +3912,8 @@ int32 FDynamicBeam2EmitterData::FillData_Noise(FAsyncBufferFillData& Me) const
 						Offset.Z	= WorkingUp.Z * Size.X * Taper;
 
 						// Generate the vertex
-						Vertex->Position	= FVector3f(InterpDrawPos + Offset);
-						Vertex->OldPosition	= (FVector3f)InterpDrawPos;
+						Vertex->Position	= FVector3f(InterpDrawPos + Offset - LWCTileOffset);
+						Vertex->OldPosition	= FVector3f(InterpDrawPos - LWCTileOffset);
 						Vertex->ParticleId	= 0;
 						Vertex->Size		= FVector2f(Size);
 						Vertex->Tex_U		= fU;
@@ -3917,8 +3923,8 @@ int32 FDynamicBeam2EmitterData::FillData_Noise(FAsyncBufferFillData& Me) const
 						Vertex++;
 						CheckVertexCount++;
 
-						Vertex->Position	= FVector3f(InterpDrawPos - Offset);
-						Vertex->OldPosition	= (FVector3f)InterpDrawPos;
+						Vertex->Position	= FVector3f(InterpDrawPos - Offset - LWCTileOffset);
+						Vertex->OldPosition	= FVector3f(InterpDrawPos - LWCTileOffset);
 						Vertex->ParticleId	= 0;
 						Vertex->Size		= FVector2f(Size);
 						Vertex->Tex_U		= fU;
@@ -4025,8 +4031,8 @@ int32 FDynamicBeam2EmitterData::FillData_Noise(FAsyncBufferFillData& Me) const
 						Offset.Z	= WorkingUp.Z * Size.X * Taper;
 
 						// Generate the vertex
-						Vertex->Position	= FVector3f(InterpDrawPos + Offset);
-						Vertex->OldPosition	= (FVector3f)InterpDrawPos;
+						Vertex->Position	= FVector3f(InterpDrawPos + Offset - LWCTileOffset);
+						Vertex->OldPosition	= FVector3f(InterpDrawPos - LWCTileOffset);
 						Vertex->ParticleId	= 0;
 						Vertex->Size		= FVector2f(Size);
 						Vertex->Tex_U		= fU;
@@ -4036,8 +4042,8 @@ int32 FDynamicBeam2EmitterData::FillData_Noise(FAsyncBufferFillData& Me) const
 						Vertex++;
 						CheckVertexCount++;
 
-						Vertex->Position	= FVector3f(InterpDrawPos - Offset);
-						Vertex->OldPosition	= (FVector3f)InterpDrawPos;
+						Vertex->Position	= FVector3f(InterpDrawPos - Offset - LWCTileOffset);
+						Vertex->OldPosition	= FVector3f(InterpDrawPos - LWCTileOffset);
 						Vertex->ParticleId	= 0;
 						Vertex->Size		= FVector2f(Size);
 						Vertex->Tex_U		= fU;
@@ -4161,8 +4167,8 @@ int32 FDynamicBeam2EmitterData::FillData_Noise(FAsyncBufferFillData& Me) const
 				LastOffset.Z	= WorkingLastUp.Z * Size.X * Taper;
 
 				// 'Lead' edge
-				Vertex->Position	= FVector3f(Location + LastOffset);
-				Vertex->OldPosition	= (FVector3f)Location;
+				Vertex->Position	= FVector3f(Location + LastOffset - LWCTileOffset);
+				Vertex->OldPosition	= FVector3f(Location - LWCTileOffset);
 				Vertex->ParticleId	= 0;
 				Vertex->Size		= FVector2f(Size);
 				Vertex->Tex_U		= fU;
@@ -4172,8 +4178,8 @@ int32 FDynamicBeam2EmitterData::FillData_Noise(FAsyncBufferFillData& Me) const
 				Vertex++;
 				CheckVertexCount++;
 
-				Vertex->Position	= FVector3f(Location - LastOffset);
-				Vertex->OldPosition	= (FVector3f)Location;
+				Vertex->Position	= FVector3f(Location - LastOffset - LWCTileOffset);
+				Vertex->OldPosition	= FVector3f(Location - LWCTileOffset);
 				Vertex->ParticleId	= 0;
 				Vertex->Size		= FVector2f(Size);
 				Vertex->Tex_U		= fU;
@@ -4320,8 +4326,8 @@ int32 FDynamicBeam2EmitterData::FillData_Noise(FAsyncBufferFillData& Me) const
 						Offset.Z	= WorkingUp.Z * Size.X * Taper;
 
 						// Generate the vertex
-						Vertex->Position	= FVector3f(InterpDrawPos + Offset);
-						Vertex->OldPosition	= (FVector3f)InterpDrawPos;
+						Vertex->Position	= FVector3f(InterpDrawPos + Offset - LWCTileOffset);
+						Vertex->OldPosition	= FVector3f(InterpDrawPos - LWCTileOffset);
 						Vertex->ParticleId	= 0;
 						Vertex->Size		= FVector2f(Size);
 						Vertex->Tex_U		= fU;
@@ -4331,8 +4337,8 @@ int32 FDynamicBeam2EmitterData::FillData_Noise(FAsyncBufferFillData& Me) const
 						Vertex++;
 						CheckVertexCount++;
 
-						Vertex->Position	= FVector3f(InterpDrawPos - Offset);
-						Vertex->OldPosition	= (FVector3f)InterpDrawPos;
+						Vertex->Position	= FVector3f(InterpDrawPos - Offset - LWCTileOffset);
+						Vertex->OldPosition	= FVector3f(InterpDrawPos - LWCTileOffset);
 						Vertex->ParticleId	= 0;
 						Vertex->Size		= FVector2f(Size);
 						Vertex->Tex_U		= fU;
@@ -4439,8 +4445,8 @@ int32 FDynamicBeam2EmitterData::FillData_Noise(FAsyncBufferFillData& Me) const
 						Offset.Z	= WorkingUp.Z * Size.X * Taper;
 
 						// Generate the vertex
-						Vertex->Position	= FVector3f(InterpDrawPos + Offset);
-						Vertex->OldPosition	= (FVector3f)InterpDrawPos;
+						Vertex->Position	= FVector3f(InterpDrawPos + Offset - LWCTileOffset);
+						Vertex->OldPosition	= FVector3f(InterpDrawPos - LWCTileOffset);
 						Vertex->ParticleId	= 0;
 						Vertex->Size		= FVector2f(Size);
 						Vertex->Tex_U		= fU;
@@ -4450,8 +4456,8 @@ int32 FDynamicBeam2EmitterData::FillData_Noise(FAsyncBufferFillData& Me) const
 						Vertex++;
 						CheckVertexCount++;
 
-						Vertex->Position	= FVector3f(InterpDrawPos - Offset);
-						Vertex->OldPosition	= (FVector3f)InterpDrawPos;
+						Vertex->Position	= FVector3f(InterpDrawPos - Offset - LWCTileOffset);
+						Vertex->OldPosition	= FVector3f(InterpDrawPos - LWCTileOffset);
 						Vertex->ParticleId	= 0;
 						Vertex->Size		= FVector2f(Size);
 						Vertex->Tex_U		= fU;
@@ -4533,6 +4539,7 @@ int32 FDynamicBeam2EmitterData::FillData_InterpolatedNoise(FAsyncBufferFillData&
 
 	FMatrix WorldToLocal = Me.WorldToLocal;
 	FMatrix LocalToWorld = Me.LocalToWorld;
+	const FVector LWCTileOffset = FVector(Source.LWCTile) * FLargeWorldRenderScalar::GetTileSize();
 
 	// Tessellate the beam along the noise points
 	for (i = 0; i < Source.ActiveParticleCount; i++)
@@ -4735,8 +4742,8 @@ int32 FDynamicBeam2EmitterData::FillData_InterpolatedNoise(FAsyncBufferFillData&
 			LastOffset.Z	= WorkingLastUp.Z * Size.X * Taper;
 
 			// 'Lead' edge
-			Vertex->Position	= FVector3f(Location + LastOffset);
-			Vertex->OldPosition	= (FVector3f)Location;
+			Vertex->Position	= FVector3f(Location + LastOffset - LWCTileOffset);
+			Vertex->OldPosition	= FVector3f(Location - LWCTileOffset);
 			Vertex->ParticleId	= 0;
 			Vertex->Size		= FVector2f(Size);
 			Vertex->Tex_U		= fU;
@@ -4746,8 +4753,8 @@ int32 FDynamicBeam2EmitterData::FillData_InterpolatedNoise(FAsyncBufferFillData&
 			Vertex++;
 			CheckVertexCount++;
 
-			Vertex->Position	= FVector3f(Location - LastOffset);
-			Vertex->OldPosition	= (FVector3f)Location;
+			Vertex->Position	= FVector3f(Location - LastOffset - LWCTileOffset);
+			Vertex->OldPosition	= FVector3f(Location - LWCTileOffset);
 			Vertex->ParticleId	= 0;
 			Vertex->Size		= FVector2f(Size);
 			Vertex->Tex_U		= fU;
@@ -4942,8 +4949,8 @@ int32 FDynamicBeam2EmitterData::FillData_InterpolatedNoise(FAsyncBufferFillData&
 					Offset.Z	= WorkingUp.Z * Size.X * Taper;
 
 					// Generate the vertex
-					Vertex->Position	= FVector3f(InterpDrawPos + Offset);
-					Vertex->OldPosition	= (FVector3f)InterpDrawPos;
+					Vertex->Position	= FVector3f(InterpDrawPos + Offset - LWCTileOffset);
+					Vertex->OldPosition	= FVector3f(InterpDrawPos - LWCTileOffset);
 					Vertex->ParticleId	= 0;
 					Vertex->Size		= FVector2f(Size);
 					Vertex->Tex_U		= fU;
@@ -4953,8 +4960,8 @@ int32 FDynamicBeam2EmitterData::FillData_InterpolatedNoise(FAsyncBufferFillData&
 					Vertex++;
 					CheckVertexCount++;
 
-					Vertex->Position	= FVector3f(InterpDrawPos - Offset);
-					Vertex->OldPosition	= (FVector3f)InterpDrawPos;
+					Vertex->Position	= FVector3f(InterpDrawPos - Offset - LWCTileOffset);
+					Vertex->OldPosition	= FVector3f(InterpDrawPos - LWCTileOffset);
 					Vertex->ParticleId	= 0;
 					Vertex->Size		= FVector2f(Size);
 					Vertex->Tex_U		= fU;
@@ -5061,8 +5068,8 @@ int32 FDynamicBeam2EmitterData::FillData_InterpolatedNoise(FAsyncBufferFillData&
 					Offset.Z	= WorkingUp.Z * Size.X * Taper;
 
 					// Generate the vertex
-					Vertex->Position	= FVector3f(InterpDrawPos + Offset);
-					Vertex->OldPosition	= (FVector3f)InterpDrawPos;
+					Vertex->Position	= FVector3f(InterpDrawPos + Offset - LWCTileOffset);
+					Vertex->OldPosition	= FVector3f(InterpDrawPos - LWCTileOffset);
 					Vertex->ParticleId	= 0;
 					Vertex->Size		= FVector2f(Size);
 					Vertex->Tex_U		= fU;
@@ -5072,8 +5079,8 @@ int32 FDynamicBeam2EmitterData::FillData_InterpolatedNoise(FAsyncBufferFillData&
 					Vertex++;
 					CheckVertexCount++;
 
-					Vertex->Position	= FVector3f(InterpDrawPos - Offset);
-					Vertex->OldPosition	= (FVector3f)InterpDrawPos;
+					Vertex->Position	= FVector3f(InterpDrawPos - Offset - LWCTileOffset);
+					Vertex->OldPosition	= FVector3f(InterpDrawPos - LWCTileOffset);
 					Vertex->ParticleId	= 0;
 					Vertex->Size		= FVector2f(Size);
 					Vertex->Tex_U		= fU;
@@ -5199,7 +5206,7 @@ void FDynamicTrailsEmitterData::GetDynamicMeshElementsEmitter(const FParticleSys
 	Mesh.VertexFactory = BeamTrailVertexFactory;
 	Mesh.LCI = NULL;
 
-	BatchElement.PrimitiveUniformBuffer = Proxy->GetWorldSpacePrimitiveUniformBuffer();
+	BatchElement.PrimitiveUniformBuffer = SourcePointer->bUseLocalSpace ? Proxy->GetUniformBuffer() : Proxy->GetWorldSpacePrimitiveUniformBuffer();
 	BatchElement.NumPrimitives = OutTriangleCount;
 	BatchElement.MinVertexIndex = 0;
 	BatchElement.MaxVertexIndex = SourcePointer->VertexCount - 1;
@@ -5582,6 +5589,7 @@ int32 FDynamicRibbonEmitterData::FillVertexData(struct FAsyncBufferFillData& Dat
 
 	FBaseParticle* PackingParticle;
 	const uint8* ParticleData = Source.DataContainer.ParticleData;
+	const FVector LWCTileOffset = FVector(Source.LWCTile) * FLargeWorldRenderScalar::GetTileSize();
 	for (int32 ParticleIdx = 0; ParticleIdx < Source.ActiveParticleCount; ParticleIdx++)
 	{
 		DECLARE_PARTICLE_PTR(Particle, ParticleData + Source.ParticleStride * Source.DataContainer.ParticleIndices[ParticleIdx]);
@@ -5694,13 +5702,13 @@ int32 FDynamicRibbonEmitterData::FillVertexData(struct FAsyncBufferFillData& Dat
 					}
 
 					FVector FinalPos = InterpPos + InterpUp * InterpSize;
-					if (Source.bUseLocalSpace)
-					{
-						FinalPos += Data.LocalToWorld.GetOrigin();
-					}
+					//if (Source.bUseLocalSpace)
+					//{
+					//	FinalPos += Data.LocalToWorld.GetOrigin();
+					//}
 					Vertex = (FParticleBeamTrailVertex*)(TempVertexData);
-					Vertex->Position = (FVector3f)FinalPos;
-					Vertex->OldPosition = (FVector3f)FinalPos;
+					Vertex->Position = FVector3f(FinalPos - LWCTileOffset);
+					Vertex->OldPosition = FVector3f(FinalPos - LWCTileOffset);
 					Vertex->ParticleId	= 0;
 					Vertex->Size.X = InterpSize;
 					Vertex->Size.Y = InterpSize;
@@ -5724,8 +5732,8 @@ int32 FDynamicRibbonEmitterData::FillVertexData(struct FAsyncBufferFillData& Dat
 
 					FinalPos = InterpPos - InterpUp * InterpSize;
 					Vertex = (FParticleBeamTrailVertex*)(TempVertexData);
-					Vertex->Position = (FVector3f)FinalPos;
-					Vertex->OldPosition = (FVector3f)FinalPos;
+					Vertex->Position = FVector3f(FinalPos - LWCTileOffset);
+					Vertex->OldPosition = FVector3f(FinalPos - LWCTileOffset);
 					Vertex->ParticleId	= 0;
 					Vertex->Size.X = InterpSize;
 					Vertex->Size.Y = InterpSize;
@@ -5767,8 +5775,8 @@ int32 FDynamicRibbonEmitterData::FillVertexData(struct FAsyncBufferFillData& Dat
 				}
 
 				Vertex = (FParticleBeamTrailVertex*)(TempVertexData);
-				Vertex->Position = FVector3f(PackingParticle->Location + WorkingUp * CurrSize);
-				Vertex->OldPosition = (FVector3f)PackingParticle->OldLocation;
+				Vertex->Position = FVector3f(PackingParticle->Location - LWCTileOffset + WorkingUp * CurrSize);
+				Vertex->OldPosition = FVector3f(PackingParticle->OldLocation - LWCTileOffset);
 				Vertex->ParticleId	= 0;
 				Vertex->Size.X = CurrSize;
 				Vertex->Size.Y = CurrSize;
@@ -5801,8 +5809,8 @@ int32 FDynamicRibbonEmitterData::FillVertexData(struct FAsyncBufferFillData& Dat
 				//PackedVertexCount++;
 
 				Vertex = (FParticleBeamTrailVertex*)(TempVertexData);
-				Vertex->Position = FVector3f(PackingParticle->Location - WorkingUp * CurrSize);
-				Vertex->OldPosition = (FVector3f)PackingParticle->OldLocation;
+				Vertex->Position = FVector3f(PackingParticle->Location - LWCTileOffset - WorkingUp * CurrSize);
+				Vertex->OldPosition = FVector3f(PackingParticle->OldLocation - LWCTileOffset);
 				Vertex->ParticleId	= 0;
 				Vertex->Size.X = CurrSize;
 				Vertex->Size.Y = CurrSize;
@@ -6287,6 +6295,7 @@ int32 FDynamicAnimTrailEmitterData::FillVertexData(struct FAsyncBufferFillData&
 	float CurrDistance = 0.0f;
 
 	const uint8* ParticleData = Source.DataContainer.ParticleData;
+	const FVector LWCTileOffset = FVector(Source.LWCTile) * FLargeWorldRenderScalar::GetTileSize();
 	for (int32 ParticleIdx = 0; ParticleIdx < Source.ActiveParticleCount; ParticleIdx++)
 	{
 		DECLARE_PARTICLE_PTR(Particle, ParticleData + Source.ParticleStride * Source.DataContainer.ParticleIndices[ParticleIdx]);
@@ -6351,8 +6360,8 @@ int32 FDynamicAnimTrailEmitterData::FillVertexData(struct FAsyncBufferFillData&
 					}
 
 					Vertex = (FParticleBeamTrailVertex*)(TempVertexData);
-					Vertex->Position = (FVector3f)FirstSocket;
-					Vertex->OldPosition = (FVector3f)FirstSocket;
+					Vertex->Position = FVector3f(FirstSocket - LWCTileOffset);
+					Vertex->OldPosition = FVector3f(FirstSocket - LWCTileOffset);
 					Vertex->ParticleId	= 0;
 					Vertex->Size.X = InterpSize;
 					Vertex->Size.Y = InterpSize;
@@ -6374,8 +6383,8 @@ int32 FDynamicAnimTrailEmitterData::FillVertexData(struct FAsyncBufferFillData&
 					TempVertexData += VertexStride;
 
 					Vertex = (FParticleBeamTrailVertex*)(TempVertexData);
-					Vertex->Position = (FVector3f)SecondSocket;
-					Vertex->OldPosition = (FVector3f)SecondSocket;
+					Vertex->Position = FVector3f(SecondSocket - LWCTileOffset);
+					Vertex->OldPosition = FVector3f(SecondSocket - LWCTileOffset);
 					Vertex->ParticleId	= 0;
 					Vertex->Size.X = InterpSize;
 					Vertex->Size.Y = InterpSize;
@@ -6414,8 +6423,8 @@ int32 FDynamicAnimTrailEmitterData::FillVertexData(struct FAsyncBufferFillData&
 				}
 
 				Vertex = (FParticleBeamTrailVertex*)(TempVertexData);
-				Vertex->Position = (FVector3f)FirstSocket;//PackingParticle->Location + TrailPayload->FirstEdge * CurrSize;
-				Vertex->OldPosition = (FVector3f)RenderData.Particle->OldLocation;
+				Vertex->Position = FVector3f(FirstSocket - LWCTileOffset);//PackingParticle->Location + TrailPayload->FirstEdge * CurrSize;
+				Vertex->OldPosition = FVector3f(RenderData.Particle->OldLocation - LWCTileOffset);
 				Vertex->ParticleId	= 0;
 				Vertex->Size.X = InterpSize;
 				Vertex->Size.Y = InterpSize;
@@ -6438,8 +6447,8 @@ int32 FDynamicAnimTrailEmitterData::FillVertexData(struct FAsyncBufferFillData&
 				//PackedVertexCount++;
 
 				Vertex = (FParticleBeamTrailVertex*)(TempVertexData);
-				Vertex->Position = (FVector3f)SecondSocket;//PackingParticle->Location - TrailPayload->SecondEdge * CurrSize;
-				Vertex->OldPosition = (FVector3f)RenderData.Particle->OldLocation;
+				Vertex->Position = FVector3f(SecondSocket - LWCTileOffset);//PackingParticle->Location - TrailPayload->SecondEdge * CurrSize;
+				Vertex->OldPosition = FVector3f(RenderData.Particle->OldLocation - LWCTileOffset);
 				Vertex->ParticleId	= 0;
 				Vertex->Size.X = InterpSize;
 				Vertex->Size.Y = InterpSize;
diff --git a/Engine/Source/Runtime/Engine/Private/Particles/ParticleTrail2EmitterInstance.cpp b/Engine/Source/Runtime/Engine/Private/Particles/ParticleTrail2EmitterInstance.cpp
--- a/Engine/Source/Runtime/Engine/Private/Particles/ParticleTrail2EmitterInstance.cpp
+++ b/Engine/Source/Runtime/Engine/Private/Particles/ParticleTrail2EmitterInstance.cpp
@@ -3601,7 +3601,7 @@ void FParticleAnimTrailEmitterInstance::UpdateBoundingBox(float DeltaTime)
 
 			DECLARE_PARTICLE_PTR(FirstParticle, ParticleData + ParticleStride * ParticleIndices[0]);
 			FAnimTrailTypeDataPayload*	FirstPayload = ((FAnimTrailTypeDataPayload*)((uint8*)FirstParticle + TypeDataOffset));
-			FVector PrevParticleLocation = FirstParticle->Location;
+			FVector PrevParticleLocation = FirstParticle->Location + PositionOffsetThisTick;
 			float PrevParticleLength = FirstPayload->Length;
 
 			bool bSkipDoubleSpawnUpdate = !SpriteTemplate->bUseLegacySpawningBehavior;
diff --git a/Engine/Source/Runtime/Engine/Public/MeshParticleVertexFactory.h b/Engine/Source/Runtime/Engine/Public/MeshParticleVertexFactory.h
--- a/Engine/Source/Runtime/Engine/Public/MeshParticleVertexFactory.h
+++ b/Engine/Source/Runtime/Engine/Public/MeshParticleVertexFactory.h
@@ -29,7 +29,9 @@ BEGIN_GLOBAL_SHADER_PARAMETER_STRUCT( FMeshParticleUniformParameters, ENGINE_API
 	SHADER_PARAMETER( uint32, TexCoordWeightA )
 	SHADER_PARAMETER( uint32, TexCoordWeightB )
 	SHADER_PARAMETER( uint32, PrevTransformAvailable )
-	SHADER_PARAMETER( float, DeltaSeconds)
+	SHADER_PARAMETER( float, DeltaSeconds )
+	SHADER_PARAMETER( uint32, bUseLocalSpace )
+	SHADER_PARAMETER( FVector3f, LWCTile )
 END_GLOBAL_SHADER_PARAMETER_STRUCT()
 typedef TUniformBufferRef<FMeshParticleUniformParameters> FMeshParticleUniformBufferRef;
 
diff --git a/Engine/Source/Runtime/Engine/Public/ParticleBeamTrailVertexFactory.h b/Engine/Source/Runtime/Engine/Public/ParticleBeamTrailVertexFactory.h
--- a/Engine/Source/Runtime/Engine/Public/ParticleBeamTrailVertexFactory.h
+++ b/Engine/Source/Runtime/Engine/Public/ParticleBeamTrailVertexFactory.h
@@ -21,6 +21,8 @@ BEGIN_GLOBAL_SHADER_PARAMETER_STRUCT( FParticleBeamTrailUniformParameters, ENGIN
 	SHADER_PARAMETER( FVector4f, CameraRight )
 	SHADER_PARAMETER( FVector4f, CameraUp )
 	SHADER_PARAMETER( FVector4f, ScreenAlignment )
+	SHADER_PARAMETER( uint32, bUseLocalSpace)
+	SHADER_PARAMETER( FVector3f, LWCTile)
 END_GLOBAL_SHADER_PARAMETER_STRUCT()
 typedef TUniformBufferRef<FParticleBeamTrailUniformParameters> FParticleBeamTrailUniformBufferRef;
 
diff --git a/Engine/Source/Runtime/Engine/Public/ParticleHelper.h b/Engine/Source/Runtime/Engine/Public/ParticleHelper.h
--- a/Engine/Source/Runtime/Engine/Public/ParticleHelper.h
+++ b/Engine/Source/Runtime/Engine/Public/ParticleHelper.h
@@ -1390,6 +1390,7 @@ struct FDynamicSpriteEmitterReplayDataBase
 	FVector3f							NormalsSphereCenter;
 	FVector3f							NormalsCylinderDirection;
 	float							InvDeltaSeconds;
+	FVector3f						LWCTile;
 	int32							MaxDrawCount;
 	int32							OrbitModuleOffset;
 	int32							DynamicParameterDataOffset;
diff --git a/Engine/Source/Runtime/Engine/Public/ParticleVertexFactory.h b/Engine/Source/Runtime/Engine/Public/ParticleVertexFactory.h
--- a/Engine/Source/Runtime/Engine/Public/ParticleVertexFactory.h
+++ b/Engine/Source/Runtime/Engine/Public/ParticleVertexFactory.h
@@ -83,6 +83,7 @@ BEGIN_GLOBAL_SHADER_PARAMETER_STRUCT( FParticleSpriteUniformParameters, ENGINE_A
 	SHADER_PARAMETER_EX( float, RotationBias, EShaderPrecisionModifier::Half )
 	SHADER_PARAMETER_EX( float, NormalsType, EShaderPrecisionModifier::Half )
 	SHADER_PARAMETER_EX( float, InvDeltaSeconds, EShaderPrecisionModifier::Half )
+	SHADER_PARAMETER( FVector3f, LWCTile )
 	SHADER_PARAMETER_EX( FVector2f, PivotOffset, EShaderPrecisionModifier::Half )
 END_GLOBAL_SHADER_PARAMETER_STRUCT()
 typedef TUniformBufferRef<FParticleSpriteUniformParameters> FParticleSpriteUniformBufferRef;
