Remove an old deprecated CreatePin function from 4.19

#rnx

diff --git a/Engine/Source/Runtime/Engine/Classes/EdGraph/EdGraphNode.h b/Engine/Source/Runtime/Engine/Classes/EdGraph/EdGraphNode.h
--- a/Engine/Source/Runtime/Engine/Classes/EdGraph/EdGraphNode.h
+++ b/Engine/Source/Runtime/Engine/Classes/EdGraph/EdGraphNode.h
@@ -845,14 +845,6 @@ public:
 	// Returns the response to use when reporting a deprecation
 	virtual FEdGraphNodeDeprecationResponse GetDeprecationResponse(EEdGraphNodeDeprecationType DeprecationType) const;
 
-	// Returns true if this node should produce a compiler warning on deprecation
-	UE_DEPRECATED(4.23, "Use GetDeprecationResponse instead.")
-	virtual bool ShouldWarnOnDeprecation() const;
-
-	// Returns the string to use when reporting the deprecation
-	UE_DEPRECATED(4.23, "Use GetDeprecationResponse instead.")
-	virtual FString GetDeprecationMessage() const;
-
 	// Returns the object that should be focused when double-clicking on this node
 	// (the object can be an actor, which selects it in the world, or a node/graph/pin)
 	virtual UObject* GetJumpTargetForDoubleClick() const;
diff --git a/Engine/Source/Runtime/Engine/Private/EdGraph/EdGraphNode.cpp b/Engine/Source/Runtime/Engine/Private/EdGraph/EdGraphNode.cpp
--- a/Engine/Source/Runtime/Engine/Private/EdGraph/EdGraphNode.cpp
+++ b/Engine/Source/Runtime/Engine/Private/EdGraph/EdGraphNode.cpp
@@ -599,38 +599,6 @@ FEdGraphNodeDeprecationResponse UEdGraphNode::GetDeprecationResponse(EEdGraphNod
 	return Response;
 }
 
-// Deprecated; implemented for backwards-compatibility.
-bool UEdGraphNode::ShouldWarnOnDeprecation() const
-{
-	if (IsDeprecated())
-	{
-		return GetDeprecationResponse(EEdGraphNodeDeprecationType::NodeTypeIsDeprecated).MessageType == EEdGraphNodeDeprecationMessageType::Warning;
-	}
-	else if (HasDeprecatedReference())
-	{
-		return GetDeprecationResponse(EEdGraphNodeDeprecationType::NodeHasDeprecatedReference).MessageType == EEdGraphNodeDeprecationMessageType::Warning;
-	}
-
-	return false;
-}
-
-// Deprecated; implemented for backwards-compatibility.
-FString UEdGraphNode::GetDeprecationMessage() const
-{
-	FText MessageText;
-
-	if (IsDeprecated())
-	{
-		MessageText = GetDeprecationResponse(EEdGraphNodeDeprecationType::NodeTypeIsDeprecated).MessageText;
-	}
-	else if (HasDeprecatedReference())
-	{
-		MessageText = GetDeprecationResponse(EEdGraphNodeDeprecationType::NodeHasDeprecatedReference).MessageText;
-	}
-
-	return MessageText.ToString();
-}
-
 void UEdGraphNode::AddReferencedObjects(UObject* InThis, FReferenceCollector& Collector)
 {
 	Super::AddReferencedObjects(InThis, Collector);
