Fix non-unity build

diff --git a/Engine/Source/Runtime/Engine/Private/WorldPartition/ContentBundle/ContentBundle.cpp b/Engine/Source/Runtime/Engine/Private/WorldPartition/ContentBundle/ContentBundle.cpp
--- a/Engine/Source/Runtime/Engine/Private/WorldPartition/ContentBundle/ContentBundle.cpp
+++ b/Engine/Source/Runtime/Engine/Private/WorldPartition/ContentBundle/ContentBundle.cpp
@@ -5,6 +5,8 @@
 #include "WorldPartition/ContentBundle/ContentBundleDescriptor.h"
 #include "WorldPartition/WorldPartitionRuntimeHash.h"
 #include "WorldPartition/ContentBundle/ContentBundleLog.h"
+#include "Engine/World.h"
+#include "Misc/Paths.h"
 
 #if WITH_EDITOR
 #include "WorldPartition/ContentBundle/ContentBundleEditor.h"
