Texture make turning on Defer Compression show uncompressed right away

diff --git a/Engine/Source/Runtime/Engine/Private/Texture.cpp b/Engine/Source/Runtime/Engine/Private/Texture.cpp
--- a/Engine/Source/Runtime/Engine/Private/Texture.cpp
+++ b/Engine/Source/Runtime/Engine/Private/Texture.cpp
@@ -667,12 +667,9 @@ void UTexture::PostEditChangeProperty(FPropertyChangedEvent& PropertyChangedEven
 		}
 	}
 
-	// Don't update the texture resource if we've turned "DeferCompression" on, as this would cause
-	// it to immediately update as an uncompressed texture.  If it's a render target, we always need
-	// to update the resource, to avoid an assert when rendering to it, due to a mismatch between the
-	// render target and scene render.
-	if( !DeferCompressionWasEnabled &&
-		((PropertyChangedEvent.ChangeType & EPropertyChangeType::Interactive) == 0 || GetTextureClass() == ETextureClass::RenderTarget))
+	// If it's a render target, we always need to update the resource, to avoid an assert when rendering to it,
+	// due to a mismatch between the render target and scene render.
+	if ( (PropertyChangedEvent.ChangeType & EPropertyChangeType::Interactive) == 0 || GetTextureClass() == ETextureClass::RenderTarget )
 	{
 		// Update the texture resource. This will recache derived data if necessary
 		// which may involve recompressing the texture.
