#ue5 - Fix type conversion warnings in misc engine files.

diff --git a/Engine/Source/Runtime/Engine/Classes/AI/Navigation/NavigationTypes.h b/Engine/Source/Runtime/Engine/Classes/AI/Navigation/NavigationTypes.h
--- a/Engine/Source/Runtime/Engine/Classes/AI/Navigation/NavigationTypes.h
+++ b/Engine/Source/Runtime/Engine/Classes/AI/Navigation/NavigationTypes.h
@@ -675,7 +675,7 @@ namespace NavMeshMemory
 		public:
 			typedef FDefaultAllocator::ForAnyElementType Super;
 		private:
-			int32 AllocatedSize;
+			int64 AllocatedSize;
 		public:
 
 			ForAnyElementType()
@@ -702,7 +702,7 @@ namespace NavMeshMemory
 
 			FORCEINLINE_DEBUGGABLE void ResizeAllocation(SizeType PreviousNumElements, SizeType NumElements, SIZE_T NumBytesPerElement)
 			{
-				const int32 NewSize = NumElements * NumBytesPerElement;
+				const int64 NewSize = NumElements * NumBytesPerElement;
 				INC_DWORD_STAT_BY(STAT_NavigationMemory, NewSize - AllocatedSize);
 				AllocatedSize = NewSize;
 
diff --git a/Engine/Source/Runtime/Engine/Classes/PhysicsEngine/AggregateGeom.h b/Engine/Source/Runtime/Engine/Classes/PhysicsEngine/AggregateGeom.h
--- a/Engine/Source/Runtime/Engine/Classes/PhysicsEngine/AggregateGeom.h
+++ b/Engine/Source/Runtime/Engine/Classes/PhysicsEngine/AggregateGeom.h
@@ -65,7 +65,7 @@ struct ENGINE_API FKAggregateGeom
 
 	int32 GetElementCount(EAggCollisionShape::Type Type) const;
 
-	uint32 GetAllocatedSize() const { return SphereElems.GetAllocatedSize() + SphylElems.GetAllocatedSize() + BoxElems.GetAllocatedSize() + ConvexElems.GetAllocatedSize() + TaperedCapsuleElems.GetAllocatedSize(); }
+	SIZE_T GetAllocatedSize() const { return SphereElems.GetAllocatedSize() + SphylElems.GetAllocatedSize() + BoxElems.GetAllocatedSize() + ConvexElems.GetAllocatedSize() + TaperedCapsuleElems.GetAllocatedSize(); }
 
 	FKShapeElem* GetElement(const EAggCollisionShape::Type Type, const int32 Index)
 	{
