Back out changelist 22010196 to restore 22008890

diff --git a/Engine/Source/Runtime/Engine/Private/WorldPartition/WorldPartition.cpp b/Engine/Source/Runtime/Engine/Private/WorldPartition/WorldPartition.cpp
--- a/Engine/Source/Runtime/Engine/Private/WorldPartition/WorldPartition.cpp
+++ b/Engine/Source/Runtime/Engine/Private/WorldPartition/WorldPartition.cpp
@@ -309,7 +309,11 @@ bool UWorldPartition::IsValidPackageName(const FString& InPackageName)
 		// Remove PIE prefix
 		FString PackageName = UWorld::RemovePIEPrefix(InPackageName);
 		// Test if package is a valid world partition PIE package
-		return GeneratedStreamingPackageNames.Contains(PackageName);
+		// return GeneratedStreamingPackageNames.Contains(PackageName);
+		// @todo_ow: Can't rely on GeneratedStreamingPackageNames until Content Bundles return their generated packages 
+		//           in FContentBundleContainer::OnPreGenerateStreaming. For now world partition memory packages will 
+		//           return true so that server streaming works (see ValidateLevelVisibility in NetConnection.cpp)
+		return true;
 	}
 	return false;
 }
