Moves child actor name logic into a function linked to that child actor's OnDestroy so that we can clean up references properly when the underlying child actor is destroyed. Previous flow would only clean up if the component was the one that triggered the child actor destruction.

diff --git a/Engine/Source/Runtime/Engine/Classes/Components/ChildActorComponent.h b/Engine/Source/Runtime/Engine/Classes/Components/ChildActorComponent.h
--- a/Engine/Source/Runtime/Engine/Classes/Components/ChildActorComponent.h
+++ b/Engine/Source/Runtime/Engine/Classes/Components/ChildActorComponent.h
@@ -208,6 +208,11 @@ public:
 
 private:
 	bool IsChildActorReplicated() const;
+
+	bool IsBeingRemovedFromLevel() const;
+
+	UFUNCTION()
+	void OnChildActorDestroyed(AActor* DestroyedActor);
 };
 
 struct FActorParentComponentSetter
diff --git a/Engine/Source/Runtime/Engine/Private/Components/ChildActorComponent.cpp b/Engine/Source/Runtime/Engine/Private/Components/ChildActorComponent.cpp
--- a/Engine/Source/Runtime/Engine/Private/Components/ChildActorComponent.cpp
+++ b/Engine/Source/Runtime/Engine/Private/Components/ChildActorComponent.cpp
@@ -623,6 +623,72 @@ bool UChildActorComponent::IsChildActorReplicated() const
 	return (bChildActorClassReplicated || bChildActorTemplateReplicated);
 }
 
+bool UChildActorComponent::IsBeingRemovedFromLevel() const
+{
+	if (AActor* MyOwner = GetOwner())
+	{
+		if (ULevel* MyLevel = MyOwner->GetLevel())
+		{
+			if (MyLevel->bIsBeingRemoved)
+			{
+				return true;
+			}
+
+#if WITH_EDITOR
+			// In the editor, the child actor can already be removed from the world if the component owner was removed first
+			UWorld* MyWorld = MyLevel->GetWorld();
+			if (MyWorld && !MyWorld->IsGameWorld() && !MyLevel->Actors.Contains(ChildActor))
+			{
+				return true;
+			}
+#endif
+		}
+	}
+
+	return false;
+}
+
+void UChildActorComponent::OnChildActorDestroyed(AActor* DestroyedActor)
+{
+	if (GExitPurge)
+	{
+		return;
+	}
+
+	if (DestroyedActor && (DestroyedActor->HasAuthority() || !IsChildActorReplicated()) && !IsBeingRemovedFromLevel())
+	{
+		UWorld* World = DestroyedActor->GetWorld();
+		// World may be nullptr during shutdown
+		if (World != nullptr)
+		{
+			UClass* ChildClass = DestroyedActor->GetClass();
+
+			// We would like to make certain that our name is not going to accidentally get taken from us while we're destroyed
+			// so we increment ClassUnique beyond our index to be certain of it.  This is ... a bit hacky.
+			if (!GFastPathUniqueNameGeneration)
+			{
+				UpdateSuffixForNextNewObject(DestroyedActor->GetOuter(), ChildClass, [DestroyedActor](int32& Index) { Index = FMath::Max(Index, DestroyedActor->GetFName().GetNumber()); });
+			}
+
+			// If we are getting here due to garbage collection we can't rename, so we'll have to abandon this child actor name and pick up a new one
+			if (!IsGarbageCollecting())
+			{
+				const FString ObjectBaseName = FString::Printf(TEXT("DESTROYED_%s_CHILDACTOR"), *ChildClass->GetName());
+				DestroyedActor->Rename(*MakeUniqueObjectName(DestroyedActor->GetOuter(), ChildClass, *ObjectBaseName).ToString(), nullptr, REN_DoNotDirty | REN_ForceNoResetLoaders);
+			}
+			else
+			{
+				ChildActorName = NAME_None;
+				if (CachedInstanceData)
+				{
+					CachedInstanceData->ChildActorName = NAME_None;
+				}
+			}
+		}
+	}
+	ChildActor = nullptr;
+}
+
 void UChildActorComponent::CreateChildActor(TFunction<void(AActor*)> CustomizerFunc)
 {
 	AActor* MyOwner = GetOwner();
@@ -701,6 +767,7 @@ void UChildActorComponent::CreateChildActor(TFunction<void(AActor*)> CustomizerF
 				FVector Location = GetComponentLocation();
 				FRotator Rotation = GetComponentRotation();
 				ChildActor = World->SpawnActor(ChildActorClass, &Location, &Rotation, Params);
+				ChildActor->OnDestroyed.AddDynamic(this, &ThisClass::OnChildActorDestroyed);
 
 				// If spawn was successful, 
 				if(ChildActor != nullptr)
@@ -774,30 +841,6 @@ void UChildActorComponent::DestroyChildActor()
 	// If we own an Actor, kill it now unless we don't have authority on it, for that we rely on the server
 	// If the level is being removed then don't destroy the child actor so re-adding it doesn't
 	// need to create a new actor
-	auto IsBeingRemovedFromLevel = [this]() -> bool
-	{
-		if (AActor* MyOwner = GetOwner())
-		{
-			if (ULevel* MyLevel = MyOwner->GetLevel())
-			{
-				if (MyLevel->bIsBeingRemoved)
-				{
-					return true;
-				}
-
-#if WITH_EDITOR
-				// In the editor, the child actor can already be removed from the world if the component owner was removed first
-				UWorld* MyWorld = MyLevel->GetWorld();
-				if (MyWorld && !MyWorld->IsGameWorld() && !MyLevel->Actors.Contains(ChildActor))
-				{
-					return true;
-				}
-#endif
-			}
-		}
-
-		return false;
-	};
 
 	if (ChildActor && (ChildActor->HasAuthority() || !IsChildActorReplicated()) && !IsBeingRemovedFromLevel())
 	{
@@ -827,30 +870,6 @@ void UChildActorComponent::DestroyChildActor()
 			// World may be nullptr during shutdown
 			if (World != nullptr)
 			{
-				UClass* ChildClass = ChildActor->GetClass();
-
-				// We would like to make certain that our name is not going to accidentally get taken from us while we're destroyed
-				// so we increment ClassUnique beyond our index to be certain of it.  This is ... a bit hacky.
-				if (!GFastPathUniqueNameGeneration)
-				{
-					UpdateSuffixForNextNewObject(ChildActor->GetOuter(), ChildClass, [this](int32& Index) { Index = FMath::Max(Index, ChildActor->GetFName().GetNumber()); });
-				}
-
-				// If we are getting here due to garbage collection we can't rename, so we'll have to abandon this child actor name and pick up a new one
-				if (!IsGarbageCollecting())
-				{
-					const FString ObjectBaseName = FString::Printf(TEXT("DESTROYED_%s_CHILDACTOR"), *ChildClass->GetName());
-					ChildActor->Rename(*MakeUniqueObjectName(ChildActor->GetOuter(), ChildClass, *ObjectBaseName).ToString(), nullptr, REN_DoNotDirty | REN_ForceNoResetLoaders);
-				}
-				else
-				{
-					ChildActorName = NAME_None;
-					if (CachedInstanceData)
-					{
-						CachedInstanceData->ChildActorName = NAME_None;
-					}
-				}
-
 				if (!bIsChildActorPendingKillOrUnreachable)
 				{
 					World->DestroyActor(ChildActor);
