World Partition: fix unsaved actors in PIE for unsaved worlds.

#rnx

diff --git a/Engine/Source/Runtime/Engine/Private/WorldPartition/WorldPartitionStreamingGeneration.cpp b/Engine/Source/Runtime/Engine/Private/WorldPartition/WorldPartitionStreamingGeneration.cpp
--- a/Engine/Source/Runtime/Engine/Private/WorldPartition/WorldPartitionStreamingGeneration.cpp
+++ b/Engine/Source/Runtime/Engine/Private/WorldPartition/WorldPartitionStreamingGeneration.cpp
@@ -341,7 +341,8 @@ class FWorldPartitionStreamingGenerator
 		{
 			for (AActor* Actor : InContainer->GetWorld()->PersistentLevel->Actors)
 			{
-				if (IsValid(Actor) && Actor->IsPackageExternal() && Actor->IsMainPackageActor() && !Actor->IsEditorOnly() && InContainer->IsActorDescHandled(Actor) && !InContainer->GetActorDesc(Actor->GetActorGuid()))
+				// @todo_ow: temp fix, don't filter out containers (mainly for content bundles) to fix PIE in unsaved maps.
+				if (IsValid(Actor) && Actor->IsPackageExternal() && Actor->IsMainPackageActor() && !Actor->IsEditorOnly() && /*InContainer->IsActorDescHandled(Actor) &&*/ !InContainer->GetActorDesc(Actor->GetActorGuid()))
 				{
 					FWorldPartitionActorDescView ModifiedActorDescView = GetModifiedActorDesc(Actor, InContainer);
 					RegisterActorDescView(Actor->GetActorGuid(), ModifiedActorDescView);
