- Fix static analysis warning

#rnx

diff --git a/Engine/Source/Runtime/Engine/Private/Particles/ParticleEmitterInstances.cpp b/Engine/Source/Runtime/Engine/Private/Particles/ParticleEmitterInstances.cpp
--- a/Engine/Source/Runtime/Engine/Private/Particles/ParticleEmitterInstances.cpp
+++ b/Engine/Source/Runtime/Engine/Private/Particles/ParticleEmitterInstances.cpp
@@ -2843,7 +2843,7 @@ bool FParticleEmitterInstance::FillReplayData( FDynamicEmitterReplayDataBase& Ou
 		NewReplayData->RequiredModule = LODLevel->RequiredModule->CreateRendererResource();
 		NewReplayData->MaterialInterface = NULL;	// Must be set by derived implementation
 		NewReplayData->InvDeltaSeconds = (LastDeltaTime > UE_KINDA_SMALL_NUMBER) ? (1.0f / LastDeltaTime) : 0.0f;
-		NewReplayData->LWCTile = LODLevel->RequiredModule->bUseLocalSpace ? FVector3f::Zero() : Component->GetLWCTile();
+		NewReplayData->LWCTile = ((Component == nullptr) || LODLevel->RequiredModule->bUseLocalSpace) ? FVector3f::Zero() : Component->GetLWCTile();
 
 		NewReplayData->MaxDrawCount =
 			(LODLevel->RequiredModule->bUseMaxDrawCount == true) ? LODLevel->RequiredModule->MaxDrawCount : -1;
diff --git a/Engine/Source/Runtime/Engine/Private/Particles/ParticleSystemRender.cpp b/Engine/Source/Runtime/Engine/Private/Particles/ParticleSystemRender.cpp
--- a/Engine/Source/Runtime/Engine/Private/Particles/ParticleSystemRender.cpp
+++ b/Engine/Source/Runtime/Engine/Private/Particles/ParticleSystemRender.cpp
@@ -1522,7 +1522,7 @@ void FDynamicMeshEmitterData::GetDynamicMeshElementsEmitter(const FParticleSyste
 					Mesh.DepthPriorityGroup = (ESceneDepthPriorityGroup)Proxy->GetDepthPriorityGroup(View);
 
 					FMeshBatchElement& BatchElement = Mesh.Elements[0];
-					BatchElement.PrimitiveUniformBuffer = SourceData->bUseLocalSpace ? Proxy->GetUniformBuffer() : Proxy->GetWorldSpacePrimitiveUniformBuffer();
+					BatchElement.PrimitiveUniformBuffer = ((SourceData == nullptr) || SourceData->bUseLocalSpace) ? Proxy->GetUniformBuffer() : Proxy->GetWorldSpacePrimitiveUniformBuffer();
 					BatchElement.FirstIndex = Section.FirstIndex;
 					BatchElement.MinVertexIndex = Section.MinVertexIndex;
 					BatchElement.MaxVertexIndex = Section.MaxVertexIndex;
