PR #8549: Enable running export tasks in runtime

diff --git a/Engine/Source/Runtime/Engine/Private/UnrealExporter.cpp b/Engine/Source/Runtime/Engine/Private/UnrealExporter.cpp
--- a/Engine/Source/Runtime/Engine/Private/UnrealExporter.cpp
+++ b/Engine/Source/Runtime/Engine/Private/UnrealExporter.cpp
@@ -209,7 +209,6 @@ void UExporter::ExportToOutputDevice(const FExportObjectInnerContext* Context, U
 
 int32 UExporter::ExportToFile( UObject* Object, UExporter* InExporter, const TCHAR* Filename, bool InSelectedOnly, bool NoReplaceIdentical, bool Prompt )
 {
-#if WITH_EDITOR
 	UAssetExportTask* ExportTask = NewObject<UAssetExportTask>();
 	FGCObjectScopeGuard ExportTaskGuard(ExportTask);
 	ExportTask->Object = Object;
@@ -222,15 +221,11 @@ int32 UExporter::ExportToFile( UObject* Object, UExporter* InExporter, const TCH
 	ExportTask->bWriteEmptyFiles = false;
 	ExportTask->bAutomated = false;
 	return RunAssetExportTask(ExportTask) ? 1 : 0;
-#else
-	return 0;
-#endif
 }
 
 
 bool UExporter::RunAssetExportTask(class UAssetExportTask* Task)
 {
-#if WITH_EDITOR
 	check(Task);
 
 	CurrentFilename = Task->Filename;
@@ -415,14 +410,10 @@ bool UExporter::RunAssetExportTask(class UAssetExportTask* Task)
 		return true;
 	}
 	return false;
-#else
-	return false;
-#endif
 }
 
 bool UExporter::RunAssetExportTasks(const TArray<UAssetExportTask*>& ExportTasks)
 {
-#if WITH_EDITOR
 	bool bSuccess = true;
 	for (UAssetExportTask* Task : ExportTasks)
 	{
@@ -432,14 +423,10 @@ bool UExporter::RunAssetExportTasks(const TArray<UAssetExportTask*>& ExportTasks
 		}
 	}
 	return bSuccess;
-#else
-	return false;
-#endif
 }
 
 int32 UExporter::ExportToFileEx( FExportToFileParams& ExportParams )
 {
-#if WITH_EDITOR
 	check(ExportParams.Object);
 	UAssetExportTask* ExportTask = NewObject<UAssetExportTask>();
 	FGCObjectScopeGuard ExportTaskGuard(ExportTask);
@@ -454,9 +441,6 @@ int32 UExporter::ExportToFileEx( FExportToFileParams& ExportParams )
 	ExportTask->IgnoreObjectList = ExportParams.IgnoreObjectList;
 	ExportTask->bAutomated = false;
 	return RunAssetExportTask(ExportTask) ? 1 : 0;
-#else
-	return 0;
-#endif
 }
 
 const bool UExporter::bEnableDebugBrackets = false;
