Change the reflection capture's encoded HDR data to RG11B10 format.

Skip reflectio for simple translucent materials for both mobile forward and deferred.

Translucent materials fall back to single reflection or sky capture on mobile deferred if clustered reflection is disabled.

diff --git a/Engine/Shaders/Private/MobileBasePassPixelShader.usf b/Engine/Shaders/Private/MobileBasePassPixelShader.usf
--- a/Engine/Shaders/Private/MobileBasePassPixelShader.usf
+++ b/Engine/Shaders/Private/MobileBasePassPixelShader.usf
@@ -565,6 +565,7 @@ void Main(
 	float2 LocalPosition = SvPosition.xy - ResolvedView.ViewRectMin.xy;
 	uint GridIndex = ComputeLightGridCellIndex(uint2(LocalPosition.x, LocalPosition.y), SvPosition.w);
 
+#if MATERIALBLENDING_MASKED || MATERIALBLENDING_SOLID || TRANSLUCENCY_LIGHTING_SURFACE_FORWARDSHADING || MATERIAL_SHADINGMODEL_SINGLELAYERWATER
 	// Reflection IBL
 	AccumulateReflection(GBuffer
 		, CameraVector
@@ -573,6 +574,7 @@ void Main(
 		, IndirectIrradiance
 		, GridIndex
 		, DirectLighting);
+#endif
 
 	// Local lights
 #if  !MATERIAL_SHADINGMODEL_SINGLELAYERWATER && ENABLE_CLUSTERED_LIGHTS
diff --git a/Engine/Shaders/Private/MobileLightingCommon.ush b/Engine/Shaders/Private/MobileLightingCommon.ush
--- a/Engine/Shaders/Private/MobileLightingCommon.ush
+++ b/Engine/Shaders/Private/MobileLightingCommon.ush
@@ -288,11 +288,7 @@ half3 GetImageBasedReflectionLighting(half3 ReflectionVector
 	else
 	{
 		half AbsoluteSpecularMip = ComputeReflectionCaptureMipFromRoughness(Roughness, ResolvedView.ReflectionCubemapMaxMip);
-		half4 SpecularIBLSample = MobileReflectionCapture.Texture.SampleLevel(MobileReflectionCapture.TextureSampler, ReflectionVector, AbsoluteSpecularMip);
-		
-		half MaxValue = MobileReflectionCapture.Params.z;
-		SpecularIBL = RGBMDecode(SpecularIBLSample, MaxValue);
-		SpecularIBL = SpecularIBL * SpecularIBL;
+		SpecularIBL = MobileReflectionCapture.Texture.SampleLevel(MobileReflectionCapture.TextureSampler, ReflectionVector, AbsoluteSpecularMip).rgb;
 
 		half ReflectionCaptureBrightness = MobileReflectionCapture.Params.w;
 		SpecularIBL = SpecularIBL * ReflectionCaptureBrightness;
diff --git a/Engine/Source/Developer/Android/AndroidTargetPlatform/Private/AndroidTargetPlatform.cpp b/Engine/Source/Developer/Android/AndroidTargetPlatform/Private/AndroidTargetPlatform.cpp
--- a/Engine/Source/Developer/Android/AndroidTargetPlatform/Private/AndroidTargetPlatform.cpp
+++ b/Engine/Source/Developer/Android/AndroidTargetPlatform/Private/AndroidTargetPlatform.cpp
@@ -257,6 +257,7 @@ FAndroidTargetPlatform::FAndroidTargetPlatform(bool bInIsClient, const TCHAR* Fl
 	, DeviceDetection(nullptr)
 	, MobileShadingPath(0)
 	, bDistanceField(false)
+	, bMobileForwardEnableClusteredReflections(false)
 
 {
 #if WITH_ENGINE
@@ -264,6 +265,7 @@ FAndroidTargetPlatform::FAndroidTargetPlatform(bool bInIsClient, const TCHAR* Fl
 	StaticMeshLODSettings.Initialize(this);
 	GetConfigSystem()->GetBool(TEXT("/Script/Engine.RendererSettings"), TEXT("r.DistanceFields"), bDistanceField, GEngineIni);
 	GetConfigSystem()->GetInt(TEXT("/Script/Engine.RendererSettings"), TEXT("r.Mobile.ShadingPath"), MobileShadingPath, GEngineIni);
+	GetConfigSystem()->GetBool(TEXT("/Script/Engine.RendererSettings"), TEXT("r.Mobile.Forward.EnableClusteredReflections"), bMobileForwardEnableClusteredReflections, GEngineIni);
 #endif
 
 	TickDelegate = FTickerDelegate::CreateRaw(this, &FAndroidTargetPlatform::HandleTicker);
@@ -672,7 +674,7 @@ void FAndroidTargetPlatform::GetReflectionCaptureFormats( TArray<FName>& OutForm
 {
 	const bool bMobileDeferredShading = (MobileShadingPath == 1);
 	
-	if (SupportsVulkanSM5() || bMobileDeferredShading)
+	if (SupportsVulkanSM5() || bMobileDeferredShading || bMobileForwardEnableClusteredReflections)
 	{
 		// use Full HDR with SM5 and Mobile Deferred
 		OutFormats.Add(FName(TEXT("FullHDR")));
diff --git a/Engine/Source/Developer/Android/AndroidTargetPlatform/Private/AndroidTargetPlatform.h b/Engine/Source/Developer/Android/AndroidTargetPlatform/Private/AndroidTargetPlatform.h
--- a/Engine/Source/Developer/Android/AndroidTargetPlatform/Private/AndroidTargetPlatform.h
+++ b/Engine/Source/Developer/Android/AndroidTargetPlatform/Private/AndroidTargetPlatform.h
@@ -208,6 +208,9 @@ protected:
 	
 	// true if DistanceField is enabled
 	bool bDistanceField;
+
+	// r.Mobile.Forward.EnableClusteredReflections value
+	bool bMobileForwardEnableClusteredReflections;
 	
 #if WITH_ENGINE
 	// Holds a cache of the target LOD settings.
diff --git a/Engine/Source/Developer/IOS/IOSTargetPlatform/Private/IOSTargetPlatform.cpp b/Engine/Source/Developer/IOS/IOSTargetPlatform/Private/IOSTargetPlatform.cpp
--- a/Engine/Source/Developer/IOS/IOSTargetPlatform/Private/IOSTargetPlatform.cpp
+++ b/Engine/Source/Developer/IOS/IOSTargetPlatform/Private/IOSTargetPlatform.cpp
@@ -34,12 +34,14 @@ FIOSTargetPlatform::FIOSTargetPlatform(bool bInIsTVOS, bool bIsClientOnly)
 	, bIsTVOS(bInIsTVOS)
 	, MobileShadingPath(0)
 	, bDistanceField(false)
+	, bMobileForwardEnableClusteredReflections(false)
 {
 #if WITH_ENGINE
 	TextureLODSettings = nullptr; // TextureLODSettings are registered by the device profile.
 	StaticMeshLODSettings.Initialize(this);
 	GetConfigSystem()->GetBool(TEXT("/Script/Engine.RendererSettings"), TEXT("r.DistanceFields"), bDistanceField, GEngineIni);
 	GetConfigSystem()->GetInt(TEXT("/Script/Engine.RendererSettings"), TEXT("r.Mobile.ShadingPath"), MobileShadingPath, GEngineIni);
+	GetConfigSystem()->GetBool(TEXT("/Script/Engine.RendererSettings"), TEXT("r.Mobile.Forward.EnableClusteredReflections"), bMobileForwardEnableClusteredReflections, GEngineIni);
 #endif // #if WITH_ENGINE
 
 	// initialize the connected device detector
diff --git a/Engine/Source/Developer/IOS/IOSTargetPlatform/Private/IOSTargetPlatform.h b/Engine/Source/Developer/IOS/IOSTargetPlatform/Private/IOSTargetPlatform.h
--- a/Engine/Source/Developer/IOS/IOSTargetPlatform/Private/IOSTargetPlatform.h
+++ b/Engine/Source/Developer/IOS/IOSTargetPlatform/Private/IOSTargetPlatform.h
@@ -139,6 +139,9 @@ private:
 	// true if DistanceField is enabled
 	bool bDistanceField;
 
+	// r.Mobile.Forward.EnableClusteredReflections value
+	bool bMobileForwardEnableClusteredReflections;
+
 #if WITH_ENGINE
 	// Holds the cache of the target LOD settings.
 	const UTextureLODSettings* TextureLODSettings;
diff --git a/Engine/Source/Editor/UnrealEd/Private/EditorEngine.cpp b/Engine/Source/Editor/UnrealEd/Private/EditorEngine.cpp
--- a/Engine/Source/Editor/UnrealEd/Private/EditorEngine.cpp
+++ b/Engine/Source/Editor/UnrealEd/Private/EditorEngine.cpp
@@ -3964,11 +3964,10 @@ void UEditorEngine::BuildReflectionCaptures(UWorld* World)
 		
 		// First capture data we will use to generate endcoded data for a mobile renderer
 		bool bCapturingForMobile = true;
-		TArray<UTextureCube*> EncodedCaptures;
+		TArray<TArray<uint8>> EncodedCaptures;
 		EncodedCaptures.AddDefaulted(ReflectionCapturesToBuild.Num());
 		{
 			UReflectionCaptureComponent::UpdateReflectionCaptureContents(World, *UpdateReason, bVerifyOnlyCapturing, bCapturingForMobile);
-			bool bIsReflectionCaptureCompressionProjectSetting = (bool)IConsoleManager::Get().FindTConsoleVariableDataInt(TEXT("r.Mobile.ReflectionCaptureCompression"))->GetValueOnAnyThread();
 			for (int32 CaptureIndex = 0; CaptureIndex < ReflectionCapturesToBuild.Num(); CaptureIndex++)
 			{ 
 				UReflectionCaptureComponent* CaptureComponent = ReflectionCapturesToBuild[CaptureIndex];
@@ -3984,14 +3983,8 @@ void UEditorEngine::BuildReflectionCaptures(UWorld* World)
 					{
 						ULevel* StorageLevel = LightingScenarios[LevelIndex] ? LightingScenarios[LevelIndex] : CaptureComponent->GetOwner()->GetLevel();
 						UMapBuildDataRegistry* Registry = StorageLevel->GetOrCreateMapBuildData();
-						if (!CaptureComponent->bModifyMaxValueRGBM)
-						{
-							CaptureComponent->MaxValueRGBM = GetMaxValueRGBM(ReadbackCaptureData.FullHDRCapturedData, ReadbackCaptureData.CubemapSize);
-						}
-						FString TextureName = CaptureComponent->GetName() + TEXT("Texture");
-						TextureName += LexToString(CaptureComponent->MapBuildDataId);
-						GenerateEncodedHDRTextureCube(Registry, ReadbackCaptureData, TextureName, CaptureComponent->MaxValueRGBM, CaptureComponent, bIsReflectionCaptureCompressionProjectSetting);
-						EncodedCaptures[CaptureIndex] = ReadbackCaptureData.EncodedCaptureData;
+
+						GenerateEncodedHDRData(ReadbackCaptureData.FullHDRCapturedData, ReadbackCaptureData.CubemapSize, EncodedCaptures[CaptureIndex]);
 					}
 				}
 			}
@@ -4021,7 +4014,7 @@ void UEditorEngine::BuildReflectionCaptures(UWorld* World)
 					UMapBuildDataRegistry* Registry = StorageLevel->GetOrCreateMapBuildData();
 					FReflectionCaptureMapBuildData& CaptureBuildData = Registry->AllocateReflectionCaptureBuildData(CaptureComponent->MapBuildDataId, true);
 					(FReflectionCaptureData&)CaptureBuildData = ReadbackCaptureData;
-					CaptureBuildData.EncodedCaptureData = EncodedCaptures[CaptureIndex];
+					CaptureBuildData.EncodedHDRCapturedData = MoveTemp(EncodedCaptures[CaptureIndex]);
 					CaptureBuildData.FinalizeLoad();
 					// Recreate capture render state now that we have valid BuildData
 					CaptureComponent->MarkRenderStateDirty();
diff --git a/Engine/Source/Runtime/Core/Public/UObject/UE5MainStreamObjectVersion.h b/Engine/Source/Runtime/Core/Public/UObject/UE5MainStreamObjectVersion.h
--- a/Engine/Source/Runtime/Core/Public/UObject/UE5MainStreamObjectVersion.h
+++ b/Engine/Source/Runtime/Core/Public/UObject/UE5MainStreamObjectVersion.h
@@ -283,6 +283,9 @@ struct CORE_API FUE5MainStreamObjectVersion
 		// Reintroducing AnimationDataModelInterface_BackedOut changes
 		ReintroduceAnimationDataModelInterface,
 
+		// Uses RG11B10 format to store the encoded reflection capture data on mobile
+		StoreReflectionCaptureEncodedHDRDataInRG11B10Format,
+
 		// -----<new versions can be added above this line>-------------------------------------------------
 		VersionPlusOne,
 		LatestVersion = VersionPlusOne - 1
diff --git a/Engine/Source/Runtime/Engine/Classes/Components/ReflectionCaptureComponent.h b/Engine/Source/Runtime/Engine/Classes/Components/ReflectionCaptureComponent.h
--- a/Engine/Source/Runtime/Engine/Classes/Components/ReflectionCaptureComponent.h
+++ b/Engine/Source/Runtime/Engine/Classes/Components/ReflectionCaptureComponent.h
@@ -24,14 +24,6 @@ enum class EReflectionSourceType : uint8
 	SpecifiedCubemap,
 };
 
-UENUM()
-enum class EMobileReflectionCompression : uint8
-{
-	Default,
-	On,
-	Off,
-};
-
 // -> will be exported to EngineDecalClasses.h
 UCLASS(abstract, hidecategories=(Collision, Object, Physics, SceneComponent, Activation, "Components|Activation", Mobility), MinimalAPI)
 class UReflectionCaptureComponent : public USceneComponent
@@ -44,10 +36,6 @@ class UReflectionCaptureComponent : public USceneComponent
 	/** Indicates where to get the reflection source from. */
 	UPROPERTY(EditAnywhere, BlueprintReadOnly, Category=ReflectionCapture)
 	EReflectionSourceType ReflectionSourceType;
-	
-	/** Indicates if the Reflection will be compressed for mobile. If set to default it will take the project wide setting.*/
-	UPROPERTY(EditAnywhere, BlueprintReadOnly, Category= ReflectionCaptureMobile)
-	EMobileReflectionCompression MobileReflectionCompression;
 
 	/** Cubemap to use for reflection if ReflectionSourceType is set to RS_SpecifiedCubemap. */
 	UPROPERTY(EditAnywhere, BlueprintReadOnly, Category=ReflectionCapture)
@@ -60,14 +48,6 @@ class UReflectionCaptureComponent : public USceneComponent
 	/** A brightness control to scale the captured scene's reflection intensity. */
 	UPROPERTY(EditAnywhere, BlueprintReadOnly, Category=ReflectionCapture, meta=(UIMin = ".5", UIMax = "4"))
 	float Brightness;
-	
-	/** Indicates if we want to manually modify the MaxRGBM value.*/
-	UPROPERTY(EditAnywhere, BlueprintReadOnly, Category = ReflectionCaptureMobile)
-	bool bModifyMaxValueRGBM;
-
-	/** M values for RGBM, it will dynamically be generated by taking the highest RGB value into account. Explictly modifying it will require a Reflection Re-Capture.*/
-	UPROPERTY(EditAnywhere, BlueprintReadOnly, Category = ReflectionCaptureMobile, meta = (UIMin = "0", UIMax = "16", EditCondition = "bModifyMaxValueRGBM == true"))
-	float MaxValueRGBM;
 
 	/** World space offset to apply before capturing. */
 	UPROPERTY(EditAnywhere, BlueprintReadWrite, Category=ReflectionCapture, AdvancedDisplay)
@@ -77,9 +57,8 @@ class UReflectionCaptureComponent : public USceneComponent
 	UPROPERTY()
 	FGuid MapBuildDataId;
 
-	/** Cached Cubemap texture from MapBuildData used for rendering with the encoded HDR values. */
-	UPROPERTY(transient)
-	TObjectPtr<UTextureCube> CachedEncodedHDRCubemap;
+	/** Cubemap texture resource used for rendering with the encoded HDR values. */
+	class FReflectionTextureCubeResource* EncodedHDRCubemapTexture;
 
 #if WITH_EDITOR
 	/** Check to see if MapBuildDataId was loaded - otherwise we need to display a warning on cook */
@@ -169,9 +148,9 @@ private:
 	//void UpdateDerivedData(FReflectionCaptureFullHDR* NewDerivedData);
 	void SerializeLegacyData(FArchive& Ar);
 
+	void SafeReleaseEncodedHDRCubemapTexture();
+
 	friend class FReflectionCaptureProxy;
 };
 
-ENGINE_API extern float GetMaxValueRGBM(const TArray<uint8>& FullHDRData, int32 CubemapSize);
-ENGINE_API extern void GenerateEncodedHDRData(const TArray<uint8>& FullHDRData, int32 CubemapSize, float MaxValueRGBM, TArray<uint8>& OutEncodedHDRData);
-ENGINE_API extern void GenerateEncodedHDRTextureCube(class UMapBuildDataRegistry* Registry, class FReflectionCaptureData& CaptureBuildData, FString& TextureName, float MaxValueRGBM, class UReflectionCaptureComponent* CaptureComponent = nullptr, bool bIsReflectionCaptureCompressionProjectSetting = false);
\ No newline at end of file
+ENGINE_API extern void GenerateEncodedHDRData(const TArray<uint8>& FullHDRData, int32 CubemapSize, TArray<uint8>& OutEncodedHDRData);
\ No newline at end of file
diff --git a/Engine/Source/Runtime/Engine/Classes/Engine/MapBuildDataRegistry.h b/Engine/Source/Runtime/Engine/Classes/Engine/MapBuildDataRegistry.h
--- a/Engine/Source/Runtime/Engine/Classes/Engine/MapBuildDataRegistry.h
+++ b/Engine/Source/Runtime/Engine/Classes/Engine/MapBuildDataRegistry.h
@@ -187,22 +187,12 @@ public:
 	float AverageBrightness;
 
 	TArray<uint8> FullHDRCapturedData;
-	UTextureCube* EncodedCaptureData;
-
-#if WITH_EDITOR
-	/** Whether the brightness is baked in the EncodedHDRCubemap*/
-	bool bBrightnessBakedInEncodedHDRCubemap;
-#endif
+	TArray<uint8> EncodedHDRCapturedData;
 
 	FReflectionCaptureData() :
 		CubemapSize(0),
 		AverageBrightness(0.0f),
-		EncodedCaptureData(nullptr),
-#if WITH_EDITOR
-		bBrightnessBakedInEncodedHDRCubemap(false),
-#endif
 		bUploadedFinal(false)
-
 	{}
 
 	bool HasBeenUploadedFinal() const
@@ -218,6 +208,7 @@ public:
 		if (!GIsEditor)
 		{
 			FullHDRCapturedData.Empty();
+			EncodedHDRCapturedData.Empty();
 			CubemapSize = 0;
 			bUploadedFinal = true;
 		}
diff --git a/Engine/Source/Runtime/Engine/Classes/Engine/RendererSettings.h b/Engine/Source/Runtime/Engine/Classes/Engine/RendererSettings.h
--- a/Engine/Source/Runtime/Engine/Classes/Engine/RendererSettings.h
+++ b/Engine/Source/Runtime/Engine/Classes/Engine/RendererSettings.h
@@ -330,11 +330,6 @@ class ENGINE_API URendererSettings : public UDeveloperSettings
 		ConfigRestartRequired = true))
 	uint32 bMobileVirtualTextures : 1;
 
-	UPROPERTY(config, EditAnywhere, Category = Mobile, meta = (
-		ConsoleVariable = "r.Mobile.ReflectionCaptureCompression", DisplayName = "Mobile Reflection Capture Compression",
-		ToolTip = "Whether to use the Reflection Capture Compression or not for mobile. It will use ETC2 format to do the compression."))
-	uint32 bReflectionCaptureCompression : 1;
-
 	UPROPERTY(config, EditAnywhere, Category = Materials, meta = (
 		ConsoleVariable = "r.DiscardUnusedQuality", DisplayName = "Game Discards Unused Material Quality Levels",
 		ToolTip = "When running in game mode, whether to keep shaders for all quality levels in memory or only those needed for the current quality level.\nUnchecked: Keep all quality levels in memory allowing a runtime quality level change. (default)\nChecked: Discard unused quality levels when loading content for the game, saving some memory."))
diff --git a/Engine/Source/Runtime/Engine/Private/Components/ReflectionCaptureComponent.cpp b/Engine/Source/Runtime/Engine/Private/Components/ReflectionCaptureComponent.cpp
--- a/Engine/Source/Runtime/Engine/Private/Components/ReflectionCaptureComponent.cpp
+++ b/Engine/Source/Runtime/Engine/Private/Components/ReflectionCaptureComponent.cpp
@@ -40,18 +40,13 @@
 #include "Engine/MapBuildDataRegistry.h"
 #include "ComponentRecreateRenderStateContext.h"
 #include "Engine/TextureCube.h"
+#include "Math/PackedVector.h"
 
 #if WITH_EDITOR
 #include "Factories/TextureFactory.h"
 #include "TextureCompiler.h"
 #endif
 
-// ES3.0+ devices support seamless cubemap filtering, averaging edges will produce artifacts on those devices
-#define MOBILE_AVERAGE_CUBEMAP_EDGES 0 
-
-// ES3.0+ devices support seamless cubemap filtering, averaging edges will produce artifacts on those devices
-#define MOBILE_AVERAGE_CUBEMAP_EDGES 0 
-
 DEFINE_LOG_CATEGORY_STATIC(LogReflectionCaptureComponent, Log, All);
 
 /** 
@@ -64,12 +59,6 @@ ENGINE_API TAutoConsoleVariable<int32> CVarReflectionCaptureSize(
 	TEXT("Set the resolution for all reflection capture cubemaps. Should be set via project's Render Settings. Must be power of 2. Defaults to 128.\n")
 	);
 
-ENGINE_API TAutoConsoleVariable<int32> CVarMobileReflectionCaptureCompression(
-	TEXT("r.Mobile.ReflectionCaptureCompression"),
-	0,
-	TEXT("Whether to use the Reflection Capture Compression or not for mobile. It will use ETC2 format to do the compression.\n")
-);
-
 TAutoConsoleVariable<int32> CVarReflectionCaptureUpdateEveryFrame(
 	TEXT("r.ReflectionCaptureUpdateEveryFrame"),
 	0,
@@ -512,40 +501,7 @@ static void EdgeWalkSetup( bool ReverseDirection, int32 Edge, int32 MipSize, int
 	}
 }
 
-float GetMaxValueRGBM(const TArray<uint8>& FullHDRData, int32 CubemapSize)
-{
-	const int32 NumMips = FMath::CeilLogTwo(CubemapSize) + 1;
-	// get MaxValue from Mip0
-	float MaxValue = 0;
-	
-	const int32 MipSize = 1 << (NumMips - 1);
-	const int32 SourceCubeFaceBytes = MipSize * MipSize * sizeof(FFloat16Color);
-
-	for (int32 CubeFace = 0; CubeFace < CubeFace_MAX; CubeFace++)
-	{
-		const int32 FaceSourceIndex = CubeFace * SourceCubeFaceBytes;
-		const FFloat16Color* FaceSourceData = (const FFloat16Color*)&FullHDRData[FaceSourceIndex];
-
-		for (int32 y = 0; y < MipSize; y++)
-		{
-			for (int32 x = 0; x < MipSize; x++)
-			{
-				int32 TexelIndex = x + y * MipSize;
-				const FLinearColor LinearColor = FLinearColor(FaceSourceData[TexelIndex]);
-				float MaxValueTexel = FMath::Max(FMath::Max(LinearColor.R, LinearColor.G), FMath::Max(LinearColor.B, UE_DELTA));
-				if (MaxValue < MaxValueTexel)
-				{
-					MaxValue = MaxValueTexel;
-				}
-			}
-		}
-	}
-
-	MaxValue = FMath::Max(MaxValue, 1.0f);
-	return MaxValue;
-}
-
-void GenerateEncodedHDRData(const TArray<uint8>& FullHDRData, int32 CubemapSize, float MaxValueRGBM, TArray<uint8>& OutEncodedHDRData)
+void GenerateEncodedHDRData(const TArray<uint8>& FullHDRData, int32 CubemapSize, TArray<uint8>& OutEncodedHDRData)
 {
 	check(FullHDRData.Num() > 0);
 	const int32 NumMips = FMath::CeilLogTwo(CubemapSize) + 1;
@@ -553,98 +509,19 @@ void GenerateEncodedHDRData(const TArray<uint8>& FullHDRData, int32 CubemapSize,
 	int32 SourceMipBaseIndex = 0;
 	int32 DestMipBaseIndex = 0;
 
-	int32 EncodedDataSize = FullHDRData.Num() * sizeof(FColor) / sizeof(FFloat16Color);
+	int32 EncodedDataSize = FullHDRData.Num() * sizeof(FFloat3Packed) / sizeof(FFloat16Color);
 
 	OutEncodedHDRData.Empty(EncodedDataSize);
 	OutEncodedHDRData.AddZeroed(EncodedDataSize);
-	
-	MaxValueRGBM = FMath::Max(MaxValueRGBM, 1.0f);
 
 	for (int32 MipIndex = 0; MipIndex < NumMips; MipIndex++)
 	{
 		const int32 MipSize = 1 << (NumMips - MipIndex - 1);
 		const int32 SourceCubeFaceBytes = MipSize * MipSize * sizeof(FFloat16Color);
-		const int32 DestCubeFaceBytes = MipSize * MipSize * sizeof(FColor);
+		const int32 DestCubeFaceBytes = MipSize * MipSize * sizeof(FFloat3Packed);
 
 		const FFloat16Color*	MipSrcData = (const FFloat16Color*)&FullHDRData[SourceMipBaseIndex];
-		FColor*					MipDstData = (FColor*)&OutEncodedHDRData[DestMipBaseIndex];
-
-#if MOBILE_AVERAGE_CUBEMAP_EDGES
-		// Fix cubemap seams by averaging colors across edges
-		int32 CornerTable[4] =
-		{
-			0,
-			MipSize - 1,
-			MipSize * (MipSize - 1),
-			MipSize * (MipSize - 1) + MipSize - 1,
-		};
-
-		// Average corner colors
-		FLinearColor AvgCornerColors[8];
-		FPlatformMemory::Memset( AvgCornerColors, 0, sizeof( AvgCornerColors ) );
-		for( int32 Face = 0; Face < CubeFace_MAX; Face++ )
-		{
-			const FFloat16Color* FaceSrcData = MipSrcData + Face * MipSize * MipSize;
-
-			for( int32 Corner = 0; Corner < 4; Corner++ )
-			{
-				AvgCornerColors[ CubeCornerList[Face][Corner] ] += FLinearColor( FaceSrcData[ CornerTable[Corner] ] );
-			}
-		}
-
-		// Encode corners
-		for( int32 Face = 0; Face < CubeFace_MAX; Face++ )
-		{
-			FColor* FaceDstData = MipDstData + Face * MipSize * MipSize;
-
-			for( int32 Corner = 0; Corner < 4; Corner++ )
-			{
-				const FLinearColor LinearColor = AvgCornerColors[ CubeCornerList[Face][Corner] ] / 3.0f;
-				FaceDstData[ CornerTable[Corner] ] = RGBMEncode( LinearColor, MaxValueRGBM);
-			}
-		}
-
-		// Average edge colors
-		for( int32 EdgeIndex = 0; EdgeIndex < 12; EdgeIndex++ )
-		{
-			int32 FaceA = CubeEdgeListA[ EdgeIndex ][0];
-			int32 EdgeA = CubeEdgeListA[ EdgeIndex ][1];
-
-			int32 FaceB = CubeEdgeListB[ EdgeIndex ][0];
-			int32 EdgeB = CubeEdgeListB[ EdgeIndex ][1];
-
-			const FFloat16Color*	FaceSrcDataA = MipSrcData + FaceA * MipSize * MipSize;
-			FColor*					FaceDstDataA = MipDstData + FaceA * MipSize * MipSize;
-
-			const FFloat16Color*	FaceSrcDataB = MipSrcData + FaceB * MipSize * MipSize;
-			FColor*					FaceDstDataB = MipDstData + FaceB * MipSize * MipSize;
-
-			int32 EdgeStartA = 0;
-			int32 EdgeStepA = 0;
-			int32 EdgeStartB = 0;
-			int32 EdgeStepB = 0;
-
-			EdgeWalkSetup( false, EdgeA, MipSize, EdgeStartA, EdgeStepA );
-			EdgeWalkSetup( EdgeA == EdgeB || EdgeA + EdgeB == 3, EdgeB, MipSize, EdgeStartB, EdgeStepB );
-
-			// Walk edge
-			// Skip corners
-			for( int32 Texel = 1; Texel < MipSize - 1; Texel++ )       
-			{
-				int32 EdgeTexelA = EdgeStartA + EdgeStepA * Texel;
-				int32 EdgeTexelB = EdgeStartB + EdgeStepB * Texel;
-
-				check( 0 <= EdgeTexelA && EdgeTexelA < MipSize * MipSize );
-				check( 0 <= EdgeTexelB && EdgeTexelB < MipSize * MipSize );
-
-				const FLinearColor EdgeColorA = FLinearColor( FaceSrcDataA[ EdgeTexelA ] );
-				const FLinearColor EdgeColorB = FLinearColor( FaceSrcDataB[ EdgeTexelB ] );
-				const FLinearColor AvgColor = 0.5f * ( EdgeColorA + EdgeColorB );
-				
-				FaceDstDataA[ EdgeTexelA ] = FaceDstDataB[ EdgeTexelB ] = RGBMEncode( AvgColor, MaxValueRGBM);
-			}
-		}
-#endif // MOBILE_AVERAGE_CUBEMAP_EDGES
+		FFloat3Packed*			MipDstData = (FFloat3Packed*)&OutEncodedHDRData[DestMipBaseIndex];
 
 		// Encode rest of texels
 		for (int32 CubeFace = 0; CubeFace < CubeFace_MAX; CubeFace++)
@@ -652,20 +529,18 @@ void GenerateEncodedHDRData(const TArray<uint8>& FullHDRData, int32 CubemapSize,
 			const int32 FaceSourceIndex = SourceMipBaseIndex + CubeFace * SourceCubeFaceBytes;
 			const int32 FaceDestIndex = DestMipBaseIndex + CubeFace * DestCubeFaceBytes;
 			const FFloat16Color* FaceSourceData = (const FFloat16Color*)&FullHDRData[FaceSourceIndex];
-			FColor* FaceDestData = (FColor*)&OutEncodedHDRData[FaceDestIndex];
+			FFloat3Packed* FaceDestData = (FFloat3Packed*)&OutEncodedHDRData[FaceDestIndex];
 
-			// Convert each texel from linear space FP16 to RGBM FColor
-			// Skip edges
-			const int32 SkipEdges = MOBILE_AVERAGE_CUBEMAP_EDGES ? 1 : 0;
+			// Convert each texel from linear space FP16 to RG11B10 FColor
 
 			// Find MaxValue
-			for (int32 y = SkipEdges; y < MipSize - SkipEdges; y++)
+			for (int32 y = 0; y < MipSize; y++)
 			{
-				for (int32 x = SkipEdges; x < MipSize - SkipEdges; x++)
+				for (int32 x = 0; x < MipSize; x++)
 				{
 					int32 TexelIndex = x + y * MipSize;
 					const FLinearColor LinearColor = FLinearColor(FaceSourceData[TexelIndex]);
-					FaceDestData[ TexelIndex ] = RGBMEncode( LinearColor, MaxValueRGBM);
+					FaceDestData[ TexelIndex ] = FFloat3Packed(LinearColor);
 				}
 			}
 		}
@@ -675,59 +550,137 @@ void GenerateEncodedHDRData(const TArray<uint8>& FullHDRData, int32 CubemapSize,
 	}
 }
 
-
-void GenerateEncodedHDRTextureCube(UMapBuildDataRegistry* Registry, FReflectionCaptureData& ReflectionCaptureData, FString& TextureName, float MaxValueRGBM, UReflectionCaptureComponent* CaptureComponent, bool bIsReflectionCaptureCompressionProjectSetting)
+/**
+ * A cubemap texture resource that knows how to upload the packed capture data from a reflection capture.
+ * @todo - support texture streaming and compression
+ */
+class FReflectionTextureCubeResource : public FTexture
 {
-#if WITH_EDITOR
-	UTextureFactory* TextureFactory = NewObject<UTextureFactory>();
+public:
 
-	TextureFactory->CompressionSettings = TC_EncodedReflectionCapture;
-	UTextureCube* TextureCube = ReflectionCaptureData.EncodedCaptureData != nullptr ? ReflectionCaptureData.EncodedCaptureData : TextureFactory->CreateTextureCube(Registry, FName(TextureName), RF_Public);
+	FReflectionTextureCubeResource() :
+		Size(0),
+		NumMips(0),
+		Format(PF_Unknown),
+		SourceData(nullptr)
+	{}
 
-	if (TextureCube)
+	void SetupParameters(int32 InSize, int32 InNumMips, EPixelFormat InFormat, TArray<uint8>* InSourceData)
 	{
-		TextureCube->PreEditChange(nullptr);
-
-		TArray<uint8> TemporaryEncodedHDRCapturedData;
-
-		GenerateEncodedHDRData(ReflectionCaptureData.FullHDRCapturedData, ReflectionCaptureData.CubemapSize, MaxValueRGBM, TemporaryEncodedHDRCapturedData);
-		const int32 NumMips = FMath::CeilLogTwo(ReflectionCaptureData.CubemapSize) + 1;
-		TextureCube->Source.Init(
-			ReflectionCaptureData.CubemapSize,
-			ReflectionCaptureData.CubemapSize,
-			6,
-			NumMips,
-			TSF_BGRA8,
-			TemporaryEncodedHDRCapturedData.GetData()
-		);
-		TextureCube->Source.UseHashAsGuid();
-		// the loader can suggest a compression setting
-		TextureCube->LODGroup = TEXTUREGROUP_World;
+		Size = InSize;
+		NumMips = InNumMips;
+		Format = InFormat;
+		SourceData = InSourceData;
+	}
+
+	virtual void InitRHI() override
+	{
+		const FRHITextureCreateDesc Desc =
+			FRHITextureCreateDesc::CreateCube(TEXT("ReflectionTextureCube"), Size, Format)
+			.SetNumMips(NumMips)
+			.SetFlags(ETextureCreateFlags::ShaderResource);
 
-		bool bIsCompressed = false;
-		if (CaptureComponent != nullptr)
+		TextureCubeRHI = RHICreateTexture(Desc);
+		TextureRHI = TextureCubeRHI;
+
+		if (SourceData)
 		{
-			bIsCompressed = CaptureComponent->MobileReflectionCompression == EMobileReflectionCompression::Default ? bIsReflectionCaptureCompressionProjectSetting : CaptureComponent->MobileReflectionCompression == EMobileReflectionCompression::On;
+			check(SourceData->Num() > 0);
+
+			const int32 BlockBytes = GPixelFormats[Format].BlockBytes;
+			int32 MipBaseIndex = 0;
+
+			for (int32 MipIndex = 0; MipIndex < NumMips; MipIndex++)
+			{
+				const int32 MipSize = 1 << (NumMips - MipIndex - 1);
+				const int32 CubeFaceBytes = MipSize * MipSize * BlockBytes;
+
+				uint32 SrcPitch = MipSize * BlockBytes;
+
+				for (int32 CubeFace = 0; CubeFace < CubeFace_MAX; CubeFace++)
+				{
+					uint32 DestPitch = 0;
+					uint8* DestBuffer = (uint8*)RHILockTextureCubeFace(TextureCubeRHI, CubeFace, 0, MipIndex, RLM_WriteOnly, DestPitch, false);
+					
+					const int32 SourceIndex = MipBaseIndex + CubeFace * CubeFaceBytes;
+					const uint8* SourceBuffer = &(*SourceData)[SourceIndex];
+
+					if (SrcPitch == DestPitch)
+					{
+						FMemory::Memcpy(DestBuffer, SourceBuffer, CubeFaceBytes);
+					}
+					else
+					{
+						// Copy data, taking the stride into account!
+						uint8* Src = (uint8*)SourceBuffer;
+						uint8* Dst = (uint8*)DestBuffer;
+						for (int32 Row = 0; Row < MipSize; ++Row)
+						{
+							FMemory::Memcpy(Dst, Src, SrcPitch);
+							Src += SrcPitch;
+							Dst += DestPitch;
+						}
+						check((PTRINT(Src) - PTRINT(SourceBuffer)) == PTRINT(CubeFaceBytes));
+					}
+
+					RHIUnlockTextureCubeFace(TextureCubeRHI, CubeFace, 0, MipIndex, false);
+				}
+
+				MipBaseIndex += CubeFaceBytes * CubeFace_MAX;
+			}
+
+			if (!GIsEditor)
+			{
+				// Toss the source data now that we've created the cubemap
+				// Note: can't do this if we ever use this texture resource in the editor and want to save the data later
+				SourceData->Empty();
+			}
 		}
 
-		TextureCube->CompressionSettings = TC_EncodedReflectionCapture;
-		TextureCube->CompressionNone = !bIsCompressed;
-		TextureCube->CompressionQuality = TCQ_Highest;
-		TextureCube->Filter = TF_Trilinear;
-		TextureCube->SRGB = 0;
+		// Create the sampler state RHI resource.
+		FSamplerStateInitializerRHI SamplerStateInitializer
+		(
+			SF_Trilinear,
+			AM_Clamp,
+			AM_Clamp,
+			AM_Clamp
+		);
+		SamplerStateRHI = GetOrCreateSamplerState(SamplerStateInitializer);
 
-		// for now we don't support mip map generation on cubemaps
-		TextureCube->MipGenSettings = TMGS_LeaveExistingMips;
+		INC_MEMORY_STAT_BY(STAT_ReflectionCaptureTextureMemory, CalcTextureSize(Size, Size, Format, NumMips) * 6);
+	}
 
-		TextureCube->UpdateResource();
-		TextureCube->MarkPackageDirty();
+	virtual void ReleaseRHI() override
+	{
+		DEC_MEMORY_STAT_BY(STAT_ReflectionCaptureTextureMemory, CalcTextureSize(Size, Size, Format, NumMips) * 6);
+		TextureCubeRHI.SafeRelease();
+		FTexture::ReleaseRHI();
+	}
 
-		TextureCube->PostEditChange();
+	virtual uint32 GetSizeX() const override
+	{
+		return Size;
 	}
-	ReflectionCaptureData.EncodedCaptureData = TextureCube;
-	ReflectionCaptureData.bBrightnessBakedInEncodedHDRCubemap = false;
-#endif
-}
+
+	virtual uint32 GetSizeY() const override //-V524
+	{
+		return Size;
+	}
+
+	FRHITexture* GetTextureRHI()
+	{
+		return TextureCubeRHI;
+	}
+
+private:
+
+	int32 Size;
+	int32 NumMips;
+	EPixelFormat Format;
+	FTextureCubeRHIRef TextureCubeRHI;
+
+	TArray<uint8>* SourceData;
+};
 
 TArray<UReflectionCaptureComponent*> UReflectionCaptureComponent::ReflectionCapturesToUpdate;
 TArray<UReflectionCaptureComponent*> UReflectionCaptureComponent::ReflectionCapturesToUpdateForLoad;
@@ -737,11 +690,9 @@ UReflectionCaptureComponent::UReflectionCaptureComponent(const FObjectInitialize
 	: Super(ObjectInitializer)
 {
 	Brightness = 1;
-	bModifyMaxValueRGBM = false;
-	MaxValueRGBM = 0.0f;
 	// Shouldn't be able to change reflection captures at runtime
 	Mobility = EComponentMobility::Static;
-	CachedEncodedHDRCubemap = nullptr;
+	EncodedHDRCubemapTexture = nullptr;
 	CachedAverageBrightness = 1.0f;
 	bNeedsRecaptureOrUpload = false;
 }
@@ -775,26 +726,59 @@ void UReflectionCaptureComponent::SendRenderTransform_Concurrent()
 	Super::SendRenderTransform_Concurrent();
 }
 
+void UReflectionCaptureComponent::SafeReleaseEncodedHDRCubemapTexture()
+{
+	if (EncodedHDRCubemapTexture)
+	{
+		BeginReleaseResource(EncodedHDRCubemapTexture);
+		ENQUEUE_RENDER_COMMAND(DeleteEncodedHDRCubemapTexture)(
+			[ParamPointerToRelease = EncodedHDRCubemapTexture](FRHICommandListImmediate& RHICmdList)
+			{
+				delete ParamPointerToRelease;
+			});
+		EncodedHDRCubemapTexture = nullptr;
+	}
+}
+
 void UReflectionCaptureComponent::OnRegister()
 {
-	const ERHIFeatureLevel::Type FeatureLevel = GetWorld()->FeatureLevel;
-	const bool bEncodedDataRequired = (FeatureLevel == ERHIFeatureLevel::ES3_1);
+	const bool bMobileDeferredShading = IsMobileDeferredShadingEnabled(GMaxRHIShaderPlatform);
+	const bool bEncodedHDRCubemapTextureRequired = (GIsEditor || GMaxRHIFeatureLevel == ERHIFeatureLevel::ES3_1) && !MobileForwardEnableClusteredReflections(GMaxRHIShaderPlatform);
 
-	if (bEncodedDataRequired)
+	if (bEncodedHDRCubemapTextureRequired)
 	{
-		const FReflectionCaptureMapBuildData* MapBuildData = GetMapBuildData();
+		FReflectionCaptureMapBuildData* MapBuildData = GetMapBuildData();
 
 		// If the MapBuildData is valid, update it. If it is not we will use the cached values, if there are any
 		if (MapBuildData)
 		{
-			CachedEncodedHDRCubemap = MapBuildData->EncodedCaptureData;
+			if (!EncodedHDRCubemapTexture)
+			{
+				EncodedHDRCubemapTexture = new FReflectionTextureCubeResource();
+				TArray<uint8>* EncodedHDRCapturedData = nullptr;
+				EPixelFormat EncodedHDRCubemapTextureFormat = PF_Unknown;
+				if (bMobileDeferredShading)
+				{
+					EncodedHDRCapturedData = &MapBuildData->FullHDRCapturedData;
+					EncodedHDRCubemapTextureFormat = PF_FloatRGBA;
+				}
+				else
+				{
+					EncodedHDRCapturedData = &MapBuildData->EncodedHDRCapturedData;
+					EncodedHDRCubemapTextureFormat = PF_FloatR11G11B10;
+				}
+
+				EncodedHDRCubemapTexture->SetupParameters(MapBuildData->CubemapSize, FMath::CeilLogTwo(MapBuildData->CubemapSize) + 1, EncodedHDRCubemapTextureFormat, EncodedHDRCapturedData);
+				BeginInitResource(EncodedHDRCubemapTexture);
+			}
+
 			CachedAverageBrightness = MapBuildData->AverageBrightness;
 		}
 	}
 	else
 	{
 		// SM5 doesn't require cached values
-		CachedEncodedHDRCubemap = nullptr;
+		SafeReleaseEncodedHDRCubemapTexture();
 		CachedAverageBrightness = 0;
 	}
 
@@ -989,9 +973,6 @@ void UReflectionCaptureComponent::PostEditChangeProperty(FPropertyChangedEvent&
 {
 	if (PropertyChangedEvent.GetPropertyName() == GET_MEMBER_NAME_CHECKED(UReflectionCaptureComponent, Cubemap) ||
 		PropertyChangedEvent.GetPropertyName() == GET_MEMBER_NAME_CHECKED(UReflectionCaptureComponent, SourceCubemapAngle) ||
-		PropertyChangedEvent.GetPropertyName() == GET_MEMBER_NAME_CHECKED(UReflectionCaptureComponent, MobileReflectionCompression) ||
-		PropertyChangedEvent.GetPropertyName() == GET_MEMBER_NAME_CHECKED(UReflectionCaptureComponent, bModifyMaxValueRGBM) ||
-		PropertyChangedEvent.GetPropertyName() == GET_MEMBER_NAME_CHECKED(UReflectionCaptureComponent, MaxValueRGBM) ||
 		PropertyChangedEvent.GetPropertyName() == GET_MEMBER_NAME_CHECKED(UReflectionCaptureComponent, ReflectionSourceType))
 	{
 		MarkDirtyForRecapture();
@@ -1019,6 +1000,11 @@ void UReflectionCaptureComponent::BeginDestroy()
 		Scene->ReleaseReflectionCubemap(this);
 	}
 
+	if (EncodedHDRCubemapTexture)
+	{
+		BeginReleaseResource(EncodedHDRCubemapTexture);
+	}
+
 	// Begin a fence to track the progress of the above BeginReleaseResource being completed on the RT
 	ReleaseResourcesFence.BeginFence();
 
@@ -1033,7 +1019,11 @@ bool UReflectionCaptureComponent::IsReadyForFinishDestroy()
 
 void UReflectionCaptureComponent::FinishDestroy()
 {
-	CachedEncodedHDRCubemap = nullptr;
+	if (EncodedHDRCubemapTexture)
+	{
+		delete EncodedHDRCubemapTexture;
+		EncodedHDRCubemapTexture = nullptr;
+	}
 
 	Super::FinishDestroy();
 }
@@ -1124,7 +1114,7 @@ void UReflectionCaptureComponent::PreFeatureLevelChange(ERHIFeatureLevel::Type P
 {
 	if (SupportsTextureCubeArray(PendingFeatureLevel))
 	{
-		CachedEncodedHDRCubemap = nullptr;
+		SafeReleaseEncodedHDRCubemapTexture();
 
 		MarkDirtyForRecaptureOrUpload();
 	}
@@ -1252,9 +1242,8 @@ FReflectionCaptureProxy::FReflectionCaptureProxy(const UReflectionCaptureCompone
 	Component = InComponent;
 	const FReflectionCaptureMapBuildData* MapBuildData = InComponent->GetMapBuildData();
 
-	EncodedHDRCubemap = Component->CachedEncodedHDRCubemap;
+	EncodedHDRCubemap = Component->EncodedHDRCubemapTexture;
 	EncodedHDRAverageBrightness = Component->CachedAverageBrightness;
-	MaxValueRGBM = Component->MaxValueRGBM;
 	SetTransform(InComponent->GetComponentTransform().ToMatrixWithScale());
 	InfluenceRadius = InComponent->GetInfluenceBoundingRadius();
 	Brightness = InComponent->Brightness;
@@ -1292,12 +1281,13 @@ void FReflectionCaptureProxy::UpdateMobileUniformBuffer()
 	FTexture* CaptureTexture = GBlackTextureCube;
 	if (EncodedHDRCubemap)
 	{
-		CaptureTexture = EncodedHDRCubemap->GetResource();
+		check(EncodedHDRCubemap->IsInitialized());
+		CaptureTexture = EncodedHDRCubemap;
 	}
 		
 	FMobileReflectionCaptureShaderParameters Parameters;
 	//To keep ImageBasedReflectionLighting coherence with PC, use AverageBrightness instead of InvAverageBrightness to calculate the IBL contribution
-	Parameters.Params = FVector4f(EncodedHDRAverageBrightness, 0.f, MaxValueRGBM <= 0.0f ? 16.0f: MaxValueRGBM, Brightness);
+	Parameters.Params = FVector4f(EncodedHDRAverageBrightness, 0.f, 0.0f, Brightness);
 	Parameters.Texture = CaptureTexture->TextureRHI;
 	Parameters.TextureSampler = CaptureTexture->SamplerStateRHI;
 
diff --git a/Engine/Source/Runtime/Engine/Private/MapBuildData.cpp b/Engine/Source/Runtime/Engine/Private/MapBuildData.cpp
--- a/Engine/Source/Runtime/Engine/Private/MapBuildData.cpp
+++ b/Engine/Source/Runtime/Engine/Private/MapBuildData.cpp
@@ -307,13 +307,6 @@ FArchive& operator<<(FArchive& Ar, FReflectionCaptureMapBuildData& ReflectionCap
 		Ar << Brightness;
 	}
 
-#if WITH_EDITOR
-	if (Ar.CustomVer(FUE5ReleaseStreamObjectVersion::GUID) < FUE5ReleaseStreamObjectVersion::ExcludeBrightnessFromEncodedHDRCubemap && Brightness != 1.0f)
-	{
-		ReflectionCaptureMapBuildData.bBrightnessBakedInEncodedHDRCubemap = true;
-	}
-#endif
-
 	static FName FullHDR(TEXT("FullHDR"));
 	static FName EncodedHDR(TEXT("EncodedHDR"));
 
@@ -335,14 +328,24 @@ FArchive& operator<<(FArchive& Ar, FReflectionCaptureMapBuildData& ReflectionCap
 		Ar << StrippedData;
 	}
 
-	if (Ar.CustomVer(FMobileObjectVersion::GUID) >= FMobileObjectVersion::StoreReflectionCaptureCompressedMobile)
+	if (Ar.CustomVer(FMobileObjectVersion::GUID) >= FMobileObjectVersion::StoreReflectionCaptureCompressedMobile
+		&& Ar.CustomVer(FUE5MainStreamObjectVersion::GUID) < FUE5MainStreamObjectVersion::StoreReflectionCaptureEncodedHDRDataInRG11B10Format)
 	{
-		Ar << ReflectionCaptureMapBuildData.EncodedCaptureData;
+		UTextureCube* EncodedCaptureData = nullptr;
+		Ar << EncodedCaptureData;
 	}
 	else
 	{
-		TArray<uint8> StrippedData;
-		Ar << StrippedData;
+		if ((Formats.Num() == 0 || Formats.Contains(EncodedHDR))
+			&& Ar.CustomVer(FUE5MainStreamObjectVersion::GUID) >= FUE5MainStreamObjectVersion::StoreReflectionCaptureEncodedHDRDataInRG11B10Format)
+		{
+			Ar << ReflectionCaptureMapBuildData.EncodedHDRCapturedData;
+		}
+		else
+		{
+			TArray<uint8> StrippedData;
+			Ar << StrippedData;
+		}
 	}
 
 	if (Ar.IsLoading())
@@ -360,13 +363,13 @@ FReflectionCaptureMapBuildData::~FReflectionCaptureMapBuildData()
 
 void FReflectionCaptureMapBuildData::FinalizeLoad()
 {
-	AllocatedSize = FullHDRCapturedData.GetAllocatedSize();
+	AllocatedSize = FullHDRCapturedData.GetAllocatedSize() + EncodedHDRCapturedData.GetAllocatedSize();
 	INC_DWORD_STAT_BY(STAT_ReflectionCaptureBuildData, AllocatedSize);
 }
 
 void FReflectionCaptureMapBuildData::AddReferencedObjects(FReferenceCollector& Collector)
 {
-	Collector.AddReferencedObject(EncodedCaptureData);
+
 }
 
 UMapBuildDataRegistry::UMapBuildDataRegistry(const FObjectInitializer& ObjectInitializer)
@@ -480,7 +483,7 @@ void UMapBuildDataRegistry::Serialize(FArchive& Ar)
 			{
 				const FReflectionCaptureMapBuildData& CaptureBuildData = It.Value();
 				// Sanity check that every reflection capture entry has valid data for at least one format
-				check(CaptureBuildData.FullHDRCapturedData.Num() > 0 || CaptureBuildData.EncodedCaptureData != nullptr);
+				check(CaptureBuildData.FullHDRCapturedData.Num() > 0 || CaptureBuildData.EncodedHDRCapturedData.Num() > 0);
 			}
 		}
 
@@ -523,10 +526,10 @@ void UMapBuildDataRegistry::PostLoad()
 
 			if (!bEncodedDataRequired)
 			{
-				CaptureBuildData.EncodedCaptureData = nullptr;
+				CaptureBuildData.EncodedHDRCapturedData.Empty();
 			}
 
-			check(CaptureBuildData.EncodedCaptureData != nullptr || CaptureBuildData.FullHDRCapturedData.Num() > 0 || FApp::CanEverRender() == false);
+			check(CaptureBuildData.EncodedHDRCapturedData.Num() > 0 || CaptureBuildData.FullHDRCapturedData.Num() > 0 || FApp::CanEverRender() == false);
 		}
 	}
 
@@ -545,11 +548,9 @@ void UMapBuildDataRegistry::HandleLegacyEncodedCubemapData()
 		for (TMap<FGuid, FReflectionCaptureMapBuildData>::TIterator It(ReflectionCaptureBuildData); It; ++It)
 		{
 			FReflectionCaptureMapBuildData& CaptureBuildData = It.Value();
-			if ((CaptureBuildData.EncodedCaptureData == nullptr || CaptureBuildData.bBrightnessBakedInEncodedHDRCubemap) && CaptureBuildData.FullHDRCapturedData.Num() != 0)
+			if (CaptureBuildData.EncodedHDRCapturedData.Num() == 0 && CaptureBuildData.FullHDRCapturedData.Num() != 0)
 			{
-				FString TextureName = TEXT("DeprecatedTexture");
-				TextureName += LexToString(It.Key());
-				GenerateEncodedHDRTextureCube(this, CaptureBuildData, TextureName, 16.0f);
+				GenerateEncodedHDRData(CaptureBuildData.FullHDRCapturedData, CaptureBuildData.CubemapSize, CaptureBuildData.EncodedHDRCapturedData);
 			}
 		}
 	}
diff --git a/Engine/Source/Runtime/Engine/Private/RendererSettings.cpp b/Engine/Source/Runtime/Engine/Private/RendererSettings.cpp
--- a/Engine/Source/Runtime/Engine/Private/RendererSettings.cpp
+++ b/Engine/Source/Runtime/Engine/Private/RendererSettings.cpp
@@ -260,8 +260,7 @@ void URendererSettings::PostEditChangeProperty(FPropertyChangedEvent& PropertyCh
 
 		ExportValuesToConsoleVariables(PropertyChangedEvent.Property);
 
-		if ((PropertyChangedEvent.Property->GetFName() == GET_MEMBER_NAME_CHECKED(URendererSettings, ReflectionCaptureResolution) ||
-			(PropertyChangedEvent.Property->GetFName() == GET_MEMBER_NAME_CHECKED(URendererSettings, bReflectionCaptureCompression))) &&
+		if (PropertyChangedEvent.Property->GetFName() == GET_MEMBER_NAME_CHECKED(URendererSettings, ReflectionCaptureResolution) &&
 			PropertyChangedEvent.ChangeType != EPropertyChangeType::Interactive)
 		{
 			if (GEditor != nullptr)
diff --git a/Engine/Source/Runtime/Engine/Public/SceneManagement.h b/Engine/Source/Runtime/Engine/Public/SceneManagement.h
--- a/Engine/Source/Runtime/Engine/Public/SceneManagement.h
+++ b/Engine/Source/Runtime/Engine/Public/SceneManagement.h
@@ -1147,7 +1147,7 @@ inline bool PrimitiveNeedsDistanceFieldSceneData(bool bTrackAllPrimitives,
 
 
 BEGIN_GLOBAL_SHADER_PARAMETER_STRUCT(FMobileReflectionCaptureShaderParameters,ENGINE_API)
-	SHADER_PARAMETER(FVector4f, Params) // x - inv average brightness, y - sky cubemap max mip, z - Max value for RGBM, w - brightness of reflection capture
+	SHADER_PARAMETER(FVector4f, Params) // x - inv average brightness, y - sky cubemap max mip, z - unused, w - brightness of reflection capture
 	SHADER_PARAMETER_TEXTURE(TextureCube, Texture)
 	SHADER_PARAMETER_SAMPLER(SamplerState, TextureSampler)
 END_GLOBAL_SHADER_PARAMETER_STRUCT()
@@ -1940,9 +1940,8 @@ public:
 
 	/** Used with mobile renderer */
 	TUniformBufferRef<FMobileReflectionCaptureShaderParameters> MobileUniformBuffer;
-	UTextureCube* EncodedHDRCubemap;
+	FTexture* EncodedHDRCubemap;
 	float EncodedHDRAverageBrightness;
-	float MaxValueRGBM;
 
 	EReflectionCaptureShape::Type Shape;
 
diff --git a/Engine/Source/Runtime/Renderer/Private/PostProcess/PostProcessCompositeEditorPrimitives.cpp b/Engine/Source/Runtime/Renderer/Private/PostProcess/PostProcessCompositeEditorPrimitives.cpp
--- a/Engine/Source/Runtime/Renderer/Private/PostProcess/PostProcessCompositeEditorPrimitives.cpp
+++ b/Engine/Source/Runtime/Renderer/Private/PostProcess/PostProcessCompositeEditorPrimitives.cpp
@@ -300,6 +300,7 @@ const FViewInfo* CreateEditorPrimitiveView(const FViewInfo& ParentView, FIntRect
 BEGIN_SHADER_PARAMETER_STRUCT(FEditorPrimitivesPassParameters, )
 	SHADER_PARAMETER_STRUCT_INCLUDE(FViewShaderParameters, View)
 	SHADER_PARAMETER_STRUCT_REF(FReflectionCaptureShaderData, ReflectionCapture)
+	SHADER_PARAMETER_STRUCT_REF(FMobileReflectionCaptureShaderData, MobileReflectionCaptureData)
 	SHADER_PARAMETER_RDG_UNIFORM_BUFFER(FOpaqueBasePassUniformParameters, BasePass)
 	SHADER_PARAMETER_RDG_UNIFORM_BUFFER(FTranslucentBasePassUniformParameters, TranslucentBasePass)
 	SHADER_PARAMETER_RDG_UNIFORM_BUFFER(FMobileBasePassUniformParameters, MobileBasePass)
@@ -500,6 +501,7 @@ FScreenPassTexture AddEditorPrimitivePass(
 			FEditorPrimitivesPassParameters* PassParameters = GraphBuilder.AllocParameters<FEditorPrimitivesPassParameters>();
 			PassParameters->View = EditorView->GetShaderParameters();
 			PassParameters->ReflectionCapture = View.ReflectionCaptureUniformBuffer;
+			PassParameters->MobileReflectionCaptureData = View.MobileReflectionCaptureUniformBuffer;
 			PassParameters->InstanceCulling = InstanceCullingManager.GetDummyInstanceCullingUniformBuffer();
 			PassParameters->RenderTargets[0] = FRenderTargetBinding(EditorPrimitiveColor, ERenderTargetLoadAction::ELoad);
 			PassParameters->RenderTargets.DepthStencil = FDepthStencilBinding(EditorPrimitiveDepth, ERenderTargetLoadAction::ELoad, ERenderTargetLoadAction::ELoad, FExclusiveDepthStencil::DepthWrite_StencilWrite);
@@ -604,6 +606,7 @@ FScreenPassTexture AddEditorPrimitivePass(
 		FEditorPrimitivesPassParameters* PassParameters = GraphBuilder.AllocParameters<FEditorPrimitivesPassParameters>();
 		PassParameters->View = EditorView->GetShaderParameters();
 		PassParameters->ReflectionCapture = View.ReflectionCaptureUniformBuffer;
+		PassParameters->MobileReflectionCaptureData = View.MobileReflectionCaptureUniformBuffer;
 		PassParameters->InstanceCulling = InstanceCullingManager.GetDummyInstanceCullingUniformBuffer();
 		PassParameters->RenderTargets[0] = Output.GetRenderTargetBinding();
 		PassParameters->RenderTargets[0].SetLoadAction(ERenderTargetLoadAction::ELoad);
diff --git a/Templates/TP_VirtualRealityBP/Config/DefaultEngine.ini b/Templates/TP_VirtualRealityBP/Config/DefaultEngine.ini
--- a/Templates/TP_VirtualRealityBP/Config/DefaultEngine.ini
+++ b/Templates/TP_VirtualRealityBP/Config/DefaultEngine.ini
@@ -23,7 +23,6 @@ r.VT.EnableCompressZlib=True
 r.VT.EnableCompressCrunch=False
 r.ClearCoatNormal=False
 r.ReflectionCaptureResolution=128
-r.Mobile.ReflectionCaptureCompression=False
 r.ReflectionEnvironmentLightmapMixBasedOnRoughness=True
 r.ForwardShading=True
 r.VertexFoggingForOpaque=True
