Add delta velocity info to contact callback info

diff --git a/Engine/Source/Runtime/Engine/Private/PhysicsEngine/Experimental/PhysScene_Chaos.cpp b/Engine/Source/Runtime/Engine/Private/PhysicsEngine/Experimental/PhysScene_Chaos.cpp
--- a/Engine/Source/Runtime/Engine/Private/PhysicsEngine/Experimental/PhysScene_Chaos.cpp
+++ b/Engine/Source/Runtime/Engine/Private/PhysicsEngine/Experimental/PhysScene_Chaos.cpp
@@ -756,15 +756,25 @@ void FPhysScene_Chaos::HandleCollisionEvents(const Chaos::FCollisionEventData& E
 								NotifyInfo.bCallEvent0 = true;
 								// if Comp1 wants this event too, it will get its own pending collision entry, so we leave it false
 
-								NotifyInfo.Info0.SetFrom(GetBodyInstanceFromProxy(PhysicsProxy0));
-								NotifyInfo.Info1.SetFrom(GetBodyInstanceFromProxy(PhysicsProxy1));
+								// Get swapped velocity deltas
+								const FVector& DeltaVelocity1 = bSwapOrder
+									? CollisionDataItem.DeltaVelocity2
+									: CollisionDataItem.DeltaVelocity1;
+								const FVector& DeltaVelocity2 = bSwapOrder
+									? CollisionDataItem.DeltaVelocity1
+									: CollisionDataItem.DeltaVelocity2;
+
+								NotifyInfo.Info0.SetFrom(GetBodyInstanceFromProxy(PhysicsProxy0), DeltaVelocity1);
+								NotifyInfo.Info1.SetFrom(GetBodyInstanceFromProxy(PhysicsProxy1), DeltaVelocity2);
 
 								FRigidBodyContactInfo& NewContact = NotifyInfo.RigidCollisionData.ContactInfos.AddZeroed_GetRef();
 								NewContact.ContactNormal = CollisionDataItem.Normal;
 								NewContact.ContactPosition = CollisionDataItem.Location;
 								NewContact.ContactPenetration = CollisionDataItem.PenetrationDepth;
 								NewContact.bContactProbe = CollisionDataItem.bProbe;
-								NotifyInfo.RigidCollisionData.bIsVelocityDeltaUnderThreshold = CollisionDataItem.DeltaVelocity1.IsNearlyZero(MinDeltaVelocityThreshold) && CollisionDataItem.DeltaVelocity2.IsNearlyZero(MinDeltaVelocityThreshold);
+								NotifyInfo.RigidCollisionData.bIsVelocityDeltaUnderThreshold =
+									DeltaVelocity1.IsNearlyZero(MinDeltaVelocityThreshold) &&
+									DeltaVelocity2.IsNearlyZero(MinDeltaVelocityThreshold);
 								
 								Chaos::FChaosPhysicsMaterial* InternalMat1 = CollisionDataItem.Mat1.Get();
 								Chaos::FChaosPhysicsMaterial* InternalMat2 = CollisionDataItem.Mat2.Get();
diff --git a/Engine/Source/Runtime/Engine/Private/PhysicsEngine/PhysUtils.cpp b/Engine/Source/Runtime/Engine/Private/PhysicsEngine/PhysUtils.cpp
--- a/Engine/Source/Runtime/Engine/Private/PhysicsEngine/PhysUtils.cpp
+++ b/Engine/Source/Runtime/Engine/Private/PhysicsEngine/PhysUtils.cpp
@@ -150,12 +150,13 @@ bool UBodySetup::CreateFromModel(UModel* InModel, bool bRemoveExisting)
 //////////////////////////////////////////////////////////////////////////
 // FRigidBodyCollisionInfo
 
-void FRigidBodyCollisionInfo::SetFrom(const FBodyInstance* BodyInst)
+void FRigidBodyCollisionInfo::SetFrom(const FBodyInstance* BodyInst, const FVector& InDeltaVelocity)
 {
 	if(BodyInst != NULL)
 	{
 		BodyIndex = BodyInst->InstanceBodyIndex;
 		BoneName = BodyInst->BodySetup.IsValid() ? BodyInst->BodySetup->BoneName : NAME_None;
+		DeltaVelocity = InDeltaVelocity;
 
 		if(BodyInst->OwnerComponent.IsValid())
 		{
@@ -169,6 +170,7 @@ void FRigidBodyCollisionInfo::SetFrom(const FBodyInstance* BodyInst)
 		Actor = NULL;
 		BodyIndex = INDEX_NONE;
 		BoneName = NAME_None;
+		DeltaVelocity = FVector::ZeroVector;
 	}
 }
 
diff --git a/Engine/Source/Runtime/Engine/Public/PhysicsPublic.h b/Engine/Source/Runtime/Engine/Public/PhysicsPublic.h
--- a/Engine/Source/Runtime/Engine/Public/PhysicsPublic.h
+++ b/Engine/Source/Runtime/Engine/Public/PhysicsPublic.h
@@ -77,13 +77,16 @@ struct ENGINE_API FRigidBodyCollisionInfo
 	/** Name of bone if a PhysicsAsset */
 	FName									BoneName;
 
+	/** Amount by which the linear velocity at the center of mass of this body has changed in the frame of this contact */
+	FVector									DeltaVelocity;
+
 	FRigidBodyCollisionInfo() :
 		BodyIndex(INDEX_NONE),
 		BoneName(NAME_None)
 	{}
 
 	/** Utility to set up the body collision info from an FBodyInstance */
-	void SetFrom(const FBodyInstance* BodyInst);
+	void SetFrom(const FBodyInstance* BodyInst, const FVector& InDeltaVelocity = FVector::ZeroVector);
 	/** Get body instance */
 	FBodyInstance* GetBodyInstance() const;
 };
