StaticLighting: Rework FLightmapResourceCluster uniform buffer initialization

diff --git a/Engine/Plugins/Experimental/GPULightmass/Source/GPULightmass/Private/LightmapStorage.cpp b/Engine/Plugins/Experimental/GPULightmass/Source/GPULightmass/Private/LightmapStorage.cpp
--- a/Engine/Plugins/Experimental/GPULightmass/Source/GPULightmass/Private/LightmapStorage.cpp
+++ b/Engine/Plugins/Experimental/GPULightmass/Source/GPULightmass/Private/LightmapStorage.cpp
@@ -227,9 +227,7 @@ void CreateLightmapPreviewVirtualTexture(FLightmapRenderStateRef LightmapRenderS
 		}
 	}
 				
-	LightmapRenderState->ResourceCluster->OverrideFeatureLevel(FeatureLevel);
-	LightmapRenderState->ResourceCluster->UpdateUniformBuffer_RenderThread();
-
+	LightmapRenderState->ResourceCluster->SetFeatureLevelAndInitialize(FeatureLevel);
 	LightmapRenderState->SetResourceCluster(nullptr);
 }
 
diff --git a/Engine/Source/Runtime/Engine/Private/LightMap.cpp b/Engine/Source/Runtime/Engine/Private/LightMap.cpp
--- a/Engine/Source/Runtime/Engine/Private/LightMap.cpp
+++ b/Engine/Source/Runtime/Engine/Private/LightMap.cpp
@@ -3375,56 +3375,93 @@ FLightmapResourceCluster::~FLightmapResourceCluster()
 	check(AllocatedVT == nullptr);
 }
 
-void FLightmapResourceCluster::UpdateUniformBuffer(ERHIFeatureLevel::Type InFeatureLevel)
+bool FLightmapResourceCluster::GetUseVirtualTexturing() const
 {
-	FLightmapResourceCluster* Cluster = this;
+	return (CVarVirtualTexturedLightMaps.GetValueOnRenderThread() != 0) && UseVirtualTexturing(GetFeatureLevel());
+}
+
+// Two stage initialization of FLightmapResourceCluster
+// 1. when UMapBuildDataRegistry is post-loaded and render resource is initialized
+// 2. when the level is made visible (ULevel::InitializeRenderingResources()), which calls UMapBuildDataRegistry::InitializeClusterRenderingResources() and fills FeatureLevel
+// When both parts are provided, TryInitialize() creates the final UB with actual content
+// Otherwise UniformBuffer is created with empty parameters
+void FLightmapResourceCluster::TryInitializeUniformBuffer()
+{
+	check(IsInRenderingThread());
+
+	FLightmapResourceClusterShaderParameters Parameters;
+
+	if (HasValidFeatureLevel())
+	{
+		ConditionalCreateAllocatedVT();
+		GetLightmapClusterResourceParameters(GetFeatureLevel(), Input, GetUseVirtualTexturing() ? GetAllocatedVT() : nullptr, Parameters);
+	}
+	else
+	{
+		GetLightmapClusterResourceParameters(GMaxRHIFeatureLevel, FLightmapClusterResourceInput(), nullptr, Parameters);
+	}
 
-	ENQUEUE_RENDER_COMMAND(SetFeatureLevel)(
-		[Cluster, InFeatureLevel](FRHICommandList& RHICmdList)
+	if (!UniformBuffer.IsValid())
 	{
-		Cluster->SetFeatureLevel(InFeatureLevel);
-		Cluster->UpdateUniformBuffer_RenderThread();
-	});
+		UniformBuffer = FLightmapResourceClusterShaderParameters::CreateUniformBuffer(Parameters, UniformBuffer_MultiFrame);
+	}
+	else
+	{
+		RHIUpdateUniformBuffer(UniformBuffer, &Parameters);
+	}
 }
 
-bool FLightmapResourceCluster::GetUseVirtualTexturing() const
+void FLightmapResourceCluster::SetFeatureLevelAndInitialize(const FStaticFeatureLevel InFeatureLevel)
 {
-	return (CVarVirtualTexturedLightMaps.GetValueOnRenderThread() != 0) && UseVirtualTexturing(GetFeatureLevel());
+	check(IsInRenderingThread());
+	SetFeatureLevel(InFeatureLevel);
+	if(IsInitialized() && GIsRHIInitialized)
+	{
+		TryInitializeUniformBuffer();
+	}
 }
 
-void FLightmapResourceCluster::UpdateUniformBuffer_RenderThread()
+void FLightmapResourceCluster::UpdateUniformBuffer()
 {
 	check(IsInRenderingThread());
 
-	FLightmapResourceClusterShaderParameters Parameters;
-	GetLightmapClusterResourceParameters(GetFeatureLevel(), Input, GetUseVirtualTexturing() ? AcquireAllocatedVT() : nullptr, Parameters);
+	if (UniformBuffer.IsValid())
+	{
+		check(HasValidFeatureLevel());
+	
+		FLightmapResourceClusterShaderParameters Parameters;
+		GetLightmapClusterResourceParameters(GetFeatureLevel(), Input, GetUseVirtualTexturing() ? GetAllocatedVT() : nullptr, Parameters);
 
-	RHIUpdateUniformBuffer(UniformBuffer, &Parameters);
+		RHIUpdateUniformBuffer(UniformBuffer, &Parameters);
+	}
 }
 
 static void OnVirtualTextureDestroyed(const FVirtualTextureProducerHandle& InHandle, void* Baton)
 {
 	FLightmapResourceCluster* Cluster = static_cast<FLightmapResourceCluster*>(Baton);
 	Cluster->ReleaseAllocatedVT();
-	Cluster->UpdateUniformBuffer_RenderThread();
+	Cluster->UpdateUniformBuffer();
 }
 
-IAllocatedVirtualTexture* FLightmapResourceCluster::AcquireAllocatedVT() const
+const IAllocatedVirtualTexture* FLightmapResourceCluster::GetAllocatedVT() const
 {
-	check(IsInRenderingThread());
-
-	bool bHighQuality = AllowHighQualityLightmaps(GetFeatureLevel());
-
-	int32 LightmapIndex = bHighQuality ? 0 : 1;
+	check(IsInParallelRenderingThread());
+	return AllocatedVT;
+}
 
-	const ULightMapVirtualTexture2D* VirtualTexture = Input.LightMapVirtualTextures[LightmapIndex];
+void FLightmapResourceCluster::ConditionalCreateAllocatedVT()
+{
+	check(IsInRenderingThread());
+	check(HasValidFeatureLevel());
+	
+	const ULightMapVirtualTexture2D* VirtualTexture = Input.LightMapVirtualTextures[AllowHighQualityLightmaps(GetFeatureLevel()) ? 0 : 1];
 	
 #if WITH_EDITOR
-		// Compilation is still pending, this function will be called back once compilation finishes.
+	// Compilation is still pending, this function will be called back once compilation finishes.
 	if (VirtualTexture && VirtualTexture->IsCompiling())
-		{
-			return nullptr;
-		}
+	{
+		return;
+	}
 #endif
 
 	if (!AllocatedVT && VirtualTexture && VirtualTexture->GetResource())
@@ -3472,8 +3509,6 @@ IAllocatedVirtualTexture* FLightmapResourceCluster::AcquireAllocatedVT() const
 
 		AllocatedVT = GetRendererModule().AllocateVirtualTexture(VTDesc);
 	}
-
-	return AllocatedVT;
 }
 
 void FLightmapResourceCluster::ReleaseAllocatedVT()
@@ -3489,22 +3524,7 @@ void FLightmapResourceCluster::ReleaseAllocatedVT()
 void FLightmapResourceCluster::InitRHI()
 {
 	SCOPED_LOADTIMER(FLightmapResourceCluster_InitRHI);
-
-	FLightmapResourceClusterShaderParameters Parameters;
-
-	// Lightmap resources are normally created before the feature level is known, so we'll use defaults and rely on a subsequent call to UpdateUniformBuffer()
-	// to set the correct level and update things accordingly. However, when we're coming from FRenderResource::ChangeFeatureLevel(), the feature level
-	// has already been set, so we can go ahead and use the correct level and input from the start (UpdateUniformBuffer() is not being called in that case).
-	if (HasValidFeatureLevel())
-	{
-		GetLightmapClusterResourceParameters(GetFeatureLevel(), Input, GetUseVirtualTexturing() ? AcquireAllocatedVT() : nullptr, Parameters);
-	}
-	else
-	{
-		GetLightmapClusterResourceParameters(GMaxRHIFeatureLevel, FLightmapClusterResourceInput(), nullptr, Parameters);
-	}
-
-	UniformBuffer = FLightmapResourceClusterShaderParameters::CreateUniformBuffer(Parameters, UniformBuffer_MultiFrame);
+	TryInitializeUniformBuffer();
 }
 
 void FLightmapResourceCluster::ReleaseRHI()
diff --git a/Engine/Source/Runtime/Engine/Private/LightmapUniformShaderParameters.cpp b/Engine/Source/Runtime/Engine/Private/LightmapUniformShaderParameters.cpp
--- a/Engine/Source/Runtime/Engine/Private/LightmapUniformShaderParameters.cpp
+++ b/Engine/Source/Runtime/Engine/Private/LightmapUniformShaderParameters.cpp
@@ -101,7 +101,7 @@ void GetPrecomputedLightingParameters(
 			check(LCI); // If LCI was nullptr, LightMapInteraction.GetType() would be LMIT_None, not LMIT_Texture
 			const FLightmapResourceCluster* ResourceCluster = LCI->GetResourceCluster();
 			check(ResourceCluster);
-			IAllocatedVirtualTexture* AllocatedVT = ResourceCluster->AcquireAllocatedVT();
+			const IAllocatedVirtualTexture* AllocatedVT = ResourceCluster->GetAllocatedVT();
 			if (AllocatedVT)
 			{
 				AllocatedVT->GetPackedPageTableUniform(&Parameters.LightmapVTPackedPageTableUniform[0]);
diff --git a/Engine/Source/Runtime/Engine/Private/MapBuildData.cpp b/Engine/Source/Runtime/Engine/Private/MapBuildData.cpp
--- a/Engine/Source/Runtime/Engine/Private/MapBuildData.cpp
+++ b/Engine/Source/Runtime/Engine/Private/MapBuildData.cpp
@@ -408,14 +408,14 @@ void UMapBuildDataRegistry::HandleAssetPostCompileEvent(const TArray<FAssetCompi
 	{
 		TRACE_CPUPROFILER_EVENT_SCOPE(UMapBuildDataRegistry::HandleAssetPostCompileEvent);
 
-		for (FLightmapResourceCluster* Cluster : ClustersToUpdate)
-		{
-			ENQUEUE_RENDER_COMMAND(UpdateClusterUniformBuffer)(
-				[Cluster](FRHICommandList& RHICmdList)
+		ENQUEUE_RENDER_COMMAND(UpdateClusterUniformBuffer)(
+			[ClustersToUpdate = ClustersToUpdate](FRHICommandList& RHICmdList)
+			{
+				for (FLightmapResourceCluster* Cluster : ClustersToUpdate)
 				{
-					Cluster->UpdateUniformBuffer_RenderThread();
-				});
-		}
+					Cluster->UpdateRHI();
+				}
+			});
 
 		for (TObjectIterator<ULandscapeComponent> It; It; ++It)
 		{
@@ -1010,12 +1010,16 @@ void UMapBuildDataRegistry::InitializeClusterRenderingResources(ERHIFeatureLevel
 	check(bSetupResourceClusters || MeshBuildData.Num() == 0);
 	// If we have any mesh build data, we must have at least one resource cluster, otherwise clusters have not been setup properly.
 	check(LightmapResourceClusters.Num() > 0 || MeshBuildData.Num() == 0);
-
-	// At this point all lightmap cluster resources are initialized and we can update cluster uniform buffers.
-	for (FLightmapResourceCluster& Cluster : LightmapResourceClusters)
-	{
-		Cluster.UpdateUniformBuffer(InFeatureLevel);
-	}
+	
+	ENQUEUE_RENDER_COMMAND(SetFeatureLevelAndInitialize)(
+		[&LightmapResourceClusters = LightmapResourceClusters, InFeatureLevel](FRHICommandList& RHICmdList)
+		{
+			// At this point all lightmap cluster resources are initialized and we can update cluster uniform buffers.
+			for (FLightmapResourceCluster& Cluster : LightmapResourceClusters)
+			{
+				Cluster.SetFeatureLevelAndInitialize(InFeatureLevel);
+			}
+		});
 }
 
 void UMapBuildDataRegistry::ReleaseResources(const TSet<FGuid>* ResourcesToKeep)
diff --git a/Engine/Source/Runtime/Engine/Private/SceneManagement.cpp b/Engine/Source/Runtime/Engine/Private/SceneManagement.cpp
--- a/Engine/Source/Runtime/Engine/Private/SceneManagement.cpp
+++ b/Engine/Source/Runtime/Engine/Private/SceneManagement.cpp
@@ -978,7 +978,7 @@ static FRHISamplerState* GetTextureSamplerState(const UTexture* Texture, FRHISam
 void GetLightmapClusterResourceParameters(
 	ERHIFeatureLevel::Type FeatureLevel, 
 	const FLightmapClusterResourceInput& Input,
-	IAllocatedVirtualTexture* AllocatedVT,
+	const IAllocatedVirtualTexture* AllocatedVT,
 	FLightmapResourceClusterShaderParameters& Parameters)
 {
 	const bool bAllowHighQualityLightMaps = AllowHighQualityLightmaps(FeatureLevel);
diff --git a/Engine/Source/Runtime/Engine/Public/LightMap.h b/Engine/Source/Runtime/Engine/Public/LightMap.h
--- a/Engine/Source/Runtime/Engine/Public/LightMap.h
+++ b/Engine/Source/Runtime/Engine/Public/LightMap.h
@@ -611,17 +611,17 @@ public:
 	virtual void InitRHI();
 	virtual void ReleaseRHI();
 
-	void OverrideFeatureLevel(const FStaticFeatureLevel InFeatureLevel) { SetFeatureLevel(InFeatureLevel); }
+	void SetFeatureLevelAndInitialize(const FStaticFeatureLevel InFeatureLevel);
+	void TryInitializeUniformBuffer();
+	void UpdateUniformBuffer();
 	
-	void UpdateUniformBuffer(ERHIFeatureLevel::Type InFeatureLevel);
-	void UpdateUniformBuffer_RenderThread();
-
 	/**
 	 * Allocates virtual texture on demand and returns it, may return nullptr if not using virtual texture
 	 * 'const' method setting 'mutable' member is not amazing, but this is required to work around ordering of render commands submitted from main thread,
 	 * it's possible for commands that want to access AllocatedVT from here execute before this has a chance to run InitRHI()
 	 */
-	IAllocatedVirtualTexture* AcquireAllocatedVT() const;
+	const IAllocatedVirtualTexture* GetAllocatedVT() const;
+	void ConditionalCreateAllocatedVT();
 	void ReleaseAllocatedVT();
 	bool GetUseVirtualTexturing() const;
 
diff --git a/Engine/Source/Runtime/Engine/Public/SceneManagement.h b/Engine/Source/Runtime/Engine/Public/SceneManagement.h
--- a/Engine/Source/Runtime/Engine/Public/SceneManagement.h
+++ b/Engine/Source/Runtime/Engine/Public/SceneManagement.h
@@ -829,7 +829,7 @@ public:
 ENGINE_API void GetLightmapClusterResourceParameters(
 	ERHIFeatureLevel::Type FeatureLevel, 
 	const FLightmapClusterResourceInput& Input,
-	IAllocatedVirtualTexture* AllocatedVT,
+	const IAllocatedVirtualTexture* AllocatedVT,
 	FLightmapResourceClusterShaderParameters& Parameters);
 
 class FDefaultLightmapResourceClusterUniformBuffer : public TUniformBuffer< FLightmapResourceClusterShaderParameters >
