Un-deprecating property until new set is migrated to key projects and new calls to old property can be properly diverted
#rnx

diff --git a/Engine/Source/Runtime/Engine/Classes/Sound/SoundModulationDestination.h b/Engine/Source/Runtime/Engine/Classes/Sound/SoundModulationDestination.h
--- a/Engine/Source/Runtime/Engine/Classes/Sound/SoundModulationDestination.h
+++ b/Engine/Source/Runtime/Engine/Classes/Sound/SoundModulationDestination.h
@@ -53,11 +53,11 @@ struct ENGINE_API FSoundModulationDestinationSettings
 	/** Whether or not to enable modulation */
 	UPROPERTY(EditAnywhere, Category = Modulation, meta = (DisplayName = "Modulate"))
 	bool bEnableModulation = false;
+#endif // WITH_EDITORONLY_DATA
 
 	/** Deprecated in favor of Modulators set. */
-	UPROPERTY(meta = (DeprecatedProperty))
+	UPROPERTY(EditAnywhere, BlueprintReadWrite, Category = Modulation)
 	TObjectPtr<USoundModulatorBase> Modulator = nullptr;
-#endif // WITH_EDITORONLY_DATA
 
 	/** Set of modulation sources, which provides values to mix with base value. */
 	UPROPERTY(EditAnywhere, BlueprintReadWrite, Category = Modulation)
