Remove ensure causing cooks to fail
#rnx

diff --git a/Engine/Source/Runtime/Engine/Private/WorldPartition/HLOD/HLODActor.cpp b/Engine/Source/Runtime/Engine/Private/WorldPartition/HLOD/HLODActor.cpp
--- a/Engine/Source/Runtime/Engine/Private/WorldPartition/HLOD/HLODActor.cpp
+++ b/Engine/Source/Runtime/Engine/Private/WorldPartition/HLOD/HLODActor.cpp
@@ -111,7 +111,7 @@ void AWorldPartitionHLOD::PostLoad()
 			// Strip "WorldName_" from the cell name
 			FString CellName = SourceCellName.ToString();
 			bool bRemoved = CellName.RemoveFromStart(WorldName + TEXT("_"), ESearchCase::CaseSensitive);
-			if (ensure(bRemoved))
+			if (bRemoved)
 			{
 				SourceCellName = *CellName;
 			}
