#UE Use the FMemoryView directly from MeshDataBuffer when getting cached FStaticMeshRenderData. This will better handle cases where the buffer size is bigger than 2gb

diff --git a/Engine/Source/Runtime/Engine/Private/StaticMesh.cpp b/Engine/Source/Runtime/Engine/Private/StaticMesh.cpp
--- a/Engine/Source/Runtime/Engine/Private/StaticMesh.cpp
+++ b/Engine/Source/Runtime/Engine/Private/StaticMesh.cpp
@@ -2894,7 +2894,7 @@ void FStaticMeshRenderData::Cache(const ITargetPlatform* TargetPlatform, UStatic
 		{
 			COOK_STAT(Timer.AddHit(MeshDataBuffer.GetSize()));
 
-			FMemoryReaderView Ar(MakeMemoryView((const uint8*)MeshDataBuffer.GetData(), MeshDataBuffer.GetSize()), /*bIsPersistent=*/ true);
+			FMemoryReaderView Ar(MeshDataBuffer.GetView(), /*bIsPersistent=*/ true);
 			Serialize(Ar, Owner, /*bCooked=*/ false);
 
 			// Reconstruct EstimatedCompressedSize
