Fix t-pose when linking in a new linked anim graph and compiling

Regression caused by new re-instancing method
Ensures that a 'full' initialization occurs in editor preview (i.e. anim BP editor) worlds.

diff --git a/Engine/Source/Runtime/Engine/Private/Animation/AnimInstance.cpp b/Engine/Source/Runtime/Engine/Private/Animation/AnimInstance.cpp
--- a/Engine/Source/Runtime/Engine/Private/Animation/AnimInstance.cpp
+++ b/Engine/Source/Runtime/Engine/Private/Animation/AnimInstance.cpp
@@ -3857,12 +3857,21 @@ void UAnimInstance::HandleObjectsReinstanced(const TMap<UObject*, UObject*>& Old
 
 			RecalcRequiredBones();
 			
-			// Reinit proxy
+			// Minimally reinit proxy (i.e. dont call per-node initialization) unless we are in an editor preview world (i.e. we are in the anim BP editor)
 			FAnimInstanceProxy& Proxy = GetProxyOnGameThread<FAnimInstanceProxy>();
-			Proxy.Initialize(this);
-			Proxy.InitializeCachedClassData();
-			Proxy.InitializeRootNode_WithRoot(Proxy.RootNode);
-	
+			UWorld* World = GetWorld();
+			if(World && World->WorldType == EWorldType::EditorPreview)
+			{
+				Proxy.Initialize(this);
+				Proxy.InitializeRootNode();
+			}
+			else
+			{
+				Proxy.Initialize(this);
+				Proxy.InitializeCachedClassData();
+				Proxy.InitializeRootNode_WithRoot(Proxy.RootNode);
+			}
+
 			if(USkeletalMeshComponent* Mesh = GetSkelMeshComponent())
 			{
 				Mesh->ClearMotionVector();
