Texture.cpp GetDefaultTextureFormatName don't change the HDR format selection due to the need for uncompresed texture output if it was already uncompressed

diff --git a/Engine/Source/Runtime/Engine/Private/Texture.cpp b/Engine/Source/Runtime/Engine/Private/Texture.cpp
--- a/Engine/Source/Runtime/Engine/Private/Texture.cpp
+++ b/Engine/Source/Runtime/Engine/Private/Texture.cpp
@@ -3083,9 +3083,30 @@ FName GetDefaultTextureFormatName( const ITargetPlatform* TargetPlatform, const
 
 		if (Texture->HasHDRSource(LayerIndex))
 		{
-			// @todo Oodle : this is not the best possible selection of output format from SourceFormat
-			// R16F and R32F is available but not used here even if their TC_ would have chosen them!
-			TextureFormatName = NameRGBA16F;
+			// if CompressionSettings was already an uncompressed HDR format, use it
+			//  this is a conservative/partial fix of the more general problem of bIsTCThatMapsToUncompressed (see 1. above)
+			if (FormatSettings.CompressionSettings == TC_HDR)
+			{
+				TextureFormatName = NameRGBA16F;
+			}
+			else if (FormatSettings.CompressionSettings == TC_HDR_F32)
+			{
+				TextureFormatName = NameRGBA32F;
+			}
+			else if (FormatSettings.CompressionSettings == TC_HalfFloat)
+			{
+				TextureFormatName = NameR16F;
+			}
+			else if (FormatSettings.CompressionSettings == TC_SingleFloat)
+			{
+				TextureFormatName = NameR32F;
+			}
+			else
+			{
+				// @todo Oodle : this is not the best possible selection of output format from SourceFormat
+				// R16F and R32F is available but not used here even if their TC_ would have chosen them!
+				TextureFormatName = NameRGBA16F;
+			}
 		}
 		else if (SourceFormat == TSF_G16)
 		{
