Fix Strata crash when opening material layer blend asset.

diff --git a/Engine/Source/Runtime/Engine/Private/Materials/Material.cpp b/Engine/Source/Runtime/Engine/Private/Materials/Material.cpp
--- a/Engine/Source/Runtime/Engine/Private/Materials/Material.cpp
+++ b/Engine/Source/Runtime/Engine/Private/Materials/Material.cpp
@@ -5901,7 +5901,7 @@ int32 UMaterial::CompilePropertyEx( FMaterialCompiler* Compiler, const FGuid& At
 		return EditorOnly->ExpressionCollection.ExpressionExecBegin->Compile(Compiler, UMaterialExpression::CompileExecutionOutputIndex);
 	}
 
-	if( bUseMaterialAttributes && MP_DiffuseColor != Property && MP_SpecularColor != Property )
+	if( bUseMaterialAttributes && MP_DiffuseColor != Property && MP_SpecularColor != Property && MP_FrontMaterial != Property)
 	{
 		return EditorOnly->MaterialAttributes.CompileWithDefault(Compiler, AttributeID);
 	}
