World Partition Streaming Generation:
- Moved generic runtime hash code to the base class to ease working on new hash types.
- Hashes can now support HLODs generation or not, hide related UI properties when they don't.

Tested Ancient/City/Lyra PIE/Game/Cook.

#rnx

diff --git a/Engine/Source/Editor/WorldPartitionEditor/Private/WorldPartition/Customizations/WorldPartitionDetailsCustomization.cpp b/Engine/Source/Editor/WorldPartitionEditor/Private/WorldPartition/Customizations/WorldPartitionDetailsCustomization.cpp
--- a/Engine/Source/Editor/WorldPartitionEditor/Private/WorldPartition/Customizations/WorldPartitionDetailsCustomization.cpp
+++ b/Engine/Source/Editor/WorldPartitionEditor/Private/WorldPartition/Customizations/WorldPartitionDetailsCustomization.cpp
@@ -88,6 +88,11 @@ void FWorldPartitionDetails::CustomizeDetails(IDetailLayoutBuilder& InDetailBuil
 				.DisplayName(LOCTEXT("RuntimeHash", "Runtime Hash"))
 				.Visibility(TAttribute<EVisibility>::CreateLambda([this]() { return WorldPartition.IsValid() && WorldPartition->IsStreamingEnabled() ? EVisibility::Visible : EVisibility::Hidden; }));
 		}
+
+		if (!WorldPartition->RuntimeHash->SupportsHLODs())
+		{
+			InDetailBuilder.HideProperty(InDetailBuilder.GetProperty(GET_MEMBER_NAME_CHECKED(UWorldPartition, DefaultHLODLayer), UWorldPartition::StaticClass()));
+		}
 	}
 }
 
diff --git a/Engine/Source/Runtime/Engine/Private/WorldPartition/ContentBundle/ContentBundleEditor.cpp b/Engine/Source/Runtime/Engine/Private/WorldPartition/ContentBundle/ContentBundleEditor.cpp
--- a/Engine/Source/Runtime/Engine/Private/WorldPartition/ContentBundle/ContentBundleEditor.cpp
+++ b/Engine/Source/Runtime/Engine/Private/WorldPartition/ContentBundle/ContentBundleEditor.cpp
@@ -298,7 +298,11 @@ void FContentBundleEditor::GenerateStreaming()
 	UWorldPartition* WorldPartition = GetInjectedWorld()->GetWorldPartition();
 	WorldPartition->GenerateContainerStreaming(ActorDescContainer.Get());
 
-	WorldPartition->RuntimeHash->GetAllStreamingCells(ContentBundleCells, true);
+	WorldPartition->RuntimeHash->ForEachStreamingCells([this](const UWorldPartitionRuntimeCell* Cell)
+	{
+		ContentBundleCells.Add(Cell);
+		return true;
+	});
 
 	FString ExternalStreamingObjectName = ObjectTools::SanitizeInvalidChars(GetDisplayName() + TEXT("_ExternalStreamingObject"), INVALID_LONGPACKAGE_CHARACTERS) ;
 	ExternalStreamingObject = WorldPartition->RuntimeHash->StoreToExternalStreamingObject(GetInjectedWorld()->ContentBundleManager, *ExternalStreamingObjectName);
diff --git a/Engine/Source/Runtime/Engine/Private/WorldPartition/HLOD/HLODSubsystem.cpp b/Engine/Source/Runtime/Engine/Private/WorldPartition/HLOD/HLODSubsystem.cpp
--- a/Engine/Source/Runtime/Engine/Private/WorldPartition/HLOD/HLODSubsystem.cpp
+++ b/Engine/Source/Runtime/Engine/Private/WorldPartition/HLOD/HLODSubsystem.cpp
@@ -171,17 +171,15 @@ void UHLODSubsystem::OnWorldPartitionInitialized(UWorldPartition* InWorldPartiti
 		check(!WorldPartitionsHLODRuntimeData.Contains(InWorldPartition));
 
 		FWorldPartitionHLODRuntimeData& WorldPartitionHLODRuntimeData = WorldPartitionsHLODRuntimeData.Add(InWorldPartition, FWorldPartitionHLODRuntimeData());
-
-		TSet<const UWorldPartitionRuntimeCell*> StreamingCells;
-		if (InWorldPartition->RuntimeHash)
-		{
-			InWorldPartition->RuntimeHash->GetAllStreamingCells(StreamingCells, /*bAllDataLayers*/ true);
-		}
-
+	
 		// Build cell to HLOD mapping
-		for (const UWorldPartitionRuntimeCell* Cell : StreamingCells)
+		if (InWorldPartition && InWorldPartition->RuntimeHash)
 		{
-			WorldPartitionHLODRuntimeData.CellsData.Emplace(Cell->GetFName());
+			InWorldPartition->RuntimeHash->ForEachStreamingCells([&WorldPartitionHLODRuntimeData](const UWorldPartitionRuntimeCell* Cell)
+			{
+				WorldPartitionHLODRuntimeData.CellsData.Emplace(Cell->GetFName());
+				return true;
+			});
 		}
 	}
 }
diff --git a/Engine/Source/Runtime/Engine/Private/WorldPartition/WorldPartitionLevelStreamingPolicy.cpp b/Engine/Source/Runtime/Engine/Private/WorldPartition/WorldPartitionLevelStreamingPolicy.cpp
--- a/Engine/Source/Runtime/Engine/Private/WorldPartition/WorldPartitionLevelStreamingPolicy.cpp
+++ b/Engine/Source/Runtime/Engine/Private/WorldPartition/WorldPartitionLevelStreamingPolicy.cpp
@@ -71,11 +71,8 @@ TSubclassOf<UWorldPartitionRuntimeCell> UWorldPartitionLevelStreamingPolicy::Get
 
 void UWorldPartitionLevelStreamingPolicy::PrepareActorToCellRemapping()
 {
-	TSet<const UWorldPartitionRuntimeCell*> StreamingCells;
-	WorldPartition->RuntimeHash->GetAllStreamingCells(StreamingCells, /*bAllDataLayers*/ true);
-
 	// Build Actor-to-Cell remapping
-	for (const UWorldPartitionRuntimeCell* Cell : StreamingCells)
+	WorldPartition->RuntimeHash->ForEachStreamingCells([this](const UWorldPartitionRuntimeCell* Cell)
 	{
 		const UWorldPartitionRuntimeLevelStreamingCell* StreamingCell = Cast<const UWorldPartitionRuntimeLevelStreamingCell>(Cell);
 		check(StreamingCell);
@@ -99,7 +96,8 @@ void UWorldPartitionLevelStreamingPolicy::PrepareActorToCellRemapping()
 				SubObjectsToCellRemapping.Add(FName(*ActorPath.Mid(LastDotPos + 1)), StreamingCell->GetFName());
 			}
 		}
-	}
+		return true;
+	});
 }
 
 void UWorldPartitionLevelStreamingPolicy::RemapSoftObjectPath(FSoftObjectPath& ObjectPath)
diff --git a/Engine/Source/Runtime/Engine/Private/WorldPartition/WorldPartitionRuntimeCell.cpp b/Engine/Source/Runtime/Engine/Private/WorldPartition/WorldPartitionRuntimeCell.cpp
--- a/Engine/Source/Runtime/Engine/Private/WorldPartition/WorldPartitionRuntimeCell.cpp
+++ b/Engine/Source/Runtime/Engine/Private/WorldPartition/WorldPartitionRuntimeCell.cpp
@@ -12,15 +12,15 @@
 int32 UWorldPartitionRuntimeCell::StreamingSourceCacheEpoch = 0;
 
 UWorldPartitionRuntimeCell::UWorldPartitionRuntimeCell(const FObjectInitializer& ObjectInitializer)
-: Super(ObjectInitializer)
-, bIsAlwaysLoaded(false)
-, MinMaxZ(-HALF_WORLD_MAX, HALF_WORLD_MAX)
-, Priority(0)
-, CachedMinSourcePriority((uint8)EStreamingSourcePriority::Lowest)
-, CachedSourceInfoEpoch(INT_MIN)
+	: Super(ObjectInitializer)
+	, bIsAlwaysLoaded(false)
+	, MinMaxZ(-HALF_WORLD_MAX, HALF_WORLD_MAX)
+	, Priority(0)
+	, CachedMinSourcePriority((uint8)EStreamingSourcePriority::Lowest)
+	, CachedSourceInfoEpoch(INT_MIN)
 #if !UE_BUILD_SHIPPING
-, DebugStreamingPriority(-1.f)
-#endif
+	, DebugStreamingPriority(-1.f)
+#endif	
 {}
 
 #if WITH_EDITOR
@@ -96,24 +96,29 @@ void UWorldPartitionRuntimeCell::DumpStateLog(FHierarchicalLogArchive& Ar)
 
 #endif
 
-bool UWorldPartitionRuntimeCell::CacheStreamingSourceInfo(const UWorldPartitionRuntimeCell::FStreamingSourceInfo& Info) const
+bool UWorldPartitionRuntimeCell::ShouldResetStreamingSourceInfo() const
 {
-	bool bWasCacheDirtied = false;
-	if (CachedSourceInfoEpoch != UWorldPartitionRuntimeCell::StreamingSourceCacheEpoch)
-	{
-		CachedSourceInfoEpoch = UWorldPartitionRuntimeCell::StreamingSourceCacheEpoch;
-		bWasCacheDirtied = true;
-		CachedSourcePriorityWeights.Reset();
-	}
+	return CachedSourceInfoEpoch != StreamingSourceCacheEpoch;
+}
 
-	static_assert((uint8)EStreamingSourcePriority::Lowest == 255);
-	CachedSourcePriorityWeights.Add(1.f - ((float)Info.Source.Priority / 255.f));
+void UWorldPartitionRuntimeCell::ResetStreamingSourceInfo() const
+{
+	check(ShouldResetStreamingSourceInfo());
+	CachedSourcePriorityWeights.Reset();
+	CachedMinSourcePriority = MAX_uint8;
+	CachedSourceInfoEpoch = StreamingSourceCacheEpoch;
+}
 
-	// If cache was dirtied, use value, else use minimum with existing cached value
-	CachedMinSourcePriority = bWasCacheDirtied ? (uint8)Info.Source.Priority : FMath::Min((uint8)Info.Source.Priority, CachedMinSourcePriority);
-	return bWasCacheDirtied;
+void UWorldPartitionRuntimeCell::AppendStreamingSourceInfo(const FWorldPartitionStreamingSource& Source, const FSphericalSector& SourceShape) const
+{
+	static_assert((uint8)EStreamingSourcePriority::Lowest == 255);
+	CachedSourcePriorityWeights.Add(1.f - ((float)Source.Priority / 255.f));
+	CachedMinSourcePriority = FMath::Min((uint8)Source.Priority, CachedMinSourcePriority);
 }
 
+void UWorldPartitionRuntimeCell::MergeStreamingSourceInfo() const
+{}
+
 int32 UWorldPartitionRuntimeCell::SortCompare(const UWorldPartitionRuntimeCell* Other) const
 {
 	// Source priority (lower value is higher prio)
@@ -133,8 +138,7 @@ bool UWorldPartitionRuntimeCell::IsDebugShown() const
 FLinearColor UWorldPartitionRuntimeCell::GetDebugStreamingPriorityColor() const
 {
 #if !UE_BUILD_SHIPPING
-	if ((CachedSourceInfoEpoch == UWorldPartitionRuntimeCell::StreamingSourceCacheEpoch) &&
-		(DebugStreamingPriority >= 0.f && DebugStreamingPriority <= 1.f))
+	if (DebugStreamingPriority >= 0.f && DebugStreamingPriority <= 1.f)
 	{
 		return FWorldPartitionDebugHelper::GetHeatMapColor(1.f - DebugStreamingPriority);
 	}
diff --git a/Engine/Source/Runtime/Engine/Private/WorldPartition/WorldPartitionRuntimeHash.cpp b/Engine/Source/Runtime/Engine/Private/WorldPartition/WorldPartitionRuntimeHash.cpp
--- a/Engine/Source/Runtime/Engine/Private/WorldPartition/WorldPartitionRuntimeHash.cpp
+++ b/Engine/Source/Runtime/Engine/Private/WorldPartition/WorldPartitionRuntimeHash.cpp
@@ -38,6 +38,38 @@ void UWorldPartitionRuntimeHash::OnEndPlay()
 	ModifiedActorDescListForPIE.Empty();
 }
 
+// In PIE, Always loaded cell is not generated. Instead, always loaded actors will be added to AlwaysLoadedActorsForPIE.
+// This will trigger loading/registration of these actors in the PersistentLevel (if not already loaded).
+// Then, duplication of world for PIE will duplicate only these actors. 
+// When stopping PIE, WorldPartition will release these FWorldPartitionReferences which 
+// will unload actors that were not already loaded in the non PIE world.
+bool UWorldPartitionRuntimeHash::ConditionalRegisterAlwaysLoadedActorsForPIE(const FWorldPartitionActorDescView& ActorDescView, bool bIsMainWorldPartition, bool bIsMainContainer, bool bIsCellAlwaysLoaded)
+{
+	if (bIsMainWorldPartition && bIsMainContainer && bIsCellAlwaysLoaded && !IsRunningCookCommandlet())
+	{
+		// This will load the actor if it isn't already loaded
+		FWorldPartitionReference Reference(GetOuterUWorldPartition(), ActorDescView.GetGuid());
+
+		if (AActor* AlwaysLoadedActor = FindObject<AActor>(nullptr, *ActorDescView.GetActorPath().ToString()))
+		{
+			AlwaysLoadedActorsForPIE.Emplace(Reference, AlwaysLoadedActor);
+
+			// Handle child actors
+			AlwaysLoadedActor->ForEachComponent<UChildActorComponent>(true, [this, &Reference](UChildActorComponent* ChildActorComponent)
+			{
+				if (AActor* ChildActor = ChildActorComponent->GetChildActor())
+				{
+					AlwaysLoadedActorsForPIE.Emplace(Reference, ChildActor);
+				}
+			});
+		}
+
+		return true;
+	}
+
+	return false;
+}
+
 void UWorldPartitionRuntimeHash::DumpStateLog(FHierarchicalLogArchive& Ar)
 {
 	Ar.Printf(TEXT("----------------------------------------------------------------------------------------------------------------"));
@@ -63,15 +95,70 @@ void UWorldPartitionRuntimeHash::ForceExternalActorLevelReference(bool bForceExt
 }
 #endif
 
-void UWorldPartitionRuntimeHash::SortStreamingCellsByImportance(const TSet<const UWorldPartitionRuntimeCell*>& InCells, const TArray<FWorldPartitionStreamingSource>& InSources, TArray<const UWorldPartitionRuntimeCell*, TInlineAllocator<256>>& OutSortedCells) const
+int32 UWorldPartitionRuntimeHash::GetAllStreamingCells(TSet<const UWorldPartitionRuntimeCell*>& Cells, bool bAllDataLayers, bool bDataLayersOnly, const TSet<FName>& InDataLayers) const
+{
+	ForEachStreamingCells([&Cells, bAllDataLayers, bDataLayersOnly, InDataLayers](const UWorldPartitionRuntimeCell* Cell)
+	{
+		if (!bDataLayersOnly && !Cell->HasDataLayers())
+		{
+			Cells.Add(Cell);
+		}
+		else if (Cell->HasDataLayers() && (bAllDataLayers || Cell->HasAnyDataLayer(InDataLayers)))
+		{
+			Cells.Add(Cell);
+		}
+		return true;
+	});
+
+	return Cells.Num();
+}
+
+bool UWorldPartitionRuntimeHash::GetStreamingCells(const FWorldPartitionStreamingQuerySource& QuerySource, TSet<const UWorldPartitionRuntimeCell*>& OutCells) const
 {
-	TRACE_CPUPROFILER_EVENT_SCOPE(UWorldPartitionRuntimeHash::SortStreamingCellsByImportance);
+	ForEachStreamingCellsQuery(QuerySource, [QuerySource, &OutCells](const UWorldPartitionRuntimeCell* Cell)
+	{
+		OutCells.Add(Cell);
+		return true;
+	});
+
+	return !!OutCells.Num();
+}
 
-	OutSortedCells = InCells.Array();
-	Algo::Sort(OutSortedCells, [](const UWorldPartitionRuntimeCell* CellA, const UWorldPartitionRuntimeCell* CellB)
+bool UWorldPartitionRuntimeHash::GetStreamingCells(const TArray<FWorldPartitionStreamingSource>& Sources, FStreamingSourceCells& OutActivateCells, FStreamingSourceCells& OutLoadCells) const
+{
+	ForEachStreamingCellsSources(Sources, [&OutActivateCells, &OutLoadCells](const UWorldPartitionRuntimeCell* Cell, EStreamingSourceTargetState TargetState)
 	{
-		return CellA->SortCompare(CellB) < 0;
+		switch (TargetState)
+		{
+		case EStreamingSourceTargetState::Loaded:
+			OutLoadCells.GetCells().Add(Cell);
+			break;
+		case EStreamingSourceTargetState::Activated:
+			OutActivateCells.GetCells().Add(Cell);
+			break;
+		}
+		return true;
 	});
+
+	return !!(OutActivateCells.Num() + OutLoadCells.Num());
+}
+
+bool UWorldPartitionRuntimeHash::IsCellRelevantFor(bool bClientOnlyVisible) const
+{
+	if (bClientOnlyVisible)
+	{
+		const UWorld* World = GetWorld();
+		if (World->IsGameWorld())
+		{
+			// Dedicated server & listen server without server streaming won't consider client-only visible cells
+			const ENetMode NetMode = World->GetNetMode();
+			if ((NetMode == NM_DedicatedServer) || ((NetMode == NM_ListenServer) && !GetOuterUWorldPartition()->IsServerStreamingEnabled()))
+			{
+				return false;
+			}
+		}
+	}
+	return true;
 }
 
 EWorldPartitionStreamingPerformance UWorldPartitionRuntimeHash::GetStreamingPerformance(const TSet<const UWorldPartitionRuntimeCell*>& CellsToActivate) const
@@ -104,9 +191,14 @@ EWorldPartitionStreamingPerformance UWorldPartitionRuntimeHash::GetStreamingPerf
 	return StreamingPerformance;
 }
 
-void UWorldPartitionRuntimeHash::FStreamingSourceCells::AddCell(const UWorldPartitionRuntimeCell* Cell, const UWorldPartitionRuntimeCell::FStreamingSourceInfo& Info)
+void UWorldPartitionRuntimeHash::FStreamingSourceCells::AddCell(const UWorldPartitionRuntimeCell* Cell, const FWorldPartitionStreamingSource& Source, const FSphericalSector& SourceShape)
 {
-	Cell->CacheStreamingSourceInfo(Info);
+	if (Cell->ShouldResetStreamingSourceInfo())
+	{
+		Cell->ResetStreamingSourceInfo();
+	}
+
+	Cell->AppendStreamingSourceInfo(Source, SourceShape);
 	Cells.Add(Cell);
 }
 
diff --git a/Engine/Source/Runtime/Engine/Private/WorldPartition/WorldPartitionRuntimeSpatialHash.cpp b/Engine/Source/Runtime/Engine/Private/WorldPartition/WorldPartitionRuntimeSpatialHash.cpp
--- a/Engine/Source/Runtime/Engine/Private/WorldPartition/WorldPartitionRuntimeSpatialHash.cpp
+++ b/Engine/Source/Runtime/Engine/Private/WorldPartition/WorldPartitionRuntimeSpatialHash.cpp
@@ -216,8 +216,18 @@ void FSpatialHashStreamingGrid::GetCells(const FWorldPartitionStreamingQuerySour
 void FSpatialHashStreamingGrid::GetCells(const TArray<FWorldPartitionStreamingSource>& Sources, const UDataLayerSubsystem* DataLayerSubsystem, UWorldPartitionRuntimeHash::FStreamingSourceCells& OutActivateCells, UWorldPartitionRuntimeHash::FStreamingSourceCells& OutLoadCells, bool bEnableZCulling) const
 {
 	TRACE_CPUPROFILER_EVENT_SCOPE(FSpatialHashStreamingGrid::GetCells);
-	
-	typedef TMap<FGridCellCoord, TArray<UWorldPartitionRuntimeCell::FStreamingSourceInfo>> FIntersectingCells;
+
+	struct FStreamingSourceInfo
+	{
+		FStreamingSourceInfo(const FWorldPartitionStreamingSource& InSource, const FSphericalSector& InSourceShape)
+			: Source(InSource)
+			, SourceShape(InSourceShape)
+		{}
+		const FWorldPartitionStreamingSource& Source;
+		const FSphericalSector& SourceShape;
+	};
+
+	typedef TMap<FGridCellCoord, TArray<FStreamingSourceInfo>> FIntersectingCells;
 	FIntersectingCells AllActivatedCells;
 
 	const float GridLoadingRange = GetLoadingRange();
@@ -226,11 +236,12 @@ void FSpatialHashStreamingGrid::GetCells(const TArray<FWorldPartitionStreamingSo
 	{
 		Source.ForEachShape(GridLoadingRange, GridName, HLODLayer, /*bProjectIn2D*/ true, [&](const FSphericalSector& Shape)
 		{
-			UWorldPartitionRuntimeCell::FStreamingSourceInfo Info(Source, Shape);
+			FStreamingSourceInfo Info(Source, Shape);
 
 			Helper.ForEachIntersectingCells(Shape, [&](const FGridCellCoord& Coords)
 			{
 				bool bAddedActivatedCell = false;
+
 #if !UE_BUILD_SHIPPING
 				if ((GFilterRuntimeSpatialHashGridLevel == INDEX_NONE) || (GFilterRuntimeSpatialHashGridLevel == Coords.Z))
 #endif
@@ -245,18 +256,18 @@ void FSpatialHashStreamingGrid::GetCells(const TArray<FWorldPartitionStreamingSo
 								{
 									if (Source.TargetState == EStreamingSourceTargetState::Loaded)
 									{
-										OutLoadCells.AddCell(Cell, Info);
+										OutLoadCells.AddCell(Cell, Source, Shape);
 									}
 									else
 									{
 										check(Source.TargetState == EStreamingSourceTargetState::Activated);
-										OutActivateCells.AddCell(Cell, Info);
+										OutActivateCells.AddCell(Cell, Source, Shape);
 										bAddedActivatedCell = !GRuntimeSpatialHashUseAlignedGridLevels && GRuntimeSpatialHashSnapNonAlignedGridLevelsToLowerLevels;
 									}
 								}
 								else if (DataLayerSubsystem && DataLayerSubsystem->IsAnyDataLayerInEffectiveRuntimeState(Cell->GetDataLayers(), EDataLayerRuntimeState::Loaded))
 								{
-									OutLoadCells.AddCell(Cell, Info);
+									OutLoadCells.AddCell(Cell, Source, Shape);
 								}
 							}
 						}
@@ -333,7 +344,7 @@ void FSpatialHashStreamingGrid::GetCells(const TArray<FWorldPartitionStreamingSo
 					{
 						for (const auto& Info : ParentCell.Value)
 						{
-							OutActivateCells.AddCell(Cell, Info);
+							OutActivateCells.AddCell(Cell, Info.Source, Info.SourceShape);
 						}
 					}
 				}
@@ -1049,38 +1060,14 @@ bool UWorldPartitionRuntimeSpatialHash::CreateStreamingGrid(const FSpatialHashRu
 				{
 					for (const FGuid& ActorGuid : ActorSetInstance->ActorSet->Actors)
 					{
-						// In PIE, Always loaded cell is not generated. Instead, always loaded actors will be added to AlwaysLoadedActorsForPIE.
-						// This will trigger loading/registration of these actors in the PersistentLevel (if not already loaded).
-						// Then, duplication of world for PIE will duplicate only these actors. 
-						// When stopping PIE, WorldPartition will release these FWorldPartitionReferences which 
-						// will unload actors that were not already loaded in the non PIE world.
-						if (bIsMainWorldPartition && bIsCellAlwaysLoaded && ActorSetInstance->ContainerID.IsMainContainer() && !IsRunningCookCommandlet())
-						{
-							// This will load the actor if it isn't already loaded
-							FWorldPartitionReference Reference(WorldPartition, ActorGuid);				
-
-							IStreamingGenerationContext::FActorInstance ActorInstance(ActorGuid, ActorSetInstance);
-							const FWorldPartitionActorDescView& ActorDescView = ActorInstance.GetActorDescView();
+						IStreamingGenerationContext::FActorInstance ActorInstance(ActorGuid, ActorSetInstance);
+						const FWorldPartitionActorDescView& ActorDescView = ActorInstance.GetActorDescView();
 
-							if (AActor* AlwaysLoadedActor = FindObject<AActor>(nullptr, *ActorDescView.GetActorSoftPath().ToString()))
-							{
-								AlwaysLoadedActorsForPIE.Emplace(Reference, AlwaysLoadedActor);
-
-								// Handle child actors
-								AlwaysLoadedActor->ForEachComponent<UChildActorComponent>(true, [this, &Reference](UChildActorComponent* ChildActorComponent)
-								{
-									if (AActor* ChildActor = ChildActorComponent->GetChildActor())
-									{
-										AlwaysLoadedActorsForPIE.Emplace(Reference, ChildActor);
-									}
-								});
-							}
-						}
-						else
+						if (!ConditionalRegisterAlwaysLoadedActorsForPIE(ActorDescView, bIsMainWorldPartition, ActorSetInstance->ContainerID.IsMainContainer(), bIsCellAlwaysLoaded))
 						{
-							FilteredActors.Emplace(ActorGuid, ActorSetInstance);
+							FilteredActors.Emplace(ActorInstance);
 						}
-					}					
+					}
 				}
 
 				if (!FilteredActors.Num())
@@ -1207,15 +1194,14 @@ bool UWorldPartitionRuntimeSpatialHash::PopulateGeneratedPackageForCook(UPackage
 TArray<UWorldPartitionRuntimeCell*> UWorldPartitionRuntimeSpatialHash::GetAlwaysLoadedCells() const
 {
 	TArray<UWorldPartitionRuntimeCell*> Result;
-	TSet<const UWorldPartitionRuntimeCell*> StreamingCells;
-	GetAllStreamingCells(StreamingCells);
-	for (const UWorldPartitionRuntimeCell* Cell : StreamingCells)
+	ForEachStreamingCells([&Result](const UWorldPartitionRuntimeCell* Cell)
 	{
 		if (Cell->IsAlwaysLoaded())
 		{
 			Result.Add(const_cast<UWorldPartitionRuntimeCell*>(Cell));
 		}
-	}
+		return true;
+	});
 	return Result;
 }
 
@@ -1284,30 +1270,12 @@ FAutoConsoleCommand UWorldPartitionRuntimeSpatialHash::OverrideLoadingRangeComma
 	})
 );
 
-bool UWorldPartitionRuntimeSpatialHash::ShouldConsiderClientOnlyVisibleCells() const
-{
-	const UWorld* World = GetWorld();
-	if (World->IsGameWorld())
-	{
-		// Dedicated server & Listen server without server streaming won't consider client-only visible cells
-		const ENetMode NetMode = World->GetNetMode();
-		if ((NetMode == NM_DedicatedServer) ||
-			((NetMode == NM_ListenServer) && !GetOuterUWorldPartition()->IsServerStreamingEnabled()))
-		{
-			return false;
-		}
-	}
-	return true;
-}
-
 // Streaming interface
-int32 UWorldPartitionRuntimeSpatialHash::GetAllStreamingCells(TSet<const UWorldPartitionRuntimeCell*>& Cells, bool bAllDataLayers, bool bDataLayersOnly, const TSet<FName>& InDataLayers) const
+void UWorldPartitionRuntimeSpatialHash::ForEachStreamingCells(TFunctionRef<bool(const UWorldPartitionRuntimeCell*)> Func) const
 {
-	const bool bShouldConsiderClientOnlyVisible = ShouldConsiderClientOnlyVisibleCells();
-
-	ForEachStreamingGrid([&](const FSpatialHashStreamingGrid& StreamingGrid)
+	for (const FSpatialHashStreamingGrid& StreamingGrid : StreamingGrids)
 	{
-		if (!StreamingGrid.bClientOnlyVisible || bShouldConsiderClientOnlyVisible)
+		if (IsCellRelevantFor(StreamingGrid.bClientOnlyVisible))
 		{
 			for (const FSpatialHashStreamingGridLevel& GridLevel : StreamingGrid.GridLevels)
 			{
@@ -1315,69 +1283,82 @@ int32 UWorldPartitionRuntimeSpatialHash::GetAllStreamingCells(TSet<const UWorldP
 				{
 					for (const UWorldPartitionRuntimeSpatialHashCell* Cell : LayerCell.GridCells)
 					{
-						if (!bDataLayersOnly && !Cell->HasDataLayers())
-						{
-							Cells.Add(Cell);
-						}
-						else if (Cell->HasDataLayers() && (bAllDataLayers || Cell->HasAnyDataLayer(InDataLayers)))
+						if (!Func(Cell))
 						{
-							Cells.Add(Cell);
+							break;
 						}
 					}
 				}
 			}
 		}
-	});
-
-	return Cells.Num();
+	}
 }
 
-bool UWorldPartitionRuntimeSpatialHash::GetStreamingCells(const FWorldPartitionStreamingQuerySource& QuerySource, TSet<const UWorldPartitionRuntimeCell*>& OutCells) const 
+void UWorldPartitionRuntimeSpatialHash::ForEachStreamingCellsQuery(const FWorldPartitionStreamingQuerySource& QuerySource, TFunctionRef<bool(const UWorldPartitionRuntimeCell*)> Func) const
 {
-	const bool bShouldConsiderClientOnlyVisible = ShouldConsiderClientOnlyVisibleCells();
-
 	for (const FSpatialHashStreamingGrid& StreamingGrid : StreamingGrids)
 	{
-		if (!StreamingGrid.bClientOnlyVisible || bShouldConsiderClientOnlyVisible)
+		if (IsCellRelevantFor(StreamingGrid.bClientOnlyVisible))
 		{
-			StreamingGrid.GetCells(QuerySource, OutCells, GetEffectiveEnableZCulling(bEnableZCulling));
+			TSet<const UWorldPartitionRuntimeCell*> Cells;
+			StreamingGrid.GetCells(QuerySource, Cells, GetEffectiveEnableZCulling(bEnableZCulling));
+
+			for (const UWorldPartitionRuntimeCell* Cell : Cells)
+			{
+				if (!Func(Cell))
+				{
+					break;
+				}
+			}
 		}
 	}
-
-	return !!OutCells.Num();
 }
 
-bool UWorldPartitionRuntimeSpatialHash::GetStreamingCells(const TArray<FWorldPartitionStreamingSource>& Sources, UWorldPartitionRuntimeHash::FStreamingSourceCells& OutActivateCells, UWorldPartitionRuntimeHash::FStreamingSourceCells& OutLoadCells) const
+void UWorldPartitionRuntimeSpatialHash::ForEachStreamingCellsSources(const TArray<FWorldPartitionStreamingSource>& Sources, TFunctionRef<bool(const UWorldPartitionRuntimeCell*, EStreamingSourceTargetState)> Func) const
 {
-	TRACE_CPUPROFILER_EVENT_SCOPE(UWorldPartitionRuntimeSpatialHash::GetStreamingCells);
-
 	const UDataLayerSubsystem* DataLayerSubsystem = GetWorld()->GetSubsystem<UDataLayerSubsystem>();
-	const bool bShouldConsiderClientOnlyVisible = ShouldConsiderClientOnlyVisibleCells();
+
+	FStreamingSourceCells ActivateStreamingSourceCells;
+	FStreamingSourceCells LoadStreamingSourceCells;
 
 	if (Sources.Num() == 0)
 	{
 		// Get always loaded cells
-		ForEachStreamingGrid([&](const FSpatialHashStreamingGrid& StreamingGrid)
+		for (const FSpatialHashStreamingGrid& StreamingGrid : StreamingGrids)
 		{
-			if (!StreamingGrid.bClientOnlyVisible || bShouldConsiderClientOnlyVisible)
+			if (IsCellRelevantFor(StreamingGrid.bClientOnlyVisible))
 			{
-				StreamingGrid.GetAlwaysLoadedCells(DataLayerSubsystem, OutActivateCells.GetCells(), OutLoadCells.GetCells());
+				StreamingGrid.GetAlwaysLoadedCells(DataLayerSubsystem, ActivateStreamingSourceCells.GetCells(), LoadStreamingSourceCells.GetCells());
 			}
-		});
+		}
 	}
 	else
 	{
 		// Get cells based on streaming sources
-		ForEachStreamingGrid([&](const FSpatialHashStreamingGrid& StreamingGrid)
+		for (const FSpatialHashStreamingGrid& StreamingGrid : StreamingGrids)
 		{
-			if (!StreamingGrid.bClientOnlyVisible || bShouldConsiderClientOnlyVisible)
+			if (IsCellRelevantFor(StreamingGrid.bClientOnlyVisible))
 			{
-				StreamingGrid.GetCells(Sources, DataLayerSubsystem, OutActivateCells, OutLoadCells, GetEffectiveEnableZCulling(bEnableZCulling));
+				StreamingGrid.GetCells(Sources, DataLayerSubsystem, ActivateStreamingSourceCells, LoadStreamingSourceCells, GetEffectiveEnableZCulling(bEnableZCulling));
 			}
-		});
+		}
 	}
 
-	return !!(OutActivateCells.Num() + OutLoadCells.Num());
+	for (const UWorldPartitionRuntimeCell* Cell : ActivateStreamingSourceCells.GetCells())
+	{
+		if (!Func(Cell, EStreamingSourceTargetState::Activated))
+		{
+			return;
+		}
+	}
+
+	for (const UWorldPartitionRuntimeCell* Cell : LoadStreamingSourceCells.GetCells())
+	{
+		if (!Func(Cell, EStreamingSourceTargetState::Loaded))
+		{
+			return;
+		}
+	}
 }
 
 bool UWorldPartitionRuntimeSpatialHash::InjectExternalStreamingObject(URuntimeHashExternalStreamingObjectBase* ExternalStreamingObject)
@@ -1653,13 +1634,12 @@ bool UWorldPartitionRuntimeSpatialHash::ContainsRuntimeHash(const FString& Name)
 
 TArray<const FSpatialHashStreamingGrid*> UWorldPartitionRuntimeSpatialHash::GetFilteredStreamingGrids() const
 {
-	const bool bShouldConsiderClientOnlyVisible = ShouldConsiderClientOnlyVisibleCells();
-
 	TArray<const FSpatialHashStreamingGrid*> FilteredStreamingGrids;
 	FilteredStreamingGrids.Reserve(GetNumGrids());
-	ForEachStreamingGrid([&FilteredStreamingGrids, bShouldConsiderClientOnlyVisible](const FSpatialHashStreamingGrid& StreamingGrid)
+
+	ForEachStreamingGrid([this, &FilteredStreamingGrids](const FSpatialHashStreamingGrid& StreamingGrid)
 	{
-		if (!StreamingGrid.bClientOnlyVisible || bShouldConsiderClientOnlyVisible)
+		if (IsCellRelevantFor(StreamingGrid.bClientOnlyVisible))
 		{
 			if (FWorldPartitionDebugHelper::IsDebugRuntimeHashGridShown(StreamingGrid.GridName))
 			{
diff --git a/Engine/Source/Runtime/Engine/Private/WorldPartition/WorldPartitionRuntimeSpatialHashCell.cpp b/Engine/Source/Runtime/Engine/Private/WorldPartition/WorldPartitionRuntimeSpatialHashCell.cpp
--- a/Engine/Source/Runtime/Engine/Private/WorldPartition/WorldPartitionRuntimeSpatialHashCell.cpp
+++ b/Engine/Source/Runtime/Engine/Private/WorldPartition/WorldPartitionRuntimeSpatialHashCell.cpp
@@ -45,26 +45,29 @@ void UWorldPartitionRuntimeSpatialHashCell::PostDuplicate(bool bDuplicateForPIE)
 }
 #endif
 
-bool UWorldPartitionRuntimeSpatialHashCell::CacheStreamingSourceInfo(const UWorldPartitionRuntimeCell::FStreamingSourceInfo& Info) const
+void UWorldPartitionRuntimeSpatialHashCell::ResetStreamingSourceInfo() const
 {
-	const bool bWasCacheDirtied = Super::CacheStreamingSourceInfo(Info);
-	if (bWasCacheDirtied)
-	{
-		CachedIsBlockingSource = false;
-		CachedMinSquareDistanceToBlockingSource = MAX_dbl;
-		CachedMinSquareDistanceToSource = MAX_dbl;
-		CachedSourceModulatedDistances.Reset();
-	}
+	Super::ResetStreamingSourceInfo();
+
+	CachedIsBlockingSource = false;
+	CachedMinSquareDistanceToBlockingSource = MAX_dbl;
+	CachedMinSquareDistanceToSource = MAX_dbl;
+	CachedSourceModulatedDistances.Reset();
+}
+
+void UWorldPartitionRuntimeSpatialHashCell::AppendStreamingSourceInfo(const FWorldPartitionStreamingSource& Source, const FSphericalSector& SourceShape) const
+{
+	Super::AppendStreamingSourceInfo(Source, SourceShape);
 
 	float AngleContribution = FMath::Clamp(GRuntimeSpatialHashCellToSourceAngleContributionToCellImportance, 0.f, 1.f);
-	const double SquareDistance = FVector::DistSquared2D(Info.SourceShape.GetCenter(), Position);
+	const double SquareDistance = FVector::DistSquared2D(SourceShape.GetCenter(), Position);
 	float AngleFactor = 1.f;
 	if (!FMath::IsNearlyZero(AngleContribution))
 	{
 		const FBox Box(FVector(Position.X - Extent, Position.Y - Extent, 0.f), FVector(Position.X + Extent, Position.Y + Extent, 0.f));
-		const FVector2D SourcePos(Info.SourceShape.GetCenter());
+		const FVector2D SourcePos(SourceShape.GetCenter());
 		const FVector StartVert(SourcePos, 0.f);
-		const FVector EndVert(SourcePos + FVector2D(Info.SourceShape.GetScaledAxis()), 0.f);
+		const FVector EndVert(SourcePos + FVector2D(SourceShape.GetScaledAxis()), 0.f);
 
 		float Angle = 0.f;
 		if (!FMath::LineBoxIntersection(Box, StartVert, EndVert, EndVert - StartVert))
@@ -72,7 +75,7 @@ bool UWorldPartitionRuntimeSpatialHashCell::CacheStreamingSourceInfo(const UWorl
 			// Find smallest angle using 4 corners and center of cell bounds
 			const FVector2D Position2D(Position);
 			float MaxDot = 0.f;
-			FVector2D SourceForward(Info.SourceShape.GetAxis());
+			FVector2D SourceForward(SourceShape.GetAxis());
 			SourceForward.Normalize();
 			FVector2D CellPoints[5];
 			CellPoints[0] = Position2D + FVector2D(-Extent, -Extent);
@@ -91,17 +94,30 @@ bool UWorldPartitionRuntimeSpatialHashCell::CacheStreamingSourceInfo(const UWorl
 		const float NormalizedAngle = FMath::Clamp(Angle, UE_PI/180.f, 1.f);
 		AngleFactor = FMath::Pow(NormalizedAngle, AngleContribution);
 	}
+
+	// Only consider blocking sources
+	if (Source.bBlockOnSlowLoading)
+	{
+		CachedIsBlockingSource = true;
+		CachedMinSquareDistanceToBlockingSource = FMath::Min(SquareDistance, CachedMinSquareDistanceToBlockingSource);
+	}
+
+	CachedMinSquareDistanceToSource = FMath::Min(SquareDistance, CachedMinSquareDistanceToSource);
+
 	// Modulate distance to cell by angle relative to source forward vector (to prioritize cells in front)
 	const double SquareAngleFactor = AngleFactor * AngleFactor;
 	const double ModulatedSquareDistance = SquareDistance * SquareAngleFactor;
 	CachedSourceModulatedDistances.Add(ModulatedSquareDistance);
+}
+
+void UWorldPartitionRuntimeSpatialHashCell::MergeStreamingSourceInfo() const
+{
+	Super::MergeStreamingSourceInfo();
+
 	int32 Count = CachedSourceModulatedDistances.Num();
 	check(Count == CachedSourcePriorityWeights.Num());
-	if (Count == 1)
-	{
-		CachedSourceSortingDistance = ModulatedSquareDistance;
-	}
-	else
+
+	if (Count)
 	{
 		double TotalSourcePriorityWeight = 0.f;
 		for (int32 i = 0; i < Count; ++i)
@@ -127,16 +143,6 @@ bool UWorldPartitionRuntimeSpatialHashCell::CacheStreamingSourceInfo(const UWorl
 		// - the weighted modulated distance
 		CachedSourceSortingDistance = FMath::Min(CachedSourceModulatedDistances[HighestPrioMinDistIndex], WeightedModulatedDistance);
 	}
-	
-	// Only consider blocking sources
-	if (Info.Source.bBlockOnSlowLoading)
-	{
-		CachedIsBlockingSource = true;
-		CachedMinSquareDistanceToBlockingSource = FMath::Min(SquareDistance, CachedMinSquareDistanceToBlockingSource);
-	}
-
-	CachedMinSquareDistanceToSource = FMath::Min(SquareDistance, CachedMinSquareDistanceToSource);
-	return bWasCacheDirtied;
 }
 
 int32 UWorldPartitionRuntimeSpatialHashCell::SortCompare(const UWorldPartitionRuntimeCell* InOther) const
diff --git a/Engine/Source/Runtime/Engine/Private/WorldPartition/WorldPartitionStreamingPolicy.cpp b/Engine/Source/Runtime/Engine/Private/WorldPartition/WorldPartitionStreamingPolicy.cpp
--- a/Engine/Source/Runtime/Engine/Private/WorldPartition/WorldPartitionStreamingPolicy.cpp
+++ b/Engine/Source/Runtime/Engine/Private/WorldPartition/WorldPartitionStreamingPolicy.cpp
@@ -58,16 +58,23 @@ static FAutoConsoleVariableRef CVarBlockOnSlowStreaming(
 	GBlockOnSlowStreaming,
 	TEXT("Set if streaming needs to block when to slow to catchup."));
 
+static void SortStreamingCellsByImportance(const TSet<const UWorldPartitionRuntimeCell*>& InCells, TArray<const UWorldPartitionRuntimeCell*>& OutSortedCells)
+{
+	TRACE_CPUPROFILER_EVENT_SCOPE(SortStreamingCellsByImportance);
+	OutSortedCells = InCells.Array();
+	Algo::Sort(OutSortedCells, [](const UWorldPartitionRuntimeCell* CellA, const UWorldPartitionRuntimeCell* CellB) { return CellA->SortCompare(CellB) < 0; });
+}
+
 UWorldPartitionStreamingPolicy::UWorldPartitionStreamingPolicy(const FObjectInitializer& ObjectInitializer)
-: Super(ObjectInitializer) 
-, bCriticalPerformanceRequestedBlockTillOnWorld(false)
-, CriticalPerformanceBlockTillLevelStreamingCompletedEpoch(0)
-, DataLayersStatesServerEpoch(INT_MIN)
-, ServerStreamingEnabledEpoch(INT_MIN)
-, StreamingPerformance(EWorldPartitionStreamingPerformance::Good)
+	: Super(ObjectInitializer) 
+	, bCriticalPerformanceRequestedBlockTillOnWorld(false)
+	, CriticalPerformanceBlockTillLevelStreamingCompletedEpoch(0)
+	, DataLayersStatesServerEpoch(INT_MIN)
+	, ServerStreamingEnabledEpoch(INT_MIN)
+	, StreamingPerformance(EWorldPartitionStreamingPerformance::Good)
 #if !UE_BUILD_SHIPPING
-, OnScreenMessageStartTime(0.0)
-, OnScreenMessageStreamingPerformance(EWorldPartitionStreamingPerformance::Good)
+	, OnScreenMessageStartTime(0.0)
+	, OnScreenMessageStreamingPerformance(EWorldPartitionStreamingPerformance::Good)
 #endif
 {
 	if (!IsTemplate())
@@ -230,11 +237,8 @@ void UWorldPartitionStreamingPolicy::UpdateStreamingState()
 	TSet<FName> ClientVisibleLevelNames;
 	bool bUpdateEpoch = false;
 
-	TSet<const UWorldPartitionRuntimeCell*>& ActivateStreamingCells = FrameActivateCells.GetCells();
-	TSet<const UWorldPartitionRuntimeCell*>& LoadStreamingCells = FrameLoadCells.GetCells();
-		
-	check(ActivateStreamingCells.IsEmpty());
-	check(LoadStreamingCells.IsEmpty());
+	check(FrameActivateCells.IsEmpty());
+	check(FrameLoadCells.IsEmpty());
 
 	ON_SCOPE_EXIT
 	{
@@ -250,16 +254,29 @@ void UWorldPartitionStreamingPolicy::UpdateStreamingState()
 	if (!bIsServer || bIsServerStreamingEnabled || AWorldPartitionReplay::IsPlaybackEnabled(World))
 	{
 		// When world partition can't stream, all cells must be unloaded
-		if (WorldPartition->CanStream())
+		if (WorldPartition->CanStream() && WorldPartition->RuntimeHash)
 		{
 			UWorldPartitionRuntimeCell::DirtyStreamingSourceCacheEpoch();
-			if (WorldPartition->RuntimeHash)
+
+			WorldPartition->RuntimeHash->ForEachStreamingCellsSources(StreamingSources, [this](const UWorldPartitionRuntimeCell* Cell, EStreamingSourceTargetState TargetState)
 			{
-				WorldPartition->RuntimeHash->GetStreamingCells(StreamingSources, FrameActivateCells, FrameLoadCells);
-			}
+				switch (TargetState)
+				{
+				case EStreamingSourceTargetState::Loaded:
+					FrameLoadCells.Add(Cell);
+					break;
+				case EStreamingSourceTargetState::Activated:
+					FrameActivateCells.Add(Cell);
+					break;
+				default:
+					check(0);
+				}
+
+				return true;
+			});
 
 			// Activation superseeds Loading
-			LoadStreamingCells = LoadStreamingCells.Difference(ActivateStreamingCells);
+			FrameLoadCells = FrameLoadCells.Difference(FrameActivateCells);
 		}
 	}
 
@@ -287,95 +304,170 @@ void UWorldPartitionStreamingPolicy::UpdateStreamingState()
 		}
 
 		if (!bIsServerStreamingEnabled)
-		{			
+		{
 			if (WorldPartition->RuntimeHash)
 			{
 				const UDataLayerSubsystem* DataLayerSubsystem = WorldPartition->GetWorld()->GetSubsystem<UDataLayerSubsystem>();
+				const TSet<FName> EffectiveActiveDataLayerNames = DataLayerSubsystem->GetEffectiveActiveDataLayerNames();
+				const TSet<FName> EffectiveLoadedDataLayerNames = DataLayerSubsystem->GetEffectiveLoadedDataLayerNames();
 
-				// Non Data Layer Cells + Active Data Layers
-				WorldPartition->RuntimeHash->GetAllStreamingCells(ActivateStreamingCells, /*bAllDataLayers=*/ false, /*bDataLayersOnly=*/ false, DataLayerSubsystem->GetEffectiveActiveDataLayerNames());
-
-				// Loaded Data Layers Cells only
-				if (DataLayerSubsystem->GetEffectiveLoadedDataLayerNames().Num())
+				WorldPartition->RuntimeHash->ForEachStreamingCells([this, &EffectiveActiveDataLayerNames, &EffectiveLoadedDataLayerNames](const UWorldPartitionRuntimeCell* Cell)
 				{
-					WorldPartition->RuntimeHash->GetAllStreamingCells(LoadStreamingCells, /*bAllDataLayers=*/ false, /*bDataLayersOnly=*/ true, DataLayerSubsystem->GetEffectiveLoadedDataLayerNames());
-				}
+					// Non Data Layer Cells + Active Data Layers
+					if (!Cell->HasDataLayers() || Cell->HasAnyDataLayer(EffectiveActiveDataLayerNames))
+					{
+						FrameActivateCells.Add(Cell);
+					}
+
+					// Loaded Data Layers Cells only
+					if (EffectiveLoadedDataLayerNames.Num())
+					{
+						if (Cell->HasDataLayers() && Cell->HasAnyDataLayer(EffectiveLoadedDataLayerNames))
+						{
+							FrameLoadCells.Add(Cell);
+						}
+					}
+
+					return true;
+				});
 			}
 		}
 	}
 
-	if (!UHLODSubsystem::IsHLODEnabled())
+	// Process cells to activate
+	auto ProcessCellsToActivate = [this](TSet<const UWorldPartitionRuntimeCell*>& Cells)
+	{
+		for (TSet<const UWorldPartitionRuntimeCell*>::TIterator It = Cells.CreateIterator(); It; ++It)
+		{
+			const UWorldPartitionRuntimeCell* Cell = *It;
+
+			if (ShouldSkipCellForPerformance(Cell))
+			{
+				It.RemoveCurrent();
+				continue;
+			}
+			
+			if (Cell->GetIsHLOD() && !UHLODSubsystem::IsHLODEnabled())
+			{
+				It.RemoveCurrent();
+				continue;
+			}
+
+			Cell->MergeStreamingSourceInfo();
+		}
+	};
+
+	// Process cells to load
+	auto ProcessCellsToLoad = [this, bIsServer, &ClientVisibleLevelNames, &bUpdateEpoch](TSet<const UWorldPartitionRuntimeCell*>& Cells)
 	{
-		// Remove all HLOD cells from the Activate & Load cells
-		auto RemoveHLODCells = [](TSet<const UWorldPartitionRuntimeCell*>& Cells)
+		for (TSet<const UWorldPartitionRuntimeCell*>::TIterator It = Cells.CreateIterator(); It; ++It)
 		{
-			for (auto It = Cells.CreateIterator(); It; ++It)
+			const UWorldPartitionRuntimeCell* Cell = *It;
+
+			if (bIsServer)
 			{
-				if ((*It)->GetIsHLOD())
+				if (ULevel* Level = Cell->GetLevel())
 				{
-					It.RemoveCurrent();
+					if (ClientVisibleLevelNames.Contains(Cell->GetLevel()->GetPackage()->GetFName()))
+					{
+						UE_CLOG(bUpdateEpoch, LogWorldPartition, Verbose, TEXT("Server epoch update delayed by client visibility"));
+						bUpdateEpoch = false;
+						It.RemoveCurrent();
+						continue;
+					}
 				}
 			}
-		};
-		RemoveHLODCells(ActivateStreamingCells);
-		RemoveHLODCells(LoadStreamingCells);
-	}
 
-	// Determine cells to load/unload
-	TSet<const UWorldPartitionRuntimeCell*> ToActivateCells = ActivateStreamingCells.Difference(ActivatedCells);
-	TSet<const UWorldPartitionRuntimeCell*> ToLoadCells = LoadStreamingCells.Difference(LoadedCells);
-	TSet<const UWorldPartitionRuntimeCell*> ToUnloadCells = ActivatedCells.Union(LoadedCells).Difference(ActivateStreamingCells.Union(LoadStreamingCells));
+			if (ShouldSkipCellForPerformance(Cell))
+			{
+				It.RemoveCurrent();
+				continue;
+			}
+			
+			if (Cell->GetIsHLOD() && !UHLODSubsystem::IsHLODEnabled())
+			{
+				It.RemoveCurrent();
+				continue;
+			}
 
-#if !UE_BUILD_SHIPPING
-	UpdateDebugCellsStreamingPriority(ActivateStreamingCells, LoadStreamingCells);
-#endif
+			Cell->MergeStreamingSourceInfo();
+		}
+	};
 
-	if(World->bMatchStarted)
-	{
-		WORLDPARTITION_LOG_UPDATESTREAMINGSTATE(Verbose);
-	}
-	else
-	{
-		WORLDPARTITION_LOG_UPDATESTREAMINGSTATE(Log);
-	}
-		
-	// Remove cells state changes based on client visibility. If any cell gets removed, delay epoch change.
-	auto ProcessClientVisibility = [bIsServer, &ClientVisibleLevelNames, &bUpdateEpoch](TSet<const UWorldPartitionRuntimeCell*>& Cells)
+	// Process cells to unload
+	auto ProcessCellsToUnload = [this, bIsServer, &ClientVisibleLevelNames, &bUpdateEpoch](TSet<const UWorldPartitionRuntimeCell*>& Cells)
 	{
 		if (bIsServer)
-		{
+		{	
 			for (TSet<const UWorldPartitionRuntimeCell*>::TIterator It = Cells.CreateIterator(); It; ++It)
 			{
-				if (ULevel* Level = (*It)->GetLevel())
+				const UWorldPartitionRuntimeCell* Cell = *It;
+
+				if (ULevel* Level = Cell->GetLevel())
 				{
-					if (ClientVisibleLevelNames.Contains((*It)->GetLevel()->GetPackage()->GetFName()))
+					if (ClientVisibleLevelNames.Contains(Cell->GetLevel()->GetPackage()->GetFName()))
 					{
-						It.RemoveCurrent();
-
 						UE_CLOG(bUpdateEpoch, LogWorldPartition, Verbose, TEXT("Server epoch update delayed by client visibility"));
 						bUpdateEpoch = false;
+						It.RemoveCurrent();
+						continue;
 					}
 				}
 			}
 		}
 	};
 
+	// Sort cells by importance
+	auto SortStreamingCells = [this](const TArray<FWorldPartitionStreamingSource>& Sources, const TSet<const UWorldPartitionRuntimeCell*>& InCells, TArray<const UWorldPartitionRuntimeCell*>& OutCells)
+	{
+		OutCells.Empty(InCells.Num());
+		SortStreamingCellsByImportance(InCells, OutCells);
+	};
+
+	// Determine cells to load/unload and sort by importance
+	TArray<const UWorldPartitionRuntimeCell*> ToActivateCells;
+	TArray<const UWorldPartitionRuntimeCell*> ToLoadCells;
+	TArray<const UWorldPartitionRuntimeCell*> ToUnloadCells;
+	{
+		TSet<const UWorldPartitionRuntimeCell*> ToActivateCellsUnsorted = FrameActivateCells.Difference(ActivatedCells);
+		TSet<const UWorldPartitionRuntimeCell*> ToLoadCellsUnsorted = FrameLoadCells.Difference(LoadedCells);
+		TSet<const UWorldPartitionRuntimeCell*> ToUnloadCellsUnsorted = ActivatedCells.Union(LoadedCells).Difference(FrameActivateCells.Union(FrameLoadCells));
+
+		ProcessCellsToActivate(ToActivateCellsUnsorted);
+		ProcessCellsToLoad(ToLoadCellsUnsorted);
+		ProcessCellsToUnload(ToUnloadCellsUnsorted);
+		
+		SortStreamingCells(StreamingSources, ToActivateCellsUnsorted, ToActivateCells);
+		SortStreamingCells(StreamingSources, ToLoadCellsUnsorted, ToLoadCells);
+	}
+
+	if(World->bMatchStarted)
+	{
+		WORLDPARTITION_LOG_UPDATESTREAMINGSTATE(Verbose);
+	}
+	else
+	{
+		WORLDPARTITION_LOG_UPDATESTREAMINGSTATE(Log);
+	}
+
+#if !UE_BUILD_SHIPPING
+	UpdateDebugCellsStreamingPriority(FrameActivateCells, FrameLoadCells);
+#endif
+
 	if (ToUnloadCells.Num() > 0)
 	{
-		ProcessClientVisibility(ToUnloadCells);
-		SetTargetStateForCells(EWorldPartitionRuntimeCellState::Unloaded, ToUnloadCells);
+		SetCellsStateToUnloaded(ToUnloadCells);
 	}
 
 	// Do Activation State first as it is higher prio than Load State (if we have a limited number of loading cells per frame)
 	if (ToActivateCells.Num() > 0)
 	{
-		SetTargetStateForCells(EWorldPartitionRuntimeCellState::Activated, ToActivateCells);
+		SetCellsStateToActivated(ToActivateCells);
 	}
 
 	if (ToLoadCells.Num() > 0)
 	{
-		ProcessClientVisibility(ToLoadCells);
-		SetTargetStateForCells(EWorldPartitionRuntimeCellState::Loaded, ToLoadCells);
+		SetCellsStateToLoaded(ToLoadCells);
 	}
 
 	// Sort cells and update streaming priority 
@@ -389,10 +481,8 @@ void UWorldPartitionStreamingPolicy::UpdateStreamingState()
 				AddToWorldCells.Add(ActivatedCell);
 			}
 		}
-		if (WorldPartition->RuntimeHash)
-		{
-			WorldPartition->RuntimeHash->SortStreamingCellsByImportance(AddToWorldCells, StreamingSources, SortedAddToWorldCells);
-		}
+
+		SortStreamingCells(StreamingSources, AddToWorldCells, SortedAddToWorldCells);
 
 		// Update level streaming priority so that UWorld::UpdateLevelStreaming will naturally process the levels in the correct order
 		const int32 MaxPrio = SortedAddToWorldCells.Num();
@@ -408,7 +498,7 @@ void UWorldPartitionStreamingPolicy::UpdateStreamingState()
 	}
 
 	// Evaluate streaming performance based on cells that should be activated
-	UpdateStreamingPerformance(ActivateStreamingCells);
+	UpdateStreamingPerformance(FrameActivateCells);
 	
 	// Update Epoch if we aren't waiting for clients anymore
 	if (bUpdateEpoch)
@@ -428,16 +518,14 @@ void UWorldPartitionStreamingPolicy::UpdateDebugCellsStreamingPriority(const TSe
 		CellsToSort.Append(ActivateStreamingCells);
 		CellsToSort.Append(LoadStreamingCells);
 
-		if (WorldPartition->RuntimeHash)
+		TArray<const UWorldPartitionRuntimeCell*> SortedCells;
+		SortStreamingCellsByImportance(CellsToSort, SortedCells);
+
+		const int32 CellCount = SortedCells.Num();
+		int32 CellPrio = 0;
+		for (const UWorldPartitionRuntimeCell* SortedCell : SortedCells)
 		{
-			TArray<const UWorldPartitionRuntimeCell*, TInlineAllocator<256>> SortedCells;
-			WorldPartition->RuntimeHash->SortStreamingCellsByImportance(CellsToSort, StreamingSources, SortedCells);
-			const int32 CellCount = SortedCells.Num();
-			int32 CellPrio = 0;
-			for (const UWorldPartitionRuntimeCell* SortedCell : SortedCells)
-			{
-				const_cast<UWorldPartitionRuntimeCell*>(SortedCell)->SetDebugStreamingPriority(float(CellPrio++) / CellCount);
-			}
+			const_cast<UWorldPartitionRuntimeCell*>(SortedCell)->SetDebugStreamingPriority(float(CellPrio++) / CellCount);
 		}
 	}
 }
@@ -545,46 +633,17 @@ int32 UWorldPartitionStreamingPolicy::GetMaxCellsToLoad() const
 	return MAX_int32;
 }
 
-void UWorldPartitionStreamingPolicy::SetTargetStateForCells(EWorldPartitionRuntimeCellState TargetState, const TSet<const UWorldPartitionRuntimeCell*>& Cells)
-{
-	switch (TargetState)
-	{
-	case EWorldPartitionRuntimeCellState::Unloaded:
-		SetCellsStateToUnloaded(Cells);
-		break;
-	case EWorldPartitionRuntimeCellState::Loaded:
-		SetCellsStateToLoaded(Cells);
-		break;
-	case EWorldPartitionRuntimeCellState::Activated:
-		SetCellsStateToActivated(Cells);
-		break;
-	default:
-		check(false);
-	}
-}
-
-void UWorldPartitionStreamingPolicy::SetCellsStateToLoaded(const TSet<const UWorldPartitionRuntimeCell*>& ToLoadCells)
+void UWorldPartitionStreamingPolicy::SetCellsStateToLoaded(const TArray<const UWorldPartitionRuntimeCell*>& ToLoadCells)
 {
 	TRACE_CPUPROFILER_EVENT_SCOPE(UWorldPartitionStreamingPolicy::SetCellsStateToLoaded);
 
 	int32 MaxCellsToLoad = GetMaxCellsToLoad();
-	// Sort cells based on importance
-	TArray<const UWorldPartitionRuntimeCell*, TInlineAllocator<256>> SortedCells;
-
-	if (WorldPartition->RuntimeHash)
-	{
-		WorldPartition->RuntimeHash->SortStreamingCellsByImportance(ToLoadCells, StreamingSources, SortedCells);
-	}
 
 	// Trigger cell loading. Depending on actual state of cell limit loading.
-	for (const UWorldPartitionRuntimeCell* Cell : SortedCells)
+	for (const UWorldPartitionRuntimeCell* Cell : ToLoadCells)
 	{
-		if (ShouldSkipCellForPerformance(Cell))
-		{
-			continue;
-		}
-
 		UE_LOG(LogWorldPartition, Verbose, TEXT("UWorldPartitionStreamingPolicy::LoadCells %s"), *Cell->GetName());
+
 		if (ActivatedCells.Contains(Cell))
 		{
 			Cell->Deactivate();
@@ -603,28 +662,17 @@ void UWorldPartitionStreamingPolicy::SetCellsStateToLoaded(const TSet<const UWor
 	}
 }
 
-void UWorldPartitionStreamingPolicy::SetCellsStateToActivated(const TSet<const UWorldPartitionRuntimeCell*>& ToActivateCells)
+void UWorldPartitionStreamingPolicy::SetCellsStateToActivated(const TArray<const UWorldPartitionRuntimeCell*>& ToActivateCells)
 {
 	TRACE_CPUPROFILER_EVENT_SCOPE(UWorldPartitionStreamingPolicy::SetCellsStateToActivated);
 
 	int32 MaxCellsToLoad = GetMaxCellsToLoad();
 
-	// Sort cells based on importance
-	TArray<const UWorldPartitionRuntimeCell*, TInlineAllocator<256>> SortedCells;
-	if (WorldPartition->RuntimeHash)
-	{	
-		WorldPartition->RuntimeHash->SortStreamingCellsByImportance(ToActivateCells, StreamingSources, SortedCells);
-	}
-
 	// Trigger cell activation. Depending on actual state of cell limit loading.
-	for (const UWorldPartitionRuntimeCell* Cell : SortedCells)
+	for (const UWorldPartitionRuntimeCell* Cell : ToActivateCells)
 	{
-		if (ShouldSkipCellForPerformance(Cell))
-		{
-			continue;
-		}
-
 		UE_LOG(LogWorldPartition, Verbose, TEXT("UWorldPartitionStreamingPolicy::ActivateCells %s"), *Cell->GetName());
+
 		if (LoadedCells.Contains(Cell))
 		{
 			LoadedCells.Remove(Cell);
@@ -643,7 +691,7 @@ void UWorldPartitionStreamingPolicy::SetCellsStateToActivated(const TSet<const U
 	}
 }
 
-void UWorldPartitionStreamingPolicy::SetCellsStateToUnloaded(const TSet<const UWorldPartitionRuntimeCell*>& ToUnloadCells)
+void UWorldPartitionStreamingPolicy::SetCellsStateToUnloaded(const TArray<const UWorldPartitionRuntimeCell*>& ToUnloadCells)
 {
 	TRACE_CPUPROFILER_EVENT_SCOPE(UWorldPartitionStreamingPolicy::SetCellsStateToUnloaded);
 
@@ -691,14 +739,10 @@ bool UWorldPartitionStreamingPolicy::IsStreamingCompleted(EWorldPartitionRuntime
 	const UDataLayerSubsystem* DataLayerSubsystem = WorldPartition->GetWorld()->GetSubsystem<UDataLayerSubsystem>();
 	const bool bIsHLODEnabled = UHLODSubsystem::IsHLODEnabled();
 
+	bool bResult = true;
 	for (const FWorldPartitionStreamingQuerySource& QuerySource : QuerySources)
 	{
-		TSet<const UWorldPartitionRuntimeCell*> Cells;
-		if (WorldPartition->RuntimeHash)
-		{
-			WorldPartition->RuntimeHash->GetStreamingCells(QuerySource, Cells);
-		}
-		for (const UWorldPartitionRuntimeCell* Cell : Cells)
+		WorldPartition->RuntimeHash->ForEachStreamingCellsQuery(QuerySource, [QuerySource, QueryState, bExactState, bIsHLODEnabled, DataLayerSubsystem, &bResult](const UWorldPartitionRuntimeCell* Cell)
 		{
 			EWorldPartitionRuntimeCellState CellState = Cell->GetCurrentState();
 			if (CellState != QueryState)
@@ -726,13 +770,15 @@ bool UWorldPartitionStreamingPolicy::IsStreamingCompleted(EWorldPartitionRuntime
 								
 				if (!bSkipCell && (bExactState || CellState < QueryState))
 				{
+					bResult = false;
 					return false;
 				}
 			}
-		}	
-	}
 
-	return true;
+			return true;
+		});
+	}
+	return bResult;
 }
 
 FVector2D UWorldPartitionStreamingPolicy::GetDrawRuntimeHash2DDesiredFootprint(const FVector2D& CanvasSize)
diff --git a/Engine/Source/Runtime/Engine/Public/WorldPartition/WorldPartitionRuntimeCell.h b/Engine/Source/Runtime/Engine/Public/WorldPartition/WorldPartitionRuntimeCell.h
--- a/Engine/Source/Runtime/Engine/Public/WorldPartition/WorldPartitionRuntimeCell.h
+++ b/Engine/Source/Runtime/Engine/Public/WorldPartition/WorldPartitionRuntimeCell.h
@@ -151,16 +151,6 @@ class ENGINE_API UWorldPartitionRuntimeCell : public UObject, public IWorldParti
 {
 	GENERATED_UCLASS_BODY()
 
-	struct FStreamingSourceInfo
-	{
-		FStreamingSourceInfo(const FWorldPartitionStreamingSource& InSource, const FSphericalSector& InSourceShape)
-			: Source(InSource)
-			, SourceShape(InSourceShape)
-		{}
-		const FWorldPartitionStreamingSource& Source;
-		const FSphericalSector& SourceShape;
-	};
-
 	virtual void Load() const PURE_VIRTUAL(UWorldPartitionRuntimeCell::Load,);
 	virtual void Unload() const PURE_VIRTUAL(UWorldPartitionRuntimeCell::Unload,);
 	virtual bool CanUnload() const PURE_VIRTUAL(UWorldPartitionRuntimeCell::CanUnload, return true;);
@@ -182,8 +172,12 @@ class ENGINE_API UWorldPartitionRuntimeCell : public UObject, public IWorldParti
 	virtual int32 SortCompare(const UWorldPartitionRuntimeCell* Other) const;
 	virtual FName GetGridName() const { return DebugInfo.GridName; }
 	virtual FGuid const& GetContentBundleID() const { return ContentBundleID; }
-	/** Caches information on streaming source that will be used later on to sort cell. Returns true if cache was reset, else returns false. */
-	virtual bool CacheStreamingSourceInfo(const UWorldPartitionRuntimeCell::FStreamingSourceInfo& Info) const;
+
+	/** Caches information on streaming source that will be used later on to sort cell. */
+	bool ShouldResetStreamingSourceInfo() const;
+	virtual void ResetStreamingSourceInfo() const;
+	virtual void AppendStreamingSourceInfo(const FWorldPartitionStreamingSource& Source, const FSphericalSector& SourceShape) const;
+	virtual void MergeStreamingSourceInfo() const;
 
 	static void DirtyStreamingSourceCacheEpoch() { ++UWorldPartitionRuntimeCell::StreamingSourceCacheEpoch; }
 
diff --git a/Engine/Source/Runtime/Engine/Public/WorldPartition/WorldPartitionRuntimeHash.h b/Engine/Source/Runtime/Engine/Public/WorldPartition/WorldPartitionRuntimeHash.h
--- a/Engine/Source/Runtime/Engine/Public/WorldPartition/WorldPartitionRuntimeHash.h
+++ b/Engine/Source/Runtime/Engine/Public/WorldPartition/WorldPartitionRuntimeHash.h
@@ -45,6 +45,7 @@ class ENGINE_API UWorldPartitionRuntimeHash : public UObject
 
 #if WITH_EDITOR
 	virtual void SetDefaultValues() {}
+	virtual bool SupportsHLODs() const { return false; }
 	virtual bool PopulateGeneratorPackageForCook(const TArray<ICookPackageSplitter::FGeneratedPackageForPreSave>& InGeneratedPackages, TArray<UPackage*>& OutModifiedPackages) { return false; }
 	virtual bool PopulateGeneratedPackageForCook(UPackage* InPackage, const FString& InPackageRelativePath, TArray<UPackage*>& OutModifiedPackages) { return false; }
 	virtual bool GenerateStreaming(class UWorldPartitionStreamingPolicy* StreamingPolicy, const IStreamingGenerationContext* StreamingGenerationContext, TArray<FString>* OutPackagesToGenerate) { return false; }
@@ -59,12 +60,16 @@ class ENGINE_API UWorldPartitionRuntimeHash : public UObject
 	// PIE/Game methods
 	void OnBeginPlay();
 	void OnEndPlay();
+
+protected:
+	bool ConditionalRegisterAlwaysLoadedActorsForPIE(const FWorldPartitionActorDescView& ActorDescView, bool bIsMainWorldPartition, bool bIsMainContainer, bool bIsCellAlwaysLoaded);
 #endif
 
+public:
 	class FStreamingSourceCells
 	{
 	public:
-		void AddCell(const UWorldPartitionRuntimeCell* Cell, const UWorldPartitionRuntimeCell::FStreamingSourceInfo& Info);
+		void AddCell(const UWorldPartitionRuntimeCell* Cell, const FWorldPartitionStreamingSource& Source, const FSphericalSector& SourceShape);
 		void Reset() { Cells.Reset(); }
 		int32 Num() const { return Cells.Num(); }
 		TSet<const UWorldPartitionRuntimeCell*>& GetCells() { return Cells; }
@@ -73,11 +78,22 @@ class ENGINE_API UWorldPartitionRuntimeHash : public UObject
 		TSet<const UWorldPartitionRuntimeCell*> Cells;
 	};
 
+	// Deprecated streaming interface
+	UE_DEPRECATED(5.1, "GetAllStreamingCells is deprecated, use ForEachStreamingCells instead.")
+	int32 GetAllStreamingCells(TSet<const UWorldPartitionRuntimeCell*>& Cells, bool bAllDataLayers = false, bool bDataLayersOnly = false, const TSet<FName>& InDataLayers = TSet<FName>()) const;
+
+	UE_DEPRECATED(5.1, "GetStreamingCells is deprecated, use ForEachStreamingCells instead.")
+	bool GetStreamingCells(const FWorldPartitionStreamingQuerySource& QuerySource, TSet<const UWorldPartitionRuntimeCell*>& OutCells) const;
+
+	UE_DEPRECATED(5.1, "GetStreamingCells is deprecated, use ForEachStreamingCells instead.")
+	bool GetStreamingCells(const TArray<FWorldPartitionStreamingSource>& Sources, UWorldPartitionRuntimeHash::FStreamingSourceCells& OutActivateCells, UWorldPartitionRuntimeHash::FStreamingSourceCells& OutLoadCells) const;
+
 	// Streaming interface
-	virtual int32 GetAllStreamingCells(TSet<const UWorldPartitionRuntimeCell*>& Cells, bool bAllDataLayers = false, bool bDataLayersOnly = false, const TSet<FName>& InDataLayers = TSet<FName>()) const { return 0; }
-	virtual bool GetStreamingCells(const FWorldPartitionStreamingQuerySource& QuerySource, TSet<const UWorldPartitionRuntimeCell*>& OutCells) const { return false; }
-	virtual bool GetStreamingCells(const TArray<FWorldPartitionStreamingSource>& Sources, UWorldPartitionRuntimeHash::FStreamingSourceCells& OutActivateCells, UWorldPartitionRuntimeHash::FStreamingSourceCells& OutLoadCells) const { return false; };
-	virtual void SortStreamingCellsByImportance(const TSet<const UWorldPartitionRuntimeCell*>& InCells, const TArray<FWorldPartitionStreamingSource>& InSources, TArray<const UWorldPartitionRuntimeCell*, TInlineAllocator<256>>& OutSortedCells) const;
+	virtual void ForEachStreamingCells(TFunctionRef<bool(const UWorldPartitionRuntimeCell*)> Func) const {}
+	virtual void ForEachStreamingCellsQuery(const FWorldPartitionStreamingQuerySource& QuerySource, TFunctionRef<bool(const UWorldPartitionRuntimeCell*)> Func) const {}
+	virtual void ForEachStreamingCellsSources(const TArray<FWorldPartitionStreamingSource>& Sources, TFunctionRef<bool(const UWorldPartitionRuntimeCell*, EStreamingSourceTargetState)> Func) const {}
+
+	bool IsCellRelevantFor(bool bClientOnlyVisible) const;
 	EWorldPartitionStreamingPerformance GetStreamingPerformance(const TSet<const UWorldPartitionRuntimeCell*>& CellToActivate) const;
 
 	virtual bool InjectExternalStreamingObject(URuntimeHashExternalStreamingObjectBase* ExternalStreamingObject) { return false; }
diff --git a/Engine/Source/Runtime/Engine/Public/WorldPartition/WorldPartitionRuntimeSpatialHash.h b/Engine/Source/Runtime/Engine/Public/WorldPartition/WorldPartitionRuntimeSpatialHash.h
--- a/Engine/Source/Runtime/Engine/Public/WorldPartition/WorldPartitionRuntimeSpatialHash.h
+++ b/Engine/Source/Runtime/Engine/Public/WorldPartition/WorldPartitionRuntimeSpatialHash.h
@@ -242,6 +242,7 @@ public:
 
 #if WITH_EDITOR
 	virtual void SetDefaultValues() override;
+	virtual bool SupportsHLODs() const override { return true; }
 	virtual bool PopulateGeneratorPackageForCook(const TArray<ICookPackageSplitter::FGeneratedPackageForPreSave>& InGeneratedPackages, TArray<UPackage*>& OutModifiedPackages) override;
 	virtual bool PopulateGeneratedPackageForCook(UPackage* InPackage, const FString& InPackageRelativePath, TArray<UPackage*>& OutModifiedPackages) override;
 	virtual void FlushStreaming() override;
@@ -257,22 +258,21 @@ public:
 #endif
 
 	// streaming interface
-	virtual int32 GetAllStreamingCells(TSet<const UWorldPartitionRuntimeCell*>& Cells, bool bAllDataLayers = false, bool bDataLayersOnly = false, const TSet<FName>& InDataLayers = TSet<FName>()) const override;
-	virtual bool GetStreamingCells(const FWorldPartitionStreamingQuerySource& QuerySource, TSet<const UWorldPartitionRuntimeCell*>& OutCells) const override;
-	virtual bool GetStreamingCells(const TArray<FWorldPartitionStreamingSource>& Sources, UWorldPartitionRuntimeHash::FStreamingSourceCells& OutActivateCells, UWorldPartitionRuntimeHash::FStreamingSourceCells& OutLoadCells) const override;
+	virtual void ForEachStreamingCells(TFunctionRef<bool(const UWorldPartitionRuntimeCell*)> Func) const override;
+	virtual void ForEachStreamingCellsQuery(const FWorldPartitionStreamingQuerySource& QuerySource, TFunctionRef<bool(const UWorldPartitionRuntimeCell*)> Func) const override;
+	virtual void ForEachStreamingCellsSources(const TArray<FWorldPartitionStreamingSource>& Sources, TFunctionRef<bool(const UWorldPartitionRuntimeCell*, EStreamingSourceTargetState)> Func) const override;
 
 	virtual bool InjectExternalStreamingObject(URuntimeHashExternalStreamingObjectBase* ExternalStreamingObject) override;
 	virtual bool RemoveExternalStreamingObject(URuntimeHashExternalStreamingObjectBase* ExternalStreamingObject) override;
 
 	uint32 GetNumGrids() const;
-		
+
 	static FString GetCellCoordString(const FGridCellCoord& InCellGlobalCoord);
 
 protected:
 	void ForEachStreamingGrid(TFunctionRef<void(FSpatialHashStreamingGrid&)> Func);
 	void ForEachStreamingGrid(TFunctionRef<void(const FSpatialHashStreamingGrid&)> Func) const;
 
-	bool ShouldConsiderClientOnlyVisibleCells() const;
 	virtual EWorldPartitionStreamingPerformance GetStreamingPerformanceForCell(const UWorldPartitionRuntimeCell* Cell) const override;
 
 #if WITH_EDITOR
diff --git a/Engine/Source/Runtime/Engine/Public/WorldPartition/WorldPartitionRuntimeSpatialHashCell.h b/Engine/Source/Runtime/Engine/Public/WorldPartition/WorldPartitionRuntimeSpatialHashCell.h
--- a/Engine/Source/Runtime/Engine/Public/WorldPartition/WorldPartitionRuntimeSpatialHashCell.h
+++ b/Engine/Source/Runtime/Engine/Public/WorldPartition/WorldPartitionRuntimeSpatialHashCell.h
@@ -15,7 +15,10 @@ class UWorldPartitionRuntimeSpatialHashCell : public UWorldPartitionRuntimeCell
 #if WITH_EDITOR
 	virtual void PostDuplicate(bool bDuplicateForPIE) override;
 #endif
-	virtual bool CacheStreamingSourceInfo(const UWorldPartitionRuntimeCell::FStreamingSourceInfo& Info) const override;
+	virtual void ResetStreamingSourceInfo() const override;
+	virtual void AppendStreamingSourceInfo(const FWorldPartitionStreamingSource& Source, const FSphericalSector& SourceShape) const override;
+	virtual void MergeStreamingSourceInfo() const override;
+
 	virtual int32 SortCompare(const UWorldPartitionRuntimeCell* InOther) const override;
 
 	UPROPERTY()
diff --git a/Engine/Source/Runtime/Engine/Public/WorldPartition/WorldPartitionStreamingPolicy.h b/Engine/Source/Runtime/Engine/Public/WorldPartition/WorldPartitionStreamingPolicy.h
--- a/Engine/Source/Runtime/Engine/Public/WorldPartition/WorldPartitionStreamingPolicy.h
+++ b/Engine/Source/Runtime/Engine/Public/WorldPartition/WorldPartitionStreamingPolicy.h
@@ -43,7 +43,6 @@ class UWorldPartitionStreamingPolicy : public UObject
 
 public:
 	virtual void UpdateStreamingState();
-	virtual void SetTargetStateForCells(EWorldPartitionRuntimeCellState TargetState, const TSet<const UWorldPartitionRuntimeCell*>& Cells);
 	virtual bool CanAddLoadedLevelToWorld(class ULevel* InLevel) const;
 	virtual FVector2D GetDrawRuntimeHash2DDesiredFootprint(const FVector2D& CanvasSize);
 	virtual bool DrawRuntimeHash2D(class UCanvas* Canvas, const FVector2D& PartitionCanvasSize, const FVector2D& Offset);
@@ -72,9 +71,9 @@ public:
 	EWorldPartitionStreamingPerformance GetStreamingPerformance() const { return StreamingPerformance; }
 
 protected:
-	virtual void SetCellsStateToLoaded(const TSet<const UWorldPartitionRuntimeCell*>& ToLoadCells);
-	virtual void SetCellsStateToActivated(const TSet<const UWorldPartitionRuntimeCell*>& ToActivateCells);
-	virtual void SetCellsStateToUnloaded(const TSet<const UWorldPartitionRuntimeCell*>& ToUnloadCells);
+	virtual void SetCellsStateToLoaded(const TArray<const UWorldPartitionRuntimeCell*>& ToLoadCells);
+	virtual void SetCellsStateToActivated(const TArray<const UWorldPartitionRuntimeCell*>& ToActivateCells);
+	virtual void SetCellsStateToUnloaded(const TArray<const UWorldPartitionRuntimeCell*>& ToUnloadCells);
 	virtual int32 GetCellLoadingCount() const { return 0; }
 	virtual int32 GetMaxCellsToLoad() const;
 	virtual void UpdateStreamingSources();
@@ -90,12 +89,12 @@ protected:
 	TArray<FWorldPartitionStreamingSource> StreamingSources;
 	TMap<FName, FStreamingSourceVelocity> StreamingSourcesVelocity;
 
-	UWorldPartitionRuntimeHash::FStreamingSourceCells FrameActivateCells;
-	UWorldPartitionRuntimeHash::FStreamingSourceCells FrameLoadCells;
+	TSet<const UWorldPartitionRuntimeCell*> FrameActivateCells;
+	TSet<const UWorldPartitionRuntimeCell*> FrameLoadCells;
 	
 	bool bCriticalPerformanceRequestedBlockTillOnWorld;
 	int32 CriticalPerformanceBlockTillLevelStreamingCompletedEpoch;
-	mutable TArray<const UWorldPartitionRuntimeCell*, TInlineAllocator<256>> SortedAddToWorldCells;
+	mutable TArray<const UWorldPartitionRuntimeCell*> SortedAddToWorldCells;
 
 	int32 DataLayersStatesServerEpoch;
 	int32 ServerStreamingEnabledEpoch;
