Backout CL 22012336 as it was causing crashes : each content bundle was considering actors of the main world persistent level as being part of their unsaved actor.

diff --git a/Engine/Source/Runtime/Engine/Private/WorldPartition/WorldPartitionStreamingGeneration.cpp b/Engine/Source/Runtime/Engine/Private/WorldPartition/WorldPartitionStreamingGeneration.cpp
--- a/Engine/Source/Runtime/Engine/Private/WorldPartition/WorldPartitionStreamingGeneration.cpp
+++ b/Engine/Source/Runtime/Engine/Private/WorldPartition/WorldPartitionStreamingGeneration.cpp
@@ -341,8 +341,8 @@ class FWorldPartitionStreamingGenerator
 		{
 			for (AActor* Actor : InContainer->GetWorld()->PersistentLevel->Actors)
 			{
-				// @todo_ow: temp fix, don't filter out containers (mainly for content bundles) to fix PIE in unsaved maps.
-				if (IsValid(Actor) && Actor->IsPackageExternal() && Actor->IsMainPackageActor() && !Actor->IsEditorOnly() && /*InContainer->IsActorDescHandled(Actor) &&*/ !InContainer->GetActorDesc(Actor->GetActorGuid()))
+				// @todo_ow: Fix PIE in unsaved maps (IsActorDescHandled always returns false for new unsaved maps)
+				if (IsValid(Actor) && Actor->IsPackageExternal() && Actor->IsMainPackageActor() && !Actor->IsEditorOnly() && InContainer->IsActorDescHandled(Actor) && !InContainer->GetActorDesc(Actor->GetActorGuid()))
 				{
 					FWorldPartitionActorDescView ModifiedActorDescView = GetModifiedActorDesc(Actor, InContainer);
 					RegisterActorDescView(Actor->GetActorGuid(), ModifiedActorDescView);
