Do not consider actor initialized when running deferred construction scripts

diff --git a/Engine/Source/Runtime/Engine/Classes/GameFramework/Actor.h b/Engine/Source/Runtime/Engine/Classes/GameFramework/Actor.h
--- a/Engine/Source/Runtime/Engine/Classes/GameFramework/Actor.h
+++ b/Engine/Source/Runtime/Engine/Classes/GameFramework/Actor.h
@@ -436,6 +436,8 @@ protected:
 	uint8 bReplicateUsingRegisteredSubObjectList : 1;
 
 private:
+	friend class FActorDeferredScriptManager;
+
 	/** Whether FinishSpawning has been called for this Actor.  If it has not, the Actor is in a malformed state */
 	uint8 bHasFinishedSpawning:1;
 
diff --git a/Engine/Source/Runtime/Engine/Private/ActorDeferredScriptManager.cpp b/Engine/Source/Runtime/Engine/Private/ActorDeferredScriptManager.cpp
--- a/Engine/Source/Runtime/Engine/Private/ActorDeferredScriptManager.cpp
+++ b/Engine/Source/Runtime/Engine/Private/ActorDeferredScriptManager.cpp
@@ -78,6 +78,8 @@ void FActorDeferredScriptManager::ProcessAsyncTasks(bool bLimitExecutionTime)
 			PendingConstructionScriptActors.PopLast();
 			if (AActor* Actor = WeakActor.Get(); Actor && Actor->GetWorld())
 			{
+				// Temporarily do not consider actor as initialized if they were when running deferred construction scripts
+				FGuardValue_Bitfield(Actor->bActorInitialized, false);
 				Actor->RerunConstructionScripts();
 			}
 			bHasTimeLeft = bLimitExecutionTime ? ((FPlatformTime::Seconds() - TickStartTime) < MaxSecondsPerFrame) : true;
