Add missing case in ERepLayoutFlags LexToString

diff --git a/Engine/Source/Runtime/Engine/Private/RepLayout.cpp b/Engine/Source/Runtime/Engine/Private/RepLayout.cpp
--- a/Engine/Source/Runtime/Engine/Private/RepLayout.cpp
+++ b/Engine/Source/Runtime/Engine/Private/RepLayout.cpp
@@ -8378,6 +8378,8 @@ const TCHAR* LexToString(ERepLayoutFlags Flag)
 		return TEXT("NoReplicatedProperties");
 	case ERepLayoutFlags::FullPushProperties:
 		return TEXT("FullPushProperties");
+	case ERepLayoutFlags::HasInitialOnlyProperties:
+		return TEXT("HasInitialOnlyProperties");
 	default:
 		check(false);
 		return TEXT("Unknown");
