Fix content validation trying to validate content bundle actors with an unknwon level.
Plugin actors validaton will only work with maps under /Game/ for the time being

diff --git a/Engine/Source/Runtime/Engine/Private/WorldPartition/ContentBundle/ContentBundlePaths.cpp b/Engine/Source/Runtime/Engine/Private/WorldPartition/ContentBundle/ContentBundlePaths.cpp
new file mode 100644
--- /dev/null
+++ b/Engine/Source/Runtime/Engine/Private/WorldPartition/ContentBundle/ContentBundlePaths.cpp
@@ -0,0 +1,27 @@
+// Copyright Epic Games, Inc. All Rights Reserved.
+#include "WorldPartition/ContentBundle/ContentBundlePaths.h"
+
+#if WITH_EDITOR
+#include "String/Find.h"
+#endif
+
+#if WITH_EDITOR
+
+FStringView ContentBundlePaths::GetRelativeExternalActorPackagePath(FStringView InContentBundleExternalActorPackagePath)
+{
+	uint32 ContentBundleStartIdx = InContentBundleExternalActorPackagePath.Find(ContentBundlePaths::ContentBundleFolder, ESearchCase::IgnoreCase);
+	if (ContentBundleStartIdx != INDEX_NONE)
+	{
+		uint32 ContentBundleUIDIdx = ContentBundleStartIdx + FCString::Strlen(ContentBundlePaths::ContentBundleFolder) + 1; // + 1 for trailing slash
+		FStringView RelativeActorPackge = InContentBundleExternalActorPackagePath.RightChop(ContentBundleUIDIdx);
+		
+		uint32 ContentBundleEndIdx = RelativeActorPackge.Find(TEXT("/"), ESearchCase::IgnoreCase); // Search for the end of the content bundle guids
+		check(ContentBundleEndIdx != INDEX_NONE);
+
+		return RelativeActorPackge.RightChop(ContentBundleEndIdx);
+	}
+
+	return FStringView();
+}
+
+#endif
\ No newline at end of file
diff --git a/Engine/Source/Runtime/Engine/Public/WorldPartition/ContentBundle/ContentBundlePaths.h b/Engine/Source/Runtime/Engine/Public/WorldPartition/ContentBundle/ContentBundlePaths.h
new file mode 100644
--- /dev/null
+++ b/Engine/Source/Runtime/Engine/Public/WorldPartition/ContentBundle/ContentBundlePaths.h
@@ -0,0 +1,16 @@
+// Copyright Epic Games, Inc. All Rights Reserved.
+#pragma once
+
+#include "CoreMinimal.h"
+#include "Containers/StringView.h"
+
+namespace ContentBundlePaths
+{
+	constexpr TCHAR ContentBundleFolder[] = TEXT("ContentBundle");
+
+#if WITH_EDITOR
+	// InContentBundleExternalActorPackagePath format is : //{Plugin}/{ExternalActorFolder}/{ContentBundleFolder}/{ContentBundleUID}/{LevelPath}/{ExternalActorPackagePath}
+	// return format is {LevelPath}/{ExternalActorPackagePath}, empty otherwise
+	ENGINE_API FStringView GetRelativeExternalActorPackagePath(FStringView InContentBundleExternalActorPackagePath);
+#endif
+}
\ No newline at end of file
