[WP Navmesh] Get the required overlap (to build world partitioned navmesh) from the navigation system instead of using a constant.

diff --git a/Engine/Source/Editor/UnrealEd/Private/WorldPartition/WorldPartitionNavigationDataBuilder.cpp b/Engine/Source/Editor/UnrealEd/Private/WorldPartition/WorldPartitionNavigationDataBuilder.cpp
--- a/Engine/Source/Editor/UnrealEd/Private/WorldPartition/WorldPartitionNavigationDataBuilder.cpp
+++ b/Engine/Source/Editor/UnrealEd/Private/WorldPartition/WorldPartitionNavigationDataBuilder.cpp
@@ -54,8 +54,9 @@ bool UWorldPartitionNavigationDataBuilder::PreRun(UWorld* World, FPackageSourceC
 	IterativeCellSize = GridSize * FMath::Max(LoadingCellSizeSetting / GridSize, 1u);
 	
 	// Extra padding around loaded cell.
-	// @todo: set value programatically.
-	IterativeCellOverlapSize = 2000; // navmesh tile size
+	// Make sure navigation is added and initialized before fetching the build overlap
+	FNavigationSystem::AddNavigationSystemToWorld(*World, FNavigationSystemRunMode::EditorWorldPartitionBuildMode);
+	IterativeCellOverlapSize = FMath::CeilToInt32(FNavigationSystem::GetWorldPartitionNavigationDataBuilderOverlap(*World));
 
 	TArray<FString> Tokens, Switches;
 	UCommandlet::ParseCommandLine(FCommandLine::Get(), Tokens, Switches);
diff --git a/Engine/Source/Runtime/Engine/Classes/AI/NavigationSystemBase.h b/Engine/Source/Runtime/Engine/Classes/AI/NavigationSystemBase.h
--- a/Engine/Source/Runtime/Engine/Classes/AI/NavigationSystemBase.h
+++ b/Engine/Source/Runtime/Engine/Classes/AI/NavigationSystemBase.h
@@ -153,6 +153,7 @@ namespace FNavigationSystem
 	
 	ENGINE_API const FNavDataConfig& GetDefaultSupportedAgent();
 	ENGINE_API const FNavDataConfig& GetBiggestSupportedAgent(const UWorld* World);
+	ENGINE_API double GetWorldPartitionNavigationDataBuilderOverlap(const UWorld& World);
 
 	ENGINE_API TSubclassOf<UNavAreaBase> GetDefaultWalkableArea();
 	ENGINE_API TSubclassOf<UNavAreaBase> GetDefaultObstacleArea();
@@ -203,6 +204,7 @@ namespace FNavigationSystem
 	DECLARE_DELEGATE_RetVal(TSubclassOf<UNavAreaBase>, FNavAreaBasedSignature);
 	DECLARE_DELEGATE_RetVal(const FNavDataConfig&, FNavDataConfigBasedSignature);
 	DECLARE_DELEGATE_RetVal_OneParam(const FNavDataConfig&, FNavDataConfigAndWorldSignature, const UWorld* /*World*/);
+	DECLARE_DELEGATE_RetVal_OneParam(double, FDoubleWorldBasedSignature, const UWorld& /*World*/);
 	DECLARE_DELEGATE_TwoParams(FWorldByteBasedSignature, UWorld& /*World*/, uint8 /*Flags*/);
 	DECLARE_DELEGATE_TwoParams(FActorBooleBasedSignature, AActor& /*Actor*/, bool /*bUpdateAttachedActors*/);
 	DECLARE_DELEGATE_ThreeParams(FComponentBoundsChangeSignature, UActorComponent& /*Comp*/, const FBox& /*NewBounds*/, const FBox& /*DirtyArea*/)
@@ -309,6 +311,7 @@ protected:
 	static FNavigationSystem::FNavigationAutoUpdateEnableSignature& SetNavigationAutoUpdateEnableDelegate();
 	static FNavigationSystem::FWorldByteBasedSignature& AddNavigationUpdateLockDelegate();
 	static FNavigationSystem::FWorldByteBasedSignature& RemoveNavigationUpdateLockDelegate();
+	static FNavigationSystem::FDoubleWorldBasedSignature& GetWorldPartitionNavigationDataBuilderOverlapDelegate();
 #endif // WITH_EDITOR
 };
 
diff --git a/Engine/Source/Runtime/Engine/Private/AI/NavigationSystemBase.cpp b/Engine/Source/Runtime/Engine/Private/AI/NavigationSystemBase.cpp
--- a/Engine/Source/Runtime/Engine/Private/AI/NavigationSystemBase.cpp
+++ b/Engine/Source/Runtime/Engine/Private/AI/NavigationSystemBase.cpp
@@ -160,6 +160,7 @@ namespace FNavigationSystem
 		FNavigationAutoUpdateEnableSignature SetNavigationAutoUpdateEnable;
 		FWorldByteBasedSignature AddNavigationUpdateLock;
 		FWorldByteBasedSignature RemoveNavigationUpdateLock;
+		FDoubleWorldBasedSignature GetWorldPartitionNavigationDataBuilderOverlap;
 #endif // WITH_EDITOR
 
 		FDelegates()
@@ -195,6 +196,7 @@ namespace FNavigationSystem
 			SetNavigationAutoUpdateEnable.BindLambda([](const bool, UNavigationSystemBase*) {});
 			AddNavigationUpdateLock.BindLambda([](UWorld&, uint8) {});
 			RemoveNavigationUpdateLock.BindLambda([](UWorld&, uint8) {});
+			GetWorldPartitionNavigationDataBuilderOverlap.BindLambda([](const UWorld&){ return 0; });
 #endif // WITH_EDITOR
 		}
 	};
@@ -221,7 +223,9 @@ namespace FNavigationSystem
 	bool HasComponentData(UActorComponent& Comp) { return Delegates.HasComponentData.Execute(Comp);	}
 	const FNavDataConfig& GetDefaultSupportedAgent() { return Delegates.GetDefaultSupportedAgent.Execute(); }
 	const FNavDataConfig& GetBiggestSupportedAgent(const UWorld* World) { return Delegates.GetBiggestSupportedAgent.Execute(World); }
-
+#if WITH_EDITOR	
+	double GetWorldPartitionNavigationDataBuilderOverlap(const UWorld& World) { return Delegates.GetWorldPartitionNavigationDataBuilderOverlap.Execute(World); }
+#endif	
 
 	TSubclassOf<UNavAreaBase> DefaultWalkableArea; 
 	TSubclassOf<UNavAreaBase> DefaultObstacleArea;
@@ -432,6 +436,7 @@ FNavigationSystem::FLevelBasedSignature& UNavigationSystemBase::UpdateLevelColli
 FNavigationSystem::FNavigationAutoUpdateEnableSignature& UNavigationSystemBase::SetNavigationAutoUpdateEnableDelegate() { return FNavigationSystem::Delegates.SetNavigationAutoUpdateEnable; }
 FNavigationSystem::FWorldByteBasedSignature& UNavigationSystemBase::AddNavigationUpdateLockDelegate() { return FNavigationSystem::Delegates.AddNavigationUpdateLock; }
 FNavigationSystem::FWorldByteBasedSignature& UNavigationSystemBase::RemoveNavigationUpdateLockDelegate() { return FNavigationSystem::Delegates.RemoveNavigationUpdateLock; }
+FNavigationSystem::FDoubleWorldBasedSignature& UNavigationSystemBase::GetWorldPartitionNavigationDataBuilderOverlapDelegate() { return FNavigationSystem::Delegates.GetWorldPartitionNavigationDataBuilderOverlap; }
 #endif // WITH_EDITOR
 //----------------------------------------------------------------------//
 // IPathFollowingManagerInterface
diff --git a/Engine/Source/Runtime/NavigationSystem/Private/NavMesh/RecastNavMesh.cpp b/Engine/Source/Runtime/NavigationSystem/Private/NavMesh/RecastNavMesh.cpp
--- a/Engine/Source/Runtime/NavigationSystem/Private/NavMesh/RecastNavMesh.cpp
+++ b/Engine/Source/Runtime/NavigationSystem/Private/NavMesh/RecastNavMesh.cpp
@@ -2586,6 +2586,13 @@ void ARecastNavMesh::OnStreamingLevelRemoved(ULevel* InLevel, UWorld* InWorld)
 	}
 }
 
+#if WITH_EDITOR
+double ARecastNavMesh::GetWorldPartitionNavigationDataBuilderOverlap() const
+{
+	return TileSizeUU;
+}
+#endif //WITH_EDITOR
+
 void ARecastNavMesh::DetachNavMeshDataChunk(URecastNavMeshDataChunk& NavDataChunk)
 {
 	const TArray<FNavTileRef> DetachedIndices = NavDataChunk.DetachTiles(*this);
diff --git a/Engine/Source/Runtime/NavigationSystem/Private/NavigationSystem.cpp b/Engine/Source/Runtime/NavigationSystem/Private/NavigationSystem.cpp
--- a/Engine/Source/Runtime/NavigationSystem/Private/NavigationSystem.cpp
+++ b/Engine/Source/Runtime/NavigationSystem/Private/NavigationSystem.cpp
@@ -552,6 +552,7 @@ UNavigationSystemV1::UNavigationSystemV1(const FObjectInitializer& ObjectInitial
 					NavSys->RemoveNavigationUpdateLock(Flags);
 				}
 			});
+			UNavigationSystemBase::GetWorldPartitionNavigationDataBuilderOverlapDelegate().BindStatic(&UNavigationSystemV1::GetWorldPartitionNavigationDataBuilderOverlap);
 #endif // WITH_EDITOR
 
 #if ENABLE_VISUAL_LOG
@@ -4869,6 +4870,30 @@ const FNavDataConfig& UNavigationSystemV1::GetBiggestSupportedAgent(const UWorld
 	return *BiggestAgent;
 }
 
+#if WITH_EDITOR
+double UNavigationSystemV1::GetWorldPartitionNavigationDataBuilderOverlap(const UWorld& World) 
+{
+	const UNavigationSystemV1* NavSys = FNavigationSystem::GetCurrent<UNavigationSystemV1>(&World);
+	if (NavSys == nullptr)
+	{
+		// If no world is available, use the CDO.
+		NavSys = GetDefault<UNavigationSystemV1>();
+	}
+	check(NavSys);
+
+	double MaxOverlap = 0;
+	for (const ANavigationData* NavData : NavSys->NavDataSet)
+	{
+		if (NavData)
+		{
+			MaxOverlap = FMath::Max(MaxOverlap, NavData->GetWorldPartitionNavigationDataBuilderOverlap());
+		}
+	}
+
+	return MaxOverlap;
+}
+#endif //WITH_EDITOR
+
 const FNavDataConfig& UNavigationSystemV1::GetDefaultSupportedAgentConfig() const 
 { 
 	static const FNavDataConfig DefaultAgent;
diff --git a/Engine/Source/Runtime/NavigationSystem/Public/NavMesh/RecastNavMesh.h b/Engine/Source/Runtime/NavigationSystem/Public/NavMesh/RecastNavMesh.h
--- a/Engine/Source/Runtime/NavigationSystem/Public/NavMesh/RecastNavMesh.h
+++ b/Engine/Source/Runtime/NavigationSystem/Public/NavMesh/RecastNavMesh.h
@@ -965,6 +965,10 @@ public:
 	
 	virtual void OnStreamingLevelAdded(ULevel* InLevel, UWorld* InWorld) override;
 	virtual void OnStreamingLevelRemoved(ULevel* InLevel, UWorld* InWorld) override;
+
+#if WITH_EDITOR
+	virtual double GetWorldPartitionNavigationDataBuilderOverlap() const override;
+#endif
 	//~ End ANavigationData Interface
 
 	virtual void AttachNavMeshDataChunk(URecastNavMeshDataChunk& NavDataChunk);
diff --git a/Engine/Source/Runtime/NavigationSystem/Public/NavigationData.h b/Engine/Source/Runtime/NavigationSystem/Public/NavigationData.h
--- a/Engine/Source/Runtime/NavigationSystem/Public/NavigationData.h
+++ b/Engine/Source/Runtime/NavigationSystem/Public/NavigationData.h
@@ -595,7 +595,11 @@ public:
 	
 	virtual void OnStreamingLevelAdded(ULevel* InLevel, UWorld* InWorld) {};
 	virtual void OnStreamingLevelRemoved(ULevel* InLevel, UWorld* InWorld) {};
-	
+
+#if WITH_EDITOR 
+	virtual double GetWorldPartitionNavigationDataBuilderOverlap() const { return 0; }
+#endif
+
 	//----------------------------------------------------------------------//
 	// Generation & data access                                                      
 	//----------------------------------------------------------------------//
diff --git a/Engine/Source/Runtime/NavigationSystem/Public/NavigationSystem.h b/Engine/Source/Runtime/NavigationSystem/Public/NavigationSystem.h
--- a/Engine/Source/Runtime/NavigationSystem/Public/NavigationSystem.h
+++ b/Engine/Source/Runtime/NavigationSystem/Public/NavigationSystem.h
@@ -689,6 +689,11 @@ public:
 
 	static const FNavDataConfig& GetDefaultSupportedAgent();
 	static const FNavDataConfig& GetBiggestSupportedAgent(const UWorld* World);
+	
+#if WITH_EDITOR
+	static double GetWorldPartitionNavigationDataBuilderOverlap(const UWorld& World);
+#endif
+	
 	const FNavDataConfig& GetDefaultSupportedAgentConfig() const;
 	FORCEINLINE const TArray<FNavDataConfig>& GetSupportedAgents() const { return SupportedAgents; }
 	void OverrideSupportedAgents(const TArray<FNavDataConfig>& NewSupportedAgents);
