Ensure that BP events are called in anim BP editor when recompiling

More fallout from the new re-instancing path

diff --git a/Engine/Source/Runtime/Engine/Private/Animation/AnimInstance.cpp b/Engine/Source/Runtime/Engine/Private/Animation/AnimInstance.cpp
--- a/Engine/Source/Runtime/Engine/Private/Animation/AnimInstance.cpp
+++ b/Engine/Source/Runtime/Engine/Private/Animation/AnimInstance.cpp
@@ -3854,19 +3854,17 @@ void UAnimInstance::HandleObjectsReinstanced(const TMap<UObject*, UObject*>& Old
 		if(bThisObjectWasReinstanced)
 		{
 			TRACE_OBJECT_LIFETIME_BEGIN(this);
-
-			RecalcRequiredBones();
-			
 			// Minimally reinit proxy (i.e. dont call per-node initialization) unless we are in an editor preview world (i.e. we are in the anim BP editor)
-			FAnimInstanceProxy& Proxy = GetProxyOnGameThread<FAnimInstanceProxy>();
 			UWorld* World = GetWorld();
 			if(World && World->WorldType == EWorldType::EditorPreview)
 			{
-				Proxy.Initialize(this);
-				Proxy.InitializeRootNode();
+				InitializeAnimation(false);
 			}
 			else
 			{
+				RecalcRequiredBones();
+
+				FAnimInstanceProxy& Proxy = GetProxyOnGameThread<FAnimInstanceProxy>();
 				Proxy.Initialize(this);
 				Proxy.InitializeCachedClassData();
 				Proxy.InitializeRootNode_WithRoot(Proxy.RootNode);
