Undo //UE5/Release-5.1/Engine/Source/Runtime/Engine/Private/... changelist 22003523

Editgrate - Backing out change as don't have time to look into now and it has caused pri 0 bug.

diff --git a/Engine/Source/Runtime/Engine/Private/PrimitiveComponentPhysics.cpp b/Engine/Source/Runtime/Engine/Private/PrimitiveComponentPhysics.cpp
--- a/Engine/Source/Runtime/Engine/Private/PrimitiveComponentPhysics.cpp
+++ b/Engine/Source/Runtime/Engine/Private/PrimitiveComponentPhysics.cpp
@@ -807,13 +807,9 @@ void UPrimitiveComponent::UnWeldChildren()
 {
 	for (USceneComponent* ChildComponent : GetAttachChildren())
 	{
-		// Invalid children should be removed in UPrimitiveComponent::DestroyComponent so to check that you have correctly destroyed attached components.
-		if (ensure(IsValid(ChildComponent)))
+		if (UPrimitiveComponent* PrimComp = Cast<UPrimitiveComponent>(ChildComponent))
 		{
-			if (UPrimitiveComponent* PrimComp = Cast<UPrimitiveComponent>(ChildComponent))
-			{
-				PrimComp->UnWeldFromParent();
-			}
+			PrimComp->UnWeldFromParent();
 		}
 	}
 }
