Fix deprecation error caused by merge.

diff --git a/Engine/Source/Runtime/Engine/Private/WorldPartition/WorldPartitionRuntimeHash.cpp b/Engine/Source/Runtime/Engine/Private/WorldPartition/WorldPartitionRuntimeHash.cpp
--- a/Engine/Source/Runtime/Engine/Private/WorldPartition/WorldPartitionRuntimeHash.cpp
+++ b/Engine/Source/Runtime/Engine/Private/WorldPartition/WorldPartitionRuntimeHash.cpp
@@ -50,7 +50,7 @@ bool UWorldPartitionRuntimeHash::ConditionalRegisterAlwaysLoadedActorsForPIE(con
 		// This will load the actor if it isn't already loaded
 		FWorldPartitionReference Reference(GetOuterUWorldPartition(), ActorDescView.GetGuid());
 
-		if (AActor* AlwaysLoadedActor = FindObject<AActor>(nullptr, *ActorDescView.GetActorPath().ToString()))
+		if (AActor* AlwaysLoadedActor = FindObject<AActor>(nullptr, *ActorDescView.GetActorSoftPath().ToString()))
 		{
 			AlwaysLoadedActorsForPIE.Emplace(Reference, AlwaysLoadedActor);
 
