Expose option on BodyInstance::SetInstanceSimulatePhysics to allow preservation of any attachment of the component to its parent. Existing behavior will not change.

diff --git a/Engine/Source/Runtime/Engine/Classes/PhysicsEngine/BodyInstance.h b/Engine/Source/Runtime/Engine/Classes/PhysicsEngine/BodyInstance.h
--- a/Engine/Source/Runtime/Engine/Classes/PhysicsEngine/BodyInstance.h
+++ b/Engine/Source/Runtime/Engine/Classes/PhysicsEngine/BodyInstance.h
@@ -769,8 +769,17 @@ public:
 	void SetInertiaConditioningEnabled(bool bEnabled);
 
 
-	/** Set this body to be fixed (kinematic) or not. */
-	void SetInstanceSimulatePhysics(bool bSimulate, bool bMaintainPhysicsBlending=false);
+	/** 
+	 * Set this body to either simulate or to be fixed/kinematic. 
+	 * 
+	 * @param bMaintainPhysicsBlending If true then the physics blend weight will not be adjusted. If false then 
+	 *        it will get set to 0 or 1 depending on bSimulate.
+	 * @param bPreserveExistingAttachments If true then any existing attachment between the owning component and 
+	 *        its parent will be preserved, even when switching to simulate (most likely useful for skeletal meshes
+	 *        that are parented to a moveable component). If false then the owning component will be detached 
+	 *        from its parent if this is the root body and it is being set to simulate.
+	 */
+	void SetInstanceSimulatePhysics(bool bSimulate, bool bMaintainPhysicsBlending=false, bool bPreserveExistingAttachment = false);
 	/** Makes sure the current kinematic state matches the simulate flag */
 	void UpdateInstanceSimulatePhysics();
 	/** Returns true if this body is simulating, false if it is fixed (kinematic) */
diff --git a/Engine/Source/Runtime/Engine/Private/PhysicsEngine/BodyInstance.cpp b/Engine/Source/Runtime/Engine/Private/PhysicsEngine/BodyInstance.cpp
--- a/Engine/Source/Runtime/Engine/Private/PhysicsEngine/BodyInstance.cpp
+++ b/Engine/Source/Runtime/Engine/Private/PhysicsEngine/BodyInstance.cpp
@@ -2298,7 +2298,7 @@ void FBodyInstance::ApplyWeldOnChildren()
 	
 }
 
-void FBodyInstance::SetInstanceSimulatePhysics(bool bSimulate, bool bMaintainPhysicsBlending)
+void FBodyInstance::SetInstanceSimulatePhysics(bool bSimulate, bool bMaintainPhysicsBlending, bool bPreserveExistingAttachment)
 {
 	if (bSimulate)
 	{
@@ -2307,7 +2307,7 @@ void FBodyInstance::SetInstanceSimulatePhysics(bool bSimulate, bool bMaintainPhy
 		// If we are enabling simulation, and we are the root body of our component (or we are welded), we detach the component 
 		if (OwnerComponentInst && OwnerComponentInst->IsRegistered() && (OwnerComponentInst->GetBodyInstance() == this || OwnerComponentInst->IsWelded()))
 		{
-			if (OwnerComponentInst->GetAttachParent())
+			if (!bPreserveExistingAttachment && OwnerComponentInst->GetAttachParent())
 			{
 				OwnerComponentInst->DetachFromComponent(FDetachmentTransformRules::KeepWorldTransform);
 			}
