Check inputs to IF material expression for equality.
If the are the same, or evaluate to the same constant then pass one directly through.
This prevents the expression from unnecessarily converting a constant expression into a non-constant expression.
We use constant expression evaluation at the final output to determine if WPO or PDO os used.
So this fix helps with culling WPO and PDO.

diff --git a/Engine/Source/Runtime/Engine/Private/Materials/HLSLMaterialTranslator.cpp b/Engine/Source/Runtime/Engine/Private/Materials/HLSLMaterialTranslator.cpp
--- a/Engine/Source/Runtime/Engine/Private/Materials/HLSLMaterialTranslator.cpp
+++ b/Engine/Source/Runtime/Engine/Private/Materials/HLSLMaterialTranslator.cpp
@@ -5337,18 +5337,64 @@ int32 FHLSLMaterialTranslator::ActorWorldPosition()
 	return CastToNonLWCIfDisabled(Result);
 }
 
-int32 FHLSLMaterialTranslator::If(int32 A,int32 B,int32 AGreaterThanB,int32 AEqualsB,int32 ALessThanB,int32 ThresholdArg)
+// Compare two inputs and return true if they are either the same, or if they evaluate to equal constant expressions.
+static bool AreEqualExpressions(int32 A, int32 B, TArray<FShaderCodeChunk> const& InCurrentScopeChunks, FMaterial const& InMaterial)
 {
-	if (IsAnalyticDerivEnabled())
+	if (A == B)
 	{
-		return DerivativeAutogen.GenerateIfFunc(*this, A, B, AGreaterThanB, AEqualsB, ALessThanB, ThresholdArg);
+		return true;
+	}
+	if (A == -1 || B == -1)
+	{
+		return false;
 	}
 	
-	if(A == INDEX_NONE || B == INDEX_NONE || AGreaterThanB == INDEX_NONE || ALessThanB == INDEX_NONE)
+	FShaderCodeChunk const& ChunkA = InCurrentScopeChunks[A];
+	FShaderCodeChunk const& ChunkB = InCurrentScopeChunks[B];
+
+	if (ChunkA.UniformExpression != nullptr && ChunkB.UniformExpression != nullptr)
+	{
+		if (ChunkA.UniformExpression == ChunkB.UniformExpression)
+		{
+			return true;
+		}
+		if (ChunkA.UniformExpression->IsConstant() && ChunkB.UniformExpression->IsConstant() && ChunkA.Type == ChunkB.Type)
+		{
+			FMaterialRenderContext DummyContext(nullptr, InMaterial, nullptr);
+			FLinearColor ValueA, ValueB;
+			ChunkA.UniformExpression->GetNumberValue(DummyContext, ValueA);
+			ChunkB.UniformExpression->GetNumberValue(DummyContext, ValueB);
+			if (ValueA == ValueB)
+			{
+				return true;
+			}
+		}
+	}
+
+	return false;
+}
+
+int32 FHLSLMaterialTranslator::If(int32 A,int32 B,int32 AGreaterThanB,int32 AEqualsB,int32 ALessThanB,int32 ThresholdArg)
+{
+	if (A == INDEX_NONE || B == INDEX_NONE || AGreaterThanB == INDEX_NONE || ALessThanB == INDEX_NONE)
 	{
 		return INDEX_NONE;
 	}
 
+	if (AreEqualExpressions(AGreaterThanB, ALessThanB, *CurrentScopeChunks, *Material))
+	{
+		if (AEqualsB == INDEX_NONE || AreEqualExpressions(AGreaterThanB, AEqualsB, *CurrentScopeChunks, *Material))
+		{
+			// All inputs are considered equal. So simply return one of them.
+			return AGreaterThanB;
+		}
+	}
+
+	if (IsAnalyticDerivEnabled())
+	{
+		return DerivativeAutogen.GenerateIfFunc(*this, A, B, AGreaterThanB, AEqualsB, ALessThanB, ThresholdArg);
+	}
+	
 	if (AEqualsB != INDEX_NONE)
 	{
 		if (ThresholdArg == INDEX_NONE)
