Avoid needlessly running RebuildRefSkeleton when running object reference collection archives

diff --git a/Engine/Source/Runtime/Engine/Private/Animation/Skeleton.cpp b/Engine/Source/Runtime/Engine/Private/Animation/Skeleton.cpp
--- a/Engine/Source/Runtime/Engine/Private/Animation/Skeleton.cpp
+++ b/Engine/Source/Runtime/Engine/Private/Animation/Skeleton.cpp
@@ -789,8 +789,13 @@ void USkeleton::Serialize( FArchive& Ar )
 		}
 	}
 
-	const bool bRebuildNameMap = false;
-	ReferenceSkeleton.RebuildRefSkeleton(this, bRebuildNameMap);
+	// This is crashing when live coding in debug - the ObjectReferenceCollector 
+	// is multithreaded and so I assume that is causing some subtle problem
+	if (!Ar.IsObjectReferenceCollector())
+	{
+		const bool bRebuildNameMap = false;
+		ReferenceSkeleton.RebuildRefSkeleton(this, bRebuildNameMap);
+	}
 }
 
 #if WITH_EDITOR
