World Partition: fix error introduced during last minute changes during review from 22027429.

#rnx

diff --git a/Engine/Source/Runtime/Engine/Private/WorldPartition/WorldPartitionStreamingPolicy.cpp b/Engine/Source/Runtime/Engine/Private/WorldPartition/WorldPartitionStreamingPolicy.cpp
--- a/Engine/Source/Runtime/Engine/Private/WorldPartition/WorldPartitionStreamingPolicy.cpp
+++ b/Engine/Source/Runtime/Engine/Private/WorldPartition/WorldPartitionStreamingPolicy.cpp
@@ -439,6 +439,7 @@ void UWorldPartitionStreamingPolicy::UpdateStreamingState()
 		
 		SortStreamingCells(StreamingSources, ToActivateCellsUnsorted, ToActivateCells);
 		SortStreamingCells(StreamingSources, ToLoadCellsUnsorted, ToLoadCells);
+		ToUnloadCells = ToUnloadCellsUnsorted.Array();
 	}
 
 	if(World->bMatchStarted)
