Add Telemetry Event to Session Summary to Distinguish UE arch on M1

Fix type being bool not string. Also add to sessionStart

diff --git a/Engine/Source/Runtime/Engine/Private/Analytics/EngineAnalyticsSessionSummary.cpp b/Engine/Source/Runtime/Engine/Private/Analytics/EngineAnalyticsSessionSummary.cpp
--- a/Engine/Source/Runtime/Engine/Private/Analytics/EngineAnalyticsSessionSummary.cpp
+++ b/Engine/Source/Runtime/Engine/Private/Analytics/EngineAnalyticsSessionSummary.cpp
@@ -153,9 +153,9 @@ FEngineAnalyticsSessionSummary::FEngineAnalyticsSessionSummary(TSharedPtr<IAnaly
 	Store->Set(TEXT("Platform"), FString(FPlatformProperties::IniPlatformName()));
 #if PLATFORM_MAC
 #if PLATFORM_MAC_ARM64
-    Store->Set(TEXT("UEBuildArch"), TEXT("AppleSilicon"));
+    Store->Set(TEXT("UEBuildArch"), FString(TEXT("AppleSilicon")));
 #else
-    Store->Set(TEXT("UEBuildArch"), TEXT("Intel(Mac)"));
+    Store->Set(TEXT("UEBuildArch"), FString(TEXT("Intel(Mac)")));
 #endif
 #endif
 	Store->Set(TEXT("OSMajor"), OSMajor);
diff --git a/Engine/Source/Runtime/Engine/Private/EngineAnalytics.cpp b/Engine/Source/Runtime/Engine/Private/EngineAnalytics.cpp
--- a/Engine/Source/Runtime/Engine/Private/EngineAnalytics.cpp
+++ b/Engine/Source/Runtime/Engine/Private/EngineAnalytics.cpp
@@ -139,6 +139,13 @@ void FEngineAnalytics::Initialize()
 			StartSessionAttributes.Emplace(TEXT("OSMinor"), OSMinor);
 			StartSessionAttributes.Emplace(TEXT("OSVersion"), FPlatformMisc::GetOSVersion());
 			StartSessionAttributes.Emplace(TEXT("Is64BitOS"), FPlatformMisc::Is64bitOperatingSystem());
+#if PLATFORM_MAC
+#if PLATFORM_MAC_ARM64
+            StartSessionAttributes.Emplace(TEXT("UEBuildArch"), FString(TEXT("AppleSilicon")));
+#else
+            StartSessionAttributes.Emplace(TEXT("UEBuildArch"), FString(TEXT("Intel(Mac)")));
+#endif
+#endif
 
 			// allow editor events to be correlated to StudioAnalytics events (if there is a studio analytics provider)
 			if (FStudioAnalytics::IsAvailable())
