[Level Actor Packing] Only pack meshes with the same determinant sign together and set ReverseCulling accordingly

#rnx

diff --git a/Engine/Source/Runtime/Engine/Private/ISMPartition/ISMComponentDescriptor.cpp b/Engine/Source/Runtime/Engine/Private/ISMPartition/ISMComponentDescriptor.cpp
--- a/Engine/Source/Runtime/Engine/Private/ISMPartition/ISMComponentDescriptor.cpp
+++ b/Engine/Source/Runtime/Engine/Private/ISMPartition/ISMComponentDescriptor.cpp
@@ -68,6 +68,7 @@ void FISMComponentDescriptor::InitFrom(const UStaticMeshComponent* Template, boo
 	bVisibleInRayTracing = Template->bVisibleInRayTracing;
 	bConsiderForActorPlacementWhenHidden = Template->bConsiderForActorPlacementWhenHidden;
 	bEvaluateWorldPositionOffset = Template->bEvaluateWorldPositionOffset;
+	bIsLocalToWorldDeterminantNegative = Template->GetRenderMatrix().Determinant() < 0;
 
 	if (const UInstancedStaticMeshComponent* ISMTemplate = Cast<UInstancedStaticMeshComponent>(Template))
 	{
@@ -133,6 +134,7 @@ bool FISMComponentDescriptor::operator==(const FISMComponentDescriptor& Other) c
 	bVisibleInRayTracing == Other.bVisibleInRayTracing &&
 	bConsiderForActorPlacementWhenHidden == Other.bConsiderForActorPlacementWhenHidden &&
 	bEvaluateWorldPositionOffset == Other.bEvaluateWorldPositionOffset &&
+	bIsLocalToWorldDeterminantNegative == Other.bIsLocalToWorldDeterminantNegative &&
 	BodyInstance.GetCollisionEnabled() == Other.BodyInstance.GetCollisionEnabled() && 
 	BodyInstance.GetCollisionResponse() == Other.BodyInstance.GetCollisionResponse() &&
 	BodyInstance.DoesUseCollisionProfile() == Other.BodyInstance.DoesUseCollisionProfile() &&
@@ -211,6 +213,7 @@ void FISMComponentDescriptor::InitComponent(UInstancedStaticMeshComponent* ISMCo
 	ISMComponent->bVisibleInRayTracing = bVisibleInRayTracing;
 	ISMComponent->bConsiderForActorPlacementWhenHidden = bConsiderForActorPlacementWhenHidden;
 	ISMComponent->bEvaluateWorldPositionOffset = bEvaluateWorldPositionOffset;
+	ISMComponent->bReverseCulling = bIsLocalToWorldDeterminantNegative;
 
 	// HISM Specific
 	if (UHierarchicalInstancedStaticMeshComponent* HISMComponent = Cast<UHierarchicalInstancedStaticMeshComponent>(ISMComponent))
diff --git a/Engine/Source/Runtime/Engine/Public/ISMPartition/ISMComponentDescriptor.h b/Engine/Source/Runtime/Engine/Public/ISMPartition/ISMComponentDescriptor.h
--- a/Engine/Source/Runtime/Engine/Public/ISMPartition/ISMComponentDescriptor.h
+++ b/Engine/Source/Runtime/Engine/Public/ISMPartition/ISMComponentDescriptor.h
@@ -171,5 +171,8 @@ public:
 
 	UPROPERTY(EditAnywhere, Category = "Component Settings")
 	uint8 bEvaluateWorldPositionOffset : 1;
+
+	UPROPERTY(EditAnywhere, Category = "Component Settings")
+	uint8 bIsLocalToWorldDeterminantNegative : 1;
 #endif
 };
