Skeletal Mesh - Move USkeletalMeshComponent::GetCurrentRefToLocalMatrices() to the USkinnedMeshComponent parent class.

diff --git a/Engine/Source/Runtime/Engine/Classes/Components/SkeletalMeshComponent.h b/Engine/Source/Runtime/Engine/Classes/Components/SkeletalMeshComponent.h
--- a/Engine/Source/Runtime/Engine/Classes/Components/SkeletalMeshComponent.h
+++ b/Engine/Source/Runtime/Engine/Classes/Components/SkeletalMeshComponent.h
@@ -1926,8 +1926,6 @@ public:
 	UE_DEPRECATED(4.27, "Use RegisterOnBoneTransformsFinalizedDelegate/UnregisterOnBoneTransformsFinalizedDelegate instead")
 	FOnBoneTransformsFinalized OnBoneTransformsFinalized;
 
-	void GetCurrentRefToLocalMatrices(TArray<FMatrix44f>& OutRefToLocals, int32 InLodIdx) const;
-
 	// Conditions used to gate when post process events happen
 	bool ShouldUpdatePostProcessInstance() const;
 	bool ShouldPostUpdatePostProcessInstance() const;
diff --git a/Engine/Source/Runtime/Engine/Classes/Components/SkinnedMeshComponent.h b/Engine/Source/Runtime/Engine/Classes/Components/SkinnedMeshComponent.h
--- a/Engine/Source/Runtime/Engine/Classes/Components/SkinnedMeshComponent.h
+++ b/Engine/Source/Runtime/Engine/Classes/Components/SkinnedMeshComponent.h
@@ -1157,6 +1157,9 @@ public:
 	/** Caches the RefToLocal matrices. */
 	void CacheRefToLocalMatrices(TArray<FMatrix44f>& OutRefToLocal) const;
 
+	/** Return the skinning matrices used for rendering. */
+	void GetCurrentRefToLocalMatrices(TArray<FMatrix44f>& OutRefToLocals, int32 InLodIdx) const;
+
 	FORCEINLINE	const USkinnedMeshComponent* GetBaseComponent()const
 	{
 		return LeaderPoseComponent.IsValid() ? LeaderPoseComponent.Get() : this;
diff --git a/Engine/Source/Runtime/Engine/Private/Components/SkeletalMeshComponent.cpp b/Engine/Source/Runtime/Engine/Private/Components/SkeletalMeshComponent.cpp
--- a/Engine/Source/Runtime/Engine/Private/Components/SkeletalMeshComponent.cpp
+++ b/Engine/Source/Runtime/Engine/Private/Components/SkeletalMeshComponent.cpp
@@ -23,7 +23,6 @@
 #include "PhysicsEngine/PhysicsAsset.h"
 #include "Animation/AnimBlueprintGeneratedClass.h"
 #include "Animation/AnimBlueprint.h"
-#include "SkeletalRender.h"
 #include "HAL/LowLevelMemTracker.h"
 #include "SkinnedAssetCompiler.h"
 #include "Rendering/SkeletalMeshRenderData.h"
@@ -4159,30 +4158,6 @@ void USkeletalMeshComponent::FinalizeBoneTransform()
 	TRACE_SKELETAL_MESH_COMPONENT(this);
 }
 
-void USkeletalMeshComponent::GetCurrentRefToLocalMatrices(TArray<FMatrix44f>& OutRefToLocals, int32 InLodIdx) const
-{
-	USkeletalMesh* SkelMesh = GetSkeletalMeshAsset();
-	if(SkelMesh)
-	{
-		FSkeletalMeshRenderData* RenderData = SkelMesh->GetResourceForRendering();
-		if (ensureMsgf(RenderData->LODRenderData.IsValidIndex(InLodIdx),
-			TEXT("GetCurrentRefToLocalMatrices (SkelMesh :%s) input LODIndex (%d) doesn't match with render data size (%d)."),
-			*SkelMesh->GetPathName(), InLodIdx, RenderData->LODRenderData.Num()))
-		{
-			UpdateRefToLocalMatrices(OutRefToLocals, this, RenderData, InLodIdx, nullptr);
-		}
-		else
-		{
-			const FReferenceSkeleton& RefSkeleton = SkelMesh->GetRefSkeleton();
-			OutRefToLocals.AddUninitialized(RefSkeleton.GetNum());
-			for (int32 Index = 0; Index < OutRefToLocals.Num(); ++Index)
-			{
-				OutRefToLocals[Index] = FMatrix44f::Identity;
-			}
-		}
-	}
-}
-
 bool USkeletalMeshComponent::ShouldUpdatePostProcessInstance() const
 {
 	return PostProcessAnimInstance && !bDisablePostProcessBlueprint;
diff --git a/Engine/Source/Runtime/Engine/Private/Components/SkinnedMeshComponent.cpp b/Engine/Source/Runtime/Engine/Private/Components/SkinnedMeshComponent.cpp
--- a/Engine/Source/Runtime/Engine/Private/Components/SkinnedMeshComponent.cpp
+++ b/Engine/Source/Runtime/Engine/Private/Components/SkinnedMeshComponent.cpp
@@ -30,6 +30,7 @@
 #include "AnimationRuntime.h"
 #include "Animation/SkinWeightProfileManager.h"
 #include "GPUSkinCache.h"
+#include "SkeletalRender.h"
 
 DEFINE_LOG_CATEGORY_STATIC(LogSkinnedMeshComp, Log, All);
 
@@ -3111,6 +3112,30 @@ void USkinnedMeshComponent::CacheRefToLocalMatrices(TArray<FMatrix44f>& OutRefTo
 	}
 }
 
+void USkinnedMeshComponent::GetCurrentRefToLocalMatrices(TArray<FMatrix44f>& OutRefToLocals, int32 InLodIdx) const
+{
+	if (const USkinnedAsset* const Asset = GetSkinnedAsset())
+	{
+		const FSkeletalMeshRenderData* const RenderData = Asset->GetResourceForRendering();
+		
+		if (ensureMsgf(RenderData->LODRenderData.IsValidIndex(InLodIdx),
+			TEXT("GetCurrentRefToLocalMatrices (SkelMesh :%s) input LODIndex (%d) doesn't match with render data size (%d)."),
+			*Asset->GetPathName(), InLodIdx, RenderData->LODRenderData.Num()))
+		{
+			UpdateRefToLocalMatrices(OutRefToLocals, this, RenderData, InLodIdx, nullptr);
+		}
+		else
+		{
+			const FReferenceSkeleton& RefSkeleton = Asset->GetRefSkeleton();
+			OutRefToLocals.AddUninitialized(RefSkeleton.GetNum());
+			for (int32 Index = 0; Index < OutRefToLocals.Num(); ++Index)
+			{
+				OutRefToLocals[Index] = FMatrix44f::Identity;
+			}
+		}
+	}
+}
+
 void USkinnedMeshComponent::ComputeSkinnedPositions(USkinnedMeshComponent* Component, TArray<FVector3f> & OutPositions, TArray<FMatrix44f>& CachedRefToLocals, const FSkeletalMeshLODRenderData& LODData, const FSkinWeightVertexBuffer& SkinWeightBuffer)
 {
 	OutPositions.Empty();
