Crash fix when drawing pose watches on skeletal mesh components
* No longer dereference pose watch UObject while in render thread, render data is instead mirrored before drawing on the render thread to prevent being modified in a worker thread.
* Cleanup of pose watch WITH_ENGINE ifdefs.

diff --git a/Engine/Plugins/Animation/GameplayInsights/Source/GameplayInsights/Private/PoseWatchTrack.cpp b/Engine/Plugins/Animation/GameplayInsights/Source/GameplayInsights/Private/PoseWatchTrack.cpp
--- a/Engine/Plugins/Animation/GameplayInsights/Source/GameplayInsights/Private/PoseWatchTrack.cpp
+++ b/Engine/Plugins/Animation/GameplayInsights/Source/GameplayInsights/Private/PoseWatchTrack.cpp
@@ -13,6 +13,7 @@
 #include "Animation/AnimBlueprint.h"
 #include "Animation/AnimBlueprintGeneratedClass.h"
 #include "Animation/AnimInstance.h"
+#include "Engine/PoseWatchRenderData.h"
 #include "Engine/PoseWatch.h"
 #include "ObjectTrace.h"
 
@@ -41,7 +42,7 @@ FPoseWatchTrack::FPoseWatchTrack(uint64 InObjectId, const FPoseWatchTrack::FPose
 				{
 					if (PoseWatch.NodeID == PostWatchTrackId.NameId)
 					{
-						PoseWatchOwner = PoseWatch.PoseWatchPoseElement.Get();
+						PoseWatchOwner = PoseWatch.PoseWatchPoseElement;
 						break;
 					}
 				}
diff --git a/Engine/Plugins/Animation/GameplayInsights/Source/RewindDebugger/Private/RewindDebugger.cpp b/Engine/Plugins/Animation/GameplayInsights/Source/RewindDebugger/Private/RewindDebugger.cpp
--- a/Engine/Plugins/Animation/GameplayInsights/Source/RewindDebugger/Private/RewindDebugger.cpp
+++ b/Engine/Plugins/Animation/GameplayInsights/Source/RewindDebugger/Private/RewindDebugger.cpp
@@ -812,13 +812,7 @@ void FRewindDebugger::Tick(float DeltaTime)
 															DebugData.ResetNodeVisitSites();
 														}
 
-														for (FAnimNodePoseWatch& PoseWatchNode : DebugData.AnimNodePoseWatch)
-														{
-															if (PoseWatchNode.PoseWatch.IsValid())
-															{
-																PoseWatchNode.PoseWatch->SetIsNodeEnabled(false);
-															}
-														}
+														DebugData.DisableAllPoseWatches();
 							
 														AnimGraphTimeline.EnumerateEvents(Frame.StartTime, Frame.EndTime, [Id, AnimationProvider, GameplayProvider, &DebugData, NodeCount](double InGraphStartTime, double InGraphEndTime, uint32 InDepth, const FAnimGraphMessage& InMessage)
 														{
@@ -917,6 +911,7 @@ void FRewindDebugger::Tick(float DeltaTime)
 																		{
 																			InPoseWatchTimeline.EnumerateEvents(InGraphStartTime, InGraphEndTime, [AnimationProvider, &DebugData](double InStartTime, double InEndTime, uint32 InDepth, const FPoseWatchMessage& InMessage)
 																				{
+#if WITH_EDITOR
 																					for (FAnimNodePoseWatch& PoseWatch : DebugData.AnimNodePoseWatch)
 																					{
 																						if (PoseWatch.NodeID == InMessage.PoseWatchId)
@@ -932,6 +927,7 @@ void FRewindDebugger::Tick(float DeltaTime)
 																							break;
 																						}
 																					}
+#endif //WITH_EDITOR
 																					return TraceServices::EEventEnumerate::Continue;
 																				});
 																		});
diff --git a/Engine/Source/Editor/Persona/Private/AnimationEditorViewportClient.cpp b/Engine/Source/Editor/Persona/Private/AnimationEditorViewportClient.cpp
--- a/Engine/Source/Editor/Persona/Private/AnimationEditorViewportClient.cpp
+++ b/Engine/Source/Editor/Persona/Private/AnimationEditorViewportClient.cpp
@@ -41,7 +41,7 @@
 #include "Animation/MorphTarget.h"
 #include "Rendering/SkeletalMeshModel.h"
 #include "UnrealWidget.h"
-#include "Engine/PoseWatch.h"
+#include "Engine/PoseWatchRenderData.h"
 
 namespace {
 	static const float AnimationEditorViewport_RotateSpeed = 0.02f;
@@ -1582,19 +1582,16 @@ void FAnimationViewportClient::DrawWatchedPoses(UDebugSkelMeshComponent * MeshCo
 {
 	if (UAnimBlueprintGeneratedClass* AnimBPGenClass = Cast<UAnimBlueprintGeneratedClass>(MeshComponent->AnimClass))
 	{
-		USkeleton* TargetSkeleton = AnimBPGenClass->GetTargetSkeleton();
-		if (TargetSkeleton)
+		if (UAnimBlueprint* Blueprint = Cast<UAnimBlueprint>(AnimBPGenClass->ClassGeneratedBy))
 		{
-			if (UAnimBlueprint* Blueprint = Cast<UAnimBlueprint>(AnimBPGenClass->ClassGeneratedBy))
+			if (Blueprint->GetObjectBeingDebugged() && MeshComponent)
 			{
-				if (Blueprint->GetObjectBeingDebugged() && MeshComponent)
+				FAnimBlueprintDebugData& DebugData = AnimBPGenClass->GetAnimBlueprintDebugData();
+				DebugData.ForEachActiveVisiblePoseWatchPoseElement([PDI, MeshComponent](FAnimNodePoseWatch& PoseWatch)
 				{
-					const FAnimBlueprintDebugData& DebugData = AnimBPGenClass->GetAnimBlueprintDebugData();
-					DebugData.ForEachActiveVisiblePoseWatchPoseElement([PDI, MeshComponent](const FAnimNodePoseWatch& PoseWatch)
-					{
-						SkeletalDebugRendering::DrawBonesFromPoseWatch(PDI, PoseWatch, MeshComponent->GetReferenceSkeleton(), /*bUseWorldTransform*/false);
-					});
-				}
+					PoseWatch.CopyPoseWatchData(MeshComponent->GetReferenceSkeleton());
+					SkeletalDebugRendering::DrawBonesFromPoseWatch(PDI, PoseWatch, /*bUseWorldTransform*/false);
+				});
 			}
 		}
 	}
diff --git a/Engine/Source/Runtime/Engine/Classes/Animation/AnimBlueprintGeneratedClass.h b/Engine/Source/Runtime/Engine/Classes/Animation/AnimBlueprintGeneratedClass.h
--- a/Engine/Source/Runtime/Engine/Classes/Animation/AnimBlueprintGeneratedClass.h
+++ b/Engine/Source/Runtime/Engine/Classes/Animation/AnimBlueprintGeneratedClass.h
@@ -12,7 +12,7 @@
 #include "Animation/AnimNodeBase.h"
 #include "Animation/BlendSpace.h"
 #include "Animation/ExposedValueHandler.h"
-#include "Engine/PoseWatch.h"
+#include "Engine/PoseWatchRenderData.h"
 
 #include "AnimBlueprintGeneratedClass.generated.h"
 
@@ -314,7 +314,8 @@ public:
 
 	void AddPoseWatch(int32 NodeID, UPoseWatchPoseElement* const InPoseWatchPoseElement);
 	void RemovePoseWatch(int32 NodeID);
-	void ForEachActiveVisiblePoseWatchPoseElement(const TFunctionRef<void(const FAnimNodePoseWatch&)>& InFunction) const;
+	void ForEachActiveVisiblePoseWatchPoseElement(const TFunctionRef<void(FAnimNodePoseWatch&)>& InFunction);
+	void DisableAllPoseWatches();
 
 	TArrayView<const FName> GetNodeAttributes(TWeakObjectPtr<UAnimGraphNode_Base> InAnimGraphNode) const;
 #endif
diff --git a/Engine/Source/Runtime/Engine/Classes/Components/SkeletalMeshComponent.h b/Engine/Source/Runtime/Engine/Classes/Components/SkeletalMeshComponent.h
--- a/Engine/Source/Runtime/Engine/Classes/Components/SkeletalMeshComponent.h
+++ b/Engine/Source/Runtime/Engine/Classes/Components/SkeletalMeshComponent.h
@@ -26,6 +26,9 @@
 #include "ClothCollisionPrim.h"
 #include "PhysicsEngine/PhysicsAsset.h"
 #include "Animation/AttributesRuntime.h"
+#if WITH_ENGINE
+#include "Engine/PoseWatchRenderData.h"
+#endif
 
 #include "SkeletalMeshComponent.generated.h"
 
@@ -1613,7 +1616,8 @@ public:
 	void DebugDrawClothingTexts(FCanvas* Canvas, const FSceneView* SceneView);
 
 #if WITH_EDITOR
-	void DebugDrawPoseWatches(FPrimitiveDrawInterface* PDI);
+	void UpdatePoseWatches();
+	TArray<FAnimNodePoseWatch> PoseWatches;
 #endif
 
 	/** Changes the value of bNotifyRigidBodyCollision
@@ -1686,6 +1690,7 @@ protected:
 	virtual bool ShouldCreatePhysicsState() const override;
 	virtual void OnCreatePhysicsState() override;
 	virtual void OnDestroyPhysicsState() override;
+	virtual void SendRenderDynamicData_Concurrent() override;
 	virtual void RegisterComponentTickFunctions(bool bRegister) override;
 public:
 	virtual void InitializeComponent() override;
@@ -2614,6 +2619,17 @@ public:
 	bool IsPostEvaluatingAnimation() const { return bPostEvaluatingAnimation; }
 };
 
+#if WITH_EDITOR
+struct FPoseWatchDynamicData
+{
+public:
+
+	FPoseWatchDynamicData(USkeletalMeshComponent* InComponent);
+
+	TArray<FAnimNodePoseWatch> PoseWatches;
+};
+#endif
+
 struct FLinkedInstancesAdapter
 {
 	static void AddLinkedInstance(USkeletalMeshComponent* InComponent, UAnimInstance* InAnimInstance)
diff --git a/Engine/Source/Runtime/Engine/Classes/Engine/PoseWatch.h b/Engine/Source/Runtime/Engine/Classes/Engine/PoseWatch.h
--- a/Engine/Source/Runtime/Engine/Classes/Engine/PoseWatch.h
+++ b/Engine/Source/Runtime/Engine/Classes/Engine/PoseWatch.h
@@ -8,60 +8,11 @@
 #include "BoneIndices.h"
 #include "PoseWatch.generated.h"
 
-
 struct FCompactHeapPose;
 class UAnimBlueprint;
 class UBlendProfile;
 struct FSlateIcon;
 
-struct FAnimNodePoseWatch
-{
-public:
-	// Object (anim instance) that this pose came from
-	TWeakObjectPtr<const UObject>			Object;
-	TWeakObjectPtr<UPoseWatch>				PoseWatch;
-	TWeakObjectPtr<UPoseWatchPoseElement>	PoseWatchPoseElement;
-	int32									NodeID;
-
-	template<class TAllocator>
-	bool SetPose(const TArray<FBoneIndexType>& InRequiredBones, const TArray<FTransform, TAllocator>& InBoneTransforms)
-	{
-		RequiredBones = InRequiredBones;
-		BoneTransforms = InBoneTransforms;
-		return true;
-	}
-
-	void SetWorldTransform(const FTransform& InWorldTransform)
-	{
-		WorldTransform = InWorldTransform;
-	}
-
-	const TArray<FBoneIndexType>& GetRequiredBones() const
-	{
-		return RequiredBones;
-	}
-
-	const TArray<FTransform>& GetBoneTransforms() const
-	{
-		return BoneTransforms;
-	}
-
-	const FTransform& GetWorldTransform() const
-	{
-		return WorldTransform;
-	}
-
-	bool IsValid() const
-	{
-		return Object.IsValid() && PoseWatch.IsValid();
-	}
-
-private:
-	FTransform WorldTransform;
-	TArray<FBoneIndexType> RequiredBones;
-	TArray<FTransform> BoneTransforms;
-};
-
 namespace PoseWatchUtil
 {
 	/** Gets all pose watches that are parented to Folder, if Folder is nullptr then gets orphans */
diff --git a/Engine/Source/Runtime/Engine/Classes/Engine/PoseWatchRenderData.h b/Engine/Source/Runtime/Engine/Classes/Engine/PoseWatchRenderData.h
new file mode 100644
--- /dev/null
+++ b/Engine/Source/Runtime/Engine/Classes/Engine/PoseWatchRenderData.h
@@ -0,0 +1,69 @@
+// Copyright Epic Games, Inc. All Rights Reserved.
+
+#pragma once
+
+#if WITH_EDITOR
+
+#include "CoreMinimal.h"
+#include "UObject/ObjectMacros.h"
+#include "UObject/Object.h"
+#include "BoneIndices.h"
+
+struct FReferenceSkeleton;
+
+struct ENGINE_API FAnimNodePoseWatch
+{
+public:
+	FAnimNodePoseWatch();
+
+	// Object (anim instance) that this pose came from
+	const UObject* Object;
+	class UPoseWatch* PoseWatch;
+	class UPoseWatchPoseElement* PoseWatchPoseElement;
+	int32 NodeID;
+
+	bool IsValid() const;
+
+	template<class TAllocator>
+	bool SetPose(const TArray<FBoneIndexType>& InRequiredBones, const TArray<FTransform, TAllocator>& InBoneTransforms)
+	{
+		RequiredBones = InRequiredBones;
+		BoneTransforms = InBoneTransforms;
+		return true;
+	}
+
+	void SetWorldTransform(const FTransform& InWorldTransform);
+
+	/**
+	 * Take a snapshot of the pose watch properties that this struct represents
+	 * to be used when drawing the debug skeleton
+	 */
+	void CopyPoseWatchData(const FReferenceSkeleton& RefSkeleton);
+
+	const TArray<FBoneIndexType>& GetRequiredBones() const;
+
+	const TArray<FTransform>& GetBoneTransforms() const;
+
+	const FTransform& GetWorldTransform() const;
+
+	FLinearColor GetBoneColor() const;
+
+	FVector GetViewportOffset() const;
+
+	TArray<int32> GetViewportAllowList() const;
+
+	TArray<int32> GetParentIndices() const;
+
+private:
+	FTransform WorldTransform;
+	TArray<FBoneIndexType> RequiredBones;
+	TArray<FTransform> BoneTransforms;
+
+	// Mirrored properties updated on CopyPoseWatchData
+	FLinearColor BoneColor;
+	FVector ViewportOffset;
+	TArray<int32> ViewportMaskAllowedList;
+	TArray<int32> ParentIndices;
+};
+
+#endif
\ No newline at end of file
diff --git a/Engine/Source/Runtime/Engine/Private/Animation/AnimBlueprintGeneratedClass.cpp b/Engine/Source/Runtime/Engine/Private/Animation/AnimBlueprintGeneratedClass.cpp
--- a/Engine/Source/Runtime/Engine/Private/Animation/AnimBlueprintGeneratedClass.cpp
+++ b/Engine/Source/Runtime/Engine/Private/Animation/AnimBlueprintGeneratedClass.cpp
@@ -262,11 +262,11 @@ void FAnimBlueprintDebugData::RemovePoseWatch(int32 NodeID)
 	}
 }
 
-void FAnimBlueprintDebugData::ForEachActiveVisiblePoseWatchPoseElement(const TFunctionRef<void(const FAnimNodePoseWatch&)>& InFunction) const
+void FAnimBlueprintDebugData::ForEachActiveVisiblePoseWatchPoseElement(const TFunctionRef<void(FAnimNodePoseWatch&)>& InFunction)
 {
-	for (const FAnimNodePoseWatch& PoseWatchNode : AnimNodePoseWatch)
+	for (FAnimNodePoseWatch& PoseWatchNode : AnimNodePoseWatch)
 	{
-		if (const UPoseWatchPoseElement* PoseWatchPoseElement = PoseWatchNode.PoseWatchPoseElement.Get())
+		if (const UPoseWatchPoseElement* PoseWatchPoseElement = PoseWatchNode.PoseWatchPoseElement)
 		{
 			if (PoseWatchPoseElement->GetIsEnabled() && PoseWatchPoseElement->GetIsVisible())
 			{
@@ -276,6 +276,17 @@ void FAnimBlueprintDebugData::ForEachActiveVisiblePoseWatchPoseElement(const TFu
 	}
 }
 
+void FAnimBlueprintDebugData::DisableAllPoseWatches()
+{
+	for (FAnimNodePoseWatch& PoseWatchNode : AnimNodePoseWatch)
+	{
+		if (PoseWatchNode.PoseWatch)
+		{
+			PoseWatchNode.PoseWatch->SetIsNodeEnabled(false);
+		}
+	}
+}
+
 TArrayView<const FName> FAnimBlueprintDebugData::GetNodeAttributes(TWeakObjectPtr<UAnimGraphNode_Base> InAnimGraphNode) const
 {
 	const TArray<FName>* Attributes = NodeAttributes.Find(InAnimGraphNode);
diff --git a/Engine/Source/Runtime/Engine/Private/Animation/AnimInstanceProxy.cpp b/Engine/Source/Runtime/Engine/Private/Animation/AnimInstanceProxy.cpp
--- a/Engine/Source/Runtime/Engine/Private/Animation/AnimInstanceProxy.cpp
+++ b/Engine/Source/Runtime/Engine/Private/Animation/AnimInstanceProxy.cpp
@@ -28,6 +28,10 @@
 #include "Animation/AnimBlueprintGeneratedClass.h"
 #include "Animation/AnimStateMachineTypes.h"
 #include "Animation/AnimTrace.h"
+#if WITH_EDITOR
+#include "Engine/PoseWatchRenderData.h"
+#include "Engine/PoseWatch.h"
+#endif
 
 #define DO_ANIMSTAT_PROCESSING(StatName) DEFINE_STAT(STAT_ ## StatName)
 #include "Animation/AnimMTStats.h"
@@ -496,11 +500,7 @@ void FAnimInstanceProxy::PreUpdate(UAnimInstance* InAnimInstance, float DeltaSec
 		FAnimBlueprintDebugData* DebugData = AnimBP->GetDebugData();
 		if (DebugData)
 		{
-			PoseWatchEntriesForThisFrame = DebugData->AnimNodePoseWatch;
-			for (FAnimNodePoseWatch& PoseWatch : PoseWatchEntriesForThisFrame)
-			{
-				PoseWatch.PoseWatch.Get()->SetIsNodeEnabled(false);
-			}
+			DebugData->DisableAllPoseWatches();
 		}
 	}
 #endif
@@ -576,18 +576,16 @@ void FAnimInstanceProxy::PostUpdate(UAnimInstance* InAnimInstance) const
 		DebugData->RecordNodeVisitArray(UpdatedNodesThisFrame);
 		DebugData->RecordNodeAttributeMaps(NodeInputAttributesThisFrame, NodeOutputAttributesThisFrame);
 		DebugData->RecordNodeSyncsArray(NodeSyncsThisFrame);
-		DebugData->AnimNodePoseWatch = PoseWatchEntriesForThisFrame;
 	}
 #endif
 
 #if WITH_EDITORONLY_DATA
-	UAnimBlueprint* AnimBP = GetAnimBlueprint();
 	if (bIsBeingDebugged)
 	{
-		FAnimBlueprintDebugData* DebugData = AnimBP->GetDebugData();
+		FAnimBlueprintDebugData* DebugData = GetAnimBlueprintDebugData();
 		if (DebugData)
 		{
-			for (const FAnimNodePoseWatch& PoseWatch : PoseWatchEntriesForThisFrame)
+			for (const FAnimNodePoseWatch& PoseWatch : DebugData->AnimNodePoseWatch)
 			{
 				if (!PoseWatch.PoseWatchPoseElement->GetIsEnabled())
 				{
@@ -3181,22 +3179,26 @@ void FAnimInstanceProxy::RegisterWatchedPose(const FCompactPose& Pose, int32 Lin
 {
 	if (bIsBeingDebugged)
 	{
-		if (USkeletalMeshComponent* SkelMeshComponent = GetSkelMeshComponent())
+		FAnimBlueprintDebugData* DebugData = GetAnimBlueprintDebugData();
+		if (DebugData)
 		{
-			for (FAnimNodePoseWatch& PoseWatch : PoseWatchEntriesForThisFrame)
+			if (USkeletalMeshComponent* SkelMeshComponent = GetSkelMeshComponent())
 			{
-				if (PoseWatch.PoseWatch.IsValid() && PoseWatch.NodeID == LinkID)
+				for (FAnimNodePoseWatch& PoseWatch : DebugData->AnimNodePoseWatch)
 				{
-					PoseWatch.Object = GetAnimInstanceObject();
-					PoseWatch.PoseWatch->SetIsNodeEnabled(true);
+					if (PoseWatch.PoseWatch && PoseWatch.NodeID == LinkID)
+					{
+						PoseWatch.Object = GetAnimInstanceObject();
+						PoseWatch.PoseWatch->SetIsNodeEnabled(true);
 
-					const TArray<FTransform, FAnimStackAllocator>& BoneTransforms = Pose.GetBones();
-					const TArray<FBoneIndexType>& TmpRequiredBones = Pose.GetBoneContainer().GetBoneIndicesArray();
-					PoseWatch.SetPose(TmpRequiredBones, BoneTransforms);
-					PoseWatch.SetWorldTransform(SkelMeshComponent->GetComponentTransform());
+						const TArray<FTransform, FAnimStackAllocator>& BoneTransforms = Pose.GetBones();
+						const TArray<FBoneIndexType>& TmpRequiredBones = Pose.GetBoneContainer().GetBoneIndicesArray();
+						PoseWatch.SetPose(TmpRequiredBones, BoneTransforms);
+						PoseWatch.SetWorldTransform(SkelMeshComponent->GetComponentTransform());
 
-					TRACE_ANIM_POSE_WATCH(*this, PoseWatch.NodeID, PoseWatch.GetBoneTransforms(), PoseWatch.GetRequiredBones(), PoseWatch.GetWorldTransform(), true);
-					break;
+						TRACE_ANIM_POSE_WATCH(*this, PoseWatch.NodeID, PoseWatch.GetBoneTransforms(), PoseWatch.GetRequiredBones(), PoseWatch.GetWorldTransform(), true);
+						break;
+					}
 				}
 			}
 		}
@@ -3207,24 +3209,28 @@ void FAnimInstanceProxy::RegisterWatchedPose(const FCSPose<FCompactPose>& Pose,
 {
 	if (bIsBeingDebugged)
 	{
-		if (USkeletalMeshComponent* SkelMeshComponent = GetSkelMeshComponent())
+		FAnimBlueprintDebugData* DebugData = GetAnimBlueprintDebugData();
+		if (DebugData)
 		{
-			for (FAnimNodePoseWatch& PoseWatch : PoseWatchEntriesForThisFrame)
+			if (USkeletalMeshComponent* SkelMeshComponent = GetSkelMeshComponent())
 			{
-				if (PoseWatch.PoseWatch.IsValid() && PoseWatch.NodeID == LinkID)
+				for (FAnimNodePoseWatch& PoseWatch : DebugData->AnimNodePoseWatch)
 				{
-					FCompactPose TempPose;
-					FCSPose<FCompactPose>::ConvertComponentPosesToLocalPoses(Pose, TempPose);
-					PoseWatch.Object = GetAnimInstanceObject();
-					PoseWatch.PoseWatch->SetIsNodeEnabled(true);
-
-					const TArray<FTransform, FAnimStackAllocator>& BoneTransforms = TempPose.GetBones();
-					const TArray<FBoneIndexType>& TmpRequiredBones = TempPose.GetBoneContainer().GetBoneIndicesArray();
-					PoseWatch.SetPose(TmpRequiredBones, BoneTransforms);
-					PoseWatch.SetWorldTransform(SkelMeshComponent->GetComponentTransform());
-
-					TRACE_ANIM_POSE_WATCH(*this, PoseWatch.NodeID, PoseWatch.GetBoneTransforms(), PoseWatch.GetRequiredBones(), PoseWatch.GetWorldTransform(), true);
-					break;
+					if (PoseWatch.PoseWatch && PoseWatch.NodeID == LinkID)
+					{
+						FCompactPose TempPose;
+						FCSPose<FCompactPose>::ConvertComponentPosesToLocalPoses(Pose, TempPose);
+						PoseWatch.Object = GetAnimInstanceObject();
+						PoseWatch.PoseWatch->SetIsNodeEnabled(true);
+
+						const TArray<FTransform, FAnimStackAllocator>& BoneTransforms = TempPose.GetBones();
+						const TArray<FBoneIndexType>& TmpRequiredBones = TempPose.GetBoneContainer().GetBoneIndicesArray();
+						PoseWatch.SetPose(TmpRequiredBones, BoneTransforms);
+						PoseWatch.SetWorldTransform(SkelMeshComponent->GetComponentTransform());
+
+						TRACE_ANIM_POSE_WATCH(*this, PoseWatch.NodeID, PoseWatch.GetBoneTransforms(), PoseWatch.GetRequiredBones(), PoseWatch.GetWorldTransform(), true);
+						break;
+					}
 				}
 			}
 		}
diff --git a/Engine/Source/Runtime/Engine/Private/Components/SkeletalMeshComponent.cpp b/Engine/Source/Runtime/Engine/Private/Components/SkeletalMeshComponent.cpp
--- a/Engine/Source/Runtime/Engine/Private/Components/SkeletalMeshComponent.cpp
+++ b/Engine/Source/Runtime/Engine/Private/Components/SkeletalMeshComponent.cpp
@@ -45,6 +45,10 @@
 #include "HAL/LowLevelMemTracker.h"
 #include "HAL/LowLevelMemStats.h"
 #include "HAL/IConsoleManager.h"
+#if WITH_EDITOR
+#include "Engine/PoseWatchRenderData.h"
+#include "Engine/PoseWatch.h"
+#endif
 
 LLM_DEFINE_TAG(SkeletalMesh_TransformData);
 
@@ -3566,6 +3570,22 @@ FRootMotionMovementParams USkeletalMeshComponent::ConsumeRootMotion()
 	return ConsumeRootMotion_Internal(InterpAlpha);
 }
 
+#if WITH_EDITOR
+FPoseWatchDynamicData::FPoseWatchDynamicData(USkeletalMeshComponent* InComponent)
+{
+	check(InComponent);
+	if (InComponent->GetSkeletalMeshAsset())
+	{
+		PoseWatches = InComponent->PoseWatches;
+
+		for (FAnimNodePoseWatch& PoseWatch : PoseWatches)
+		{
+			PoseWatch.CopyPoseWatchData(InComponent->GetSkeletalMeshAsset()->GetRefSkeleton());
+		}
+	}
+}
+#endif
+
 FRootMotionMovementParams USkeletalMeshComponent::ConsumeRootMotion_Internal(float InAlpha)
 {
 	FRootMotionMovementParams RootMotion;
@@ -3901,6 +3921,33 @@ void USkeletalMeshComponent::RefreshMorphTargets()
 	UpdateMorphTargetOverrideCurves();
 }
 
+void USkeletalMeshComponent::SendRenderDynamicData_Concurrent()
+{
+	Super::SendRenderDynamicData_Concurrent();
+
+#if WITH_EDITOR
+	if (SceneProxy)
+	{
+		UpdatePoseWatches();
+		FPoseWatchDynamicData* NewDynamicData = new FPoseWatchDynamicData(this);
+
+		FSkeletalMeshSceneProxy* TargetProxy = (FSkeletalMeshSceneProxy*)SceneProxy;
+
+		ENQUEUE_RENDER_COMMAND(PoseWatchDynamicDataCommand)(
+			[TargetProxy, NewDynamicData](FRHICommandListImmediate& RHICommandList)
+			{
+				if (TargetProxy->PoseWatchDynamicData)
+				{
+					delete TargetProxy->PoseWatchDynamicData;
+				}
+
+				TargetProxy->PoseWatchDynamicData = NewDynamicData;
+			}
+		);
+	}
+#endif
+}
+
 void USkeletalMeshComponent::ParallelAnimationEvaluation() 
 {
 	if (AnimEvaluationContext.bDoInterpolation)
@@ -4566,8 +4613,9 @@ TArray<FTransform> USkeletalMeshComponent::GetBoneSpaceTransforms()
 }
 
 #if WITH_EDITOR
-void USkeletalMeshComponent::DebugDrawPoseWatches(FPrimitiveDrawInterface* PDI)
+void USkeletalMeshComponent::UpdatePoseWatches()
 {
+	PoseWatches.Empty();
 	UAnimInstance* AnimInstance = GetAnimInstance();
 	if (AnimInstance)
 	{
@@ -4578,11 +4626,17 @@ void USkeletalMeshComponent::DebugDrawPoseWatches(FPrimitiveDrawInterface* PDI)
 			{
 				if (USkeletalMesh* TmpSkeletalMesh = GetSkeletalMeshAsset())
 				{
-					const FAnimBlueprintDebugData& DebugData = AnimBPGenClass->GetAnimBlueprintDebugData();
-					DebugData.ForEachActiveVisiblePoseWatchPoseElement([PDI, TmpSkeletalMesh](const FAnimNodePoseWatch& PoseWatch)
+					FAnimBlueprintDebugData& DebugData = AnimBPGenClass->GetAnimBlueprintDebugData();
+					for (FAnimNodePoseWatch& PoseWatch : DebugData.AnimNodePoseWatch)
 					{
-						SkeletalDebugRendering::DrawBonesFromPoseWatch(PDI, PoseWatch, TmpSkeletalMesh->GetRefSkeleton(), /*bUseWorldTransform*/true);
-					});
+						if (const UPoseWatchPoseElement* PoseWatchPoseElement = PoseWatch.PoseWatchPoseElement)
+						{
+							if (PoseWatchPoseElement->GetIsEnabled() && PoseWatchPoseElement->GetIsVisible())
+							{
+								PoseWatches.Add(PoseWatch); // Copy
+							}
+						}
+					}
 				}
 			}
 		}
diff --git a/Engine/Source/Runtime/Engine/Private/Engine/PoseWatchRenderData.cpp b/Engine/Source/Runtime/Engine/Private/Engine/PoseWatchRenderData.cpp
new file mode 100644
--- /dev/null
+++ b/Engine/Source/Runtime/Engine/Private/Engine/PoseWatchRenderData.cpp
@@ -0,0 +1,92 @@
+// Copyright Epic Games, Inc. All Rights Reserved.
+
+#if WITH_EDITOR
+
+#include "Engine/PoseWatchRenderData.h"
+#include "Animation/BlendProfile.h"
+#include "Engine/PoseWatch.h"
+
+FAnimNodePoseWatch::FAnimNodePoseWatch()
+{
+	Object = nullptr;
+	PoseWatch = nullptr;
+	PoseWatchPoseElement = nullptr;
+	NodeID = INDEX_NONE;
+	WorldTransform = FTransform::Identity;
+	BoneColor = FLinearColor::White;
+	ViewportOffset = FVector::ZeroVector;
+}
+
+void FAnimNodePoseWatch::SetWorldTransform(const FTransform& InWorldTransform)
+{
+	WorldTransform = InWorldTransform;
+}
+
+void FAnimNodePoseWatch::CopyPoseWatchData(const FReferenceSkeleton& RefSkeleton)
+{
+	BoneColor = FLinearColor(PoseWatchPoseElement->GetColor());
+	ViewportOffset = PoseWatchPoseElement->ViewportOffset;
+	ParentIndices.Empty();
+	ViewportMaskAllowedList.Empty();
+
+	const UBlendProfile* Mask = PoseWatchPoseElement->ViewportMask;
+	for (const FBoneIndexType& BoneIndex : RequiredBones)
+	{
+		ParentIndices.Add(RefSkeleton.GetParentIndex(BoneIndex));
+
+		if (Mask)
+		{
+			const FName& BoneName = RefSkeleton.GetBoneName(BoneIndex);
+			const bool bBoneBlendScaleMeetsThreshold = Mask->GetBoneBlendScale(BoneName) > PoseWatchPoseElement->BlendScaleThreshold;
+			const bool bBoneRequired = (!PoseWatchPoseElement->bInvertViewportMask && bBoneBlendScaleMeetsThreshold) || (PoseWatchPoseElement->bInvertViewportMask && !bBoneBlendScaleMeetsThreshold);
+			if (!bBoneRequired)
+			{
+				continue;
+			}
+		}
+
+		ViewportMaskAllowedList.Add(BoneIndex);
+	}
+}
+
+const TArray<FBoneIndexType>& FAnimNodePoseWatch::GetRequiredBones() const
+{
+	return RequiredBones;
+}
+
+const TArray<FTransform>& FAnimNodePoseWatch::GetBoneTransforms() const
+{
+	return BoneTransforms;
+}
+
+const FTransform& FAnimNodePoseWatch::GetWorldTransform() const
+{
+	return WorldTransform;
+}
+
+bool FAnimNodePoseWatch::IsValid() const
+{
+	return Object != nullptr && PoseWatch != nullptr;
+}
+
+FLinearColor FAnimNodePoseWatch::GetBoneColor() const
+{ 
+	return BoneColor;
+}
+
+FVector FAnimNodePoseWatch::GetViewportOffset() const
+{ 
+	return ViewportOffset;
+}
+
+TArray<int32> FAnimNodePoseWatch::GetViewportAllowList() const
+{ 
+	return ViewportMaskAllowedList;
+}
+
+TArray<int32> FAnimNodePoseWatch::GetParentIndices() const
+{ 
+	return ParentIndices; 
+}
+
+#endif // WITH_EDITOR
\ No newline at end of file
diff --git a/Engine/Source/Runtime/Engine/Private/SkeletalDebugRendering.cpp b/Engine/Source/Runtime/Engine/Private/SkeletalDebugRendering.cpp
--- a/Engine/Source/Runtime/Engine/Private/SkeletalDebugRendering.cpp
+++ b/Engine/Source/Runtime/Engine/Private/SkeletalDebugRendering.cpp
@@ -4,8 +4,8 @@
 #include "DrawDebugHelpers.h"
 #include "SceneManagement.h"
 #include "Engine/PoseWatch.h"
+#include "Engine/PoseWatchRenderData.h"
 #include "BonePose.h"
-#include "Animation/BlendProfile.h"
 #include "HitProxies.h"
 
 
@@ -170,20 +170,13 @@ void DrawRootCone(
 #endif
 }
 
+#if WITH_EDITOR
 void DrawBonesFromPoseWatch(
 	FPrimitiveDrawInterface* PDI,
 	const FAnimNodePoseWatch& PoseWatch,
-	const FReferenceSkeleton& RefSkeleton,
 	const bool bUseWorldTransform
 )
 {
-#if WITH_EDITOR
-	if (!PoseWatch.PoseWatchPoseElement.IsValid())
-	{
-		return;
-	}
-	const FLinearColor BoneColor = FLinearColor(PoseWatch.PoseWatchPoseElement->GetColor());
-
 	const TArray<FTransform>& InBoneTransforms = PoseWatch.GetBoneTransforms();
 	const TArray<FBoneIndexType>& InRequiredBones = PoseWatch.GetRequiredBones();
 	if (InRequiredBones.Num() == 0 || InBoneTransforms.Num() < InRequiredBones.Num())
@@ -191,26 +184,26 @@ void DrawBonesFromPoseWatch(
 		return;
 	}
 
-	const UBlendProfile* ViewportMask = PoseWatch.PoseWatchPoseElement->ViewportMask;
 	const FTransform WorldTransform = bUseWorldTransform ? PoseWatch.GetWorldTransform() : FTransform::Identity;
-	const FVector RelativeOffset = WorldTransform.GetRotation().RotateVector(PoseWatch.PoseWatchPoseElement->ViewportOffset);
+	const FVector RelativeOffset = WorldTransform.GetRotation().RotateVector(PoseWatch.GetViewportOffset());
 
 	TArray<FTransform> UseWorldTransforms;
-	UseWorldTransforms.AddDefaulted(RefSkeleton.GetNum());
+	UseWorldTransforms.AddDefaulted(InRequiredBones.Num());
 	
 	TArray<FBoneIndexType> UseRequiredBones;
 	UseRequiredBones.Reserve(InRequiredBones.Num());
+
+	TArray<int32> ViewportMaskAllowList = PoseWatch.GetViewportAllowList();
+	TArray<int32> ParentIndices = PoseWatch.GetParentIndices();
 	
 	for (int32 Index = 0; Index < InRequiredBones.Num(); ++Index)
 	{
 		const FBoneIndexType& BoneIndex = InRequiredBones[Index];
-		if (!RefSkeleton.IsValidIndex(BoneIndex))
-		{
-			continue;
-		}
 
-		const int32 ParentIndex = RefSkeleton.GetParentIndex(BoneIndex);
-		if (!RefSkeleton.IsValidIndex(ParentIndex))
+		check(ParentIndices.IsValidIndex(BoneIndex));
+		const int32 ParentIndex = ParentIndices[BoneIndex];
+
+		if (ParentIndex == INDEX_NONE)
 		{
 			UseWorldTransforms[BoneIndex] = InBoneTransforms[Index] * WorldTransform;
 			UseWorldTransforms[BoneIndex].AddToTranslation(RelativeOffset);
@@ -220,20 +213,16 @@ void DrawBonesFromPoseWatch(
 			UseWorldTransforms[BoneIndex] = InBoneTransforms[Index] * UseWorldTransforms[ParentIndex];
 		}
 
-		if (ViewportMask)
+		if (!ViewportMaskAllowList.Contains(BoneIndex))
 		{
-			const FName& BoneName = RefSkeleton.GetBoneName(BoneIndex);
-			const bool bBoneBlendScaleMeetsThreshold = ViewportMask->GetBoneBlendScale(BoneName) > PoseWatch.PoseWatchPoseElement->BlendScaleThreshold;
-			const bool bBoneRequired = (!PoseWatch.PoseWatchPoseElement->bInvertViewportMask && bBoneBlendScaleMeetsThreshold) || (PoseWatch.PoseWatchPoseElement->bInvertViewportMask && !bBoneBlendScaleMeetsThreshold);
-			if (!bBoneRequired)
-			{
-				continue;
-			}
+			continue;
 		}
 
 		UseRequiredBones.Add(BoneIndex);
 	}
 
+	const FLinearColor BoneColor = PoseWatch.GetBoneColor();
+
 	FSkelDebugDrawConfig DrawConfig;
 	DrawConfig.BoneDrawMode = EBoneDrawMode::All;
 	DrawConfig.BoneDrawSize = 1.f;
@@ -244,19 +233,18 @@ void DrawBonesFromPoseWatch(
 	DrawConfig.SelectedBoneColor = BoneColor;
 	DrawConfig.ParentOfSelectedBoneColor = BoneColor;
 
-	SkeletalDebugRendering::DrawBones(
+	SkeletalDebugRendering::DrawBonesInternal(
 		PDI,
 		WorldTransform.GetLocation() + RelativeOffset,
 		UseRequiredBones,
-		RefSkeleton,
+		ParentIndices,
 		UseWorldTransforms,
 		/*SelectedBones*/TArray<int32>(),
 		/*BoneColors*/TArray<FLinearColor>(),
 		/*HitProxies*/TArray<TRefCountPtr<HHitProxy>>(),
-		DrawConfig
-	);
-#endif
+		DrawConfig);
 }
+#endif
 
 void DrawBones(
 	FPrimitiveDrawInterface* PDI,
@@ -269,8 +257,49 @@ void DrawBones(
 	const TArray<TRefCountPtr<HHitProxy>>& HitProxies,
 	const FSkelDebugDrawConfig& DrawConfig)
 {
+	TArray<int32> ParentIndices;
+	ParentIndices.AddUninitialized(RefSkeleton.GetNum());
+
+	for (int32 BoneIndex = 0; BoneIndex < RefSkeleton.GetNum(); ++BoneIndex)
+	{
+		ParentIndices[BoneIndex] = RefSkeleton.GetParentIndex(BoneIndex);
+	}
+
+	SkeletalDebugRendering::DrawBonesInternal(
+		PDI,
+		ComponentOrigin,
+		RequiredBones,
+		ParentIndices,
+		WorldTransforms,
+		InSelectedBones,
+		BoneColors,
+		HitProxies,
+		DrawConfig);
+}
+
+void DrawBonesInternal(
+	FPrimitiveDrawInterface* PDI,
+	const FVector& ComponentOrigin,
+	const TArray<FBoneIndexType>& RequiredBones,
+	const TArray<int32>& ParentIndices,
+	const TArray<FTransform>& WorldTransforms,
+	const TArray<int32>& InSelectedBones,
+	const TArray<FLinearColor>& BoneColors,
+	const TArray<TRefCountPtr<HHitProxy>>& HitProxies,
+	const FSkelDebugDrawConfig& DrawConfig)
+{
+	const auto GetParentIndex = [ParentIndices](const int32 InBoneIndex) -> int32
+	{
+		if (ParentIndices.IsValidIndex(InBoneIndex))
+		{
+			return ParentIndices[InBoneIndex];
+		}
+		return INDEX_NONE;
+	};
+
 	// first determine which bones to draw, and which to filter out
-	TBitArray<> BonesToDraw(false, RefSkeleton.GetNum());
+	const int32 NumBones = ParentIndices.Num();
+	TBitArray<> BonesToDraw(false, NumBones);
 	const bool bDrawSelected = DrawConfig.BoneDrawMode == EBoneDrawMode::Selected;
 	const bool bDrawSelectedAndParents = DrawConfig.BoneDrawMode == EBoneDrawMode::SelectedAndParents;
 	const bool bDrawSelectedAndChildren = DrawConfig.BoneDrawMode == EBoneDrawMode::SelectedAndChildren;
@@ -291,9 +320,9 @@ void DrawBones(
 	// add children of selected
 	if (bDrawSelectedAndChildren || bDrawSelectedAndParentsAndChildren)
 	{
-		for (int32 BoneIndex = 0; BoneIndex < RefSkeleton.GetNum(); ++BoneIndex)
+		for (int32 BoneIndex = 0; BoneIndex < NumBones; ++BoneIndex)
 		{
-			const int32 ParentIndex = RefSkeleton.GetParentIndex(BoneIndex);
+			const int32 ParentIndex = GetParentIndex(BoneIndex);
 			if (ParentIndex != INDEX_NONE && BonesToDraw[ParentIndex])
 			{
 				BonesToDraw[BoneIndex] = true;
@@ -308,7 +337,7 @@ void DrawBones(
 		{
 			if (BoneIndex != INDEX_NONE)
 			{
-				for (int32 ParentIndex = RefSkeleton.GetParentIndex(BoneIndex); ParentIndex != INDEX_NONE; ParentIndex = RefSkeleton.GetParentIndex(ParentIndex))
+				for (int32 ParentIndex = GetParentIndex(BoneIndex); ParentIndex != INDEX_NONE; ParentIndex = GetParentIndex(ParentIndex))
 				{
 					BonesToDraw[ParentIndex] = true;
 				}
@@ -317,10 +346,10 @@ void DrawBones(
 	}
 
 	// determine which bones are "affected" (these are ALL children of selected bones)
-	TBitArray<> AffectedBones(false, RefSkeleton.GetNum());
-	for (int32 BoneIndex = 0; BoneIndex < RefSkeleton.GetNum(); ++BoneIndex)
+	TBitArray<> AffectedBones(false, NumBones);
+	for (int32 BoneIndex = 0; BoneIndex < NumBones; ++BoneIndex)
 	{
-		for (int32 ParentIndex = RefSkeleton.GetParentIndex(BoneIndex); ParentIndex != INDEX_NONE; ParentIndex = RefSkeleton.GetParentIndex(ParentIndex))
+		for (int32 ParentIndex = GetParentIndex(BoneIndex); ParentIndex != INDEX_NONE; ParentIndex = GetParentIndex(ParentIndex))
 		{
 			if (InSelectedBones.Contains(ParentIndex))
 			{
@@ -362,9 +391,9 @@ void DrawBones(
 		// but use a different color if this bone is NOT selected, but the child IS selected
 		TArray<FVector> ChildPositions;
 		TArray<FLinearColor> ChildColors;
-		for (int32 ChildIndex = 0; ChildIndex < RefSkeleton.GetNum(); ++ChildIndex)
+		for (int32 ChildIndex = 0; ChildIndex < NumBones; ++ChildIndex)
 		{
-			const int32 ParentIndex = RefSkeleton.GetParentIndex(ChildIndex);
+			const int32 ParentIndex = GetParentIndex(ChildIndex);
 			if (ParentIndex >= WorldTransforms.Num())
 			{
 				continue;
@@ -394,7 +423,7 @@ void DrawBones(
 			SDPG_Foreground,
 			BoneRadius,
 			bDrawAxesInsideBone);
-		if (RefSkeleton.GetParentIndex(BoneIndex) == INDEX_NONE)
+		if (GetParentIndex(BoneIndex) == INDEX_NONE)
 		{
 			SkeletalDebugRendering::DrawRootCone(PDI, BoneTransform, ComponentOrigin, BoneRadius);
 		}
diff --git a/Engine/Source/Runtime/Engine/Private/SkeletalMesh.cpp b/Engine/Source/Runtime/Engine/Private/SkeletalMesh.cpp
--- a/Engine/Source/Runtime/Engine/Private/SkeletalMesh.cpp
+++ b/Engine/Source/Runtime/Engine/Private/SkeletalMesh.cpp
@@ -75,7 +75,8 @@
 #include "IMeshReductionManagerModule.h"
 #include "SkeletalMeshReductionSettings.h"
 #include "Engine/RendererSettings.h"
-
+#include "Engine/PoseWatchRenderData.h"
+#include "SkeletalDebugRendering.h"
 #endif // #if WITH_EDITOR
 
 #include "UnrealEngine.h"
@@ -5649,6 +5650,10 @@ FSkeletalMeshSceneProxy::FSkeletalMeshSceneProxy(const USkinnedMeshComponent* Co
 	check(SkeletalMeshRenderData);
 	check(SkeletalMeshForDebug);
 
+#if WITH_EDITOR
+	PoseWatchDynamicData = nullptr;
+#endif
+
 	// Skeletal meshes DO deform internally, unless bRenderStatic is used to force static mesh behaviour.
 	bHasDeformableMesh = !bRenderStatic;
 
@@ -6686,12 +6691,11 @@ void FSkeletalMeshSceneProxy::DebugDrawPoseWatchSkeletons(int32 ViewIndex, FMesh
 {
 	FPrimitiveDrawInterface* PDI = Collector.GetPDI(ViewIndex);
 
-	if (Owner)
+	if (PoseWatchDynamicData)
 	{
-		USkeletalMeshComponent* SkeletalMeshComp = Cast<USkeletalMeshComponent>(Owner->GetComponentByClass(USkeletalMeshComponent::StaticClass()));
-		if (SkeletalMeshComp)
+		for (const FAnimNodePoseWatch& PoseWatch : PoseWatchDynamicData->PoseWatches)
 		{
-			SkeletalMeshComp->DebugDrawPoseWatches(PDI);
+			SkeletalDebugRendering::DrawBonesFromPoseWatch(PDI, PoseWatch, true);
 		}
 	}
 }
diff --git a/Engine/Source/Runtime/Engine/Public/Animation/AnimInstanceProxy.h b/Engine/Source/Runtime/Engine/Public/Animation/AnimInstanceProxy.h
--- a/Engine/Source/Runtime/Engine/Public/Animation/AnimInstanceProxy.h
+++ b/Engine/Source/Runtime/Engine/Public/Animation/AnimInstanceProxy.h
@@ -979,9 +979,6 @@ private:
 
 	/** Map of node syncs this frame - maps from player node index to graph-determined group name */
 	TMap<int32, FName> NodeSyncsThisFrame;
-
-	/** Array of nodes to watch this frame */
-	TArray<FAnimNodePoseWatch> PoseWatchEntriesForThisFrame;
 #endif
 
 #if ENABLE_ANIM_LOGGING
diff --git a/Engine/Source/Runtime/Engine/Public/SkeletalDebugRendering.h b/Engine/Source/Runtime/Engine/Public/SkeletalDebugRendering.h
--- a/Engine/Source/Runtime/Engine/Public/SkeletalDebugRendering.h
+++ b/Engine/Source/Runtime/Engine/Public/SkeletalDebugRendering.h
@@ -110,17 +110,14 @@ ENGINE_API	void DrawRootCone(
 	const FVector& ComponentOrigin,
 	const float SphereRadius);
 
-
+#if WITH_EDITOR
 ENGINE_API void DrawBonesFromPoseWatch(
 	FPrimitiveDrawInterface* PDI,
 	const FAnimNodePoseWatch& PoseWatch,
-	const FReferenceSkeleton& RefSkeleton,
 	const bool bUseWorldTransform);
+#endif
 
 
-/**
- * Draw skeleton bones. If BoneColors is empty, the DefaultBoneColor in the draw config will be used.
- */
 ENGINE_API void DrawBones(
 	FPrimitiveDrawInterface* PDI,
 	const FVector& ComponentOrigin,
@@ -131,4 +128,16 @@ ENGINE_API void DrawBones(
 	const TArray<FLinearColor>& BoneColors,
 	const TArray<TRefCountPtr<HHitProxy>>& HitProxies,
 	const FSkelDebugDrawConfig& DrawConfig);
+
+
+void DrawBonesInternal(
+	FPrimitiveDrawInterface* PDI,
+	const FVector& ComponentOrigin,
+	const TArray<FBoneIndexType>& RequiredBones,
+	const TArray<int32>& ParentIndices,
+	const TArray<FTransform>& WorldTransforms,
+	const TArray<int32>& InSelectedBones,
+	const TArray<FLinearColor>& BoneColors,
+	const TArray<TRefCountPtr<HHitProxy>>& HitProxies,
+	const FSkelDebugDrawConfig& DrawConfig);
 }
\ No newline at end of file
diff --git a/Engine/Source/Runtime/Engine/Public/SkeletalMeshTypes.h b/Engine/Source/Runtime/Engine/Public/SkeletalMeshTypes.h
--- a/Engine/Source/Runtime/Engine/Public/SkeletalMeshTypes.h
+++ b/Engine/Source/Runtime/Engine/Public/SkeletalMeshTypes.h
@@ -338,6 +338,11 @@ protected:
 	uint8 GetCurrentFirstLODIdx_Internal() const;
 private:
 	void CreateBaseMeshBatch(const FSceneView* View, const FSkeletalMeshLODRenderData& LODData, const int32 LODIndex, const int32 SectionIndex, const FSectionElementInfo& SectionElementInfo, FMeshBatch& Mesh, ESkinVertexFactoryMode VFMode = ESkinVertexFactoryMode::Default) const;
+
+public:
+#if WITH_EDITORONLY_DATA
+	struct FPoseWatchDynamicData* PoseWatchDynamicData;
+#endif
 };
 
 /** Used to recreate all skinned mesh components for a given skeletal mesh */
