Fix for sequence player start position description not matching functionality
* Sequence player nodes previously ask for values input to 'Start Position' to be in the [0, 1] interval but is used as though in the [0, length] range.

diff --git a/Engine/Source/Runtime/Engine/Classes/Animation/AnimNode_SequencePlayer.h b/Engine/Source/Runtime/Engine/Classes/Animation/AnimNode_SequencePlayer.h
--- a/Engine/Source/Runtime/Engine/Classes/Animation/AnimNode_SequencePlayer.h
+++ b/Engine/Source/Runtime/Engine/Classes/Animation/AnimNode_SequencePlayer.h
@@ -126,7 +126,7 @@ protected:
 	UPROPERTY()
 	FInputScaleBiasClamp PlayRateScaleBiasClamp_DEPRECATED;
 
-	// The start position in [0, 1] to use when initializing. When looping, play will still jump back to the beginning when reaching the end.
+	// The start position between 0 and the length of the sequence to use when initializing. When looping, play will still jump back to the beginning when reaching the end.
 	UPROPERTY(EditAnywhere, Category = Settings, meta = (PinHiddenByDefault, FoldProperty))
 	float StartPosition = 0.0f;
 
@@ -205,7 +205,7 @@ protected:
 	UPROPERTY(EditAnywhere, Category = Settings, meta = (DisplayName="PlayRateScaleBiasClamp"))
 	FInputScaleBiasClampConstants PlayRateScaleBiasClampConstants;
 
-	// The start position in [0, 1] to use when initializing. When looping, play will still jump back to the beginning when reaching the end.
+	// The start position between 0 and the length of the sequence to use when initializing. When looping, play will still jump back to the beginning when reaching the end.
 	UPROPERTY(EditAnywhere, Category = Settings, meta = (PinHiddenByDefault))
 	float StartPosition = 0.0f;
 
