Fix asset overrides not working correctly for blendspace & sequence players that use identical defaults

Folding properties prevents derived classes from properly overriding these properties

diff --git a/Engine/Source/Runtime/AnimGraphRuntime/Private/AnimNodes/AnimNode_BlendSpacePlayer.cpp b/Engine/Source/Runtime/AnimGraphRuntime/Private/AnimNodes/AnimNode_BlendSpacePlayer.cpp
--- a/Engine/Source/Runtime/AnimGraphRuntime/Private/AnimNodes/AnimNode_BlendSpacePlayer.cpp
+++ b/Engine/Source/Runtime/AnimGraphRuntime/Private/AnimNodes/AnimNode_BlendSpacePlayer.cpp
@@ -301,18 +301,8 @@ bool FAnimNode_BlendSpacePlayer::SetResetPlayTimeWhenBlendSpaceChanges(bool bInR
 
 bool FAnimNode_BlendSpacePlayer::SetBlendSpace(UBlendSpace* InBlendSpace)
 {
-#if WITH_EDITORONLY_DATA
 	BlendSpace = InBlendSpace;
-	GET_MUTABLE_ANIM_NODE_DATA(TObjectPtr<UBlendSpace>, BlendSpace) = InBlendSpace;
-#endif
-
-	if (TObjectPtr<UBlendSpace>* BlendSpacePtr = GET_INSTANCE_ANIM_NODE_DATA_PTR(TObjectPtr<UBlendSpace>, BlendSpace))
-	{
-		*BlendSpacePtr = InBlendSpace;
-		return true;
-	}
-
-	return false;
+	return true;
 }
 
 FVector FAnimNode_BlendSpacePlayer::GetPosition() const
@@ -397,5 +387,5 @@ float FAnimNode_BlendSpacePlayer::GetStartPosition() const
 
 UBlendSpace* FAnimNode_BlendSpacePlayer::GetBlendSpace() const
 {
-	return GET_ANIM_NODE_DATA(TObjectPtr<UBlendSpace>, BlendSpace);
+	return BlendSpace;
 }
\ No newline at end of file
diff --git a/Engine/Source/Runtime/AnimGraphRuntime/Public/AnimNodes/AnimNode_BlendSpacePlayer.h b/Engine/Source/Runtime/AnimGraphRuntime/Public/AnimNodes/AnimNode_BlendSpacePlayer.h
--- a/Engine/Source/Runtime/AnimGraphRuntime/Public/AnimNodes/AnimNode_BlendSpacePlayer.h
+++ b/Engine/Source/Runtime/AnimGraphRuntime/Public/AnimNodes/AnimNode_BlendSpacePlayer.h
@@ -165,11 +165,11 @@ private:
 	// The start position in [0, 1] to use when initializing. When looping, play will still jump back to the beginning when reaching the end.
 	UPROPERTY(EditAnywhere, Category = Settings, meta = (DefaultValue = "0.f", PinHiddenByDefault, FoldProperty))
 	float StartPosition = 0.0f;
+#endif
 
 	// The blendspace asset to play
-	UPROPERTY(EditAnywhere, Category = Settings, meta = (PinHiddenByDefault, FoldProperty))
+	UPROPERTY(EditAnywhere, Category = Settings, meta = (PinHiddenByDefault))
 	TObjectPtr<UBlendSpace> BlendSpace = nullptr;
-#endif
 
 public:
 
diff --git a/Engine/Source/Runtime/Engine/Classes/Animation/AnimNode_SequencePlayer.h b/Engine/Source/Runtime/Engine/Classes/Animation/AnimNode_SequencePlayer.h
--- a/Engine/Source/Runtime/Engine/Classes/Animation/AnimNode_SequencePlayer.h
+++ b/Engine/Source/Runtime/Engine/Classes/Animation/AnimNode_SequencePlayer.h
@@ -104,11 +104,13 @@ protected:
 	// If true, "Relevant anim" nodes that look for the highest weighted animation in a state will ignore this node
 	UPROPERTY(EditAnywhere, Category=Relevancy, meta=(FoldProperty, PinHiddenByDefault))
 	bool bIgnoreForRelevancyTest = false;
-	
+#endif
+
 	// The animation sequence asset to play
-	UPROPERTY(EditAnywhere, Category = Settings, meta = (PinHiddenByDefault, DisallowedClasses="/Script/Engine.AnimMontage", FoldProperty))
+	UPROPERTY(EditAnywhere, Category = Settings, meta = (PinHiddenByDefault, DisallowedClasses="/Script/Engine.AnimMontage"))
 	TObjectPtr<UAnimSequenceBase> Sequence = nullptr;
 
+#if WITH_EDITORONLY_DATA
 	// The Basis in which the PlayRate is expressed in. This is used to rescale PlayRate inputs.
 	// For example a Basis of 100 means that the PlayRate input will be divided by 100.
 	UPROPERTY(EditAnywhere, Category = Settings, meta = (PinHiddenByDefault, FoldProperty))
diff --git a/Engine/Source/Runtime/Engine/Private/Animation/AnimNode_SequencePlayer.cpp b/Engine/Source/Runtime/Engine/Private/Animation/AnimNode_SequencePlayer.cpp
--- a/Engine/Source/Runtime/Engine/Private/Animation/AnimNode_SequencePlayer.cpp
+++ b/Engine/Source/Runtime/Engine/Private/Animation/AnimNode_SequencePlayer.cpp
@@ -173,18 +173,8 @@ float FAnimNode_SequencePlayerBase::GetEffectiveStartPosition(const FAnimationBa
 
 bool FAnimNode_SequencePlayer::SetSequence(UAnimSequenceBase* InSequence)
 {
-#if WITH_EDITORONLY_DATA
 	Sequence = InSequence;
-	GET_MUTABLE_ANIM_NODE_DATA(TObjectPtr<UAnimSequenceBase>, Sequence) = InSequence;
-#endif
-	
-	if(TObjectPtr<UAnimSequenceBase>* SequencePtr = GET_INSTANCE_ANIM_NODE_DATA_PTR(TObjectPtr<UAnimSequenceBase>, Sequence))
-	{
-		*SequencePtr = InSequence;
-		return true;
-	}
-
-	return false;
+	return true;
 }
 
 bool FAnimNode_SequencePlayer::SetLoopAnimation(bool bInLoopAnimation)
@@ -204,7 +194,7 @@ bool FAnimNode_SequencePlayer::SetLoopAnimation(bool bInLoopAnimation)
 
 UAnimSequenceBase* FAnimNode_SequencePlayer::GetSequence() const
 {
-	return GET_ANIM_NODE_DATA(TObjectPtr<UAnimSequenceBase>, Sequence);
+	return Sequence;
 }
 
 float FAnimNode_SequencePlayer::GetPlayRateBasis() const
