Fixed leaking linked anim instances when reinitializing animation

Fixes cases where linked instance array was not being cleared like SetSkeletalMesh

diff --git a/Engine/Source/Runtime/Engine/Private/Components/SkeletalMeshComponent.cpp b/Engine/Source/Runtime/Engine/Private/Components/SkeletalMeshComponent.cpp
--- a/Engine/Source/Runtime/Engine/Private/Components/SkeletalMeshComponent.cpp
+++ b/Engine/Source/Runtime/Engine/Private/Components/SkeletalMeshComponent.cpp
@@ -893,6 +893,8 @@ bool USkeletalMeshComponent::InitializeAnimScriptInstance(bool bForceReinit, boo
 
 				if (AnimScriptInstance)
 				{
+					ResetLinkedAnimInstances();
+				
 					AnimScriptInstance->InitializeAnimation(bInDeferRootNodeInitialization);
 					bInitializedMainInstance = true;
 				}
@@ -947,6 +949,8 @@ bool USkeletalMeshComponent::InitializeAnimScriptInstance(bool bForceReinit, boo
 
 		if (AnimScriptInstance && !bInitializedMainInstance && bForceReinit)
 		{
+			ResetLinkedAnimInstances();
+		
 			AnimScriptInstance->InitializeAnimation(bInDeferRootNodeInitialization);
 			bInitializedMainInstance = true;
 		}
