Add type trait WithSerializer and implementation for FRawAnimSequenceTrack, to speed up load-times (after resave)

diff --git a/Engine/Source/Runtime/Core/Public/UObject/UE5ReleaseStreamObjectVersion.h b/Engine/Source/Runtime/Core/Public/UObject/UE5ReleaseStreamObjectVersion.h
--- a/Engine/Source/Runtime/Core/Public/UObject/UE5ReleaseStreamObjectVersion.h
+++ b/Engine/Source/Runtime/Core/Public/UObject/UE5ReleaseStreamObjectVersion.h
@@ -126,6 +126,9 @@ struct CORE_API FUE5ReleaseStreamObjectVersion
 
 		// Uses RG11B10 format to store the encoded reflection capture data on mobile
 		StoreReflectionCaptureEncodedHDRDataInRG11B10Format,
+		
+		// Add WithSerializer type trait and implementation for FRawAnimSequenceTrack
+		RawAnimSequenceTrackSerializer,
 
 		// -----<new versions can be added above this line>-------------------------------------------------
 		VersionPlusOne,
diff --git a/Engine/Source/Runtime/Engine/Public/Animation/AnimTypes.h b/Engine/Source/Runtime/Engine/Public/Animation/AnimTypes.h
--- a/Engine/Source/Runtime/Engine/Public/Animation/AnimTypes.h
+++ b/Engine/Source/Runtime/Engine/Public/Animation/AnimTypes.h
@@ -10,6 +10,7 @@
 #include "Animation/AnimEnums.h"
 #include "Misc/SecureHash.h"
 #include "Kismet/BlueprintFunctionLibrary.h"
+#include "UObject/UE5ReleaseStreamObjectVersion.h"
 #include "AnimTypes.generated.h"
 
 struct FMarkerPair;
@@ -905,11 +906,31 @@ struct ENGINE_API FRawAnimSequenceTrack
 		bContainsNaN = CheckForNan(ScaleKeys);
 
 		return bContainsNaN;
-	}	
+	}
+
+	bool Serialize(FArchive& Ar)
+	{
+		Ar.UsingCustomVersion(FUE5ReleaseStreamObjectVersion::GUID); 
+
+		// Previously generated content was saved without bulk array serialization, so have to rely on UProperty serialization until resaved
+		if (Ar.CustomVer(FUE5ReleaseStreamObjectVersion::GUID) < FUE5ReleaseStreamObjectVersion::RawAnimSequenceTrackSerializer)
+		{
+			return false;
+		}
+ 
+		Ar << *this;
+ 
+		return true;
+	}
 
 	static const uint32 SingleKeySize = sizeof(FVector3f) + sizeof(FQuat4f) + sizeof(FVector3f);
 };
 
+template<> struct TStructOpsTypeTraits<FRawAnimSequenceTrack> : public TStructOpsTypeTraitsBase2<FRawAnimSequenceTrack>
+{
+	enum { WithSerializer = true };
+};
+
 UCLASS()
 class ENGINE_API URawAnimSequenceTrackExtensions : public UBlueprintFunctionLibrary
 {
