Manual robo-merge and retrofit for CL 22026643:
PR #9482: Expose bind pose in decompression context and expose anim sequence in bone codec DDC key API (Contributed by [REDACTED_PERSON])

diff --git a/Engine/Source/Runtime/Engine/Classes/Animation/AnimBoneCompressionCodec.h b/Engine/Source/Runtime/Engine/Classes/Animation/AnimBoneCompressionCodec.h
--- a/Engine/Source/Runtime/Engine/Classes/Animation/AnimBoneCompressionCodec.h
+++ b/Engine/Source/Runtime/Engine/Classes/Animation/AnimBoneCompressionCodec.h
@@ -33,10 +33,17 @@ class ENGINE_API UAnimBoneCompressionCodec : public UObject
 	/** Compresses the curve data from an animation sequence. */
 	virtual bool Compress(const FCompressibleAnimData& CompressibleAnimData, FCompressibleAnimDataResult& OutResult) PURE_VIRTUAL(UAnimCurveCompressionCodec::Compress, return false;);
 
+	/*
+	 * Called to generate a unique DDC key for this codec instance and input anim sequence.
+	 * A suitable key should be generated from: the InstanceGuid, a codec version, and all relevant properties that drive the behavior.
+	 */
+	virtual void PopulateDDCKey(const UAnimSequenceBase& AnimSeq, FArchive& Ar);
+
 	/*
 	 * Called to generate a unique DDC key for this codec instance.
 	 * A suitable key should be generated from: the InstanceGuid, a codec version, and all relevant properties that drive the behavior.
 	 */
+	UE_DEPRECATED(5.1, "This function has been deprecated. Override the one above instead.")
 	virtual void PopulateDDCKey(FArchive& Ar);
 #endif
 
@@ -67,7 +74,11 @@ class ENGINE_API UAnimBoneCompressionCodec : public UObject
 	 */
 	virtual void ByteSwapOut(ICompressedAnimData& AnimData, TArrayView<uint8> CompressedData, FMemoryWriter& MemoryStream) const PURE_VIRTUAL(UAnimCurveCompressionCodec::ByteSwapOut, );
 
-	/** Decompresses all the specified bone tracks. */
+	/**
+	 * Decompresses all the specified bone tracks.
+	 * The caller is responsible for pre-filling the output pose with sensible values (e.g. reference/bind/additive identity pose) as
+	 * the codec will only decompress and write out tracks that are contained in the compressed data.
+	 */
 	virtual void DecompressPose(FAnimSequenceDecompressionContext& DecompContext, const BoneTrackArray& RotationPairs, const BoneTrackArray& TranslationPairs, const BoneTrackArray& ScalePairs, TArrayView<FTransform>& OutAtoms) const PURE_VIRTUAL(UAnimCurveCompressionCodec::DecompressPose, );
 
 	/** Decompress a single bone. */
diff --git a/Engine/Source/Runtime/Engine/Classes/Animation/AnimBoneCompressionSettings.h b/Engine/Source/Runtime/Engine/Classes/Animation/AnimBoneCompressionSettings.h
--- a/Engine/Source/Runtime/Engine/Classes/Animation/AnimBoneCompressionSettings.h
+++ b/Engine/Source/Runtime/Engine/Classes/Animation/AnimBoneCompressionSettings.h
@@ -6,6 +6,7 @@
 #include "AnimBoneCompressionSettings.generated.h"
 
 class UAnimBoneCompressionCodec;
+class UAnimSequenceBase;
 struct FCompressibleAnimData;
 struct FCompressibleAnimDataResult;
 
@@ -50,7 +51,11 @@ class ENGINE_API UAnimBoneCompressionSettings : public UObject
 	 */
 	bool Compress(const FCompressibleAnimData& AnimSeq, FCompressibleAnimDataResult& OutCompressedData) const;
 
+	/** Generates a DDC key that takes into account the current settings, selected codec, and input anim sequence. */
+	void PopulateDDCKey(const UAnimSequenceBase& AnimSeq, FArchive& Ar);
+
 	/** Generates a DDC key that takes into account the current settings and selected codec. */
+	UE_DEPRECATED(5.1, "This function has been deprecated. Override the one above instead.")
 	void PopulateDDCKey(FArchive& Ar);
 #endif
 };
diff --git a/Engine/Source/Runtime/Engine/Classes/Animation/AnimCompress.h b/Engine/Source/Runtime/Engine/Classes/Animation/AnimCompress.h
--- a/Engine/Source/Runtime/Engine/Classes/Animation/AnimCompress.h
+++ b/Engine/Source/Runtime/Engine/Classes/Animation/AnimCompress.h
@@ -505,10 +505,11 @@ public:
 		bool IncludeKeyTable = false);
 
 #if WITH_EDITOR
-	void PopulateDDCKeyArchive(FArchive& Ar) { PopulateDDCKey(Ar); }
+	UE_DEPRECATED(5.1, "PopulateDDCKeyArchive has been deprecated")
+	void PopulateDDCKeyArchive(FArchive& Ar) {}
 
 protected:
-	virtual void PopulateDDCKey(FArchive& Ar);
+	virtual void PopulateDDCKey(const UAnimSequenceBase& AnimSeq, FArchive& Ar) override;
 #endif // WITH_EDITOR
 
 	/**
diff --git a/Engine/Source/Runtime/Engine/Classes/Animation/AnimCompress_PerTrackCompression.h b/Engine/Source/Runtime/Engine/Classes/Animation/AnimCompress_PerTrackCompression.h
--- a/Engine/Source/Runtime/Engine/Classes/Animation/AnimCompress_PerTrackCompression.h
+++ b/Engine/Source/Runtime/Engine/Classes/Animation/AnimCompress_PerTrackCompression.h
@@ -137,7 +137,7 @@ protected:
 #if WITH_EDITOR
 	virtual void PostEditChangeProperty(struct FPropertyChangedEvent& PropertyChangedEvent) override;
 	virtual bool DoReduction(const FCompressibleAnimData& CompressibleAnimData, FCompressibleAnimDataResult& OutResult) override;
-	virtual void PopulateDDCKey(FArchive& Ar) override;
+	virtual void PopulateDDCKey(const UAnimSequenceBase& AnimSeq, FArchive& Ar) override;
 #endif // WITH_EDITOR
 	//~ Begin UAnimCompress Interface
 
diff --git a/Engine/Source/Runtime/Engine/Classes/Animation/AnimCompress_RemoveEverySecondKey.h b/Engine/Source/Runtime/Engine/Classes/Animation/AnimCompress_RemoveEverySecondKey.h
--- a/Engine/Source/Runtime/Engine/Classes/Animation/AnimCompress_RemoveEverySecondKey.h
+++ b/Engine/Source/Runtime/Engine/Classes/Animation/AnimCompress_RemoveEverySecondKey.h
@@ -33,7 +33,7 @@ protected:
 	//~ Begin UAnimCompress Interface
 #if WITH_EDITOR
 	virtual bool DoReduction(const FCompressibleAnimData& CompressibleAnimData, FCompressibleAnimDataResult& OutResult) override;
-	virtual void PopulateDDCKey(FArchive& Ar) override;
+	virtual void PopulateDDCKey(const UAnimSequenceBase& AnimSeq, FArchive& Ar) override;
 #endif // WITH_EDITOR
 	//~ Begin UAnimCompress Interface
 };
diff --git a/Engine/Source/Runtime/Engine/Classes/Animation/AnimCompress_RemoveLinearKeys.h b/Engine/Source/Runtime/Engine/Classes/Animation/AnimCompress_RemoveLinearKeys.h
--- a/Engine/Source/Runtime/Engine/Classes/Animation/AnimCompress_RemoveLinearKeys.h
+++ b/Engine/Source/Runtime/Engine/Classes/Animation/AnimCompress_RemoveLinearKeys.h
@@ -80,7 +80,7 @@ protected:
 	//~ Begin UAnimCompress Interface
 #if WITH_EDITOR
 	virtual bool DoReduction(const FCompressibleAnimData& CompressibleAnimData, FCompressibleAnimDataResult& OutResult) override;
-	virtual void PopulateDDCKey(FArchive& Ar) override;
+	virtual void PopulateDDCKey(const UAnimSequenceBase& AnimSeq, FArchive& Ar) override;
 #endif // WITH_EDITOR
 	//~ Begin UAnimCompress Interface
 
diff --git a/Engine/Source/Runtime/Engine/Classes/Animation/AnimCompress_RemoveTrivialKeys.h b/Engine/Source/Runtime/Engine/Classes/Animation/AnimCompress_RemoveTrivialKeys.h
--- a/Engine/Source/Runtime/Engine/Classes/Animation/AnimCompress_RemoveTrivialKeys.h
+++ b/Engine/Source/Runtime/Engine/Classes/Animation/AnimCompress_RemoveTrivialKeys.h
@@ -32,7 +32,7 @@ protected:
 	//~ Begin UAnimCompress Interface
 #if WITH_EDITOR
 	virtual bool DoReduction(const FCompressibleAnimData& CompressibleAnimData, FCompressibleAnimDataResult& OutResult) override;
-	virtual void PopulateDDCKey(FArchive& Ar) override;
+	virtual void PopulateDDCKey(const UAnimSequenceBase& AnimSeq, FArchive& Ar) override;
 #endif // WITH_EDITOR
 	//~ Begin UAnimCompress Interface
 };
diff --git a/Engine/Source/Runtime/Engine/Private/Animation/AnimBoneCompressionCodec.cpp b/Engine/Source/Runtime/Engine/Private/Animation/AnimBoneCompressionCodec.cpp
--- a/Engine/Source/Runtime/Engine/Private/Animation/AnimBoneCompressionCodec.cpp
+++ b/Engine/Source/Runtime/Engine/Private/Animation/AnimBoneCompressionCodec.cpp
@@ -43,6 +43,10 @@ FString UAnimBoneCompressionCodec::GetCodecDDCHandle() const
 }
 
 #if WITH_EDITORONLY_DATA
+void UAnimBoneCompressionCodec::PopulateDDCKey(const UAnimSequenceBase& AnimSeq, FArchive& Ar)
+{
+}
+
 void UAnimBoneCompressionCodec::PopulateDDCKey(FArchive& Ar)
 {
 }
diff --git a/Engine/Source/Runtime/Engine/Private/Animation/AnimBoneCompressionSettings.cpp b/Engine/Source/Runtime/Engine/Private/Animation/AnimBoneCompressionSettings.cpp
--- a/Engine/Source/Runtime/Engine/Private/Animation/AnimBoneCompressionSettings.cpp
+++ b/Engine/Source/Runtime/Engine/Private/Animation/AnimBoneCompressionSettings.cpp
@@ -271,6 +271,38 @@ bool UAnimBoneCompressionSettings::Compress(const FCompressibleAnimData& AnimSeq
 	return Success;
 }
 
+void UAnimBoneCompressionSettings::PopulateDDCKey(const UAnimSequenceBase& AnimSeq, FArchive& Ar)
+{
+	Ar << ErrorThreshold;
+	Ar << bForceBelowThreshold;
+
+	int32 NumValidCodecs = 0;
+	for (UAnimBoneCompressionCodec* Codec : Codecs)
+	{
+		if (Codec != nullptr)
+		{
+			const int64 archiveOffset = Ar.Tell();
+			Codec->PopulateDDCKey(AnimSeq, Ar);
+
+			if (archiveOffset == Ar.Tell())
+			{
+				// If nothing was written, perhaps the codec implements the old deprecated API, call it just in case
+				PRAGMA_DISABLE_DEPRECATION_WARNINGS
+				Codec->PopulateDDCKey(Ar);
+				PRAGMA_ENABLE_DEPRECATION_WARNINGS
+			}
+
+			NumValidCodecs++;
+		}
+	}
+
+	if (NumValidCodecs == 0)
+	{
+		static FString NoCodecString(TEXT("<Missing Codec>"));
+		Ar << NoCodecString;
+	}
+}
+
 void UAnimBoneCompressionSettings::PopulateDDCKey(FArchive& Ar)
 {
 	Ar << ErrorThreshold;
@@ -281,7 +313,12 @@ void UAnimBoneCompressionSettings::PopulateDDCKey(FArchive& Ar)
 	{
 		if (Codec != nullptr)
 		{
+			// We have no choice but to call the deprecated version, might not work correctly with newer codecs
+			// that leverage the new argument
+			PRAGMA_DISABLE_DEPRECATION_WARNINGS
 			Codec->PopulateDDCKey(Ar);
+			PRAGMA_ENABLE_DEPRECATION_WARNINGS
+
 			NumValidCodecs++;
 		}
 	}
diff --git a/Engine/Source/Runtime/Engine/Private/Animation/AnimCompress.cpp b/Engine/Source/Runtime/Engine/Private/Animation/AnimCompress.cpp
--- a/Engine/Source/Runtime/Engine/Private/Animation/AnimCompress.cpp
+++ b/Engine/Source/Runtime/Engine/Private/Animation/AnimCompress.cpp
@@ -703,9 +703,9 @@ void UAnimCompress::DecompressBone(FAnimSequenceDecompressionContext& DecompCont
 
 #if WITH_EDITOR
 
-void UAnimCompress::PopulateDDCKey(FArchive& Ar)
+void UAnimCompress::PopulateDDCKey(const UAnimSequenceBase& AnimSeq, FArchive& Ar)
 {
-	Super::PopulateDDCKey(Ar);
+	Super::PopulateDDCKey(AnimSeq, Ar);
 
 	uint8 TCF, RCF, SCF;
 	TCF = (uint8)TranslationCompressionFormat.GetValue();
diff --git a/Engine/Source/Runtime/Engine/Private/Animation/AnimCompress_PerTrackCompression.cpp b/Engine/Source/Runtime/Engine/Private/Animation/AnimCompress_PerTrackCompression.cpp
--- a/Engine/Source/Runtime/Engine/Private/Animation/AnimCompress_PerTrackCompression.cpp
+++ b/Engine/Source/Runtime/Engine/Private/Animation/AnimCompress_PerTrackCompression.cpp
@@ -1594,9 +1594,9 @@ void WriteEnumArrayToKey(FArchive& Ar, TArray<TEnumAsByte<enum AnimationCompress
 		Ar << Val;
 	}
 }
-void UAnimCompress_PerTrackCompression::PopulateDDCKey(FArchive& Ar)
+void UAnimCompress_PerTrackCompression::PopulateDDCKey(const UAnimSequenceBase& AnimSeq, FArchive& Ar)
 {
-	Super::PopulateDDCKey(Ar);
+	Super::PopulateDDCKey(AnimSeq, Ar);
 
 	Ar << MaxZeroingThreshold;
 	Ar << MaxPosDiffBitwise;
diff --git a/Engine/Source/Runtime/Engine/Private/Animation/AnimCompress_RemoveEverySecondKey.cpp b/Engine/Source/Runtime/Engine/Private/Animation/AnimCompress_RemoveEverySecondKey.cpp
--- a/Engine/Source/Runtime/Engine/Private/Animation/AnimCompress_RemoveEverySecondKey.cpp
+++ b/Engine/Source/Runtime/Engine/Private/Animation/AnimCompress_RemoveEverySecondKey.cpp
@@ -59,9 +59,9 @@ bool UAnimCompress_RemoveEverySecondKey::DoReduction(const FCompressibleAnimData
 	return true;
 }
 
-void UAnimCompress_RemoveEverySecondKey::PopulateDDCKey(FArchive& Ar)
+void UAnimCompress_RemoveEverySecondKey::PopulateDDCKey(const UAnimSequenceBase& AnimSeq, FArchive& Ar)
 {
-	Super::PopulateDDCKey(Ar);
+	Super::PopulateDDCKey(AnimSeq, Ar);
 	Ar << MinKeys;
 	bool bVal = bStartAtSecondKey;
 	Ar << bVal;
diff --git a/Engine/Source/Runtime/Engine/Private/Animation/AnimCompress_RemoveLinearKeys.cpp b/Engine/Source/Runtime/Engine/Private/Animation/AnimCompress_RemoveLinearKeys.cpp
--- a/Engine/Source/Runtime/Engine/Private/Animation/AnimCompress_RemoveLinearKeys.cpp
+++ b/Engine/Source/Runtime/Engine/Private/Animation/AnimCompress_RemoveLinearKeys.cpp
@@ -1105,9 +1105,9 @@ bool UAnimCompress_RemoveLinearKeys::DoReduction(const FCompressibleAnimData& Co
 	return true;
 }
 
-void UAnimCompress_RemoveLinearKeys::PopulateDDCKey(FArchive& Ar)
+void UAnimCompress_RemoveLinearKeys::PopulateDDCKey(const UAnimSequenceBase& AnimSeq, FArchive& Ar)
 {
-	Super::PopulateDDCKey(Ar);
+	Super::PopulateDDCKey(AnimSeq, Ar);
 	Ar << MaxPosDiff;
 	Ar << MaxAngleDiff;
 	Ar << MaxScaleDiff;
diff --git a/Engine/Source/Runtime/Engine/Private/Animation/AnimCompress_RemoveTrivialKeys.cpp b/Engine/Source/Runtime/Engine/Private/Animation/AnimCompress_RemoveTrivialKeys.cpp
--- a/Engine/Source/Runtime/Engine/Private/Animation/AnimCompress_RemoveTrivialKeys.cpp
+++ b/Engine/Source/Runtime/Engine/Private/Animation/AnimCompress_RemoveTrivialKeys.cpp
@@ -53,9 +53,9 @@ bool UAnimCompress_RemoveTrivialKeys::DoReduction(const FCompressibleAnimData& C
 	return true;
 }
 
-void UAnimCompress_RemoveTrivialKeys::PopulateDDCKey(FArchive& Ar)
+void UAnimCompress_RemoveTrivialKeys::PopulateDDCKey(const UAnimSequenceBase& AnimSeq, FArchive& Ar)
 {
-	Super::PopulateDDCKey(Ar);
+	Super::PopulateDDCKey(AnimSeq, Ar);
 	Ar << MaxPosDiff;
 	Ar << MaxAngleDiff;
 	Ar << MaxScaleDiff;
diff --git a/Engine/Source/Runtime/Engine/Private/Animation/AnimCompressionTypes.cpp b/Engine/Source/Runtime/Engine/Private/Animation/AnimCompressionTypes.cpp
--- a/Engine/Source/Runtime/Engine/Private/Animation/AnimCompressionTypes.cpp
+++ b/Engine/Source/Runtime/Engine/Private/Animation/AnimCompressionTypes.cpp
@@ -1037,7 +1037,7 @@ void DecompressPose(FCompactPose& OutPose, const FCompressedAnimSequence& Compre
 {
 	const TArray<FTransform>& RetargetTransforms = SourceSkeleton->GetRefLocalPoses(RetargetSource);
 	PRAGMA_DISABLE_DEPRECATION_WARNINGS
-	FAnimSequenceDecompressionContext DecompressionContext(SequenceLength, Interpolation, SourceName, *CompressedData.CompressedDataStructure.Get(), SourceSkeleton, bIsBakedAdditive);
+	FAnimSequenceDecompressionContext DecompressionContext(SequenceLength, Interpolation, SourceName, *CompressedData.CompressedDataStructure.Get());
 	PRAGMA_ENABLE_DEPRECATION_WARNINGS
 	UE::Anim::Decompression::DecompressPose(OutPose, CompressedData, ExtractionContext, DecompressionContext, RetargetTransforms, RootMotionReset);
 }
@@ -1045,7 +1045,7 @@ void DecompressPose(FCompactPose& OutPose, const FCompressedAnimSequence& Compre
 void DecompressPose(FCompactPose& OutPose, const FCompressedAnimSequence& CompressedData, const FAnimExtractContext& ExtractionContext, USkeleton* SourceSkeleton, float SequenceLength, EAnimInterpolationType Interpolation, bool bIsBakedAdditive, const TArray<FTransform>& RetargetTransforms, FName SourceName, const FRootMotionReset& RootMotionReset)
 {
 	PRAGMA_DISABLE_DEPRECATION_WARNINGS
-	FAnimSequenceDecompressionContext DecompressionContext(SequenceLength, Interpolation, SourceName, *CompressedData.CompressedDataStructure.Get(), SourceSkeleton, bIsBakedAdditive);
+	FAnimSequenceDecompressionContext DecompressionContext(SequenceLength, Interpolation, SourceName, *CompressedData.CompressedDataStructure.Get());
 	PRAGMA_ENABLE_DEPRECATION_WARNINGS
 	UE::Anim::Decompression::DecompressPose(OutPose, CompressedData, ExtractionContext, DecompressionContext, RetargetTransforms, RootMotionReset);
 }
diff --git a/Engine/Source/Runtime/Engine/Private/Animation/AnimSequence.cpp b/Engine/Source/Runtime/Engine/Private/Animation/AnimSequence.cpp
--- a/Engine/Source/Runtime/Engine/Private/Animation/AnimSequence.cpp
+++ b/Engine/Source/Runtime/Engine/Private/Animation/AnimSequence.cpp
@@ -288,7 +288,7 @@ FString GetAnimSequenceSpecificCacheKeySuffix(const UAnimSequence& Seq, bool bPe
 
 	ArcToHexString.Ar << CompressionErrorThresholdScale;
 	ArcToHexString.Ar << bPerformStripping;
-	Seq.BoneCompressionSettings->PopulateDDCKey(ArcToHexString.Ar);
+	Seq.BoneCompressionSettings->PopulateDDCKey(Seq, ArcToHexString.Ar);
 	Seq.CurveCompressionSettings->PopulateDDCKey(ArcToHexString.Ar);
 
 	FString Ret = FString::Printf(TEXT("%i_%s%s%s_%c%c%i_%s_%s_%i"),
@@ -1160,7 +1160,7 @@ void UAnimSequence::GetBoneTransform(FTransform& OutAtom, int32 TrackIndex, doub
 
 	if ( bEvaluateCompressedData && IsCompressedDataValid())
 	{
-		FAnimSequenceDecompressionContext DecompContext(TargetFrameRate, TargetFrameRate.AsFrameNumber(GetPlayLength()).Value, Interpolation, GetRetargetTransformsSourceName(), *CompressedData.CompressedDataStructure, GetSkeleton(), IsValidAdditive());
+		FAnimSequenceDecompressionContext DecompContext(TargetFrameRate, TargetFrameRate.AsFrameNumber(GetPlayLength()).Value, Interpolation, GetRetargetTransformsSourceName(), *CompressedData.CompressedDataStructure, GetSkeleton()->GetRefLocalPoses(), CompressedData.CompressedTrackToSkeletonMapTable, GetSkeleton(), IsValidAdditive());
 		DecompContext.Seek(Time);
 		if (CompressedData.BoneCompressionCodec)
 		{
@@ -1563,7 +1563,7 @@ void UAnimSequence::GetBonePose(FAnimationPoseData& OutAnimationPoseData, const
 	if (NumTracks != 0)
 	{
 		// Evaluate compressed bone data
-		FAnimSequenceDecompressionContext DecompContext(TargetFrameRate, TargetFrameRate.AsFrameTime(GetPlayLength()).RoundToFrame().Value, Interpolation, GetRetargetTransformsSourceName(), *CompressedData.CompressedDataStructure, GetSkeleton(), IsValidAdditive());
+		FAnimSequenceDecompressionContext DecompContext(TargetFrameRate, TargetFrameRate.AsFrameTime(GetPlayLength()).RoundToFrame().Value, Interpolation, GetRetargetTransformsSourceName(), *CompressedData.CompressedDataStructure, GetSkeleton()->GetRefLocalPoses(), CompressedData.CompressedTrackToSkeletonMapTable, GetSkeleton(), IsValidAdditive());
 		UE::Anim::Decompression::DecompressPose(OutPose, CompressedData, ExtractionContext, DecompContext, GetRetargetTransforms(), RootMotionReset);
 	}
 
diff --git a/Engine/Source/Runtime/Engine/Private/Animation/AnimStreamable.cpp b/Engine/Source/Runtime/Engine/Private/Animation/AnimStreamable.cpp
--- a/Engine/Source/Runtime/Engine/Private/Animation/AnimStreamable.cpp
+++ b/Engine/Source/Runtime/Engine/Private/Animation/AnimStreamable.cpp
@@ -400,7 +400,7 @@ void UAnimStreamable::GetAnimationPose(FAnimationPoseData& OutAnimationPoseData,
 		ChunkExtractionContext.PoseCurves = ExtractionContext.PoseCurves;
 
 		//FPlatformMisc::LowLevelOutputDebugStringf(TEXT("Playing Streaming Anim %s Time: %.2f Chunk:%i\n"), *GetName(), ExtractionContext.CurrentTime, ChunkIndex);
-		FAnimSequenceDecompressionContext Context(SamplingFrameRate, CurrentChunk.NumFrames, Interpolation, GetFName(), *CompressedData->CompressedDataStructure.Get(), GetSkeleton(), bIsBakedAdditive);
+		FAnimSequenceDecompressionContext Context(SamplingFrameRate, CurrentChunk.NumFrames, Interpolation, GetFName(), *CompressedData->CompressedDataStructure.Get(), GetSkeleton()->GetRefLocalPoses(), CompressedData->CompressedTrackToSkeletonMapTable, GetSkeleton(), bIsBakedAdditive);
 		UE::Anim::Decompression::DecompressPose(OutPose, *CompressedData, ChunkExtractionContext, Context, RetargetSource, RootMotionReset);
 	}
 }
@@ -773,7 +773,7 @@ FString UAnimStreamable::GetBaseDDCKey(uint32 NumChunks) const
 	FArcToHexString ArcToHexString;
 
 	ArcToHexString.Ar << NumChunks;
-	BoneCompressionSettings->PopulateDDCKey(ArcToHexString.Ar);
+	BoneCompressionSettings->PopulateDDCKey(*this, ArcToHexString.Ar);
 	CurveCompressionSettings->PopulateDDCKey(ArcToHexString.Ar);
 
 	FString Ret = FString::Printf(TEXT("%s%s%s%s_%s"),
diff --git a/Engine/Source/Runtime/Engine/Private/Animation/AnimationUtils.cpp b/Engine/Source/Runtime/Engine/Private/Animation/AnimationUtils.cpp
--- a/Engine/Source/Runtime/Engine/Private/Animation/AnimationUtils.cpp
+++ b/Engine/Source/Runtime/Engine/Private/Animation/AnimationUtils.cpp
@@ -279,7 +279,8 @@ void FAnimationUtils::ComputeCompressionError(const FCompressibleAnimData& Compr
 		const FTransform EndEffectorDummyBoneSocket(FQuat::Identity, FVector(END_EFFECTOR_DUMMY_BONE_LENGTH_SOCKET));
 		const FTransform EndEffectorDummyBone(FQuat::Identity, FVector(END_EFFECTOR_DUMMY_BONE_LENGTH));
 
-		FAnimSequenceDecompressionContext DecompContext(CompressibleAnimData.SampledFrameRate, CompressibleAnimData.GetNumberOfFrames(), CompressibleAnimData.Interpolation, CompressibleAnimData.AnimFName, *CompressedData.AnimData, nullptr, CompressibleAnimData.bIsValidAdditive);
+		FAnimSequenceDecompressionContext DecompContext(CompressibleAnimData.SampledFrameRate, CompressibleAnimData.GetNumberOfFrames(), CompressibleAnimData.Interpolation, CompressibleAnimData.AnimFName, *CompressedData.AnimData, RefPose,
+			CompressibleAnimData.TrackToSkeletonMapTable, nullptr, CompressibleAnimData.bIsValidAdditive);
 
 		const TArray<FBoneData>& BoneData = CompressibleAnimData.BoneData;
 
@@ -1068,7 +1069,7 @@ void FAnimationUtils::ExtractTransformFromCompressionData(const FCompressibleAni
 		// Build our read-only version from the mutable source
 		FUECompressedAnimData AnimData(AnimDataMutable);
 
-		FAnimSequenceDecompressionContext DecompContext(CompressibleAnimData.SampledFrameRate, CompressibleAnimData.GetNumberOfFrames(), CompressibleAnimData.Interpolation, CompressibleAnimData.AnimFName, AnimData, nullptr, CompressibleAnimData.bIsValidAdditive);
+		FAnimSequenceDecompressionContext DecompContext(CompressibleAnimData.SampledFrameRate, CompressibleAnimData.GetNumberOfFrames(), CompressibleAnimData.Interpolation, CompressibleAnimData.AnimFName, AnimData, CompressibleAnimData.RefLocalPoses, CompressibleAnimData.TrackToSkeletonMapTable, nullptr, CompressibleAnimData.bIsValidAdditive);
 		DecompContext.Seek(Time);
 		CompressedAnimData.Codec->DecompressBone(DecompContext, TrackIndex, OutBoneTransform);
 		return;
diff --git a/Engine/Source/Runtime/Engine/Public/Animation/AnimSequenceDecompressionContext.h b/Engine/Source/Runtime/Engine/Public/Animation/AnimSequenceDecompressionContext.h
--- a/Engine/Source/Runtime/Engine/Public/Animation/AnimSequenceDecompressionContext.h
+++ b/Engine/Source/Runtime/Engine/Public/Animation/AnimSequenceDecompressionContext.h
@@ -12,8 +12,8 @@ class USkeleton;
 /* Encapsulates decompression related data used by bone compression codecs. */
 struct FAnimSequenceDecompressionContext
 {
-	UE_DEPRECATED(5.1, "FAnimSequenceDecompressionContext signature has been deprecated, use different one")
-	FAnimSequenceDecompressionContext(float SequenceLength_, EAnimInterpolationType Interpolation_, const FName& AnimName_, const ICompressedAnimData& CompressedAnimData_,  const USkeleton* InSourceSkeleton, bool bIsBakedAdditive)
+	UE_DEPRECATED(5.1, "his constructor is deprecated. Use the other constructor by passing the ref pose and track to skeleton map to ensure safe usage with all codecs")
+	FAnimSequenceDecompressionContext(float SequenceLength_, EAnimInterpolationType Interpolation_, const FName& AnimName_, const ICompressedAnimData& CompressedAnimData_)
 		:
 		PRAGMA_DISABLE_DEPRECATION_WARNINGS
 		SequenceLength(SequenceLength_),
@@ -22,17 +22,33 @@ struct FAnimSequenceDecompressionContext
 		PRAGMA_DISABLE_DEPRECATION_WARNINGS
 		Time(0.f), RelativePos(0.f),
 		PRAGMA_ENABLE_DEPRECATION_WARNINGS
-		SourceSkeleton(InSourceSkeleton),
-		bAdditiveAnimation(bIsBakedAdditive)
+		SourceSkeleton(nullptr),
+		bAdditiveAnimation(false)
 	{
 	}
 
-	FAnimSequenceDecompressionContext(const FFrameRate& InSamplingRate, const int32 InNumberOfFrames, EAnimInterpolationType Interpolation_, const FName& AnimName_, const ICompressedAnimData& CompressedAnimData_, const USkeleton* InSourceSkeleton, bool bIsBakedAdditive)
+	UE_DEPRECATED(5.2, "FAnimSequenceDecompressionContext signature has been deprecated, use different one with sampling rate and number of frames")
+	FAnimSequenceDecompressionContext(float SequenceLength_, EAnimInterpolationType Interpolation_, const FName& AnimName_, const ICompressedAnimData& CompressedAnimData_,
+		const TArray<FTransform>& RefPoses_, const TArray<FTrackToSkeletonMap>& TrackToSkeletonMap_)
+		:
+		PRAGMA_DISABLE_DEPRECATION_WARNINGS
+		SequenceLength(SequenceLength_),
+		PRAGMA_ENABLE_DEPRECATION_WARNINGS
+		Interpolation(Interpolation_), AnimName(AnimName_), CompressedAnimData(CompressedAnimData_),
+		RefPoses(RefPoses_.GetData(), RefPoses_.Num()), TrackToSkeletonMap(TrackToSkeletonMap_.GetData(), TrackToSkeletonMap_.Num()),
+		PRAGMA_DISABLE_DEPRECATION_WARNINGS
+		Time(0.f), RelativePos(0.f)
+		PRAGMA_ENABLE_DEPRECATION_WARNINGS
+	{
+	}
+	
+	FAnimSequenceDecompressionContext(const FFrameRate& InSamplingRate, const int32 InNumberOfFrames, EAnimInterpolationType Interpolation_, const FName& AnimName_, const ICompressedAnimData& CompressedAnimData_, const TArray<FTransform>& RefPoses_, const TArray<FTrackToSkeletonMap>& TrackToSkeletonMap_, const USkeleton* InSourceSkeleton, bool bIsBakedAdditive)
 		:
 		PRAGMA_DISABLE_DEPRECATION_WARNINGS
 		SequenceLength(InSamplingRate.AsSeconds(InNumberOfFrames)),
 		PRAGMA_ENABLE_DEPRECATION_WARNINGS
 		Interpolation(Interpolation_), AnimName(AnimName_), CompressedAnimData(CompressedAnimData_),
+		RefPoses(RefPoses_.GetData(), RefPoses_.Num()), TrackToSkeletonMap(TrackToSkeletonMap_.GetData(), TrackToSkeletonMap_.Num()),
 		PRAGMA_DISABLE_DEPRECATION_WARNINGS
 		Time(0.f), RelativePos(0.f),
 		PRAGMA_ENABLE_DEPRECATION_WARNINGS
@@ -51,6 +67,14 @@ struct FAnimSequenceDecompressionContext
 
 	const ICompressedAnimData& CompressedAnimData;
 
+private:
+	// Immutable reference/bind pose, maps a bone index to its reference transform
+	TArrayView<const FTransform> RefPoses;
+
+	// Immutable array that maps from a track index to a bone index
+	TArrayView<const FTrackToSkeletonMap> TrackToSkeletonMap;
+
+public:
 	UE_DEPRECATED(5.1, "Direct access to Time has been deprecated use GetEvaluationTime or GetInterpolatedEvaluationTime instead")
 	float Time;
 	
@@ -95,6 +119,16 @@ struct FAnimSequenceDecompressionContext
 	const USkeleton* GetSourceSkeleton() const { return SourceSkeleton; }
 	bool IsAdditiveAnimation() const { return bAdditiveAnimation; }	
 
+	const TArrayView<const FTransform>& GetRefLocalPoses() const
+	{
+		return RefPoses;
+	}
+
+	const TArrayView<const FTrackToSkeletonMap>& GetTrackToSkeletonMap() const
+	{
+		return TrackToSkeletonMap;
+	}
+
 protected:
 	FFrameRate SamplingRate;
 	FFrameTime SamplingTime;
