Anim Sequence: Can't remove last frames of animation for this asset.

diff --git a/Engine/Source/Editor/Persona/Private/AnimTimeline/SAnimTimeline.cpp b/Engine/Source/Editor/Persona/Private/AnimTimeline/SAnimTimeline.cpp
--- a/Engine/Source/Editor/Persona/Private/AnimTimeline/SAnimTimeline.cpp
+++ b/Engine/Source/Editor/Persona/Private/AnimTimeline/SAnimTimeline.cpp
@@ -515,7 +515,7 @@ void SAnimTimeline::OnCropAnimSequence( bool bFromStart, float CurrentTime )
 				const float TrimEnd = bFromStart ? CurrentTime : AnimSequence->GetPlayLength();
 
 				// Trim off the user-selected part of the raw anim data.
-				UE::Anim::AnimationData::Trim(AnimSequence, TrimStart, TrimEnd);
+				UE::Anim::AnimationData::Trim(AnimSequence, TrimStart, TrimEnd, !bFromStart);
 
 
 				//Resetting slider position to the first frame
diff --git a/Engine/Source/Runtime/Engine/Classes/Animation/AnimSequenceHelpers.h b/Engine/Source/Runtime/Engine/Classes/Animation/AnimSequenceHelpers.h
--- a/Engine/Source/Runtime/Engine/Classes/Animation/AnimSequenceHelpers.h
+++ b/Engine/Source/Runtime/Engine/Classes/Animation/AnimSequenceHelpers.h
@@ -130,10 +130,11 @@ namespace Anim {
 		* @param	InSequence		Animation Sequence to add the animated frame to
 		* @param	TrimStart		Time value at which to start removing frames
 		* @param	TrimEnd			Time value at which to stop removing frames
+		* @param	bInclusiveEnd   Whether or not the TrimEnd should be included in the trimming range
 		*
 		* @return Whether or not the operation was succesful
 		*/
-		ENGINE_API bool Trim(UAnimSequence* InSequence, float TrimStart, float TrimEnd);
+		ENGINE_API bool Trim(UAnimSequence* InSequence, float TrimStart, float TrimEnd, bool bInclusiveEnd=false);
 				
 		/**
 		* Insert duplicate key(s) for all tracks for the provided Animation Sequence
diff --git a/Engine/Source/Runtime/Engine/Private/Animation/AnimSequenceHelpers.cpp b/Engine/Source/Runtime/Engine/Private/Animation/AnimSequenceHelpers.cpp
--- a/Engine/Source/Runtime/Engine/Private/Animation/AnimSequenceHelpers.cpp
+++ b/Engine/Source/Runtime/Engine/Private/Animation/AnimSequenceHelpers.cpp
@@ -661,7 +661,7 @@ bool AnimationData::AddLoopingInterpolation(UAnimSequence* InSequence)
 	return false;
 }
 
-bool AnimationData::Trim(UAnimSequence* InSequence, float TrimStart, float TrimEnd)
+bool AnimationData::Trim(UAnimSequence* InSequence, float TrimStart, float TrimEnd, bool bInclusiveEnd /*=false*/ )
 {
 	const IAnimationDataModel* DataModel = InSequence->GetDataModel();
 	IAnimationDataController& Controller = InSequence->GetController();
@@ -672,9 +672,6 @@ bool AnimationData::Trim(UAnimSequence* InSequence, float TrimStart, float TrimE
 	if (NumTracks > 0 && NumKeys > 0)
 	{
 		const FFrameRate& FrameRate = DataModel->GetFrameRate();
-				
-		// Save Total Number of Keys before trimming
-		int32 TotalNumOfKeys = NumKeys;
 
 		// if there is only one key, there is nothing to trim away
 		if (NumKeys <= 1)
@@ -686,9 +683,9 @@ bool AnimationData::Trim(UAnimSequence* InSequence, float TrimStart, float TrimE
 		const FFrameNumber EndFrameTrim = FrameRate.AsFrameTime(TrimEnd).RoundToFrame();
 
 		const int32 StartTrimKeyIndex = StartFrameTrim.Value;
-		const int32 NumTrimmedFrames = EndFrameTrim.Value - StartFrameTrim.Value;
+		const int32 NumTrimmedFrames = (EndFrameTrim.Value - StartFrameTrim.Value) + (bInclusiveEnd ? 1 : 0);
 
-		if (StartTrimKeyIndex == 0 && NumTrimmedFrames + 1 == NumKeys)
+		if (StartTrimKeyIndex == 0 && NumTrimmedFrames == NumKeys)
 		{
 			return false;
 		}
@@ -803,12 +800,12 @@ void AnimationData::RemoveKeys(UAnimSequence* InSequence, int32 StartKeyIndex, i
 			Controller.SetBoneTrackKeys(AnimationTrack.Name, TrackData.PosKeys, TrackData.RotKeys, TrackData.ScaleKeys);
 		}
 
-		const int32 NewNumberOfFrames = FMath::Max(NewNumberOfKeys - 1, 0);
+		const int32 NewNumberOfFrames = FMath::Max(NewNumberOfKeys - 1, 1);
 
 		const int32 StartKey = FMath::Max(StartKeyIndex, 0);
 
 		// Notify will happen with time slice that was removed
-		Controller.ResizeInFrames(NewNumberOfFrames, StartKey, StartKey + NumberOfKeys);
+		Controller.ResizeInFrames(NewNumberOfFrames, StartKey, StartKey + NumKeysToRemove);
 	}
 }
 
