Chaos; Fixed OnHit event propagation to correctly identify the shape that the collision came from and the component that owns it.
Previously we resolved the component just from the proxies in the collision but in the case where one component was welded to another we would report the parent (weld target) componenet as the instigator for the hit which is incorrect. Now we resolve the original body instance from the shape that was involved in the collision

diff --git a/Engine/Source/Runtime/Engine/Private/PhysicsEngine/Experimental/PhysScene_Chaos.cpp b/Engine/Source/Runtime/Engine/Private/PhysicsEngine/Experimental/PhysScene_Chaos.cpp
--- a/Engine/Source/Runtime/Engine/Private/PhysicsEngine/Experimental/PhysScene_Chaos.cpp
+++ b/Engine/Source/Runtime/Engine/Private/PhysicsEngine/Experimental/PhysScene_Chaos.cpp
@@ -682,6 +682,24 @@ FBodyInstance* FPhysScene_Chaos::GetBodyInstanceFromProxy(const IPhysicsProxyBas
 	return BodyInstance;
 }
 
+const FBodyInstance* FPhysScene_Chaos::GetBodyInstanceFromProxyAndShape(IPhysicsProxyBase* InProxy, int32 InShapeIndex) const
+{
+	const FBodyInstance* ProxyInstance = GetBodyInstanceFromProxy(InProxy);
+
+	if(ProxyInstance && InShapeIndex != INDEX_NONE)
+	{
+		PhysicsInterfaceTypes::FInlineShapeArray ShapeHandles;
+		FPhysicsInterface::GetAllShapes_AssumedLocked(ProxyInstance->GetPhysicsActorHandle(), ShapeHandles);
+
+		if(ShapeHandles.IsValidIndex(InShapeIndex))
+		{
+			ProxyInstance = ProxyInstance->GetOriginalBodyInstance(ShapeHandles[InShapeIndex]);
+		}
+	}
+
+	return ProxyInstance;
+}
+
 FCollisionNotifyInfo& FPhysScene_Chaos::GetPendingCollisionForContactPair(const void* P0, const void* P1, bool& bNewEntry)
 {
 	const FUniqueContactPairKey Key = { P0, P1 };
@@ -763,9 +781,9 @@ void FPhysScene_Chaos::HandleCollisionEvents(const Chaos::FCollisionEventData& E
 								const FVector& DeltaVelocity2 = bSwapOrder
 									? CollisionDataItem.DeltaVelocity1
 									: CollisionDataItem.DeltaVelocity2;
-
-								NotifyInfo.Info0.SetFrom(GetBodyInstanceFromProxy(PhysicsProxy0), DeltaVelocity1);
-								NotifyInfo.Info1.SetFrom(GetBodyInstanceFromProxy(PhysicsProxy1), DeltaVelocity2);
+									
+								NotifyInfo.Info0.SetFrom(GetBodyInstanceFromProxyAndShape(PhysicsProxy0, CollisionDataItem.ShapeIndex1), DeltaVelocity1);
+								NotifyInfo.Info1.SetFrom(GetBodyInstanceFromProxyAndShape(PhysicsProxy1, CollisionDataItem.ShapeIndex2), DeltaVelocity2);
 
 								FRigidBodyContactInfo& NewContact = NotifyInfo.RigidCollisionData.ContactInfos.AddZeroed_GetRef();
 								NewContact.ContactNormal = CollisionDataItem.Normal;
diff --git a/Engine/Source/Runtime/Engine/Public/Physics/Experimental/PhysScene_Chaos.h b/Engine/Source/Runtime/Engine/Public/Physics/Experimental/PhysScene_Chaos.h
--- a/Engine/Source/Runtime/Engine/Public/Physics/Experimental/PhysScene_Chaos.h
+++ b/Engine/Source/Runtime/Engine/Public/Physics/Experimental/PhysScene_Chaos.h
@@ -159,7 +159,7 @@ public:
 
 	/** Given a physics proxy, returns its associated body instance if any */
 	FBodyInstance* GetBodyInstanceFromProxy(const IPhysicsProxyBase* PhysicsProxy) const;
-
+	const FBodyInstance* GetBodyInstanceFromProxyAndShape(IPhysicsProxyBase* InProxy, int32 InShapeIndex) const;
 	/**
 	 * Callback when a world ends, to mark updated packages dirty. This can't be done in final
 	 * sync as the editor will ignore packages being dirtied in PIE. Also used to clean up any other references
diff --git a/Engine/Source/Runtime/Experimental/Chaos/Private/EventDefaults.cpp b/Engine/Source/Runtime/Experimental/Chaos/Private/EventDefaults.cpp
--- a/Engine/Source/Runtime/Experimental/Chaos/Private/EventDefaults.cpp
+++ b/Engine/Source/Runtime/Experimental/Chaos/Private/EventDefaults.cpp
@@ -217,6 +217,8 @@ namespace Chaos
 								const FPerShapeData* Shape0 = Constraint.GetShape0();
 								const FPerShapeData* Shape1 = Constraint.GetShape1();
 
+								Data.ShapeIndex1 = Shape0 ? Shape0->GetShapeIndex() : INDEX_NONE;
+								Data.ShapeIndex2 = Shape1 ? Shape1->GetShapeIndex() : INDEX_NONE;
 								Data.Mat1 = ResolveMaterial(Shape0, Constraint);
 								Data.Mat2 = ResolveMaterial(Shape1, Constraint);
 
diff --git a/Engine/Source/Runtime/Experimental/Chaos/Public/Chaos/ExternalCollisionData.h b/Engine/Source/Runtime/Experimental/Chaos/Public/Chaos/ExternalCollisionData.h
--- a/Engine/Source/Runtime/Experimental/Chaos/Public/Chaos/ExternalCollisionData.h
+++ b/Engine/Source/Runtime/Experimental/Chaos/Public/Chaos/ExternalCollisionData.h
@@ -35,10 +35,13 @@ namespace Chaos
 			, bProbe(false)
 			, Proxy1(nullptr)
 			, Proxy2(nullptr)
+			, ShapeIndex1(INDEX_NONE)
+			, ShapeIndex2(INDEX_NONE)
 		{}
 
 		FCollidingData(FVec3 InLocation, FVec3 InAccumulatedImpulse, FVec3 InNormal, FVec3 InVelocity1, FVec3 InVelocity2, FVec3 InDeltaVelocity1, FVec3 InDeltaVelocity2
-			, FVec3 InAngularVelocity1, FVec3 InAngularVelocity2, FReal InMass1,FReal InMass2,  FReal InPenetrationDepth, IPhysicsProxyBase* InProxy1, IPhysicsProxyBase* InProxy2)
+			, FVec3 InAngularVelocity1, FVec3 InAngularVelocity2, FReal InMass1,FReal InMass2,  FReal InPenetrationDepth, IPhysicsProxyBase* InProxy1, IPhysicsProxyBase* InProxy2
+			, int32 InShapeIndex1, int32 InShapeIndex2)
 			: Location(InLocation)
 			, AccumulatedImpulse(InAccumulatedImpulse)
 			, Normal(InNormal)
@@ -56,6 +59,8 @@ namespace Chaos
 			, bProbe(false)
 			, Proxy1(InProxy1)
 			, Proxy2(InProxy2)
+			, ShapeIndex1(InShapeIndex1)
+			, ShapeIndex2(InShapeIndex2)
 		{}
 
 		FVec3 Location;
@@ -82,6 +87,9 @@ namespace Chaos
 		// If either Proxy is nullptr, the structure is invalid.
 		IPhysicsProxyBase* Proxy1;
 		IPhysicsProxyBase* Proxy2;
+
+		int32 ShapeIndex1;
+		int32 ShapeIndex2;
 	};
 
 	/*
