Rename Whitelist to Allowlist and update config redirects.

diff --git a/Engine/Config/ConfigRedirects.ini b/Engine/Config/ConfigRedirects.ini
--- a/Engine/Config/ConfigRedirects.ini
+++ b/Engine/Config/ConfigRedirects.ini
@@ -28,3 +28,6 @@ BlacklistLocalizationTargets=DisallowedLocalizationTargets
 [/Script/UnrealEd.ProjectPackagingSettings]
 IniKeyBlacklist=IniKeyDenylist
 IniSectionBlacklist=IniSectionDenylist
+
+[/Script/Engine.RPCDoSDetectionConfig]
+RPCBlockWhitelist=RPCBlockAllowlist
diff --git a/Engine/Source/Runtime/Engine/Private/Net/RPCDoSDetection.cpp b/Engine/Source/Runtime/Engine/Private/Net/RPCDoSDetection.cpp
--- a/Engine/Source/Runtime/Engine/Private/Net/RPCDoSDetection.cpp
+++ b/Engine/Source/Runtime/Engine/Private/Net/RPCDoSDetection.cpp
@@ -378,7 +378,7 @@ void FRPCDoSDetection::InitConfig(FName NetDriverName)
 			NextTimeQuotaCheck = FPlatformTime::Seconds() + (CurConfigObj->InitialConnectToleranceMS / 1000.0);
 		}
 
-		RPCBlockAllowList = CurConfigObj->RPCBlockWhitelist;
+		RPCBlockAllowList = CurConfigObj->RPCBlockAllowlist;
 	}
 
 
diff --git a/Engine/Source/Runtime/Engine/Private/Net/RPCDoSDetectionConfig.h b/Engine/Source/Runtime/Engine/Private/Net/RPCDoSDetectionConfig.h
--- a/Engine/Source/Runtime/Engine/Private/Net/RPCDoSDetectionConfig.h
+++ b/Engine/Source/Runtime/Engine/Private/Net/RPCDoSDetectionConfig.h
@@ -80,10 +80,14 @@ public:
 	UPROPERTY(config)
 	int32 InitialConnectToleranceMS;
 
-	/** List of RPC's which should never be blocked */
+	UE_DEPRECATED(5.1, "This property is no longer supported. Use RPCBlockAllowlist.")
 	UPROPERTY(config)
 	TArray<FName> RPCBlockWhitelist;
 
+	/** List of RPC's which should never be blocked */
+	UPROPERTY(config)
+	TArray<FName> RPCBlockAllowlist;
+
 	/** Custom thresholds for when specific RPC's should be included in analytics */
 	UPROPERTY(config)
 	TArray<FRPCAnalyticsThreshold> RPCAnalyticsThresholds;
