Deprecate unused bRaiseMaxErrorToExisting variable in UAnimationSettings

diff --git a/Engine/Source/Runtime/Engine/Classes/Animation/AnimationSettings.h b/Engine/Source/Runtime/Engine/Classes/Animation/AnimationSettings.h
--- a/Engine/Source/Runtime/Engine/Classes/Animation/AnimationSettings.h
+++ b/Engine/Source/Runtime/Engine/Classes/Animation/AnimationSettings.h
@@ -47,6 +47,7 @@ class ENGINE_API UAnimationSettings : public UDeveloperSettings
 	bool bFirstRecompressUsingCurrentOrDefault;
 
 	/** If true and the existing compression error is greater than Alternative Compression Threshold, then Alternative Compression Threshold will be effectively raised to the existing error level */
+	UE_DEPRECATED(5.1, "This is being removed because it is unused")
 	UPROPERTY(config, EditAnywhere, Category = Compression)
 	bool bRaiseMaxErrorToExisting;
 
