Added option to prevent server from streaming-out levels when server streaming is enabled (wp.Runtime.EnableServerStreamingOut)

diff --git a/Engine/Source/Runtime/Engine/Private/WorldPartition/WorldPartition.cpp b/Engine/Source/Runtime/Engine/Private/WorldPartition/WorldPartition.cpp
--- a/Engine/Source/Runtime/Engine/Private/WorldPartition/WorldPartition.cpp
+++ b/Engine/Source/Runtime/Engine/Private/WorldPartition/WorldPartition.cpp
@@ -79,6 +79,12 @@ static FAutoConsoleVariableRef CVarEnableServerStreaming(
 	GEnableServerStreaming,
 	TEXT("Set to 1 to enable server streaming, set to 2 to only enable it in PIE."));
 
+static bool GEnableServerStreamingOut = false;
+static FAutoConsoleVariableRef CVarEnableServerStreamingOut(
+	TEXT("wp.Runtime.EnableServerStreamingOut"),
+	GEnableServerStreamingOut,
+	TEXT("Turn on/off to allow or not the server to stream out levels (only relevant when server streaming is enabled)"));
+
 #define LOCTEXT_NAMESPACE "WorldPartition"
 
 #if WITH_EDITOR
@@ -923,6 +929,11 @@ bool UWorldPartition::IsServerStreamingEnabled() const
 	return false;
 }
 
+bool UWorldPartition::IsServerStreamingOutEnabled() const
+{
+	return IsServerStreamingEnabled() && GEnableServerStreamingOut;
+}
+
 bool UWorldPartition::IsSimulating()
 {
 #if WITH_EDITOR
diff --git a/Engine/Source/Runtime/Engine/Private/WorldPartition/WorldPartitionStreamingPolicy.cpp b/Engine/Source/Runtime/Engine/Private/WorldPartition/WorldPartitionStreamingPolicy.cpp
--- a/Engine/Source/Runtime/Engine/Private/WorldPartition/WorldPartitionStreamingPolicy.cpp
+++ b/Engine/Source/Runtime/Engine/Private/WorldPartition/WorldPartitionStreamingPolicy.cpp
@@ -249,6 +249,7 @@ void UWorldPartitionStreamingPolicy::UpdateStreamingState()
 
 	const bool bIsServer = WorldPartition->IsServer();
 	const bool bIsServerStreamingEnabled = WorldPartition->IsServerStreamingEnabled();
+	const bool bCanDeactivateOrUnloadCells = !bIsServer || WorldPartition->IsServerStreamingOutEnabled();
 	const int32 NewServerStreamingEnabledEpoch = bIsServer ? (bIsServerStreamingEnabled ? 1 : 0) : INT_MIN;
 
 	if (!bIsServer || bIsServerStreamingEnabled || AWorldPartitionReplay::IsPlaybackEnabled(World))
@@ -334,84 +335,77 @@ void UWorldPartitionStreamingPolicy::UpdateStreamingState()
 		}
 	}
 
+	auto ShouldWaitForClientVisibility = [bIsServer, &ClientVisibleLevelNames, &bUpdateEpoch](const UWorldPartitionRuntimeCell* Cell)
+	{
+		check(bIsServer);
+		if (ULevel* Level = Cell->GetLevel())
+		{
+			if (ClientVisibleLevelNames.Contains(Cell->GetLevel()->GetPackage()->GetFName()))
+			{
+				UE_CLOG(bUpdateEpoch, LogWorldPartition, Verbose, TEXT("Server epoch update delayed by client visibility"));
+				bUpdateEpoch = false;
+				return true;
+			}
+		}
+		return false;
+	};
+
+	auto ShouldSkipDisabledHLODCell = [](const UWorldPartitionRuntimeCell* Cell)
+	{
+		return Cell->GetIsHLOD() && !UHLODSubsystem::IsHLODEnabled();
+	};
+
 	// Process cells to activate
-	auto ProcessCellsToActivate = [this](TSet<const UWorldPartitionRuntimeCell*>& Cells)
+	auto ProcessCellsToActivate = [this, &ShouldSkipDisabledHLODCell](TSet<const UWorldPartitionRuntimeCell*>& Cells)
 	{
 		for (TSet<const UWorldPartitionRuntimeCell*>::TIterator It = Cells.CreateIterator(); It; ++It)
 		{
 			const UWorldPartitionRuntimeCell* Cell = *It;
-
-			if (ShouldSkipCellForPerformance(Cell))
+			if (ShouldSkipCellForPerformance(Cell) || ShouldSkipDisabledHLODCell(Cell))
 			{
 				It.RemoveCurrent();
-				continue;
 			}
-			
-			if (Cell->GetIsHLOD() && !UHLODSubsystem::IsHLODEnabled())
+			else
 			{
-				It.RemoveCurrent();
-				continue;
+				Cell->MergeStreamingSourceInfo();
 			}
-
-			Cell->MergeStreamingSourceInfo();
 		}
 	};
 
 	// Process cells to load
-	auto ProcessCellsToLoad = [this, bIsServer, &ClientVisibleLevelNames, &bUpdateEpoch](TSet<const UWorldPartitionRuntimeCell*>& Cells)
+	auto ProcessCellsToLoad = [this, bIsServer, bCanDeactivateOrUnloadCells, &ShouldWaitForClientVisibility, &ShouldSkipDisabledHLODCell](TSet<const UWorldPartitionRuntimeCell*>& Cells)
 	{
 		for (TSet<const UWorldPartitionRuntimeCell*>::TIterator It = Cells.CreateIterator(); It; ++It)
 		{
 			const UWorldPartitionRuntimeCell* Cell = *It;
+			// Only deactivated server cells need to call ShouldWaitForClientVisibility (those part of ActivatedCells)
+			const bool bIsServerCellToDeactivate = bIsServer && ActivatedCells.Contains(Cell);
+			const bool bShoudSkipServerCell = bIsServerCellToDeactivate && (!bCanDeactivateOrUnloadCells || ShouldWaitForClientVisibility(Cell));
 
-			if (bIsServer)
-			{
-				if (ULevel* Level = Cell->GetLevel())
-				{
-					if (ClientVisibleLevelNames.Contains(Cell->GetLevel()->GetPackage()->GetFName()))
-					{
-						UE_CLOG(bUpdateEpoch, LogWorldPartition, Verbose, TEXT("Server epoch update delayed by client visibility"));
-						bUpdateEpoch = false;
-						It.RemoveCurrent();
-						continue;
-					}
-				}
-			}
-
-			if (ShouldSkipCellForPerformance(Cell))
+			if (bShoudSkipServerCell || ShouldSkipCellForPerformance(Cell) || ShouldSkipDisabledHLODCell(Cell))
 			{
 				It.RemoveCurrent();
-				continue;
 			}
-			
-			if (Cell->GetIsHLOD() && !UHLODSubsystem::IsHLODEnabled())
+			else
 			{
-				It.RemoveCurrent();
-				continue;
+				Cell->MergeStreamingSourceInfo();
 			}
-
-			Cell->MergeStreamingSourceInfo();
 		}
 	};
 
 	// Process cells to unload
-	auto ProcessCellsToUnload = [this, bIsServer, &ClientVisibleLevelNames, &bUpdateEpoch](TSet<const UWorldPartitionRuntimeCell*>& Cells)
+	auto ProcessCellsToUnload = [this, bIsServer, bCanDeactivateOrUnloadCells, &ShouldWaitForClientVisibility](TSet<const UWorldPartitionRuntimeCell*>& Cells)
 	{
+		check(bCanDeactivateOrUnloadCells || Cells.IsEmpty());
+		// Only loop for server as ShouldWaitForClientVisibility only concerns server
 		if (bIsServer)
-		{	
+		{
 			for (TSet<const UWorldPartitionRuntimeCell*>::TIterator It = Cells.CreateIterator(); It; ++It)
 			{
 				const UWorldPartitionRuntimeCell* Cell = *It;
-
-				if (ULevel* Level = Cell->GetLevel())
+				if (ShouldWaitForClientVisibility(Cell))
 				{
-					if (ClientVisibleLevelNames.Contains(Cell->GetLevel()->GetPackage()->GetFName()))
-					{
-						UE_CLOG(bUpdateEpoch, LogWorldPartition, Verbose, TEXT("Server epoch update delayed by client visibility"));
-						bUpdateEpoch = false;
-						It.RemoveCurrent();
-						continue;
-					}
+					It.RemoveCurrent();
 				}
 			}
 		}
@@ -424,21 +418,28 @@ void UWorldPartitionStreamingPolicy::UpdateStreamingState()
 		SortStreamingCellsByImportance(InCells, OutCells);
 	};
 
-	// Determine cells to load/unload and sort by importance
+	// Determine cells to activate (sorted by importance)
 	TArray<const UWorldPartitionRuntimeCell*> ToActivateCells;
-	TArray<const UWorldPartitionRuntimeCell*> ToLoadCells;
-	TArray<const UWorldPartitionRuntimeCell*> ToUnloadCells;
 	{
 		TSet<const UWorldPartitionRuntimeCell*> ToActivateCellsUnsorted = FrameActivateCells.Difference(ActivatedCells);
-		TSet<const UWorldPartitionRuntimeCell*> ToLoadCellsUnsorted = FrameLoadCells.Difference(LoadedCells);
-		TSet<const UWorldPartitionRuntimeCell*> ToUnloadCellsUnsorted = ActivatedCells.Union(LoadedCells).Difference(FrameActivateCells.Union(FrameLoadCells));
-
 		ProcessCellsToActivate(ToActivateCellsUnsorted);
-		ProcessCellsToLoad(ToLoadCellsUnsorted);
-		ProcessCellsToUnload(ToUnloadCellsUnsorted);
-		
 		SortStreamingCells(StreamingSources, ToActivateCellsUnsorted, ToActivateCells);
+	}
+
+	// Determine cells to load (sorted by importance)
+	TArray<const UWorldPartitionRuntimeCell*> ToLoadCells;
+	{
+		TSet<const UWorldPartitionRuntimeCell*> ToLoadCellsUnsorted = FrameLoadCells.Difference(LoadedCells);
+		ProcessCellsToLoad(ToLoadCellsUnsorted);
 		SortStreamingCells(StreamingSources, ToLoadCellsUnsorted, ToLoadCells);
+	}
+
+	// Determine cells to unload
+	TArray<const UWorldPartitionRuntimeCell*> ToUnloadCells;
+	if (bCanDeactivateOrUnloadCells)
+	{
+		TSet<const UWorldPartitionRuntimeCell*> ToUnloadCellsUnsorted = ActivatedCells.Union(LoadedCells).Difference(FrameActivateCells.Union(FrameLoadCells));
+		ProcessCellsToUnload(ToUnloadCellsUnsorted);
 		ToUnloadCells = ToUnloadCellsUnsorted.Array();
 	}
 
diff --git a/Engine/Source/Runtime/Engine/Public/WorldPartition/WorldPartition.h b/Engine/Source/Runtime/Engine/Public/WorldPartition/WorldPartition.h
--- a/Engine/Source/Runtime/Engine/Public/WorldPartition/WorldPartition.h
+++ b/Engine/Source/Runtime/Engine/Public/WorldPartition/WorldPartition.h
@@ -214,6 +214,7 @@ public:
 	bool CanStream() const;
 	bool IsServer() const;
 	bool IsServerStreamingEnabled() const;
+	bool IsServerStreamingOutEnabled() const;
 
 	bool IsMainWorldPartition() const;
 
