Implements MEVP_ResolutionFraction to allows some effect to remain resolution independent

diff --git a/Engine/Source/Runtime/Engine/Classes/Materials/MaterialExpressionViewProperty.h b/Engine/Source/Runtime/Engine/Classes/Materials/MaterialExpressionViewProperty.h
--- a/Engine/Source/Runtime/Engine/Classes/Materials/MaterialExpressionViewProperty.h
+++ b/Engine/Source/Runtime/Engine/Classes/Materials/MaterialExpressionViewProperty.h
@@ -39,6 +39,8 @@ enum EMaterialExposedViewProperty
 	MEVP_PreExposure UMETA(DisplayName = "Pre-Exposure"),
 	/** Maximum mip level of Runtime Virtual Texture that Runtime Virtual Texture Output is rendering to. */
 	MEVP_RuntimeVirtualTextureMaxLevel UMETA(DisplayName = "Virtual Texture Max Level"),
+	/** Screen percentage at which the rendering resolution happens, to allow tech-art to remain consistent with dynamic resolution. */
+	MEVP_ResolutionFraction UMETA(DisplayName = "ScreenPercentage / 100"),
 
 	MEVP_MAX,
 };
diff --git a/Engine/Source/Runtime/Engine/Private/Materials/HLSLMaterialTranslator.cpp b/Engine/Source/Runtime/Engine/Private/Materials/HLSLMaterialTranslator.cpp
--- a/Engine/Source/Runtime/Engine/Private/Materials/HLSLMaterialTranslator.cpp
+++ b/Engine/Source/Runtime/Engine/Private/Materials/HLSLMaterialTranslator.cpp
@@ -4343,6 +4343,7 @@ int32 FHLSLMaterialTranslator::ViewProperty(EMaterialExposedViewProperty Propert
 		{MEVP_RuntimeVirtualTextureOutputDerivative, MCT_Float2, TEXT("View.RuntimeVirtualTextureMipLevel.zw"), nullptr},
 		{MEVP_PreExposure, MCT_Float1, TEXT("View.PreExposure.x"), TEXT("View.OneOverPreExposure.x")},
 		{MEVP_RuntimeVirtualTextureMaxLevel, MCT_Float1, TEXT("View.RuntimeVirtualTextureMipLevel.y"), nullptr},
+		{MEVP_ResolutionFraction, MCT_Float1, TEXT("View.ResolutionFractionAndInv.x"), TEXT("View.ResolutionFractionAndInv.y")},
 	};
 	static_assert((sizeof(ViewPropertyMetaArray) / sizeof(ViewPropertyMetaArray[0])) == MEVP_MAX, "incoherency between EMaterialExposedViewProperty and ViewPropertyMetaArray");
 
diff --git a/Engine/Source/Runtime/Engine/Private/Materials/MaterialExpressionHLSL.cpp b/Engine/Source/Runtime/Engine/Private/Materials/MaterialExpressionHLSL.cpp
--- a/Engine/Source/Runtime/Engine/Private/Materials/MaterialExpressionHLSL.cpp
+++ b/Engine/Source/Runtime/Engine/Private/Materials/MaterialExpressionHLSL.cpp
@@ -496,6 +496,7 @@ bool UMaterialExpressionViewProperty::GenerateHLSLExpression(FMaterialHLSLGenera
 	case MEVP_RuntimeVirtualTextureOutputDerivative: Input = EExternalInput::RuntimeVirtualTextureOutputDerivative; bApplyRcp = bRequestedRcp; break;
 	case MEVP_PreExposure: Input = bRequestedRcp ? EExternalInput::RcpPreExposure : EExternalInput::PreExposure; break;
 	case MEVP_RuntimeVirtualTextureMaxLevel: Input = EExternalInput::RuntimeVirtualTextureMaxLevel; bApplyRcp = bRequestedRcp; break;
+	case MEVP_ResolutionFraction: Input = bRequestedRcp ? EExternalInput::RcpResolutionFraction : EExternalInput::ResolutionFraction;
 	default: checkNoEntry(); break;
 	}
 
diff --git a/Engine/Source/Runtime/Engine/Private/Materials/MaterialHLSLTree.cpp b/Engine/Source/Runtime/Engine/Private/Materials/MaterialHLSLTree.cpp
--- a/Engine/Source/Runtime/Engine/Private/Materials/MaterialHLSLTree.cpp
+++ b/Engine/Source/Runtime/Engine/Private/Materials/MaterialHLSLTree.cpp
@@ -100,6 +100,8 @@ FExternalInputDescription GetExternalInputDescription(EExternalInput Input)
 	case EExternalInput::RuntimeVirtualTextureOutputLevel: return FExternalInputDescription(TEXT("RuntimeVirtualTextureOutputLevel"), Shader::EValueType::Float1);
 	case EExternalInput::RuntimeVirtualTextureOutputDerivative: return FExternalInputDescription(TEXT("RuntimeVirtualTextureOutputDerivative"), Shader::EValueType::Float2);
 	case EExternalInput::RuntimeVirtualTextureMaxLevel: return FExternalInputDescription(TEXT("RuntimeVirtualTextureMaxLevel"), Shader::EValueType::Float1);
+	case EExternalInput::ResolutionFraction: return FExternalInputDescription(TEXT("ResolutionFraction"), Shader::EValueType::Float1);
+	case EExternalInput::RcpResolutionFraction: return FExternalInputDescription(TEXT("RcpResolutionFraction"), Shader::EValueType::Float1);
 
 	case EExternalInput::CameraVector: return FExternalInputDescription(TEXT("CameraVector"), Shader::EValueType::Float3);
 	case EExternalInput::CameraWorldPosition: return FExternalInputDescription(TEXT("CameraWorldPosition"), Shader::EValueType::Double3, EExternalInput::None, EExternalInput::None, EExternalInput::PrevCameraWorldPosition);
@@ -288,6 +290,8 @@ void FExpressionExternalInput::EmitValueShader(FEmitContext& Context, FEmitScope
 		case EExternalInput::RuntimeVirtualTextureOutputLevel:  Code = TEXT("View.RuntimeVirtualTextureMipLevel.x"); break;
 		case EExternalInput::RuntimeVirtualTextureOutputDerivative: Code = TEXT("View.RuntimeVirtualTextureMipLevel.zw"); break;
 		case EExternalInput::RuntimeVirtualTextureMaxLevel:  Code = TEXT("View.RuntimeVirtualTextureMipLevel.y"); break;
+		case EExternalInput::ResolutionFraction: Code = TEXT("View.ResolutionFractionAndInv.x"); break;
+		case EExternalInput::RcpResolutionFraction: Code = TEXT("View.ResolutionFractionAndInv.y"); break;
 
 		case EExternalInput::CameraVector: Code = TEXT("Parameters.CameraVector"); break;
 		case EExternalInput::CameraWorldPosition: Code = TEXT("ResolvedView.WorldCameraOrigin"); break;
diff --git a/Engine/Source/Runtime/Engine/Private/SceneView.cpp b/Engine/Source/Runtime/Engine/Private/SceneView.cpp
--- a/Engine/Source/Runtime/Engine/Private/SceneView.cpp
+++ b/Engine/Source/Runtime/Engine/Private/SceneView.cpp
@@ -2629,6 +2629,7 @@ void FSceneView::SetupCommonViewUniformBufferParameters(
 	ViewUniformShaderParameters.NearPlane = InViewMatrices.ComputeNearPlane();
 	ViewUniformShaderParameters.MaterialTextureMipBias = 0.0f;
 	ViewUniformShaderParameters.MaterialTextureDerivativeMultiply = 1.0f;
+	ViewUniformShaderParameters.ResolutionFractionAndInv = FVector2f(1.0f, 1.0f);
 
 	ViewUniformShaderParameters.bCheckerboardSubsurfaceProfileRendering = 0;
 
diff --git a/Engine/Source/Runtime/Engine/Public/MaterialHLSLTree.h b/Engine/Source/Runtime/Engine/Public/MaterialHLSLTree.h
--- a/Engine/Source/Runtime/Engine/Public/MaterialHLSLTree.h
+++ b/Engine/Source/Runtime/Engine/Public/MaterialHLSLTree.h
@@ -132,6 +132,8 @@ enum class EExternalInput : uint8
 	RuntimeVirtualTextureOutputLevel,
 	RuntimeVirtualTextureOutputDerivative,
 	RuntimeVirtualTextureMaxLevel,
+	ResolutionFraction,
+	RcpResolutionFraction,
 
 	CameraVector,
 	CameraWorldPosition,
diff --git a/Engine/Source/Runtime/Engine/Public/SceneView.h b/Engine/Source/Runtime/Engine/Public/SceneView.h
--- a/Engine/Source/Runtime/Engine/Public/SceneView.h
+++ b/Engine/Source/Runtime/Engine/Public/SceneView.h
@@ -651,6 +651,7 @@ enum ETranslucencyVolumeCascade
 	VIEW_UNIFORM_BUFFER_MEMBER(FVector4f, BufferBilinearUVMinMax) \
 	VIEW_UNIFORM_BUFFER_MEMBER(FVector4f, ScreenToViewSpace) \
 	VIEW_UNIFORM_BUFFER_MEMBER(FVector2f, BufferToSceneTextureScale) \
+	VIEW_UNIFORM_BUFFER_MEMBER(FVector2f, ResolutionFractionAndInv) \
 	VIEW_UNIFORM_BUFFER_MEMBER(int32, NumSceneColorMSAASamples) \
 	VIEW_UNIFORM_BUFFER_MEMBER(float, SeparateWaterMainDirLightLuminance) \
 	VIEW_UNIFORM_BUFFER_MEMBER(float, PreExposure) \
diff --git a/Engine/Source/Runtime/Renderer/Private/SceneRendering.cpp b/Engine/Source/Runtime/Renderer/Private/SceneRendering.cpp
--- a/Engine/Source/Runtime/Renderer/Private/SceneRendering.cpp
+++ b/Engine/Source/Runtime/Renderer/Private/SceneRendering.cpp
@@ -1564,6 +1564,13 @@ void FViewInfo::SetupUniformBufferParameters(
 			TemporalJitterPixels.Y);
 	}
 
+	{
+		float ResolutionFraction = float(ViewRect.Width()) / float(UnscaledViewRect.Width());
+
+		ViewUniformShaderParameters.ResolutionFractionAndInv.X = ResolutionFraction;
+		ViewUniformShaderParameters.ResolutionFractionAndInv.Y = 1.0f / ResolutionFraction;
+	}
+
 	uint32 FrameIndex = 0;
 	if (ViewState)
 	{
