Volumetric Cloud - conservative density is now a float4
- The VolumetricAdvancedInput node still ahs its first output be a float3 to make sure we do not break any existing content and do not really a heavy conversion step.
- Another VolumetricAdvancedInput node output is float4.
- Conservative density of VolumetricAdvancedOutput node takes any floatX and force cast to float4 to be compatible with the code assuming conservative dnrsity is always a float4.

Tested with multiple assets.

diff --git a/Engine/Shaders/Private/MaterialTemplate.ush b/Engine/Shaders/Private/MaterialTemplate.ush
--- a/Engine/Shaders/Private/MaterialTemplate.ush
+++ b/Engine/Shaders/Private/MaterialTemplate.ush
@@ -458,7 +458,7 @@ struct FMaterialPixelParameters
 	float CloudSampleAltitude;
 	float CloudSampleAltitudeInLayer;
 	float CloudSampleNormAltitudeInLayer;
-	float3 VolumeSampleConservativeDensity;
+	float4 VolumeSampleConservativeDensity;
 	float ShadowSampleDistance;
 #endif
 
@@ -1741,12 +1741,12 @@ float MaterialExpressionCloudSampleNormAltitudeInLayer(FMaterialPixelParameters
 #endif
 }
 
-float3 MaterialExpressionVolumeSampleConservativeDensity(FMaterialPixelParameters Parameters)
+float4 MaterialExpressionVolumeSampleConservativeDensity(FMaterialPixelParameters Parameters)
 {
 #if CLOUD_LAYER_PIXEL_SHADER
 	return Parameters.VolumeSampleConservativeDensity;
 #else
-	return float3(0.0f, 0.0f, 0.0f);
+	return float4(0.0f, 0.0f, 0.0f, 0.0f);
 #endif
 }
 
diff --git a/Engine/Source/Runtime/Engine/Private/Materials/HLSLMaterialTranslator.cpp b/Engine/Source/Runtime/Engine/Private/Materials/HLSLMaterialTranslator.cpp
--- a/Engine/Source/Runtime/Engine/Private/Materials/HLSLMaterialTranslator.cpp
+++ b/Engine/Source/Runtime/Engine/Private/Materials/HLSLMaterialTranslator.cpp
@@ -9398,7 +9398,7 @@ int32 FHLSLMaterialTranslator::GetCloudSampleShadowSampleDistance()
 
 int32 FHLSLMaterialTranslator::GetVolumeSampleConservativeDensity()
 {
-	return AddCodeChunk(MCT_Float3, TEXT("MaterialExpressionVolumeSampleConservativeDensity(Parameters)"));
+	return AddCodeChunk(MCT_Float4, TEXT("MaterialExpressionVolumeSampleConservativeDensity(Parameters)"));
 }
 
 int32 FHLSLMaterialTranslator::CustomPrimitiveData(int32 OutputIndex, EMaterialValueType Type)
diff --git a/Engine/Source/Runtime/Engine/Private/Materials/MaterialExpressionHLSL.cpp b/Engine/Source/Runtime/Engine/Private/Materials/MaterialExpressionHLSL.cpp
--- a/Engine/Source/Runtime/Engine/Private/Materials/MaterialExpressionHLSL.cpp
+++ b/Engine/Source/Runtime/Engine/Private/Materials/MaterialExpressionHLSL.cpp
@@ -2223,12 +2223,13 @@ static UE::Shader::FType GetCustomOutputType(const FMaterialHLSLGenerator& Gener
 
 bool UMaterialExpressionVolumetricAdvancedMaterialInput::GenerateHLSLExpression(FMaterialHLSLGenerator& Generator, UE::HLSLTree::FScope& Scope, int32 OutputIndex, UE::HLSLTree::FExpression const*& OutExpression) const
 {
-	if (OutputIndex != 0)
+	switch (OutputIndex)
 	{
-		return Generator.Error(TEXT("Invalid output"));
+	case 0: OutExpression = Generator.GetTree().NewExpression<UE::HLSLTree::FExpressionInlineCustomHLSL>(UE::Shader::EValueType::Float3, TEXT("MaterialExpressionVolumeSampleConservativeDensity(Parameters).rgb"));
+	case 1: OutExpression = Generator.GetTree().NewExpression<UE::HLSLTree::FExpressionInlineCustomHLSL>(UE::Shader::EValueType::Float4, TEXT("MaterialExpressionVolumeSampleConservativeDensity(Parameters)"));
+	default: return Generator.Error(TEXT("Invlid output"));
 	}
-	OutExpression = Generator.GetTree().NewExpression<UE::HLSLTree::FExpressionInlineCustomHLSL>(UE::Shader::EValueType::Float3, TEXT("MaterialExpressionVolumeSampleConservativeDensity(Parameters)"));
-	return true;
+	return OutExpression != nullptr;
 }
 
 bool UMaterialExpressionVolumetricAdvancedMaterialOutput::GenerateHLSLExpression(FMaterialHLSLGenerator& Generator, UE::HLSLTree::FScope& Scope, int32 OutputIndex, UE::HLSLTree::FExpression const*& OutExpression) const
@@ -2241,7 +2242,7 @@ bool UMaterialExpressionVolumetricAdvancedMaterialOutput::GenerateHLSLExpression
 	case 3: OutExpression = MultiScatteringContribution.AcquireHLSLExpressionOrConstant(Generator, Scope, ConstMultiScatteringContribution); break;
 	case 4: OutExpression = MultiScatteringOcclusion.AcquireHLSLExpressionOrConstant(Generator, Scope, ConstMultiScatteringOcclusion); break;
 	case 5: OutExpression = MultiScatteringEccentricity.AcquireHLSLExpressionOrConstant(Generator, Scope, ConstMultiScatteringEccentricity); break;
-	case 6: OutExpression = ConservativeDensity.AcquireHLSLExpressionOrConstant(Generator, Scope, FVector3f(1.0f, 1.0f, 1.0f)); break;
+	case 6: OutExpression = ConservativeDensity.AcquireHLSLExpressionOrConstant(Generator, Scope, FVector4f(1.0f, 1.0f, 1.0f, 1.0f)); break;
 	default: return Generator.Error(TEXT("Invlid output"));
 	}
 	return OutExpression != nullptr;
diff --git a/Engine/Source/Runtime/Engine/Private/Materials/MaterialExpressions.cpp b/Engine/Source/Runtime/Engine/Private/Materials/MaterialExpressions.cpp
--- a/Engine/Source/Runtime/Engine/Private/Materials/MaterialExpressions.cpp
+++ b/Engine/Source/Runtime/Engine/Private/Materials/MaterialExpressions.cpp
@@ -21155,6 +21155,9 @@ int32 UMaterialExpressionVolumetricAdvancedMaterialOutput::Compile(class FMateri
 	else if (OutputIndex == 6)
 	{
 		CodeInput = ConservativeDensity.IsConnected() ? ConservativeDensity.Compile(Compiler) : Compiler->Constant3(1.0f, 1.0f, 1.0f);
+
+		// We force a cast to a float4 because that is what we use behind the scene and is needed by artists in some cases.
+		CodeInput = Compiler->ForceCast(CodeInput, MCT_Float4);
 	}
 
 	return Compiler->CustomOutput(this, OutputIndex, CodeInput);
@@ -21216,7 +21219,8 @@ UMaterialExpressionVolumetricAdvancedMaterialInput::UMaterialExpressionVolumetri
 	bShowOutputNameOnPin = true;
 
 	Outputs.Reset();
-	Outputs.Add(FExpressionOutput(TEXT("ConservativeDensity")));
+	Outputs.Add(FExpressionOutput(TEXT("ConservativeDensity as Float3")));
+	Outputs.Add(FExpressionOutput(TEXT("ConservativeDensity as Float4")));
 #endif
 }
 
@@ -21224,6 +21228,10 @@ UMaterialExpressionVolumetricAdvancedMaterialInput::UMaterialExpressionVolumetri
 int32 UMaterialExpressionVolumetricAdvancedMaterialInput::Compile(class FMaterialCompiler* Compiler, int32 OutputIndex)
 {
 	if (OutputIndex == 0)
+	{
+		return Compiler->ForceCast(Compiler->GetVolumeSampleConservativeDensity(), MCT_Float3);
+	}
+	else if (OutputIndex == 1)
 	{
 		return Compiler->GetVolumeSampleConservativeDensity();
 	}
