Chaos: Change coincident vertex Warning during tri-mesh cleaning to Display. This is a recoverable issue as we just delete those triangles and move on. Display will still message to the user that there are invalid triangles to address.

diff --git a/Engine/Source/Runtime/Engine/Private/PhysicsEngine/Experimental/ChaosDerivedDataUtil.cpp b/Engine/Source/Runtime/Engine/Private/PhysicsEngine/Experimental/ChaosDerivedDataUtil.cpp
--- a/Engine/Source/Runtime/Engine/Private/PhysicsEngine/Experimental/ChaosDerivedDataUtil.cpp
+++ b/Engine/Source/Runtime/Engine/Private/PhysicsEngine/Experimental/ChaosDerivedDataUtil.cpp
@@ -71,7 +71,7 @@ namespace Chaos
 				}
 			}
 		}
-		UE_CLOG(NumBadTris > 0, LogChaos, Warning, TEXT("Input trimesh contains %d bad triangles."), NumBadTris);
+		UE_CLOG(NumBadTris > 0, LogChaos, Display, TEXT("Input trimesh contains %d bad triangles."), NumBadTris);
 
 		float WeldThresholdSq = 0.0f;// SMALL_NUMBER * SMALL_NUMBER;
 
