Supporting render target textures in listtextures command
#RNX

diff --git a/Engine/Source/Runtime/Engine/Classes/Engine/TextureRenderTarget2D.h b/Engine/Source/Runtime/Engine/Classes/Engine/TextureRenderTarget2D.h
--- a/Engine/Source/Runtime/Engine/Classes/Engine/TextureRenderTarget2D.h
+++ b/Engine/Source/Runtime/Engine/Classes/Engine/TextureRenderTarget2D.h
@@ -203,6 +203,7 @@ class UTextureRenderTarget2D : public UTextureRenderTarget
 	virtual uint32 GetSurfaceArraySize() const override { return 0; }
 	ENGINE_API virtual FTextureResource* CreateResource() override;
 	ENGINE_API virtual EMaterialValueType GetMaterialType() const override;
+	virtual uint32 CalcTextureMemorySizeEnum(ETextureMipCount Enum) const override;
 	//~ End UTexture Interface.
 
 	//~ Begin UObject Interface
diff --git a/Engine/Source/Runtime/Engine/Private/TextureRenderTarget2D.cpp b/Engine/Source/Runtime/Engine/Private/TextureRenderTarget2D.cpp
--- a/Engine/Source/Runtime/Engine/Private/TextureRenderTarget2D.cpp
+++ b/Engine/Source/Runtime/Engine/Private/TextureRenderTarget2D.cpp
@@ -65,6 +65,19 @@ FTextureResource* UTextureRenderTarget2D::CreateResource()
 	return Result;
 }
 
+uint32 UTextureRenderTarget2D::CalcTextureMemorySizeEnum(ETextureMipCount Enum) const
+{
+	// Calculate size based on format.  All mips are resident on render targets so we always return the same value.
+	EPixelFormat Format = GetFormat();
+	int32 BlockSizeX = GPixelFormats[Format].BlockSizeX;
+	int32 BlockSizeY = GPixelFormats[Format].BlockSizeY;
+	int32 BlockBytes = GPixelFormats[Format].BlockBytes;
+	int32 NumBlocksX = (SizeX + BlockSizeX - 1) / BlockSizeX;
+	int32 NumBlocksY = (SizeY + BlockSizeY - 1) / BlockSizeY;
+	int32 NumBytes = NumBlocksX * NumBlocksY * BlockBytes;
+	return NumBytes;
+}
+
 EMaterialValueType UTextureRenderTarget2D::GetMaterialType() const
 {
 	return MCT_Texture2D;
@@ -74,14 +87,7 @@ void UTextureRenderTarget2D::GetResourceSizeEx(FResourceSizeEx& CumulativeResour
 {
 	Super::GetResourceSizeEx(CumulativeResourceSize);
 
-	// Calculate size based on format.
-	EPixelFormat Format = GetFormat();
-	int32 BlockSizeX	= GPixelFormats[Format].BlockSizeX;
-	int32 BlockSizeY	= GPixelFormats[Format].BlockSizeY;
-	int32 BlockBytes	= GPixelFormats[Format].BlockBytes;
-	int32 NumBlocksX	= (SizeX + BlockSizeX - 1) / BlockSizeX;
-	int32 NumBlocksY	= (SizeY + BlockSizeY - 1) / BlockSizeY;
-	int32 NumBytes	= NumBlocksX * NumBlocksY * BlockBytes;
+	int32 NumBytes = CalcTextureMemorySizeEnum(TMC_AllMips);
 
 	CumulativeResourceSize.AddUnknownMemoryBytes(NumBytes);
 }
diff --git a/Engine/Source/Runtime/Engine/Private/UnrealEngine.cpp b/Engine/Source/Runtime/Engine/Private/UnrealEngine.cpp
--- a/Engine/Source/Runtime/Engine/Private/UnrealEngine.cpp
+++ b/Engine/Source/Runtime/Engine/Private/UnrealEngine.cpp
@@ -74,6 +74,7 @@ UnrealEngine.cpp: Implements the UEngine class and helpers.
 #include "Components/SkeletalMeshComponent.h"
 #include "Engine/Texture.h"
 #include "Engine/Texture2D.h"
+#include "Engine/TextureRenderTarget2D.h"
 #include "Engine/VolumeTexture.h"
 #include "ParticleHelper.h"
 #include "Particles/ParticleModule.h"
@@ -5967,6 +5968,7 @@ bool UEngine::HandleListTexturesCommand( const TCHAR* Cmd, FOutputDevice& Ar )
 		UTexture2D*			Texture2D			= Cast<UTexture2D>(Texture);
 		UTextureCube*		TextureCube			= Cast<UTextureCube>(Texture);
 		UVolumeTexture*		Texture3D			= Cast<UVolumeTexture>(Texture);
+		UTextureRenderTarget2D* RenderTexture	= Cast<UTextureRenderTarget2D>(Texture);
 
 		int32				LODGroup			= Texture->LODGroup;
 		int32				NumMips				= 0;
@@ -6045,6 +6047,21 @@ bool UEngine::HandleListTexturesCommand( const TCHAR* Cmd, FOutputDevice& Ar )
 			CurSizeY			= FMath::Max<int32>(Texture3D->GetSizeY() >> DroppedMips, 1);
 			bIsUncompressed		= Texture3D->IsUncompressed();
 		}
+		else if (RenderTexture != nullptr)
+		{
+			NumMips = RenderTexture->GetNumMips();
+			MaxResLODBias = RenderTexture->GetCachedLODBias();
+			MaxAllowedSizeX = FMath::Max<int32>(RenderTexture->SizeX >> MaxResLODBias, 1);
+			MaxAllowedSizeY = FMath::Max<int32>(RenderTexture->SizeY >> MaxResLODBias, 1);
+			Format = RenderTexture->GetFormat();
+			DroppedMips = MaxResLODBias;
+			CurSizeX = FMath::Max<int32>(RenderTexture->SizeX >> DroppedMips, 1);
+			CurSizeY = FMath::Max<int32>(RenderTexture->SizeY >> DroppedMips, 1);
+			bIsStreamingTexture = RenderTexture->GetStreamingIndex() != INDEX_NONE;
+			bIsForced = RenderTexture->ShouldMipLevelsBeForcedResident() && bIsStreamingTexture;
+			bIsVirtual = RenderTexture->IsCurrentlyVirtualTextured();
+			bIsUncompressed = RenderTexture->IsUncompressed();
+		}
 
 		if (bListUnused && UsageCount != 0)
 		{
