Fix non-editor-ticking skeletal meshes producing unexpected results post-compile

Ensure that we tick at least once with LOD0 after compilation

diff --git a/Engine/Source/Runtime/Engine/Private/Components/SkinnedMeshComponent.cpp b/Engine/Source/Runtime/Engine/Private/Components/SkinnedMeshComponent.cpp
--- a/Engine/Source/Runtime/Engine/Private/Components/SkinnedMeshComponent.cpp
+++ b/Engine/Source/Runtime/Engine/Private/Components/SkinnedMeshComponent.cpp
@@ -578,7 +578,25 @@ void USkinnedMeshComponent::OnRegister()
 		CachedSceneFeatureLevel = ERHIFeatureLevel::Num;
 	}
 
+#if WITH_EDITOR
+	// When we are reinstancing, ensure that the initial setup is done at LOD 0 so that non-ticking components dont
+	// get unexpected behavior when transitioning LODs post-compile.
+	const int32 OldForcedLOD = GetForcedLOD();
+	if(GIsReinstancing)
+	{
+		SetForcedLOD(1);
+	}
+#endif
+	
 	UpdateLODStatus();
+
+#if WITH_EDITOR
+	if(GIsReinstancing)
+	{
+		SetForcedLOD(OldForcedLOD);
+	}
+#endif
+
 	InvalidateCachedBounds();
 
 	UMeshDeformer* ActiveMeshDeformer = GetActiveMeshDeformer();
