Added wp.Runtime.DebugFilterOutContentBundles to filter content bundle from world partition debug display (on by default)

diff --git a/Engine/Source/Runtime/Engine/Private/WorldPartition/WorldPartitionDebugHelper.cpp b/Engine/Source/Runtime/Engine/Private/WorldPartition/WorldPartitionDebugHelper.cpp
--- a/Engine/Source/Runtime/Engine/Private/WorldPartition/WorldPartitionDebugHelper.cpp
+++ b/Engine/Source/Runtime/Engine/Private/WorldPartition/WorldPartitionDebugHelper.cpp
@@ -12,6 +12,12 @@
 #include "Algo/AnyOf.h"
 #include "Algo/AllOf.h"
 
+bool FWorldPartitionDebugHelper::bDebugFilterOutContentBundles = true;
+FAutoConsoleVariableRef FWorldPartitionDebugHelper::DebugFilterOutContentBundlesCommand(
+	TEXT("wp.Runtime.DebugFilterOutContentBundles"),
+	FWorldPartitionDebugHelper::bDebugFilterOutContentBundles,
+	TEXT("Filter content bundle from world partition debug display."));
+
 TSet<FName> FWorldPartitionDebugHelper::DebugRuntimeHashFilter;
 FAutoConsoleCommand FWorldPartitionDebugHelper::DebugFilterByRuntimeHashGridNameCommand(
 	TEXT("wp.Runtime.DebugFilterByRuntimeHashGridName"),
diff --git a/Engine/Source/Runtime/Engine/Private/WorldPartition/WorldPartitionRuntimeSpatialHash.cpp b/Engine/Source/Runtime/Engine/Private/WorldPartition/WorldPartitionRuntimeSpatialHash.cpp
--- a/Engine/Source/Runtime/Engine/Private/WorldPartition/WorldPartitionRuntimeSpatialHash.cpp
+++ b/Engine/Source/Runtime/Engine/Private/WorldPartition/WorldPartitionRuntimeSpatialHash.cpp
@@ -1439,20 +1439,23 @@ void UWorldPartitionRuntimeSpatialHash::ForEachStreamingGrid(TFunctionRef<void(F
 	}
 }
 
-void UWorldPartitionRuntimeSpatialHash::ForEachStreamingGrid(TFunctionRef<void(const FSpatialHashStreamingGrid&)> Func) const
+void UWorldPartitionRuntimeSpatialHash::ForEachStreamingGrid(TFunctionRef<void(const FSpatialHashStreamingGrid&)> Func, bool bIncludeExternalStreamingObjects) const
 {
 	for (const FSpatialHashStreamingGrid& StreamingGrid : StreamingGrids)
 	{
 		Func(StreamingGrid);
 	}
 
-	for (const TWeakObjectPtr<URuntimeSpatialHashExternalStreamingObject>& StreamingObject : ExternalStreamingObjects)
+	if (bIncludeExternalStreamingObjects)
 	{
-		if (StreamingObject.IsValid())
+		for (const TWeakObjectPtr<URuntimeSpatialHashExternalStreamingObject>& StreamingObject : ExternalStreamingObjects)
 		{
-			for (const FSpatialHashStreamingGrid& StreamingGrid : StreamingObject->StreamingGrids)
+			if (StreamingObject.IsValid())
 			{
-				Func(StreamingGrid);
+				for (const FSpatialHashStreamingGrid& StreamingGrid : StreamingObject->StreamingGrids)
+				{
+					Func(StreamingGrid);
+				}
 			}
 		}
 	}
@@ -1646,7 +1649,7 @@ TArray<const FSpatialHashStreamingGrid*> UWorldPartitionRuntimeSpatialHash::GetF
 				FilteredStreamingGrids.Add(&StreamingGrid);
 			}
 		}
-	});
+	}, !FWorldPartitionDebugHelper::GetDebugFilterOutContentBundles());
 
 	return FilteredStreamingGrids;
 }
diff --git a/Engine/Source/Runtime/Engine/Public/WorldPartition/WorldPartitionDebugHelper.h b/Engine/Source/Runtime/Engine/Public/WorldPartition/WorldPartitionDebugHelper.h
--- a/Engine/Source/Runtime/Engine/Public/WorldPartition/WorldPartitionDebugHelper.h
+++ b/Engine/Source/Runtime/Engine/Public/WorldPartition/WorldPartitionDebugHelper.h
@@ -24,6 +24,7 @@ public:
 	static void DrawLegendItem(UCanvas* Canvas, const FString& Text, const UFont* Font, const FColor& Color, const FColor& TextColor, FVector2D& Pos, float* MaxItemWidth = nullptr);
 	static FLinearColor GetHeatMapColor(float Value);
 	FORCEINLINE static bool IsRuntimeSpatialHashCellStreamingPriorityShown() { return bShowRuntimeSpatialHashCellStreamingPriority; }
+	FORCEINLINE static bool GetDebugFilterOutContentBundles() { return bDebugFilterOutContentBundles; }
 
 private:
 	static FAutoConsoleCommand DebugFilterByStreamingStatusCommand;
@@ -40,4 +41,7 @@ private:
 
 	static FAutoConsoleVariableRef ShowRuntimeSpatialHashCellStreamingPriorityCommand;
 	static bool bShowRuntimeSpatialHashCellStreamingPriority;
+
+	static FAutoConsoleVariableRef DebugFilterOutContentBundlesCommand;
+	static bool bDebugFilterOutContentBundles;
 };
\ No newline at end of file
diff --git a/Engine/Source/Runtime/Engine/Public/WorldPartition/WorldPartitionRuntimeSpatialHash.h b/Engine/Source/Runtime/Engine/Public/WorldPartition/WorldPartitionRuntimeSpatialHash.h
--- a/Engine/Source/Runtime/Engine/Public/WorldPartition/WorldPartitionRuntimeSpatialHash.h
+++ b/Engine/Source/Runtime/Engine/Public/WorldPartition/WorldPartitionRuntimeSpatialHash.h
@@ -271,7 +271,7 @@ public:
 
 protected:
 	void ForEachStreamingGrid(TFunctionRef<void(FSpatialHashStreamingGrid&)> Func);
-	void ForEachStreamingGrid(TFunctionRef<void(const FSpatialHashStreamingGrid&)> Func) const;
+	void ForEachStreamingGrid(TFunctionRef<void(const FSpatialHashStreamingGrid&)> Func, bool bIncludeExternalStreamingObjects = true) const;
 
 	virtual EWorldPartitionStreamingPerformance GetStreamingPerformanceForCell(const UWorldPartitionRuntimeCell* Cell) const override;
 
