World Partition Content Bundle Support: added content bundle GUID in actors and actor descriptors and fix unsaved actor descriptors PIE filtering.

#rnx

diff --git a/Engine/Source/Editor/DetailCustomizations/Private/ActorDetails.cpp b/Engine/Source/Editor/DetailCustomizations/Private/ActorDetails.cpp
--- a/Engine/Source/Editor/DetailCustomizations/Private/ActorDetails.cpp
+++ b/Engine/Source/Editor/DetailCustomizations/Private/ActorDetails.cpp
@@ -619,22 +619,45 @@ void FActorDetails::AddActorCategory( IDetailLayoutBuilder& DetailBuilder, const
 	{
 		if (AActor* Actor = GEditor->GetSelectedActors()->GetTop<AActor>())
 		{
-			const FText ActorGuidText = FText::FromString(Actor->GetActorGuid().ToString());
-			ActorCategory.AddCustomRow( LOCTEXT("ActorGuid", "ActorGuid") )
-				.NameContent()
-				[
-					SNew(STextBlock)
-					.Text(LOCTEXT("ActorGuid2", "Actor Guid"))
-					.ToolTipText(LOCTEXT("ActorGuid_ToolTip", "Actor Guid"))
-					.Font(IDetailLayoutBuilder::GetDetailFont())
-				]
-				.ValueContent()
-				[
-					SNew(STextBlock)
-						.Text(ActorGuidText)
+			if (Actor->GetActorGuid().IsValid())
+			{
+				const FText ActorGuidText = FText::FromString(Actor->GetActorGuid().ToString());
+				ActorCategory.AddCustomRow( LOCTEXT("ActorGuid", "ActorGuid") )
+					.NameContent()
+					[
+						SNew(STextBlock)
+						.Text(LOCTEXT("ActorGuid2", "Actor Guid"))
+						.ToolTipText(LOCTEXT("ActorGuid_ToolTip", "Actor Guid"))
 						.Font(IDetailLayoutBuilder::GetDetailFont())
-						.IsEnabled(false)
-				];
+					]
+					.ValueContent()
+					[
+						SNew(STextBlock)
+							.Text(ActorGuidText)
+							.Font(IDetailLayoutBuilder::GetDetailFont())
+							.IsEnabled(false)
+					];
+			}
+
+			if (Actor->GetContentBundleGuid().IsValid())
+			{
+				const FText ActorContentBundleGuidText = FText::FromString(Actor->GetContentBundleGuid().ToString());
+				ActorCategory.AddCustomRow( LOCTEXT("ContentBundleGuid", "ContentBundleGuid") )
+					.NameContent()
+					[
+						SNew(STextBlock)
+						.Text(LOCTEXT("ContentBundleGuid2", "Content Bundle Guid"))
+						.ToolTipText(LOCTEXT("ActorContentBundleGuid_ToolTip", "Actor Content BundleGuid"))
+						.Font(IDetailLayoutBuilder::GetDetailFont())
+					]
+					.ValueContent()
+					[
+						SNew(STextBlock)
+							.Text(ActorContentBundleGuidText)
+							.Font(IDetailLayoutBuilder::GetDetailFont())
+							.IsEnabled(false)
+					];
+			}
 		}
 	};
 
diff --git a/Engine/Source/Runtime/Core/Public/UObject/FortniteMainBranchObjectVersion.h b/Engine/Source/Runtime/Core/Public/UObject/FortniteMainBranchObjectVersion.h
--- a/Engine/Source/Runtime/Core/Public/UObject/FortniteMainBranchObjectVersion.h
+++ b/Engine/Source/Runtime/Core/Public/UObject/FortniteMainBranchObjectVersion.h
@@ -244,6 +244,9 @@ struct CORE_API FFortniteMainBranchObjectVersion
 		// Changed how world partition streaming cells are named
 		WorldPartitionStreamingCellsNamingShortened,
 
+		// Serialize ContentBundleGuid in WorldPartitionActorDesc
+		WorldPartitionActorDescSerializeContentBundleGuid,
+
 		// -----<new versions can be added above this line>-------------------------------------------------
 		VersionPlusOne,
 		LatestVersion = VersionPlusOne - 1
diff --git a/Engine/Source/Runtime/Engine/Classes/GameFramework/Actor.h b/Engine/Source/Runtime/Engine/Classes/GameFramework/Actor.h
--- a/Engine/Source/Runtime/Engine/Classes/GameFramework/Actor.h
+++ b/Engine/Source/Runtime/Engine/Classes/GameFramework/Actor.h
@@ -862,6 +862,12 @@ protected:
 	UPROPERTY(BluePrintReadOnly, AdvancedDisplay, Category=Actor, NonPIEDuplicateTransient, TextExportTransient, NonTransactional)
 	FGuid ActorGuid;
 
+	/**
+	 * The GUID for this actor's content bundle.
+	 */
+	UPROPERTY(BluePrintReadOnly, AdvancedDisplay, Category=Actor, TextExportTransient, NonTransactional)
+	FGuid ContentBundleGuid;
+
 	/** DataLayers the actor belongs to.*/
 	UPROPERTY(VisibleAnywhere, AdvancedDisplay, Category = DataLayers)
 	TArray<FActorDataLayer> DataLayers;
@@ -904,7 +910,6 @@ public:
 
 	FActorOnPackagingModeChanged OnPackagingModeChanged;
 
-
 	/** Returns this actor's current target runtime grid. */
 	virtual FName GetRuntimeGrid() const { return RuntimeGrid; }
 
@@ -914,9 +919,12 @@ public:
 	/** Gets the property name for RuntimeGrid. */
 	static const FName GetRuntimeGridPropertyName()	{ return GET_MEMBER_NAME_CHECKED(AActor, RuntimeGrid); }
 
-	/** Returns this actor's current Guid. Actor Guids are only available in development builds. */
+	/** Returns this actor's Guid. Actor Guids are only available in editor builds. */
 	inline const FGuid& GetActorGuid() const { return ActorGuid; }
 
+	/** Returns this actor's content bundle Guid. */
+	inline const FGuid& GetContentBundleGuid() const { return ContentBundleGuid; }
+
 	/** Returns true if actor location should be locked. */
 	virtual bool IsLockLocation() const { return bLockLocation; }
 
@@ -3844,6 +3852,7 @@ private:
 
 	friend struct FSetActorHiddenInSceneOutliner;
 	friend struct FSetActorGuid;
+	friend struct FSetActorContentBundleGuid;
 	friend struct FSetActorSelectable;
 #endif
 
@@ -4099,6 +4108,16 @@ private:
 	friend class UExternalActorsCommandlet;
 	friend class UWorldPartitionConvertCommandlet;
 };
+
+struct FSetActorContentBundleGuid
+{
+private:
+	FSetActorContentBundleGuid(AActor* InActor, const FGuid& InContentBundleGuid)
+	{
+		InActor->ContentBundleGuid = InContentBundleGuid;
+	}
+	friend class FContentBundleEditor;
+};
 #endif
 
 /** Helper function for executing tick functions based on the normal conditions previous found in UActorComponent::ConditionalTick */
diff --git a/Engine/Source/Runtime/Engine/Private/Actor.cpp b/Engine/Source/Runtime/Engine/Private/Actor.cpp
--- a/Engine/Source/Runtime/Engine/Private/Actor.cpp
+++ b/Engine/Source/Runtime/Engine/Private/Actor.cpp
@@ -20,6 +20,7 @@
 #include "UObject/UE5MainStreamObjectVersion.h"
 #include "UObject/UE5ReleaseStreamObjectVersion.h"
 #include "UObject/UE5PrivateFrostyStreamObjectVersion.h"
+#include "UObject/FortniteMainBranchObjectVersion.h"
 #include "Kismet/GameplayStatics.h"
 #include "Engine/BlueprintGeneratedClass.h"
 #include "Engine/LocalPlayer.h"
@@ -75,6 +76,7 @@
 #include "WorldPartition/DataLayer/DataLayerInstance.h"
 #include "WorldPartition/DataLayer/DataLayerAsset.h"
 #include "WorldPartition/DataLayer/DataLayerSubsystem.h"
+#include "WorldPartition/ContentBundle/ContentBundlePaths.h"
 
 DEFINE_LOG_CATEGORY(LogActor);
 
@@ -762,6 +764,7 @@ void AActor::Serialize(FArchive& Ar)
 	Ar.UsingCustomVersion(FUE5MainStreamObjectVersion::GUID);
 	Ar.UsingCustomVersion(FUE5PrivateFrostyStreamObjectVersion::GUID);
 	Ar.UsingCustomVersion(FUE5ReleaseStreamObjectVersion::GUID);
+	Ar.UsingCustomVersion(FFortniteMainBranchObjectVersion::GUID);
 
 #if WITH_EDITOR
 	// Prior to load, map natively-constructed component instances for Blueprint-generated class types to any serialized properties that might reference them.
@@ -869,6 +872,11 @@ void AActor::Serialize(FArchive& Ar)
 		{
 			bIsSpatiallyLoaded = GetClass()->GetDefaultObject<AActor>()->bIsSpatiallyLoaded;
 		}
+
+		if (Ar.CustomVer(FFortniteMainBranchObjectVersion::GUID) < FFortniteMainBranchObjectVersion::WorldPartitionActorDescSerializeContentBundleGuid)
+		{
+			ContentBundleGuid = ContentBundlePaths::GetGuidFromPath(GetPackage()->GetFName().ToString());
+		}
 	}
 #endif
 }
diff --git a/Engine/Source/Runtime/Engine/Private/WorldPartition/ActorDescContainer.cpp b/Engine/Source/Runtime/Engine/Private/WorldPartition/ActorDescContainer.cpp
--- a/Engine/Source/Runtime/Engine/Private/WorldPartition/ActorDescContainer.cpp
+++ b/Engine/Source/Runtime/Engine/Private/WorldPartition/ActorDescContainer.cpp
@@ -112,9 +112,13 @@ FString UActorDescContainer::GetExternalActorPath() const
 
 bool UActorDescContainer::IsActorDescHandled(const AActor* Actor) const
 {
-	const FString ActorPackageName = Actor->GetPackage()->GetName();
-	const FString ExternalActorPath = GetExternalActorPath() / TEXT("");
-	return ActorPackageName.StartsWith(ExternalActorPath);
+	if (Actor->GetContentBundleGuid() == GetContentBundleGuid())
+	{
+		const FString ActorPackageName = Actor->GetPackage()->GetName();
+		const FString ExternalActorPath = GetExternalActorPath() / TEXT("");
+		return ActorPackageName.StartsWith(ExternalActorPath);
+	}
+	return false;
 }
 
 bool UActorDescContainer::IsMainPartitionContainer() const
diff --git a/Engine/Source/Runtime/Engine/Private/WorldPartition/ContentBundle/ContentBundleEditor.cpp b/Engine/Source/Runtime/Engine/Private/WorldPartition/ContentBundle/ContentBundleEditor.cpp
--- a/Engine/Source/Runtime/Engine/Private/WorldPartition/ContentBundle/ContentBundleEditor.cpp
+++ b/Engine/Source/Runtime/Engine/Private/WorldPartition/ContentBundle/ContentBundleEditor.cpp
@@ -7,26 +7,23 @@
 #include "WorldPartition/WorldPartition.h"
 #include "WorldPartition/DataLayer/WorldDataLayers.h"
 #include "WorldPartition/WorldPartitionRuntimeCell.h"
-#include "Engine/World.h"
 #include "WorldPartition/WorldPartitionRuntimeHash.h"
 #include "WorldPartition/ContentBundle/ContentBundleDescriptor.h"
 #include "WorldPartition/ContentBundle/ContentBundleEditorSubsystemInterface.h"
+#include "WorldPartition/ContentBundle/ContentBundleWorldSubsystem.h"
 #include "WorldPartition/ContentBundle/ContentBundleLog.h"
+#include "Engine/World.h"
 #include "PackageTools.h"
 #include "ObjectTools.h"
 #include "Editor.h"
-#include "WorldPartition/ContentBundle/ContentBundleWorldSubsystem.h"
 
 FContentBundleEditor::FContentBundleEditor(TSharedPtr<FContentBundleClient>& InClient, UWorld* InWorld)
-	:FContentBundleBase(InClient, InWorld),
-	UnsavedActorMonitor(nullptr),
-	ExternalStreamingObject(nullptr),
-	Guid(FGuid::NewGuid()),
-	bIsBeingEdited(false)
-{
-
-}
-
+	: FContentBundleBase(InClient, InWorld)
+	, UnsavedActorMonitor(nullptr)
+	, ExternalStreamingObject(nullptr)
+	, Guid(FGuid::NewGuid())
+	, bIsBeingEdited(false)
+{}
 
 void FContentBundleEditor::DoInitialize()
 {
@@ -62,6 +59,9 @@ void FContentBundleEditor::DoInjectContent()
 		{
 			UE_LOG(LogContentBundle, Log, TEXT("[CB: %s] ExternalActors in %s found. %u actors were injected"), *GetDescriptor()->GetDisplayName(), *ActorDescContainer->GetExternalActorPath(), ActorDescContainer->GetActorDescCount());
 
+			check(GetDescriptor()->GetGuid().IsValid());
+			ActorDescContainer->SetContentBundleGuid(GetDescriptor()->GetGuid());
+
 			if (!ActorDescContainer->IsEmpty())
 			{
 				WorldDataLayersActorReference = FWorldDataLayersReference(ActorDescContainer.Get(), BuildWorlDataLayersName());
@@ -149,12 +149,14 @@ bool FContentBundleEditor::AddActor(AActor* InActor)
 
 	check(GetStatus() == EContentBundleStatus::ContentInjected);
 
+	// Assign the container guid to the actor
+	check(!InActor->GetContentBundleGuid().IsValid());
+	check(GetDescriptor()->GetGuid().IsValid());
+	FSetActorContentBundleGuid SetActorContentBundleGuid(InActor, GetDescriptor()->GetGuid());
+
+	// Rename the actor so it is saved in the content bundle location
 	FName ActorPackageNameInContentBundle(*ULevel::GetActorPackageName(ActorDescContainer->GetExternalActorPath(), EActorPackagingScheme::Reduced, InActor->GetName()));
-	FName ActorPackageName = InActor->GetPackage()->GetFName();
-	if (ActorPackageName != ActorPackageNameInContentBundle)
-	{
-		InActor->GetPackage()->Rename(*ActorPackageNameInContentBundle.ToString());
-	}
+	verify(InActor->GetPackage()->Rename(*ActorPackageNameInContentBundle.ToString()));
 
 	UnsavedActorMonitor->MonitorActor(InActor);
 
diff --git a/Engine/Source/Runtime/Engine/Private/WorldPartition/ContentBundle/ContentBundlePaths.cpp b/Engine/Source/Runtime/Engine/Private/WorldPartition/ContentBundle/ContentBundlePaths.cpp
--- a/Engine/Source/Runtime/Engine/Private/WorldPartition/ContentBundle/ContentBundlePaths.cpp
+++ b/Engine/Source/Runtime/Engine/Private/WorldPartition/ContentBundle/ContentBundlePaths.cpp
@@ -9,10 +9,10 @@
 
 FStringView ContentBundlePaths::GetRelativeExternalActorPackagePath(FStringView InContentBundleExternalActorPackagePath)
 {
-	uint32 ContentBundleStartIdx = InContentBundleExternalActorPackagePath.Find(ContentBundlePaths::ContentBundleFolder, ESearchCase::IgnoreCase);
+	uint32 ContentBundleStartIdx = UE::String::FindFirst(InContentBundleExternalActorPackagePath, ContentBundleFolder, ESearchCase::IgnoreCase);
 	if (ContentBundleStartIdx != INDEX_NONE)
 	{
-		uint32 ContentBundleUIDIdx = ContentBundleStartIdx + FCString::Strlen(ContentBundlePaths::ContentBundleFolder) + 1; // + 1 for trailing slash
+		uint32 ContentBundleUIDIdx = ContentBundleStartIdx + FCString::Strlen(ContentBundleFolder) + 1; // + 1 for trailing slash
 		FStringView RelativeActorPackge = InContentBundleExternalActorPackagePath.RightChop(ContentBundleUIDIdx);
 		
 		uint32 ContentBundleEndIdx = RelativeActorPackge.Find(TEXT("/"), ESearchCase::IgnoreCase); // Search for the end of the content bundle guids
@@ -24,4 +24,21 @@ FStringView ContentBundlePaths::GetRelativeExternalActorPackagePath(FStringView
 	return FStringView();
 }
 
+FGuid ContentBundlePaths::GetGuidFromPath(FStringView InPath)
+{
+	FGuid Result;
+
+	static const TCHAR* ContentBundlerFolderPath = TEXT("/ContentBundle/");
+	if (int FoundIdx = UE::String::FindFirst(InPath, ContentBundlerFolderPath, ESearchCase::IgnoreCase); FoundIdx != INDEX_NONE)
+	{
+		FString GuidPath(InPath);
+		const uint32 StartUIDIdx = FoundIdx + FCString::Strlen(ContentBundlerFolderPath);
+		const uint32 EndUIDIdx = GuidPath.Find(TEXT("/"), ESearchCase::IgnoreCase, ESearchDir::FromStart, StartUIDIdx);
+		GuidPath.MidInline(StartUIDIdx, EndUIDIdx - StartUIDIdx);		
+		verify(FGuid::Parse(GuidPath, Result));
+	}
+
+	return Result;
+}
+
 #endif
\ No newline at end of file
diff --git a/Engine/Source/Runtime/Engine/Private/WorldPartition/WorldPartitionActorDesc.cpp b/Engine/Source/Runtime/Engine/Private/WorldPartition/WorldPartitionActorDesc.cpp
--- a/Engine/Source/Runtime/Engine/Private/WorldPartition/WorldPartitionActorDesc.cpp
+++ b/Engine/Source/Runtime/Engine/Private/WorldPartition/WorldPartitionActorDesc.cpp
@@ -16,6 +16,7 @@
 #include "UObject/UE5MainStreamObjectVersion.h"
 #include "UObject/UE5ReleaseStreamObjectVersion.h"
 #include "UObject/FortniteNCBranchObjectVersion.h"
+#include "UObject/FortniteMainBranchObjectVersion.h"
 #include "WorldPartition/WorldPartitionLog.h"
 #include "WorldPartition/ActorDescContainer.h"
 #include "WorldPartition/HLOD/HLODLayer.h"
@@ -23,6 +24,7 @@
 #include "WorldPartition/DataLayer/DataLayerSubsystem.h"
 #include "WorldPartition/DataLayer/DataLayerAsset.h"
 #include "WorldPartition/DataLayer/DataLayerUtils.h"
+#include "WorldPartition/ContentBundle/ContentBundlePaths.h"
 #include "WorldPartition/ErrorHandling/WorldPartitionStreamingGenerationErrorHandler.h"
 #include "ActorReferencesUtils.h"
 
@@ -115,6 +117,8 @@ void FWorldPartitionActorDesc::Init(const AActor* InActor)
 	{
 		ParentActor = AttachParentActor->GetActorGuid();
 	}
+
+	ContentBundleGuid = InActor->GetContentBundleGuid();
 	
 	TArray<AActor*> ActorReferences = ActorsReferencesUtils::GetExternalActorReferences(const_cast<AActor*>(InActor));
 
@@ -315,6 +319,7 @@ void FWorldPartitionActorDesc::Serialize(FArchive& Ar)
 	Ar.UsingCustomVersion(FUE5MainStreamObjectVersion::GUID);
 	Ar.UsingCustomVersion(FUE5ReleaseStreamObjectVersion::GUID);
 	Ar.UsingCustomVersion(FFortniteNCBranchObjectVersion::GUID);
+	Ar.UsingCustomVersion(FFortniteMainBranchObjectVersion::GUID);
 
 	if (Ar.CustomVer(FFortniteNCBranchObjectVersion::GUID) >= FFortniteNCBranchObjectVersion::WorldPartitionActorDescNativeBaseClassSerialization)
 	{
@@ -456,6 +461,15 @@ PRAGMA_ENABLE_DEPRECATION_WARNINGS
 	{
 		Ar << Properties;
 	}
+
+	if (Ar.CustomVer(FFortniteMainBranchObjectVersion::GUID) < FFortniteMainBranchObjectVersion::WorldPartitionActorDescSerializeContentBundleGuid)
+	{
+		ContentBundleGuid = ContentBundlePaths::GetGuidFromPath(ActorPackage.ToString());
+	}
+	else
+	{
+		Ar << ContentBundleGuid;
+	}
 }
 
 FBox FWorldPartitionActorDesc::GetBounds() const
@@ -492,24 +506,9 @@ FName FWorldPartitionActorDesc::GetDisplayClassName() const
 	return BaseClass.IsNull() ? GetCleanClassName(NativeClass) : GetCleanClassName(BaseClass);
 }
 
-FGuid FWorldPartitionActorDesc::GetContentBundleUID() const
+FGuid FWorldPartitionActorDesc::GetContentBundleGuid() const
 {
-	constexpr TCHAR ContentBundleFolder[] = TEXT("/ContentBundle/");
-	FString ActorPackageString = ActorPackage.ToString();
-	int FoundIdx = ActorPackageString.Find(ContentBundleFolder, ESearchCase::IgnoreCase);
-	if (FoundIdx != INDEX_NONE)
-	{
-		uint32 StartUIDIdx = FoundIdx + FCString::Strlen(ContentBundleFolder);
-		uint32 EndUIDIdx = ActorPackageString.Find(TEXT("/"), ESearchCase::IgnoreCase, ESearchDir::FromStart, StartUIDIdx);
-		ActorPackageString.MidInline(StartUIDIdx, EndUIDIdx - StartUIDIdx);
-		
-		FGuid ContentBundleUID;
-		FGuid::Parse(ActorPackageString, ContentBundleUID);
-
-		return ContentBundleUID;
-	}
-
-	return FGuid();
+	return ContentBundleGuid;
 }
 
 void FWorldPartitionActorDesc::CheckForErrors(IStreamingGenerationErrorHandler* ErrorHandler) const
diff --git a/Engine/Source/Runtime/Engine/Private/WorldPartition/WorldPartitionActorDescView.cpp b/Engine/Source/Runtime/Engine/Private/WorldPartition/WorldPartitionActorDescView.cpp
--- a/Engine/Source/Runtime/Engine/Private/WorldPartition/WorldPartitionActorDescView.cpp
+++ b/Engine/Source/Runtime/Engine/Private/WorldPartition/WorldPartitionActorDescView.cpp
@@ -140,9 +140,9 @@ const FGuid& FWorldPartitionActorDescView::GetFolderGuid() const
 	return ActorDesc->GetFolderGuid();
 }
 
-FGuid FWorldPartitionActorDescView::GetContentBundleUID() const
+FGuid FWorldPartitionActorDescView::GetContentBundleGuid() const
 {
-	return ActorDesc->GetContentBundleUID();
+	return ActorDesc->GetContentBundleGuid();
 }
 
 bool FWorldPartitionActorDescView::IsContainerInstance() const
diff --git a/Engine/Source/Runtime/Engine/Private/WorldPartition/WorldPartitionRuntimeLevelStreamingCell.cpp b/Engine/Source/Runtime/Engine/Private/WorldPartition/WorldPartitionRuntimeLevelStreamingCell.cpp
--- a/Engine/Source/Runtime/Engine/Private/WorldPartition/WorldPartitionRuntimeLevelStreamingCell.cpp
+++ b/Engine/Source/Runtime/Engine/Private/WorldPartition/WorldPartitionRuntimeLevelStreamingCell.cpp
@@ -103,7 +103,7 @@ void UWorldPartitionRuntimeLevelStreamingCell::AddActorToCell(const FWorldPartit
 	check(!ActorDescView.GetActorIsEditorOnly());
 	// Leaving this using deprecated functions until the serialization format of Packages is updated.
 PRAGMA_DISABLE_DEPRECATION_WARNINGS
-	Packages.Emplace(ActorDescView.GetActorPackage(), ActorDescView.GetActorPath(), InContainerID, InContainerTransform, InContainer->GetContainerPackage(), GetWorld()->GetPackage()->GetFName(), ActorDescView.GetContentBundleUID());
+	Packages.Emplace(ActorDescView.GetActorPackage(), ActorDescView.GetActorPath(), InContainerID, InContainerTransform, InContainer->GetContainerPackage(), GetWorld()->GetPackage()->GetFName(), ActorDescView.GetContentBundleGuid());
 PRAGMA_ENABLE_DEPRECATION_WARNINGS
 }
 
diff --git a/Engine/Source/Runtime/Engine/Private/WorldPartition/WorldPartitionStreamingGeneration.cpp b/Engine/Source/Runtime/Engine/Private/WorldPartition/WorldPartitionStreamingGeneration.cpp
--- a/Engine/Source/Runtime/Engine/Private/WorldPartition/WorldPartitionStreamingGeneration.cpp
+++ b/Engine/Source/Runtime/Engine/Private/WorldPartition/WorldPartitionStreamingGeneration.cpp
@@ -117,7 +117,7 @@ class FWorldPartitionStreamingGenerator
 					ActorSetInstance.ActorSet = &ActorSet;
 					ActorSetInstance.ContainerID = ContainerInstanceDescriptor.ID;
 					ActorSetInstance.Transform = ContainerInstanceDescriptor.Transform;
-					ActorSetInstance.ContentBundleID = ReferenceActorDescView.GetContentBundleUID();
+					ActorSetInstance.ContentBundleID = ReferenceActorDescView.GetContentBundleGuid();
 
 					if (ContainerInstanceDescriptor.ID.IsMainContainer())
 					{
@@ -341,8 +341,7 @@ class FWorldPartitionStreamingGenerator
 		{
 			for (AActor* Actor : InContainer->GetWorld()->PersistentLevel->Actors)
 			{
-				// @todo_ow: Fix PIE in unsaved maps (IsActorDescHandled always returns false for new unsaved maps)
-				if (IsValid(Actor) && Actor->IsPackageExternal() && Actor->IsMainPackageActor() && !Actor->IsEditorOnly() && InContainer->IsActorDescHandled(Actor) && !InContainer->GetActorDesc(Actor->GetActorGuid()))
+				if (IsValid(Actor) && Actor->IsPackageExternal() && Actor->IsMainPackageActor() && !Actor->IsEditorOnly() && (Actor->GetContentBundleGuid() == InContainer->GetContentBundleGuid()) && !InContainer->GetActorDesc(Actor->GetActorGuid()))
 				{
 					FWorldPartitionActorDescView ModifiedActorDescView = GetModifiedActorDesc(Actor, InContainer);
 					RegisterActorDescView(Actor->GetActorGuid(), ModifiedActorDescView);
diff --git a/Engine/Source/Runtime/Engine/Public/WorldPartition/ActorDescContainer.h b/Engine/Source/Runtime/Engine/Public/WorldPartition/ActorDescContainer.h
--- a/Engine/Source/Runtime/Engine/Public/WorldPartition/ActorDescContainer.h
+++ b/Engine/Source/Runtime/Engine/Public/WorldPartition/ActorDescContainer.h
@@ -32,8 +32,11 @@ public:
 	void OnObjectsReplaced(const TMap<UObject*, UObject*>& OldToNewObjectMap);
 
 	FName GetContainerPackage() const { return ContainerPackageName; }
-	void SetContainerPackage(const FName& InContainerPackageName) { ContainerPackageName = InContainerPackageName;	}
-	
+	void SetContainerPackage(const FName& InContainerPackageName) { ContainerPackageName = InContainerPackageName; }
+
+	FGuid GetContentBundleGuid() const { return ContentBundleGuid; }
+	void SetContentBundleGuid(const FGuid& InGetContentBundleGuid) { ContentBundleGuid = InGetContentBundleGuid; }
+
 	bool IsMainPartitionContainer() const;
 
 	FString GetExternalActorPath() const;
@@ -85,6 +88,7 @@ protected:
 	bool bContainerInitialized;
 
 	FName ContainerPackageName;
+	FGuid ContentBundleGuid;
 #endif
 
 private:
diff --git a/Engine/Source/Runtime/Engine/Public/WorldPartition/ContentBundle/ContentBundlePaths.h b/Engine/Source/Runtime/Engine/Public/WorldPartition/ContentBundle/ContentBundlePaths.h
--- a/Engine/Source/Runtime/Engine/Public/WorldPartition/ContentBundle/ContentBundlePaths.h
+++ b/Engine/Source/Runtime/Engine/Public/WorldPartition/ContentBundle/ContentBundlePaths.h
@@ -12,5 +12,7 @@ namespace ContentBundlePaths
 	// InContentBundleExternalActorPackagePath format is : //{Plugin}/{ExternalActorFolder}/{ContentBundleFolder}/{ContentBundleUID}/{LevelPath}/{ExternalActorPackagePath}
 	// return format is {LevelPath}/{ExternalActorPackagePath}, empty otherwise
 	ENGINE_API FStringView GetRelativeExternalActorPackagePath(FStringView InContentBundleExternalActorPackagePath);
+
+	ENGINE_API FGuid GetGuidFromPath(FStringView InPath);
 #endif
 }
\ No newline at end of file
diff --git a/Engine/Source/Runtime/Engine/Public/WorldPartition/WorldPartitionActorDesc.h b/Engine/Source/Runtime/Engine/Public/WorldPartition/WorldPartitionActorDesc.h
--- a/Engine/Source/Runtime/Engine/Public/WorldPartition/WorldPartitionActorDesc.h
+++ b/Engine/Source/Runtime/Engine/Public/WorldPartition/WorldPartitionActorDesc.h
@@ -120,7 +120,7 @@ PRAGMA_ENABLE_DEPRECATION_WARNINGS
 	virtual FName GetLevelPackage() const { return NAME_None; }
 	virtual bool GetContainerInstance(const UActorDescContainer*& OutLevelContainer, FTransform& OutLevelTransform, EContainerClusterMode& OutClusterMode) const { return false; }
 
-	FGuid GetContentBundleUID() const;
+	FGuid GetContentBundleGuid() const;
 
 	virtual const FGuid& GetSceneOutlinerParent() const { return GetParentActor(); }
 	virtual bool IsResaveNeeded() const { return false; }
@@ -258,8 +258,8 @@ protected:
 	FPropertyPairsMap				Properties;
 	FName							FolderPath;
 	FGuid							FolderGuid;
-
 	FGuid							ParentActor; // Used to validate settings against parent (to warn on layer/placement compatibility issues)
+	FGuid							ContentBundleGuid;
 	
 	// Transient
 	mutable uint32					SoftRefCount;
diff --git a/Engine/Source/Runtime/Engine/Public/WorldPartition/WorldPartitionActorDescView.h b/Engine/Source/Runtime/Engine/Public/WorldPartition/WorldPartitionActorDescView.h
--- a/Engine/Source/Runtime/Engine/Public/WorldPartition/WorldPartitionActorDescView.h
+++ b/Engine/Source/Runtime/Engine/Public/WorldPartition/WorldPartitionActorDescView.h
@@ -55,7 +55,7 @@ PRAGMA_ENABLE_DEPRECATION_WARNINGS
 	FName GetActorName() const;
 	const FGuid& GetFolderGuid() const;
 
-	FGuid GetContentBundleUID() const;
+	FGuid GetContentBundleGuid() const;
 
 	bool IsContainerInstance() const;
 	bool GetContainerInstance(const UActorDescContainer*& OutLevelContainer, FTransform& OutLevelTransform, EContainerClusterMode& OutClusterMode) const;
