World Partition: don't cluster spatially loaded actors referencing non-spatially loaded ones together.

#rnx

diff --git a/Engine/Source/Runtime/Engine/Private/WorldPartition/WorldPartitionStreamingGeneration.cpp b/Engine/Source/Runtime/Engine/Private/WorldPartition/WorldPartitionStreamingGeneration.cpp
--- a/Engine/Source/Runtime/Engine/Private/WorldPartition/WorldPartitionStreamingGeneration.cpp
+++ b/Engine/Source/Runtime/Engine/Private/WorldPartition/WorldPartitionStreamingGeneration.cpp
@@ -111,6 +111,15 @@ class FWorldPartitionStreamingGenerator
 				{
 					const FWorldPartitionActorDescView& ReferenceActorDescView = ActorSetContainer.ActorDescViewMap->FindByGuidChecked(ActorSet.Actors[0]);
 
+					// Validate assumptions
+					for (const FGuid& ActorGuid : ActorSet.Actors)
+					{
+						const FWorldPartitionActorDescView& ActorDescView = ActorSetContainer.ActorDescViewMap->FindByGuidChecked(ActorGuid);						
+						check(ActorDescView.GetRuntimeGrid() == ReferenceActorDescView.GetRuntimeGrid());
+						check(ActorDescView.GetIsSpatiallyLoaded() == ReferenceActorDescView.GetIsSpatiallyLoaded());
+						check(ActorDescView.GetContentBundleGuid() == ReferenceActorDescView.GetContentBundleGuid());
+					}
+
 					FActorSetInstance& ActorSetInstance = ActorSetInstances.Emplace_GetRef();
 				
 					ActorSetInstance.ContainerInstance = &ActorSetContainer;
@@ -728,7 +737,28 @@ class FWorldPartitionStreamingGenerator
 	{
 		// Build clusters for this container - at this point, all actors references should be in the same data layers, grid, etc because of actor descriptors validation.
 		TArray<TPair<FGuid, TArray<FGuid>>> ActorsWithRefs;
-		ContainerDescriptor.ActorDescViewMap.ForEachActorDescView([&ActorsWithRefs](const FWorldPartitionActorDescView& ActorDescView) { ActorsWithRefs.Emplace(ActorDescView.GetGuid(), ActorDescView.GetReferences()); });
+		ContainerDescriptor.ActorDescViewMap.ForEachActorDescView([&ContainerDescriptor, &ActorsWithRefs](const FWorldPartitionActorDescView& ActorDescView)
+		{
+			TArray<FGuid> ActorReferences;
+			ActorReferences.Reserve(ActorDescView.GetReferences().Num());
+
+			for (const FGuid& ActorReferenceGuid : ActorDescView.GetReferences())
+			{
+				const FWorldPartitionActorDescView& ReferenceActorDescView = ContainerDescriptor.ActorDescViewMap.FindByGuidChecked(ActorReferenceGuid);
+
+				const bool bIsActorDescSpatiallyLoaded = ActorDescView.GetIsSpatiallyLoaded();
+				const bool bIsActorDescRefSpatiallyLoaded = ReferenceActorDescView.GetIsSpatiallyLoaded();
+
+				// The only case we support right now is spatially loaded actors referencing non-spatially loaded actors, when target is not in data layers.
+				// For this to work with data layers, we need to implement dependency logic support in the content cooker splitter.
+				if (!bIsActorDescSpatiallyLoaded || bIsActorDescRefSpatiallyLoaded || !ReferenceActorDescView.GetDataLayers().IsEmpty())
+				{
+					ActorReferences.Add(ActorReferenceGuid);
+				}
+			}
+			
+			ActorsWithRefs.Emplace(ActorDescView.GetGuid(), MoveTemp(ActorReferences));
+		});
 		ContainerDescriptor.Clusters = GenerateObjectsClusters(ActorsWithRefs);
 	}
 
