Shader Model 6 warning dialogs
- If a static mesh is detected to use Nanite or a project is configured to use Virtual Shadow Maps and SM6 is not in the project settings, show a warning that SM6 needs to be enabled.
- Both cases only warn if the project was already configured to use a D3D platform.
- Added a UI option for enabling Nanite in project settings.

diff --git a/Engine/Source/Editor/StaticMeshEditor/Private/StaticMeshEditorTools.cpp b/Engine/Source/Editor/StaticMeshEditor/Private/StaticMeshEditorTools.cpp
--- a/Engine/Source/Editor/StaticMeshEditor/Private/StaticMeshEditorTools.cpp
+++ b/Engine/Source/Editor/StaticMeshEditor/Private/StaticMeshEditorTools.cpp
@@ -5392,7 +5392,10 @@ void FNaniteSettingsLayout::ApplyChanges()
 
 		StaticMesh->Modify();
 		StaticMesh->NaniteSettings = NaniteSettings;
-		StaticMesh->PostEditChange();
+
+		FProperty* ChangedProperty = FindFProperty<FProperty>(UStaticMesh::StaticClass(), GET_MEMBER_NAME_CHECKED(UStaticMesh, NaniteSettings));
+		FPropertyChangedEvent Event(ChangedProperty);
+		StaticMesh->PostEditChangeProperty(Event);
 	}
 
 
diff --git a/Engine/Source/Runtime/Engine/Classes/Engine/RendererSettings.h b/Engine/Source/Runtime/Engine/Classes/Engine/RendererSettings.h
--- a/Engine/Source/Runtime/Engine/Classes/Engine/RendererSettings.h
+++ b/Engine/Source/Runtime/Engine/Classes/Engine/RendererSettings.h
@@ -562,6 +562,12 @@ class ENGINE_API URendererSettings : public UDeveloperSettings
 		ConfigRestartRequired=true))
 	float DistanceFieldVoxelDensity;
 
+	UPROPERTY(config, EditAnywhere, Category = Nanite, meta = (
+		ConsoleVariable = "r.Nanite.ProjectEnabled",
+		ToolTip = "Whether to enable Nanite rendering.",
+		ConfigRestartRequired = true))
+	uint32 bNanite : 1;
+
 	UPROPERTY(config, EditAnywhere, Category=MiscLighting, meta=(
 		ConsoleVariable="r.AllowStaticLighting",
 		ToolTip="Whether to allow any static lighting to be generated and used, like lightmaps and shadowmaps. Games that only use dynamic lighting should set this to 0 to save some static lighting overhead. Disabling allows Material AO and Material BentNormal to work with Lumen Global Illumination.  Changing this setting requires restarting the editor.",
@@ -1157,6 +1163,11 @@ private:
 #if WITH_EDITOR
 	/** shadow copy saved before effects of PostEditChange() to provide option to roll back edit. */
 	int32 PreEditReflectionCaptureResolution = 128;
+
+	// Notification about missing required shader models.
+	TWeakPtr<class SNotificationItem> ShaderModelNotificationPtr;
+
+	void CheckForMissingShaderModels();
 #endif // WITH_EDITOR
 
 	void SanatizeReflectionCaptureResolution();
diff --git a/Engine/Source/Runtime/Engine/Classes/Engine/StaticMesh.h b/Engine/Source/Runtime/Engine/Classes/Engine/StaticMesh.h
--- a/Engine/Source/Runtime/Engine/Classes/Engine/StaticMesh.h
+++ b/Engine/Source/Runtime/Engine/Classes/Engine/StaticMesh.h
@@ -2039,6 +2039,11 @@ private:
 	// Filled at CommitDescription time and reused during build
 	TOptional<FBoxSphereBounds> CachedMeshDescriptionBounds;
 
+	// Notification about missing Nanite required shader models.
+	TWeakPtr<class SNotificationItem> ShaderModelNotificationPtr;
+
+	void CheckForMissingShaderModels();
+
 	FOnPreMeshBuild PreMeshBuild;
 	FOnPostMeshBuild PostMeshBuild;
 
diff --git a/Engine/Source/Runtime/Engine/Private/RendererSettings.cpp b/Engine/Source/Runtime/Engine/Private/RendererSettings.cpp
--- a/Engine/Source/Runtime/Engine/Private/RendererSettings.cpp
+++ b/Engine/Source/Runtime/Engine/Private/RendererSettings.cpp
@@ -14,6 +14,10 @@
 #include "Misc/ConfigCacheIni.h"
 #include "HAL/PlatformFileManager.h"
 
+#include "Framework/Notifications/NotificationManager.h"
+#include "Widgets/Notifications/SNotificationList.h"
+#include "Framework/Docking/TabManager.h"
+
 /** The editor object. */
 extern UNREALED_API class UEditorEngine* GEditor;
 #endif // #if WITH_EDITOR
@@ -98,6 +102,8 @@ void URendererSettings::PostInitProperties()
 	UpdateWorkingColorSpaceAndChromaticities();
 
 #if WITH_EDITOR
+	CheckForMissingShaderModels();
+
 	if (IsTemplate())
 	{
 		ImportConsoleVariableValues();
@@ -287,6 +293,11 @@ void URendererSettings::PostEditChangeProperty(FPropertyChangedEvent& PropertyCh
 		{
 			UpdateWorkingColorSpaceAndChromaticities();
 		}
+
+		if (PropertyChangedEvent.GetPropertyName() == GET_MEMBER_NAME_CHECKED(URendererSettings, ShadowMapMethod))
+		{
+			CheckForMissingShaderModels();
+		}
 	}
 }
 
@@ -331,6 +342,57 @@ bool URendererSettings::CanEditChange(const FProperty* InProperty) const
 
 	return true;
 }
+
+void URendererSettings::CheckForMissingShaderModels()
+{
+	if (GIsEditor && ShadowMapMethod == EShadowMapMethod::VirtualShadowMaps)
+	{
+		TArray<FString> D3D11TargetedShaderFormats;
+		GConfig->GetArray(TEXT("/Script/WindowsTargetPlatform.WindowsTargetSettings"), TEXT("D3D11TargetedShaderFormats"), D3D11TargetedShaderFormats, GEngineIni);
+
+		TArray<FString> D3D12TargetedShaderFormats;
+		GConfig->GetArray(TEXT("/Script/WindowsTargetPlatform.WindowsTargetSettings"), TEXT("D3D12TargetedShaderFormats"), D3D12TargetedShaderFormats, GEngineIni);
+
+		const bool bProjectUsesD3D = (D3D11TargetedShaderFormats.Num() + D3D12TargetedShaderFormats.Num()) > 0;
+
+		if (bProjectUsesD3D && !D3D12TargetedShaderFormats.Contains(TEXT("PCD3D_SM6")))
+		{
+			auto DismissNotification = [this]()
+			{
+				if (TSharedPtr<SNotificationItem> NotificationPin = ShaderModelNotificationPtr.Pin())
+				{
+					NotificationPin->SetCompletionState(SNotificationItem::CS_None);
+					NotificationPin->ExpireAndFadeout();
+					ShaderModelNotificationPtr.Reset();
+				}
+			};
+
+			auto OpenProjectSettings = []()
+			{
+				FGlobalTabmanager::Get()->TryInvokeTab(FName("ProjectSettings"));
+			};
+
+			FNotificationInfo Info(LOCTEXT("NeedProjectSettings", "Missing Project Settings!"));
+			Info.bFireAndForget = false;
+			Info.FadeOutDuration = 0.0f;
+			Info.ExpireDuration = 0.0f;
+			Info.WidthOverride = FOptionalSize();
+
+			Info.ButtonDetails.Add(FNotificationButtonInfo(
+				LOCTEXT("GuidelineDismiss", "Dismiss"),
+				LOCTEXT("GuidelineDismissTT", "Dismiss this notification."),
+				FSimpleDelegate::CreateLambda(DismissNotification),
+				SNotificationItem::CS_None));
+
+			Info.Text = LOCTEXT("NeedProjectSettings", "Missing Project Settings!");
+			Info.SubText = LOCTEXT("VirtualShadowMapsNeedsSM6Setting", "Shader Model 6 (SM6) is required to use Virtual Shadow Maps. Please enable this in:\n  Project Settings -> Platforms -> Windows -> D3D12 Targeted Shader Formats\nVirtual shadow maps will not work until this is enabled.");
+			Info.HyperlinkText = LOCTEXT("ProjectSettingsHyperlinkText", "Open Project Settings");
+			Info.Hyperlink = FSimpleDelegate::CreateLambda(OpenProjectSettings);
+
+			ShaderModelNotificationPtr = FSlateNotificationManager::Get().AddNotification(Info);
+		}
+	}
+}
 #endif // #if WITH_EDITOR
 
 void URendererSettings::SanatizeReflectionCaptureResolution()
diff --git a/Engine/Source/Runtime/Engine/Private/StaticMesh.cpp b/Engine/Source/Runtime/Engine/Private/StaticMesh.cpp
--- a/Engine/Source/Runtime/Engine/Private/StaticMesh.cpp
+++ b/Engine/Source/Runtime/Engine/Private/StaticMesh.cpp
@@ -81,6 +81,9 @@
 #include "DerivedDataCacheRecord.h"
 #include "DerivedDataValue.h"
 #include "Compression/OodleDataCompression.h"
+
+#include "Widgets/Notifications/SNotificationList.h"
+#include "Framework/Docking/TabManager.h"
 #endif // #if WITH_EDITOR
 
 #include "Engine/StaticMeshSocket.h"
@@ -3837,6 +3840,11 @@ void UStaticMesh::PostEditChangeProperty(FPropertyChangedEvent& PropertyChangedE
 		}
 	}
 
+	if (PropertyName == GET_MEMBER_NAME_CHECKED(UStaticMesh, NaniteSettings))
+	{
+		CheckForMissingShaderModels();
+	}
+
 	UStaticMesh::FBuildParameters BuildParameters;
 	BuildParameters.bInSilent = true;
 	BuildParameters.bInRebuildUVChannelData = true;
@@ -5674,6 +5682,8 @@ void UStaticMesh::BeginPostLoadInternal(FStaticMeshPostLoadContext& Context)
 	// before trying to access it from multiple threads
 	GetStaticMeshDerivedDataVersion();
 
+	CheckForMissingShaderModels();
+
 	// Lock all properties that should not be modified/accessed during async post-load
 	AcquireAsyncProperty();
 
@@ -6198,6 +6208,62 @@ void UStaticMesh::ExecutePostLoadInternal(FStaticMeshPostLoadContext& Context)
 #endif // #if WITH_EDITOR
 }
 
+#if WITH_EDITOR
+void UStaticMesh::CheckForMissingShaderModels()
+{
+	static bool bWarnedAboutMissingShaderModel = false;
+	if (GIsEditor && NaniteSettings.bEnabled && !bWarnedAboutMissingShaderModel)
+	{
+		TArray<FString> D3D11TargetedShaderFormats;
+		GConfig->GetArray(TEXT("/Script/WindowsTargetPlatform.WindowsTargetSettings"), TEXT("D3D11TargetedShaderFormats"), D3D11TargetedShaderFormats, GEngineIni);
+
+		TArray<FString> D3D12TargetedShaderFormats;
+		GConfig->GetArray(TEXT("/Script/WindowsTargetPlatform.WindowsTargetSettings"), TEXT("D3D12TargetedShaderFormats"), D3D12TargetedShaderFormats, GEngineIni);
+
+		const bool bProjectUsesD3D = (D3D11TargetedShaderFormats.Num() + D3D12TargetedShaderFormats.Num()) > 0;
+
+		if (bProjectUsesD3D && !D3D12TargetedShaderFormats.Contains(TEXT("PCD3D_SM6")))
+		{
+			bWarnedAboutMissingShaderModel = true;
+
+			auto DismissNotification = [this]()
+			{
+				if (TSharedPtr<SNotificationItem> NotificationPin = ShaderModelNotificationPtr.Pin())
+				{
+					NotificationPin->SetCompletionState(SNotificationItem::CS_None);
+					NotificationPin->ExpireAndFadeout();
+					ShaderModelNotificationPtr.Reset();
+				}
+			};
+
+			auto OpenProjectSettings = []()
+			{
+				FGlobalTabmanager::Get()->TryInvokeTab(FName("ProjectSettings"));
+			};
+
+			FNotificationInfo Info(LOCTEXT("NeedProjectSettings", "Missing Project Settings!"));
+			Info.bFireAndForget = false;
+			Info.FadeOutDuration = 0.0f;
+			Info.ExpireDuration = 0.0f;
+			Info.WidthOverride = FOptionalSize();
+
+			Info.ButtonDetails.Add(FNotificationButtonInfo(
+				LOCTEXT("GuidelineDismiss", "Dismiss"),
+				LOCTEXT("GuidelineDismissTT", "Dismiss this notification."),
+				FSimpleDelegate::CreateLambda(DismissNotification),
+				SNotificationItem::CS_None));
+
+			Info.Text = LOCTEXT("NeedProjectSettings", "Missing Project Settings!");
+			Info.SubText = LOCTEXT("NaniteNeedsSM6Setting", "Shader Model 6 (SM6) is required to use Nanite assets. Please enable this in:\n  Project Settings -> Platforms -> Windows -> D3D12 Targeted Shader Formats\nNanite assets will not display properly until this is enabled.");
+			Info.HyperlinkText = LOCTEXT("ProjectSettingsHyperlinkText", "Open Project Settings");
+			Info.Hyperlink = FSimpleDelegate::CreateLambda(OpenProjectSettings);
+
+			ShaderModelNotificationPtr = FSlateNotificationManager::Get().AddNotification(Info);
+		}
+	}
+}
+#endif
+
 void UStaticMesh::FinishPostLoadInternal(FStaticMeshPostLoadContext& Context)
 {
 	TRACE_CPUPROFILER_EVENT_SCOPE(UStaticMesh::FinishPostLoad);
