Add support for adding textures with different pixels formats to a Texture2DArray.

diff --git a/Engine/Source/Runtime/Engine/Classes/Engine/Texture2DArray.cpp b/Engine/Source/Runtime/Engine/Classes/Engine/Texture2DArray.cpp
--- a/Engine/Source/Runtime/Engine/Classes/Engine/Texture2DArray.cpp
+++ b/Engine/Source/Runtime/Engine/Classes/Engine/Texture2DArray.cpp
@@ -23,6 +23,7 @@
 #include "Engine/Public/ImageUtils.h"
 #include "Misc/ScopedSlowTask.h"
 #include "UObject/StrongObjectPtr.h"
+#include "ImageCoreUtils.h"
 
 #define LOCTEXT_NAMESPACE "UTexture2DArray"
 
@@ -366,12 +367,6 @@ ENGINE_API bool UTexture2DArray::CheckArrayTexturesCompatibility()
 			UE_LOG(LogTexture, Warning, TEXT("Texture2DArray creation failed. Textures %s and %s have different aspect ratios."), *SourceTextures[TextureIndex]->GetFName().ToString(), *BaseTextureName);
 			return false;
 		}
-
-		if (SourceTextures[TextureIndex]->Source.GetFormat() != BaseTextureSource.GetFormat())
-		{
-			UE_LOG(LogTexture, Warning, TEXT("Texture2DArray creation failed. Textures %s and %s have incompatible pixel formats."), *SourceTextures[TextureIndex]->GetFName().ToString(), *BaseTextureName);
-			return false;
-		}
 	}
 
 	return true;
@@ -389,6 +384,10 @@ ENGINE_API bool UTexture2DArray::UpdateSourceFromSourceTextures(bool bCreatingNe
 	{
 		Modify();
 
+		ETextureSourceFormat SourceFormat = SourceTextures[0]->Source.Format;
+		bool bUseSRGBSource = SourceTextures[0]->Source.GetGammaSpace(0) == EGammaSpace::sRGB;
+		bool bMismatchedFormats = false;
+		bool bMismatchedGammaSpace = false;
 		int32 InitialSourceTextureIndex = 0;
 		for (int32 SourceTextureIndex = 1; SourceTextureIndex < SourceTextures.Num(); ++SourceTextureIndex)
 		{
@@ -398,12 +397,33 @@ ENGINE_API bool UTexture2DArray::UpdateSourceFromSourceTextures(bool bCreatingNe
 			{
 				InitialSourceTextureIndex = SourceTextureIndex;
 			}
+
+			if ((SourceTextures[SourceTextureIndex]->Source.GetGammaSpace(0) == EGammaSpace::sRGB) != bUseSRGBSource)
+			{
+				bMismatchedGammaSpace = true;
+				bUseSRGBSource = false;
+				SourceFormat = TSF_RGBA32F;
+			}
+
+			if (SourceFormat != SourceTextures[SourceTextureIndex]->Source.Format)
+			{
+				bMismatchedFormats = true;
+				SourceFormat = FImageCoreUtils::GetCommonSourceFormat(SourceFormat, SourceTextures[SourceTextureIndex]->Source.Format);
+			}
+		}
+
+		if (bMismatchedGammaSpace)
+		{
+			UE_LOG(LogTexture, Warning, TEXT("Mismatched source gamma spaces, converting all to RGBA32F/Linear ..."));
+		}
+		else if (bMismatchedFormats)
+		{
+			UE_LOG(LogTexture, Warning, TEXT("Mismatched source pixel formats, converting all to %s/%s ..."),
+				ERawImageFormat::GetName(FImageCoreUtils::ConvertToRawImageFormat(SourceFormat)), bUseSRGBSource ? TEXT("sRGB") : TEXT("Linear"));
 		}
 
 		FTextureSource& InitialSource = SourceTextures[InitialSourceTextureIndex]->Source;
-		// Format and format size.
-		ETextureSourceFormat Format = InitialSource.GetFormat();
-		int32 FormatDataSize = InitialSource.GetBytesPerPixel();
+		int32 FormatDataSize = FTextureSource::GetBytesPerPixel(SourceFormat);
 		// X,Y,Z size of the array. Switched to use Source sizes - previously would use PlatformData sizes which caused an error in headless commandlet MHC export.
 		int32 SizeX = InitialSource.GetSizeX();
 		int32 SizeY = InitialSource.GetSizeY();
@@ -416,13 +436,14 @@ ENGINE_API bool UTexture2DArray::UpdateSourceFromSourceTextures(bool bCreatingNe
 			MipGenSettings = TMGS_NoMipmaps;
 			PowerOfTwoMode = ETexturePowerOfTwoSetting::None;
 			LODGroup = SourceTextures[0]->LODGroup;
-			SRGB = SourceTextures[0]->SRGB;
 			NeverStream = true;
 		}
 
+		SRGB = bUseSRGBSource;
+
 		// Create the source texture for this Texture2DArray.
 		// Currently only single-mip Texture2DArray's are supported, therefore only the first mip is copied from each element source.
-		Source.Init(SizeX, SizeY, ArraySize, 1, Format);
+		Source.Init(SizeX, SizeY, ArraySize, 1, SourceFormat);
 
 		uint8* DestMipData = Source.LockMip(0);
 		int64 MipSizeBytes = Source.CalcMipSize(0) / ArraySize;
@@ -431,7 +452,7 @@ ENGINE_API bool UTexture2DArray::UpdateSourceFromSourceTextures(bool bCreatingNe
 		{
 			FTextureSource& TextureSource = SourceTextures[SourceTexIndex]->Source;
 			void* DestSliceData = DestMipData + MipSizeBytes * SourceTexIndex;
-			if (TextureSource.SizeX == SizeX && TextureSource.SizeY == SizeY)
+			if (TextureSource.SizeX == SizeX && TextureSource.SizeY == SizeY && TextureSource.Format == SourceFormat)
 			{
 				void* SourceData = TextureSource.LockMip(0);
 				check(TextureSource.CalcMipSize(0) == MipSizeBytes);
@@ -442,10 +463,24 @@ ENGINE_API bool UTexture2DArray::UpdateSourceFromSourceTextures(bool bCreatingNe
 			{
 				FImage SourceImage;
 				SourceTextures[SourceTexIndex]->Source.GetMipImage(SourceImage, 0);
-				FImage ResizedSourceImage;
-				SourceImage.ResizeTo(ResizedSourceImage, SizeX, SizeY, SourceImage.Format, SourceImage.GammaSpace);
-				check(ResizedSourceImage.RawData.Num() == MipSizeBytes);
-				FMemory::Memcpy(DestSliceData, ResizedSourceImage.RawData.GetData(), MipSizeBytes);
+
+				if (TextureSource.Format != SourceFormat)
+				{
+					// note: there is no need to check if the gamma space is different, because in such case we would fallback to TSF_RGBA32F, which is always linear
+					FImage ConvertedSourceImage;
+					SourceImage.CopyTo(ConvertedSourceImage, FImageCoreUtils::ConvertToRawImageFormat(SourceFormat), bUseSRGBSource ? EGammaSpace::sRGB : EGammaSpace::Linear);
+					ConvertedSourceImage.Swap(SourceImage);
+				}
+
+				if (TextureSource.SizeX != SizeX || TextureSource.SizeY != SizeY)
+				{
+					FImage ResizedSourceImage;
+					SourceImage.ResizeTo(ResizedSourceImage, SizeX, SizeY, SourceImage.Format, SourceImage.GammaSpace);
+					ResizedSourceImage.Swap(SourceImage);
+				}
+				
+				check(SourceImage.RawData.Num() == MipSizeBytes);
+				FMemory::Memcpy(DestSliceData, SourceImage.RawData.GetData(), MipSizeBytes);
 			}
 		}
 
diff --git a/Engine/Source/Runtime/ImageCore/Private/ImageCoreUtils.cpp b/Engine/Source/Runtime/ImageCore/Private/ImageCoreUtils.cpp
--- a/Engine/Source/Runtime/ImageCore/Private/ImageCoreUtils.cpp
+++ b/Engine/Source/Runtime/ImageCore/Private/ImageCoreUtils.cpp
@@ -131,3 +131,38 @@ ENUM_SUPPORTED_FORMATS(DECL_FORMAT_NAME);
 		return FName();
 	}
 }
+
+IMAGECORE_API ETextureSourceFormat FImageCoreUtils::GetCommonSourceFormat(ETextureSourceFormat Format1, ETextureSourceFormat Format2)
+{
+	if (Format1 == Format2)
+	{
+		return Format1;
+	}
+
+	// calculation is based on the following oriented graph, where formats A and B are linked if conversion from A to B is considered to be acceptable, specifically:
+	// TSF_G8 -> TSF_BGRA8 -> TSF_RGBA16
+	// TSF_G8 -> TSF_G16 -> TSF_RGBA16
+	// TSF_RGBA16 -> TSF_RGBA32F
+	// TSF_R16F -> TSF_RGBA16F -> TSF_RGBA32F
+	// TSF_R16F -> TSF_R32F -> TSF_RGBA32F
+	// TSF_BGRE8 -> TSF_RGBA32F
+	// the function returns the lowest possible parent node for two input nodes
+
+	switch (Format1)
+	{
+	case TSF_G8:
+		return !UE::TextureDefines::IsHDR(Format2) ? Format2 : TSF_RGBA32F;
+	case TSF_G16:
+	case TSF_BGRA8:
+		return Format2 == TSF_G8 ? Format1 : !UE::TextureDefines::IsHDR(Format2) ? TSF_RGBA16 : TSF_RGBA32F;
+	case TSF_RGBA16:
+		return !UE::TextureDefines::IsHDR(Format2) ? Format1 : TSF_RGBA32F;
+	case TSF_R16F:
+		return Format2 == TSF_R32F || Format2 == TSF_RGBA16F ? Format2 : TSF_RGBA32F;
+	case TSF_R32F:
+	case TSF_RGBA16F:
+		return Format2 == TSF_R16F ? Format1 : TSF_RGBA32F;
+	default:
+		return TSF_RGBA32F;
+	}
+}
diff --git a/Engine/Source/Runtime/ImageCore/Public/ImageCoreUtils.h b/Engine/Source/Runtime/ImageCore/Public/ImageCoreUtils.h
--- a/Engine/Source/Runtime/ImageCore/Public/ImageCoreUtils.h
+++ b/Engine/Source/Runtime/ImageCore/Public/ImageCoreUtils.h
@@ -37,4 +37,12 @@ namespace FImageCoreUtils
 							
 	IMAGECORE_API FName ConvertToUncompressedTextureFormatName(ERawImageFormat::Type Format);
 
+	/**
+	 * Returns ETextureSourceFormat which can be used to efficiently store data encoded in both input formats.
+	 * Can be used in cases when multiple sources need to be mixed together (i.e. in a texture array or UDIM)
+	 * @param Format1 - First source format.
+	 * @param Format2 - Second source format.
+	 */
+	IMAGECORE_API ETextureSourceFormat GetCommonSourceFormat(ETextureSourceFormat Format1, ETextureSourceFormat Format2);
+
 };
