Added Engine.VerifyLoadMapWorldCleanupEnabled CVar to toggle whether the VerifyLoadMapWorldCleanup routine runs, disabling it by default on shipping builds.

diff --git a/Engine/Source/Runtime/Engine/Private/UnrealEngine.cpp b/Engine/Source/Runtime/Engine/Private/UnrealEngine.cpp
--- a/Engine/Source/Runtime/Engine/Private/UnrealEngine.cpp
+++ b/Engine/Source/Runtime/Engine/Private/UnrealEngine.cpp
@@ -436,6 +436,14 @@ static FAutoConsoleVariableRef GDelayTrimMemoryDuringMapLoadModeCVar(
 	ECVF_Default
 );
 
+bool GVerifyLoadMapWorldCleanupEnabled = !UE_BUILD_SHIPPING;
+static FAutoConsoleVariableRef GVerifyLoadMapWorldCleanupEnabledCVar(
+	TEXT("Engine.VerifyLoadMapWorldCleanupEnabled"),
+	GVerifyLoadMapWorldCleanupEnabled,
+	TEXT("Enables logging references preventing the previous world objects from being cleaned up during map load"),
+	ECVF_Default
+);
+
 #if !UE_BUILD_SHIPPING
 TAutoConsoleVariable<bool> CVarCsvAlwaysShowFrameCount(
 	TEXT("csv.AlwaysShowFrameCount"),
@@ -15541,6 +15549,11 @@ bool UEngine::IsWorldDuplicate(const UWorld* const InWorld)
 
 void UEngine::VerifyLoadMapWorldCleanup(FWorldContext* ForWorldContext)
 {
+	if (!GVerifyLoadMapWorldCleanupEnabled)
+	{
+		return;
+	}
+
 	TRACE_CPUPROFILER_EVENT_SCOPE(VerifyLoadMapWorldCleanup);
 	// All worlds at this point should be the CurrentWorld of some context, preview worlds, or streaming level
 	// worlds that are owned by the CurrentWorld of some context.
