Tweak to last change, hide the global until needed.

diff --git a/Engine/Source/Runtime/Engine/Private/UnrealEngine.cpp b/Engine/Source/Runtime/Engine/Private/UnrealEngine.cpp
--- a/Engine/Source/Runtime/Engine/Private/UnrealEngine.cpp
+++ b/Engine/Source/Runtime/Engine/Private/UnrealEngine.cpp
@@ -436,7 +436,7 @@ static FAutoConsoleVariableRef GDelayTrimMemoryDuringMapLoadModeCVar(
 	ECVF_Default
 );
 
-bool GVerifyLoadMapWorldCleanupEnabled = !UE_BUILD_SHIPPING;
+static bool GVerifyLoadMapWorldCleanupEnabled = !UE_BUILD_SHIPPING;
 static FAutoConsoleVariableRef GVerifyLoadMapWorldCleanupEnabledCVar(
 	TEXT("Engine.VerifyLoadMapWorldCleanupEnabled"),
 	GVerifyLoadMapWorldCleanupEnabled,
