Foliage: Fix ensure when painting/erasing (race condition for InstanceBuffer.GetNumInstances() being updated)

diff --git a/Engine/Source/Runtime/Engine/Private/InstancedStaticMesh.cpp b/Engine/Source/Runtime/Engine/Private/InstancedStaticMesh.cpp
--- a/Engine/Source/Runtime/Engine/Private/InstancedStaticMesh.cpp
+++ b/Engine/Source/Runtime/Engine/Private/InstancedStaticMesh.cpp
@@ -225,6 +225,7 @@ FInstanceUpdateCmdBuffer::FInstanceUpdateCmdBuffer()
 	  NumUpdates(0),
 	  NumRemoves(0)
 	, NumEdits(0)
+	, NumEditInstances(0)
 {
 }
 
@@ -442,6 +443,10 @@ void FStaticMeshInstanceBuffer::UpdateFromCommandBuffer_Concurrent(FInstanceUpda
 	
 	// leave NumEdits unchanged in commandbuffer
 	CmdBuffer.NumEdits = NewCmdBuffer->NumEdits; 
+
+	// Compute render instances (same value that will computed on the render thread in UpdateFromCommandBuffer_RenderThread)
+	// Any query of number of render instances on game thread should use this instead of InstanceData->GetNumInstances();
+	CmdBuffer.NumEditInstances = NewCmdBuffer->NumAdds + InstanceData->GetNumInstances();
 	CmdBuffer.ResetInlineCommands();
 		
 	ENQUEUE_RENDER_COMMAND(InstanceBuffer_UpdateFromPreallocatedData)(
@@ -1457,7 +1462,7 @@ void FInstancedStaticMeshSceneProxy::SetupProxy(UInstancedStaticMeshComponent* I
 
 		// NumRenderInstances is the extent that the reorder table can map to.
 		// Temporarily when removing instances from a HISM this can be sparse so that InComponent->GetInstanceCount() < NumRenderInstances.
-		const int32 NumRenderInstances = FMath::Max<int32>(InComponent->PerInstanceRenderData->InstanceBuffer.GetNumInstances(), InComponent->GetInstanceCount());
+		const int32 NumRenderInstances = FMath::Max<int32>(InComponent->InstanceUpdateCmdBuffer.NumEditInstances, InComponent->GetInstanceCount());
 
 		bSupportsInstanceDataBuffer = true;
 		InstanceSceneData.SetNumZeroed(NumRenderInstances);
diff --git a/Engine/Source/Runtime/Engine/Public/PrimitiveInstanceUpdateCommand.h b/Engine/Source/Runtime/Engine/Public/PrimitiveInstanceUpdateCommand.h
--- a/Engine/Source/Runtime/Engine/Public/PrimitiveInstanceUpdateCommand.h
+++ b/Engine/Source/Runtime/Engine/Public/PrimitiveInstanceUpdateCommand.h
@@ -70,4 +70,5 @@ struct FInstanceUpdateCmdBuffer
 	int32 NumCustomFloatUpdates;
 	int32 NumRemoves;
 	int32 NumEdits;
+	int32 NumEditInstances;
 };
