Show an error message in the console if the command Trace.Screenshot is used and the Screenshot channel is off.

diff --git a/Engine/Source/Runtime/Engine/Private/ProfilingDebugging/TraceScreenshot.cpp b/Engine/Source/Runtime/Engine/Private/ProfilingDebugging/TraceScreenshot.cpp
--- a/Engine/Source/Runtime/Engine/Private/ProfilingDebugging/TraceScreenshot.cpp
+++ b/Engine/Source/Runtime/Engine/Private/ProfilingDebugging/TraceScreenshot.cpp
@@ -30,7 +30,7 @@ static void TraceScreenshotCommandCallback(const TArray<FString>& Args)
 		bShowUI = FCString::ToBool(*Args[1]);
 	}
 
-	FTraceScreenshot::RequestScreenshot(Name, bShowUI);
+	FTraceScreenshot::RequestScreenshot(Name, bShowUI, LogConsoleResponse);
 }
 
 static FAutoConsoleCommand TraceScreenshotCmd(
@@ -41,11 +41,11 @@ static FAutoConsoleCommand TraceScreenshotCmd(
 
 bool FTraceScreenshot::bSuppressWritingToFile = false;
 
-void FTraceScreenshot::RequestScreenshot(FString Name, bool bShowUI)
+void FTraceScreenshot::RequestScreenshot(FString Name, bool bShowUI, const FLogCategoryAlias& LogCategory)
 {
 	if (!SHOULD_TRACE_SCREENSHOT())
 	{
-		UE_LOG(LogCore, Error, TEXT("Could not trace screenshot because the screenshot trace channel is off. Turn it on using \"Trace.Enable Screenshot\"."));
+		UE_LOG_REF(LogCategory, Error, TEXT("Could not trace screenshot because the screenshot trace channel is off. Turn it on using \"Trace.Enable Screenshot\"."));
 		return;
 	}
 
diff --git a/Engine/Source/Runtime/Engine/Public/ProfilingDebugging/TraceScreenshot.h b/Engine/Source/Runtime/Engine/Public/ProfilingDebugging/TraceScreenshot.h
--- a/Engine/Source/Runtime/Engine/Public/ProfilingDebugging/TraceScreenshot.h
+++ b/Engine/Source/Runtime/Engine/Public/ProfilingDebugging/TraceScreenshot.h
@@ -10,7 +10,16 @@ class ULevel;
 class ENGINE_API FTraceScreenshot
 {
 public:
-	static void RequestScreenshot(FString Name, bool bShowUI);
+
+// In no logging configurations all log categories are of type FNoLoggingCategory, which has no relation with
+// FLogCategoryBase. In order to not need to conditionally set the argument alias the type here.
+#if NO_LOGGING
+	typedef FNoLoggingCategory FLogCategoryAlias;
+#else
+	typedef FLogCategoryBase FLogCategoryAlias;
+#endif
+
+	static void RequestScreenshot(FString Name, bool bShowUI, const FLogCategoryAlias& LogCategory = LogCore);
 
 	/* 
 	* Add the provided screenshot to the trace.
@@ -31,6 +40,7 @@ public:
 	* Reset the internal state of the TraceScreenshot system.
 	*/
 	static void Reset();
+
 private:
 	FTraceScreenshot();
 	FTraceScreenshot(const FTraceScreenshot& Other) {}
