World Partition - HLOD
* Exposed proper Nanite generation settings: To control mesh fallback mesh generation, relative error is prefered over triangle percent
* Exposed "Support Ray Tracing" option

diff --git a/Engine/Source/Developer/MeshMergeUtilities/Private/MeshMergeUtilities.cpp b/Engine/Source/Developer/MeshMergeUtilities/Private/MeshMergeUtilities.cpp
--- a/Engine/Source/Developer/MeshMergeUtilities/Private/MeshMergeUtilities.cpp
+++ b/Engine/Source/Developer/MeshMergeUtilities/Private/MeshMergeUtilities.cpp
@@ -2920,6 +2920,9 @@ void FMeshMergeUtilities::MergeComponentsToStaticMesh(const TArray<UPrimitiveCom
 			StaticMesh->SetLightMapCoordinateIndex(LightMapUVChannel);
 		}
 
+		// Ray tracing support
+		StaticMesh->bSupportRayTracing = InSettings.bSupportRayTracing;
+
 		const bool bContainsImposters = ImposterComponents.Num() > 0;
 		TArray<UMaterialInterface*> ImposterMaterials;
 		FBox ImposterBounds(EForceInit::ForceInit);
@@ -3044,9 +3047,7 @@ void FMeshMergeUtilities::MergeComponentsToStaticMesh(const TArray<UPrimitiveCom
 		StaticMesh->SetLightMapResolution(InSettings.bComputedLightMapResolution ? DataTracker.GetLightMapDimension() : InSettings.TargetLightMapResolution);
 
 		// Nanite settings
-		StaticMesh->NaniteSettings.bEnabled = InSettings.bGenerateNaniteEnabledMesh;
-		StaticMesh->NaniteSettings.FallbackPercentTriangles = InSettings.NaniteFallbackTrianglePercent * 0.01f;
-		StaticMesh->NaniteSettings.PositionPrecision = MIN_int32;
+		StaticMesh->NaniteSettings = InSettings.NaniteSettings;
 
 #if WITH_EDITOR
 		//If we are running the automation test
diff --git a/Engine/Source/Developer/MeshMergeUtilities/Private/ProxyGenerationProcessor.cpp b/Engine/Source/Developer/MeshMergeUtilities/Private/ProxyGenerationProcessor.cpp
--- a/Engine/Source/Developer/MeshMergeUtilities/Private/ProxyGenerationProcessor.cpp
+++ b/Engine/Source/Developer/MeshMergeUtilities/Private/ProxyGenerationProcessor.cpp
@@ -202,6 +202,8 @@ void FProxyGenerationProcessor::ProcessJob(const FGuid& JobGuid, FProxyGeneratio
 	StaticMesh->SetLightMapResolution(Data->MergeData->InProxySettings.LightMapResolution);
 	StaticMesh->SetLightMapCoordinateIndex(1);
 
+	// Ray tracing support
+	StaticMesh->bSupportRayTracing = Data->MergeData->InProxySettings.bSupportRayTracing;
 
 	FStaticMeshSourceModel& SrcModel = StaticMesh->AddSourceModel();
 	/*Don't allow the engine to recalculate normals*/
@@ -296,9 +298,7 @@ void FProxyGenerationProcessor::ProcessJob(const FGuid& JobGuid, FProxyGeneratio
 	}
 
 	// Nanite settings
-	StaticMesh->NaniteSettings.bEnabled = Data->MergeData->InProxySettings.bGenerateNaniteEnabledMesh;
-	StaticMesh->NaniteSettings.FallbackPercentTriangles = Data->MergeData->InProxySettings.NaniteProxyTrianglePercent * 0.01f;
-	StaticMesh->NaniteSettings.PositionPrecision = MIN_int32;
+	StaticMesh->NaniteSettings = Data->MergeData->InProxySettings.NaniteSettings;
 
 	//Set the Imported version before calling the build
 	StaticMesh->ImportVersion = EImportStaticMeshVersion::LastVersion;
diff --git a/Engine/Source/Runtime/Engine/Classes/Engine/MeshMerging.h b/Engine/Source/Runtime/Engine/Classes/Engine/MeshMerging.h
--- a/Engine/Source/Runtime/Engine/Classes/Engine/MeshMerging.h
+++ b/Engine/Source/Runtime/Engine/Classes/Engine/MeshMerging.h
@@ -242,6 +242,12 @@ struct FMeshProxySettings
 
 	UPROPERTY()
 	uint8 bBakeVertexData_DEPRECATED:1;
+
+	UPROPERTY()
+	uint8 bGenerateNaniteEnabledMesh_DEPRECATED : 1;
+	
+	UPROPERTY()
+	float NaniteProxyTrianglePercent_DEPRECATED;
 #endif
 
 	/** Distance at which meshes should be merged together, this can close gaps like doors and windows in distant geometry */
@@ -300,6 +306,10 @@ struct FMeshProxySettings
 	UPROPERTY(BlueprintReadWrite, EditAnywhere, Category = LandscapeCulling)
 	uint8 bUseLandscapeCulling:1;
 
+	/** Whether ray tracing will be supported on this mesh. Disable this to save memory if the generated mesh will only be rendered in the distance. */
+	UPROPERTY(BlueprintReadWrite, EditAnywhere, Category = ProxySettings)
+	uint8 bSupportRayTracing : 1;
+
 	/** Whether to allow distance field to be computed for this mesh. Disable this to save memory if the merged mesh will only be rendered in the distance. */
 	UPROPERTY(BlueprintReadWrite, EditAnywhere, Category = ProxySettings)
 	uint8 bAllowDistanceField:1;
@@ -324,13 +334,9 @@ struct FMeshProxySettings
 	UPROPERTY(BlueprintReadWrite, EditAnywhere, Category = ProxySettings)
 	uint8 bGenerateLightmapUVs:1;
 
-	/** Whether to generate a nanite-enabled mesh */
-	UPROPERTY(BlueprintReadWrite, EditAnywhere, Category = NaniteSettings)
-	uint8 bGenerateNaniteEnabledMesh : 1;
-
-	/** Percentage of triangles to reduce down to for generating a coarse proxy mesh from the Nanite mesh */
-	UPROPERTY(BlueprintReadWrite, EditAnywhere, AdvancedDisplay, Category = NaniteSettings, meta = (EditConditionHides, EditCondition = "bGenerateNaniteEnabledMesh", ClampMin = 0, ClampMax = 100))
-	float NaniteProxyTrianglePercent;
+	/** Settings related to building Nanite data. */
+	UPROPERTY(EditAnywhere, BlueprintReadWrite, Category = NaniteSettings)
+	FMeshNaniteSettings NaniteSettings;
 
 	/** Default settings. */
 	FMeshProxySettings()
@@ -344,6 +350,8 @@ struct FMeshProxySettings
 		, bExportRoughnessMap_DEPRECATED(false)
 		, bExportSpecularMap_DEPRECATED(false)
 		, bBakeVertexData_DEPRECATED(false)
+		, bGenerateNaniteEnabledMesh_DEPRECATED(false)
+		, NaniteProxyTrianglePercent_DEPRECATED(100)
 #endif
 		, MergeDistance(0)
 		, UnresolvedGeometryColor(FColor::Black)
@@ -359,14 +367,13 @@ struct FMeshProxySettings
 		, bComputeLightMapResolution(false)
 		, bRecalculateNormals(true)
 		, bUseLandscapeCulling(false)
+		, bSupportRayTracing(true)
 		, bAllowDistanceField(false)
 		, bReuseMeshLightmapUVs(true)
 		, bGroupIdenticalMeshesForBaking(false)
 		, bCreateCollision(true)
 		, bAllowVertexColors(false)
 		, bGenerateLightmapUVs(false)
-		, bGenerateNaniteEnabledMesh(false)
-		, NaniteProxyTrianglePercent(100)
 	{
 		MaterialSettings.MaterialMergeType = EMaterialMergeType::MaterialMergeType_Simplygon;
 	}
@@ -391,14 +398,14 @@ struct FMeshProxySettings
 			&& bComputeLightMapResolution == Other.bComputeLightMapResolution
 			&& bRecalculateNormals == Other.bRecalculateNormals
 			&& bUseLandscapeCulling == Other.bUseLandscapeCulling
+			&& bSupportRayTracing == Other.bSupportRayTracing
 			&& bAllowDistanceField == Other.bAllowDistanceField
 			&& bReuseMeshLightmapUVs == Other.bReuseMeshLightmapUVs
 			&& bGroupIdenticalMeshesForBaking == Other.bGroupIdenticalMeshesForBaking
 			&& bCreateCollision == Other.bCreateCollision
 			&& bAllowVertexColors == Other.bAllowVertexColors
 			&& bGenerateLightmapUVs == Other.bGenerateLightmapUVs
-			&& bGenerateNaniteEnabledMesh == Other.bGenerateNaniteEnabledMesh
-			&& NaniteProxyTrianglePercent == Other.NaniteProxyTrianglePercent;
+			&& NaniteSettings == Other.NaniteSettings;
 	}
 
 	/** Inequality. */
@@ -537,41 +544,48 @@ struct FMeshMergingSettings
 	UPROPERTY(EditAnywhere, BlueprintReadWrite, Category = MeshSettings)
 	uint8 bIncludeImposters:1;
 
+	/** Whether ray tracing will be supported on this mesh. Disable this to save memory if the generated mesh will only be rendered in the distance. */
+	UPROPERTY(EditAnywhere, BlueprintReadWrite, Category = MeshSettings)
+	uint8 bSupportRayTracing : 1;
+
 	/** Whether to allow distance field to be computed for this mesh. Disable this to save memory if the merged mesh will only be rendered in the distance. */
 	UPROPERTY(EditAnywhere, Category = MeshSettings)
 	uint8 bAllowDistanceField:1;
 
-	/** Whether to generate a nanite-enabled mesh */
+	/** Settings related to building Nanite data. */
 	UPROPERTY(EditAnywhere, BlueprintReadWrite, Category = NaniteSettings)
-	uint8 bGenerateNaniteEnabledMesh : 1;
-
-	/** Percentage of triangles to reduce down to for generating a coarse fallback mesh from the Nanite mesh */
-	UPROPERTY(BlueprintReadWrite, EditAnywhere, AdvancedDisplay, Category = NaniteSettings, meta = (EditConditionHides, EditCondition = "bGenerateNaniteEnabledMesh", ClampMin = 0, ClampMax = 100))
-	float NaniteFallbackTrianglePercent;
+	FMeshNaniteSettings NaniteSettings;
 
 #if WITH_EDITORONLY_DATA
-	/** Whether we should import vertex colors into merged mesh */
 	UPROPERTY()
 	uint8 bImportVertexColors_DEPRECATED:1;
+
 	UPROPERTY()
 	uint8 bCalculateCorrectLODModel_DEPRECATED:1;
-	/** Whether to export normal maps for material merging */
+
 	UPROPERTY()
 	uint8 bExportNormalMap_DEPRECATED:1;
-	/** Whether to export metallic maps for material merging */
+
 	UPROPERTY()
 	uint8 bExportMetallicMap_DEPRECATED:1;
-	/** Whether to export roughness maps for material merging */
+
 	UPROPERTY()
 	uint8 bExportRoughnessMap_DEPRECATED:1;
-	/** Whether to export specular maps for material merging */
+
 	UPROPERTY()
 	uint8 bExportSpecularMap_DEPRECATED:1;
-	/** Merged material texture atlas resolution */
+
 	UPROPERTY()
 	int32 MergedMaterialAtlasResolution_DEPRECATED;
+
 	UPROPERTY()
 	int32 ExportSpecificLOD_DEPRECATED;
+
+	UPROPERTY()
+	uint8 bGenerateNaniteEnabledMesh_DEPRECATED : 1;
+
+	UPROPERTY()
+	float NaniteFallbackTrianglePercent_DEPRECATED;
 #endif
 
 	EMeshMergeType MergeType;
@@ -595,9 +609,8 @@ struct FMeshMergingSettings
 		, bMergeEquivalentMaterials(true)
 		, bUseLandscapeCulling(false)
 		, bIncludeImposters(true)
+		, bSupportRayTracing(true)
 		, bAllowDistanceField(false)
-		, bGenerateNaniteEnabledMesh(false)
-		, NaniteFallbackTrianglePercent(100)
 #if WITH_EDITORONLY_DATA
 		, bImportVertexColors_DEPRECATED(false)
 		, bCalculateCorrectLODModel_DEPRECATED(false)
@@ -607,6 +620,8 @@ struct FMeshMergingSettings
 		, bExportSpecularMap_DEPRECATED(false)
 		, MergedMaterialAtlasResolution_DEPRECATED(1024)
 		, ExportSpecificLOD_DEPRECATED(0)
+		, bGenerateNaniteEnabledMesh_DEPRECATED(false)
+		, NaniteFallbackTrianglePercent_DEPRECATED(100)
 #endif
 		, MergeType(EMeshMergeType::MeshMergeType_Default)
 	{
diff --git a/Engine/Source/Runtime/Engine/Private/EngineTypes.cpp b/Engine/Source/Runtime/Engine/Private/EngineTypes.cpp
--- a/Engine/Source/Runtime/Engine/Private/EngineTypes.cpp
+++ b/Engine/Source/Runtime/Engine/Private/EngineTypes.cpp
@@ -28,6 +28,12 @@ void FMeshProxySettings::PostSerialize(const FArchive& Ar)
 	if (Ar.IsLoading())
 	{
 		MaterialSettings.MaterialMergeType = EMaterialMergeType::MaterialMergeType_Simplygon;
+
+		if (bGenerateNaniteEnabledMesh_DEPRECATED)
+		{
+			NaniteSettings.bEnabled = true;
+			NaniteSettings.FallbackPercentTriangles = NaniteProxyTrianglePercent_DEPRECATED / 100.0f;
+		}
 	}
 }
 
@@ -73,6 +79,12 @@ void FMeshMergingSettings::PostSerialize(const FArchive& Ar)
 			SpecificLOD = ExportSpecificLOD_DEPRECATED;
 			LODSelectionType = EMeshLODSelectionType::SpecificLOD;
 		}
+
+		if (bGenerateNaniteEnabledMesh_DEPRECATED)
+		{
+			NaniteSettings.bEnabled = true;
+			NaniteSettings.FallbackPercentTriangles = NaniteFallbackTrianglePercent_DEPRECATED / 100.0f;
+		}
 	}
 }
 #endif
