LandscapeActorDesc: Do not return itself as SceneOutlinerParent

Fixes stack overflow crash. This is a stop gap because ALandscape is a AlwaysLoaded actor and should never be unloaded.

#rnx

diff --git a/Engine/Source/Runtime/Engine/Private/WorldPartition/Landscape/LandscapeActorDesc.cpp b/Engine/Source/Runtime/Engine/Private/WorldPartition/Landscape/LandscapeActorDesc.cpp
--- a/Engine/Source/Runtime/Engine/Private/WorldPartition/Landscape/LandscapeActorDesc.cpp
+++ b/Engine/Source/Runtime/Engine/Private/WorldPartition/Landscape/LandscapeActorDesc.cpp
@@ -71,7 +71,14 @@ bool FLandscapeActorDesc::Equals(const FWorldPartitionActorDesc* Other) const
 
 const FGuid& FLandscapeActorDesc::GetSceneOutlinerParent() const
 {
-	return LandscapeActorGuid;
+	// Landscape can't parent itself
+	if (LandscapeActorGuid != GetGuid())
+	{
+		return LandscapeActorGuid;
+	}
+
+	static FGuid NoParent;
+	return NoParent;
 }
 
 #endif
