Fix for a memory leak affecting nDisplay where the FViewInfo destructor was not properly getting called when FSceneViewFamilyContext destroyed its views because the FSceneView destructor was not virtual. FSceneViewFamilyContext stores views as FSceneView pointers, and since FSceneView did not declare a virtual destructor, FSceneViewFamilyContext could not property destroy any FSceneView-derived classes, specifically FViewInfo, which was causing a memory leak in the case where an FViewInfo object was created through FRendererModule::CreateAndInitSingleView and managed using a FSceneViewFamilyContext, as used in a handful of cases in the engine and nDisplay. Fixed by adding a default virtual destructor to FSceneView and marking FViewInfo's destructor as virtual.

diff --git a/Engine/Source/Runtime/Engine/Public/SceneView.h b/Engine/Source/Runtime/Engine/Public/SceneView.h
--- a/Engine/Source/Runtime/Engine/Public/SceneView.h
+++ b/Engine/Source/Runtime/Engine/Public/SceneView.h
@@ -1298,6 +1298,9 @@ public:
 	FSceneView(FSceneView&& Other) = default;
 	PRAGMA_ENABLE_DEPRECATION_WARNINGS
 
+	/** Default destructor */
+	virtual ~FSceneView() = default;
+
 #if DO_CHECK || USING_CODE_ANALYSIS
 	/** Verifies all the assertions made on members. */
 	bool VerifyMembersChecks() const;
diff --git a/Engine/Source/Runtime/Renderer/Private/SceneRendering.h b/Engine/Source/Runtime/Renderer/Private/SceneRendering.h
--- a/Engine/Source/Runtime/Renderer/Private/SceneRendering.h
+++ b/Engine/Source/Runtime/Renderer/Private/SceneRendering.h
@@ -1570,7 +1570,7 @@ public:
 	/** 
 	* Destructor. 
 	*/
-	RENDERER_API ~FViewInfo();
+	RENDERER_API virtual ~FViewInfo();
 
 #if DO_CHECK || USING_CODE_ANALYSIS
 	/** Verifies all the assertions made on members. */
