Optimus: Reinit the mesh deformer instance when changing the skel mesh on a skinned mesh component.
The mesh deformer may be inherited from the skel mesh, so we need to take into account when it changes.

diff --git a/Engine/Source/Runtime/Engine/Private/Components/SkinnedMeshComponent.cpp b/Engine/Source/Runtime/Engine/Private/Components/SkinnedMeshComponent.cpp
--- a/Engine/Source/Runtime/Engine/Private/Components/SkinnedMeshComponent.cpp
+++ b/Engine/Source/Runtime/Engine/Private/Components/SkinnedMeshComponent.cpp
@@ -1909,6 +1909,10 @@ void USkinnedMeshComponent::SetSkinnedAssetAndUpdate(USkinnedAsset* InSkinnedAss
 			MorphTargetWeights.Empty();
 			ExternalMorphWeightData.Empty();
 		}
+
+		// Re-init the MeshDeformer which might come from the SkelMesh.
+		UMeshDeformer* ActiveMeshDeformer = GetActiveMeshDeformer();
+		MeshDeformerInstance = ActiveMeshDeformer != nullptr ? ActiveMeshDeformer->CreateInstance(this, MeshDeformerInstanceSettings) : nullptr;
 	}
 
 	// Update external weight array sizes.
