Fix content bundle cells not being requested

diff --git a/Engine/Source/Runtime/Engine/Private/WorldPartition/WorldPartitionRuntimeSpatialHash.cpp b/Engine/Source/Runtime/Engine/Private/WorldPartition/WorldPartitionRuntimeSpatialHash.cpp
--- a/Engine/Source/Runtime/Engine/Private/WorldPartition/WorldPartitionRuntimeSpatialHash.cpp
+++ b/Engine/Source/Runtime/Engine/Private/WorldPartition/WorldPartitionRuntimeSpatialHash.cpp
@@ -1273,7 +1273,7 @@ FAutoConsoleCommand UWorldPartitionRuntimeSpatialHash::OverrideLoadingRangeComma
 // Streaming interface
 void UWorldPartitionRuntimeSpatialHash::ForEachStreamingCells(TFunctionRef<bool(const UWorldPartitionRuntimeCell*)> Func) const
 {
-	for (const FSpatialHashStreamingGrid& StreamingGrid : StreamingGrids)
+	ForEachStreamingGrid([&](const FSpatialHashStreamingGrid& StreamingGrid)
 	{
 		if (IsCellRelevantFor(StreamingGrid.bClientOnlyVisible))
 		{
@@ -1291,12 +1291,12 @@ void UWorldPartitionRuntimeSpatialHash::ForEachStreamingCells(TFunctionRef<bool(
 				}
 			}
 		}
-	}
+	});
 }
 
 void UWorldPartitionRuntimeSpatialHash::ForEachStreamingCellsQuery(const FWorldPartitionStreamingQuerySource& QuerySource, TFunctionRef<bool(const UWorldPartitionRuntimeCell*)> Func) const
 {
-	for (const FSpatialHashStreamingGrid& StreamingGrid : StreamingGrids)
+	ForEachStreamingGrid([&](const FSpatialHashStreamingGrid& StreamingGrid)
 	{
 		if (IsCellRelevantFor(StreamingGrid.bClientOnlyVisible))
 		{
@@ -1311,7 +1311,7 @@ void UWorldPartitionRuntimeSpatialHash::ForEachStreamingCellsQuery(const FWorldP
 				}
 			}
 		}
-	}
+	});
 }
 
 void UWorldPartitionRuntimeSpatialHash::ForEachStreamingCellsSources(const TArray<FWorldPartitionStreamingSource>& Sources, TFunctionRef<bool(const UWorldPartitionRuntimeCell*, EStreamingSourceTargetState)> Func) const
@@ -1324,24 +1324,24 @@ void UWorldPartitionRuntimeSpatialHash::ForEachStreamingCellsSources(const TArra
 	if (Sources.Num() == 0)
 	{
 		// Get always loaded cells
-		for (const FSpatialHashStreamingGrid& StreamingGrid : StreamingGrids)
+		ForEachStreamingGrid([&] (const FSpatialHashStreamingGrid& StreamingGrid)
 		{
 			if (IsCellRelevantFor(StreamingGrid.bClientOnlyVisible))
 			{
 				StreamingGrid.GetAlwaysLoadedCells(DataLayerSubsystem, ActivateStreamingSourceCells.GetCells(), LoadStreamingSourceCells.GetCells());
 			}
-		}
-	}
+		});
+}
 	else
 	{
 		// Get cells based on streaming sources
-		for (const FSpatialHashStreamingGrid& StreamingGrid : StreamingGrids)
+		ForEachStreamingGrid([&](const FSpatialHashStreamingGrid& StreamingGrid)
 		{
 			if (IsCellRelevantFor(StreamingGrid.bClientOnlyVisible))
 			{
 				StreamingGrid.GetCells(Sources, DataLayerSubsystem, ActivateStreamingSourceCells, LoadStreamingSourceCells, GetEffectiveEnableZCulling(bEnableZCulling));
 			}
-		}
+		});
 	}
 
 	for (const UWorldPartitionRuntimeCell* Cell : ActivateStreamingSourceCells.GetCells())
@@ -1439,6 +1439,31 @@ void UWorldPartitionRuntimeSpatialHash::ForEachStreamingGrid(TFunctionRef<void(F
 	}
 }
 
+void UWorldPartitionRuntimeSpatialHash::ForEachStreamingGridBreakable(TFunctionRef<bool(const FSpatialHashStreamingGrid&)> Func) const
+{
+	for (const FSpatialHashStreamingGrid& StreamingGrid : StreamingGrids)
+	{
+		if (!Func(StreamingGrid))
+		{
+			return;
+		}
+	}
+
+	for (const TWeakObjectPtr<URuntimeSpatialHashExternalStreamingObject>& StreamingObject : ExternalStreamingObjects)
+	{
+		if (StreamingObject.IsValid())
+		{
+			for (FSpatialHashStreamingGrid& StreamingGrid : StreamingObject->StreamingGrids)
+			{
+				if (!Func(StreamingGrid))
+				{
+					return;
+				}
+			}
+		}
+	}
+}
+
 void UWorldPartitionRuntimeSpatialHash::ForEachStreamingGrid(TFunctionRef<void(const FSpatialHashStreamingGrid&)> Func, bool bIncludeExternalStreamingObjects) const
 {
 	for (const FSpatialHashStreamingGrid& StreamingGrid : StreamingGrids)
@@ -1625,14 +1650,14 @@ void UWorldPartitionRuntimeSpatialHash::Draw3D(const TArray<FWorldPartitionStrea
 
 bool UWorldPartitionRuntimeSpatialHash::ContainsRuntimeHash(const FString& Name) const
 {
-	for (const FSpatialHashStreamingGrid& StreamingGrid : StreamingGrids)
+	bool bResult = false;
+	ForEachStreamingGridBreakable([&](const FSpatialHashStreamingGrid& StreamingGrid)
 	{
-		if (StreamingGrid.GridName.ToString().Equals(Name, ESearchCase::IgnoreCase))
-		{
-			return true;
-		}
-	}
-	return false;
+		bResult = StreamingGrid.GridName.ToString().Equals(Name, ESearchCase::IgnoreCase);
+		return !bResult;
+	});
+
+	return bResult;
 }
 
 TArray<const FSpatialHashStreamingGrid*> UWorldPartitionRuntimeSpatialHash::GetFilteredStreamingGrids() const
diff --git a/Engine/Source/Runtime/Engine/Public/WorldPartition/WorldPartitionRuntimeSpatialHash.h b/Engine/Source/Runtime/Engine/Public/WorldPartition/WorldPartitionRuntimeSpatialHash.h
--- a/Engine/Source/Runtime/Engine/Public/WorldPartition/WorldPartitionRuntimeSpatialHash.h
+++ b/Engine/Source/Runtime/Engine/Public/WorldPartition/WorldPartitionRuntimeSpatialHash.h
@@ -271,6 +271,7 @@ public:
 
 protected:
 	void ForEachStreamingGrid(TFunctionRef<void(FSpatialHashStreamingGrid&)> Func);
+	void ForEachStreamingGridBreakable(TFunctionRef<bool(const FSpatialHashStreamingGrid&)> Func) const;
 	void ForEachStreamingGrid(TFunctionRef<void(const FSpatialHashStreamingGrid&)> Func, bool bIncludeExternalStreamingObjects = true) const;
 
 	virtual EWorldPartitionStreamingPerformance GetStreamingPerformanceForCell(const UWorldPartitionRuntimeCell* Cell) const override;
