Add logging to determine retained vs non-retained memory trimming in the audio stream cache.

diff --git a/Engine/Source/Runtime/Engine/Private/AudioStreamingCache.cpp b/Engine/Source/Runtime/Engine/Private/AudioStreamingCache.cpp
--- a/Engine/Source/Runtime/Engine/Private/AudioStreamingCache.cpp
+++ b/Engine/Source/Runtime/Engine/Private/AudioStreamingCache.cpp
@@ -1018,9 +1018,12 @@ uint64 FAudioChunkCache::TrimMemory(uint64 BytesToFree, bool bInAllowRetainedChu
 	}
 
 	uint64 RetainedBytesFreed = 0;
+	uint32 NumRetainedElementsEvicted = 0;
 	// If we have run out of non-retained and in-flight load audio chunks to trim, eat into the 
 	if (bInAllowRetainedChunkTrimming && EnableTrimmingRetainedAudioCVar > 0 && BytesFreed < BytesToFree)
 	{
+		UE_LOG(LogAudioStreamCaching, Verbose, TEXT("TrimMemory: Num Non-Retained Elements Evicted: %d. Non-Retained Bytes Freed: %d"), NumElementsEvicted, BytesFreed);
+
 		CurrentElement = LeastRecentElement;
 		ElementToStopAt = MostRecentElement->LessRecentElement;
 		while (CurrentElement != ElementToStopAt && BytesFreed < BytesToFree)
@@ -1051,14 +1054,17 @@ uint64 FAudioChunkCache::TrimMemory(uint64 BytesToFree, bool bInAllowRetainedChu
 					CurrentElement->DebugInfo.Reset();
 #endif
 					NumElementsEvicted++;
+					NumRetainedElementsEvicted++;
 				}
 			}
 
 			CurrentElement = CurrentElement->MoreRecentElement;
 		}
+
+		UE_LOG(LogAudioStreamCaching, Verbose, TEXT("TrimMemory: Num Retained Elements Evicted: %d. Retained Bytes Freed: %d"), NumRetainedElementsEvicted, RetainedBytesFreed);
 	}
 
-	UE_LOG(LogAudioStreamCaching, Verbose, TEXT("TrimMemory: NumElements Evicted: %d. Bytes Freed: %d"), NumElementsEvicted, BytesFreed);
+	UE_LOG(LogAudioStreamCaching, Verbose, TEXT("TrimMemory: Total Num Elements Evicted: %d. Total Bytes Freed: %d"), NumElementsEvicted, BytesFreed);
 
 	return BytesFreed;
 }
