Unify ContentBundlePaths interface
Made functions resilient to plugins being named "ContentBundle"

diff --git a/Engine/Source/Runtime/Engine/Private/Actor.cpp b/Engine/Source/Runtime/Engine/Private/Actor.cpp
--- a/Engine/Source/Runtime/Engine/Private/Actor.cpp
+++ b/Engine/Source/Runtime/Engine/Private/Actor.cpp
@@ -875,7 +875,7 @@ void AActor::Serialize(FArchive& Ar)
 
 		if (Ar.CustomVer(FFortniteMainBranchObjectVersion::GUID) < FFortniteMainBranchObjectVersion::WorldPartitionActorDescSerializeContentBundleGuid)
 		{
-			ContentBundleGuid = ContentBundlePaths::GetGuidFromPath(GetPackage()->GetFName().ToString());
+			ContentBundleGuid = ContentBundlePaths::GetContentBundleGuidFromExternalActorPackagePath(GetPackage()->GetFName().ToString());
 		}
 	}
 #endif
diff --git a/Engine/Source/Runtime/Engine/Private/WorldPartition/ContentBundle/ContentBundleEditor.cpp b/Engine/Source/Runtime/Engine/Private/WorldPartition/ContentBundle/ContentBundleEditor.cpp
--- a/Engine/Source/Runtime/Engine/Private/WorldPartition/ContentBundle/ContentBundleEditor.cpp
+++ b/Engine/Source/Runtime/Engine/Private/WorldPartition/ContentBundle/ContentBundleEditor.cpp
@@ -12,6 +12,7 @@
 #include "WorldPartition/ContentBundle/ContentBundleEditorSubsystemInterface.h"
 #include "WorldPartition/ContentBundle/ContentBundleWorldSubsystem.h"
 #include "WorldPartition/ContentBundle/ContentBundleLog.h"
+#include "WorldPartition/ContentBundle/ContentBundlePaths.h"
 #include "Engine/World.h"
 #include "PackageTools.h"
 #include "ObjectTools.h"
@@ -155,8 +156,8 @@ bool FContentBundleEditor::AddActor(AActor* InActor)
 	FSetActorContentBundleGuid SetActorContentBundleGuid(InActor, GetDescriptor()->GetGuid());
 
 	// Rename the actor so it is saved in the content bundle location
-	FName ActorPackageNameInContentBundle(*ULevel::GetActorPackageName(ActorDescContainer->GetExternalActorPath(), EActorPackagingScheme::Reduced, InActor->GetName()));
-	verify(InActor->GetPackage()->Rename(*ActorPackageNameInContentBundle.ToString()));
+	FString ActorPackageNameInContentBundle = ContentBundlePaths::MakeExternalActorPackagePath(ActorDescContainer->GetExternalActorPath(), InActor->GetName());
+	verify(InActor->GetPackage()->Rename(*ActorPackageNameInContentBundle));
 
 	UnsavedActorMonitor->MonitorActor(InActor);
 
diff --git a/Engine/Source/Runtime/Engine/Private/WorldPartition/ContentBundle/ContentBundlePaths.cpp b/Engine/Source/Runtime/Engine/Private/WorldPartition/ContentBundle/ContentBundlePaths.cpp
--- a/Engine/Source/Runtime/Engine/Private/WorldPartition/ContentBundle/ContentBundlePaths.cpp
+++ b/Engine/Source/Runtime/Engine/Private/WorldPartition/ContentBundle/ContentBundlePaths.cpp
@@ -3,42 +3,70 @@
 
 #if WITH_EDITOR
 #include "String/Find.h"
+#include "Engine/Level.h"
 #endif
 
 #if WITH_EDITOR
 
+FString ContentBundlePaths::MakeExternalActorPackagePath(const FString& ContentBundleExternalActorFolder, const FString& ActorName)
+{
+	const FString ContentBundleExternalActor = ULevel::GetActorPackageName(ContentBundleExternalActorFolder, EActorPackagingScheme::Reduced, *ActorName);
+	check(IsAContentBundlePackagePath(ContentBundleExternalActor));
+	return ContentBundleExternalActor;
+}
+
+bool ContentBundlePaths::IsAContentBundlePackagePath(FStringView InPackagePath)
+{
+	return GetContentBundleGuidFromExternalActorPackagePath(InPackagePath).IsValid();
+}
+
 FStringView ContentBundlePaths::GetRelativeExternalActorPackagePath(FStringView InContentBundleExternalActorPackagePath)
 {
-	uint32 ContentBundleStartIdx = UE::String::FindFirst(InContentBundleExternalActorPackagePath, ContentBundleFolder, ESearchCase::IgnoreCase);
-	if (ContentBundleStartIdx != INDEX_NONE)
+	FStringView RelativeContentBundlePath = GetActorPathRelativeToExternalActors(InContentBundleExternalActorPackagePath);
+	if (!RelativeContentBundlePath.IsEmpty())
 	{
-		uint32 ContentBundleUIDIdx = ContentBundleStartIdx + FCString::Strlen(ContentBundleFolder) + 1; // + 1 for trailing slash
-		FStringView RelativeActorPackge = InContentBundleExternalActorPackagePath.RightChop(ContentBundleUIDIdx);
-		
-		uint32 ContentBundleEndIdx = RelativeActorPackge.Find(TEXT("/"), ESearchCase::IgnoreCase); // Search for the end of the content bundle guids
-		check(ContentBundleEndIdx != INDEX_NONE);
-
-		return RelativeActorPackge.RightChop(ContentBundleEndIdx);
+		check(RelativeContentBundlePath.Left(GetContentBundleFolder().Len()).Equals(GetContentBundleFolder()));
+		RelativeContentBundlePath = RelativeContentBundlePath.RightChop(GetContentBundleFolder().Len());
+		if (!RelativeContentBundlePath.IsEmpty())
+		{
+			return RelativeContentBundlePath.RightChop(RelativeContentBundlePath.Find(TEXT("/")));
+		}
 	}
-
+	
 	return FStringView();
 }
 
-FGuid ContentBundlePaths::GetGuidFromPath(FStringView InPath)
+FGuid ContentBundlePaths::GetContentBundleGuidFromExternalActorPackagePath(FStringView InContentBundleExternalActorPackagePath)
 {
 	FGuid Result;
 
-	static const TCHAR* ContentBundlerFolderPath = TEXT("/ContentBundle/");
-	if (int FoundIdx = UE::String::FindFirst(InPath, ContentBundlerFolderPath, ESearchCase::IgnoreCase); FoundIdx != INDEX_NONE)
+	FStringView RelativeContentBundlePath = GetActorPathRelativeToExternalActors(InContentBundleExternalActorPackagePath);
+	if (!RelativeContentBundlePath.IsEmpty())
 	{
-		FString GuidPath(InPath);
-		const uint32 StartUIDIdx = FoundIdx + FCString::Strlen(ContentBundlerFolderPath);
-		const uint32 EndUIDIdx = GuidPath.Find(TEXT("/"), ESearchCase::IgnoreCase, ESearchDir::FromStart, StartUIDIdx);
-		GuidPath.MidInline(StartUIDIdx, EndUIDIdx - StartUIDIdx);		
-		verify(FGuid::Parse(GuidPath, Result));
+		check(RelativeContentBundlePath.Left(GetContentBundleFolder().Len()).Equals(GetContentBundleFolder()));
+		RelativeContentBundlePath = RelativeContentBundlePath.RightChop(GetContentBundleFolder().Len());
+		if (!RelativeContentBundlePath.IsEmpty())
+		{
+			RelativeContentBundlePath= RelativeContentBundlePath.LeftChop(RelativeContentBundlePath.Len() - RelativeContentBundlePath.Find(TEXT("/")));
+			verify(FGuid::Parse(FString(RelativeContentBundlePath), Result));
+		}
 	}
 
 	return Result;
 }
 
+FStringView ContentBundlePaths::GetActorPathRelativeToExternalActors(FStringView InContentBundleExternalActorPackagePath)
+{
+	uint32 ExternalActorIdx = UE::String::FindFirst(InContentBundleExternalActorPackagePath, FPackagePath::GetExternalActorsFolderName(), ESearchCase::IgnoreCase);
+	if (ExternalActorIdx != INDEX_NONE)
+	{
+		FStringView RelativeContentBundlePath = InContentBundleExternalActorPackagePath.RightChop(ExternalActorIdx + FCString::Strlen(FPackagePath::GetExternalActorsFolderName()));
+		if (RelativeContentBundlePath.Left(GetContentBundleFolder().Len()).Equals(GetContentBundleFolder()))
+		{
+			return RelativeContentBundlePath;
+		}
+	}
+	return FStringView();
+}
+
 #endif
\ No newline at end of file
diff --git a/Engine/Source/Runtime/Engine/Private/WorldPartition/WorldPartitionActorDesc.cpp b/Engine/Source/Runtime/Engine/Private/WorldPartition/WorldPartitionActorDesc.cpp
--- a/Engine/Source/Runtime/Engine/Private/WorldPartition/WorldPartitionActorDesc.cpp
+++ b/Engine/Source/Runtime/Engine/Private/WorldPartition/WorldPartitionActorDesc.cpp
@@ -464,7 +464,7 @@ PRAGMA_ENABLE_DEPRECATION_WARNINGS
 
 	if (Ar.CustomVer(FFortniteMainBranchObjectVersion::GUID) < FFortniteMainBranchObjectVersion::WorldPartitionActorDescSerializeContentBundleGuid)
 	{
-		ContentBundleGuid = ContentBundlePaths::GetGuidFromPath(ActorPackage.ToString());
+		ContentBundleGuid = ContentBundlePaths::GetContentBundleGuidFromExternalActorPackagePath(ActorPackage.ToString());
 	}
 	else
 	{
diff --git a/Engine/Source/Runtime/Engine/Public/WorldPartition/ContentBundle/ContentBundlePaths.h b/Engine/Source/Runtime/Engine/Public/WorldPartition/ContentBundle/ContentBundlePaths.h
--- a/Engine/Source/Runtime/Engine/Public/WorldPartition/ContentBundle/ContentBundlePaths.h
+++ b/Engine/Source/Runtime/Engine/Public/WorldPartition/ContentBundle/ContentBundlePaths.h
@@ -6,13 +6,38 @@
 
 namespace ContentBundlePaths
 {
-	constexpr TCHAR ContentBundleFolder[] = TEXT("ContentBundle");
+	constexpr FStringView ContentBundleFolder = TEXTVIEW("/ContentBundle/");
+
+	// Returns "/ContentBundle/"
+	constexpr FStringView GetContentBundleFolder()
+	{
+		return ContentBundleFolder;
+	}
+	
+	// Returns "ContentBundle"
+	constexpr FStringView GetContentBundleFolderName()
+	{
+		// Start at index 1 to remove the first "/", Lenght - 2 to remove the trailing "/"
+		return FStringView(&GetContentBundleFolder()[1], GetContentBundleFolder().Len() - 2);
+	}
 
 #if WITH_EDITOR
-	// InContentBundleExternalActorPackagePath format is : //{Plugin}/{ExternalActorFolder}/{ContentBundleFolder}/{ContentBundleUID}/{LevelPath}/{ExternalActorPackagePath}
-	// return format is {LevelPath}/{ExternalActorPackagePath}, empty otherwise
+	// return an ExternalActor path following the format : /{MountPoint}/{ExternalActorFolder}/{ContentBundleFolder}/{ContentBundleUID}/{LevelPath}/{ExternalActorPackagePath}
+	ENGINE_API FString MakeExternalActorPackagePath(const FString& ContentBundleExternalActorFolder, const FString& ActorName);
+
+	// return true if InPackage follow format : //{MountPoint}/{ExternalActorFolder}/{ContentBundleFolder}/{ContentBundleUID}/, false otherwise
+	ENGINE_API bool IsAContentBundlePackagePath(FStringView InPackagePath);
+
+	// InContentBundleExternalActorPackagePath format is : /{MountPoint}/{ExternalActorFolder}/{ContentBundleFolder}/{ContentBundleUID}/{LevelPath}/{ExternalActorPackagePath}
+	// return format is /{LevelPath}/{ExternalActorPackagePath}, empty otherwise
 	ENGINE_API FStringView GetRelativeExternalActorPackagePath(FStringView InContentBundleExternalActorPackagePath);
 
-	ENGINE_API FGuid GetGuidFromPath(FStringView InPath);
+	// InContentBundleExternalActorPackagePath format is : /{MountPoint}/{ExternalActorFolder}/{ContentBundleFolder}/{ContentBundleUID}/{LevelPath}/{ExternalActorPackagePath}
+	// return format is {ContentBundleUID}, 0 otherwise
+	ENGINE_API FGuid GetContentBundleGuidFromExternalActorPackagePath(FStringView InContentBundleExternalActorPackagePath);
+
+	// InExternalActorPath format is : /{MountPoint}/{ExternalActorFolder}/{ContentBundleFolder}/{ContentBundleUID}/{LevelPath}/{ExternalActorPackagePath}
+	// return format is /{ContentBundleFolder}/{ContentBundleUID}/{LevelPath}/{ExternalActorPackagePath}, empty otherwise
+	ENGINE_API FStringView GetActorPathRelativeToExternalActors(FStringView InContentBundleExternalActorPackagePath);
 #endif
 }
\ No newline at end of file
