Fix UnrealEditor NonUnity Win64 (Strict, VS2019) not trapped by preflight 6328c1c0f87253e021d66841.

#rnx

diff --git a/Engine/Source/Runtime/Engine/Public/WorldPartition/ActorDescContainer.h b/Engine/Source/Runtime/Engine/Public/WorldPartition/ActorDescContainer.h
--- a/Engine/Source/Runtime/Engine/Public/WorldPartition/ActorDescContainer.h
+++ b/Engine/Source/Runtime/Engine/Public/WorldPartition/ActorDescContainer.h
@@ -6,6 +6,7 @@
 #include "UObject/Object.h"
 #include "WorldPartition/ActorDescList.h"
 #include "WorldPartition/WorldPartitionHandle.h"
+#include "AssetRegistry/AssetData.h"
 #include "ActorDescContainer.generated.h"
 
 class FLinkerInstancingContext;
