Fix deprecated call to GetActortPath

diff --git a/Engine/Source/Runtime/Engine/Private/WorldPartition/WorldPartitionActorDescUtils.cpp b/Engine/Source/Runtime/Engine/Private/WorldPartition/WorldPartitionActorDescUtils.cpp
--- a/Engine/Source/Runtime/Engine/Private/WorldPartition/WorldPartitionActorDescUtils.cpp
+++ b/Engine/Source/Runtime/Engine/Private/WorldPartition/WorldPartitionActorDescUtils.cpp
@@ -75,7 +75,7 @@ TUniquePtr<FWorldPartitionActorDesc> FWorldPartitionActorDescUtils::GetActorDesc
 
 				if (!BlueprintClass.Num())
 				{
-					UE_LOG(LogWorldPartition, Warning, TEXT("Failed to find class '%s' for actor '%s"), *InAssetData.AssetClassPath.ToString(), *NewActorDesc->GetActorPath().ToString());
+					UE_LOG(LogWorldPartition, Warning, TEXT("Failed to find class '%s' for actor '%s"), *InAssetData.AssetClassPath.ToString(), *NewActorDesc->GetActorSoftPath().ToString());
 					return nullptr;
 				}
 			}
