World Partition - HLOD: Added support for automatic texture sizing for Merged & Simplified HLODs
* Reused existing code to share it between all HLOD generation methods

diff --git a/Engine/Plugins/Editor/WorldPartitionHLODUtilities/Source/Private/HLODBuilderMeshApproximate.cpp b/Engine/Plugins/Editor/WorldPartitionHLODUtilities/Source/Private/HLODBuilderMeshApproximate.cpp
--- a/Engine/Plugins/Editor/WorldPartitionHLODUtilities/Source/Private/HLODBuilderMeshApproximate.cpp
+++ b/Engine/Plugins/Editor/WorldPartitionHLODUtilities/Source/Private/HLODBuilderMeshApproximate.cpp
@@ -71,10 +71,10 @@ TArray<UActorComponent*> UHLODBuilderMeshApproximate::Build(const FHLODBuildCont
 {
 	TRACE_CPUPROFILER_EVENT_SCOPE(UHLODBuilderMeshApproximate::Build);
 
-	TArray<UStaticMeshComponent*> StaticMeshComponents = FilterComponents<UStaticMeshComponent>(InSourceComponents);
+	TArray<UPrimitiveComponent*> PrimitiveComponents = FilterComponents<UPrimitiveComponent>(InSourceComponents);
 
 	TSet<AActor*> Actors;
-	Algo::Transform(StaticMeshComponents, Actors, [](UPrimitiveComponent* PrimitiveComponent) { return PrimitiveComponent->GetOwner(); });
+	Algo::Transform(PrimitiveComponents, Actors, [](UPrimitiveComponent* PrimitiveComponent) { return PrimitiveComponent->GetOwner(); });
 	
 	IGeometryProcessingInterfacesModule& GeomProcInterfaces = FModuleManager::Get().LoadModuleChecked<IGeometryProcessingInterfacesModule>("GeometryProcessingInterfaces");
 	IGeometryProcessing_ApproximateActors* ApproxActorsAPI = GeomProcInterfaces.GetApproximateActorsImplementation();
@@ -103,45 +103,11 @@ TArray<UActorComponent*> UHLODBuilderMeshApproximate::Build(const FHLODBuildCont
 	Options.bUsePackedMRS = true;
 	Options.PackedMRSTexParamName = FName("PackedTexture");
 
-	// Gather bounds of the input components
-	auto GetActorsBounds = [&]() -> FBoxSphereBounds
-	{
-		FBoxSphereBounds Bounds;
-		bool bFirst = true;
-
-		for (UPrimitiveComponent* Component : StaticMeshComponents)
-		{
-			FBoxSphereBounds ComponentBounds = Component->Bounds;
-			Bounds = bFirst ? ComponentBounds : Bounds + ComponentBounds;
-			bFirst = false;
-		}
-
-		return Bounds;
-	};
-	
 	// Compute texel density if needed, depending on the TextureSizingType setting
-	ETextureSizingType TextureSizingType = UseSettings.MaterialSettings.TextureSizingType;
-	float TexelDensityPerMeter = 0.0f;
-
-	IGeometryProcessing_ApproximateActors::ETextureSizePolicy TextureSizePolicy = IGeometryProcessing_ApproximateActors::ETextureSizePolicy::TextureSize;
-	if (TextureSizingType == ETextureSizingType::TextureSizingType_AutomaticFromTexelDensity)
-	{
-		TexelDensityPerMeter = UseSettings.MaterialSettings.TargetTexelDensityPerMeter;
-		TextureSizePolicy = IGeometryProcessing_ApproximateActors::ETextureSizePolicy::TexelDensity;
-	}
-	else if (TextureSizingType == ETextureSizingType::TextureSizingType_AutomaticFromMeshScreenSize)
-	{
-		TexelDensityPerMeter = FMaterialUtilities::ComputeRequiredTexelDensityFromScreenSize(UseSettings.MaterialSettings.MeshMaxScreenSizePercent, GetActorsBounds().SphereRadius);
-		TextureSizePolicy = IGeometryProcessing_ApproximateActors::ETextureSizePolicy::TexelDensity;
+	if (UseSettings.MaterialSettings.ResolveTexelDensity(PrimitiveComponents, Options.MeshTexelDensity))
+	{ 
+		Options.TextureSizePolicy = IGeometryProcessing_ApproximateActors::ETextureSizePolicy::TexelDensity;
 	}
-	else if (TextureSizingType == ETextureSizingType::TextureSizingType_AutomaticFromMeshDrawDistance)
-	{
-		TexelDensityPerMeter = FMaterialUtilities::ComputeRequiredTexelDensityFromDrawDistance(UseSettings.MaterialSettings.MeshMinDrawDistance, GetActorsBounds().SphereRadius);
-		TextureSizePolicy = IGeometryProcessing_ApproximateActors::ETextureSizePolicy::TexelDensity;
-	}
-
-	Options.MeshTexelDensity = TexelDensityPerMeter;
-	Options.TextureSizePolicy = TextureSizePolicy;
 
 	// run actor approximation computation
 	IGeometryProcessing_ApproximateActors::FResults Results;
diff --git a/Engine/Plugins/Experimental/MeshModelingToolsetExp/Source/GeometryProcessingAdapters/GeometryProcessingAdapters.Build.cs b/Engine/Plugins/Experimental/MeshModelingToolsetExp/Source/GeometryProcessingAdapters/GeometryProcessingAdapters.Build.cs
--- a/Engine/Plugins/Experimental/MeshModelingToolsetExp/Source/GeometryProcessingAdapters/GeometryProcessingAdapters.Build.cs
+++ b/Engine/Plugins/Experimental/MeshModelingToolsetExp/Source/GeometryProcessingAdapters/GeometryProcessingAdapters.Build.cs
@@ -53,6 +53,7 @@ public class GeometryProcessingAdapters : ModuleRules
 				"Engine",
 				"RenderCore",
 
+				"MaterialUtilities",
 				"MeshBuilder",
 				"MeshUtilitiesCommon",  
 				"MeshReductionInterface", // for UE standard simplification 
diff --git a/Engine/Plugins/Experimental/MeshModelingToolsetExp/Source/GeometryProcessingAdapters/Private/ApproximateActorsImpl.cpp b/Engine/Plugins/Experimental/MeshModelingToolsetExp/Source/GeometryProcessingAdapters/Private/ApproximateActorsImpl.cpp
--- a/Engine/Plugins/Experimental/MeshModelingToolsetExp/Source/GeometryProcessingAdapters/Private/ApproximateActorsImpl.cpp
+++ b/Engine/Plugins/Experimental/MeshModelingToolsetExp/Source/GeometryProcessingAdapters/Private/ApproximateActorsImpl.cpp
@@ -36,6 +36,7 @@
 
 #include "Materials/Material.h"
 #include "Materials/MaterialInstanceConstant.h"
+#include "MaterialUtilities.h"
 
 #include "Engine/MeshMerging.h"
 
@@ -828,47 +829,6 @@ static TSharedPtr<FApproximationMeshData> GenerateApproximationMesh(
 	return Result;
 }
 
-
-static int32 GetMeshTextureSizeFromTargetTexelDensity(const FDynamicMesh3& Mesh, float TargetTexelDensity)
-{
-	const FDynamicMeshUVOverlay* UVOverlay = Mesh.Attributes()->PrimaryUV();
-	double Mesh3DArea = 0;
-	double MeshUVArea = 0;
-	for (int TriangleID : Mesh.TriangleIndicesItr())
-	{
-		// World space area
-		Mesh3DArea += Mesh.GetTriArea(TriangleID);
-
-		FIndex3i UVVertices = UVOverlay->GetTriangle(TriangleID);
-		FTriangle2d TriangleUV = FTriangle2d(
-			(FVector2d)UVOverlay->GetElement(UVVertices.A),
-			(FVector2d)UVOverlay->GetElement(UVVertices.B),
-			(FVector2d)UVOverlay->GetElement(UVVertices.C));
-
-		// UV space area
-		MeshUVArea += TriangleUV.Area();
-	}
-	double TexelRatio = FMath::Sqrt(MeshUVArea / Mesh3DArea) * 100;
-
-	// Compute the perfect texture size that would get us to our texture density
-	// Also compute the nearest power of two sizes (below and above our target)
-	const int32 SizePerfect = FMath::CeilToInt(TargetTexelDensity / TexelRatio);
-	const int32 SizeHi = FMath::RoundUpToPowerOfTwo(SizePerfect);
-	const int32 SizeLo = SizeHi >> 1;
-
-	// Compute the texel density we achieve with these two texture sizes
-	const double TexelDensityLo = SizeLo * TexelRatio;
-	const double TexelDensityHi = SizeHi * TexelRatio;
-
-	// Select best match between low & high res textures.
-	const double TexelDensityLoDiff = TargetTexelDensity - TexelDensityLo;
-	const double TexelDensityHiDiff = TexelDensityHi - TargetTexelDensity;
-	const int32 BestTextureSize = TexelDensityLoDiff < TexelDensityHiDiff ? SizeLo : SizeHi;
-
-	return BestTextureSize;
-}
-
-
 IGeometryProcessing_ApproximateActors::FOptions FApproximateActorsImpl::ConstructOptions(const FMeshApproximationSettings& UseSettings)
 {
 	//
@@ -1178,21 +1138,9 @@ void FApproximateActorsImpl::GenerateApproximationForActorSet(const TArray<AActo
 	// evaluate required texture size if needed
 	if (Options.TextureSizePolicy == ETextureSizePolicy::TexelDensity)
 	{
-		const int32 MaxTextureSize = 8192;
-		const int32 BestTextureSize = GetMeshTextureSizeFromTargetTexelDensity(FinalMesh, Options.MeshTexelDensity);
-
-		if (BestTextureSize > MaxTextureSize)
-		{
-			UE_LOG(LogApproximateActors, Warning, TEXT("Mesh would require %dx%d textures, clamping down to maximum (%dx%d)"), BestTextureSize, BestTextureSize, MaxTextureSize, MaxTextureSize);
-			OverridenOptions.TextureImageSize = MaxTextureSize;
-		}
-		else
-		{
-			OverridenOptions.TextureImageSize = BestTextureSize;
-		}
+		OverridenOptions.TextureImageSize = FMaterialUtilities::GetTextureSizeFromTargetTexelDensity(FinalMesh, Options.MeshTexelDensity);
 	}
 
-
 	// bake textures for Actor
 	FGeneratedResultTextures GeneratedTextures;
 	BakeTexturesFromPhotoCapture(SceneCapture, OverridenOptions,
diff --git a/Engine/Plugins/Experimental/ProxyLODPlugin/Source/ProxyLOD/Private/ProxyLODGrid2d.h b/Engine/Plugins/Experimental/ProxyLODPlugin/Source/ProxyLOD/Private/ProxyLODGrid2d.h
--- a/Engine/Plugins/Experimental/ProxyLODPlugin/Source/ProxyLOD/Private/ProxyLODGrid2d.h
+++ b/Engine/Plugins/Experimental/ProxyLODPlugin/Source/ProxyLOD/Private/ProxyLODGrid2d.h
@@ -13,9 +13,18 @@ namespace ProxyLOD
 	*/
 	struct FTextureAtlasDesc
 	{
-		FTextureAtlasDesc(const FIntPoint& S, const double G) :
-			Size(S), Gutter(G)
-		{}
+		FTextureAtlasDesc()
+			: Size()
+			, Gutter()
+		{
+		}
+
+		FTextureAtlasDesc(const FIntPoint& S, const double G)
+			: Size(S)
+			, Gutter(G)
+		{
+		}
+
 		FIntPoint Size;
 		float Gutter;
 	};
diff --git a/Engine/Plugins/Experimental/ProxyLODPlugin/Source/ProxyLOD/Private/ProxyLODMeshUtilities.cpp b/Engine/Plugins/Experimental/ProxyLODPlugin/Source/ProxyLOD/Private/ProxyLODMeshUtilities.cpp
--- a/Engine/Plugins/Experimental/ProxyLODPlugin/Source/ProxyLOD/Private/ProxyLODMeshUtilities.cpp
+++ b/Engine/Plugins/Experimental/ProxyLODPlugin/Source/ProxyLOD/Private/ProxyLODMeshUtilities.cpp
@@ -15,6 +15,7 @@ THIRD_PARTY_INCLUDES_START
 THIRD_PARTY_INCLUDES_END
 
 #include "StaticMeshAttributes.h"
+#include "TriangleTypes.h"
 
 #include <vector>
 #include <map>
@@ -1588,7 +1589,30 @@ int32 ProxyLOD::CorrectCollapsedWalls(FVertexDataMesh& InOutMesh,
 	BuildkDOPTree(InOutMesh, kDOPTree);
 
 	return CorrectCollapsedWalls(kDOPTree, InOutMesh.Indices, InOutMesh.Points, VoxelSize);
+}
+
+FBoxSphereBounds ProxyLOD::GetBounds(const FVertexDataMesh& InMesh)
+{
+	FBoxSphereBounds3f Bounds(InMesh.Points.GetData(), InMesh.Points.Num());
+	return FBoxSphereBounds(Bounds);
+}
+
+double ProxyLOD::GetWorldSpaceArea(const FVertexDataMesh& InMesh)
+{
+	const uint32* Indices = InMesh.Indices.GetData();
+	const FVector3f* Points = InMesh.Points.GetData();
+	const int32 NumTriangles = InMesh.Indices.Num() / 3;
+
+	double Mesh3DArea = 0;
+
+	FVector3f Pos[3];
+	for (int32 IdxTri = 0; IdxTri < NumTriangles; ++IdxTri)
+	{
+		const uint32 Offset = 3 * IdxTri;
+		Mesh3DArea += UE::Geometry::VectorUtil::Area(Points[Indices[Offset + 0]], Points[Indices[Offset + 1]], Points[Indices[Offset + 2]]);
+	}
 
+	return Mesh3DArea;
 }
 
 // Test to insure that no two vertexes share the same position.
diff --git a/Engine/Plugins/Experimental/ProxyLODPlugin/Source/ProxyLOD/Private/ProxyLODMeshUtilities.h b/Engine/Plugins/Experimental/ProxyLODPlugin/Source/ProxyLOD/Private/ProxyLODMeshUtilities.h
--- a/Engine/Plugins/Experimental/ProxyLODPlugin/Source/ProxyLOD/Private/ProxyLODMeshUtilities.h
+++ b/Engine/Plugins/Experimental/ProxyLODPlugin/Source/ProxyLOD/Private/ProxyLODMeshUtilities.h
@@ -147,6 +147,9 @@ namespace ProxyLOD
 								 const float LengthScale);
 	int32 CorrectCollapsedWalls( FVertexDataMesh& InOutMesh,
 		                         const float LengthScale );
+
+	FBoxSphereBounds GetBounds(const FVertexDataMesh& InMesh);
+	double GetWorldSpaceArea(const FVertexDataMesh& InMesh);
 	
 #if 0
 	// Can be used to transfer things like vertex colors from source geometry to output.
diff --git a/Engine/Plugins/Experimental/ProxyLODPlugin/Source/ProxyLOD/Private/ProxyLODPlugin.cpp b/Engine/Plugins/Experimental/ProxyLODPlugin/Source/ProxyLOD/Private/ProxyLODPlugin.cpp
--- a/Engine/Plugins/Experimental/ProxyLODPlugin/Source/ProxyLOD/Private/ProxyLODPlugin.cpp
+++ b/Engine/Plugins/Experimental/ProxyLODPlugin/Source/ProxyLOD/Private/ProxyLODPlugin.cpp
@@ -542,64 +542,13 @@ static float BBoxMajorAxisLength(const ProxyLOD::FBBox& BBox)
 *
 * NB: A min size of 64x64 is enforced.
 */
-static FIntPoint GetTexelGridSize(const  FMaterialProxySettings& MaterialSettings)
+static FIntPoint GetTexelGridSize(const FMaterialProxySettings& InMaterialSettings)
 {
-	auto MaxBBox = [](const FIntPoint& A, const FIntPoint B)
-	{
-		return FIntPoint(FMath::Max(A.X, B.X), FMath::Max(A.Y, B.Y));
-	};
-
-	// Start with a min value of 64x64.  
-	FIntPoint MaxTextureSize(64, 64);
-
-	const bool bUseTextureSize =
-		(MaterialSettings.TextureSizingType == ETextureSizingType::TextureSizingType_UseSingleTextureSize) ||
-		(MaterialSettings.TextureSizingType == ETextureSizingType::TextureSizingType_UseAutomaticBiasedSizes);
-
-	if (bUseTextureSize)
-	{
-		MaxTextureSize = MaxBBox(MaxTextureSize, MaterialSettings.TextureSize);
-	}
-	else if (MaterialSettings.TextureSizingType == ETextureSizingType::TextureSizingType_UseManualOverrideTextureSize)
-	{
-		MaxTextureSize = MaxBBox(MaxTextureSize, MaterialSettings.DiffuseTextureSize);
-
-		if (MaterialSettings.bNormalMap)
-		{
-			MaxTextureSize = MaxBBox(MaxTextureSize, MaterialSettings.NormalTextureSize);
-		}
-		if (MaterialSettings.bMetallicMap)
-		{
-			MaxTextureSize = MaxBBox(MaxTextureSize, MaterialSettings.MetallicTextureSize);
-		}
-		if (MaterialSettings.bRoughnessMap)
-		{
-			MaxTextureSize = MaxBBox(MaxTextureSize, MaterialSettings.RoughnessTextureSize);
-		}
-		if (MaterialSettings.bSpecularMap)
-		{
-			MaxTextureSize = MaxBBox(MaxTextureSize, MaterialSettings.SpecularTextureSize);
-		}
-		if (MaterialSettings.bEmissiveMap)
-		{
-			MaxTextureSize = MaxBBox(MaxTextureSize, MaterialSettings.EmissiveTextureSize);
-		}
-		if (MaterialSettings.bOpacityMap)
-		{
-			MaxTextureSize = MaxBBox(MaxTextureSize, MaterialSettings.OpacityTextureSize);
-		}
-		if (MaterialSettings.bOpacityMaskMap)
-		{
-			MaxTextureSize = MaxBBox(MaxTextureSize, MaterialSettings.OpacityMaskTextureSize);
-		}
-		if (MaterialSettings.bAmbientOcclusionMap)
-		{
-			MaxTextureSize = MaxBBox(MaxTextureSize, MaterialSettings.AmbientOcclusionTextureSize);
-		}
-	}
+	FIntPoint MaxTextureSize = InMaterialSettings.GetMaxTextureSize();
 
 	// Make the UVs square
 	int32 MaxLength = FMath::Max(MaxTextureSize.X, MaxTextureSize.Y);
+	MaxLength = FMath::Max(MaxLength, 64);
 	return FIntPoint(MaxLength, MaxLength);
 }
 
@@ -707,6 +656,7 @@ void FVoxelizeMeshMerging::ProxyLOD(FMeshDescriptionArrayAdapter& InSrcGeometryA
 		return;
 	}
 
+	FMaterialProxySettings MaterialSettings = InProxySettings.MaterialSettings;
 	
 	// Container for the raw mesh that will hold the simplified geometry
 	// and the FlattenMaterial that will hold the materials for the output mesh.
@@ -716,7 +666,7 @@ void FVoxelizeMeshMerging::ProxyLOD(FMeshDescriptionArrayAdapter& InSrcGeometryA
 	FMeshDescription OutRawMesh;
 	FStaticMeshAttributes(OutRawMesh).Register();
 
-	FFlattenMaterial OutMaterial         = FMaterialUtilities::CreateFlattenMaterialWithSettings(InProxySettings.MaterialSettings);
+	FFlattenMaterial OutMaterial;
 	const FColor UnresolvedGeometryColor = InProxySettings.UnresolvedGeometryColor;
 
 	bool bProxyGenerationSuccess = true;
@@ -781,13 +731,9 @@ void FVoxelizeMeshMerging::ProxyLOD(FMeshDescriptionArrayAdapter& InSrcGeometryA
 		FVertexDataMesh VertexDataMesh;
 
 		// Description of texture atlas.
-		// Make the UV space work with the largest texture size.
-		
-		const  FMaterialProxySettings& MaterialSettings = InProxySettings.MaterialSettings;
-		const FIntPoint UVSize = GetTexelGridSize(MaterialSettings);
-	
-		
-		ProxyLOD::FTextureAtlasDesc TextureAtlasDesc(UVSize, MaterialSettings.GutterSpace);
+		ProxyLOD::FTextureAtlasDesc TextureAtlasDesc;
+		TextureAtlasDesc.Size = FIntPoint(64, 64);
+		TextureAtlasDesc.Gutter = MaterialSettings.GutterSpace;
 	    
 		// --- Create New (High Poly) Geometry --
 		// 1) Voxelize the source geometry & maybe gap fill.
@@ -905,7 +851,7 @@ void FVoxelizeMeshMerging::ProxyLOD(FMeshDescriptionArrayAdapter& InSrcGeometryA
 
 			PrepGeometryAndBakeMaterialsTaskGroup.Run([&AOSMeshedVolume,
 				&SrcGeometryPolyField, &VertexDataMesh,
-				&TextureAtlasDesc, &bUVGenerationSuccess, &InProxySettings,
+				&TextureAtlasDesc, &MaterialSettings, &bUVGenerationSuccess, &InProxySettings,
 				VoxelSize, bColorVertsByChart, bSingleThreadedSimplify, bDoCollapsedWallFix, bSplitHardAngles, HardAngle]()
 			{
 				TRACE_CPUPROFILER_EVENT_SCOPE(PrepGeometryTask)
@@ -967,7 +913,13 @@ void FVoxelizeMeshMerging::ProxyLOD(FMeshDescriptionArrayAdapter& InSrcGeometryA
 				ProxyLOD::ComputeVertexNormals(VertexDataMesh, NormalMethod);
 
 
-				// 3) Generate UVs for Simplified Geometry
+				// 3) Resolve texture size if required
+				float WorldSpaceRadius = ProxyLOD::GetBounds(VertexDataMesh).SphereRadius;
+				double WorldSpaceArea = ProxyLOD::GetWorldSpaceArea(VertexDataMesh);
+				MaterialSettings.ResolveTextureSize(WorldSpaceRadius, WorldSpaceArea);
+				TextureAtlasDesc.Size = GetTexelGridSize(MaterialSettings);
+
+				// 4) Generate UVs for Simplified Geometry
 				// UV Atlas create, this can fail.
 				// NB: Vertices are split on UV seams. 
 
@@ -1007,6 +959,8 @@ void FVoxelizeMeshMerging::ProxyLOD(FMeshDescriptionArrayAdapter& InSrcGeometryA
 
 		// Using the new UVs fill OutMaterial texture atlas. 
 
+		OutMaterial = FMaterialUtilities::CreateFlattenMaterialWithSettings(MaterialSettings);
+
 		if (bUVGenerationSuccess && SrcGeometryPolyField)
 		{
 			ProxyLOD::FRasterGrid::Ptr DstUVGrid;
diff --git a/Engine/Source/Developer/HierarchicalLODUtilities/Private/HierarchicalLODUtilities.cpp b/Engine/Source/Developer/HierarchicalLODUtilities/Private/HierarchicalLODUtilities.cpp
--- a/Engine/Source/Developer/HierarchicalLODUtilities/Private/HierarchicalLODUtilities.cpp
+++ b/Engine/Source/Developer/HierarchicalLODUtilities/Private/HierarchicalLODUtilities.cpp
@@ -422,7 +422,7 @@ struct FHLODBuildParams
 	ALODActor*								LODActor;
 	UHLODProxy*								Proxy;
 	const TArray<UPrimitiveComponent*>&		Components;
-	const FHierarchicalSimplification&		LODSetup;
+	FHierarchicalSimplification				LODSetup;
 	UMaterialInterface*						BaseMaterial;
 	UPackage*								AssetsOuter;
 	FString									PackageName;
@@ -431,6 +431,15 @@ struct FHLODBuildParams
 		: Components(InComponents)
 		, LODSetup(InLODSetup)
 	{
+		FMaterialProxySettings* MaterialProxySettings = LODSetup.GetSimplificationMethodMaterialSettings();
+		if (MaterialProxySettings)
+		{
+			MaterialProxySettings->ResolveTexelDensity(InComponents);
+		}
+		else
+		{
+			UE_LOG(LogHierarchicalLODUtilities, Error, TEXT("Unsupported simplification method provided"));
+		}
 	}
 };
 
@@ -556,46 +565,6 @@ FHLODBuildResults GenerateHLODMesh_Approximate(const FHLODBuildParams& InBuildPa
 	Options.bUsePackedMRS = true;
 	Options.PackedMRSTexParamName = FName("PackedTexture");
 
-	// Gather bounds of the input components
-	auto GetActorsBounds = [&]() -> FBoxSphereBounds
-	{
-		FBoxSphereBounds Bounds;
-		bool bFirst = true;
-
-		for (UPrimitiveComponent* Component : InBuildParams.Components)
-		{
-			FBoxSphereBounds ComponentBounds = Component->Bounds;
-			Bounds = bFirst ? ComponentBounds : Bounds + ComponentBounds;
-			bFirst = false;
-		}
-
-		return Bounds;
-	};
-
-	// Compute texel density if needed, depending on the TextureSizingType setting
-	ETextureSizingType TextureSizingType = UseSettings.MaterialSettings.TextureSizingType;
-	float TexelDensityPerMeter = 0.0f;
-
-	IGeometryProcessing_ApproximateActors::ETextureSizePolicy TextureSizePolicy = IGeometryProcessing_ApproximateActors::ETextureSizePolicy::TextureSize;
-	if (TextureSizingType == ETextureSizingType::TextureSizingType_AutomaticFromTexelDensity)
-	{
-		TexelDensityPerMeter = UseSettings.MaterialSettings.TargetTexelDensityPerMeter;
-		TextureSizePolicy = IGeometryProcessing_ApproximateActors::ETextureSizePolicy::TexelDensity;
-	}
-	else if (TextureSizingType == ETextureSizingType::TextureSizingType_AutomaticFromMeshScreenSize)
-	{
-		TexelDensityPerMeter = FMaterialUtilities::ComputeRequiredTexelDensityFromScreenSize(UseSettings.MaterialSettings.MeshMaxScreenSizePercent, GetActorsBounds().SphereRadius);
-		TextureSizePolicy = IGeometryProcessing_ApproximateActors::ETextureSizePolicy::TexelDensity;
-	}
-	else if (TextureSizingType == ETextureSizingType::TextureSizingType_AutomaticFromMeshDrawDistance)
-	{
-		TexelDensityPerMeter = FMaterialUtilities::ComputeRequiredTexelDensityFromDrawDistance(UseSettings.MaterialSettings.MeshMinDrawDistance, GetActorsBounds().SphereRadius);
-		TextureSizePolicy = IGeometryProcessing_ApproximateActors::ETextureSizePolicy::TexelDensity;
-	}
-
-	Options.MeshTexelDensity = TexelDensityPerMeter;
-	Options.TextureSizePolicy = TextureSizePolicy;
-
 	// Use temp packages - otherwise Approximate Actors will create it's mesh using the same name as the HLOD Proxy object.
 	const FString NewAssetNamePrefix(TEXT("NEWASSET_"));
 	FString PackageName = InBuildParams.AssetsOuter->GetPackage()->GetName();
diff --git a/Engine/Source/Developer/MaterialUtilities/MaterialUtilities.Build.cs b/Engine/Source/Developer/MaterialUtilities/MaterialUtilities.Build.cs
--- a/Engine/Source/Developer/MaterialUtilities/MaterialUtilities.Build.cs
+++ b/Engine/Source/Developer/MaterialUtilities/MaterialUtilities.Build.cs
@@ -31,7 +31,8 @@ public class MaterialUtilities : ModuleRules
         PublicDependencyModuleNames.AddRange(
 			new string [] {
                  "MeshDescription",
-				 "StaticMeshDescription"
+				 "StaticMeshDescription",
+				 "GeometryCore"
 			}
 		);      
 
diff --git a/Engine/Source/Developer/MaterialUtilities/Private/MaterialUtilities.cpp b/Engine/Source/Developer/MaterialUtilities/Private/MaterialUtilities.cpp
--- a/Engine/Source/Developer/MaterialUtilities/Private/MaterialUtilities.cpp
+++ b/Engine/Source/Developer/MaterialUtilities/Private/MaterialUtilities.cpp
@@ -45,6 +45,10 @@
 #include "StaticMeshAttributes.h"
 #include "TextureCompiler.h"
 
+#include "TriangleTypes.h"
+#include "DynamicMesh/DynamicMesh3.h"
+#include "DynamicMesh/DynamicMeshAttributeSet.h"
+
 #if WITH_EDITOR
 #include "DeviceProfiles/DeviceProfile.h"
 #include "Tests/AutomationEditorCommon.h"
@@ -1508,11 +1512,11 @@ FFlattenMaterial FMaterialUtilities::CreateFlattenMaterialWithSettings(const FMa
 	// Create new material.
 	FFlattenMaterial Material;
 
-	// TODO REMOVE THIS FEATURE?
-	FIntPoint MaximumSize = FIntPoint::ZeroValue;// InMaterialLODSettings.TextureSize;
-	// If the user is manually overriding the texture size, make sure we have the max texture size to render with
+	FIntPoint MaximumSize = FIntPoint::ZeroValue;
+
 	if (InMaterialLODSettings.TextureSizingType == TextureSizingType_UseManualOverrideTextureSize)
 	{
+		// If the user is manually overriding the texture size, make sure we have the max texture size to render with
 		MaximumSize = (MaximumSize.X < InMaterialLODSettings.DiffuseTextureSize.X) ? InMaterialLODSettings.DiffuseTextureSize : MaximumSize ;
 		MaximumSize = (InMaterialLODSettings.bSpecularMap && (MaximumSize.X < InMaterialLODSettings.SpecularTextureSize.X)) ? InMaterialLODSettings.SpecularTextureSize	:	MaximumSize;
 		MaximumSize = (InMaterialLODSettings.bMetallicMap && (MaximumSize.X < InMaterialLODSettings.MetallicTextureSize.X)) ? InMaterialLODSettings.MetallicTextureSize	:	MaximumSize;
@@ -1520,12 +1524,8 @@ FFlattenMaterial FMaterialUtilities::CreateFlattenMaterialWithSettings(const FMa
 		MaximumSize = (InMaterialLODSettings.bNormalMap && (MaximumSize.X < InMaterialLODSettings.NormalTextureSize.X)) ? InMaterialLODSettings.NormalTextureSize :			MaximumSize;
 		MaximumSize = (InMaterialLODSettings.bEmissiveMap && (MaximumSize.X < InMaterialLODSettings.EmissiveTextureSize.X)) ? InMaterialLODSettings.EmissiveTextureSize :	MaximumSize;
 		MaximumSize = (InMaterialLODSettings.bOpacityMap && (MaximumSize.X < InMaterialLODSettings.OpacityTextureSize.X)) ? InMaterialLODSettings.OpacityTextureSize :		MaximumSize;
-	}
 	
-	if (InMaterialLODSettings.TextureSizingType == TextureSizingType_UseManualOverrideTextureSize)
-	{
 		Material.RenderSize = MaximumSize;
-
 		Material.SetPropertySize(EFlattenMaterialProperties::Diffuse, InMaterialLODSettings.DiffuseTextureSize);
 		Material.SetPropertySize(EFlattenMaterialProperties::Specular, InMaterialLODSettings.bSpecularMap ? InMaterialLODSettings.SpecularTextureSize : FIntPoint::ZeroValue);
 		Material.SetPropertySize(EFlattenMaterialProperties::Metallic, InMaterialLODSettings.bMetallicMap ? InMaterialLODSettings.MetallicTextureSize : FIntPoint::ZeroValue);
@@ -1538,15 +1538,13 @@ FFlattenMaterial FMaterialUtilities::CreateFlattenMaterialWithSettings(const FMa
 	{
 		Material.RenderSize = InMaterialLODSettings.TextureSize;
 		
-		int NormalSizeX, DiffuseSizeX, PropertiesSizeX;
-		NormalSizeX = InMaterialLODSettings.TextureSize.X;
-		DiffuseSizeX = FMath::Max(InMaterialLODSettings.TextureSize.X >> 1, 32);
-		PropertiesSizeX = FMath::Max(InMaterialLODSettings.TextureSize.X >> 2, 16);
+		int NormalSize = InMaterialLODSettings.TextureSize.X;
+		int DiffuseSize = FMath::Max(InMaterialLODSettings.TextureSize.X >> 1, 32);
+		int OtherSize = FMath::Max(InMaterialLODSettings.TextureSize.X >> 2, 16);
+		FIntPoint PropertiesSize = FIntPoint(OtherSize, OtherSize);
 
-		Material.SetPropertySize(EFlattenMaterialProperties::Diffuse, FIntPoint(DiffuseSizeX, DiffuseSizeX));
-		Material.SetPropertySize(EFlattenMaterialProperties::Normal, (InMaterialLODSettings.bNormalMap) ? FIntPoint(NormalSizeX, NormalSizeX) : FIntPoint::ZeroValue);
-
-		FIntPoint PropertiesSize = FIntPoint(PropertiesSizeX, PropertiesSizeX);
+		Material.SetPropertySize(EFlattenMaterialProperties::Diffuse, FIntPoint(DiffuseSize, DiffuseSize));
+		Material.SetPropertySize(EFlattenMaterialProperties::Normal, (InMaterialLODSettings.bNormalMap) ? FIntPoint(NormalSize, NormalSize) : FIntPoint::ZeroValue);
 		Material.SetPropertySize(EFlattenMaterialProperties::Specular, (InMaterialLODSettings.bSpecularMap) ? PropertiesSize : FIntPoint::ZeroValue );
 		Material.SetPropertySize(EFlattenMaterialProperties::Metallic, (InMaterialLODSettings.bMetallicMap) ? PropertiesSize : FIntPoint::ZeroValue );
 		Material.SetPropertySize(EFlattenMaterialProperties::Roughness, (InMaterialLODSettings.bRoughnessMap) ? PropertiesSize : FIntPoint::ZeroValue );
@@ -1564,6 +1562,10 @@ FFlattenMaterial FMaterialUtilities::CreateFlattenMaterialWithSettings(const FMa
 		Material.SetPropertySize(EFlattenMaterialProperties::Emissive, (InMaterialLODSettings.bEmissiveMap) ? InMaterialLODSettings.TextureSize : FIntPoint::ZeroValue);
 		Material.SetPropertySize(EFlattenMaterialProperties::Opacity, (InMaterialLODSettings.bOpacityMap) ? InMaterialLODSettings.TextureSize : FIntPoint::ZeroValue);
 	}
+	else
+	{
+		UE_LOG(LogMaterialUtilities, Error, TEXT("Unsupported TextureSizingType value. You should resolve the material texture size first with ResolveTextureSize()"));
+	}
 
 	return Material;
 }
@@ -2822,3 +2824,108 @@ float FMaterialUtilities::ComputeRequiredTexelDensityFromDrawDistance(const floa
 
 	return TexelDensityPerMeter;
 }
+
+static int32 ComputeTextureSizeFromTexelRatio(const double TexelRatio, const double TargetTexelDensity)
+{
+	// Compute the perfect texture size that would get us to our texture density
+	// Also compute the nearest power of two sizes (below and above our target)
+	const int32 SizePerfect = FMath::CeilToInt(TargetTexelDensity / TexelRatio);
+	const int32 SizeHi = FMath::RoundUpToPowerOfTwo(SizePerfect);
+	const int32 SizeLo = SizeHi >> 1;
+
+	// Compute the texel density we achieve with these two texture sizes
+	const double TexelDensityLo = SizeLo * TexelRatio;
+	const double TexelDensityHi = SizeHi * TexelRatio;
+
+	// Select best match between low & high res textures.
+	const double TexelDensityLoDiff = TargetTexelDensity - TexelDensityLo;
+	const double TexelDensityHiDiff = TexelDensityHi - TargetTexelDensity;
+	const int32 BestTextureSize = TexelDensityLoDiff < TexelDensityHiDiff ? SizeLo : SizeHi;
+
+	return BestTextureSize;
+}
+
+int32 FMaterialUtilities::GetTextureSizeFromTargetTexelDensity(const UE::Geometry::FDynamicMesh3& Mesh, float InTargetTexelDensity)
+{
+	const UE::Geometry::FDynamicMeshUVOverlay* UVOverlay = Mesh.Attributes()->PrimaryUV();
+
+	double Mesh3DArea = 0;
+	double MeshUVArea = 0;
+
+	// If no UVs, assume perfect UV space usage.
+	const bool bHasUVs = UVOverlay != nullptr;
+	if (!bHasUVs)
+	{
+		MeshUVArea = 1.0;
+	}
+
+	for (int TriangleID : Mesh.TriangleIndicesItr())
+	{
+		// World space area
+		Mesh3DArea += Mesh.GetTriArea(TriangleID);
+
+		// UV space area
+		if (bHasUVs)
+		{
+			UE::Geometry::FIndex3i UVVertices = UVOverlay->GetTriangle(TriangleID);
+			UE::Geometry::FTriangle2d TriangleUV = UE::Geometry::FTriangle2d(
+				(FVector2d)UVOverlay->GetElement(UVVertices.A),
+				(FVector2d)UVOverlay->GetElement(UVVertices.B),
+				(FVector2d)UVOverlay->GetElement(UVVertices.C));
+
+			MeshUVArea += TriangleUV.Area();
+		}
+	}
+
+	return GetTextureSizeFromTargetTexelDensity(Mesh3DArea, MeshUVArea, InTargetTexelDensity);
+}
+
+int32 FMaterialUtilities::GetTextureSizeFromTargetTexelDensity(const FMeshDescription& InMesh, float InTargetTexelDensity)
+{
+	FStaticMeshConstAttributes Attributes(InMesh);
+
+	TVertexAttributesConstRef<FVector3f> Positions = Attributes.GetVertexPositions();
+	TUVAttributesConstRef<FVector2f> UVs = Attributes.GetUVCoordinates(0);
+
+	double Mesh3DArea = 0;
+	double MeshUVArea = 0;
+
+	// If no UVs, assume perfect UV space usage.
+	const bool bHasUVs = UVs.IsValid() && UVs.GetNumElements() != 0;
+	if (!bHasUVs)
+	{
+		MeshUVArea = 1.0;
+	}
+
+	for (const FTriangleID TriangleID : InMesh.Triangles().GetElementIDs())
+	{
+		// World space area
+		TArrayView<const FVertexID> TriVertices = InMesh.GetTriangleVertices(TriangleID);
+		Mesh3DArea += UE::Geometry::VectorUtil::Area(Positions[TriVertices[0]], Positions[TriVertices[1]], Positions[TriVertices[2]]);
+
+		// UV space area
+		if (bHasUVs)
+		{
+			TArrayView<const FUVID> TriUVs = InMesh.GetTriangleUVIndices(TriangleID);
+			MeshUVArea += UE::Geometry::VectorUtil::Area(UVs[0], UVs[1], UVs[2]);
+		}
+	}
+
+	return GetTextureSizeFromTargetTexelDensity(Mesh3DArea, MeshUVArea, InTargetTexelDensity);
+}
+
+int32 FMaterialUtilities::GetTextureSizeFromTargetTexelDensity(double InMesh3DArea, double InMeshUVArea, double InTargetTexelDensity)
+{
+	double TexelRatio = FMath::Sqrt(InMeshUVArea / InMesh3DArea) * 100;
+
+	static const int32 MinTextureSize = 16;
+	static const int32 MaxTextureSize = 8192;
+
+	int32 TextureSize = ComputeTextureSizeFromTexelRatio(TexelRatio, InTargetTexelDensity);
+	if (TextureSize > MaxTextureSize)
+	{
+		UE_LOG(LogMaterialUtilities, Warning, TEXT("Mesh would require %d x %d textures, clamping down to maximum (%d x %d)"), TextureSize, TextureSize, MaxTextureSize, MaxTextureSize);
+	}
+
+	return FMath::Clamp(TextureSize, MinTextureSize, MaxTextureSize);
+}
diff --git a/Engine/Source/Developer/MaterialUtilities/Public/MaterialUtilities.h b/Engine/Source/Developer/MaterialUtilities/Public/MaterialUtilities.h
--- a/Engine/Source/Developer/MaterialUtilities/Public/MaterialUtilities.h
+++ b/Engine/Source/Developer/MaterialUtilities/Public/MaterialUtilities.h
@@ -240,6 +240,9 @@ class ULandscapeComponent;
 class FPrimitiveComponentId;
 class UMaterialInstanceConstant;
 struct FMaterialMergeData;
+struct FMeshDescription;
+
+namespace UE::Geometry { class FDynamicMesh3; }
 
 /**
  * Material utilities
@@ -491,6 +494,27 @@ public:
 	*/
 	static float ComputeRequiredTexelDensityFromDrawDistance(const float InDrawDistance, float InWorldSpaceRadius);
 
+	/** Compute the required texture size to achieve a target texel density for the given mesh.
+	* @param InMesh					The mesh for which we want to create a flatten material.
+	* @param InTargetTexelDensity	The target texel density.
+	* @return The texture size needed to achieve the required texel density.
+	*/
+	static int32 GetTextureSizeFromTargetTexelDensity(const FMeshDescription& InMesh, float InTargetTexelDensity);
+
+	/** Compute the required texture size to achieve a target texel density for the given mesh.
+	* @param InMesh					The mesh for which we want to create a flatten material.
+	* @param InTargetTexelDensity	The target texel density.
+	* @return The texture size needed to achieve the required texel density.
+	*/
+	static int32 GetTextureSizeFromTargetTexelDensity(const UE::Geometry::FDynamicMesh3& Mesh, float TargetTexelDensity);
+
+	/** Compute the required texture size to achieve a target texel density, given the world to UV space ratio.
+	* @param InMesh3DArea	World space area of the mesh
+	* @param InMeshUVArea	UV space area of the mesh
+	* @return The texture size needed to achieve the required texel density.
+	*/
+	static int32 GetTextureSizeFromTargetTexelDensity(double InMesh3DArea, double InMeshUVArea, double InTargetTexelDensity);
+
 	/** Validate that the provided material has all the required parameters needed to be considered a flattening material */
 	static bool IsValidFlattenMaterial(const UMaterialInterface* InBaseMaterial);
 
diff --git a/Engine/Source/Developer/MeshMergeUtilities/Private/MeshMergeDataTracker.cpp b/Engine/Source/Developer/MeshMergeUtilities/Private/MeshMergeDataTracker.cpp
--- a/Engine/Source/Developer/MeshMergeUtilities/Private/MeshMergeDataTracker.cpp
+++ b/Engine/Source/Developer/MeshMergeUtilities/Private/MeshMergeDataTracker.cpp
@@ -5,6 +5,8 @@
 #include "Misc/Crc.h"
 #include "Engine/StaticMesh.h"
 #include "StaticMeshAttributes.h"
+#include "TriangleTypes.h"
+#include "MaterialUtilities.h"
 
 FMeshMergeDataTracker::FMeshMergeDataTracker()
 	: AvailableLightMapUVChannel(INDEX_NONE), SummedLightMapPixels(0)
@@ -303,4 +305,28 @@ void FMeshMergeDataTracker::ProcessRawMeshes()
 			break;
 		}
 	}
+}
+
+double FMeshMergeDataTracker::GetTextureSizeFromTargetTexelDensity(float InTargetTexelDensity) const
+{
+	double Mesh3DArea = 0;
+	const double MeshUVArea = 1.0;		// UVs are not available yet, assume perfect UV space usage.
+
+	for (const TPair<FMeshLODKey, FMeshDescription>& MeshPair : RawMeshLODs)
+	{
+		const FMeshDescription& MeshDescription = MeshPair.Value;
+			
+		FStaticMeshConstAttributes Attributes(MeshDescription);
+		TVertexAttributesConstRef<FVector3f> Positions = Attributes.GetVertexPositions();
+		TUVAttributesConstRef<FVector2f> UVs = Attributes.GetUVCoordinates(0);
+
+		for (const FTriangleID TriangleID : MeshDescription.Triangles().GetElementIDs())
+		{
+			// World space area
+			TArrayView<const FVertexID> TriVertices = MeshDescription.GetTriangleVertices(TriangleID);
+			Mesh3DArea += UE::Geometry::VectorUtil::Area(Positions[TriVertices[0]], Positions[TriVertices[1]], Positions[TriVertices[2]]);
+		}
+	}
+
+	return FMaterialUtilities::GetTextureSizeFromTargetTexelDensity(Mesh3DArea, MeshUVArea, InTargetTexelDensity);
 }
\ No newline at end of file
diff --git a/Engine/Source/Developer/MeshMergeUtilities/Private/MeshMergeUtilities.cpp b/Engine/Source/Developer/MeshMergeUtilities/Private/MeshMergeUtilities.cpp
--- a/Engine/Source/Developer/MeshMergeUtilities/Private/MeshMergeUtilities.cpp
+++ b/Engine/Source/Developer/MeshMergeUtilities/Private/MeshMergeUtilities.cpp
@@ -69,6 +69,8 @@
 #include "RawMesh.h"
 #include "StaticMeshAttributes.h"
 #include "StaticMeshOperations.h"
+#include "TriangleTypes.h"
+#include "MaterialUtilities.h"
 
 #include "Async/Future.h"
 #include "Async/Async.h"
@@ -607,8 +609,6 @@ UMaterialOptions* FMeshMergeUtilities::PopulateMaterialOptions(const FMaterialPr
 	MaterialOptions->Properties.Empty();	
 	MaterialOptions->TextureSize = MaterialSettings.TextureSize;
 	
-	const bool bCustomSizes = MaterialSettings.TextureSizingType == TextureSizingType_UseManualOverrideTextureSize;
-
 	FPropertyEntry Property;
 	PopulatePropertyEntry(MaterialSettings, MP_BaseColor, Property);
 	MaterialOptions->Properties.Add(Property);
@@ -735,6 +735,10 @@ void FMeshMergeUtilities::PopulatePropertyEntry(const FMaterialProxySettings& Ma
 			InOutPropertyEntry.CustomSize = MaterialSettings.TextureSize;
 			break;
 		}
+
+		default:
+			UE_LOG(LogMeshMerging, Error, TEXT("Unsupported TextureSizingType value. You should resolve the material texture size first with ResolveTextureSize()"));
+
 	}
 	/** Check whether or not a constant value should be used for this property */
 	InOutPropertyEntry.bUseConstantValue = [MaterialSettings, MaterialProperty]() -> bool
@@ -1638,7 +1642,39 @@ void FMeshMergeUtilities::CreateProxyMesh(const TArray<UStaticMeshComponent*>& I
 	TArray<FMeshData> GlobalMeshSettings;
 	TArray<FMaterialData> GlobalMaterialSettings;
 
-	UMaterialOptions* Options = PopulateMaterialOptions(InMeshProxySettings.MaterialSettings);
+	FMaterialProxySettings MaterialProxySettings = InMeshProxySettings.MaterialSettings;
+	TArray<UPrimitiveComponent*> PrimitiveComponents;
+	Algo::Transform(StaticMeshComponents, PrimitiveComponents, [](UStaticMeshComponent* SMComponent) { return SMComponent; });
+	if (MaterialProxySettings.ResolveTexelDensity(PrimitiveComponents))
+	{
+		double Total3DArea = 0;
+
+		for (const FInstancedMeshDescriptionData& InstancedMeshDescriptionData : MeshDescriptionData)
+		{
+			double Mesh3DArea = 0;
+
+			const FMeshDescription& MeshDescription = *InstancedMeshDescriptionData.MeshDescription;
+
+			FStaticMeshConstAttributes Attributes(MeshDescription);
+			TVertexAttributesConstRef<FVector3f> Positions = Attributes.GetVertexPositions();
+
+			for (const FTriangleID TriangleID : MeshDescription.Triangles().GetElementIDs())
+			{
+				// World space area
+				TArrayView<const FVertexID> TriVertices = MeshDescription.GetTriangleVertices(TriangleID);
+				Mesh3DArea += UE::Geometry::VectorUtil::Area(Positions[TriVertices[0]], Positions[TriVertices[1]], Positions[TriVertices[2]]);
+			}
+
+			// Account for multiple instances (no transforms means a single instance)
+			uint32 NumInstances = FMath::Max(1, InstancedMeshDescriptionData.InstancesTransforms.Num());
+			Total3DArea += Mesh3DArea * NumInstances;
+		}
+
+		MaterialProxySettings.TextureSize = FMaterialUtilities::GetTextureSizeFromTargetTexelDensity(Total3DArea, 1.0f, MaterialProxySettings.TargetTexelDensityPerMeter);
+		MaterialProxySettings.TextureSizingType = ETextureSizingType::TextureSizingType_UseSingleTextureSize;
+	}
+
+	UMaterialOptions* Options = PopulateMaterialOptions(MaterialProxySettings);
 	TArray<EMaterialProperty> MaterialProperties;
 	for (const FPropertyEntry& Entry : Options->Properties)
 	{
@@ -2305,7 +2341,17 @@ void FMeshMergeUtilities::MergeComponentsToStaticMesh(const TArray<UPrimitiveCom
 	// If the user wants to merge materials into a single one
 	if (bMergeMaterialData && UniqueMaterials.Num() != 0)
 	{
-		UMaterialOptions* MaterialOptions = PopulateMaterialOptions(InSettings.MaterialSettings);
+		TArray<UPrimitiveComponent*> PrimitiveComponents;
+		Algo::Transform(StaticMeshComponentsToMerge, PrimitiveComponents, [](UStaticMeshComponent* SMComponent) { return SMComponent; });
+
+		FMaterialProxySettings MaterialProxySettings = InSettings.MaterialSettings;
+		if (MaterialProxySettings.ResolveTexelDensity(PrimitiveComponents))
+		{
+			MaterialProxySettings.TextureSize = DataTracker.GetTextureSizeFromTargetTexelDensity(MaterialProxySettings.TargetTexelDensityPerMeter);
+			MaterialProxySettings.TextureSizingType = ETextureSizingType::TextureSizingType_UseSingleTextureSize;
+		}
+
+		UMaterialOptions* MaterialOptions = PopulateMaterialOptions(MaterialProxySettings);
 		// Check each material to see if the shader actually uses vertex data and collect flags
 		TArray<bool> bMaterialUsesVertexData;
 		DetermineMaterialVertexDataUsage(bMaterialUsesVertexData, UniqueMaterials, MaterialOptions);
diff --git a/Engine/Source/Developer/MeshMergeUtilities/Private/ProxyGenerationProcessor.cpp b/Engine/Source/Developer/MeshMergeUtilities/Private/ProxyGenerationProcessor.cpp
--- a/Engine/Source/Developer/MeshMergeUtilities/Private/ProxyGenerationProcessor.cpp
+++ b/Engine/Source/Developer/MeshMergeUtilities/Private/ProxyGenerationProcessor.cpp
@@ -146,6 +146,8 @@ void FProxyGenerationProcessor::ProcessJob(const FGuid& JobGuid, FProxyGeneratio
 
 	if (!Data->RawMesh.IsEmpty())
 	{
+		Data->MergeData->InProxySettings.MaterialSettings.ResolveTextureSize(Data->RawMesh);
+
 		// Don't recreate render states with the material update context as we will manually do it through
 		// the FStaticMeshComponentRecreateRenderStateContext below
 		FMaterialUpdateContext MaterialUpdateContext(FMaterialUpdateContext::EOptions::Default & ~FMaterialUpdateContext::EOptions::RecreateRenderStates);
diff --git a/Engine/Source/Developer/MeshMergeUtilities/Public/MeshMergeDataTracker.h b/Engine/Source/Developer/MeshMergeUtilities/Public/MeshMergeDataTracker.h
--- a/Engine/Source/Developer/MeshMergeUtilities/Public/MeshMergeDataTracker.h
+++ b/Engine/Source/Developer/MeshMergeUtilities/Public/MeshMergeDataTracker.h
@@ -189,6 +189,8 @@ public:
 	void AddComponentToWedgeMapping(int32 MeshIndex, int32 LODIndex, uint32 WedgeIndex);
 	uint32 GetComponentToWedgeMappng(int32 MeshIndex, int32 LODIndex) const;
 
+	double GetTextureSizeFromTargetTexelDensity(float InTargetTexelDensity) const;
+
 protected:
 	// Mesh / LOD index, RawMesh
 	TMap<FMeshLODKey, FMeshDescription> RawMeshLODs;
diff --git a/Engine/Source/Editor/MergeActors/Private/MeshApproximationTool/MeshApproximationTool.cpp b/Engine/Source/Editor/MergeActors/Private/MeshApproximationTool/MeshApproximationTool.cpp
--- a/Engine/Source/Editor/MergeActors/Private/MeshApproximationTool/MeshApproximationTool.cpp
+++ b/Engine/Source/Editor/MergeActors/Private/MeshApproximationTool/MeshApproximationTool.cpp
@@ -30,6 +30,7 @@
 #include "AssetToolsModule.h"
 #include "ObjectTools.h"
 #include "Algo/ForEach.h"
+#include "Algo/Transform.h"
 
 #include "IGeometryProcessingInterfacesModule.h"
 #include "GeometryProcessingInterfaces/ApproximateActors.h"
@@ -169,53 +170,6 @@ bool FMeshApproximationTool::RunMerge(const FString& PackageName, const TArray<T
 		return false;
 	}
 
-    // Gather bounds of the input components
-	auto GetActorsBounds = [&]() -> FBoxSphereBounds
-	{
-		FBoxSphereBounds Bounds;
-		bool bFirst = true;
-
-		for (auto& Component : SelectedComponents)
-		{
-			if (!Component.IsValid() || !Component.Get()->bShouldIncorporate)
-			{
-				continue;
-			}
-
-			if (UPrimitiveComponent* PrimitiveComponent = Cast<UPrimitiveComponent>(Component.Get()->PrimComponent))
-			{
-				// Append bounds of visible components only
-				FBoxSphereBounds ComponentBounds = PrimitiveComponent->Bounds;
-				Bounds = bFirst ? ComponentBounds : Bounds + ComponentBounds;
-				bFirst = false;
-			}
-		}
-
-		return Bounds;
-	};
-	
-
-	// Compute texel density if needed, depending on the TextureSizingType setting
-	ETextureSizingType TextureSizingType = SettingsObject->Settings.MaterialSettings.TextureSizingType;
-	float TexelDensityPerMeter = 0.0f;
-
-	IGeometryProcessing_ApproximateActors::ETextureSizePolicy TextureSizePolicy = IGeometryProcessing_ApproximateActors::ETextureSizePolicy::TextureSize;
-	if (TextureSizingType == ETextureSizingType::TextureSizingType_AutomaticFromTexelDensity)
-	{
-		TexelDensityPerMeter = SettingsObject->Settings.MaterialSettings.TargetTexelDensityPerMeter;
-		TextureSizePolicy = IGeometryProcessing_ApproximateActors::ETextureSizePolicy::TexelDensity;
-	}
-	else if (TextureSizingType == ETextureSizingType::TextureSizingType_AutomaticFromMeshScreenSize)
-	{
-		TexelDensityPerMeter = FMaterialUtilities::ComputeRequiredTexelDensityFromScreenSize(SettingsObject->Settings.MaterialSettings.MeshMaxScreenSizePercent, GetActorsBounds().SphereRadius);
-		TextureSizePolicy = IGeometryProcessing_ApproximateActors::ETextureSizePolicy::TexelDensity;
-	}
-	else if (TextureSizingType == ETextureSizingType::TextureSizingType_AutomaticFromMeshDrawDistance)
-	{
-		TexelDensityPerMeter = FMaterialUtilities::ComputeRequiredTexelDensityFromDrawDistance(SettingsObject->Settings.MaterialSettings.MeshMinDrawDistance, GetActorsBounds().SphereRadius);
-		TextureSizePolicy = IGeometryProcessing_ApproximateActors::ETextureSizePolicy::TexelDensity;
-	}
-
 	const FMeshApproximationSettings& UseSettings = SettingsObject->Settings;
 
 	IGeometryProcessingInterfacesModule& GeomProcInterfaces = FModuleManager::Get().LoadModuleChecked<IGeometryProcessingInterfacesModule>("GeometryProcessingInterfaces");
@@ -226,8 +180,17 @@ bool FMeshApproximationTool::RunMerge(const FString& PackageName, const TArray<T
 	//
 
 	IGeometryProcessing_ApproximateActors::FOptions Options = ApproxActorsAPI->ConstructOptions(UseSettings);
-	Options.TextureSizePolicy = TextureSizePolicy;
-	Options.MeshTexelDensity = TexelDensityPerMeter;
+
+	TArray<UPrimitiveComponent*> PrimitiveComponents;
+	auto IsValidPrimitiveComponent = [](const TSharedPtr<FMergeComponentData>& Component) { return Component.IsValid() && !Component.Get()->bShouldIncorporate; };
+	auto GetPrimitiveComponent = [](const TSharedPtr<FMergeComponentData>& Component) { return Component.Get()->PrimComponent.Get(); };
+	Algo::TransformIf(SelectedComponents, PrimitiveComponents, IsValidPrimitiveComponent, GetPrimitiveComponent);
+
+	// Compute texel density if needed, depending on the TextureSizingType setting
+	if (UseSettings.MaterialSettings.ResolveTexelDensity(PrimitiveComponents, Options.MeshTexelDensity))
+	{
+		Options.TextureSizePolicy = IGeometryProcessing_ApproximateActors::ETextureSizePolicy::TexelDensity;
+	}
 
 	// Use temp packages - Needed to allow proper replacement of existing assets (performed below)
 	const FString NewAssetNamePrefix(TEXT("NEWASSET_"));
diff --git a/Engine/Source/Runtime/Engine/Classes/Engine/MaterialMerging.h b/Engine/Source/Runtime/Engine/Classes/Engine/MaterialMerging.h
--- a/Engine/Source/Runtime/Engine/Classes/Engine/MaterialMerging.h
+++ b/Engine/Source/Runtime/Engine/Classes/Engine/MaterialMerging.h
@@ -7,6 +7,8 @@
 #include "EngineTypes.h"
 #include "MaterialMerging.generated.h"
 
+struct FMeshDescription;
+
 UENUM()
 enum ETextureSizingType
 {
@@ -28,7 +30,7 @@ enum EMaterialMergeType
 };
 
 USTRUCT(Blueprintable)
-struct FMaterialProxySettings
+struct ENGINE_API FMaterialProxySettings
 {
 	GENERATED_USTRUCT_BODY()
 
@@ -37,7 +39,7 @@ struct FMaterialProxySettings
 	TEnumAsByte<ETextureSizingType> TextureSizingType;
 
 	// Size of generated BaseColor map
-	UPROPERTY(Category = Material, BlueprintReadWrite, EditAnywhere, meta =(ClampMin = "1", UIMin = "1", EditConditionHides, EditCondition = "TextureSizingType == ETextureSizingType::TextureSizingType_UseSingleTextureSize || TextureSizingType == ETextureSizingType::TextureSizingType_UseAutomaticBiasedSizes"))
+	UPROPERTY(Category = Material, EditAnywhere, meta =(ClampMin = "1", UIMin = "1", EditConditionHides, EditCondition = "TextureSizingType == ETextureSizingType::TextureSizingType_UseSingleTextureSize || TextureSizingType == ETextureSizingType::TextureSizingType_UseAutomaticBiasedSizes"))
 	FIntPoint TextureSize;
 
 	// Target texel density
@@ -179,90 +181,18 @@ struct FMaterialProxySettings
 	UPROPERTY(Category = Material, BlueprintReadWrite, AdvancedDisplay, EditAnywhere, meta = (ClampMin = "1", UIMin = "1"))
 	FIntPoint AmbientOcclusionTextureSize;
 
-	FMaterialProxySettings()
-		: TextureSizingType(TextureSizingType_UseSingleTextureSize)
-		, TextureSize(1024, 1024)
-		, TargetTexelDensityPerMeter(5.0f)
-		, MeshMaxScreenSizePercent(0.5f)
-		, MeshMinDrawDistance(10000.0f)
-		, GutterSpace(4.0f)
-		, MetallicConstant(0.0f)
-		, RoughnessConstant(0.5f)
-		, AnisotropyConstant(0.0f)
-		, SpecularConstant(0.5f)
-		, OpacityConstant(1.0f)
-		, OpacityMaskConstant(1.0f)
-		, AmbientOcclusionConstant(1.0f)
-		, MaterialMergeType(EMaterialMergeType::MaterialMergeType_Default)
-		, BlendMode(BLEND_Opaque)
-		, bAllowTwoSidedMaterial(true)
-		, bNormalMap(true)
-		, bTangentMap(false)
-		, bMetallicMap(false)
-		, bRoughnessMap(false)
-		, bAnisotropyMap(false)
-		, bSpecularMap(false)
-		, bEmissiveMap(false)
-		, bOpacityMap(false)
-		, bOpacityMaskMap(false)
-		, bAmbientOcclusionMap(false)
-		, DiffuseTextureSize(1024, 1024)
-		, NormalTextureSize(1024, 1024)
-		, TangentTextureSize(1024, 1024)
-		, MetallicTextureSize(1024, 1024)
-		, RoughnessTextureSize(1024, 1024)
-		, AnisotropyTextureSize(1024, 1024)
-		, SpecularTextureSize(1024, 1024)
-		, EmissiveTextureSize(1024, 1024)
-		, OpacityTextureSize(1024, 1024)
-		, OpacityMaskTextureSize(1024, 1024)
-		, AmbientOcclusionTextureSize(1024, 1024)
-	{
-	}
-
-	bool operator == (const FMaterialProxySettings& Other) const
-	{
-		return TextureSize == Other.TextureSize
-			&& TextureSizingType == Other.TextureSizingType
-			&& TargetTexelDensityPerMeter == Other.TargetTexelDensityPerMeter
-			&& MeshMaxScreenSizePercent == Other.MeshMaxScreenSizePercent
-			&& MeshMinDrawDistance == Other.MeshMinDrawDistance
-			&& GutterSpace == Other.GutterSpace
-			&& MetallicConstant == Other.MetallicConstant
-			&& RoughnessConstant == Other.RoughnessConstant
-			&& AnisotropyConstant == Other.AnisotropyConstant
-			&& SpecularConstant == Other.SpecularConstant
-			&& OpacityConstant == Other.OpacityConstant
-			&& OpacityMaskConstant == Other.OpacityMaskConstant
-			&& AmbientOcclusionConstant == Other.AmbientOcclusionConstant
-			&& MaterialMergeType == Other.MaterialMergeType
-			&& BlendMode == Other.BlendMode
-			&& bAllowTwoSidedMaterial == Other.bAllowTwoSidedMaterial
-			&& bNormalMap == Other.bNormalMap
-			&& bTangentMap == Other.bTangentMap
-			&& bMetallicMap == Other.bMetallicMap
-			&& bRoughnessMap == Other.bRoughnessMap
-			&& bAnisotropyMap == Other.bAnisotropyMap
-			&& bSpecularMap == Other.bSpecularMap
-			&& bEmissiveMap == Other.bEmissiveMap
-			&& bOpacityMap == Other.bOpacityMap
-			&& bOpacityMaskMap == Other.bOpacityMaskMap
-			&& bAmbientOcclusionMap == Other.bAmbientOcclusionMap
-			&& DiffuseTextureSize == Other.DiffuseTextureSize
-			&& NormalTextureSize == Other.NormalTextureSize
-			&& TangentTextureSize == Other.TangentTextureSize
-			&& MetallicTextureSize == Other.MetallicTextureSize
-			&& RoughnessTextureSize == Other.RoughnessTextureSize
-			&& AnisotropyTextureSize == Other.AnisotropyTextureSize
-			&& SpecularTextureSize == Other.SpecularTextureSize
-			&& EmissiveTextureSize == Other.EmissiveTextureSize
-			&& OpacityTextureSize == Other.OpacityTextureSize
-			&& OpacityMaskTextureSize == Other.OpacityMaskTextureSize
-			&& AmbientOcclusionTextureSize == Other.AmbientOcclusionTextureSize;
-	}
-
-	bool operator != (const FMaterialProxySettings& Other) const
-	{
-		return !(*this == Other);
-	}
+	FMaterialProxySettings();
+
+	bool operator == (const FMaterialProxySettings& Other) const;
+	bool operator != (const FMaterialProxySettings& Other) const;
+
+	FIntPoint GetMaxTextureSize() const;
+
+#if WITH_EDITOR
+	bool ResolveTexelDensity(const TArray<class UPrimitiveComponent*>& InComponents);
+	bool ResolveTexelDensity(const TArray<class UPrimitiveComponent*>& InComponents, float& OutTexelDensity) const;
+
+	void ResolveTextureSize(const FMeshDescription& InMesh);
+	void ResolveTextureSize(const float InWorldSpaceRadius, const double InWorldSpaceArea, const double InUVSpaceArea = 1.0);
+#endif
 };
diff --git a/Engine/Source/Runtime/Engine/Classes/GameFramework/WorldSettings.h b/Engine/Source/Runtime/Engine/Classes/GameFramework/WorldSettings.h
--- a/Engine/Source/Runtime/Engine/Classes/GameFramework/WorldSettings.h
+++ b/Engine/Source/Runtime/Engine/Classes/GameFramework/WorldSettings.h
@@ -359,6 +359,9 @@ struct ENGINE_API FHierarchicalSimplification
 	/** Handles deprecated properties */
 	void PostSerialize(const FArchive& Ar);
 #endif
+
+	/** Retrieve the correct material proxy settings based on the simplification method. */
+	FMaterialProxySettings* GetSimplificationMethodMaterialSettings();
 };
 
 template<>
diff --git a/Engine/Source/Runtime/Engine/Private/MaterialMerging.cpp b/Engine/Source/Runtime/Engine/Private/MaterialMerging.cpp
new file mode 100644
--- /dev/null
+++ b/Engine/Source/Runtime/Engine/Private/MaterialMerging.cpp
@@ -0,0 +1,269 @@
+// Copyright Epic Games, Inc. All Rights Reserved.
+
+#include "Engine/MaterialMerging.h"
+
+#if WITH_EDITOR
+#include "Components/PrimitiveComponent.h"
+#include "MaterialUtilities.h"
+#include "MeshDescription.h"
+#endif
+
+DEFINE_LOG_CATEGORY_STATIC(LogMaterialMerging, Log, All);
+
+FMaterialProxySettings::FMaterialProxySettings()
+	: TextureSizingType(TextureSizingType_UseSingleTextureSize)
+	, TextureSize(1024, 1024)
+	, TargetTexelDensityPerMeter(5.0f)
+	, MeshMaxScreenSizePercent(0.5f)
+	, MeshMinDrawDistance(10000.0f)
+	, GutterSpace(4.0f)
+	, MetallicConstant(0.0f)
+	, RoughnessConstant(0.5f)
+	, AnisotropyConstant(0.0f)
+	, SpecularConstant(0.5f)
+	, OpacityConstant(1.0f)
+	, OpacityMaskConstant(1.0f)
+	, AmbientOcclusionConstant(1.0f)
+	, MaterialMergeType(EMaterialMergeType::MaterialMergeType_Default)
+	, BlendMode(BLEND_Opaque)
+	, bAllowTwoSidedMaterial(true)
+	, bNormalMap(true)
+	, bTangentMap(false)
+	, bMetallicMap(false)
+	, bRoughnessMap(false)
+	, bAnisotropyMap(false)
+	, bSpecularMap(false)
+	, bEmissiveMap(false)
+	, bOpacityMap(false)
+	, bOpacityMaskMap(false)
+	, bAmbientOcclusionMap(false)
+	, DiffuseTextureSize(1024, 1024)
+	, NormalTextureSize(1024, 1024)
+	, TangentTextureSize(1024, 1024)
+	, MetallicTextureSize(1024, 1024)
+	, RoughnessTextureSize(1024, 1024)
+	, AnisotropyTextureSize(1024, 1024)
+	, SpecularTextureSize(1024, 1024)
+	, EmissiveTextureSize(1024, 1024)
+	, OpacityTextureSize(1024, 1024)
+	, OpacityMaskTextureSize(1024, 1024)
+	, AmbientOcclusionTextureSize(1024, 1024)
+{
+}
+
+bool FMaterialProxySettings::operator == (const FMaterialProxySettings& Other) const
+{
+	return TextureSize == Other.TextureSize
+		&& TextureSizingType == Other.TextureSizingType
+		&& TargetTexelDensityPerMeter == Other.TargetTexelDensityPerMeter
+		&& MeshMaxScreenSizePercent == Other.MeshMaxScreenSizePercent
+		&& MeshMinDrawDistance == Other.MeshMinDrawDistance
+		&& GutterSpace == Other.GutterSpace
+		&& MetallicConstant == Other.MetallicConstant
+		&& RoughnessConstant == Other.RoughnessConstant
+		&& AnisotropyConstant == Other.AnisotropyConstant
+		&& SpecularConstant == Other.SpecularConstant
+		&& OpacityConstant == Other.OpacityConstant
+		&& OpacityMaskConstant == Other.OpacityMaskConstant
+		&& AmbientOcclusionConstant == Other.AmbientOcclusionConstant
+		&& MaterialMergeType == Other.MaterialMergeType
+		&& BlendMode == Other.BlendMode
+		&& bAllowTwoSidedMaterial == Other.bAllowTwoSidedMaterial
+		&& bNormalMap == Other.bNormalMap
+		&& bTangentMap == Other.bTangentMap
+		&& bMetallicMap == Other.bMetallicMap
+		&& bRoughnessMap == Other.bRoughnessMap
+		&& bAnisotropyMap == Other.bAnisotropyMap
+		&& bSpecularMap == Other.bSpecularMap
+		&& bEmissiveMap == Other.bEmissiveMap
+		&& bOpacityMap == Other.bOpacityMap
+		&& bOpacityMaskMap == Other.bOpacityMaskMap
+		&& bAmbientOcclusionMap == Other.bAmbientOcclusionMap
+		&& DiffuseTextureSize == Other.DiffuseTextureSize
+		&& NormalTextureSize == Other.NormalTextureSize
+		&& TangentTextureSize == Other.TangentTextureSize
+		&& MetallicTextureSize == Other.MetallicTextureSize
+		&& RoughnessTextureSize == Other.RoughnessTextureSize
+		&& AnisotropyTextureSize == Other.AnisotropyTextureSize
+		&& SpecularTextureSize == Other.SpecularTextureSize
+		&& EmissiveTextureSize == Other.EmissiveTextureSize
+		&& OpacityTextureSize == Other.OpacityTextureSize
+		&& OpacityMaskTextureSize == Other.OpacityMaskTextureSize
+		&& AmbientOcclusionTextureSize == Other.AmbientOcclusionTextureSize;
+}
+
+bool FMaterialProxySettings::operator != (const FMaterialProxySettings& Other) const
+{
+	return !(*this == Other);
+}
+
+FIntPoint FMaterialProxySettings::GetMaxTextureSize() const
+{
+	auto MaxBBox = [](const FIntPoint& A, const FIntPoint B)
+	{
+		return FIntPoint(FMath::Max(A.X, B.X), FMath::Max(A.Y, B.Y));
+	};
+
+	FIntPoint MaxTextureSize = FIntPoint(1, 1);
+
+	const bool bUseTextureSize =
+		(TextureSizingType == ETextureSizingType::TextureSizingType_UseSingleTextureSize) ||
+		(TextureSizingType == ETextureSizingType::TextureSizingType_UseAutomaticBiasedSizes);
+
+	switch (TextureSizingType)
+	{
+	case ETextureSizingType::TextureSizingType_UseSingleTextureSize:
+	case ETextureSizingType::TextureSizingType_UseAutomaticBiasedSizes:
+		MaxTextureSize = MaxBBox(MaxTextureSize, TextureSize);
+		break;
+	
+	case ETextureSizingType::TextureSizingType_UseManualOverrideTextureSize:
+		MaxTextureSize = MaxBBox(MaxTextureSize, DiffuseTextureSize);
+		if (bNormalMap)
+		{
+			MaxTextureSize = MaxBBox(MaxTextureSize, NormalTextureSize);
+		}
+		if (bTangentMap)
+		{
+			MaxTextureSize = MaxBBox(MaxTextureSize, TangentTextureSize);			
+		}
+		if (bMetallicMap)
+		{
+			MaxTextureSize = MaxBBox(MaxTextureSize, MetallicTextureSize);
+		}
+		if (bRoughnessMap)
+		{
+			MaxTextureSize = MaxBBox(MaxTextureSize, RoughnessTextureSize);
+		}
+		if (bAnisotropyMap)
+		{
+			MaxTextureSize = MaxBBox(MaxTextureSize, AnisotropyTextureSize);
+		}
+		if (bSpecularMap)
+		{
+			MaxTextureSize = MaxBBox(MaxTextureSize, SpecularTextureSize);
+		}
+		if (bEmissiveMap)
+		{
+			MaxTextureSize = MaxBBox(MaxTextureSize, EmissiveTextureSize);
+		}
+		if (bOpacityMap)
+		{
+			MaxTextureSize = MaxBBox(MaxTextureSize, OpacityTextureSize);
+		}
+		if (bOpacityMaskMap)
+		{
+			MaxTextureSize = MaxBBox(MaxTextureSize, OpacityMaskTextureSize);
+		}
+		if (bAmbientOcclusionMap)
+		{
+			MaxTextureSize = MaxBBox(MaxTextureSize, AmbientOcclusionTextureSize);
+		}
+		break;
+
+	case TextureSizingType_UseSimplygonAutomaticSizing:
+	case TextureSizingType_AutomaticFromTexelDensity:
+	case TextureSizingType_AutomaticFromMeshScreenSize:
+	case TextureSizingType_AutomaticFromMeshDrawDistance:
+	default:
+		UE_LOG(LogMaterialMerging, Error, TEXT("Unsupported TextureSizingType value. You should resolve the material texture size first with ResolveTextureSize()"));
+		break;
+	}
+
+	return MaxTextureSize;
+}
+
+#if WITH_EDITOR
+
+bool FMaterialProxySettings::ResolveTexelDensity(const TArray<UPrimitiveComponent*>& InComponents)
+{
+	if (ResolveTexelDensity(InComponents, TargetTexelDensityPerMeter))
+	{
+		TextureSizingType = ETextureSizingType::TextureSizingType_AutomaticFromTexelDensity;
+		return true;
+	}
+
+	return false;
+}
+
+bool FMaterialProxySettings::ResolveTexelDensity(const TArray<UPrimitiveComponent*>& InComponents, float& OutTexelDensity) const
+{
+	// Gather bounds of the input components
+	auto GetActorsBounds = [&]() -> FBoxSphereBounds
+	{
+		FBoxSphereBounds Bounds;
+		bool bFirst = true;
+
+		for (UPrimitiveComponent* Component : InComponents)
+		{
+			if (Component)
+			{
+				FBoxSphereBounds ComponentBounds = Component->Bounds;
+				Bounds = bFirst ? ComponentBounds : Bounds + ComponentBounds;
+				bFirst = false;
+			}
+		}
+
+		return Bounds;
+	};
+
+	if (TextureSizingType == ETextureSizingType::TextureSizingType_AutomaticFromMeshScreenSize)
+	{
+		OutTexelDensity = FMaterialUtilities::ComputeRequiredTexelDensityFromScreenSize(MeshMaxScreenSizePercent, GetActorsBounds().SphereRadius);
+		return true;
+	}
+	else if (TextureSizingType == ETextureSizingType::TextureSizingType_AutomaticFromMeshDrawDistance)
+	{
+		OutTexelDensity = FMaterialUtilities::ComputeRequiredTexelDensityFromDrawDistance(MeshMinDrawDistance, GetActorsBounds().SphereRadius);
+		return true;
+	}
+	else if (TextureSizingType == ETextureSizingType::TextureSizingType_AutomaticFromTexelDensity)
+	{
+		OutTexelDensity = TargetTexelDensityPerMeter;
+		return true;
+	}
+
+	return false;
+}
+
+void FMaterialProxySettings::ResolveTextureSize(const FMeshDescription& InMesh)
+{
+	if (TextureSizingType == ETextureSizingType::TextureSizingType_AutomaticFromTexelDensity ||
+		TextureSizingType == ETextureSizingType::TextureSizingType_AutomaticFromMeshScreenSize ||
+		TextureSizingType == ETextureSizingType::TextureSizingType_AutomaticFromMeshDrawDistance)
+	{
+		if (TextureSizingType == ETextureSizingType::TextureSizingType_AutomaticFromMeshScreenSize)
+		{
+			TargetTexelDensityPerMeter = FMaterialUtilities::ComputeRequiredTexelDensityFromScreenSize(MeshMaxScreenSizePercent, InMesh.GetBounds().SphereRadius);
+		}
+		else if (TextureSizingType == ETextureSizingType::TextureSizingType_AutomaticFromMeshDrawDistance)
+		{
+			TargetTexelDensityPerMeter = FMaterialUtilities::ComputeRequiredTexelDensityFromDrawDistance(MeshMinDrawDistance, InMesh.GetBounds().SphereRadius);
+		}
+
+		TextureSize = FMaterialUtilities::GetTextureSizeFromTargetTexelDensity(InMesh, TargetTexelDensityPerMeter);
+		TextureSizingType = ETextureSizingType::TextureSizingType_UseSingleTextureSize;
+	}
+}
+
+void FMaterialProxySettings::ResolveTextureSize(const float InWorldSpaceRadius, const double InWorldSpaceArea, const double InUVSpaceArea)
+{
+	if (TextureSizingType == ETextureSizingType::TextureSizingType_AutomaticFromTexelDensity ||
+		TextureSizingType == ETextureSizingType::TextureSizingType_AutomaticFromMeshScreenSize ||
+		TextureSizingType == ETextureSizingType::TextureSizingType_AutomaticFromMeshDrawDistance)
+	{
+		if (TextureSizingType == ETextureSizingType::TextureSizingType_AutomaticFromMeshScreenSize)
+		{
+			TargetTexelDensityPerMeter = FMaterialUtilities::ComputeRequiredTexelDensityFromScreenSize(MeshMaxScreenSizePercent, InWorldSpaceRadius);
+		}
+		else if (TextureSizingType == ETextureSizingType::TextureSizingType_AutomaticFromMeshDrawDistance)
+		{
+			TargetTexelDensityPerMeter = FMaterialUtilities::ComputeRequiredTexelDensityFromDrawDistance(MeshMinDrawDistance, InWorldSpaceRadius);
+		}
+
+		TextureSize = FMaterialUtilities::GetTextureSizeFromTargetTexelDensity(InWorldSpaceArea, InUVSpaceArea, TargetTexelDensityPerMeter);
+		TextureSizingType = ETextureSizingType::TextureSizingType_UseSingleTextureSize;
+	}
+}
+
+#endif // WITH_EDITOR
diff --git a/Engine/Source/Runtime/Engine/Private/WorldSettings.cpp b/Engine/Source/Runtime/Engine/Private/WorldSettings.cpp
--- a/Engine/Source/Runtime/Engine/Private/WorldSettings.cpp
+++ b/Engine/Source/Runtime/Engine/Private/WorldSettings.cpp
@@ -1058,4 +1058,21 @@ void FHierarchicalSimplification::PostSerialize(const FArchive& Ar)
 
 #endif
 
+FMaterialProxySettings* FHierarchicalSimplification::GetSimplificationMethodMaterialSettings()
+{
+	switch (SimplificationMethod)
+	{
+	case EHierarchicalSimplificationMethod::Merge:
+		return &MergeSetting.MaterialSettings;
+
+	case EHierarchicalSimplificationMethod::Simplify:
+		return &ProxySetting.MaterialSettings;
+
+	case EHierarchicalSimplificationMethod::Approximate:
+		return &ApproximateSettings.MaterialSettings;
+	}
+
+	return nullptr;
+}
+
 #undef LOCTEXT_NAMESPACE
