PSO Precache Niagara & Cascade materials used by ParticleSystems

Precache is done at particle system level at not at component level - component only know their possible overrides just before drawing which is too late and almost all required materials and data is known at particle system post load
The vertex factory type and pritimitive type are retrieved for each particle type

diff --git a/Engine/Plugins/FX/Niagara/Source/Niagara/Classes/NiagaraSystem.h b/Engine/Plugins/FX/Niagara/Source/Niagara/Classes/NiagaraSystem.h
--- a/Engine/Plugins/FX/Niagara/Source/Niagara/Classes/NiagaraSystem.h
+++ b/Engine/Plugins/FX/Niagara/Source/Niagara/Classes/NiagaraSystem.h
@@ -334,6 +334,8 @@ public:
 
 	bool AllowScalabilityForLocalPlayerFX()const;
 
+	void PrecachePSOs();
+
 private:
 	bool IsReadyToRunInternal() const;
 
diff --git a/Engine/Plugins/FX/Niagara/Source/Niagara/Private/NiagaraMeshRendererProperties.cpp b/Engine/Plugins/FX/Niagara/Source/Niagara/Private/NiagaraMeshRendererProperties.cpp
--- a/Engine/Plugins/FX/Niagara/Source/Niagara/Private/NiagaraMeshRendererProperties.cpp
+++ b/Engine/Plugins/FX/Niagara/Source/Niagara/Private/NiagaraMeshRendererProperties.cpp
@@ -449,6 +449,11 @@ void UNiagaraMeshRendererProperties::GetUsedMeshMaterials(int32 MeshIndex, const
 	}
 }
 
+const FVertexFactoryType* UNiagaraMeshRendererProperties::GetVertexFactoryType() const
+{
+	return NeedsPreciseMotionVectors() ? &FNiagaraMeshVertexFactoryEx::StaticType : &FNiagaraMeshVertexFactory::StaticType;
+}
+
 void UNiagaraMeshRendererProperties::GetUsedMaterials(const FNiagaraEmitterInstance* InEmitter, TArray<UMaterialInterface*>& OutMaterials) const
 {
 	TArray<UMaterialInterface*> OrderedMeshMaterials;
diff --git a/Engine/Plugins/FX/Niagara/Source/Niagara/Private/NiagaraRibbonRendererProperties.cpp b/Engine/Plugins/FX/Niagara/Source/Niagara/Private/NiagaraRibbonRendererProperties.cpp
--- a/Engine/Plugins/FX/Niagara/Source/Niagara/Private/NiagaraRibbonRendererProperties.cpp
+++ b/Engine/Plugins/FX/Niagara/Source/Niagara/Private/NiagaraRibbonRendererProperties.cpp
@@ -174,6 +174,11 @@ void UNiagaraRibbonRendererProperties::GetUsedMaterials(const FNiagaraEmitterIns
 	OutMaterials.Add(MaterialInterface ? MaterialInterface : ToRawPtr(Material));
 }
 
+const FVertexFactoryType* UNiagaraRibbonRendererProperties::GetVertexFactoryType() const
+{
+	return &FNiagaraRibbonVertexFactory::StaticType;
+}
+
 bool UNiagaraRibbonRendererProperties::PopulateRequiredBindings(FNiagaraParameterStore& InParameterStore)
 {
 	bool bAnyAdded = Super::PopulateRequiredBindings(InParameterStore);
diff --git a/Engine/Plugins/FX/Niagara/Source/Niagara/Private/NiagaraSpriteRendererProperties.cpp b/Engine/Plugins/FX/Niagara/Source/Niagara/Private/NiagaraSpriteRendererProperties.cpp
--- a/Engine/Plugins/FX/Niagara/Source/Niagara/Private/NiagaraSpriteRendererProperties.cpp
+++ b/Engine/Plugins/FX/Niagara/Source/Niagara/Private/NiagaraSpriteRendererProperties.cpp
@@ -119,6 +119,11 @@ void UNiagaraSpriteRendererProperties::GetUsedMaterials(const FNiagaraEmitterIns
 	OutMaterials.Add(MaterialInterface ? MaterialInterface : ToRawPtr(Material));
 }
 
+const FVertexFactoryType* UNiagaraSpriteRendererProperties::GetVertexFactoryType() const
+{
+	return NeedsPreciseMotionVectors() ? &FNiagaraSpriteVertexFactoryEx::StaticType : &FNiagaraSpriteVertexFactory::StaticType;
+}
+
 void UNiagaraSpriteRendererProperties::PostLoad()
 {
 	Super::PostLoad();
diff --git a/Engine/Plugins/FX/Niagara/Source/Niagara/Private/NiagaraSystem.cpp b/Engine/Plugins/FX/Niagara/Source/Niagara/Private/NiagaraSystem.cpp
--- a/Engine/Plugins/FX/Niagara/Source/Niagara/Private/NiagaraSystem.cpp
+++ b/Engine/Plugins/FX/Niagara/Source/Niagara/Private/NiagaraSystem.cpp
@@ -33,6 +33,7 @@
 #include "UObject/Package.h"
 #include "UObject/UObjectThreadContext.h"
 #include "ShaderCompiler.h"
+#include "PipelineStateCache.h"
 
 #define LOCTEXT_NAMESPACE "NiagaraSystem"
 
@@ -939,8 +940,71 @@ void UNiagaraSystem::PostLoad()
 		TemplateSpecification = bIsTemplateAsset_DEPRECATED ? ENiagaraScriptTemplateSpecification::Template : ENiagaraScriptTemplateSpecification::None;
 	}
 #endif // WITH_EDITORONLY_DATA
+
+	PrecachePSOs();
 }
 
+void UNiagaraSystem::PrecachePSOs()
+{
+	if (!FApp::CanEverRender() || !PipelineStateCache::IsPSOPrecachingEnabled())
+	{
+		return;
+	}
+
+	struct VFsPerMaterialData
+	{
+		UMaterialInterface* MaterialInterface;
+		TArray<const FVertexFactoryType*, TInlineAllocator<2>> VertexFactoryTypes;
+	};
+	TArray<VFsPerMaterialData, TInlineAllocator<2>> VFsPerMaterials;
+
+	for (const FNiagaraEmitterHandle& EmitterHandle : GetEmitterHandles())
+	{
+		FVersionedNiagaraEmitterData* EmitterData = EmitterHandle.GetEmitterData();
+		if (EmitterData && EmitterHandle.GetIsEnabled())
+		{
+			EmitterData->ForEachEnabledRenderer(
+				[&](UNiagaraRendererProperties* Properties)
+				{
+					const FVertexFactoryType* VFType = Properties->GetVertexFactoryType();
+					if (VFType == nullptr)
+					{
+						return;
+					}
+
+					// Don't have an instance yet to retrieve the possible material override data from
+					const FNiagaraEmitterInstance* EmitterInstance = nullptr;
+					TArray<UMaterialInterface*> Mats;
+					Properties->GetUsedMaterials(EmitterInstance, Mats);
+
+					for (UMaterialInterface* MaterialInterface : Mats)
+					{
+						VFsPerMaterialData* VFsPerMaterial = VFsPerMaterials.FindByPredicate([MaterialInterface](const VFsPerMaterialData& Other) { return Other.MaterialInterface == MaterialInterface; });
+						if (VFsPerMaterial == nullptr)
+						{
+							VFsPerMaterial = &VFsPerMaterials.AddDefaulted_GetRef();
+							VFsPerMaterial->MaterialInterface = MaterialInterface;
+						}
+						VFsPerMaterial->VertexFactoryTypes.AddUnique(VFType);
+					}
+				}
+			);
+		}
+	}
+
+	FPSOPrecacheParams PreCachePSOParams;
+	PreCachePSOParams.SetMobility(EComponentMobility::Movable);
+
+	for (VFsPerMaterialData& VFsPerMaterial : VFsPerMaterials)
+	{
+		if (VFsPerMaterial.MaterialInterface)
+		{
+			VFsPerMaterial.MaterialInterface->PrecachePSOs(VFsPerMaterial.VertexFactoryTypes, PreCachePSOParams);
+		}
+	}
+}
+
+
 #if WITH_EDITORONLY_DATA
 
 UNiagaraEditorDataBase* UNiagaraSystem::GetEditorData()
diff --git a/Engine/Plugins/FX/Niagara/Source/Niagara/Public/NiagaraMeshRendererProperties.h b/Engine/Plugins/FX/Niagara/Source/Niagara/Public/NiagaraMeshRendererProperties.h
--- a/Engine/Plugins/FX/Niagara/Source/Niagara/Public/NiagaraMeshRendererProperties.h
+++ b/Engine/Plugins/FX/Niagara/Source/Niagara/Public/NiagaraMeshRendererProperties.h
@@ -177,6 +177,7 @@ public:
 	virtual class FNiagaraBoundsCalculator* CreateBoundsCalculator() override;
 	virtual void GetUsedMaterials(const FNiagaraEmitterInstance* InEmitter, TArray<UMaterialInterface*>& OutMaterials) const override;
 	virtual void GetStreamingMeshInfo(const FBoxSphereBounds& OwnerBounds, const FNiagaraEmitterInstance* InEmitter, TArray<FStreamingRenderAssetPrimitiveInfo>& OutStreamingRenderAssets) const override;
+	virtual const FVertexFactoryType* GetVertexFactoryType() const override;
 	virtual bool IsSimTargetSupported(ENiagaraSimTarget InSimTarget) const override { return true; };
 	virtual bool PopulateRequiredBindings(FNiagaraParameterStore& InParameterStore) override;
 #if WITH_EDITORONLY_DATA
diff --git a/Engine/Plugins/FX/Niagara/Source/Niagara/Public/NiagaraRendererProperties.h b/Engine/Plugins/FX/Niagara/Source/Niagara/Public/NiagaraRendererProperties.h
--- a/Engine/Plugins/FX/Niagara/Source/Niagara/Public/NiagaraRendererProperties.h
+++ b/Engine/Plugins/FX/Niagara/Source/Niagara/Public/NiagaraRendererProperties.h
@@ -263,6 +263,7 @@ public:
 	virtual FNiagaraRenderer* CreateEmitterRenderer(ERHIFeatureLevel::Type FeatureLevel, const FNiagaraEmitterInstance* Emitter, const FNiagaraSystemInstanceController& InController) PURE_VIRTUAL ( UNiagaraRendererProperties::CreateEmitterRenderer, return nullptr;);
 	virtual class FNiagaraBoundsCalculator* CreateBoundsCalculator() PURE_VIRTUAL(UNiagaraRendererProperties::CreateBoundsCalculator, return nullptr;);
 	virtual void GetUsedMaterials(const FNiagaraEmitterInstance* InEmitter, TArray<UMaterialInterface*>& OutMaterials) const PURE_VIRTUAL(UNiagaraRendererProperties::GetUsedMaterials,);
+	virtual const FVertexFactoryType* GetVertexFactoryType() const { return nullptr; }
 
 	virtual void GetStreamingMeshInfo(const FBoxSphereBounds& OwnerBounds, const FNiagaraEmitterInstance* InEmitter, TArray<FStreamingRenderAssetPrimitiveInfo>& OutStreamingRenderAssets) const {}
 
diff --git a/Engine/Plugins/FX/Niagara/Source/Niagara/Public/NiagaraRibbonRendererProperties.h b/Engine/Plugins/FX/Niagara/Source/Niagara/Public/NiagaraRibbonRendererProperties.h
--- a/Engine/Plugins/FX/Niagara/Source/Niagara/Public/NiagaraRibbonRendererProperties.h
+++ b/Engine/Plugins/FX/Niagara/Source/Niagara/Public/NiagaraRibbonRendererProperties.h
@@ -213,6 +213,7 @@ public:
 	virtual FNiagaraRenderer* CreateEmitterRenderer(ERHIFeatureLevel::Type FeatureLevel, const FNiagaraEmitterInstance* Emitter, const FNiagaraSystemInstanceController& InController) override;
 	virtual class FNiagaraBoundsCalculator* CreateBoundsCalculator() override;
 	virtual void GetUsedMaterials(const FNiagaraEmitterInstance* InEmitter, TArray<UMaterialInterface*>& OutMaterials) const override;
+	virtual const FVertexFactoryType* GetVertexFactoryType() const override;
 	virtual bool IsSimTargetSupported(ENiagaraSimTarget InSimTarget) const override { return true; };
 	virtual bool PopulateRequiredBindings(FNiagaraParameterStore& InParameterStore) override;
 
diff --git a/Engine/Plugins/FX/Niagara/Source/Niagara/Public/NiagaraSpriteRendererProperties.h b/Engine/Plugins/FX/Niagara/Source/Niagara/Public/NiagaraSpriteRendererProperties.h
--- a/Engine/Plugins/FX/Niagara/Source/Niagara/Public/NiagaraSpriteRendererProperties.h
+++ b/Engine/Plugins/FX/Niagara/Source/Niagara/Public/NiagaraSpriteRendererProperties.h
@@ -123,6 +123,7 @@ public:
 	virtual FNiagaraRenderer* CreateEmitterRenderer(ERHIFeatureLevel::Type FeatureLevel, const FNiagaraEmitterInstance* Emitter, const FNiagaraSystemInstanceController& InController) override;
 	virtual class FNiagaraBoundsCalculator* CreateBoundsCalculator() override;
 	virtual void GetUsedMaterials(const FNiagaraEmitterInstance* InEmitter, TArray<UMaterialInterface*>& OutMaterials) const override;
+	virtual const FVertexFactoryType* GetVertexFactoryType() const override;
 	virtual bool IsSimTargetSupported(ENiagaraSimTarget InSimTarget) const override { return true; };
 	virtual bool PopulateRequiredBindings(FNiagaraParameterStore& InParameterStore)  override;
 #if WITH_EDITOR
diff --git a/Engine/Plugins/FX/Niagara/Source/NiagaraShader/Private/NiagaraShared.cpp b/Engine/Plugins/FX/Niagara/Source/NiagaraShader/Private/NiagaraShared.cpp
--- a/Engine/Plugins/FX/Niagara/Source/NiagaraShader/Private/NiagaraShared.cpp
+++ b/Engine/Plugins/FX/Niagara/Source/NiagaraShader/Private/NiagaraShared.cpp
@@ -516,6 +516,14 @@ void FNiagaraShaderScript::UpdateCachedData_PostCompile(bool bCalledFromSerializ
 				CachedData.bExternalConstantBufferUsed |= NiagaraShader->ExternalConstantBufferParam[i].IsBound() ? BitToSet : 0;
 			}
 			CachedData.bViewUniformBufferUsed |= NiagaraShader->bNeedsViewUniformBuffer;
+
+			// request precache the compute shader
+			{
+				check(NiagaraShader->GetFrequency() == SF_Compute)
+				FRHIShader* RHIShader = GameThreadShaderMap->GetResource()->GetShader(Shader->GetResourceIndex());
+				FRHIComputeShader* RHIComputeShader = static_cast<FRHIComputeShader*>(RHIShader);
+				PipelineStateCache::PrecacheComputePipelineState(RHIComputeShader);
+			}
 		}
 	}
 	else
diff --git a/Engine/Source/Runtime/Engine/Classes/Particles/ParticleSystem.h b/Engine/Source/Runtime/Engine/Classes/Particles/ParticleSystem.h
--- a/Engine/Source/Runtime/Engine/Classes/Particles/ParticleSystem.h
+++ b/Engine/Source/Runtime/Engine/Classes/Particles/ParticleSystem.h
@@ -437,6 +437,8 @@ public:
 
 	//~ End UObject Interface.
 
+	void PrecachePSOs();
+
 	bool CanBePooled()const;
 
 	// @todo document
diff --git a/Engine/Source/Runtime/Engine/Classes/Particles/TypeData/ParticleModuleTypeDataAnimTrail.h b/Engine/Source/Runtime/Engine/Classes/Particles/TypeData/ParticleModuleTypeDataAnimTrail.h
--- a/Engine/Source/Runtime/Engine/Classes/Particles/TypeData/ParticleModuleTypeDataAnimTrail.h
+++ b/Engine/Source/Runtime/Engine/Classes/Particles/TypeData/ParticleModuleTypeDataAnimTrail.h
@@ -81,6 +81,8 @@ class UParticleModuleTypeDataAnimTrail : public UParticleModuleTypeDataBase
 
 	//~ Begin UParticleModuleTypeDataBase Interface
 	virtual FParticleEmitterInstance* CreateInstance(UParticleEmitter* InEmitterParent, UParticleSystemComponent* InComponent) override;
+	virtual const FVertexFactoryType* GetVertexFactoryType() const override;
+	virtual EPrimitiveType GetPrimitiveType() const override { return PT_TriangleStrip; }
 	//~ End UParticleModuleTypeDataBase Interface
 };
 
diff --git a/Engine/Source/Runtime/Engine/Classes/Particles/TypeData/ParticleModuleTypeDataBase.h b/Engine/Source/Runtime/Engine/Classes/Particles/TypeData/ParticleModuleTypeDataBase.h
--- a/Engine/Source/Runtime/Engine/Classes/Particles/TypeData/ParticleModuleTypeDataBase.h
+++ b/Engine/Source/Runtime/Engine/Classes/Particles/TypeData/ParticleModuleTypeDataBase.h
@@ -36,6 +36,9 @@ class UParticleModuleTypeDataBase : public UParticleModule
 	// @todo document
 	virtual FParticleEmitterInstance* CreateInstance(UParticleEmitter* InEmitterParent, UParticleSystemComponent* InComponent);
 
+	virtual const FVertexFactoryType* GetVertexFactoryType() const { return nullptr; }
+	virtual EPrimitiveType GetPrimitiveType() const { return PT_TriangleList; }
+
 	/** Cache any desired module pointers inside this type data */
 	virtual void CacheModuleInfo(UParticleEmitter* Emitter) {}
 
diff --git a/Engine/Source/Runtime/Engine/Classes/Particles/TypeData/ParticleModuleTypeDataBeam2.h b/Engine/Source/Runtime/Engine/Classes/Particles/TypeData/ParticleModuleTypeDataBeam2.h
--- a/Engine/Source/Runtime/Engine/Classes/Particles/TypeData/ParticleModuleTypeDataBeam2.h
+++ b/Engine/Source/Runtime/Engine/Classes/Particles/TypeData/ParticleModuleTypeDataBeam2.h
@@ -224,6 +224,8 @@ class UParticleModuleTypeDataBeam2 : public UParticleModuleTypeDataBase
 	//~ Begin UParticleModuleTypeDataBase Interface
 	virtual FParticleEmitterInstance* CreateInstance(UParticleEmitter* InEmitterParent, UParticleSystemComponent* InComponent) override;
 	virtual void CacheModuleInfo(UParticleEmitter* Emitter) override;
+	virtual const FVertexFactoryType* GetVertexFactoryType() const override;
+	virtual EPrimitiveType GetPrimitiveType() const { return PT_TriangleStrip; }
 	//~ End UParticleModuleTypeDataBase Interface
 
 
diff --git a/Engine/Source/Runtime/Engine/Classes/Particles/TypeData/ParticleModuleTypeDataGpu.h b/Engine/Source/Runtime/Engine/Classes/Particles/TypeData/ParticleModuleTypeDataGpu.h
--- a/Engine/Source/Runtime/Engine/Classes/Particles/TypeData/ParticleModuleTypeDataGpu.h
+++ b/Engine/Source/Runtime/Engine/Classes/Particles/TypeData/ParticleModuleTypeDataGpu.h
@@ -473,6 +473,7 @@ class UParticleModuleTypeDataGpu : public UParticleModuleTypeDataBase
 	virtual void Build( struct FParticleEmitterBuildInfo& EmitterBuildInfo ) override;
 	virtual bool RequiresBuild() const override { return true; }
 	virtual FParticleEmitterInstance* CreateInstance(UParticleEmitter* InEmitterParent, UParticleSystemComponent* InComponent) override;
+	virtual const FVertexFactoryType* GetVertexFactoryType() const override;
 	//~ End UParticleModuleTypeDataBase Interface
 };
 
diff --git a/Engine/Source/Runtime/Engine/Classes/Particles/TypeData/ParticleModuleTypeDataMesh.h b/Engine/Source/Runtime/Engine/Classes/Particles/TypeData/ParticleModuleTypeDataMesh.h
--- a/Engine/Source/Runtime/Engine/Classes/Particles/TypeData/ParticleModuleTypeDataMesh.h
+++ b/Engine/Source/Runtime/Engine/Classes/Particles/TypeData/ParticleModuleTypeDataMesh.h
@@ -224,7 +224,8 @@ class UParticleModuleTypeDataMesh : public UParticleModuleTypeDataBase
 	//~ End UParticleModule Interface
 
 	//~ Begin UParticleModuleTypeDataBase Interface
-	virtual FParticleEmitterInstance*	CreateInstance(UParticleEmitter* InEmitterParent, UParticleSystemComponent* InComponent) override;
+	virtual FParticleEmitterInstance* CreateInstance(UParticleEmitter* InEmitterParent, UParticleSystemComponent* InComponent) override;
+	virtual const FVertexFactoryType* GetVertexFactoryType() const override;
 	virtual bool	SupportsSpecificScreenAlignmentFlags() const override {	return true;	}	
 	virtual bool	SupportsSubUV() const override { return true; }
 	virtual bool	IsAMeshEmitter() const override { return true; }
diff --git a/Engine/Source/Runtime/Engine/Classes/Particles/TypeData/ParticleModuleTypeDataRibbon.h b/Engine/Source/Runtime/Engine/Classes/Particles/TypeData/ParticleModuleTypeDataRibbon.h
--- a/Engine/Source/Runtime/Engine/Classes/Particles/TypeData/ParticleModuleTypeDataRibbon.h
+++ b/Engine/Source/Runtime/Engine/Classes/Particles/TypeData/ParticleModuleTypeDataRibbon.h
@@ -178,6 +178,8 @@ class UParticleModuleTypeDataRibbon : public UParticleModuleTypeDataBase
 
 	//~ Begin UParticleModuleTypeDataBase Interface
 	virtual FParticleEmitterInstance* CreateInstance(UParticleEmitter* InEmitterParent, UParticleSystemComponent* InComponent) override;
+	virtual const FVertexFactoryType* GetVertexFactoryType() const override;
+	virtual EPrimitiveType GetPrimitiveType() const override { return PT_TriangleStrip; }
 	//~ End UParticleModuleTypeDataBase Interface
 };
 
diff --git a/Engine/Source/Runtime/Engine/Private/Particles/ParticleBeamModules.cpp b/Engine/Source/Runtime/Engine/Private/Particles/ParticleBeamModules.cpp
--- a/Engine/Source/Runtime/Engine/Private/Particles/ParticleBeamModules.cpp
+++ b/Engine/Source/Runtime/Engine/Private/Particles/ParticleBeamModules.cpp
@@ -22,6 +22,7 @@
 #include "Particles/ParticleLODLevel.h"
 #include "Distributions/DistributionFloatConstantCurve.h"
 #include "Engine/InterpCurveEdSetup.h"
+#include "ParticleBeamTrailVertexFactory.h"
 
 /*-----------------------------------------------------------------------------
 	Abstract base modules used for categorization.
@@ -695,6 +696,11 @@ FParticleEmitterInstance* UParticleModuleTypeDataBeam2::CreateInstance(UParticle
 	return Instance;
 }
 
+const FVertexFactoryType* UParticleModuleTypeDataBeam2::GetVertexFactoryType() const
+{
+	return &FParticleBeamTrailVertexFactory::StaticType;
+}
+
 void UParticleModuleTypeDataBeam2::CacheModuleInfo(UParticleEmitter* Emitter)
 {
 	int32 LODCount = Emitter->LODLevels.Num();
diff --git a/Engine/Source/Runtime/Engine/Private/Particles/ParticleComponents.cpp b/Engine/Source/Runtime/Engine/Private/Particles/ParticleComponents.cpp
--- a/Engine/Source/Runtime/Engine/Private/Particles/ParticleComponents.cpp
+++ b/Engine/Source/Runtime/Engine/Private/Particles/ParticleComponents.cpp
@@ -50,6 +50,7 @@
 #include "Interfaces/ITargetPlatform.h"
 #include "DeviceProfiles/DeviceProfile.h"
 #include "DeviceProfiles/DeviceProfileManager.h"
+#include "PipelineStateCache.h"
 #if WITH_EDITOR
 #include "Engine/InterpCurveEdSetup.h"
 #include "ObjectEditorUtils.h"
@@ -2504,6 +2505,8 @@ void UParticleSystem::PostLoad()
 		}
 	}
 
+	PrecachePSOs();
+
 #if WITH_EDITOR
 	// Due to there still being some ways that LODLevel counts get mismatched,
 	// when loading in the editor LOD levels will always be checked and fixed
@@ -2604,6 +2607,77 @@ void UParticleSystem::PostLoad()
 	SetupSoloing();
 }
 
+void UParticleSystem::PrecachePSOs()
+{
+	if (!FApp::CanEverRender() || !PipelineStateCache::IsPSOPrecachingEnabled())
+	{
+		return;
+	}
+
+	struct VFsPerMaterialData
+	{
+		UMaterialInterface* MaterialInterface;
+		EPrimitiveType PrimitiveType;
+		TArray<const FVertexFactoryType*, TInlineAllocator<2>> VertexFactoryTypes;
+	};
+	TArray<VFsPerMaterialData, TInlineAllocator<2>> VFsPerMaterials;
+
+	// No per component emitter materials known at this point in time
+	TArray<UMaterialInterface*> EmptyEmitterMaterials;
+	// Cached array to collect all materials used for LOD level
+	TArray<UMaterialInterface*> Materials;
+
+	for (int32 EmitterIdx = 0; EmitterIdx < Emitters.Num(); ++EmitterIdx)
+	{
+		const UParticleEmitter* Emitter = Emitters[EmitterIdx];
+		if (!Emitter)
+		{
+			continue;
+		}
+
+		for (int32 LodIndex = 0; LodIndex < Emitter->LODLevels.Num(); ++LodIndex)
+		{
+			const UParticleLODLevel* LOD = Emitter->LODLevels[LodIndex];
+			if (LOD && LOD->bEnabled)
+			{
+				const FVertexFactoryType* VFType = LOD->TypeDataModule ? LOD->TypeDataModule->GetVertexFactoryType() : &FParticleSpriteVertexFactory::StaticType;
+				check(VFType);
+				EPrimitiveType PrimitiveType = LOD->TypeDataModule ? LOD->TypeDataModule->GetPrimitiveType() : PT_TriangleList;
+
+				Materials.Empty();
+				LOD->GetUsedMaterials(Materials, NamedMaterialSlots, EmptyEmitterMaterials);
+
+				for (UMaterialInterface* MaterialInterface : Materials)
+				{
+					VFsPerMaterialData* VFsPerMaterial = VFsPerMaterials.FindByPredicate([MaterialInterface, PrimitiveType](const VFsPerMaterialData& Other)
+						{
+							return Other.MaterialInterface == MaterialInterface && Other.PrimitiveType == PrimitiveType;
+						});
+					if (VFsPerMaterial == nullptr)
+					{
+						VFsPerMaterial = &VFsPerMaterials.AddDefaulted_GetRef();
+						VFsPerMaterial->MaterialInterface = MaterialInterface;
+						VFsPerMaterial->PrimitiveType = PrimitiveType;
+					}
+					VFsPerMaterial->VertexFactoryTypes.AddUnique(VFType);
+				}
+			}
+		}
+	}
+
+	FPSOPrecacheParams PreCachePSOParams;
+	PreCachePSOParams.SetMobility(EComponentMobility::Movable);
+
+	for (VFsPerMaterialData& VFsPerMaterial : VFsPerMaterials)
+	{
+		if (VFsPerMaterial.MaterialInterface)
+		{
+			PreCachePSOParams.PrimitiveType = VFsPerMaterial.PrimitiveType;
+			VFsPerMaterial.MaterialInterface->PrecachePSOs(VFsPerMaterial.VertexFactoryTypes, PreCachePSOParams);
+		}
+	}
+}
+
 void UParticleSystem::Serialize(FArchive& Ar)
 {
 	Ar.UsingCustomVersion(FParticleSystemCustomVersion::GUID);
@@ -7643,7 +7717,7 @@ void UParticleLODLevel::GetUsedMaterials(TArray<UMaterialInterface*>& OutMateria
 	{
 		const UParticleModuleTypeDataMesh* MeshTypeData = Cast<UParticleModuleTypeDataMesh>(TypeDataModule);
 
-		if (MeshTypeData && MeshTypeData->Mesh)
+		if (MeshTypeData && MeshTypeData->Mesh && MeshTypeData->Mesh->GetRenderData())
 		{
 			const FStaticMeshLODResources& LODModel = MeshTypeData->Mesh->GetRenderData()->LODResources[0];
 
diff --git a/Engine/Source/Runtime/Engine/Private/Particles/ParticleModules.cpp b/Engine/Source/Runtime/Engine/Private/Particles/ParticleModules.cpp
--- a/Engine/Source/Runtime/Engine/Private/Particles/ParticleModules.cpp
+++ b/Engine/Source/Runtime/Engine/Private/Particles/ParticleModules.cpp
@@ -3400,6 +3400,11 @@ FParticleEmitterInstance* UParticleModuleTypeDataMesh::CreateInstance(UParticleE
 	return Instance;
 }
 
+const FVertexFactoryType* UParticleModuleTypeDataMesh::GetVertexFactoryType() const
+{
+	return &FMeshParticleVertexFactory::StaticType;
+}
+
 void UParticleModuleTypeDataMesh::SetToSensibleDefaults(UParticleEmitter* Owner)
 {
 	if ((Mesh == NULL) && GIsEditor )
@@ -4910,6 +4915,10 @@ FParticleEmitterInstance* UParticleModuleTypeDataGpu::CreateInstance(UParticleEm
 	return Instance;
 }
 
+const FVertexFactoryType* UParticleModuleTypeDataGpu::GetVertexFactoryType() const
+{
+	return &FGPUSpriteVertexFactory::StaticType;
+}
 
 /*-----------------------------------------------------------------------------
 	UParticleModulePivotOffset implementation.
diff --git a/Engine/Source/Runtime/Engine/Private/Particles/ParticleTrailModules.cpp b/Engine/Source/Runtime/Engine/Private/Particles/ParticleTrailModules.cpp
--- a/Engine/Source/Runtime/Engine/Private/Particles/ParticleTrailModules.cpp
+++ b/Engine/Source/Runtime/Engine/Private/Particles/ParticleTrailModules.cpp
@@ -13,6 +13,7 @@
 #include "Particles/Trail/ParticleModuleTrailSource.h"
 #include "Particles/TypeData/ParticleModuleTypeDataAnimTrail.h"
 #include "Particles/TypeData/ParticleModuleTypeDataRibbon.h"
+#include "ParticleBeamTrailVertexFactory.h"
 
 UParticleModuleTrailBase::UParticleModuleTrailBase(const FObjectInitializer& ObjectInitializer)
 	: Super(ObjectInitializer)
@@ -237,6 +238,11 @@ FParticleEmitterInstance* UParticleModuleTypeDataRibbon::CreateInstance(UParticl
 	return Instance;
 }
 
+const FVertexFactoryType* UParticleModuleTypeDataRibbon::GetVertexFactoryType() const
+{
+	return &FParticleBeamTrailVertexFactory::StaticType;
+}
+
 /*-----------------------------------------------------------------------------
 	UParticleModuleTypeDataAnimTrail implementation.
 -----------------------------------------------------------------------------*/
@@ -262,3 +268,8 @@ FParticleEmitterInstance* UParticleModuleTypeDataAnimTrail::CreateInstance(UPart
 	Instance->InitParameters(InEmitterParent, InComponent);
 	return Instance;
 }
+
+const FVertexFactoryType* UParticleModuleTypeDataAnimTrail::GetVertexFactoryType() const
+{
+	return &FParticleBeamTrailVertexFactory::StaticType;
+}
