Stash light function material when switching between non-static and static, and restore it when switching back.

diff --git a/Engine/Source/Runtime/Engine/Classes/Components/LightComponent.h b/Engine/Source/Runtime/Engine/Classes/Components/LightComponent.h
--- a/Engine/Source/Runtime/Engine/Classes/Components/LightComponent.h
+++ b/Engine/Source/Runtime/Engine/Classes/Components/LightComponent.h
@@ -158,6 +158,12 @@ class ENGINE_API ULightComponent : public ULightComponentBase
 	UPROPERTY(EditAnywhere, BlueprintReadOnly, Category=LightFunction)
 	TObjectPtr<class UMaterialInterface> LightFunctionMaterial;
 
+#if WITH_EDITORONLY_DATA
+	/** When clearing the light func, e.g. because the light is made static, this field remembers the last value */
+	UPROPERTY(Transient)
+	TObjectPtr<class UMaterialInterface> StashedLightFunctionMaterial;
+#endif
+
 	/** Scales the light function projection.  X and Y scale in the directions perpendicular to the light's direction, Z scales along the light direction. */
 	UPROPERTY(EditAnywhere, BlueprintReadOnly, Category=LightFunction, meta=(AllowPreserveRatio = "true"))
 	FVector LightFunctionScale;
@@ -432,6 +438,9 @@ public:
 	/** Set the MaterialInterface to use for the given element index (if valid) */
 	virtual void SetMaterial(int32 ElementIndex, UMaterialInterface* InMaterial);
 
+	/** Set the light func to null, but remember the current value so it can be restored later */
+	void ClearLightFunctionMaterial();
+
 	virtual class FLightSceneProxy* CreateSceneProxy() const
 	{
 		return NULL;
diff --git a/Engine/Source/Runtime/Engine/Private/Components/LightComponent.cpp b/Engine/Source/Runtime/Engine/Private/Components/LightComponent.cpp
--- a/Engine/Source/Runtime/Engine/Private/Components/LightComponent.cpp
+++ b/Engine/Source/Runtime/Engine/Private/Components/LightComponent.cpp
@@ -655,7 +655,7 @@ void ULightComponent::PostLoad()
 	if (LightFunctionMaterial && HasStaticLighting())
 	{
 		// Light functions can only be used on dynamic lights
-		LightFunctionMaterial = NULL;
+		ClearLightFunctionMaterial();
 	}
 
 	PreviewShadowMapChannel = INDEX_NONE;
@@ -773,8 +773,19 @@ void ULightComponent::PostEditChangeProperty(FPropertyChangedEvent& PropertyChan
 	if (HasStaticLighting())
 	{
 		// Lightmapped lights must not have light functions
-		LightFunctionMaterial = NULL;
+		ClearLightFunctionMaterial();
+	}
+#if WITH_EDITOR
+	else if (PropertyName == GET_MEMBER_NAME_STRING_CHECKED(ULightComponent, LightFunctionMaterial))
+	{
+		StashedLightFunctionMaterial = nullptr;
 	}
+	else if (StashedLightFunctionMaterial != nullptr)
+	{
+		// Light has been made non-static, restore previous light function
+		LightFunctionMaterial = StashedLightFunctionMaterial;
+	}
+#endif
 
 	// Unbuild lighting because a property changed
 	// Exclude properties that don't affect built lighting
@@ -1079,11 +1090,22 @@ void ULightComponent::SetLightFunctionMaterial(UMaterialInterface* NewLightFunct
 	if (AreDynamicDataChangesAllowed()
 		&& NewLightFunctionMaterial != LightFunctionMaterial)
 	{
+#if WITH_EDITOR
+		StashedLightFunctionMaterial = nullptr;
+#endif
 		LightFunctionMaterial = NewLightFunctionMaterial;
 		MarkRenderStateDirty();
 	}
 }
 
+void ULightComponent::ClearLightFunctionMaterial()
+{
+#if WITH_EDITOR
+	StashedLightFunctionMaterial = LightFunctionMaterial;
+#endif
+	LightFunctionMaterial = nullptr;
+}
+
 void ULightComponent::SetLightFunctionScale(FVector NewLightFunctionScale)
 {
 	if (AreDynamicDataChangesAllowed()
@@ -1486,6 +1508,9 @@ void ULightComponent::SetMaterial(int32 ElementIndex, UMaterialInterface* InMate
 {
 	if (ElementIndex == 0)
 	{
+#if WITH_EDITOR
+		StashedLightFunctionMaterial = nullptr;
+#endif
 		LightFunctionMaterial = InMaterial;
 		MarkRenderStateDirty();
 	}
diff --git a/Engine/Source/Runtime/Engine/Private/Light.cpp b/Engine/Source/Runtime/Engine/Private/Light.cpp
--- a/Engine/Source/Runtime/Engine/Private/Light.cpp
+++ b/Engine/Source/Runtime/Engine/Private/Light.cpp
@@ -185,7 +185,7 @@ void APointLight::PostLoad()
 
 	if (GetLightComponent()->Mobility == EComponentMobility::Static)
 	{
-		GetLightComponent()->LightFunctionMaterial = NULL;
+		GetLightComponent()->ClearLightFunctionMaterial();
 	}
 }
 
@@ -271,7 +271,7 @@ void ADirectionalLight::PostLoad()
 
 	if (GetLightComponent()->Mobility == EComponentMobility::Static)
 	{
-		GetLightComponent()->LightFunctionMaterial = NULL;
+		GetLightComponent()->ClearLightFunctionMaterial();
 	}
 #if WITH_EDITORONLY_DATA
 	if(ArrowComponent != nullptr)
diff --git a/Engine/Source/Runtime/Engine/Private/RectLight.cpp b/Engine/Source/Runtime/Engine/Private/RectLight.cpp
--- a/Engine/Source/Runtime/Engine/Private/RectLight.cpp
+++ b/Engine/Source/Runtime/Engine/Private/RectLight.cpp
@@ -18,7 +18,7 @@ void ARectLight::PostLoad()
 
 	if (GetLightComponent()->Mobility == EComponentMobility::Static)
 	{
-		GetLightComponent()->LightFunctionMaterial = NULL;
+		GetLightComponent()->ClearLightFunctionMaterial();
 	}
 }
 
diff --git a/Engine/Source/Runtime/Engine/Private/SpotLight.cpp b/Engine/Source/Runtime/Engine/Private/SpotLight.cpp
--- a/Engine/Source/Runtime/Engine/Private/SpotLight.cpp
+++ b/Engine/Source/Runtime/Engine/Private/SpotLight.cpp
@@ -47,7 +47,7 @@ void ASpotLight::PostLoad()
 
 	if (GetLightComponent()->Mobility == EComponentMobility::Static)
 	{
-		GetLightComponent()->LightFunctionMaterial = NULL;
+		GetLightComponent()->ClearLightFunctionMaterial();
 	}
 
 #if WITH_EDITORONLY_DATA
