Added ContentBounds to world partition runtime cells : requested by licensees to keep track of what bounds of the world is affected whern a streaming cell unloads its streaming level. This also replaces the cell's MinMaxZ.

diff --git a/Engine/Source/Runtime/Engine/Private/WorldPartition/WorldPartitionRuntimeCell.cpp b/Engine/Source/Runtime/Engine/Private/WorldPartition/WorldPartitionRuntimeCell.cpp
--- a/Engine/Source/Runtime/Engine/Private/WorldPartition/WorldPartitionRuntimeCell.cpp
+++ b/Engine/Source/Runtime/Engine/Private/WorldPartition/WorldPartitionRuntimeCell.cpp
@@ -14,7 +14,7 @@ int32 UWorldPartitionRuntimeCell::StreamingSourceCacheEpoch = 0;
 UWorldPartitionRuntimeCell::UWorldPartitionRuntimeCell(const FObjectInitializer& ObjectInitializer)
 	: Super(ObjectInitializer)
 	, bIsAlwaysLoaded(false)
-	, MinMaxZ(-HALF_WORLD_MAX, HALF_WORLD_MAX)
+	, ContentBounds(ForceInit)
 	, Priority(0)
 	, CachedMinSourcePriority((uint8)EStreamingSourcePriority::Lowest)
 	, CachedSourceInfoEpoch(INT_MIN)
@@ -43,11 +43,6 @@ void UWorldPartitionRuntimeCell::SetDataLayers(const TArray<const UDataLayerInst
 	UpdateDebugName();
 }
 
-void UWorldPartitionRuntimeCell::SetMinMaxZ(const FVector2D& InMinMaxZ)
-{
-	MinMaxZ = InMinMaxZ;
-}
-
 void UWorldPartitionRuntimeCell::SetDebugInfo(int64 InCoordX, int64 InCoordY, int64 InCoordZ, FName InGridName)
 {
 	DebugInfo.CoordX = InCoordX;
@@ -172,4 +167,4 @@ bool UWorldPartitionRuntimeCell::ContainsDataLayer(const UDataLayerAsset* DataLa
 bool UWorldPartitionRuntimeCell::ContainsDataLayer(const UDataLayerInstance* DataLayerInstance) const
 {
 	return GetDataLayers().Contains(DataLayerInstance->GetDataLayerFName());
-}
+}
\ No newline at end of file
diff --git a/Engine/Source/Runtime/Engine/Private/WorldPartition/WorldPartitionRuntimeSpatialHash.cpp b/Engine/Source/Runtime/Engine/Private/WorldPartition/WorldPartitionRuntimeSpatialHash.cpp
--- a/Engine/Source/Runtime/Engine/Private/WorldPartition/WorldPartitionRuntimeSpatialHash.cpp
+++ b/Engine/Source/Runtime/Engine/Private/WorldPartition/WorldPartitionRuntimeSpatialHash.cpp
@@ -1125,14 +1125,13 @@ bool UWorldPartitionRuntimeSpatialHash::CreateStreamingGrid(const FSpatialHashRu
 					}
 				}
 
-				FVector2D CellMinMaxZ(UE_BIG_NUMBER, -UE_BIG_NUMBER);
+				FBox CellContentBounds(ForceInit);
 				for (const IStreamingGenerationContext::FActorInstance& ActorInstance : FilteredActors)
 				{
 					const FWorldPartitionActorDescView& ActorDescView = ActorInstance.GetActorDescView();
 					StreamingCell->AddActorToCell(ActorDescView, ActorInstance.GetContainerID(), ActorInstance.GetTransform(), ActorInstance.GetActorDescContainer());
 					
-					CellMinMaxZ.X = FMath::Min(CellMinMaxZ.X, ActorDescView.GetBounds().Min.Z);
-					CellMinMaxZ.Y = FMath::Max(CellMinMaxZ.Y, ActorDescView.GetBounds().Max.Z);
+					CellContentBounds += ActorDescView.GetBounds();
 
 					if (ActorInstance.GetContainerID().IsMainContainer() && StreamingCell->UnsavedActorsContainer)
 					{
@@ -1151,7 +1150,7 @@ bool UWorldPartitionRuntimeSpatialHash::CreateStreamingGrid(const FSpatialHashRu
 						}
 					}
 				}
-				StreamingCell->SetMinMaxZ(CellMinMaxZ);
+				StreamingCell->SetContentBounds(CellContentBounds);
 
 				// Always loaded cell actors are transfered to World's Persistent Level (see UWorldPartitionRuntimeSpatialHash::PopulateGeneratorPackageForCook)
 				if (StreamingCell->GetActorCount() && !StreamingCell->IsAlwaysLoaded())
diff --git a/Engine/Source/Runtime/Engine/Public/WorldPartition/WorldPartitionRuntimeCell.h b/Engine/Source/Runtime/Engine/Public/WorldPartition/WorldPartitionRuntimeCell.h
--- a/Engine/Source/Runtime/Engine/Public/WorldPartition/WorldPartitionRuntimeCell.h
+++ b/Engine/Source/Runtime/Engine/Public/WorldPartition/WorldPartitionRuntimeCell.h
@@ -191,9 +191,10 @@ class ENGINE_API UWorldPartitionRuntimeCell : public UObject, public IWorldParti
 	{
 		return Algo::AnyOf(DataLayers, [&InDataLayers](const FName& DataLayer) { return InDataLayers.Contains(DataLayer); });
 	}
+	virtual const FBox& GetContentBounds() const override { return ContentBounds; }
 	// IWorldPartitionCell Interface end
 
-	const FVector2D& GetMinMaxZ() const { return MinMaxZ; }
+	const FVector2D GetMinMaxZ() const { return FVector2D(ContentBounds.Min.Z, ContentBounds.Max.Z); }
 
 	bool GetBlockOnSlowLoading() const { return bBlockOnSlowLoading; }
 #if WITH_EDITOR
@@ -205,7 +206,7 @@ class ENGINE_API UWorldPartitionRuntimeCell : public UObject, public IWorldParti
 
 	void SetDataLayers(const TArray<const UDataLayerInstance*>& InDataLayerInstances);
 	void SetContentBundleUID(const FGuid& InContentBundleID) { ContentBundleID = InContentBundleID; }
-	void SetMinMaxZ(const FVector2D& InMinMaxZ);
+	void SetContentBounds(const FBox& InBounds) { ContentBounds = InBounds; }
 	void SetDebugInfo(int64 InCoordX, int64 InCoordY, int64 InCoordZ, FName InGridName);
 	void SetGridName(FName InGridName);
 	virtual void AddActorToCell(const FWorldPartitionActorDescView& ActorDescView, const FActorContainerID& InContainerID, const FTransform& InContainerTransform, const UActorDescContainer* InContainer) PURE_VIRTUAL(UWorldPartitionRuntimeCell::AddActorToCell,);
@@ -243,7 +244,7 @@ private:
 	TArray<FName> DataLayers;
 
 	UPROPERTY()
-	FVector2D MinMaxZ;
+	FBox ContentBounds;
 
 	// Debug Info
 	UPROPERTY()
diff --git a/Engine/Source/Runtime/Engine/Public/WorldPartition/WorldPartitionRuntimeCellInterface.h b/Engine/Source/Runtime/Engine/Public/WorldPartition/WorldPartitionRuntimeCellInterface.h
--- a/Engine/Source/Runtime/Engine/Public/WorldPartition/WorldPartitionRuntimeCellInterface.h
+++ b/Engine/Source/Runtime/Engine/Public/WorldPartition/WorldPartitionRuntimeCellInterface.h
@@ -15,5 +15,6 @@ public:
 	virtual bool HasDataLayers() const = 0;
 	virtual const TArray<FName>& GetDataLayers() const = 0;
 	virtual bool HasAnyDataLayer(const TSet<FName>& InDataLayers) const = 0;
+	virtual const FBox& GetContentBounds() const = 0;
 };
 
