Crash in UE::Anim::Attributes::InterpolateAttributes when enabling URO on B_Hero_ShooterMannequin in LyraStarterGame
#fix fixed mistake of using incorrect index value
#misc added and improved automated tests

diff --git a/Engine/Source/Runtime/Engine/Private/Animation/AttributesRuntime.cpp b/Engine/Source/Runtime/Engine/Private/Animation/AttributesRuntime.cpp
--- a/Engine/Source/Runtime/Engine/Private/Animation/AttributesRuntime.cpp
+++ b/Engine/Source/Runtime/Engine/Private/Animation/AttributesRuntime.cpp
@@ -242,7 +242,7 @@ void Attributes::InterpolateAttributes(FMeshAttributeContainer& FromAttributes,
 			const IAttributeBlendOperator* Operator = AttributeTypes::GetTypeOperator(WeakScriptStruct);
 			for (int32 EntryIndex = 0; EntryIndex < Identifiers.Num(); ++EntryIndex)
 			{
-				uint8* OutAttributeData = FromAttributes.Find(WeakScriptStruct.Get(), Identifiers[TypeIndex]);
+				uint8* OutAttributeData = FromAttributes.Find(WeakScriptStruct.Get(), Identifiers[EntryIndex]);
 				ensure(OutAttributeData);
 
 				// Interpolate, from default value to stored value
@@ -260,10 +260,10 @@ void Attributes::InterpolateAttributes(FMeshAttributeContainer& FromAttributes,
 			const IAttributeBlendOperator* Operator = AttributeTypes::GetTypeOperator(WeakScriptStruct);
 			for (int32 EntryIndex = 0; EntryIndex < Identifiers.Num(); ++EntryIndex)
 			{
-				const uint8* AttributeData = ToAttributes.Find(WeakScriptStruct.Get(), Identifiers[TypeIndex]);
+				const uint8* AttributeData = ToAttributes.Find(WeakScriptStruct.Get(), Identifiers[EntryIndex]);
 				ensure(AttributeData);
 
-				uint8* OutAttributeData = FromAttributes.Add(WeakScriptStruct.Get(), Identifiers[TypeIndex]);
+				uint8* OutAttributeData = FromAttributes.Add(WeakScriptStruct.Get(), Identifiers[EntryIndex]);
 				ensure(OutAttributeData);
 
 				Operator->Interpolate(OutAttributeData, AttributeData, Alpha, OutAttributeData);
@@ -325,7 +325,7 @@ void Attributes::InterpolateAttributes(FMeshAttributeContainer& FromAttributes,
 				uint8* FromAttributeData = FromAttributes.Add(WeakScriptStruct.Get(), ToIdentifiers[ToValueIndex]);
 				ensure(FromAttributeData);
 
-				Operator->Interpolate(FromAttributeData, ToAttributeData, Alpha, FromAttributeData);
+				Operator->Interpolate(DefaultValue.GetPtr<uint8>(), ToAttributeData, Alpha, FromAttributeData);
 			}
 		}
 	}
