Don't call RecreateAndLinkNode() when calling SetParameterValue() on a channel mask material expression.
The pins for that expression type are fixed so I don't think there is a use for RecreateAndLinkNode() here.
Normally it is used for things like material functions where the pins change.
The problem is that this function is called on copy/paste as part of the post paste ValidateParameterName(). At that point we don't have enough information to recreate the node links after destroying them. So we end up with broken links after the paste.

diff --git a/Engine/Source/Runtime/Engine/Private/Materials/MaterialExpressions.cpp b/Engine/Source/Runtime/Engine/Private/Materials/MaterialExpressions.cpp
--- a/Engine/Source/Runtime/Engine/Private/Materials/MaterialExpressions.cpp
+++ b/Engine/Source/Runtime/Engine/Private/Materials/MaterialExpressions.cpp
@@ -8777,12 +8777,6 @@ bool UMaterialExpressionChannelMaskParameter::SetParameterValue(FName InParamete
 			SendPostEditChangeProperty(this, TEXT("MaskChannel"));
 		}
 
-		// Update caption and preview
-		if (GraphNode)
-		{
-			CastChecked<UMaterialGraphNode>(GraphNode)->RecreateAndLinkNode();
-		}
-
 		return true;
 	}
 
