LevelStreaming: Fix BlockTillLevelStreamingCompleted when level package is redirected

diff --git a/Engine/Source/Runtime/Engine/Private/LevelStreaming.cpp b/Engine/Source/Runtime/Engine/Private/LevelStreaming.cpp
--- a/Engine/Source/Runtime/Engine/Private/LevelStreaming.cpp
+++ b/Engine/Source/Runtime/Engine/Private/LevelStreaming.cpp
@@ -842,7 +842,19 @@ void ULevelStreaming::UpdateStreamingState(bool& bOutUpdateAgain, bool& bOutRede
 				// This rare case can happen if desired level (typically LODPackage) changed between last RequestLevel call and AsyncLevelLoadComplete completion callback.
 				DiscardPendingUnloadLevel(World);
 			}
+
 			UpdateStreamingState_RequestLevel();
+			
+			// This is to fix the Blocking load on a redirected world package
+			// When loading a redirected streaming level the state will go from: Unloaded -> LoadedNotVisible
+			// This state change will generate a new RequestLevel call which will load the redirected package.
+			// In blocking load, loading will be done after the UpdateStreamingState_RequestLevel call and leave us in the LoadedNotVisible (with a loadedlevel) state which would prevent the bOutUpdateAgain from being set to true.
+			// So this condition here makes sure that we aren't loading (async) and that we should be visible (LoadedNotVisible isn't our final target).
+			// If that is the case we request another update.
+			if (CurrentState != ECurrentState::Loading)
+			{
+				bOutUpdateAgain |= ShouldBeVisible();
+			}
 			break;
 
 		default:
