Add units (milliseconds) to description of ExactPing variable (Contributed by [REDACTED_PERSON])

#rnx
#9543

diff --git a/Engine/Source/Runtime/Engine/Classes/GameFramework/PlayerState.h b/Engine/Source/Runtime/Engine/Classes/GameFramework/PlayerState.h
--- a/Engine/Source/Runtime/Engine/Classes/GameFramework/PlayerState.h
+++ b/Engine/Source/Runtime/Engine/Classes/GameFramework/PlayerState.h
@@ -143,7 +143,7 @@ public:
 	UPROPERTY()
 	TSubclassOf<class ULocalMessage> EngineMessageClass;
 
-	/** Exact ping as float (rounded and compressed in replicated CompressedPing) */
+	/** Exact ping in milliseconds as float (rounded and compressed in replicated CompressedPing) */
 	float ExactPing;
 
 	UE_DEPRECATED(4.27, "Please use ExactPing instead.")
