Fix for push model fast arrays caching the wrong object

diff --git a/Engine/Source/Runtime/CoreUObject/Public/UObject/CoreNet.h b/Engine/Source/Runtime/CoreUObject/Public/UObject/CoreNet.h
--- a/Engine/Source/Runtime/CoreUObject/Public/UObject/CoreNet.h
+++ b/Engine/Source/Runtime/CoreUObject/Public/UObject/CoreNet.h
@@ -660,9 +660,12 @@ struct FNetDeltaSerializeInfo
 	 */
 	bool bInternalAck = false;
 
-	/** The object that owns the struct we're serializing. */
+	/** The object that owns the struct we're serializing, may be an archetype. */
 	UObject* Object = nullptr;
 
+	/** Used by SendCustomDeltaProperty to distinguish between the source object (archetype) and the replicating object. */
+	UObject* CustomDeltaObject = nullptr;
+
 	/**
 	 * When non-null, this indicates that we're gathering Guid References.
 	 * Any Guids the struct is referencing should be added.
diff --git a/Engine/Source/Runtime/Engine/Private/DataReplication.cpp b/Engine/Source/Runtime/Engine/Private/DataReplication.cpp
--- a/Engine/Source/Runtime/Engine/Private/DataReplication.cpp
+++ b/Engine/Source/Runtime/Engine/Private/DataReplication.cpp
@@ -360,6 +360,7 @@ bool FObjectReplicator::SendCustomDeltaProperty(UObject* InObject, uint16 Custom
 
 	FNetDeltaSerializeInfo Parms;
 	Parms.Object = InObject;
+	Parms.CustomDeltaObject = GetObject();
 	Parms.Writer = &OutBunch;
 	Parms.Map = Connection->PackageMap;
 	Parms.OldState = OldState.Get();
diff --git a/Engine/Source/Runtime/Engine/Private/RepLayout.cpp b/Engine/Source/Runtime/Engine/Private/RepLayout.cpp
--- a/Engine/Source/Runtime/Engine/Private/RepLayout.cpp
+++ b/Engine/Source/Runtime/Engine/Private/RepLayout.cpp
@@ -4505,9 +4505,10 @@ bool FRepLayout::SendCustomDeltaProperty(FNetDeltaSerializeInfo& Params, uint16
 	}
 
 #if WITH_PUSH_MODEL
-	if (EnumHasAnyFlags(Parent.Flags, ERepParentFlags::IsFastArray) && IsPushModelProperty(CustomDeltaProperty.PropertyRepIndex))
+	if (EnumHasAnyFlags(Parent.Flags, ERepParentFlags::IsFastArray) && IsPushModelProperty(CustomDeltaProperty.PropertyRepIndex) && Params.CustomDeltaObject)
 	{
-		static_cast<FFastArraySerializer*>(Params.Data)->CachePushModelState(Params.Object, CustomDeltaProperty.PropertyRepIndex);
+		void* ArrayData = FRepObjectDataBuffer(Params.CustomDeltaObject) + Parent;
+		static_cast<FFastArraySerializer*>(ArrayData)->CachePushModelState(Params.CustomDeltaObject, CustomDeltaProperty.PropertyRepIndex);
 	}
 #endif
 
