Adding csv stats for open/ticking channel counts on servers

diff --git a/Engine/Source/Runtime/Engine/Classes/Engine/NetConnection.h b/Engine/Source/Runtime/Engine/Classes/Engine/NetConnection.h
--- a/Engine/Source/Runtime/Engine/Classes/Engine/NetConnection.h
+++ b/Engine/Source/Runtime/Engine/Classes/Engine/NetConnection.h
@@ -1221,6 +1221,8 @@ public:
 	/** Removes a channel from the ticking list directly */
 	void StopTickingChannel(UChannel* Channel) { ChannelsToTick.Remove(Channel); }
 
+	int32 GetNumTickingChannels() const { return ChannelsToTick.Num(); }
+
 	FORCEINLINE FHistogram GetNetHistogram() const { return NetConnectionHistogram; }
 
 	/** Whether or not a client packet has been received - used serverside, to delay any packet sends */
diff --git a/Engine/Source/Runtime/Engine/Private/NetDriver.cpp b/Engine/Source/Runtime/Engine/Private/NetDriver.cpp
--- a/Engine/Source/Runtime/Engine/Private/NetDriver.cpp
+++ b/Engine/Source/Runtime/Engine/Private/NetDriver.cpp
@@ -4955,6 +4955,18 @@ struct FScopedNetDriverStats
 			CSV_CUSTOM_STAT(Replication, SharedSerializationPropertyHit, GNumSharedSerializationHit, ECsvCustomStatOp::Set);
 			CSV_CUSTOM_STAT(Replication, SharedSerializationPropertyMiss, GNumSharedSerializationMiss, ECsvCustomStatOp::Set);
 
+			int32 NumOpenChannels = 0;
+			int32 NumTickingChannels = 0;
+
+			for (const UNetConnection* ClientConnection : NetDriver->ClientConnections)
+			{
+				NumOpenChannels += ClientConnection->OpenChannels.Num();
+				NumTickingChannels += ClientConnection->GetNumTickingChannels();
+			}
+
+			CSV_CUSTOM_STAT(Replication, NumOpenChannels, (float)NumOpenChannels, ECsvCustomStatOp::Set);
+			CSV_CUSTOM_STAT(Replication, NumTickingChannels, (float)NumTickingChannels, ECsvCustomStatOp::Set);
+
 			// Note: we want to reset this at the end of the frame since the RPC stats are incremented at the top (recv)
 			GNumSharedSerializationHit = 0;
 			GNumSharedSerializationMiss = 0;
