World Partition: don't log when trying to convert deprecated short class names and changed a warning to a log when encountering invalid actor classes (can happen for editor-only actors or actors coming from editor-only modules/plugins when running -game).

#rnx

diff --git a/Engine/Source/Runtime/Engine/Private/WorldPartition/WorldPartitionActorDesc.cpp b/Engine/Source/Runtime/Engine/Private/WorldPartition/WorldPartitionActorDesc.cpp
--- a/Engine/Source/Runtime/Engine/Private/WorldPartition/WorldPartitionActorDesc.cpp
+++ b/Engine/Source/Runtime/Engine/Private/WorldPartition/WorldPartitionActorDesc.cpp
@@ -328,7 +328,7 @@ void FWorldPartitionActorDesc::Serialize(FArchive& Ar)
 			FName BaseClassPathName;
 			Ar << BaseClassPathName;
 PRAGMA_DISABLE_DEPRECATION_WARNINGS
-			BaseClass = FAssetData::TryConvertShortClassNameToPathName(BaseClassPathName);
+			BaseClass = FAssetData::TryConvertShortClassNameToPathName(BaseClassPathName, ELogVerbosity::NoLogging);
 PRAGMA_ENABLE_DEPRECATION_WARNINGS
 		}
 		else
@@ -342,7 +342,7 @@ PRAGMA_ENABLE_DEPRECATION_WARNINGS
 		FName NativeClassPathName;
 		Ar << NativeClassPathName;
 PRAGMA_DISABLE_DEPRECATION_WARNINGS
-		NativeClass = FAssetData::TryConvertShortClassNameToPathName(NativeClassPathName);
+		NativeClass = FAssetData::TryConvertShortClassNameToPathName(NativeClassPathName, ELogVerbosity::NoLogging);
 PRAGMA_ENABLE_DEPRECATION_WARNINGS
 	}
 	else
diff --git a/Engine/Source/Runtime/Engine/Private/WorldPartition/WorldPartitionActorDescUtils.cpp b/Engine/Source/Runtime/Engine/Private/WorldPartition/WorldPartitionActorDescUtils.cpp
--- a/Engine/Source/Runtime/Engine/Private/WorldPartition/WorldPartitionActorDescUtils.cpp
+++ b/Engine/Source/Runtime/Engine/Private/WorldPartition/WorldPartitionActorDescUtils.cpp
@@ -57,7 +57,7 @@ TUniquePtr<FWorldPartitionActorDesc> FWorldPartitionActorDescUtils::GetActorDesc
 			
 		if (!ActorDescInitData.NativeClass)
 		{
-			UE_LOG(LogWorldPartition, Warning, TEXT("Invalid class for actor guid `%s` ('%s') from package '%s'"), *NewActorDesc->GetGuid().ToString(), *NewActorDesc->GetActorName().ToString(), *NewActorDesc->GetActorPackage().ToString());
+			UE_LOG(LogWorldPartition, Log, TEXT("Invalid class for actor guid `%s` ('%s') from package '%s'"), *NewActorDesc->GetGuid().ToString(), *NewActorDesc->GetActorName().ToString(), *NewActorDesc->GetActorPackage().ToString());
 			return nullptr;
 		}
 		else if (UClass* Class = FindObject<UClass>(InAssetData.AssetClassPath); !Class)
@@ -71,7 +71,7 @@ TUniquePtr<FWorldPartitionActorDesc> FWorldPartitionActorDescUtils::GetActorDesc
 
 				if (!BlueprintClass.Num())
 				{
-					UE_LOG(LogWorldPartition, Warning, TEXT("Failed to find class '%s' for actor '%s"), *InAssetData.AssetClassPath.ToString(), *NewActorDesc->GetActorSoftPath().ToString());
+					UE_LOG(LogWorldPartition, Log, TEXT("Failed to find class '%s' for actor '%s"), *InAssetData.AssetClassPath.ToString(), *NewActorDesc->GetActorSoftPath().ToString());
 					return nullptr;
 				}
 			}
