Chaos : fix event hit not reporting other comp or actor when it is a geometry collection

diff --git a/Engine/Source/Runtime/Engine/Private/PhysicsEngine/Experimental/PhysScene_Chaos.cpp b/Engine/Source/Runtime/Engine/Private/PhysicsEngine/Experimental/PhysScene_Chaos.cpp
--- a/Engine/Source/Runtime/Engine/Private/PhysicsEngine/Experimental/PhysScene_Chaos.cpp
+++ b/Engine/Source/Runtime/Engine/Private/PhysicsEngine/Experimental/PhysScene_Chaos.cpp
@@ -785,6 +785,18 @@ void FPhysScene_Chaos::HandleCollisionEvents(const Chaos::FCollisionEventData& E
 								NotifyInfo.Info0.SetFrom(GetBodyInstanceFromProxyAndShape(PhysicsProxy0, CollisionDataItem.ShapeIndex1), DeltaVelocity1);
 								NotifyInfo.Info1.SetFrom(GetBodyInstanceFromProxyAndShape(PhysicsProxy1, CollisionDataItem.ShapeIndex2), DeltaVelocity2);
 
+								// in some case ( like with geometry collections ) we don't have a body instance so the component part will null, we need to handle that 
+								if (NotifyInfo.Info0.Component == nullptr)
+								{
+									NotifyInfo.Info0.Component = GetOwningComponent<UPrimitiveComponent>(PhysicsProxy0);
+									NotifyInfo.Info0.Actor = (NotifyInfo.Info0.Component != nullptr) ? NotifyInfo.Info0.Component->GetOwner() : nullptr;
+								}
+								if (NotifyInfo.Info1.Component == nullptr)
+								{
+									NotifyInfo.Info1.Component = GetOwningComponent<UPrimitiveComponent>(PhysicsProxy1);
+									NotifyInfo.Info1.Actor = (NotifyInfo.Info1.Component != nullptr) ? NotifyInfo.Info1.Component->GetOwner() : nullptr;
+								}
+
 								FRigidBodyContactInfo& NewContact = NotifyInfo.RigidCollisionData.ContactInfos.AddZeroed_GetRef();
 								NewContact.ContactNormal = CollisionDataItem.Normal;
 								NewContact.ContactPosition = CollisionDataItem.Location;
diff --git a/Engine/Source/Runtime/Experimental/ChaosSolverEngine/Private/Chaos/ChaosGameplayEventDispatcher.cpp b/Engine/Source/Runtime/Experimental/ChaosSolverEngine/Private/Chaos/ChaosGameplayEventDispatcher.cpp
--- a/Engine/Source/Runtime/Experimental/ChaosSolverEngine/Private/Chaos/ChaosGameplayEventDispatcher.cpp
+++ b/Engine/Source/Runtime/Experimental/ChaosSolverEngine/Private/Chaos/ChaosGameplayEventDispatcher.cpp
@@ -327,6 +327,7 @@ void UChaosGameplayEventDispatcher::UnregisterChaosEvents()
 
 void UChaosGameplayEventDispatcher::HandleCollisionEvents(const Chaos::FCollisionEventData& Event)
 {
+	// todo(chaos) : this code is very similar to FPhysScene_Chaos::HandleCollisionEvents, we should propably consolidate if possible or share as much code as possible 
 	SCOPE_CYCLE_COUNTER(STAT_DispatchCollisionEvents);
 
 	FPhysScene_Chaos& Scene = *(GetWorld()->GetPhysicsScene());
@@ -391,6 +392,7 @@ void UChaosGameplayEventDispatcher::HandleCollisionEvents(const Chaos::FCollisio
 										NotifyInfo.bCallEvent0 = true;
 										// if Comp1 wants this event too, it will get its own pending collision entry, so we leave it false
 
+										// @todo(chaos) this may not handle welded objects properly as the component returned may be thewrong one ( see FPhysScene_Chaos::HandleCollisionEvents ) 
 										SetCollisionInfoFromComp(NotifyInfo.Info0, Comp0);
 										SetCollisionInfoFromComp(NotifyInfo.Info1, Comp1);
 
