CIS: Added some needed null checks to ActorConstruction.cpp

#rnx

diff --git a/Engine/Source/Runtime/Engine/Private/ActorConstruction.cpp b/Engine/Source/Runtime/Engine/Private/ActorConstruction.cpp
--- a/Engine/Source/Runtime/Engine/Private/ActorConstruction.cpp
+++ b/Engine/Source/Runtime/Engine/Private/ActorConstruction.cpp
@@ -1129,10 +1129,21 @@ UActorComponent* AActor::CreateComponentFromTemplateData(const FBlueprintCookedC
 
 UActorComponent* AActor::AddComponent(FName TemplateName, bool bManualAttachment, const FTransform& RelativeTransform, const UObject* ComponentTemplateContext, bool bDeferredFinish)
 {
-	UWorld* World = GetWorld();
-	if (World->bIsTearingDown)
+	if (const UWorld* World = GetWorld())
+	{
+		if (World->bIsTearingDown)
+		{
+			UE_LOG(LogActor, Warning, TEXT("AddComponent failed for actor: [%s] with param TemplateName: [%s] because we are in the process of tearing down the world")
+				, *GetName()
+				, *TemplateName.ToString());
+			return nullptr;
+		}
+	}
+	else
 	{
-		UE_LOG(LogActor, Warning, TEXT("AddComponent failed because we are in the process of tearing down the world"));
+		UE_LOG(LogActor, Warning, TEXT("AddComponent failed for actor: [%s] with param TemplateName: [%s] because world == nullptr")
+			, *GetName()
+			, *TemplateName.ToString());
 		return nullptr;
 	}
 
@@ -1175,10 +1186,21 @@ UActorComponent* AActor::AddComponentByClass(TSubclassOf<UActorComponent> Class,
 		return nullptr;
 	}
 
-	UWorld* World = GetWorld();
-	if (World->bIsTearingDown)
+	if (const UWorld* World = GetWorld())
+	{
+		if (World->bIsTearingDown)
+		{
+			UE_LOG(LogActor, Warning, TEXT("AddComponentByClass failed for actor: [%s] with param Class: [%s] because we are in the process of tearing down the world")
+				, *GetName()
+				, *GetNameSafe(Class));
+			return nullptr;
+		}
+	}
+	else
 	{
-		UE_LOG(LogActor, Warning, TEXT("AddComponent failed because we are in the process of tearing down the world"));
+		UE_LOG(LogActor, Warning, TEXT("AddComponentByClass failed for actor: [%s] with param Class: [%s] because world == nullptr")
+			, *GetName()
+			, *GetNameSafe(Class));
 		return nullptr;
 	}
 
