Additional handling of replay recording errors

diff --git a/Engine/Source/Runtime/Engine/Classes/Engine/DemoNetDriver.h b/Engine/Source/Runtime/Engine/Classes/Engine/DemoNetDriver.h
--- a/Engine/Source/Runtime/Engine/Classes/Engine/DemoNetDriver.h
+++ b/Engine/Source/Runtime/Engine/Classes/Engine/DemoNetDriver.h
@@ -261,6 +261,8 @@ protected:
 	uint8 bIsWaitingForHeaderDownload : 1;
 	uint8 bIsWaitingForStream : 1;
 
+	TOptional<EDemoRecordFailure> PendingRecordFailure;
+
 private:
 	TArray<FNetworkGUID> NonQueuedGUIDsForScrubbing;
 
@@ -320,6 +322,7 @@ public:
 	virtual bool InitConnect(FNetworkNotify* InNotify, const FURL& ConnectURL, FString& Error) override;
 	virtual bool InitListen(FNetworkNotify* InNotify, FURL& ListenURL, bool bReuseAddressAndPort, FString& Error) override;
 	virtual void TickFlush(float DeltaSeconds) override;
+	virtual void PostTickFlush() override;
 	virtual void TickDispatch(float DeltaSeconds) override;
 	virtual void ProcessRemoteFunction(class AActor* Actor, class UFunction* Function, void* Parameters, struct FOutParmRec* OutParms, struct FFrame* Stack, class UObject* SubObject = nullptr) override;
 	virtual bool IsAvailable() const override { return true; }
@@ -692,6 +695,7 @@ protected:
 	bool DemoReplicateActor(AActor* Actor, UNetConnection* Connection, bool bMustReplicate);
 
 	void NotifyDemoPlaybackFailure(EDemoPlayFailure::Type FailureType);
+	void NotifyDemoRecordFailure(EDemoRecordFailure FailureType);
 
 	TArray<FQueuedDemoPacket> QueuedPacketsBeforeTravel;
 
diff --git a/Engine/Source/Runtime/Engine/Classes/Engine/GameInstance.h b/Engine/Source/Runtime/Engine/Classes/Engine/GameInstance.h
--- a/Engine/Source/Runtime/Engine/Classes/Engine/GameInstance.h
+++ b/Engine/Source/Runtime/Engine/Classes/Engine/GameInstance.h
@@ -11,10 +11,10 @@
 #include "Engine/EngineBaseTypes.h"
 #include "Engine/NetworkDelegates.h"
 #include "RHIDefinitions.h"
-
 #include "Subsystems/GameInstanceSubsystem.h"
 #include "Subsystems/SubsystemCollection.h"
 #include "GameFramework/OnlineReplStructs.h"
+#include "ReplayTypes.h"
 
 #if WITH_EDITOR
 #include "Settings/LevelEditorPlaySettings.h"
@@ -430,7 +430,10 @@ public:
 	void CleanupGameViewport();
 
 	/** Called when demo playback fails for any reason */
-	virtual void HandleDemoPlaybackFailure( EDemoPlayFailure::Type FailureType, const FString& ErrorString = TEXT("") ) { }
+	virtual void HandleDemoPlaybackFailure(EDemoPlayFailure::Type FailureType, const FString& ErrorString = TEXT("")) { }
+
+	/** Called when demo recording fails for any reason */
+	virtual void HandleDemoRecordFailure(EDemoRecordFailure FailureType) { }
 
 	/** This gets called when the player scrubs in a replay to a different level */
 	virtual void OnSeamlessTravelDuringReplay() { }
diff --git a/Engine/Source/Runtime/Engine/Private/DemoNetDriver.cpp b/Engine/Source/Runtime/Engine/Private/DemoNetDriver.cpp
--- a/Engine/Source/Runtime/Engine/Private/DemoNetDriver.cpp
+++ b/Engine/Source/Runtime/Engine/Private/DemoNetDriver.cpp
@@ -1258,6 +1258,25 @@ void UDemoNetDriver::TickFlush(float DeltaSeconds)
 	}
 }
 
+void UDemoNetDriver::PostTickFlush()
+{
+	Super::PostTickFlush();
+
+	// process this outside of tick in case we want to destroy the driver in response
+	if (PendingRecordFailure.IsSet())
+	{
+		if (World)
+		{
+			if (UGameInstance* GameInstance = World->GetGameInstance())
+			{
+				GameInstance->HandleDemoRecordFailure(PendingRecordFailure.GetValue());
+			}
+		}
+
+		PendingRecordFailure.Reset();
+	}
+}
+
 void UDemoNetDriver::TickFlushAsyncEndOfFrame(float DeltaSeconds)
 {
 	if (ShouldTickFlushAsyncEndOfFrame())
@@ -1299,7 +1318,7 @@ void UDemoNetDriver::TickFlushInternal(float DeltaSeconds)
 	if (Streamer->GetLastError() != ENetworkReplayError::None)
 	{
 		UE_LOG(LogDemo, Error, TEXT("UDemoNetDriver::TickFlush: ReplayStreamer ERROR: %s"), ENetworkReplayError::ToString(Streamer->GetLastError()));
-		StopDemo();
+		NotifyDemoRecordFailure(EDemoRecordFailure::Generic);
 		return;
 	}
 
@@ -1313,7 +1332,7 @@ void UDemoNetDriver::TickFlushInternal(float DeltaSeconds)
 	if (FileAr == nullptr)
 	{
 		UE_LOG(LogDemo, Error, TEXT("UDemoNetDriver::TickFlush: FileAr == nullptr"));
-		StopDemo();
+		NotifyDemoRecordFailure(EDemoRecordFailure::StreamingArchive);
 		return;
 	}
 
@@ -3150,7 +3169,7 @@ void UDemoNetDriver::ReplayStreamingReady(const FStartStreamingResult& Result)
 
 		if (Result.bRecording)
 		{
-			StopDemo();
+			NotifyDemoRecordFailure(EDemoRecordFailure::StartStreaming);
 		}
 		else
 		{
@@ -5296,6 +5315,20 @@ void UDemoNetDriver::NotifyDemoPlaybackFailure(EDemoPlayFailure::Type FailureTyp
 	}
 }
 
+void UDemoNetDriver::NotifyDemoRecordFailure(EDemoRecordFailure FailureType)
+{
+	UE_LOG(LogDemo, Warning, TEXT("Demo record failure: '%s'"), LexToString(FailureType));
+
+	const bool bIsRecording = IsRecording();
+
+	StopDemo();
+
+	if (bIsRecording && !PendingRecordFailure.IsSet())
+	{
+		PendingRecordFailure.Emplace(FailureType);
+	}
+}
+
 FString UDemoNetDriver::GetDemoPath() const
 {
 	if (ReplayHelper.ReplayStreamer.IsValid())
diff --git a/Engine/Source/Runtime/Engine/Private/ReplayHelper.cpp b/Engine/Source/Runtime/Engine/Private/ReplayHelper.cpp
--- a/Engine/Source/Runtime/Engine/Private/ReplayHelper.cpp
+++ b/Engine/Source/Runtime/Engine/Private/ReplayHelper.cpp
@@ -2291,6 +2291,22 @@ const TCHAR* LexToString(EReplayHeaderFlags Flag)
 	}
 }
 
+const TCHAR* LexToString(EDemoRecordFailure FailureType)
+{
+	switch (FailureType)
+	{
+	case EDemoRecordFailure::Generic:
+		return TEXT("Generic");
+	case EDemoRecordFailure::StreamingArchive:
+		return TEXT("StreamingArchive");
+	case EDemoRecordFailure::StartStreaming:
+		return TEXT("StartStreaming");
+	default:
+		check(false);
+		return TEXT("Unknown");
+	}
+}
+
 void FReplayHelper::RequestCheckpoint()
 {
 	bPendingCheckpointRequest = true;
diff --git a/Engine/Source/Runtime/Engine/Public/ReplayTypes.h b/Engine/Source/Runtime/Engine/Public/ReplayTypes.h
--- a/Engine/Source/Runtime/Engine/Public/ReplayTypes.h
+++ b/Engine/Source/Runtime/Engine/Public/ReplayTypes.h
@@ -43,6 +43,15 @@ enum class EWriteDemoFrameFlags : uint32
 
 ENUM_CLASS_FLAGS(EWriteDemoFrameFlags);
 
+enum class EDemoRecordFailure : uint32
+{
+	Generic,
+	StreamingArchive,
+	StartStreaming,
+};
+
+ENGINE_API const TCHAR* LexToString(EDemoRecordFailure FailureType);
+
 struct FPlaybackPacket
 {
 	TArray<uint8>		Data;
