fix for Skeletal mesh bounds becoming too large if bIncludeComponentLocationIntoBounds is set to true.
the issue was introduced with CL#18366871

diff --git a/Engine/Source/Runtime/Engine/Private/Components/SkeletalMeshComponent.cpp b/Engine/Source/Runtime/Engine/Private/Components/SkeletalMeshComponent.cpp
--- a/Engine/Source/Runtime/Engine/Private/Components/SkeletalMeshComponent.cpp
+++ b/Engine/Source/Runtime/Engine/Private/Components/SkeletalMeshComponent.cpp
@@ -2853,7 +2853,17 @@ FBoxSphereBounds USkeletalMeshComponent::CalcBounds(const FTransform& LocalToWor
 		if (bIncludeComponentLocationIntoBounds)
 		{
 			const FVector ComponentLocation = GetComponentLocation();
-			NewBounds = NewBounds + FBoxSphereBounds(ComponentLocation, FVector(1.0f), 1.0f);
+			const FBoxSphereBounds ComponentLocationBounds(ComponentLocation, FVector(1.0f), 1.0f);
+			if (bCacheLocalSpaceBounds)
+			{
+				NewBounds = NewBounds.TransformBy(LocalToWorld);
+				NewBounds = NewBounds + ComponentLocationBounds;
+				NewBounds = NewBounds.TransformBy(LocalToWorld.ToInverseMatrixWithScale());
+			}
+			else
+			{
+				NewBounds = NewBounds + ComponentLocationBounds;
+			}			
 		}
 
 		AddClothingBounds(NewBounds, CachedBoundsTransform);
