World Partition: disable invalid actor files detection temporarily until we figure out why it fails during HLOD generation.

#rnx

diff --git a/Engine/Source/Runtime/Engine/Private/WorldPartition/WorldPartitionActorDescUtils.cpp b/Engine/Source/Runtime/Engine/Private/WorldPartition/WorldPartitionActorDescUtils.cpp
--- a/Engine/Source/Runtime/Engine/Private/WorldPartition/WorldPartitionActorDescUtils.cpp
+++ b/Engine/Source/Runtime/Engine/Private/WorldPartition/WorldPartitionActorDescUtils.cpp
@@ -60,7 +60,7 @@ TUniquePtr<FWorldPartitionActorDesc> FWorldPartitionActorDescUtils::GetActorDesc
 			UE_LOG(LogWorldPartition, Log, TEXT("Invalid class for actor guid `%s` ('%s') from package '%s'"), *NewActorDesc->GetGuid().ToString(), *NewActorDesc->GetActorName().ToString(), *NewActorDesc->GetActorPackage().ToString());
 			return nullptr;
 		}
-		else if (UClass* Class = FindObject<UClass>(InAssetData.AssetClassPath); !Class)
+		/*else if (UClass* Class = FindObject<UClass>(InAssetData.AssetClassPath); !Class)
 		{
 			// We can't detect mising BP classes for inactive plugins, etc.
 			if (InAssetData.AssetClassPath.GetPackageName().ToString().StartsWith(TEXT("/Game/")))
@@ -75,7 +75,7 @@ TUniquePtr<FWorldPartitionActorDesc> FWorldPartitionActorDescUtils::GetActorDesc
 					return nullptr;
 				}
 			}
-		}
+		}*/
 
 		return NewActorDesc;
 	}
