Exempt /Engine/ and /Game/ material paths from restrictive mode.

#rnx

diff --git a/Engine/Source/Runtime/Engine/Private/Materials/HLSLMaterialTranslator.cpp b/Engine/Source/Runtime/Engine/Private/Materials/HLSLMaterialTranslator.cpp
--- a/Engine/Source/Runtime/Engine/Private/Materials/HLSLMaterialTranslator.cpp
+++ b/Engine/Source/Runtime/Engine/Private/Materials/HLSLMaterialTranslator.cpp
@@ -11161,11 +11161,49 @@ int32 FHLSLMaterialTranslator::AccessMaterialAttribute(int32 CodeIndex, const FG
 		*AttributeName);
 }
 
+static bool IsMaterialPathExemptFromRestrictiveMode(FString& OutPathName)
+{
+	// Anything currently being edited is NOT exempt from restrictive mode
+	const TCHAR* EngineTransientPath = TEXT("/Engine/Transient.");
+	if (OutPathName.StartsWith(EngineTransientPath))
+	{
+		OutPathName.RightChopInline(FCString::Strlen(EngineTransientPath));
+		return false;
+	}
+
+	// Anything else in /Engine/ and /Game/ folders is exempt
+	if (OutPathName.StartsWith(TEXT("/Engine/")) || OutPathName.StartsWith(TEXT("/Game/")))
+	{
+		return true;
+	}
+
+	return false;
+}
+
+static bool IsMaterialExpressionExemptFromRestrictiveMode(const UMaterialExpression* Expression, FString& OutPathName)
+{
+	if (Expression)
+	{
+		if (Expression->Material)
+		{
+			OutPathName = Expression->Material->GetPathName();
+			return IsMaterialPathExemptFromRestrictiveMode(OutPathName);
+		}
+		if (Expression->Function)
+		{
+			OutPathName = Expression->Function->GetPathName();
+			return IsMaterialPathExemptFromRestrictiveMode(OutPathName);
+		}
+	}
+	return false;
+}
+
 int32 FHLSLMaterialTranslator::CustomExpression( class UMaterialExpressionCustom* Custom, int32 OutputIndex, TArray<int32>& CompiledInputs )
 {
-	if (!IsExpressionClassPermitted(Custom->GetClass()))
+	FString MaterialPathName;
+	if (!IsExpressionClassPermitted(Custom->GetClass()) && !IsMaterialExpressionExemptFromRestrictiveMode(Custom, MaterialPathName))
 	{
-		return Errorf(TEXT("Custom HLSL expressions are not permitted"));
+		return Errorf(TEXT("Custom HLSL expressions are not permitted in %s"), *MaterialPathName);
 	}
 
 	const FMaterialCustomExpressionEntry* CustomEntry = nullptr;
