Removed overzealous ensures on register/unregister in FDebugDrawDelegateHelper that did not match the reality when working with BPs and level visibility toggling in the editor

diff --git a/Engine/Source/Runtime/Engine/Private/DebugRenderSceneProxy.cpp b/Engine/Source/Runtime/Engine/Private/DebugRenderSceneProxy.cpp
--- a/Engine/Source/Runtime/Engine/Private/DebugRenderSceneProxy.cpp
+++ b/Engine/Source/Runtime/Engine/Private/DebugRenderSceneProxy.cpp
@@ -27,7 +27,8 @@ FDebugRenderSceneProxy::FDebugRenderSceneProxy(const UPrimitiveComponent* InComp
 
 void FDebugDrawDelegateHelper::RegisterDebugDrawDelegateInternal()
 {
-	ensureMsgf(State != RegisteredState, TEXT("DrawDelegate is already Registered!"));
+	// note that it's possible at this point for State == RegisteredState since RegisterDebugDrawDelegateInternal can get 
+	// called in multiple scenarios, most notably blueprint recompilation or changing level visibility in the editor.
 	if (State == InitializedState)
 	{
 		DebugTextDrawingDelegate = FDebugDrawDelegate::CreateRaw(this, &FDebugDrawDelegateHelper::DrawDebugLabels);
@@ -57,7 +58,8 @@ void FDebugDrawDelegateHelper::ProcessDeferredRegister()
 
 void FDebugDrawDelegateHelper::UnregisterDebugDrawDelegate()
 {
-	ensureMsgf(State != InitializedState, TEXT("DrawDelegate is in an invalid State: %i !"), State);
+	// note that it's possible at this point for State == InitializedState since UnregisterDebugDrawDelegate can get 
+	// called in multiple scenarios, most notably blueprint recompilation or changing level visibility in the editor.
 	if (State == RegisteredState)
 	{
 		check(DebugTextDrawingDelegate.IsBound());
