Enable Random Sequence Player react to automatic transition rules

diff --git a/Engine/Source/Editor/AnimGraph/Private/AnimGraphNode_StateMachineBase.cpp b/Engine/Source/Editor/AnimGraph/Private/AnimGraphNode_StateMachineBase.cpp
--- a/Engine/Source/Editor/AnimGraph/Private/AnimGraphNode_StateMachineBase.cpp
+++ b/Engine/Source/Editor/AnimGraph/Private/AnimGraphNode_StateMachineBase.cpp
@@ -23,6 +23,7 @@
 #include "AnimGraphNode_SequencePlayer.h"
 #include "AnimStateConduitNode.h"
 #include "AnimGraphNode_LinkedAnimGraphBase.h"
+#include "AnimGraphNode_RandomPlayer.h"
 #include "AnimGraphNode_TransitionPoseEvaluator.h"
 #include "AnimGraphNode_CustomTransitionResult.h"
 #include "AnimBlueprintExtension_StateMachine.h"
@@ -564,9 +565,11 @@ void UAnimGraphNode_StateMachineBase::OnProcessDuringCompilation(IAnimBlueprintC
 
 		TArray<UAnimGraphNode_LinkedAnimGraphBase*> LinkedAnimGraphNodes;
 		TArray<UAnimGraphNode_AssetPlayerBase*> AssetPlayerNodes;
+		TArray<UAnimGraphNode_RandomPlayer*> AssetRandomPlayerPlayerNodes;
 		for (UEdGraph* ChildGraph : GraphsToCheck)
 		{
 			ChildGraph->GetNodesOfClass(AssetPlayerNodes);
+			ChildGraph->GetNodesOfClass(AssetRandomPlayerPlayerNodes);
 			ChildGraph->GetNodesOfClass(LinkedAnimGraphNodes);
 		}
 
@@ -578,6 +581,14 @@ void UAnimGraphNode_StateMachineBase::OnProcessDuringCompilation(IAnimBlueprintC
 			}
 		}
 
+		for (UAnimGraphNode_RandomPlayer* Node : AssetRandomPlayerPlayerNodes)
+		{
+			if (int32* IndexPtr = OutCompiledData.GetAnimBlueprintDebugData().NodeGuidToIndexMap.Find(Node->NodeGuid))
+			{
+				BakedState.PlayerNodeIndices.Add(*IndexPtr);
+			}
+		}
+
 		for (UAnimGraphNode_LinkedAnimGraphBase* Node : LinkedAnimGraphNodes)
 		{
 			if (int32* IndexPtr = OutCompiledData.GetAnimBlueprintDebugData().NodeGuidToIndexMap.Find(Node->NodeGuid))
diff --git a/Engine/Source/Runtime/AnimGraphRuntime/Private/AnimNodes/AnimNode_RandomPlayer.cpp b/Engine/Source/Runtime/AnimGraphRuntime/Private/AnimNodes/AnimNode_RandomPlayer.cpp
--- a/Engine/Source/Runtime/AnimGraphRuntime/Private/AnimNodes/AnimNode_RandomPlayer.cpp
+++ b/Engine/Source/Runtime/AnimGraphRuntime/Private/AnimNodes/AnimNode_RandomPlayer.cpp
@@ -114,6 +114,8 @@ void FAnimNode_RandomPlayer::Initialize_AnyThread(const FAnimationInitializeCont
 
 void FAnimNode_RandomPlayer::Update_AnyThread(const FAnimationUpdateContext& Context)
 {
+	BlendWeight = Context.GetFinalBlendWeight();
+
 	DECLARE_SCOPE_HIERARCHICAL_COUNTER_ANIMNODE(Update_AnyThread)
 	GetEvaluateGraphExposedInputs().Execute(Context);
 
@@ -319,6 +321,63 @@ void FAnimNode_RandomPlayer::GatherDebugData(FNodeDebugData& DebugData)
 	DebugData.AddDebugItem(DebugLine, true);
 }
 
+UAnimationAsset* FAnimNode_RandomPlayer::GetAnimAsset() const
+{
+	UAnimationAsset* AnimationAsset = nullptr;
+
+	if (ValidEntries.Num() > 0)
+	{
+		const FRandomAnimPlayData& CurrentPlayData = GetPlayData(ERandomDataIndexType::Current);
+		AnimationAsset = (CurrentPlayData.Entry != nullptr) ? CurrentPlayData.Entry->Sequence : nullptr;
+	}
+
+	return AnimationAsset;
+}
+
+float FAnimNode_RandomPlayer::GetAccumulatedTime() const
+{
+	float AccumulatedTime = 0.f;
+
+	if (ValidEntries.Num() > 0)
+	{
+		const FRandomAnimPlayData& CurrentPlayData = GetPlayData(ERandomDataIndexType::Current);
+
+		return CurrentPlayData.CurrentPlayTime;
+	}
+
+	return AccumulatedTime;
+}
+
+bool FAnimNode_RandomPlayer::GetIgnoreForRelevancyTest() const
+{
+	return GET_ANIM_NODE_DATA(bool, bIgnoreForRelevancyTest);
+}
+
+bool FAnimNode_RandomPlayer::SetIgnoreForRelevancyTest(bool bInIgnoreForRelevancyTest)
+{
+#if WITH_EDITORONLY_DATA
+	bIgnoreForRelevancyTest = bInIgnoreForRelevancyTest;
+#endif
+
+	if (bool* bIgnoreForRelevancyTestPtr = GET_INSTANCE_ANIM_NODE_DATA_PTR(bool, bIgnoreForRelevancyTest))
+	{
+		*bIgnoreForRelevancyTestPtr = bInIgnoreForRelevancyTest;
+		return true;
+	}
+
+	return false;
+}
+
+float FAnimNode_RandomPlayer::GetCachedBlendWeight() const
+{
+	return BlendWeight;
+}
+
+void FAnimNode_RandomPlayer::ClearCachedBlendWeight()
+{
+	BlendWeight = 0.f;
+}
+
 int32 FAnimNode_RandomPlayer::GetNextValidEntryIndex()
 {
 	check(ValidEntries.Num() > 0);
@@ -361,12 +420,17 @@ FRandomAnimPlayData& FAnimNode_RandomPlayer::GetPlayData(ERandomDataIndexType Ty
 	}
 }
 
-void FAnimNode_RandomPlayer::InitPlayData(FRandomAnimPlayData& Data, int32 ValidEntryIndex, float BlendWeight)
+const FRandomAnimPlayData& FAnimNode_RandomPlayer::GetPlayData(ERandomDataIndexType Type) const
+{
+	return const_cast<FAnimNode_RandomPlayer*>(this)->GetPlayData(Type);
+}
+
+void FAnimNode_RandomPlayer::InitPlayData(FRandomAnimPlayData& Data, int32 InValidEntryIndex, float InBlendWeight)
 {
-	FRandomPlayerSequenceEntry* Entry = ValidEntries[ValidEntryIndex];
+	FRandomPlayerSequenceEntry* Entry = ValidEntries[InValidEntryIndex];
 
 	Data.Entry = Entry;
-	Data.BlendWeight = BlendWeight;
+	Data.BlendWeight = InBlendWeight;
 	Data.PlayRate = RandomStream.FRandRange(Entry->MinPlayRate, Entry->MaxPlayRate);
 	Data.RemainingLoops = FMath::Clamp(RandomStream.RandRange(Entry->MinLoopCount, Entry->MaxLoopCount), 0, MAX_int32);
 
diff --git a/Engine/Source/Runtime/AnimGraphRuntime/Public/AnimNodes/AnimNode_RandomPlayer.h b/Engine/Source/Runtime/AnimGraphRuntime/Public/AnimNodes/AnimNode_RandomPlayer.h
--- a/Engine/Source/Runtime/AnimGraphRuntime/Public/AnimNodes/AnimNode_RandomPlayer.h
+++ b/Engine/Source/Runtime/AnimGraphRuntime/Public/AnimNodes/AnimNode_RandomPlayer.h
@@ -3,7 +3,7 @@
 #pragma once
 
 #include "AlphaBlend.h"
-#include "Animation/AnimNodeBase.h"
+#include "Animation/AnimNode_RelevantAssetPlayerBase.h"
 #include "Animation/AnimationAsset.h"
 #include "CoreMinimal.h"
 #include "Math/RandomStream.h"
@@ -93,7 +93,7 @@ struct FRandomAnimPlayData
 };
 
 USTRUCT(BlueprintInternalUseOnly)
-struct ANIMGRAPHRUNTIME_API FAnimNode_RandomPlayer : public FAnimNode_Base
+struct ANIMGRAPHRUNTIME_API FAnimNode_RandomPlayer : public FAnimNode_AssetPlayerRelevancyBase
 {
 	GENERATED_BODY()
 
@@ -111,6 +111,15 @@ public:
 	virtual void GatherDebugData(FNodeDebugData& DebugData) override;
 	// End of FAnimNode_Base interface
 
+	// FAnimNode_RelevantAssetPlayerBase
+	virtual UAnimationAsset* GetAnimAsset() const override;
+	virtual float GetAccumulatedTime() const override;
+	virtual bool GetIgnoreForRelevancyTest() const override;
+	virtual bool SetIgnoreForRelevancyTest(bool bInIgnoreForRelevancyTest) override;
+	virtual float GetCachedBlendWeight() const override;
+	virtual void ClearCachedBlendWeight() override;
+	// End of FAnimNode_RelevantAssetPlayerBase
+
 private:
 	// Return the index of the next FRandomPlayerSequenceEntry to play, from the list
 	// of valid playable entries (ValidEntries).
@@ -119,10 +128,11 @@ private:
 	// Return the play data for either the currently playing animation or the next
 	// animation to blend into.
 	FRandomAnimPlayData& GetPlayData(ERandomDataIndexType Type);
+	const FRandomAnimPlayData& GetPlayData(ERandomDataIndexType Type) const;
 
 	// Initialize the play data with the given index into the ValidEntries array and
 	// a specific blend weight. All other member data will be reset to their default values.
-	void InitPlayData(FRandomAnimPlayData& Data, int32 ValidEntryIndex, float BlendWeight);
+	void InitPlayData(FRandomAnimPlayData& Data, int32 InValidEntryIndex, float InBlendWeight);
 
 	// Advance to the next playable sequence. This is only called once a sequence is fully
 	// blended or there's a hard switch to the same playable entry.
@@ -152,6 +162,17 @@ private:
 	// Random number source
 	FRandomStream RandomStream;
 
+#if WITH_EDITORONLY_DATA
+	// If true, "Relevant anim" nodes that look for the highest weighted animation in a state will ignore this node
+	UPROPERTY(EditAnywhere, Category = Relevancy, meta = (FoldProperty, PinHiddenByDefault))
+	bool bIgnoreForRelevancyTest = false;
+#endif // WITH_EDITORONLY_DATA
+
+protected:
+	/** Last encountered blend weight for this node */
+	UPROPERTY(BlueprintReadWrite, Transient, Category = DoNotEdit)
+	float BlendWeight = 0.0f;
+
 public:
 	/** When shuffle mode is active we will never loop a sequence beyond MaxLoopCount
 	  * without visiting each sequence in turn (no repeats). Enabling this will ignore
diff --git a/Engine/Source/Runtime/Engine/Classes/Animation/AnimInstance.h b/Engine/Source/Runtime/Engine/Classes/Animation/AnimInstance.h
--- a/Engine/Source/Runtime/Engine/Classes/Animation/AnimInstance.h
+++ b/Engine/Source/Runtime/Engine/Classes/Animation/AnimInstance.h
@@ -27,6 +27,7 @@ class UAnimInstance;
 class UCanvas;
 struct FAnimInstanceProxy;
 struct FAnimNode_AssetPlayerBase;
+struct FAnimNode_AssetPlayerRelevancyBase;
 struct FAnimNode_StateMachine;
 struct FAnimNode_LinkedInputPose;
 struct FBakedAnimationStateMachine;
@@ -1104,11 +1105,24 @@ public:
 	/** Returns all Animation Nodes of FAnimNode_AssetPlayerBase class within the specified (named) Animation Graph */
 	TArray<FAnimNode_AssetPlayerBase*> GetMutableInstanceAssetPlayers(const FName& GraphName);
 
+	/** Returns all Animation Nodes implementing FAnimNode_AssetPlayerRelevancyBase within the specified (named) Animation Graph */
+	TArray<const FAnimNode_AssetPlayerRelevancyBase*> GetInstanceRelevantAssetPlayers(const FName& GraphName) const;
+
+	/** Returns all Animation Nodes of FAnimNode_AssetPlayerRelevancyBase class within the specified (named) Animation Graph */
+	TArray<FAnimNode_AssetPlayerRelevancyBase*> GetMutableInstanceRelevantAssetPlayers(const FName& GraphName);
+
 	/** Gets the runtime instance desc of the state machine specified by name */
 	const FBakedAnimationStateMachine* GetStateMachineInstanceDesc(FName MachineName) const;
 
 	/** Gets the most relevant asset player in a specified state */
-	const FAnimNode_AssetPlayerBase* GetRelevantAssetPlayerFromState(int32 MachineIndex, int32 StateIndex) const;
+	UE_DEPRECATED(5.1, "Please use GetRelevantAssetPlayerInterfaceFromState")
+	const FAnimNode_AssetPlayerBase* GetRelevantAssetPlayerFromState(int32 MachineIndex, int32 StateIndex) const
+	{
+		return nullptr;
+	}
+
+	/** Gets the most relevant asset player in a specified state */
+	const FAnimNode_AssetPlayerRelevancyBase* GetRelevantAssetPlayerInterfaceFromState(int32 MachineIndex, int32 StateIndex) const;
 
 	//////////////////////////////////////////////////////////////////////////
 
diff --git a/Engine/Source/Runtime/Engine/Classes/Animation/AnimNode_AssetPlayerBase.h b/Engine/Source/Runtime/Engine/Classes/Animation/AnimNode_AssetPlayerBase.h
--- a/Engine/Source/Runtime/Engine/Classes/Animation/AnimNode_AssetPlayerBase.h
+++ b/Engine/Source/Runtime/Engine/Classes/Animation/AnimNode_AssetPlayerBase.h
@@ -5,12 +5,12 @@
 #include "CoreMinimal.h"
 #include "UObject/ObjectMacros.h"
 #include "Animation/AnimationAsset.h"
-#include "Animation/AnimNodeBase.h"
+#include "Animation/AnimNode_RelevantAssetPlayerBase.h"
 #include "AnimNode_AssetPlayerBase.generated.h"
 
 /* Base class for any asset playing anim node */
 USTRUCT(BlueprintInternalUseOnly)
-struct ENGINE_API FAnimNode_AssetPlayerBase : public FAnimNode_Base
+struct ENGINE_API FAnimNode_AssetPlayerBase : public FAnimNode_AssetPlayerRelevancyBase
 {
 	GENERATED_BODY()
 
@@ -18,21 +18,6 @@ struct ENGINE_API FAnimNode_AssetPlayerBase : public FAnimNode_Base
 
 	FAnimNode_AssetPlayerBase() = default;
 
-	/** Get the last encountered blend weight for this node */
-	virtual float GetCachedBlendWeight() const;
-	
-	/** Set the cached blendweight to zero */
-	void ClearCachedBlendWeight();
-
-	/** Get the currently referenced time within the asset player node */
-	virtual float GetAccumulatedTime() const;
-
-	/** Override the currently accumulated time */
-	virtual void SetAccumulatedTime(float NewTime);
-
-	/** Get the animation asset associated with the node, derived classes should implement this */
-	virtual UAnimationAsset* GetAnimAsset() const;
-
 	/** Initialize function for setup purposes */
 	virtual void Initialize_AnyThread(const FAnimationInitializeContext& Context) override;
 
@@ -47,11 +32,6 @@ struct ENGINE_API FAnimNode_AssetPlayerBase : public FAnimNode_Base
 	// Create a tick record for this node
 	void CreateTickRecordForNode(const FAnimationUpdateContext& Context, UAnimSequenceBase* Sequence, bool bLooping, float PlayRate);
 
-	// Functions to report data to getters, this is required for all asset players (but can't be pure abstract because of struct instantiation generated code).
-	virtual float GetCurrentAssetLength() const { return 0.0f; }
-	virtual float GetCurrentAssetTime() const { return 0.0f; }
-	virtual float GetCurrentAssetTimePlayRateAdjusted() const { return GetCurrentAssetTime(); }
-
 	// Get the sync group name we are using
 	virtual FName GetGroupName() const { return NAME_None; }
 
@@ -61,9 +41,6 @@ struct ENGINE_API FAnimNode_AssetPlayerBase : public FAnimNode_Base
 	// Get the sync group method we are using
 	virtual EAnimSyncMethod GetGroupMethod() const { return EAnimSyncMethod::DoNotSync; }
 
-	// Check whether this node should be ignored when testing for relevancy in state machines
-	virtual bool GetIgnoreForRelevancyTest() const { return false; }
-
 	// Set the sync group name we are using
 	virtual bool SetGroupName(FName InGroupName) { return false; }
 
@@ -73,8 +50,12 @@ struct ENGINE_API FAnimNode_AssetPlayerBase : public FAnimNode_Base
 	// Set the sync group method we are using
 	virtual bool SetGroupMethod(EAnimSyncMethod InMethod) { return false; }
 
-	// Set whether this node should be ignored when testing for relevancy in state machines
-	virtual bool SetIgnoreForRelevancyTest(bool bInIgnoreForRelevancyTest) { return false; }
+	// --- FAnimNode_RelevantAssetPlayerBase ---
+	virtual float GetAccumulatedTime() const override;
+	virtual void SetAccumulatedTime(float NewTime) override;
+	virtual float GetCachedBlendWeight() const override;
+	virtual void ClearCachedBlendWeight() override;
+	// --- End of FAnimNode_RelevantAssetPlayerBase ---
 
 private:
 #if WITH_EDITORONLY_DATA
diff --git a/Engine/Source/Runtime/Engine/Classes/Animation/AnimNode_RelevantAssetPlayerBase.h b/Engine/Source/Runtime/Engine/Classes/Animation/AnimNode_RelevantAssetPlayerBase.h
new file mode 100644
--- /dev/null
+++ b/Engine/Source/Runtime/Engine/Classes/Animation/AnimNode_RelevantAssetPlayerBase.h
@@ -0,0 +1,43 @@
+// Copyright Epic Games, Inc. All Rights Reserved.
+
+#pragma once
+
+#include "CoreMinimal.h"
+#include "UObject/ObjectMacros.h"
+#include "Animation/AnimationAsset.h"
+#include "Animation/AnimNodeBase.h"
+#include "AnimNode_RelevantAssetPlayerBase.generated.h"
+
+/* Base class for any asset playing anim node */
+USTRUCT(BlueprintInternalUseOnly)
+struct ENGINE_API FAnimNode_AssetPlayerRelevancyBase : public FAnimNode_Base
+{
+	GENERATED_BODY()
+
+public:
+	/** Get the animation asset associated with the node, derived classes should implement this */
+	virtual class UAnimationAsset* GetAnimAsset() const;
+
+	/** Get the currently referenced time within the asset player node */
+	virtual float GetAccumulatedTime() const;
+
+	/** Override the currently accumulated time */
+	virtual void SetAccumulatedTime(float NewTime);
+
+	// Functions to report data to getters, this is required for all asset players (but can't be pure abstract because of struct instantiation generated code).
+	virtual float GetCurrentAssetLength() const;
+	virtual float GetCurrentAssetTime() const;
+	virtual float GetCurrentAssetTimePlayRateAdjusted() const;
+
+	// Check whether this node should be ignored when testing for relevancy in state machines
+	virtual bool GetIgnoreForRelevancyTest() const;
+
+	// Set whether this node should be ignored when testing for relevancy in state machines
+	virtual bool SetIgnoreForRelevancyTest(bool bInIgnoreForRelevancyTest);
+
+	/** Get the last encountered blend weight for this node */
+	virtual float GetCachedBlendWeight() const;
+
+	/** Set the cached blendweight to zero */
+	virtual void ClearCachedBlendWeight();
+};
diff --git a/Engine/Source/Runtime/Engine/Classes/Animation/AnimNode_StateMachine.h b/Engine/Source/Runtime/Engine/Classes/Animation/AnimNode_StateMachine.h
--- a/Engine/Source/Runtime/Engine/Classes/Animation/AnimNode_StateMachine.h
+++ b/Engine/Source/Runtime/Engine/Classes/Animation/AnimNode_StateMachine.h
@@ -15,6 +15,7 @@
 
 class IAnimClassInterface;
 struct FAnimNode_AssetPlayerBase;
+struct FAnimNode_AssetPlayerRelevancyBase;
 struct FAnimNode_StateMachine;
 struct FAnimNode_TransitionPoseEvaluator;
 
@@ -306,10 +307,22 @@ protected:
 		return GetRelevantAnimTimeRemainingFraction(Context.AnimInstanceProxy, StateIndex);
 	}
 
-	const FAnimNode_AssetPlayerBase* GetRelevantAssetPlayerFromState(const FAnimInstanceProxy* InAnimInstanceProxy, const FBakedAnimationState& StateInfo) const;
+	UE_DEPRECATED(5.1, "Please use GetRelevantAssetPlayerInterfaceFromState")
+	const FAnimNode_AssetPlayerBase* GetRelevantAssetPlayerFromState(const FAnimInstanceProxy* InAnimInstanceProxy, const FBakedAnimationState& StateInfo) const
+	{
+		return nullptr;
+	}
+
+	UE_DEPRECATED(5.1, "Please use GetRelevantAssetPlayerInterfaceFromState")
 	const FAnimNode_AssetPlayerBase* GetRelevantAssetPlayerFromState(const FAnimationUpdateContext& Context, const FBakedAnimationState& StateInfo) const
 	{
-		return GetRelevantAssetPlayerFromState(Context.AnimInstanceProxy, StateInfo);
+		return nullptr;
+	}
+
+	const FAnimNode_AssetPlayerRelevancyBase* GetRelevantAssetPlayerInterfaceFromState(const FAnimInstanceProxy* InAnimInstanceProxy, const FBakedAnimationState& StateInfo) const;
+	const FAnimNode_AssetPlayerRelevancyBase* GetRelevantAssetPlayerInterfaceFromState(const FAnimationUpdateContext& Context, const FBakedAnimationState& StateInfo) const
+	{
+		return GetRelevantAssetPlayerInterfaceFromState(Context.AnimInstanceProxy, StateInfo);
 	}
 
 	void LogInertializationRequestError(const FAnimationUpdateContext& Context, int32 PreviousState, int32 NextState);
diff --git a/Engine/Source/Runtime/Engine/Private/Animation/AnimInstance.cpp b/Engine/Source/Runtime/Engine/Private/Animation/AnimInstance.cpp
--- a/Engine/Source/Runtime/Engine/Private/Animation/AnimInstance.cpp
+++ b/Engine/Source/Runtime/Engine/Private/Animation/AnimInstance.cpp
@@ -3757,9 +3757,14 @@ TArray<FAnimNode_AssetPlayerBase*> UAnimInstance::GetMutableInstanceAssetPlayers
 	return GetProxyOnAnyThread<FAnimInstanceProxy>().GetMutableInstanceAssetPlayers(GraphName);
 }
 
-const FAnimNode_AssetPlayerBase* UAnimInstance::GetRelevantAssetPlayerFromState(int32 MachineIndex, int32 StateIndex) const
+TArray<const FAnimNode_AssetPlayerRelevancyBase*> UAnimInstance::GetInstanceRelevantAssetPlayers(const FName& GraphName) const
 {
-	return GetProxyOnGameThread<FAnimInstanceProxy>().GetRelevantAssetPlayerFromState(MachineIndex, StateIndex);
+	return GetProxyOnAnyThread<FAnimInstanceProxy>().GetInstanceRelevantAssetPlayers(GraphName);
+}
+
+TArray<FAnimNode_AssetPlayerRelevancyBase*> UAnimInstance::GetMutableInstanceRelevantAssetPlayers(const FName& GraphName)
+{
+	return GetProxyOnAnyThread<FAnimInstanceProxy>().GetMutableInstanceRelevantAssetPlayers(GraphName);
 }
 
 int32 UAnimInstance::GetSyncGroupIndexFromName(FName SyncGroupName) const
diff --git a/Engine/Source/Runtime/Engine/Private/Animation/AnimInstanceProxy.cpp b/Engine/Source/Runtime/Engine/Private/Animation/AnimInstanceProxy.cpp
--- a/Engine/Source/Runtime/Engine/Private/Animation/AnimInstanceProxy.cpp
+++ b/Engine/Source/Runtime/Engine/Private/Animation/AnimInstanceProxy.cpp
@@ -2382,7 +2382,7 @@ float FAnimInstanceProxy::GetRelevantAnimTimeRemainingFraction(int32 MachineInde
 
 float FAnimInstanceProxy::GetRelevantAnimLength(int32 MachineIndex, int32 StateIndex) const
 {
-	if(const FAnimNode_AssetPlayerBase* AssetPlayer = GetRelevantAssetPlayerFromState(MachineIndex, StateIndex))
+	if(const FAnimNode_AssetPlayerRelevancyBase* AssetPlayer = GetRelevantAssetPlayerInterfaceFromState(MachineIndex, StateIndex))
 	{
 		if(AssetPlayer->GetAnimAsset())
 		{
@@ -2395,7 +2395,7 @@ float FAnimInstanceProxy::GetRelevantAnimLength(int32 MachineIndex, int32 StateI
 
 float FAnimInstanceProxy::GetRelevantAnimTime(int32 MachineIndex, int32 StateIndex) const
 {
-	if(const FAnimNode_AssetPlayerBase* AssetPlayer = GetRelevantAssetPlayerFromState(MachineIndex, StateIndex))
+	if(const FAnimNode_AssetPlayerRelevancyBase* AssetPlayer = GetRelevantAssetPlayerInterfaceFromState(MachineIndex, StateIndex))
 	{
 		return AssetPlayer->GetCurrentAssetTimePlayRateAdjusted();
 	}
@@ -2405,7 +2405,7 @@ float FAnimInstanceProxy::GetRelevantAnimTime(int32 MachineIndex, int32 StateInd
 
 float FAnimInstanceProxy::GetRelevantAnimTimeFraction(int32 MachineIndex, int32 StateIndex) const
 {
-	if(const FAnimNode_AssetPlayerBase* AssetPlayer = GetRelevantAssetPlayerFromState(MachineIndex, StateIndex))
+	if(const FAnimNode_AssetPlayerRelevancyBase* AssetPlayer = GetRelevantAssetPlayerInterfaceFromState(MachineIndex, StateIndex))
 	{
 		float Length = AssetPlayer->GetCurrentAssetLength();
 		if(Length > 0.0f)
@@ -2654,11 +2654,11 @@ bool FAnimInstanceProxy::QueryAndMarkTransitionEvent(int32 MachineIndex, int32 T
 	return false;
 }
 
-const FAnimNode_AssetPlayerBase* FAnimInstanceProxy::GetRelevantAssetPlayerFromState(int32 MachineIndex, int32 StateIndex) const
+const FAnimNode_AssetPlayerRelevancyBase* FAnimInstanceProxy::GetRelevantAssetPlayerInterfaceFromState(int32 MachineIndex, int32 StateIndex) const
 {
-	if(const FAnimNode_StateMachine* StateMachine = GetStateMachineInstance(MachineIndex))
+	if (const FAnimNode_StateMachine* StateMachine = GetStateMachineInstance(MachineIndex))
 	{
-		return StateMachine->GetRelevantAssetPlayerFromState(this, StateMachine->GetStateInfo(StateIndex));
+		return StateMachine->GetRelevantAssetPlayerInterfaceFromState(this, StateMachine->GetStateInfo(StateIndex));
 	}
 
 	return nullptr;
@@ -3156,6 +3156,51 @@ TArray<FAnimNode_AssetPlayerBase*> FAnimInstanceProxy::GetMutableInstanceAssetPl
 
 	return Nodes;
 }
+TArray<const FAnimNode_AssetPlayerRelevancyBase*> FAnimInstanceProxy::GetInstanceRelevantAssetPlayers(const FName& GraphName) const
+{
+	TArray<const FAnimNode_AssetPlayerRelevancyBase*> Nodes;
+
+	// Retrieve all asset player nodes from the (named) Animation Layer Graph
+	if (AnimClassInterface)
+	{
+		const TMap<FName, FGraphAssetPlayerInformation>& GrapInformationMap = AnimClassInterface->GetGraphAssetPlayerInformation();
+		if (const FGraphAssetPlayerInformation* Information = GrapInformationMap.Find(GraphName))
+		{
+			for (const int32& NodeIndex : Information->PlayerNodeIndices)
+			{
+				if (const FAnimNode_AssetPlayerRelevancyBase* Node = GetNodeFromIndex<FAnimNode_AssetPlayerRelevancyBase>(NodeIndex))
+				{
+					Nodes.Add(Node);
+				}
+			}
+		}
+	}
+
+	return Nodes;
+}
+
+TArray<FAnimNode_AssetPlayerRelevancyBase*> FAnimInstanceProxy::GetMutableInstanceRelevantAssetPlayers(const FName& GraphName)
+{
+	TArray<FAnimNode_AssetPlayerRelevancyBase*> Nodes;
+
+	// Retrieve all asset player nodes from the (named) Animation Layer Graph	
+	if (AnimClassInterface)
+	{
+		const TMap<FName, FGraphAssetPlayerInformation>& GrapInformationMap = AnimClassInterface->GetGraphAssetPlayerInformation();
+		if (const FGraphAssetPlayerInformation* Information = GrapInformationMap.Find(GraphName))
+		{
+			for (const int32& NodeIndex : Information->PlayerNodeIndices)
+			{
+				if (FAnimNode_AssetPlayerRelevancyBase* Node = GetMutableNodeFromIndex<FAnimNode_AssetPlayerRelevancyBase>(NodeIndex))
+				{
+					Nodes.Add(Node);
+				}
+			}
+		}
+	}
+
+	return Nodes;
+}
 
 #if WITH_EDITOR
 void FAnimInstanceProxy::RecordNodeVisit(int32 TargetNodeIndex, int32 SourceNodeIndex, float BlendWeight)
@@ -3509,5 +3554,4 @@ void FAnimInstanceProxy::ForEachStateMachine(const TFunctionRef<void(FAnimNode_S
 		}
 	}
 }
-
 #undef LOCTEXT_NAMESPACE
diff --git a/Engine/Source/Runtime/Engine/Private/Animation/AnimNode_AssetPlayerBase.cpp b/Engine/Source/Runtime/Engine/Private/Animation/AnimNode_AssetPlayerBase.cpp
--- a/Engine/Source/Runtime/Engine/Private/Animation/AnimNode_AssetPlayerBase.cpp
+++ b/Engine/Source/Runtime/Engine/Private/Animation/AnimNode_AssetPlayerBase.cpp
@@ -49,11 +49,6 @@ void FAnimNode_AssetPlayerBase::CreateTickRecordForNode(const FAnimationUpdateCo
 	TRACE_ANIM_TICK_RECORD(Context, TickRecord);
 }
 
-float FAnimNode_AssetPlayerBase::GetCachedBlendWeight() const
-{
-	return BlendWeight;
-}
-
 float FAnimNode_AssetPlayerBase::GetAccumulatedTime() const
 {
 	return InternalTimeAccumulator;
@@ -64,12 +59,13 @@ void FAnimNode_AssetPlayerBase::SetAccumulatedTime(float NewTime)
 	InternalTimeAccumulator = NewTime;
 }
 
-UAnimationAsset* FAnimNode_AssetPlayerBase::GetAnimAsset() const
+float FAnimNode_AssetPlayerBase::GetCachedBlendWeight() const
 {
-	return nullptr;
+	return BlendWeight;
 }
 
 void FAnimNode_AssetPlayerBase::ClearCachedBlendWeight()
 {
 	BlendWeight = 0.0f;
 }
+
diff --git a/Engine/Source/Runtime/Engine/Private/Animation/AnimNode_RelevantAssetPlayerBase.cpp b/Engine/Source/Runtime/Engine/Private/Animation/AnimNode_RelevantAssetPlayerBase.cpp
new file mode 100644
--- /dev/null
+++ b/Engine/Source/Runtime/Engine/Private/Animation/AnimNode_RelevantAssetPlayerBase.cpp
@@ -0,0 +1,52 @@
+// Copyright Epic Games, Inc. All Rights Reserved.
+
+#include "Animation/AnimNode_RelevantAssetPlayerBase.h"
+#include "Animation/AnimSyncScope.h"
+
+UAnimationAsset* FAnimNode_AssetPlayerRelevancyBase::GetAnimAsset() const
+{
+	return nullptr;
+}
+
+float FAnimNode_AssetPlayerRelevancyBase::GetAccumulatedTime() const
+{
+	return 0.f;
+}
+
+void FAnimNode_AssetPlayerRelevancyBase::SetAccumulatedTime(float NewTime)
+{
+}
+
+float FAnimNode_AssetPlayerRelevancyBase::GetCurrentAssetLength() const
+{
+	return 0.f;
+}
+
+float FAnimNode_AssetPlayerRelevancyBase::GetCurrentAssetTime() const
+{
+	return 0.f;
+}
+
+float FAnimNode_AssetPlayerRelevancyBase::GetCurrentAssetTimePlayRateAdjusted() const
+{
+	return 0.f;
+}
+
+bool FAnimNode_AssetPlayerRelevancyBase::GetIgnoreForRelevancyTest() const
+{
+	return false;
+}
+
+bool FAnimNode_AssetPlayerRelevancyBase::SetIgnoreForRelevancyTest(bool bInIgnoreForRelevancyTest)
+{
+	return false;
+}
+
+float FAnimNode_AssetPlayerRelevancyBase::GetCachedBlendWeight() const
+{
+	return 0.f;
+}
+
+void FAnimNode_AssetPlayerRelevancyBase::ClearCachedBlendWeight()
+{
+}
diff --git a/Engine/Source/Runtime/Engine/Private/Animation/AnimNode_StateMachine.cpp b/Engine/Source/Runtime/Engine/Private/Animation/AnimNode_StateMachine.cpp
--- a/Engine/Source/Runtime/Engine/Private/Animation/AnimNode_StateMachine.cpp
+++ b/Engine/Source/Runtime/Engine/Private/Animation/AnimNode_StateMachine.cpp
@@ -555,7 +555,7 @@ void FAnimNode_StateMachine::Update_AnyThread(const FAnimationUpdateContext& Con
 
 float FAnimNode_StateMachine::GetRelevantAnimTimeRemaining(const FAnimInstanceProxy* InAnimInstanceProxy, int32 StateIndex) const
 {
-	if (const FAnimNode_AssetPlayerBase* AssetPlayer = GetRelevantAssetPlayerFromState(InAnimInstanceProxy, GetStateInfo(StateIndex)))
+	if (const FAnimNode_AssetPlayerRelevancyBase* AssetPlayer = GetRelevantAssetPlayerInterfaceFromState(InAnimInstanceProxy, GetStateInfo(StateIndex)))
 	{
 		if (AssetPlayer->GetAnimAsset())
 		{
@@ -568,7 +568,7 @@ float FAnimNode_StateMachine::GetRelevantAnimTimeRemaining(const FAnimInstancePr
 
 float FAnimNode_StateMachine::GetRelevantAnimTimeRemainingFraction(const FAnimInstanceProxy* InAnimInstanceProxy, int32 StateIndex) const
 {
-	if (const FAnimNode_AssetPlayerBase* AssetPlayer = GetRelevantAssetPlayerFromState(InAnimInstanceProxy, GetStateInfo(StateIndex)))
+	if (const FAnimNode_AssetPlayerRelevancyBase* AssetPlayer = GetRelevantAssetPlayerInterfaceFromState(InAnimInstanceProxy, GetStateInfo(StateIndex)))
 	{
 		if (AssetPlayer->GetAnimAsset())
 		{
@@ -583,12 +583,12 @@ float FAnimNode_StateMachine::GetRelevantAnimTimeRemainingFraction(const FAnimIn
 	return 1.f;
 }
 
-const FAnimNode_AssetPlayerBase* FAnimNode_StateMachine::GetRelevantAssetPlayerFromState(const FAnimInstanceProxy* InAnimInstanceProxy, const FBakedAnimationState& StateInfo) const
+const FAnimNode_AssetPlayerRelevancyBase* FAnimNode_StateMachine::GetRelevantAssetPlayerInterfaceFromState(const FAnimInstanceProxy* InAnimInstanceProxy, const FBakedAnimationState& StateInfo) const
 {
-	const FAnimNode_AssetPlayerBase* ResultPlayer = nullptr;
+	const FAnimNode_AssetPlayerRelevancyBase* ResultPlayer = nullptr;
 	float MaxWeight = 0.0f;
 
-	auto EvaluatePlayerWeight = [&MaxWeight, &ResultPlayer](const FAnimNode_AssetPlayerBase* Player)
+	auto EvaluatePlayerWeight = [&MaxWeight, &ResultPlayer](const FAnimNode_AssetPlayerRelevancyBase* Player)
 	{
 		if (!Player->GetIgnoreForRelevancyTest() && Player->GetCachedBlendWeight() > MaxWeight)
 		{
@@ -599,7 +599,7 @@ const FAnimNode_AssetPlayerBase* FAnimNode_StateMachine::GetRelevantAssetPlayerF
 
 	for (const int32& PlayerIdx : StateInfo.PlayerNodeIndices)
 	{
-		if (const FAnimNode_AssetPlayerBase* Player = InAnimInstanceProxy->GetNodeFromIndex<FAnimNode_AssetPlayerBase>(PlayerIdx))
+		if (const FAnimNode_AssetPlayerRelevancyBase* Player = InAnimInstanceProxy->GetNodeFromIndex<FAnimNode_AssetPlayerRelevancyBase>(PlayerIdx))
 		{
 			EvaluatePlayerWeight(Player);
 		}
@@ -625,8 +625,8 @@ const FAnimNode_AssetPlayerBase* FAnimNode_StateMachine::GetRelevantAssetPlayerF
 				}
 
 				// Retrieve all asset player nodes from the corresponding Anim blueprint class and apply same logic to find highest weighted asset player 
-				TArray<const FAnimNode_AssetPlayerBase*> PlayerNodesInLayer = CurrentTarget->GetInstanceAssetPlayers(GraphName);
-				for (const FAnimNode_AssetPlayerBase* Player : PlayerNodesInLayer)
+				TArray<const FAnimNode_AssetPlayerRelevancyBase*> PlayerNodesInLayer = CurrentTarget->GetInstanceRelevantAssetPlayers(GraphName);
+				for (const FAnimNode_AssetPlayerRelevancyBase* Player : PlayerNodesInLayer)
 				{
 					EvaluatePlayerWeight(Player);
 				}
@@ -697,7 +697,7 @@ bool FAnimNode_StateMachine::FindValidTransition(const FAnimationUpdateContext&
 		else if (TransitionRule.bAutomaticRemainingTimeRule)
 		{
 			bool bCanEnterTransition = false;
-			if (const FAnimNode_AssetPlayerBase* RelevantPlayer = GetRelevantAssetPlayerFromState(Context, StateInfo))
+			if (const FAnimNode_AssetPlayerRelevancyBase* RelevantPlayer = GetRelevantAssetPlayerInterfaceFromState(Context, StateInfo))
 			{
 				if (const UAnimationAsset* AnimAsset = RelevantPlayer->GetAnimAsset())
 				{
@@ -1038,7 +1038,7 @@ void FAnimNode_StateMachine::SetState(const FAnimationBaseContext& Context, int3
 		// This stops any zero length blends holding on to old weights
 		for(int32 PlayerIndex : GetStateInfo(CurrentState).PlayerNodeIndices)
 		{
-			if(FAnimNode_AssetPlayerBase* Player = Context.AnimInstanceProxy->GetMutableNodeFromIndex<FAnimNode_AssetPlayerBase>(PlayerIndex))
+			if(FAnimNode_AssetPlayerRelevancyBase* Player = Context.AnimInstanceProxy->GetMutableNodeFromIndex<FAnimNode_AssetPlayerRelevancyBase>(PlayerIndex))
 			{
 				Player->ClearCachedBlendWeight();
 			}
@@ -1064,8 +1064,8 @@ void FAnimNode_StateMachine::SetState(const FAnimationBaseContext& Context, int3
 					}
 
 					// Retrieve all asset player nodes from the corresponding Anim blueprint class and clear their cached blend weight
-					TArray<FAnimNode_AssetPlayerBase*> PlayerNodesInLayer = CurrentTarget->GetMutableInstanceAssetPlayers(GraphName);
-					for (FAnimNode_AssetPlayerBase* Player : PlayerNodesInLayer)
+					TArray<FAnimNode_AssetPlayerRelevancyBase*> PlayerNodesInLayer = CurrentTarget->GetMutableInstanceRelevantAssetPlayers(GraphName);
+					for (FAnimNode_AssetPlayerRelevancyBase* Player : PlayerNodesInLayer)
 					{
 						Player->ClearCachedBlendWeight();
 					}
diff --git a/Engine/Source/Runtime/Engine/Public/Animation/AnimInstanceProxy.h b/Engine/Source/Runtime/Engine/Public/Animation/AnimInstanceProxy.h
--- a/Engine/Source/Runtime/Engine/Public/Animation/AnimInstanceProxy.h
+++ b/Engine/Source/Runtime/Engine/Public/Animation/AnimInstanceProxy.h
@@ -35,6 +35,7 @@ struct FAnimBlueprintDebugData_AttributeRecord;
 struct FAnimGroupInstance;
 struct FAnimNodePoseWatch;
 struct FAnimNode_AssetPlayerBase;
+struct FAnimNode_AssetPlayerRelevancyBase;
 struct FAnimNode_Base;
 struct FAnimNode_LinkedInputPose;
 struct FAnimNode_SaveCachedPose;
@@ -520,6 +521,12 @@ public:
 	/** Returns all Animation Nodes of FAnimNode_AssetPlayerBase class within the specified (named) Animation Graph */
 	TArray<FAnimNode_AssetPlayerBase*> GetMutableInstanceAssetPlayers(const FName& GraphName);
 
+	/** Returns all Animation Nodes of FAnimNode_AssetPlayerRelevancyBase class within the specified (named) Animation Graph */
+	TArray<const FAnimNode_AssetPlayerRelevancyBase*> GetInstanceRelevantAssetPlayers(const FName& GraphName) const;
+
+	/** Returns all Animation Nodes of FAnimNode_AssetPlayerRelevancyBase class within the specified (named) Animation Graph */
+	TArray<FAnimNode_AssetPlayerRelevancyBase*> GetMutableInstanceRelevantAssetPlayers(const FName& GraphName);
+
 	/** Returns true if SyncGroupName is valid (exists, and if is based on markers, has valid markers) */
 	bool IsSyncGroupValid(FName InSyncGroupName) const;
 
@@ -756,8 +763,16 @@ protected:
 	const TArray<FMontageEvaluationState>& GetMontageEvaluationData() const;
 
 	/** Check whether we have active morph target curves */
+
+	/** Gets the most relevant asset player in a specified state */
+	UE_DEPRECATED(5.1, "Please use GetRelevantAssetPlayerInterfaceFromState")
+	const FAnimNode_AssetPlayerBase* GetRelevantAssetPlayerFromState(int32 MachineIndex, int32 StateIndex) const
+	{
+		return nullptr;
+	}
+
 	/** Gets the most relevant asset player in a specified state */
-	const FAnimNode_AssetPlayerBase* GetRelevantAssetPlayerFromState(int32 MachineIndex, int32 StateIndex) const;
+	const FAnimNode_AssetPlayerRelevancyBase* GetRelevantAssetPlayerInterfaceFromState(int32 MachineIndex, int32 StateIndex) const;
 
 	/** Gets an unchecked (can return nullptr) node given a property of the anim instance */
 	template<class NodeType>
