Reduce multi influence cloth vertex factory shader permutations by merging with single influence cloth VF

CitySample PC cook stats (recompile all shaders in both runs):
Baseline: 75384 shaders compiled, Cook time 1h 35m 47s
WithChange: 71431 shaders compiled, Cook time 1h 25m 23s
Shader reduction 5.24%, Cook time reduction 10mins

diff --git a/Engine/Shaders/Private/GpuSkinCacheComputeShader.usf b/Engine/Shaders/Private/GpuSkinCacheComputeShader.usf
--- a/Engine/Shaders/Private/GpuSkinCacheComputeShader.usf
+++ b/Engine/Shaders/Private/GpuSkinCacheComputeShader.usf
@@ -47,15 +47,13 @@ void GetMorphBuffer(uint VertexIndex, out float3 DeltaPosition, out half3 DeltaT
 		MorphBuffer[Offset + 5]);
 }
 #endif
-#if GPUSKIN_APEX_CLOTH
 
-#if GPUSKIN_MULTIPLE_CLOTH_INFLUENCES
+#if GPUSKIN_APEX_CLOTH
 // This must match NUM_INFLUENCES_PER_VERTEX in ClothingMeshUtils.cpp and GpuSkinVertexFactory.ush
+// It represents the maximum number of influences so perhaps should be renamed. Right now we are keeping the name matching for easy search.
 // TODO: Make this easier to change in without messing things up
 #define NUM_INFLUENCES_PER_VERTEX 5
-#else
-#define NUM_INFLUENCES_PER_VERTEX 1
-#endif
+uint ClothNumInfluencesPerVertex;
 
 struct FVertexTriangleInfluence
 {
@@ -78,8 +76,9 @@ float4x4 ClothToLocal;
 
 void GetClothBuffer(uint VertexIndex, uint ClothIndexOffset, out FVertexTriangleInfluence Influences[NUM_INFLUENCES_PER_VERTEX])
 {
-	const uint VertexOffset = NUM_INFLUENCES_PER_VERTEX * VertexIndex + ClothIndexOffset;
-	for (int i = 0; i < NUM_INFLUENCES_PER_VERTEX; ++i )
+	const uint VertexOffset = ClothNumInfluencesPerVertex * VertexIndex + ClothIndexOffset;
+	UNROLL_N(NUM_INFLUENCES_PER_VERTEX)
+	for (int i = 0; i < ClothNumInfluencesPerVertex; ++i )
 	{
 		const uint Offset = (VertexOffset + i) * NUM_FLOAT4S_PER_VERTEX_INFLUENCE;
 
@@ -351,7 +350,8 @@ void SkinCacheUpdateBatchCS(uint3 GroupID : SV_GroupID,
 	float SumWeights = 0.0;
 	float ClothWeight = 0.0;
 	
-    for (int i = 0; i < NUM_INFLUENCES_PER_VERTEX; ++i )
+	UNROLL_N(NUM_INFLUENCES_PER_VERTEX)
+    for (int i = 0; i < ClothNumInfluencesPerVertex; ++i )
     {
         const FVertexTriangleInfluence Influence = Unpacked.Influences[i];
 
@@ -370,7 +370,7 @@ void SkinCacheUpdateBatchCS(uint3 GroupID : SV_GroupID,
 	    	ClothWeight += ClothBlendWeight * (1.0f - (Influence.SourceMeshVertIndices.w / 65535.0f));
 
 			float Weight = 1.0f;
-			if ( NUM_INFLUENCES_PER_VERTEX > 1 )
+			if ( ClothNumInfluencesPerVertex > 1 )
 			{
 				// Weight is packed in the last coordinate
 				Weight = Influence.Weight;
@@ -422,7 +422,10 @@ void SkinCacheUpdateBatchCS(uint3 GroupID : SV_GroupID,
 		Intermediates.SimulatedPosition = float3(0, 0, 0);
 	}
 	
-	ClothWeight /= NUM_INFLUENCES_PER_VERTEX;
+	if (ClothNumInfluencesPerVertex > 1)
+	{
+		ClothWeight /= ClothNumInfluencesPerVertex;
+	}
 #endif // GPUSKIN_APEX_CLOTH
 #endif // GPUSKIN_MORPH_BLEND
 
diff --git a/Engine/Shaders/Private/GpuSkinVertexFactory.ush b/Engine/Shaders/Private/GpuSkinVertexFactory.ush
--- a/Engine/Shaders/Private/GpuSkinVertexFactory.ush
+++ b/Engine/Shaders/Private/GpuSkinVertexFactory.ush
@@ -32,10 +32,6 @@
 #define GPUSKIN_USE_BONES_SRV_BUFFER	0	// Use a Buffer<float4> if 1, otherwise a uniform buffer   
 #endif
 
-#ifndef GPUSKIN_MULTIPLE_CLOTH_INFLUENCES
-#define GPUSKIN_MULTIPLE_CLOTH_INFLUENCES 0
-#endif
-
 #define MAX_INFLUENCES_PER_STREAM	4
 
 #define FIXED_VERTEX_INDEX 0xFFFF
@@ -138,13 +134,11 @@ struct FVertexFactoryInput
 // .x = Draw Index Buffer offset, .y = Offset into Cloth Vertex Buffer
 uint2 GPUSkinApexClothStartIndexOffset;
 
-#if GPUSKIN_MULTIPLE_CLOTH_INFLUENCES
 // This must match NUM_INFLUENCES_PER_VERTEX in ClothingMeshUtils.cpp and GpuSkinCacheComputeShader.usf
+// It represents the maximum number of influences so perhaps should be renamed. Right now we are keeping the name matching for easy search.
 // TODO: Make this easier to change in without messing things up
 #define NUM_INFLUENCES_PER_VERTEX 5
-#else
-#define NUM_INFLUENCES_PER_VERTEX 1
-#endif
+uint ClothNumInfluencesPerVertex;
 
 struct FClothVertex
 {
@@ -179,9 +173,10 @@ void SetupClothVertex(uint VertexIndex, out FClothVertex ClothVertexInfluences[N
 {
 	const uint MinVertexIndex = GPUSkinApexClothStartIndexOffset.x;
 	const uint ClothIndexOffset = GPUSkinApexClothStartIndexOffset.y;
-    VertexIndex = NUM_INFLUENCES_PER_VERTEX * (VertexIndex - MinVertexIndex) + ClothIndexOffset;
+    VertexIndex = ClothNumInfluencesPerVertex * (VertexIndex - MinVertexIndex) + ClothIndexOffset;
 
-    for (int i = 0; i < NUM_INFLUENCES_PER_VERTEX; ++i )
+    UNROLL_N(NUM_INFLUENCES_PER_VERTEX)
+	for (int i = 0; i < ClothNumInfluencesPerVertex; ++i )
     {
         const uint Offset = (VertexIndex + i) * NUM_FLOAT4S_PER_VERTEX_INFLUENCE;
 
@@ -420,7 +415,8 @@ float4 ClothingPosition(FClothVertex ClothVertexInfluences[NUM_INFLUENCES_PER_VE
 	float SimulWeight = 0.0f; 
 	float3 AveragedSimulatedPosition = float3(0, 0, 0);
 
-    for (int i = 0; i < NUM_INFLUENCES_PER_VERTEX; ++i )
+	UNROLL_N(NUM_INFLUENCES_PER_VERTEX)
+    for (int i = 0; i < ClothNumInfluencesPerVertex; ++i )
     {
         const FClothVertex Influence = ClothVertexInfluences[i];
 
@@ -439,7 +435,7 @@ float4 ClothingPosition(FClothVertex ClothVertexInfluences[NUM_INFLUENCES_PER_VE
 	    	SimulWeight += GetClothSimulWeight(Influence.SimulIndices.w);
 	    	
 			float Weight = 1.0f;
-			if ( NUM_INFLUENCES_PER_VERTEX > 1 )
+			if ( ClothNumInfluencesPerVertex > 1 )
 			{
 				Weight = Influence.Weight;
 				SumWeights += Weight;
@@ -474,7 +470,10 @@ float4 ClothingPosition(FClothVertex ClothVertexInfluences[NUM_INFLUENCES_PER_VE
 		SimulWeight = 0.0f;			// Fall back to skinned vertex position by setting this lerp param to zero
 	}
 
-	SimulWeight /= NUM_INFLUENCES_PER_VERTEX;
+	if (ClothNumInfluencesPerVertex > 1)
+	{
+		SimulWeight /= ClothNumInfluencesPerVertex;
+	}
 
 	return float4(AveragedSimulatedPosition, SimulWeight);
 }
@@ -712,7 +711,7 @@ float3x3 SkinTangents( FVertexFactoryInput Input, FVertexFactoryIntermediates In
 	    float SumWeights = 0.0;
 		float SimulWeight = 0.0;
 
-        for (int i = 0; i < NUM_INFLUENCES_PER_VERTEX; ++i )
+        for (int i = 0; i < ClothNumInfluencesPerVertex; ++i )
         {
             const FClothVertex Influence = Intermediates.ClothVertexInfluences[i];
 
@@ -732,7 +731,7 @@ float3x3 SkinTangents( FVertexFactoryInput Input, FVertexFactoryIntermediates In
 	        	SimulWeight += GetClothSimulWeight(Influence.SimulIndices.w);
 	        	
 			    float Weight = 1.0f;
-			    if ( NUM_INFLUENCES_PER_VERTEX > 1 )
+			    if ( ClothNumInfluencesPerVertex > 1 )
 			    {
 				    // Weight is packed in the last coordinate
 				    Weight = Influence.Weight;
@@ -755,7 +754,10 @@ float3x3 SkinTangents( FVertexFactoryInput Input, FVertexFactoryIntermediates In
 	        }
         }
 
-		SimulWeight /= NUM_INFLUENCES_PER_VERTEX;
+		if (ClothNumInfluencesPerVertex > 1)
+		{
+			SimulWeight /= ClothNumInfluencesPerVertex;
+		}
 	
         if ( NumInfluences > 0 && SumWeights > 1e-4f )
         {
diff --git a/Engine/Source/Runtime/Engine/Private/GPUSkinCache.cpp b/Engine/Source/Runtime/Engine/Private/GPUSkinCache.cpp
--- a/Engine/Source/Runtime/Engine/Private/GPUSkinCache.cpp
+++ b/Engine/Source/Runtime/Engine/Private/GPUSkinCache.cpp
@@ -369,7 +369,7 @@ public:
         // cloth input
 		uint32 ClothBufferOffset = 0;
         float ClothBlendWeight = 0.0f;
-
+		uint32 ClothNumInfluencesPerVertex = 1;
         FMatrix44f ClothToLocal = FMatrix44f::Identity;
 
 		// triangle index buffer (input for the RecomputeSkinTangents, might need special index buffer unique to position and normal, not considering UV/vertex color)
@@ -527,20 +527,30 @@ public:
 			Data.MorphBufferOffset = Section->BaseVertexIndex;
 		}
 
-		if (ClothVertexBuffer && ClothVertexBuffer->GetClothIndexMapping().Num() > SectionIndex)
+		if (ClothVertexBuffer)
 		{
-			const FClothBufferIndexMapping& ClothBufferIndexMapping = ClothVertexBuffer->GetClothIndexMapping()[SectionIndex];
+			constexpr int32 ClothLODBias0 = 0;  // Use the same cloth LOD mapping (= 0 bias) to get the number of deformer weights
+			const uint32 NumWrapDeformerWeights = Data.Section->ClothMappingDataLODs.Num() ? Data.Section->ClothMappingDataLODs[ClothLODBias0].Num() : 0;
+			// NumInfluencesPerVertex should be a whole integer
+			check(NumWrapDeformerWeights % Data.NumVertices == 0);
+			Data.ClothNumInfluencesPerVertex = NumWrapDeformerWeights / Data.NumVertices;
+
+			if (ClothVertexBuffer->GetClothIndexMapping().Num() > SectionIndex)
+			{
+
+				const FClothBufferIndexMapping& ClothBufferIndexMapping = ClothVertexBuffer->GetClothIndexMapping()[SectionIndex];
 
-			check(SimData->LODIndex != INDEX_NONE && SimData->LODIndex <= LOD);
-			const uint32 ClothLODBias = (uint32)(LOD - SimData->LODIndex);
+				check(SimData->LODIndex != INDEX_NONE && SimData->LODIndex <= LOD);
+				const uint32 ClothLODBias = (uint32)(LOD - SimData->LODIndex);
 
-			const uint32 ClothBufferOffset = ClothBufferIndexMapping.MappingOffset + ClothBufferIndexMapping.LODBiasStride * ClothLODBias;
+				const uint32 ClothBufferOffset = ClothBufferIndexMapping.MappingOffset + ClothBufferIndexMapping.LODBiasStride * ClothLODBias;
 
-			// Set the buffer offset depending on whether enough deformer mapping data exists (RaytracingMinLOD/RaytracingLODBias/ClothLODBiasMode settings)
-			const uint32 NumInfluences = NumVertices ? ClothBufferIndexMapping.LODBiasStride / NumVertices : 1;
-			Data.ClothBufferOffset = (ClothBufferOffset + NumVertices * NumInfluences <= ClothVertexBuffer->GetNumVertices()) ?
-				ClothBufferOffset :                     // If the offset is valid, set the calculated LODBias offset
-				ClothBufferIndexMapping.MappingOffset;  // Otherwise fallback to a 0 ClothLODBias to prevent from reading pass the buffer (but still raytrace broken shadows/reflections/etc.)
+				// Set the buffer offset depending on whether enough deformer mapping data exists (RaytracingMinLOD/RaytracingLODBias/ClothLODBiasMode settings)
+				const uint32 NumInfluences = NumVertices ? ClothBufferIndexMapping.LODBiasStride / NumVertices : 1;
+				Data.ClothBufferOffset = (ClothBufferOffset + NumVertices * NumInfluences <= ClothVertexBuffer->GetNumVertices()) ?
+					ClothBufferOffset :                     // If the offset is valid, set the calculated LODBias offset
+					ClothBufferIndexMapping.MappingOffset;  // Otherwise fallback to a 0 ClothLODBias to prevent from reading pass the buffer (but still raytrace broken shadows/reflections/etc.)
+			}
 		}
 
 		//INC_DWORD_STAT(STAT_GPUSkinCache_TotalNumChunks);
@@ -657,6 +667,7 @@ public:
 		ClothBufferOffset.Bind(Initializer.ParameterMap, TEXT("ClothBufferOffset"));
 		ClothBlendWeight.Bind(Initializer.ParameterMap, TEXT("ClothBlendWeight"));
 		ClothToLocal.Bind(Initializer.ParameterMap, TEXT("ClothToLocal"));
+		ClothNumInfluencesPerVertex.Bind(Initializer.ParameterMap, TEXT("ClothNumInfluencesPerVertex"));
 	}
 
 	void SetParameters(
@@ -705,6 +716,7 @@ public:
 			SetShaderValue(RHICmdList, ShaderRHI, ClothBufferOffset, DispatchData.ClothBufferOffset);
 			SetShaderValue(RHICmdList, ShaderRHI, ClothBlendWeight, DispatchData.ClothBlendWeight);
 			SetShaderValue(RHICmdList, ShaderRHI, ClothToLocal, DispatchData.ClothToLocal);
+			SetShaderValue(RHICmdList, ShaderRHI, ClothNumInfluencesPerVertex, DispatchData.ClothNumInfluencesPerVertex);
 		}
 
 		SetShaderValue(RHICmdList, ShaderRHI, SkinCacheDebug, CVarGPUSkinCacheDebug.GetValueOnRenderThread());
@@ -751,18 +763,16 @@ private:
 	LAYOUT_FIELD(FShaderParameter, ClothBufferOffset)
 	LAYOUT_FIELD(FShaderParameter, ClothBlendWeight)
 	LAYOUT_FIELD(FShaderParameter, ClothToLocal)
-
+	LAYOUT_FIELD(FShaderParameter, ClothNumInfluencesPerVertex)
 };
 
 /** Compute shader that skins a batch of vertices. */
 // @param SkinType 0:normal, 1:with morph targets calculated outside the cache, 2: with cloth, 3:with morph target calculated insde the cache (not yet implemented)
 //        BoneInfluenceType 0:normal, 1:extra bone influences, 2:unlimited bone influences
 //        BoneIndex16 0: 8-bit indices, 1: 16-bit indices
-//        MultipleClothInfluences 0:single influence 1:multiple influences
 template <int Permutation>
 class TGPUSkinCacheCS : public FBaseGPUSkinCacheCS
 {
-	constexpr static bool bMultipleClothInfluences = (32 == (Permutation & 32));
 	constexpr static bool bBoneIndex16 = (16 == (Permutation & 16));
 	constexpr static bool bUnlimitedBoneInfluence = (8 == (Permutation & 12));
 	constexpr static bool bUseExtraBoneInfluencesT = (4 == (Permutation & 12));
@@ -785,12 +795,10 @@ public:
 		const uint32 UseExtraBoneInfluences = bUseExtraBoneInfluencesT;
 		const uint32 MorphBlend = bMorphBlend;
 		const uint32 ApexCloth = bApexCloth;
-		const uint32 MultipleClothInfluences = bMultipleClothInfluences;
 		OutEnvironment.SetDefine(TEXT("GPUSKIN_UNLIMITED_BONE_INFLUENCE"), UnlimitedBoneInfluence);
 		OutEnvironment.SetDefine(TEXT("GPUSKIN_USE_EXTRA_INFLUENCES"), UseExtraBoneInfluences);
 		OutEnvironment.SetDefine(TEXT("GPUSKIN_MORPH_BLEND"), MorphBlend);
 		OutEnvironment.SetDefine(TEXT("GPUSKIN_APEX_CLOTH"), ApexCloth);
-		OutEnvironment.SetDefine(TEXT("GPUSKIN_MULTIPLE_CLOTH_INFLUENCES"), MultipleClothInfluences);
 		OutEnvironment.SetDefine(TEXT("GPUSKIN_RWBUFFER_OFFSET_TANGENT_X"), FGPUSkinCache::RWTangentXOffsetInFloats);
 		OutEnvironment.SetDefine(TEXT("GPUSKIN_RWBUFFER_OFFSET_TANGENT_Z"), FGPUSkinCache::RWTangentZOffsetInFloats);
 		OutEnvironment.SetDefine(TEXT("GPUSKIN_BONE_INDEX_UINT16"), BoneIndex16);
@@ -821,12 +829,6 @@ IMPLEMENT_SHADER_TYPE(template<>, TGPUSkinCacheCS<18>, TEXT("/Engine/Private/Gpu
 IMPLEMENT_SHADER_TYPE(template<>, TGPUSkinCacheCS<20>, TEXT("/Engine/Private/GpuSkinCacheComputeShader.usf"), TEXT("SkinCacheUpdateBatchCS"), SF_Compute);  // 16bit_1, BoneInfluenceType_1, SkinType_0 
 IMPLEMENT_SHADER_TYPE(template<>, TGPUSkinCacheCS<21>, TEXT("/Engine/Private/GpuSkinCacheComputeShader.usf"), TEXT("SkinCacheUpdateBatchCS"), SF_Compute);  // 16bit_1, BoneInfluenceType_1, SkinType_1 
 IMPLEMENT_SHADER_TYPE(template<>, TGPUSkinCacheCS<22>, TEXT("/Engine/Private/GpuSkinCacheComputeShader.usf"), TEXT("SkinCacheUpdateBatchCS"), SF_Compute);  // 16bit_1, BoneInfluenceType_1, SkinType_2 
-// Multi-influences for cloth:
-IMPLEMENT_SHADER_TYPE(template<>, TGPUSkinCacheCS<34>, TEXT("/Engine/Private/GpuSkinCacheComputeShader.usf"), TEXT("SkinCacheUpdateBatchCS"), SF_Compute);	// 16bit_0, BoneInfluenceType_0, SkinType_2, MultipleClothInfluences_1
-IMPLEMENT_SHADER_TYPE(template<>, TGPUSkinCacheCS<38>, TEXT("/Engine/Private/GpuSkinCacheComputeShader.usf"), TEXT("SkinCacheUpdateBatchCS"), SF_Compute);	// 16bit_0, BoneInfluenceType_1, SkinType_2, MultipleClothInfluences_1 
-IMPLEMENT_SHADER_TYPE(template<>, TGPUSkinCacheCS<42>, TEXT("/Engine/Private/GpuSkinCacheComputeShader.usf"), TEXT("SkinCacheUpdateBatchCS"), SF_Compute);	// 16bit_0, BoneInfluenceType_2, SkinType_2, MultipleClothInfluences_1
-IMPLEMENT_SHADER_TYPE(template<>, TGPUSkinCacheCS<50>, TEXT("/Engine/Private/GpuSkinCacheComputeShader.usf"), TEXT("SkinCacheUpdateBatchCS"), SF_Compute);  // 16bit_1, BoneInfluenceType_0, SkinType_2, MultipleClothInfluences_1 
-IMPLEMENT_SHADER_TYPE(template<>, TGPUSkinCacheCS<54>, TEXT("/Engine/Private/GpuSkinCacheComputeShader.usf"), TEXT("SkinCacheUpdateBatchCS"), SF_Compute);  // 16bit_1, BoneInfluenceType_1, SkinType_2, MultipleClothInfluences_1 
 
 FGPUSkinCache::FGPUSkinCache(ERHIFeatureLevel::Type InFeatureLevel, bool bInRequiresMemoryLimit, UWorld* InWorld)
 	: UsedMemoryInBytes(0)
@@ -1850,10 +1852,6 @@ void FGPUSkinCache::DispatchUpdateSkinning(FRHICommandListImmediate& RHICmdList,
 	FGPUBaseSkinVertexFactory::FShaderDataType& ShaderData = DispatchData.SourceVertexFactory->GetShaderData();
 	const FString RayTracingTag = (Entry->Mode == EGPUSkinCacheEntryMode::RayTracing ? TEXT("[RT]") : TEXT(""));
 	
-	constexpr int32 ClothLODBias = 0;  // Use the same cloth LOD mapping (= 0 bias) to get the number of deformer weights
-	const uint32 NumWrapDeformerWeights = DispatchData.Section->ClothMappingDataLODs.Num() ? DispatchData.Section->ClothMappingDataLODs[ClothLODBias].Num(): 0;
-	const bool bMultipleWrapDeformerInfluences = (DispatchData.NumVertices < NumWrapDeformerWeights);
-
 	SCOPED_DRAW_EVENTF(RHICmdList, SkinCacheDispatch,
 		TEXT("%sSkinning%d%d%d Mesh=%s LOD=%d Chunk=%d InStreamStart=%d OutStart=%d Vert=%d Morph=%d/%d"),
 		*RayTracingTag, (int32)Entry->bUse16BitBoneIndex, (int32)Entry->BoneInfluenceType, DispatchData.SkinType, *GetSkeletalMeshObjectName(Entry->GPUSkin), Entry->LOD,
@@ -1875,13 +1873,6 @@ void FGPUSkinCache::DispatchUpdateSkinning(FRHICommandListImmediate& RHICmdList,
 	TShaderMapRef<TGPUSkinCacheCS<21>>  SkinCacheCS111(GlobalShaderMap);	// 16bit_1, BoneInfluenceType_1, SkinType_1
 	TShaderMapRef<TGPUSkinCacheCS<22>>  SkinCacheCS112(GlobalShaderMap);	// 16bit_1, BoneInfluenceType_1, SkinType_2
 
-	// Multi-influences for cloth:
-	TShaderMapRef<TGPUSkinCacheCS<34>>  SkinCacheCS0021(GlobalShaderMap);	// 16bit_0, BoneInfluenceType_0, SkinType_2, MultipleClothInfluences_1
-	TShaderMapRef<TGPUSkinCacheCS<38>>  SkinCacheCS0121(GlobalShaderMap);	// 16bit_0, BoneInfluenceType_1, SkinType_2, MultipleClothInfluences_1
-	TShaderMapRef<TGPUSkinCacheCS<42>>  SkinCacheCS0221(GlobalShaderMap);	// 16bit_0, BoneInfluenceType_2, SkinType_2, MultipleClothInfluences_1
-	TShaderMapRef<TGPUSkinCacheCS<50>>  SkinCacheCS1021(GlobalShaderMap);	// 16bit_1, BoneInfluenceType_0, SkinType_2, MultipleClothInfluences_1
-	TShaderMapRef<TGPUSkinCacheCS<54>>  SkinCacheCS1121(GlobalShaderMap);	// 16bit_1, BoneInfluenceType_1, SkinType_2, MultipleClothInfluences_1
-
 	TShaderRef<FBaseGPUSkinCacheCS> Shader;
 	switch (DispatchData.SkinType)
 	{
@@ -1918,41 +1909,20 @@ void FGPUSkinCache::DispatchUpdateSkinning(FRHICommandListImmediate& RHICmdList,
 		}
 		break;
 	case 2:
-		if (bMultipleWrapDeformerInfluences)
+		// Cloth
+		if (Entry->BoneInfluenceType == 0)
 		{
-			// Multiple influences for cloth skinning
-			if (Entry->BoneInfluenceType == 0)
-			{
-				if (Entry->bUse16BitBoneIndex) Shader = SkinCacheCS1021;
-				else Shader = SkinCacheCS0021;
-			}
-			else if (Entry->BoneInfluenceType == 1)
-			{
-				if (Entry->bUse16BitBoneIndex) Shader = SkinCacheCS1121;
-				else Shader = SkinCacheCS0121;
-			}
-			else
-			{
-				Shader = SkinCacheCS0221;
-			}
+			if (Entry->bUse16BitBoneIndex) Shader = SkinCacheCS102;
+			else Shader = SkinCacheCS002;
+		}
+		else if (Entry->BoneInfluenceType == 1)
+		{
+			if (Entry->bUse16BitBoneIndex) Shader = SkinCacheCS112;
+			else Shader = SkinCacheCS012;
 		}
 		else
 		{
-			// Single influence for cloth skinning
-			if (Entry->BoneInfluenceType == 0)
-			{
-				if (Entry->bUse16BitBoneIndex) Shader = SkinCacheCS102;
-				else Shader = SkinCacheCS002;
-			}
-			else if (Entry->BoneInfluenceType == 1)
-			{
-				if (Entry->bUse16BitBoneIndex) Shader = SkinCacheCS112;
-				else Shader = SkinCacheCS012;
-			}
-			else
-			{
-				Shader = SkinCacheCS022;
-			}
+			Shader = SkinCacheCS022;
 		}
 		break;
 	default:
diff --git a/Engine/Source/Runtime/Engine/Private/GPUSkinVertexFactory.cpp b/Engine/Source/Runtime/Engine/Private/GPUSkinVertexFactory.cpp
--- a/Engine/Source/Runtime/Engine/Private/GPUSkinVertexFactory.cpp
+++ b/Engine/Source/Runtime/Engine/Private/GPUSkinVertexFactory.cpp
@@ -1064,6 +1064,7 @@ public:
 		ClothBlendWeightParameter.Bind(ParameterMap, TEXT("ClothBlendWeight"));
 		GPUSkinApexClothParameter.Bind(ParameterMap, TEXT("GPUSkinApexCloth"));
 		GPUSkinApexClothStartIndexOffsetParameter.Bind(ParameterMap, TEXT("GPUSkinApexClothStartIndexOffset"));
+		ClothNumInfluencesPerVertexParameter.Bind(ParameterMap, TEXT("ClothNumInfluencesPerVertex"));
 	}
 	
 	void GetElementShaderBindings(
@@ -1092,6 +1093,7 @@ public:
 		ShaderBindings.Add(ClothSimulVertsPositionsNormalsParameter, ClothShaderData.GetClothBufferForReading(false, FrameNumber).VertexBufferSRV);
 		ShaderBindings.Add(ClothToLocalParameter, ClothShaderData.GetClothToLocalForReading(false, FrameNumber));
 		ShaderBindings.Add(ClothBlendWeightParameter,ClothShaderData.ClothBlendWeight);
+		ShaderBindings.Add(ClothNumInfluencesPerVertexParameter, ClothShaderData.NumInfluencesPerVertex);
 
 		// Mobile doesn't support motion blur, no need to feed the previous frame cloth data
 		const EShaderPlatform ShaderPlatform = GetFeatureLevelShaderPlatform(FeatureLevel);
@@ -1117,7 +1119,7 @@ protected:
 	LAYOUT_FIELD(FShaderParameter, ClothBlendWeightParameter);
 	LAYOUT_FIELD(FShaderResourceParameter, GPUSkinApexClothParameter);
 	LAYOUT_FIELD(FShaderParameter, GPUSkinApexClothStartIndexOffsetParameter);
-	
+	LAYOUT_FIELD(FShaderParameter, ClothNumInfluencesPerVertexParameter);
 };
 
 IMPLEMENT_TYPE_LAYOUT(TGPUSkinAPEXClothVertexFactoryShaderParameters);
@@ -1219,7 +1221,6 @@ void TGPUSkinAPEXClothVertexFactory<BoneInfluenceType>::ModifyCompilationEnviron
 {
 	Super::ModifyCompilationEnvironment(Parameters, OutEnvironment);
 	OutEnvironment.SetDefine(TEXT("GPUSKIN_APEX_CLOTH"),TEXT("1"));
-	OutEnvironment.SetDefine(TEXT("GPUSKIN_MULTIPLE_CLOTH_INFLUENCES"), TEXT("0"));
 	
 	// Mobile doesn't support motion blur, don't use previous frame data.
 	const bool bIsMobile = IsMobilePlatform(Parameters.Platform);
@@ -1274,28 +1275,5 @@ IMPLEMENT_GPUSKINNING_VERTEX_FACTORY_TYPE(TGPUSkinAPEXClothVertexFactory, "/Engi
 	| EVertexFactoryFlags::SupportsPSOPrecaching
 );
 
-
-template<GPUSkinBoneInfluenceType BoneInfluenceType>
-void TMultipleInfluenceClothVertexFactory<BoneInfluenceType>::ModifyCompilationEnvironment(
-	const FVertexFactoryShaderPermutationParameters& Parameters, FShaderCompilerEnvironment& OutEnvironment)
-{
-	Super::ModifyCompilationEnvironment(Parameters, OutEnvironment);
-	OutEnvironment.SetDefine(TEXT("GPUSKIN_MULTIPLE_CLOTH_INFLUENCES"), TEXT("1"));
-}
-
-
-template <GPUSkinBoneInfluenceType BoneInfluenceType>
-bool TMultipleInfluenceClothVertexFactory<BoneInfluenceType>::ShouldCompilePermutation(const FVertexFactoryShaderPermutationParameters& Parameters)
-{
-	return Super::ShouldCompilePermutation(Parameters);
-}
-
-IMPLEMENT_GPUSKINNING_VERTEX_FACTORY_PARAMETER_TYPE(TMultipleInfluenceClothVertexFactory, SF_Vertex, TGPUSkinAPEXClothVertexFactoryShaderParameters);
-IMPLEMENT_GPUSKINNING_VERTEX_FACTORY_TYPE(TMultipleInfluenceClothVertexFactory, "/Engine/Private/GpuSkinVertexFactory.ush",
-	  EVertexFactoryFlags::UsedWithMaterials
-	| EVertexFactoryFlags::SupportsDynamicLighting
-	| EVertexFactoryFlags::SupportsPSOPrecaching
-);
-
 #undef IMPLEMENT_GPUSKINNING_VERTEX_FACTORY_PARAMETER_TYPE
 #undef IMPLEMENT_GPUSKINNING_VERTEX_FACTORY_TYPE
diff --git a/Engine/Source/Runtime/Engine/Private/SkeletalRenderGPUSkin.cpp b/Engine/Source/Runtime/Engine/Private/SkeletalRenderGPUSkin.cpp
--- a/Engine/Source/Runtime/Engine/Private/SkeletalRenderGPUSkin.cpp
+++ b/Engine/Source/Runtime/Engine/Private/SkeletalRenderGPUSkin.cpp
@@ -1889,32 +1889,18 @@ static void UpdateVertexFactoryMorph(TArray<TUniquePtr<FGPUBaseSkinVertexFactory
 static void CreateVertexFactoryCloth(TArray<TUniquePtr<FGPUBaseSkinAPEXClothVertexFactory>>& VertexFactories,
 						 const FSkeletalMeshObjectGPUSkin::FVertexFactoryBuffers& InVertexBuffers,
 						 ERHIFeatureLevel::Type FeatureLevel,
-						 bool bUseMultipleInfluences
+						 uint32 NumInfluencesPerVertex
 						 )
 {
 	FGPUBaseSkinAPEXClothVertexFactory* ClothVertexFactory = nullptr;
 	GPUSkinBoneInfluenceType BoneInfluenceType = InVertexBuffers.SkinWeightVertexBuffer->GetBoneInfluenceType();
 	if (BoneInfluenceType == GPUSkinBoneInfluenceType::DefaultBoneInfluence)
 	{
-		if (bUseMultipleInfluences)
-		{
-			ClothVertexFactory = new TMultipleInfluenceClothVertexFactory<GPUSkinBoneInfluenceType::DefaultBoneInfluence>(FeatureLevel, InVertexBuffers.NumVertices);
-		}
-		else
-		{
-			ClothVertexFactory = new TGPUSkinAPEXClothVertexFactory<GPUSkinBoneInfluenceType::DefaultBoneInfluence>(FeatureLevel, InVertexBuffers.NumVertices);
-		}
+		ClothVertexFactory = new TGPUSkinAPEXClothVertexFactory<GPUSkinBoneInfluenceType::DefaultBoneInfluence>(FeatureLevel, InVertexBuffers.NumVertices, NumInfluencesPerVertex);
 	}
 	else
 	{
-		if (bUseMultipleInfluences)
-		{
-			ClothVertexFactory = new TMultipleInfluenceClothVertexFactory<GPUSkinBoneInfluenceType::UnlimitedBoneInfluence>(FeatureLevel, InVertexBuffers.NumVertices);
-		}
-		else
-		{
-			ClothVertexFactory = new TGPUSkinAPEXClothVertexFactory<GPUSkinBoneInfluenceType::UnlimitedBoneInfluence>(FeatureLevel, InVertexBuffers.NumVertices);
-		}
+		ClothVertexFactory = new TGPUSkinAPEXClothVertexFactory<GPUSkinBoneInfluenceType::UnlimitedBoneInfluence>(FeatureLevel, InVertexBuffers.NumVertices, NumInfluencesPerVertex);
 	}
 	VertexFactories.Add(TUniquePtr<FGPUBaseSkinAPEXClothVertexFactory>(ClothVertexFactory));
 
@@ -2061,8 +2047,10 @@ void FSkeletalMeshObjectGPUSkin::FVertexFactoryData::InitAPEXClothVertexFactorie
 			constexpr int32 ClothLODBias = 0;
 			const uint32 NumClothWeights = Section.ClothMappingDataLODs.Num() ? Section.ClothMappingDataLODs[ClothLODBias].Num(): 0;
 			const uint32 NumPositionVertices = Section.NumVertices;
-			const bool bUseMultipleInfluences = (NumClothWeights > NumPositionVertices);
-			CreateVertexFactoryCloth(ClothVertexFactories, VertexBuffers, InFeatureLevel, bUseMultipleInfluences);
+			// NumInfluencesPerVertex should be a whole integer
+			check(NumClothWeights % NumPositionVertices == 0);
+			const uint32 NumInfluencesPerVertex = NumClothWeights / NumPositionVertices;
+			CreateVertexFactoryCloth(ClothVertexFactories, VertexBuffers, InFeatureLevel, NumInfluencesPerVertex);
 		}
 		else
 		{
diff --git a/Engine/Source/Runtime/Engine/Public/GPUSkinVertexFactory.h b/Engine/Source/Runtime/Engine/Public/GPUSkinVertexFactory.h
--- a/Engine/Source/Runtime/Engine/Public/GPUSkinVertexFactory.h
+++ b/Engine/Source/Runtime/Engine/Public/GPUSkinVertexFactory.h
@@ -604,7 +604,6 @@ public:
 	struct ClothShaderType
 	{
 		ClothShaderType()
-			: ClothBlendWeight(1.0f)
 		{
 			Reset();
 		}
@@ -703,7 +702,8 @@ public:
 		 * weight to blend between simulated positions and key-framed poses
 		 * if ClothBlendWeight is 1.0, it shows only simulated positions and if it is 0.0, it shows only key-framed animation
 		 */
-		float ClothBlendWeight;
+		float ClothBlendWeight = 1.0f;
+		uint32 NumInfluencesPerVertex = 1;
 
 	private:
 		// fallback for ClothSimulPositionNormalBuffer if the shadermodel doesn't allow it
@@ -801,6 +801,11 @@ public:
 		}
 	};
 
+	FGPUBaseSkinAPEXClothVertexFactory(uint32 InNumInfluencesPerVertex)
+	{
+		ClothShaderData.NumInfluencesPerVertex = InNumInfluencesPerVertex;
+	}
+
 	virtual ~FGPUBaseSkinAPEXClothVertexFactory() {}
 
 	/** accessor */
@@ -881,8 +886,9 @@ public:
 	 *
 	 * @param	InBoneMatrices	Reference to shared bone matrices array.
 	 */
-	TGPUSkinAPEXClothVertexFactory(ERHIFeatureLevel::Type InFeatureLevel, uint32 InNumVertices)
-		: TGPUSkinVertexFactory<BoneInfluenceType>(InFeatureLevel, InNumVertices)
+	TGPUSkinAPEXClothVertexFactory(ERHIFeatureLevel::Type InFeatureLevel, uint32 InNumVertices, uint32 InNumInfluencesPerVertex)
+		: FGPUBaseSkinAPEXClothVertexFactory(InNumInfluencesPerVertex)
+		, TGPUSkinVertexFactory<BoneInfluenceType>(InFeatureLevel, InNumVertices)
 	{}
 
 	static void ModifyCompilationEnvironment(const FVertexFactoryShaderPermutationParameters& Parameters, FShaderCompilerEnvironment& OutEnvironment);
@@ -929,20 +935,3 @@ protected:
 	FGPUSkinAPEXClothDataType* ClothDataPtr = nullptr;
 };
 
-
-template<GPUSkinBoneInfluenceType BoneInfluenceType>
-class TMultipleInfluenceClothVertexFactory : public TGPUSkinAPEXClothVertexFactory<BoneInfluenceType>
-{
-	DECLARE_VERTEX_FACTORY_TYPE(TMultipleInfluenceClothVertexFactory<BoneInfluenceType>);
-
-	typedef TGPUSkinAPEXClothVertexFactory<BoneInfluenceType> Super;
-
-public:
-
-	TMultipleInfluenceClothVertexFactory(ERHIFeatureLevel::Type InFeatureLevel, uint32 InNumVertices)
-		: Super(InFeatureLevel, InNumVertices)
-	{}
-
-	static void ModifyCompilationEnvironment(const FVertexFactoryShaderPermutationParameters& Parameters, FShaderCompilerEnvironment& OutEnvironment);
-	static bool ShouldCompilePermutation(const FVertexFactoryShaderPermutationParameters& Parameters);
-};
