When console doc generation fails, report the failure instead of reporting success.

diff --git a/Engine/Source/Runtime/Core/Private/Misc/Parse.cpp b/Engine/Source/Runtime/Core/Private/Misc/Parse.cpp
--- a/Engine/Source/Runtime/Core/Private/Misc/Parse.cpp
+++ b/Engine/Source/Runtime/Core/Private/Misc/Parse.cpp
@@ -79,7 +79,7 @@ public:
 	}
 };
 
-void ConsoleCommandLibrary_DumpLibrary(UWorld* InWorld, FExec& SubSystem, const FString& Pattern, FOutputDevice& Ar)
+bool ConsoleCommandLibrary_DumpLibrary(UWorld* InWorld, FExec& SubSystem, const FString& Pattern, FOutputDevice& Ar)
 {
 	ConsoleCommandLibrary LocalConsoleCommandLibrary(Pattern);
 
@@ -110,9 +110,11 @@ void ConsoleCommandLibrary_DumpLibrary(UWorld* InWorld, FExec& SubSystem, const
 		Ar.Logf(TEXT("ERROR: The function was supposed to only find matching commands but not have any side effect."));
 		Ar.Logf(TEXT("However Exec() returned true which means we either executed a command or the command parsing returned true where it shouldn't."));
 	}
+
+	return true;
 }
 
-void ConsoleCommandLibrary_DumpLibraryHTML(UWorld* InWorld, FExec& SubSystem, const FString& OutPath)
+bool ConsoleCommandLibrary_DumpLibraryHTML(UWorld* InWorld, FExec& SubSystem, const FString& OutPath)
 {
 	const FString& Pattern(TEXT("*"));
 	ConsoleCommandLibrary LocalConsoleCommandLibrary(Pattern);
@@ -205,9 +207,12 @@ void ConsoleCommandLibrary_DumpLibraryHTML(UWorld* InWorld, FExec& SubSystem, co
 
 			delete File;
 			File = 0;
+
+			return true;
 		}
 	}
 
+	return false;
 /*
 	// the pattern (e.g. Motion*) should not really trigger the execution
 	if(bExecuted)
diff --git a/Engine/Source/Runtime/Core/Public/Misc/Parse.h b/Engine/Source/Runtime/Core/Public/Misc/Parse.h
--- a/Engine/Source/Runtime/Core/Public/Misc/Parse.h
+++ b/Engine/Source/Runtime/Core/Public/Misc/Parse.h
@@ -126,7 +126,7 @@ struct CORE_API FParse
 
 #if !UE_BUILD_SHIPPING
 /** Needed for the console command "DumpConsoleCommands" */
-CORE_API void ConsoleCommandLibrary_DumpLibrary(class UWorld* InWorld, FExec& SubSystem, const FString& Pattern, FOutputDevice& Ar);
+CORE_API bool ConsoleCommandLibrary_DumpLibrary(class UWorld* InWorld, FExec& SubSystem, const FString& Pattern, FOutputDevice& Ar);
 /** Needed for the console command "Help" */
-CORE_API void ConsoleCommandLibrary_DumpLibraryHTML(class UWorld* InWorld, FExec& SubSystem, const FString& OutPath);
+CORE_API bool ConsoleCommandLibrary_DumpLibraryHTML(class UWorld* InWorld, FExec& SubSystem, const FString& OutPath);
 #endif // !UE_BUILD_SHIPPING
diff --git a/Engine/Source/Runtime/Engine/Private/UnrealEngine.cpp b/Engine/Source/Runtime/Engine/Private/UnrealEngine.cpp
--- a/Engine/Source/Runtime/Engine/Private/UnrealEngine.cpp
+++ b/Engine/Source/Runtime/Engine/Private/UnrealEngine.cpp
@@ -5303,34 +5303,54 @@ static void DumpHelp(UWorld* InWorld)
 	UE_LOG(LogEngine, Display, TEXT("To browse console variables open this: '%s'"), *FilePath);
 	UE_LOG(LogEngine, Display, TEXT(" "));
 
-	ConsoleCommandLibrary_DumpLibraryHTML(InWorld, *GEngine, FilePath);
+	const bool bSuccess = ConsoleCommandLibrary_DumpLibraryHTML(InWorld, *GEngine, FilePath);
 
 	// Notification in editor
 #if WITH_EDITOR
+	if (bSuccess && FPaths::FileExists(FilePath))
 	{
-		const FText Message = NSLOCTEXT("UnrealEd", "ConsoleHelpExported", "ConsoleHelp.html was saved as");
+		const FText Message = NSLOCTEXT("UnrealEd", "ConsoleHelpExported", "Console help file saved at:");
 		FNotificationInfo Info(Message);
 		Info.bFireAndForget = true;
 		Info.ExpireDuration = 5.0f;
 		Info.bUseSuccessFailIcons = false;
 		Info.bUseLargeFont = false;
 
-		const FString HyperLinkText = FPaths::ConvertRelativePathToFull(FilePath);
-		Info.Hyperlink = FSimpleDelegate::CreateStatic([](FString SourceFilePath) 
+		const FString HyperLinkPath = FPaths::ConvertRelativePathToFull(FilePath);
+		Info.SubText = FText::FromString(HyperLinkPath);
+		Info.Hyperlink = FSimpleDelegate::CreateStatic([](FString SourceFilePath)
 		{
 			// open folder, you can choose the browser yourself
 			FPlatformProcess::ExploreFolder(*(FPaths::GetPath(SourceFilePath)));
-		}, HyperLinkText);
-		Info.HyperlinkText = FText::FromString(HyperLinkText);
+		}, HyperLinkPath);
+		Info.HyperlinkText = NSLOCTEXT("UnrealEd", "ConsoleHelpFolderLink", "Open file location...");
+		Info.WidthOverride = FOptionalSize();
 
 		FSlateNotificationManager::Get().AddNotification(Info);
 
 		// Always try to open the help file on Windows (including in -game, etc...)
 #if PLATFORM_WINDOWS
-		const FString LaunchableURL = FString(TEXT("file://")) + HyperLinkText;
+		const FString LaunchableURL = FString(TEXT("file://")) + HyperLinkPath;
 		FPlatformProcess::LaunchURL(*LaunchableURL, nullptr, nullptr);
 #endif
 	}
+	else
+	{
+		const FText Message = NSLOCTEXT("UnrealEd", "ConsoleHelpExportFailed", "Console help file generation failed.");
+		FNotificationInfo Info(Message);
+
+		const FString DocTemplateFile = FPaths::Combine(FPaths::EngineDir(), "Documentation/Extras/ConsoleHelpTemplate.html");
+		if (!FPaths::FileExists(DocTemplateFile))
+		{
+				Info.SubText = NSLOCTEXT("UnrealEd", "ConsoleHelpExportFailedTip", "The template HTML file could not be found in the Engine/Documentation/Extras folder.\n\nIf you use a version control system or Unreal Game Sync (UGS), check that this folder is not excluded by your sync settings.");
+		}
+		Info.bFireAndForget = true;
+		Info.ExpireDuration = 10.0f;
+		Info.bUseSuccessFailIcons = false;
+		Info.bUseLargeFont = false;
+
+		FSlateNotificationManager::Get().AddNotification(Info);
+	}
 #endif// WITH_EDITOR
 }
 static FAutoConsoleCommandWithWorld GConsoleCommandHelp(
