PR #9255: Fixed FAnimNotifyEvent::IsBlueprintNotify() to support metadata filtering by AnimNotifyList (Contributed by [REDACTED_PERSON])

diff --git a/Engine/Source/Runtime/Engine/Private/Animation/AnimTypes.cpp b/Engine/Source/Runtime/Engine/Private/Animation/AnimTypes.cpp
--- a/Engine/Source/Runtime/Engine/Private/Animation/AnimTypes.cpp
+++ b/Engine/Source/Runtime/Engine/Private/Animation/AnimTypes.cpp
@@ -110,6 +110,13 @@ bool FAnimNotifyEvent::IsBranchingPoint() const
 	return GetLinkedMontage() && ((MontageTickType == EMontageNotifyTickType::BranchingPoint) || (Notify && Notify->bIsNativeBranchingPoint) || (NotifyStateClass && NotifyStateClass->bIsNativeBranchingPoint));
 }
 
+bool FAnimNotifyEvent::IsBlueprintNotify() const
+{
+	return (
+		(Notify && Notify->GetClass() && !Notify->GetClass()->IsNative()) ||
+		(NotifyStateClass && NotifyStateClass->GetClass() && !NotifyStateClass->GetClass()->IsNative())
+		);
+}
 
 void FAnimNotifyEvent::SetTime(float NewTime, EAnimLinkMethod::Type ReferenceFrame /*= EAnimLinkMethod::Absolute*/)
 {
diff --git a/Engine/Source/Runtime/Engine/Public/Animation/AnimTypes.h b/Engine/Source/Runtime/Engine/Public/Animation/AnimTypes.h
--- a/Engine/Source/Runtime/Engine/Public/Animation/AnimTypes.h
+++ b/Engine/Source/Runtime/Engine/Public/Animation/AnimTypes.h
@@ -423,10 +423,7 @@ public:
 	ENGINE_API bool IsBranchingPoint() const;
 
 	/** Returns true if this is blueprint derived notifies **/
-	bool IsBlueprintNotify() const
-	{
-		return Notify != nullptr || NotifyStateClass != nullptr;
-	}
+	bool IsBlueprintNotify() const;
 
 	bool operator ==(const FAnimNotifyEvent& Other) const
 	{
