Fix invalid Single/Simple categorization when a weight operator is applied to a BSDF.

diff --git a/Engine/Source/Runtime/Engine/Private/Materials/HLSLMaterialTranslator.cpp b/Engine/Source/Runtime/Engine/Private/Materials/HLSLMaterialTranslator.cpp
--- a/Engine/Source/Runtime/Engine/Private/Materials/HLSLMaterialTranslator.cpp
+++ b/Engine/Source/Runtime/Engine/Private/Materials/HLSLMaterialTranslator.cpp
@@ -10223,6 +10223,13 @@ bool FHLSLMaterialTranslator::StrataGenerateDerivedMaterialOperatorData()
 					}
 					break;
 				}
+				case STRATA_OPERATOR_WEIGHT:
+				{
+					// If a BSDF modified by a weight operator, its weight will be < 1.0f, and it won't be a single material anymore
+					bStrataMaterialIsSimple = false;
+					bStrataMaterialIsSingle = false;
+					break;
+				}
 				}
 			}
 			bStrataMaterialIsSingle = bStrataMaterialIsSingle && !bStrataMaterialIsSimple;
