Interchange: Fix the custom LOD import issues
- Static mesh editor LOD number not updated when importing custom LOD asynchronously
- Static mesh reimport LOD with new file need to dirty the ddc key
- Gltf translator must set a valid local transform on the root scene node

#rnx

diff --git a/Engine/Plugins/Interchange/Runtime/Source/Import/Private/Gltf/InterchangeGltfTranslator.cpp b/Engine/Plugins/Interchange/Runtime/Source/Import/Private/Gltf/InterchangeGltfTranslator.cpp
--- a/Engine/Plugins/Interchange/Runtime/Source/Import/Private/Gltf/InterchangeGltfTranslator.cpp
+++ b/Engine/Plugins/Interchange/Runtime/Source/Import/Private/Gltf/InterchangeGltfTranslator.cpp
@@ -1038,8 +1038,10 @@ bool UInterchangeGltfTranslator::Translate( UInterchangeBaseNodeContainer& NodeC
 			SceneNode->InitializeNode( SceneNodeUid, SceneName, EInterchangeNodeContainerType::TranslatedScene );
 			NodeContainer.AddNode( SceneNode );
 
-			TArray<int32> SkinnedMeshNodes;
+			//All scene node should have a valid local transform
+			SceneNode->SetCustomLocalTransform(&NodeContainer, FTransform::Identity);
 
+			TArray<int32> SkinnedMeshNodes;
 			for ( const int32 NodeIndex : GltfScene.Nodes )
 			{
 				if ( GltfAsset.Nodes.IsValidIndex( NodeIndex ) )
diff --git a/Engine/Source/Editor/StaticMeshEditor/Private/StaticMeshEditorTools.cpp b/Engine/Source/Editor/StaticMeshEditor/Private/StaticMeshEditorTools.cpp
--- a/Engine/Source/Editor/StaticMeshEditor/Private/StaticMeshEditorTools.cpp
+++ b/Engine/Source/Editor/StaticMeshEditor/Private/StaticMeshEditorTools.cpp
@@ -3277,6 +3277,15 @@ FLevelOfDetailSettingsLayout::FLevelOfDetailSettingsLayout( FStaticMeshEditor& I
 	LODCount = StaticMeshEditor.GetStaticMesh()->GetNumLODs();
 
 	UpdateLODNames();
+
+	OnAssetPostLODImportDelegateHandle = GEditor->GetEditorSubsystem<UImportSubsystem>()->OnAssetPostLODImport.AddLambda([this](UObject* InObject, int32 InLODIndex)
+		{
+			UStaticMesh* StaticMesh = StaticMeshEditor.GetStaticMesh();
+			if (InObject == StaticMesh)
+			{
+				StaticMeshEditor.RefreshTool();
+			}
+		});
 }
 
 /** Returns true if automatic mesh reduction is available. */
@@ -3871,6 +3880,7 @@ void FLevelOfDetailSettingsLayout::AddLODLevelCategories( IDetailLayoutBuilder&
 
 FLevelOfDetailSettingsLayout::~FLevelOfDetailSettingsLayout()
 {
+	GEditor->GetEditorSubsystem<UImportSubsystem>()->OnAssetPostLODImport.Remove(OnAssetPostLODImportDelegateHandle);
 }
 
 FString FLevelOfDetailSettingsLayout::GetSourceImportFilename(int32 LODIndex) const
diff --git a/Engine/Source/Editor/StaticMeshEditor/Private/StaticMeshEditorTools.h b/Engine/Source/Editor/StaticMeshEditor/Private/StaticMeshEditorTools.h
--- a/Engine/Source/Editor/StaticMeshEditor/Private/StaticMeshEditorTools.h
+++ b/Engine/Source/Editor/StaticMeshEditor/Private/StaticMeshEditorTools.h
@@ -664,6 +664,8 @@ private:
 	IDetailCategoryBuilder* LodCustomCategory;
 
 	bool DetailDisplayLODs[MAX_STATIC_MESH_LODS];
+
+	FDelegateHandle OnAssetPostLODImportDelegateHandle;
 };
 
 /**
diff --git a/Engine/Source/Runtime/Engine/Private/StaticMesh.cpp b/Engine/Source/Runtime/Engine/Private/StaticMesh.cpp
--- a/Engine/Source/Runtime/Engine/Private/StaticMesh.cpp
+++ b/Engine/Source/Runtime/Engine/Private/StaticMesh.cpp
@@ -4111,19 +4111,28 @@ bool UStaticMesh::SetCustomLOD(const UStaticMesh* SourceStaticMesh, int32 Destin
 
 	//Copy the mesh description of the source into the destination
 	*DestinationMeshDescription = *SourceMeshDescription;
+	
 
 	if(IsInGameThread())
 	{
 		PostEditChange();
-		MarkPackageDirty();
+		//Commit the mesh description to update the ddc key
+		FCommitMeshDescriptionParams CommitMeshDescriptionParams;
+		CommitMeshDescriptionParams.bMarkPackageDirty = true;
+		CommitMeshDescriptionParams.bUseHashAsGuid = false;
+		CommitMeshDescription(DestinationLodIndex, CommitMeshDescriptionParams);
 	}
 	else
 	{
 		UStaticMesh* ThisMesh = this;
-		Async(EAsyncExecution::TaskGraphMainThread, [ThisMesh]()
+		Async(EAsyncExecution::TaskGraphMainThread, [ThisMesh, DestinationLodIndex]()
 		{
 			ThisMesh->PostEditChange();
-			ThisMesh->MarkPackageDirty();
+			//Commit the mesh description to update the ddc key
+			FCommitMeshDescriptionParams CommitMeshDescriptionParams;
+			CommitMeshDescriptionParams.bMarkPackageDirty = true;
+			CommitMeshDescriptionParams.bUseHashAsGuid = false;
+			ThisMesh->CommitMeshDescription(DestinationLodIndex, CommitMeshDescriptionParams);
 		});
 	}
 
