Fixed city sample rule processor EUW related to data layers usage/exposition to users

#rnx

diff --git a/Engine/Source/Runtime/Engine/Public/WorldPartition/DataLayer/DataLayerAsset.h b/Engine/Source/Runtime/Engine/Public/WorldPartition/DataLayer/DataLayerAsset.h
--- a/Engine/Source/Runtime/Engine/Public/WorldPartition/DataLayer/DataLayerAsset.h
+++ b/Engine/Source/Runtime/Engine/Public/WorldPartition/DataLayer/DataLayerAsset.h
@@ -11,7 +11,7 @@
 
 #include "DataLayerAsset.generated.h"
 
-UCLASS(editinlinenew)
+UCLASS(BlueprintType, editinlinenew)
 class ENGINE_API UDataLayerAsset : public UObject
 {
 	GENERATED_UCLASS_BODY()
