Allow global shaders to load before precaching PSOs (UE-163502)

- Also, only issue the warning if we waited at least 1ms in that call.

diff --git a/Engine/Source/Runtime/Engine/Private/ShaderCompiler/ShaderCompiler.cpp b/Engine/Source/Runtime/Engine/Private/ShaderCompiler/ShaderCompiler.cpp
--- a/Engine/Source/Runtime/Engine/Private/ShaderCompiler/ShaderCompiler.cpp
+++ b/Engine/Source/Runtime/Engine/Private/ShaderCompiler/ShaderCompiler.cpp
@@ -6821,15 +6821,6 @@ void VerifyGlobalShaders(EShaderPlatform Platform, const ITargetPlatform* Target
 					//SharedShaderJobs.Add(ShaderTypePermutation, Job);
 					PermutationCountToCompile++;
 				}
-				else if (GlobalShader.IsValid())
-				{
-					// try and precache if compute
-					if (GlobalShader->GetFrequency() == SF_Compute)
-					{
-						FRHIComputeShader* RHIComputeShader = GlobalShader.GetComputeShader();
-						PipelineStateCache::PrecacheComputePipelineState(RHIComputeShader);
-					}
-				}
 			}
 		}
 
@@ -6925,6 +6916,37 @@ void VerifyGlobalShaders(EShaderPlatform Platform, bool bLoadedFromCacheFile, co
 	VerifyGlobalShaders(Platform, nullptr, bLoadedFromCacheFile, OutdatedShaderTypes, OutdatedShaderPipelineTypes);
 }
 
+void PrecacheComputePipelineStatesForGlobalShaders(EShaderPlatform Platform, const ITargetPlatform* TargetPlatform)
+{
+	FPlatformTypeLayoutParameters LayoutParams;
+	LayoutParams.InitializeForPlatform(TargetPlatform);
+	EShaderPermutationFlags PermutationFlags = GetShaderPermutationFlags(LayoutParams);
+
+	FGlobalShaderMap* GlobalShaderMap = GetGlobalShaderMap(Platform);
+
+	for (TLinkedList<FShaderType*>::TIterator ShaderTypeIt(FShaderType::GetTypeList()); ShaderTypeIt; ShaderTypeIt.Next())
+	{
+		FGlobalShaderType* GlobalShaderType = ShaderTypeIt->GetGlobalShaderType();
+		if (!GlobalShaderType)
+		{
+			continue;
+		}
+
+		for (int32 PermutationId = 0; PermutationId < GlobalShaderType->GetPermutationCount(); PermutationId++)
+		{
+			if (GlobalShaderType->ShouldCompilePermutation(Platform, PermutationId, PermutationFlags))
+			{
+				TShaderRef<FShader> GlobalShader = GlobalShaderMap->GetShader(GlobalShaderType, PermutationId);
+				if (GlobalShader.IsValid() && GlobalShader->GetFrequency() == SF_Compute)
+				{
+					FRHIComputeShader* RHIComputeShader = GlobalShader.GetComputeShader();
+					PipelineStateCache::PrecacheComputePipelineState(RHIComputeShader);
+				}
+			}
+		}
+	}
+}
+
 #include "Misc/PreLoadFile.h"
 #include "Serialization/LargeMemoryReader.h"
 static FPreLoadFile GGlobalShaderPreLoadFile(*(FString(TEXT("../../../Engine")) / TEXT("GlobalShaderCache-SP_") + FPlatformProperties::IniPlatformName() + TEXT(".bin")));
diff --git a/Engine/Source/Runtime/Engine/Public/ShaderCompiler.h b/Engine/Source/Runtime/Engine/Public/ShaderCompiler.h
--- a/Engine/Source/Runtime/Engine/Public/ShaderCompiler.h
+++ b/Engine/Source/Runtime/Engine/Public/ShaderCompiler.h
@@ -1198,6 +1198,9 @@ extern ENGINE_API FOnGlobalShadersCompilation& GetOnGlobalShaderCompilation();
 extern ENGINE_API void VerifyGlobalShaders(EShaderPlatform Platform, bool bLoadedFromCacheFile, const TArray<const FShaderType*>* OutdatedShaderTypes = nullptr, const TArray<const FShaderPipelineType*>* OutdatedShaderPipelineTypes = nullptr);
 extern ENGINE_API void VerifyGlobalShaders(EShaderPlatform Platform, const ITargetPlatform* TargetPlatform, bool bLoadedFromCacheFile, const TArray<const FShaderType*>* OutdatedShaderTypes = nullptr, const TArray<const FShaderPipelineType*>* OutdatedShaderPipelineTypes = nullptr);
 
+/** Precreates compute PSOs for global shaders. Separate from Compile/VerifyGlobalShaders to avoid blocking loads. */
+extern ENGINE_API void PrecacheComputePipelineStatesForGlobalShaders(EShaderPlatform Platform, const ITargetPlatform* TargetPlatform);
+
 /**
 * Forces a recompile of the global shaders.
 */
diff --git a/Engine/Source/Runtime/Launch/Private/LaunchEngineLoop.cpp b/Engine/Source/Runtime/Launch/Private/LaunchEngineLoop.cpp
--- a/Engine/Source/Runtime/Launch/Private/LaunchEngineLoop.cpp
+++ b/Engine/Source/Runtime/Launch/Private/LaunchEngineLoop.cpp
@@ -3894,6 +3894,11 @@ int32 FEngineLoop::PreInitPostStartupScreen(const TCHAR* CmdLine)
 		GetHighResScreenshotConfig().Init();
 	}
 
+	// precache compute PSOs for global shaders - enough time should have passed since loading global SM to avoid a blocking load
+	{
+		PrecacheComputePipelineStatesForGlobalShaders(GShaderPlatformForFeatureLevel[GMaxRHIFeatureLevel], nullptr);
+	}
+
 #else // WITH_ENGINE
 	InitEngineTextLocalization();
 	InitGameTextLocalization();
diff --git a/Engine/Source/Runtime/RenderCore/Private/ShaderCodeArchive.cpp b/Engine/Source/Runtime/RenderCore/Private/ShaderCodeArchive.cpp
--- a/Engine/Source/Runtime/RenderCore/Private/ShaderCodeArchive.cpp
+++ b/Engine/Source/Runtime/RenderCore/Private/ShaderCodeArchive.cpp
@@ -2044,8 +2044,16 @@ TRefCountPtr<FRHIShader> FIoStoreShaderCodeArchive::CreateShader(int32 ShaderInd
 	if (bNeededToWait)
 	{
 		TRACE_CPUPROFILER_EVENT_SCOPE(BlockingShaderLoad);
-		UE_LOG(LogShaderLibrary, Warning, TEXT("Blocking wait for shader preload, NumRefs: %d, FramePreloadStarted: %d"), PreloadEntryPtr->NumRefs, PreloadEntryPtr->FramePreloadStarted);
+
+		double TimeStarted = FPlatformTime::Seconds();
 		FTaskGraphInterface::Get().WaitUntilTaskCompletes(Event);
+		double WaitDuration = FPlatformTime::Seconds() - TimeStarted;
+
+		// only complain if we spent more than 1ms waiting
+		if (TimeStarted > 0.001)
+		{
+			UE_LOG(LogShaderLibrary, Warning, TEXT("Spent %.2f ms in a blocking wait for shader preload, NumRefs: %d, FramePreloadStarted: %d, CurrentFrame: %d"), WaitDuration * 1000.0, PreloadEntryPtr->NumRefs, PreloadEntryPtr->FramePreloadStarted, GFrameNumber);
+		}
 	}
 
 	const uint8* ShaderCode = PreloadEntryPtr->IoRequest.GetResultOrDie().Data();
