Fix invalid detection of rough refraction for Strata material.

diff --git a/Engine/Source/Runtime/Engine/Private/Materials/HLSLMaterialTranslator.cpp b/Engine/Source/Runtime/Engine/Private/Materials/HLSLMaterialTranslator.cpp
--- a/Engine/Source/Runtime/Engine/Private/Materials/HLSLMaterialTranslator.cpp
+++ b/Engine/Source/Runtime/Engine/Private/Materials/HLSLMaterialTranslator.cpp
@@ -9982,7 +9982,7 @@ bool FHLSLMaterialTranslator::StrataGenerateDerivedMaterialOperatorData()
 				}
 
 				// When at least one vertical operator exists that is not parameter blending, we can enabled writing to opaque rough refraction buffer.
-				bStrataUsesVerticalLayering = !CurrentOperator.bUseParameterBlending && CurrentOperator.OperatorType == STRATA_OPERATOR_VERTICAL;
+				bStrataUsesVerticalLayering = bStrataUsesVerticalLayering || (!CurrentOperator.bUseParameterBlending && CurrentOperator.OperatorType == STRATA_OPERATOR_VERTICAL);
 			};
 
 			WalkOperators(*StrataMaterialRootOperator, false);
