World Partition - HLOD: Support a main world non partitioned

diff --git a/Engine/Source/Runtime/Engine/Private/WorldPartition/HLOD/HLODSubsystem.cpp b/Engine/Source/Runtime/Engine/Private/WorldPartition/HLOD/HLODSubsystem.cpp
--- a/Engine/Source/Runtime/Engine/Private/WorldPartition/HLOD/HLODSubsystem.cpp
+++ b/Engine/Source/Runtime/Engine/Private/WorldPartition/HLOD/HLODSubsystem.cpp
@@ -122,19 +122,6 @@ bool UHLODSubsystem::IsHLODEnabled()
 	return UHLODSubsystem::WorldPartitionHLODEnabled;
 }
 
-bool UHLODSubsystem::ShouldCreateSubsystem(UObject* Outer) const
-{
-	if (!Super::ShouldCreateSubsystem(Outer))
-	{
-		return false;
-	}
-
-	UWorld* World = Cast<UWorld>(Outer);
-	check(World);
-
-	return World->IsPartitionedWorld();
-}
-
 bool UHLODSubsystem::DoesSupportWorldType(const EWorldType::Type WorldType) const
 {
 	return WorldType == EWorldType::Game || WorldType == EWorldType::PIE;
diff --git a/Engine/Source/Runtime/Engine/Private/WorldPartition/WorldPartitionRuntimeLevelStreamingCell.cpp b/Engine/Source/Runtime/Engine/Private/WorldPartition/WorldPartitionRuntimeLevelStreamingCell.cpp
--- a/Engine/Source/Runtime/Engine/Private/WorldPartition/WorldPartitionRuntimeLevelStreamingCell.cpp
+++ b/Engine/Source/Runtime/Engine/Private/WorldPartition/WorldPartitionRuntimeLevelStreamingCell.cpp
@@ -297,7 +297,7 @@ UWorldPartitionLevelStreamingDynamic* UWorldPartitionRuntimeLevelStreamingCell::
 	}
 #endif
 
-	if (LevelStreaming && LevelStreaming->GetWorld()->HasSubsystem<UHLODSubsystem>())
+	if (LevelStreaming)
 	{
 		LevelStreaming->OnLevelShown.AddUniqueDynamic(this, &UWorldPartitionRuntimeLevelStreamingCell::OnLevelShown);
 		LevelStreaming->OnLevelHidden.AddUniqueDynamic(this, &UWorldPartitionRuntimeLevelStreamingCell::OnLevelHidden);
diff --git a/Engine/Source/Runtime/Engine/Public/WorldPartition/HLOD/HLODSubsystem.h b/Engine/Source/Runtime/Engine/Public/WorldPartition/HLOD/HLODSubsystem.h
--- a/Engine/Source/Runtime/Engine/Public/WorldPartition/HLOD/HLODSubsystem.h
+++ b/Engine/Source/Runtime/Engine/Public/WorldPartition/HLOD/HLODSubsystem.h
@@ -40,7 +40,6 @@ public:
 	virtual ~UHLODSubsystem();
 
 	//~ Begin USubsystem Interface.
-	virtual bool ShouldCreateSubsystem(UObject* Outer) const override;
 	virtual void Initialize(FSubsystemCollectionBase& Collection) override;
 	virtual void Deinitialize() override;
 	//~ End USubsystem Interface.
