Changed "Failed to recover Nanite streaming from DDC for.." message from Warning to Log

diff --git a/Engine/Source/Runtime/Engine/Private/StaticMesh.cpp b/Engine/Source/Runtime/Engine/Private/StaticMesh.cpp
--- a/Engine/Source/Runtime/Engine/Private/StaticMesh.cpp
+++ b/Engine/Source/Runtime/Engine/Private/StaticMesh.cpp
@@ -4754,7 +4754,7 @@ bool UStaticMesh::IsCachedCookedPlatformDataLoaded(const ITargetPlatform* Target
 
 	if (bFailed)
 	{
-		UE_LOG(LogStaticMesh, Warning, TEXT("Failed to recover Nanite streaming from DDC for '%s'. Rebuilding and retrying."), *GetPathName());
+		UE_LOG(LogStaticMesh, Log, TEXT("Failed to recover Nanite streaming from DDC for '%s'. Rebuilding and retrying."), *GetPathName());
 
 		// This should be a very rare event
 		// For simplicity, just rebuild the entire RenderData
