adding DrawCentripetalCatmullRomSpline to DrawDebugHelpers

diff --git a/Engine/Source/Runtime/Engine/Private/DrawDebugHelpers.cpp b/Engine/Source/Runtime/Engine/Private/DrawDebugHelpers.cpp
--- a/Engine/Source/Runtime/Engine/Private/DrawDebugHelpers.cpp
+++ b/Engine/Source/Runtime/Engine/Private/DrawDebugHelpers.cpp
@@ -1045,6 +1045,84 @@ void DrawDebugCamera(const UWorld* InWorld, FVector const& Location, FRotator co
 	}
 }
 
+// https://en.wikipedia.org/wiki/Centripetal_Catmull%E2%80%93Rom_spline
+void DrawCentripetalCatmullRomSpline(const UWorld* InWorld, TConstArrayView<FVector> Points, FColor const& Color, float Alpha, int32 NumSamplesPerSegment, bool bPersistentLines, float LifeTime, uint8 DepthPriority, float Thickness)
+{
+	TConstArrayView<FColor> Colors(&Color, 1);
+	DrawCentripetalCatmullRomSpline(InWorld, Points, Colors, Alpha, NumSamplesPerSegment, bPersistentLines, LifeTime, DepthPriority, Thickness);
+}
+
+void DrawCentripetalCatmullRomSpline(const UWorld* InWorld, TConstArrayView<FVector> Points, TConstArrayView<FColor> Colors, float Alpha, int32 NumSamplesPerSegment, bool bPersistentLines, float LifeTime, uint8 DepthPriority, float Thickness)
+{
+	const int32 NumPoints = Points.Num();
+	const int32 NumColors = Colors.Num();
+	if (NumPoints > 1)
+	{
+		auto GetT = [](float T, float Alpha, const FVector& P0, const FVector& P1)
+		{
+			const FVector P1P0 = P1 - P0;
+			const float Dot = P1P0 | P1P0;
+			const float Pow = FMath::Pow(Dot, Alpha * .5f);
+			return Pow + T;
+		};
+
+		auto LerpColor = [](FColor A, FColor B, float T) -> FColor
+		{
+			return FColor(
+				FMath::RoundToInt(float(A.R) * (1.f - T) + float(B.R) * T),
+				FMath::RoundToInt(float(A.G) * (1.f - T) + float(B.G) * T),
+				FMath::RoundToInt(float(A.B) * (1.f - T) + float(B.B) * T),
+				FMath::RoundToInt(float(A.A) * (1.f - T) + float(B.A) * T));
+		};
+
+		FVector PrevPoint = Points[0];
+		for (int i = 0; i < NumPoints - 1; ++i)
+		{
+			const FVector& P0 = Points[FMath::Max(i - 1, 0)];
+			const FVector& P1 = Points[i];
+			const FVector& P2 = Points[i + 1];
+			const FVector& P3 = Points[FMath::Min(i + 2, NumPoints - 1)];
+
+			const float T0 = 0.0f;
+			const float T1 = GetT(T0, Alpha, P0, P1);
+			const float T2 = GetT(T1, Alpha, P1, P2);
+			const float T3 = GetT(T2, Alpha, P2, P3);
+
+			const float T1T0 = T1 - T0;
+			const float T2T1 = T2 - T1;
+			const float T3T2 = T3 - T2;
+			const float T2T0 = T2 - T0;
+			const float T3T1 = T3 - T1;
+
+			const bool bIsNearlyZeroT1T0 = FMath::IsNearlyZero(T1T0);
+			const bool bIsNearlyZeroT2T1 = FMath::IsNearlyZero(T2T1);
+			const bool bIsNearlyZeroT3T2 = FMath::IsNearlyZero(T3T2);
+			const bool bIsNearlyZeroT2T0 = FMath::IsNearlyZero(T2T0);
+			const bool bIsNearlyZeroT3T1 = FMath::IsNearlyZero(T3T1);
+
+			const FColor Color1 = Colors[FMath::Min(i, NumColors - 1)];
+			const FColor Color2 = Colors[FMath::Min(i + 1, NumColors - 1)];
+
+			for (int SampleIndex = 1; SampleIndex < NumSamplesPerSegment; ++SampleIndex)
+			{
+				const float ParametricDistance = float(SampleIndex) / float(NumSamplesPerSegment - 1);
+
+				const float T = FMath::Lerp(T1, T2, ParametricDistance);
+
+				const FVector A1 = bIsNearlyZeroT1T0 ? P0 : (T1 - T) / T1T0 * P0 + (T - T0) / T1T0 * P1;
+				const FVector A2 = bIsNearlyZeroT2T1 ? P1 : (T2 - T) / T2T1 * P1 + (T - T1) / T2T1 * P2;
+				const FVector A3 = bIsNearlyZeroT3T2 ? P2 : (T3 - T) / T3T2 * P2 + (T - T2) / T3T2 * P3;
+				const FVector B1 = bIsNearlyZeroT2T0 ? A1 : (T2 - T) / T2T0 * A1 + (T - T0) / T2T0 * A2;
+				const FVector B2 = bIsNearlyZeroT3T1 ? A2 : (T3 - T) / T3T1 * A2 + (T - T1) / T3T1 * A3;
+				const FVector Point = bIsNearlyZeroT2T1 ? B1 : (T2 - T) / T2T1 * B1 + (T - T1) / T2T1 * B2;
+
+				DrawDebugLine(InWorld, PrevPoint, Point, LerpColor(Color1, Color2, ParametricDistance), bPersistentLines, LifeTime, DepthPriority, Thickness);
+
+				PrevPoint = Point;
+			}
+		}
+	}
+}
 
 void DrawDebugFloatHistory(UWorld const & WorldRef, FDebugFloatHistory const & FloatHistory, FTransform const & DrawTransform, FVector2D const & DrawSize, FColor const & DrawColor, bool const & bPersistent, float const & LifeTime, uint8 const & DepthPriority)
 {
diff --git a/Engine/Source/Runtime/Engine/Public/DrawDebugHelpers.h b/Engine/Source/Runtime/Engine/Public/DrawDebugHelpers.h
--- a/Engine/Source/Runtime/Engine/Public/DrawDebugHelpers.h
+++ b/Engine/Source/Runtime/Engine/Public/DrawDebugHelpers.h
@@ -56,6 +56,10 @@ ENGINE_API void DrawCircle(const UWorld* InWorld, const FVector& Base, const FVe
 ENGINE_API void DrawDebugCapsule(const UWorld* InWorld, FVector const& Center, float HalfHeight, float Radius, const FQuat& Rotation, FColor const& Color, bool bPersistentLines = false, float LifeTime = -1.f, uint8 DepthPriority = 0, float Thickness = 0);
 /** Draw a debug camera shape.  FOV is full angle in degrees. */
 ENGINE_API void DrawDebugCamera(const UWorld* InWorld, FVector const& Location, FRotator const& Rotation, float FOVDeg, float Scale=1.f, FColor const& Color=FColor::White, bool bPersistentLines=false, float LifeTime=-1.f, uint8 DepthPriority = 0);
+/** Draw a centripetal catmull rom spline. Alpha values between 0 and 1 */
+ENGINE_API void DrawCentripetalCatmullRomSpline(const UWorld* InWorld, TConstArrayView<FVector> Points, FColor const& Color, float Alpha = 0.5f, int32 NumSamplesPerSegment = 8, bool bPersistentLines = false, float LifeTime = -1.f, uint8 DepthPriority = 0, float Thickness = 0.f);
+ENGINE_API void DrawCentripetalCatmullRomSpline(const UWorld* InWorld, TConstArrayView<FVector> Points, TConstArrayView<FColor> Colors, float Alpha = 0.5f, int32 NumSamplesPerSegment = 8, bool bPersistentLines = false, float LifeTime = -1.f, uint8 DepthPriority = 0, float Thickness = 0.f);
+
 ENGINE_API void FlushDebugStrings(const UWorld* InWorld);
 
 /** Draw a solid box.  Various API's provided for convenience, including ones that use FBox as well as ones that prefer (Center, Extent). **/
@@ -197,6 +201,8 @@ FORCEINLINE void DrawDebugFrustum(const UWorld* InWorld, const FMatrix& FrustumT
 FORCEINLINE void DrawCircle(const UWorld* InWorld, const FVector& Base, const FVector& X, const FVector& Y, const FColor& Color, float Radius, int32 NumSides, bool bPersistentLines = false, float LifeTime = -1.f, uint8 DepthPriority = 0, float Thickness = 0) {}
 FORCEINLINE void DrawDebugCapsule(const UWorld* InWorld, FVector const& Center, float HalfHeight, float Radius, const FQuat& Rotation, FColor const& Color, bool bPersistentLines = false, float LifeTime = -1.f, uint8 DepthPriority = 0, float Thickness = 0) {}
 FORCEINLINE void DrawDebugCamera(const UWorld* InWorld, FVector const& Location, FRotator const& Rotation, float FOVDeg, float Scale = 1.f, FColor const& Color = FColor::White, bool bPersistentLines = false, float LifeTime = -1.f, uint8 DepthPriority = 0) {}
+FORCEINLINE void DrawCentripetalCatmullRomSpline(const UWorld* InWorld, TConstArrayView<FVector> Points, FColor const& Color, float Alpha = 0.5f, int32 NumSamplesPerSegment = 8, bool bPersistentLines = false, float LifeTime = -1.f, uint8 DepthPriority = 0, float Thickness = 0.f) {}
+FORCEINLINE void DrawCentripetalCatmullRomSpline(const UWorld* InWorld, TConstArrayView<FVector> Points, TConstArrayView<FColor> Colors, float Alpha = 0.5f, int32 NumSamplesPerSegment = 8, bool bPersistentLines = false, float LifeTime = -1.f, uint8 DepthPriority = 0, float Thickness = 0.f);
 FORCEINLINE void FlushDebugStrings(const UWorld* InWorld) {}
 FORCEINLINE void DrawDebugSolidBox(const UWorld* InWorld, FBox const& Box, FColor const& Color, const FTransform& Transform = FTransform::Identity, bool bPersistent = false, float LifeTime = -1.f, uint8 DepthPriority = 0) {}
 FORCEINLINE void DrawDebugSolidBox(const UWorld* InWorld, FVector const& Center, FVector const& Extent, FColor const& Color, bool bPersistent = false, float LifeTime = -1.f, uint8 DepthPriority = 0) {}
