Adding `Total Thread Preprocess Time` for shader compilation.

- Instrument how long shader preprocessing takes for each shader compilation job.
- Keep a running total in our shader stats per shader type.
- Add logging to track total thread preprocessing time for all shaders.
- Fixup local variable names to match the output log names.

diff --git a/Engine/Source/Developer/Apple/MetalShaderFormat/Private/MetalShaderCompiler.cpp b/Engine/Source/Developer/Apple/MetalShaderFormat/Private/MetalShaderCompiler.cpp
--- a/Engine/Source/Developer/Apple/MetalShaderFormat/Private/MetalShaderCompiler.cpp
+++ b/Engine/Source/Developer/Apple/MetalShaderFormat/Private/MetalShaderCompiler.cpp
@@ -1003,6 +1003,8 @@ void CompileShader_Metal(const FShaderCompilerInput& _Input,FShaderCompilerOutpu
 
 	AdditionalDefines.SetDefine(TEXT("COMPILER_SUPPORTS_DUAL_SOURCE_BLENDING_SLOT_DECORATION"), (uint32)1);
 
+	const double StartPreprocessTime = FPlatformTime::Seconds();
+
 	if (Input.bSkipPreprocessedCache)
 	{
 		if (!FFileHelper::LoadFileToString(PreprocessedShader, *Input.VirtualSourceFilePath))
@@ -1050,6 +1052,8 @@ void CompileShader_Metal(const FShaderCompilerInput& _Input,FShaderCompilerOutpu
 
 	// Process TEXT macro.
 	TransformStringIntoCharacterArray(PreprocessedShader);
+
+	Output.PreprocessTime = FPlatformTime::Seconds() - StartPreprocessTime;
 	
 	uint32 CCFlags = HLSLCC_NoPreprocess | HLSLCC_PackUniformsIntoUniformBufferWithNames | HLSLCC_FixAtomicReferences | HLSLCC_RetainSizes | HLSLCC_KeepSamplerAndImageNames;
 	if (!bDirectCompile || UE_BUILD_DEBUG)
diff --git a/Engine/Source/Developer/VulkanShaderFormat/Private/VulkanShaderCompiler.cpp b/Engine/Source/Developer/VulkanShaderFormat/Private/VulkanShaderCompiler.cpp
--- a/Engine/Source/Developer/VulkanShaderFormat/Private/VulkanShaderCompiler.cpp
+++ b/Engine/Source/Developer/VulkanShaderFormat/Private/VulkanShaderCompiler.cpp
@@ -2327,6 +2327,8 @@ void DoCompileVulkanShader(const FShaderCompilerInput& Input, FShaderCompilerOut
 	AdditionalDefines.SetDefine(TEXT("VULKAN_BINDLESS_SAMPLER_ARRAY_PREFIX"), VulkanBindless::kBindlessSamplerArrayPrefix);
 	AdditionalDefines.SetDefine(TEXT("VULKAN_BINDLESS_RESOURCE_ARRAY_PREFIX"), VulkanBindless::kBindlessResourceArrayPrefix);
 
+	const double StartPreprocessTime = FPlatformTime::Seconds();
+
 	// Preprocess the shader.
 	FString PreprocessedShaderSource;
 	const bool bDirectCompile = FParse::Param(FCommandLine::Get(), TEXT("directcompile"));
@@ -2363,6 +2365,8 @@ void DoCompileVulkanShader(const FShaderCompilerInput& Input, FShaderCompilerOut
 	// Process TEXT macro.
 	TransformStringIntoCharacterArray(PreprocessedShaderSource);
 
+	Output.PreprocessTime = FPlatformTime::Seconds() - StartPreprocessTime;
+
 	FCompilerInfo CompilerInfo(Input, WorkingDirectory, HlslFrequency);
 
 	// Setup hlslcc flags. Needed here as it will be used when dumping debug info
diff --git a/Engine/Source/Developer/Windows/ShaderFormatD3D/Private/D3DShaderCompiler.cpp b/Engine/Source/Developer/Windows/ShaderFormatD3D/Private/D3DShaderCompiler.cpp
--- a/Engine/Source/Developer/Windows/ShaderFormatD3D/Private/D3DShaderCompiler.cpp
+++ b/Engine/Source/Developer/Windows/ShaderFormatD3D/Private/D3DShaderCompiler.cpp
@@ -1225,6 +1225,8 @@ void CompileD3DShader(const FShaderCompilerInput& Input, FShaderCompilerOutput&
 			AdditionalDefines.SetDefine(TEXT("PLATFORM_SUPPORTS_REAL_TYPES"), 1);
 		}
 	}
+
+	const double StartPreprocessTime = FPlatformTime::Seconds();
 	
 	if (Input.bSkipPreprocessedCache)
 	{
@@ -1353,6 +1355,8 @@ void CompileD3DShader(const FShaderCompilerInput& Input, FShaderCompilerOutput&
 		}
 	}
 
+	Output.PreprocessTime = FPlatformTime::Seconds() - StartPreprocessTime;
+
 	// @TODO - implement different material path to allow us to remove backwards compat flag on sm5 shaders
 	uint32 CompileFlags = D3DCOMPILE_ENABLE_BACKWARDS_COMPATIBILITY
 		// Unpack uniform matrices as row-major to match the CPU layout.
diff --git a/Engine/Source/Runtime/Engine/Private/ShaderCompiler/ShaderCompiler.cpp b/Engine/Source/Runtime/Engine/Private/ShaderCompiler/ShaderCompiler.cpp
--- a/Engine/Source/Runtime/Engine/Private/ShaderCompiler/ShaderCompiler.cpp
+++ b/Engine/Source/Runtime/Engine/Private/ShaderCompiler/ShaderCompiler.cpp
@@ -3134,17 +3134,21 @@ void FShaderCompilerStats::WriteStatSummary()
 	if (ShaderTimings.Num())
 	{
 		// calculate effective parallelization (total time needed to compile all shaders divided by actual wall clock time spent processing at least 1 shader)
-		double TotalTimeForAllShaders = 0;
+		double TotalThreadTimeForAllShaders = 0.0;
+		double TotalThreadPreprocessTimeForAllShaders = 0.0;
 		for (TMap<FString, FShaderTimings>::TConstIterator Iter(ShaderTimings); Iter; ++Iter)
 		{
-			TotalTimeForAllShaders += Iter.Value().TotalCompileTime;
+			TotalThreadTimeForAllShaders += Iter.Value().TotalCompileTime;
+			TotalThreadPreprocessTimeForAllShaders += Iter.Value().TotalPreprocessTime;
 		}
 
-		UE_LOG(LogShaderCompilers, Display, TEXT("Total thread time: %.2f s"), TotalTimeForAllShaders);
+		UE_LOG(LogShaderCompilers, Display, TEXT("Total thread time: %.2f s"), TotalThreadTimeForAllShaders);
+		UE_LOG(LogShaderCompilers, Display, TEXT("Total thread preprocess time: %.2f s"), TotalThreadPreprocessTimeForAllShaders);
+		UE_LOG(LogShaderCompilers, Display, TEXT("Percentage time preprocessing: %.2f%%"), TotalThreadTimeForAllShaders > 0.0 ? (TotalThreadPreprocessTimeForAllShaders / TotalThreadTimeForAllShaders) * 100.0 : 0.0);
 
 		if (TotalTimeAtLeastOneJobWasInFlight > 0.0)
 		{
-			double EffectiveParallelization = TotalTimeForAllShaders / TotalTimeAtLeastOneJobWasInFlight;
+			double EffectiveParallelization = TotalThreadTimeForAllShaders / TotalTimeAtLeastOneJobWasInFlight;
 			if (DistributedJobBatchesSeen == 0)
 			{
 				UE_LOG(LogShaderCompilers, Display, TEXT("Effective parallelization: %.2f (times faster than compiling all shaders on one thread). Compare with number of workers: %d"), EffectiveParallelization, GShaderCompilingManager->GetNumLocalWorkers());
@@ -3185,7 +3189,7 @@ void FShaderCompilerStats::WriteStatSummary()
 			const FShaderTimings& Timings = Iter.Value();
 
 			UE_LOG(LogShaderCompilers, Display, TEXT("%60s - %.2f%% of total time (compiled %4d times, average %4.2f sec, max %4.2f sec, min %4.2f sec)"), 
-				*Iter.Key(), 100.0 * Timings.TotalCompileTime / TotalTimeForAllShaders, Timings.NumCompiled, Timings.AverageCompileTime, Timings.MaxCompileTime, Timings.MinCompileTime);
+				*Iter.Key(), 100.0 * Timings.TotalCompileTime / TotalThreadTimeForAllShaders, Timings.NumCompiled, Timings.AverageCompileTime, Timings.MaxCompileTime, Timings.MinCompileTime);
 			if (++Idx >= MaxShadersToPrint)
 			{
 				break;
@@ -3295,6 +3299,7 @@ void FShaderCompilerStats::RegisterFinishedJob(FShaderCommonCompileJob& Job)
 			Existing->MinCompileTime = FMath::Min(Existing->MinCompileTime, static_cast<float>(SingleJob->Output.CompileTime));
 			Existing->MaxCompileTime = FMath::Max(Existing->MaxCompileTime, static_cast<float>(SingleJob->Output.CompileTime));
 			Existing->TotalCompileTime += SingleJob->Output.CompileTime;
+			Existing->TotalPreprocessTime += SingleJob->Output.PreprocessTime;
 			Existing->NumCompiled++;
 			// calculate as an optimization to make sorting later faster
 			Existing->AverageCompileTime = Existing->TotalCompileTime / static_cast<float>(Existing->NumCompiled);
@@ -3305,6 +3310,7 @@ void FShaderCompilerStats::RegisterFinishedJob(FShaderCommonCompileJob& Job)
 			New.MinCompileTime = SingleJob->Output.CompileTime;
 			New.MaxCompileTime = New.MinCompileTime;
 			New.TotalCompileTime = New.MinCompileTime;
+			New.TotalPreprocessTime += SingleJob->Output.PreprocessTime;
 			New.AverageCompileTime = New.MinCompileTime;
 			New.NumCompiled = 1;
 			ShaderTimings.Add(ShaderName, New);
diff --git a/Engine/Source/Runtime/Engine/Public/ShaderCompiler.h b/Engine/Source/Runtime/Engine/Public/ShaderCompiler.h
--- a/Engine/Source/Runtime/Engine/Public/ShaderCompiler.h
+++ b/Engine/Source/Runtime/Engine/Public/ShaderCompiler.h
@@ -707,11 +707,12 @@ public:
 	/** Structure used to describe compiling time of a shader type (for all the instances of it that we have seen). */
 	struct FShaderTimings
 	{
-		float MinCompileTime;
-		float MaxCompileTime;
-		float TotalCompileTime;
-		int32 NumCompiled;
-		float AverageCompileTime;	// stored explicitly as an optimization
+		float MinCompileTime = 0.0f;
+		float MaxCompileTime = 0.0f;
+		float TotalCompileTime = 0.0f;
+		float TotalPreprocessTime = 0.0f;
+		int32 NumCompiled = 0;
+		float AverageCompileTime = 0.0f;	// stored explicitly as an optimization
 	};
 
 	ENGINE_API void RegisterCookedShaders(uint32 NumCooked, float CompileTime, EShaderPlatform Platform, const FString MaterialPath, FString PermutationString = FString(""));
diff --git a/Engine/Source/Runtime/RenderCore/Public/ShaderCompilerCore.h b/Engine/Source/Runtime/RenderCore/Public/ShaderCompilerCore.h
--- a/Engine/Source/Runtime/RenderCore/Public/ShaderCompilerCore.h
+++ b/Engine/Source/Runtime/RenderCore/Public/ShaderCompilerCore.h
@@ -19,7 +19,7 @@ class Error;
 // this is for the protocol, not the data, bump if FShaderCompilerInput or ProcessInputFromArchive changes.
 const int32 ShaderCompileWorkerInputVersion = 17;
 // this is for the protocol, not the data, bump if FShaderCompilerOutput or WriteToOutputArchive changes.
-const int32 ShaderCompileWorkerOutputVersion = 7;
+const int32 ShaderCompileWorkerOutputVersion = 8;
 // this is for the protocol, not the data.
 const int32 ShaderCompileWorkerSingleJobHeader = 'S';
 // this is for the protocol, not the data.
@@ -445,6 +445,7 @@ struct FShaderCompilerOutput
 	:	NumInstructions(0)
 	,	NumTextureSamplers(0)
 	,	CompileTime(0.0)
+	,	PreprocessTime(0.0)
 	,	bSucceeded(false)
 	,	bFailedRemovingUnused(false)
 	,	bSupportsQueryingUsedAttributes(false)
@@ -461,6 +462,7 @@ struct FShaderCompilerOutput
 	uint32 NumInstructions;
 	uint32 NumTextureSamplers;
 	double CompileTime;
+	double PreprocessTime;
 	bool bSucceeded;
 	bool bFailedRemovingUnused;
 	bool bSupportsQueryingUsedAttributes;
@@ -484,6 +486,7 @@ struct FShaderCompilerOutput
 		Ar << Output.ParameterMap << Output.Errors << Output.Target << Output.ShaderCode << Output.OutputHash << Output.NumInstructions << Output.NumTextureSamplers << Output.bSucceeded;
 		Ar << Output.bFailedRemovingUnused << Output.bSupportsQueryingUsedAttributes << Output.UsedAttributes;
 		Ar << Output.CompileTime;
+		Ar << Output.PreprocessTime;
 		Ar << Output.OptionalPreprocessedShaderSource;
 		Ar << Output.OptionalFinalShaderSource;
 		Ar << Output.PlatformDebugData;
