Exclude primitives without raytracing geometry (eg: "Support Ray Tracing" disabled in Static Mesh Editor) from raytracing scene.

diff --git a/Engine/Source/Runtime/Engine/Private/PrimitiveSceneProxy.cpp b/Engine/Source/Runtime/Engine/Private/PrimitiveSceneProxy.cpp
--- a/Engine/Source/Runtime/Engine/Private/PrimitiveSceneProxy.cpp
+++ b/Engine/Source/Runtime/Engine/Private/PrimitiveSceneProxy.cpp
@@ -1486,7 +1486,12 @@ ERayTracingPrimitiveFlags FPrimitiveSceneProxy::GetCachedRayTracingInstance(FRay
 
 	if (IsRayTracingStaticRelevant())
 	{
-		check(PrimitiveSceneInfo->GetRayTracingGeometryNum() > 0 && PrimitiveSceneInfo->StaticMeshes.Num() > 0);
+		if (PrimitiveSceneInfo->GetRayTracingGeometryNum() == 0)
+		{
+			return ERayTracingPrimitiveFlags::Excluded;
+		}
+
+		check(PrimitiveSceneInfo->StaticMeshes.Num() > 0);
 
 		// overwrite flag if static
 		ResultFlags = ERayTracingPrimitiveFlags::StaticMesh | ERayTracingPrimitiveFlags::ComputeLOD | ERayTracingPrimitiveFlags::CacheMeshCommands;
