World Partition: resolve actor descriptor views on container instances to fix missing data layers on level instances.

#rnx

diff --git a/Engine/Source/Runtime/Engine/Private/WorldPartition/WorldPartitionStreamingGeneration.cpp b/Engine/Source/Runtime/Engine/Private/WorldPartition/WorldPartitionStreamingGeneration.cpp
--- a/Engine/Source/Runtime/Engine/Private/WorldPartition/WorldPartitionStreamingGeneration.cpp
+++ b/Engine/Source/Runtime/Engine/Private/WorldPartition/WorldPartitionStreamingGeneration.cpp
@@ -210,7 +210,15 @@ class FWorldPartitionStreamingGenerator
 		FActorContainerID ID;
 	};
 
-	void ResolveRuntimeDataLayers(FWorldPartitionActorDescView& ActorDescView, const FActorDescViewMap& ActorDescViewMap)
+	void ResolveRuntimeSpatiallyLoaded(FWorldPartitionActorDescView& ActorDescView)
+	{
+		if (!bEnableStreaming)
+		{
+			ActorDescView.SetForcedNonSpatiallyLoaded();
+		}
+	}
+
+	void ResolveRuntimeDataLayers(FWorldPartitionActorDescView& ActorDescView)
 	{
 		TArray<FName> RuntimeDataLayerInstanceNames;
 		RuntimeDataLayerInstanceNames.Reserve(ActorDescView.GetDataLayers().Num());
@@ -476,16 +484,22 @@ class FWorldPartitionStreamingGenerator
 	 */
 	void ResolveContainerDescriptor(FContainerDescriptor& ContainerDescriptor)
 	{
-		ContainerDescriptor.ActorDescViewMap.ForEachActorDescView([this, &ContainerDescriptor](FWorldPartitionActorDescView& ActorDescView)
+		auto ResolveActorDescView = [this, &ContainerDescriptor](FWorldPartitionActorDescView& ActorDescView)
 		{
-			if (!bEnableStreaming)
-			{
-				ActorDescView.SetForcedNonSpatiallyLoaded();
-			}
-
+			ResolveRuntimeSpatiallyLoaded(ActorDescView);
 			ResolveRuntimeDataLayers(ActorDescView, ContainerDescriptor.ActorDescViewMap);
 			ResolveRuntimeReferences(ActorDescView, ContainerDescriptor.ActorDescViewMap);
+		};
+
+		ContainerDescriptor.ActorDescViewMap.ForEachActorDescView([this, &ResolveActorDescView](FWorldPartitionActorDescView& ActorDescView)
+		{
+			ResolveActorDescView(ActorDescView);
 		});
+
+		for (FWorldPartitionActorDescView& ContainerInstanceView : ContainerDescriptor.ContainerInstanceViews)
+		{
+			ResolveActorDescView(ContainerInstanceView);
+		}
 	}
 
 	/** 
