Fix bad edit during review

#rnx

diff --git a/Engine/Source/Runtime/Engine/Private/WorldPartition/WorldPartitionStreamingGeneration.cpp b/Engine/Source/Runtime/Engine/Private/WorldPartition/WorldPartitionStreamingGeneration.cpp
--- a/Engine/Source/Runtime/Engine/Private/WorldPartition/WorldPartitionStreamingGeneration.cpp
+++ b/Engine/Source/Runtime/Engine/Private/WorldPartition/WorldPartitionStreamingGeneration.cpp
@@ -218,7 +218,7 @@ class FWorldPartitionStreamingGenerator
 		}
 	}
 
-	void ResolveRuntimeDataLayers(FWorldPartitionActorDescView& ActorDescView)
+	void ResolveRuntimeDataLayers(FWorldPartitionActorDescView& ActorDescView, const FActorDescViewMap& ActorDescViewMap)
 	{
 		TArray<FName> RuntimeDataLayerInstanceNames;
 		RuntimeDataLayerInstanceNames.Reserve(ActorDescView.GetDataLayers().Num());
