Fix: Texture streaming state should not request a top LOD index in the nonstreaming mips
fixes crash on 4x1024 texture when you set cinematic mip lod bias

diff --git a/Engine/Source/Runtime/Engine/Private/Texture.cpp b/Engine/Source/Runtime/Engine/Private/Texture.cpp
--- a/Engine/Source/Runtime/Engine/Private/Texture.cpp
+++ b/Engine/Source/Runtime/Engine/Private/Texture.cpp
@@ -1367,6 +1367,9 @@ FStreamableRenderResourceState UTexture::GetResourcePostInitState(const FTexture
 
 	const int32 NumOfNonOptionalMips = FMath::Min<int32>(NumMips, PlatformData->GetNumNonOptionalMips());
 	const int32 NumOfNonStreamingMips = FMath::Min<int32>(NumMips, PlatformData->GetNumNonStreamingMips());
+	// Optional mips must be streaming mips :
+	check( NumOfNonOptionalMips >= NumOfNonStreamingMips );
+
 	const int32 AssetMipIdxForResourceFirstMip = FMath::Max<int32>(0, PlatformData->Mips.Num() - NumMips);
 
 	bool bMakeStreamble = false;
@@ -1420,6 +1423,9 @@ FStreamableRenderResourceState UTexture::GetResourcePostInitState(const FTexture
 		{
 			NumRequestedMips = NumOfNonOptionalMips;
 		}
+
+		// Ensure we don't request a top mip in the NonStreamingMips
+		NumRequestedMips = FMath::Max(NumRequestedMips,NumOfNonStreamingMips);
 	}
 
 	if (NumRequestedMips < MinRequestMipCount && MinRequestMipCount < NumMips)
