WorldPartitionAssetDependencyGatherer: Avoid filtering AssetClass as it can be null. Since we are scanning ExternalActors folder consider all assets as valid

diff --git a/Engine/Source/Runtime/Engine/Private/WorldPartition/WorldPartitionAssetDependencyGatherer.cpp b/Engine/Source/Runtime/Engine/Private/WorldPartition/WorldPartitionAssetDependencyGatherer.cpp
--- a/Engine/Source/Runtime/Engine/Private/WorldPartition/WorldPartitionAssetDependencyGatherer.cpp
+++ b/Engine/Source/Runtime/Engine/Private/WorldPartition/WorldPartitionAssetDependencyGatherer.cpp
@@ -27,10 +27,7 @@ void FWorldPartitionAssetDependencyGatherer::GatherDependencies(const FAssetData
 
 		for (const FAssetData& FilteredAsset : FilteredAssets)
 		{
-			if (const UClass* FilteredAssetClass = FilteredAsset.GetClass(); FilteredAssetClass && FilteredAssetClass->IsChildOf(AActor::StaticClass()))
-			{
-				OutDependencies.Emplace(IAssetDependencyGatherer::FGathereredDependency{ FilteredAsset.PackageName, UE::AssetRegistry::EDependencyProperty::Game });
-			}
+			OutDependencies.Emplace(IAssetDependencyGatherer::FGathereredDependency{ FilteredAsset.PackageName, UE::AssetRegistry::EDependencyProperty::Game });
 		}
 	}
 }
