#UE Avoid scanning primary asset types more than once in editor cases where background asset discovery is used

diff --git a/Engine/Source/Runtime/Engine/Private/AssetManager.cpp b/Engine/Source/Runtime/Engine/Private/AssetManager.cpp
--- a/Engine/Source/Runtime/Engine/Private/AssetManager.cpp
+++ b/Engine/Source/Runtime/Engine/Private/AssetManager.cpp
@@ -826,6 +826,9 @@ int32 UAssetManager::ScanPathsForPrimaryAssets(FPrimaryAssetType PrimaryAssetTyp
 		{
 			TypeData.DeferredAssetScanPaths.AddUnique(Path);
 		}
+
+		// Since we are still asynchronously discovering assets, we'll wait until that is done before populating with primary assets
+		return 0;
 	}
 #endif
 
