Bugfix [MetaSounds] [Asset Importing] - Can't Import Wave Files Created from MetaSound Wave Writer

diff --git a/Engine/Source/Runtime/Engine/Private/Audio.cpp b/Engine/Source/Runtime/Engine/Private/Audio.cpp
--- a/Engine/Source/Runtime/Engine/Private/Audio.cpp
+++ b/Engine/Source/Runtime/Engine/Private/Audio.cpp
@@ -1563,7 +1563,8 @@ bool FWaveModInfo::ReadWaveInfo( const uint8* WaveData, int32 WaveDataSize, FStr
 
 	if (*pFormatTag != 0x0001 // WAVE_FORMAT_PCM
 		&& *pFormatTag != 0x0002 // WAVE_FORMAT_ADPCM
-		&& *pFormatTag != 0x0011) // WAVE_FORMAT_DVI_ADPCM
+		&& *pFormatTag != 0x0011 // WAVE_FORMAT_DVI_ADPCM
+		&& *pFormatTag != 0x0003) // WAVE_FORMAT_IEEE_FLOAT
 	{
 		ReportImportFailure();
 		if (ErrorReason) *ErrorReason = TEXT("Unsupported wave file format.  Only PCM, ADPCM, and DVI ADPCM can be imported.");
