Fix to allow Gpu morph resources to be created when there associated LodRenderData is finish streaming.
Wrapped the fix around r.SkeletalMesh.AllowSkinnedMorphDataStreaming,disable by default.

diff --git a/Engine/Source/Runtime/Engine/Private/SkeletalMeshLODRenderData.cpp b/Engine/Source/Runtime/Engine/Private/SkeletalMeshLODRenderData.cpp
--- a/Engine/Source/Runtime/Engine/Private/SkeletalMeshLODRenderData.cpp
+++ b/Engine/Source/Runtime/Engine/Private/SkeletalMeshLODRenderData.cpp
@@ -30,6 +30,13 @@ static FAutoConsoleVariableRef CVarStripSkeletalMeshLodsBelowMinLod(
 	TEXT("If set will strip skeletal mesh LODs under the minimum renderable LOD for the target platform during cooking.")
 );
 
+int32 GAllowSkinnedMorphDataStreaming = 0;
+static FAutoConsoleVariableRef CVarAllowSkinnedMorphDataStreaming(
+	TEXT("r.Skinned.AllowSkinnedMorphDataStreaming"),
+	GAllowSkinnedMorphDataStreaming,
+	TEXT("Call BeginInitResource when the morph resources are finish streaming (DoFinishUpdate)")
+);
+
 namespace
 {
 	struct FReverseOrderBitArraysBySetBits
@@ -159,6 +166,14 @@ FArchive& operator<<(FArchive& Ar, FSkelMeshRenderSection& S)
 	return Ar;
 }
 
+void FSkeletalMeshLODRenderData::InitMorphResources()
+{
+	if (GAllowSkinnedMorphDataStreaming && !MorphTargetVertexInfoBuffers.IsRHIIntialized() && MorphTargetVertexInfoBuffers.IsMorphCPUDataValid() && MorphTargetVertexInfoBuffers.NumTotalBatches > 0)
+	{
+		BeginInitResource(&MorphTargetVertexInfoBuffers);
+	}
+}
+
 void FSkeletalMeshLODRenderData::InitResources(bool bNeedsVertexColors, int32 LODIndex, TArray<UMorphTarget*>& InMorphTargets, USkinnedAsset* Owner)
 {
 	if (bStreamedDataInlined)
diff --git a/Engine/Source/Runtime/Engine/Private/Streaming/SkeletalMeshUpdate.cpp b/Engine/Source/Runtime/Engine/Private/Streaming/SkeletalMeshUpdate.cpp
--- a/Engine/Source/Runtime/Engine/Private/Streaming/SkeletalMeshUpdate.cpp
+++ b/Engine/Source/Runtime/Engine/Private/Streaming/SkeletalMeshUpdate.cpp
@@ -198,6 +198,7 @@ void FSkeletalMeshStreamIn::DoFinishUpdate(const FContext& Context)
 			{
 				FSkeletalMeshLODRenderData& LODResource = *Context.LODResourcesView[LODIndex];
 				LODResource.IncrementMemoryStats(Mesh->GetHasVertexColors());
+				LODResource.InitMorphResources();
 				IntermediateBuffersArray[LODIndex].TransferBuffers(LODResource, Batcher);
 			}
 		}
diff --git a/Engine/Source/Runtime/Engine/Public/Rendering/SkeletalMeshLODRenderData.h b/Engine/Source/Runtime/Engine/Public/Rendering/SkeletalMeshLODRenderData.h
--- a/Engine/Source/Runtime/Engine/Public/Rendering/SkeletalMeshLODRenderData.h
+++ b/Engine/Source/Runtime/Engine/Public/Rendering/SkeletalMeshLODRenderData.h
@@ -177,6 +177,8 @@ public:
 	*/
 	void InitResources(bool bNeedsVertexColors, int32 LODIndex, TArray<class UMorphTarget*>& InMorphTargets, USkinnedAsset* Owner);
 
+	void InitMorphResources();
+
 	/**
 	* Releases the LOD's render resources.
 	*/
