fix for Issue 215711: Compile errors in Module.PoseSearch.cpp

diff --git a/Engine/Source/Runtime/Engine/Public/DrawDebugHelpers.h b/Engine/Source/Runtime/Engine/Public/DrawDebugHelpers.h
--- a/Engine/Source/Runtime/Engine/Public/DrawDebugHelpers.h
+++ b/Engine/Source/Runtime/Engine/Public/DrawDebugHelpers.h
@@ -202,7 +202,7 @@ FORCEINLINE void DrawCircle(const UWorld* InWorld, const FVector& Base, const FV
 FORCEINLINE void DrawDebugCapsule(const UWorld* InWorld, FVector const& Center, float HalfHeight, float Radius, const FQuat& Rotation, FColor const& Color, bool bPersistentLines = false, float LifeTime = -1.f, uint8 DepthPriority = 0, float Thickness = 0) {}
 FORCEINLINE void DrawDebugCamera(const UWorld* InWorld, FVector const& Location, FRotator const& Rotation, float FOVDeg, float Scale = 1.f, FColor const& Color = FColor::White, bool bPersistentLines = false, float LifeTime = -1.f, uint8 DepthPriority = 0) {}
 FORCEINLINE void DrawCentripetalCatmullRomSpline(const UWorld* InWorld, TConstArrayView<FVector> Points, FColor const& Color, float Alpha = 0.5f, int32 NumSamplesPerSegment = 8, bool bPersistentLines = false, float LifeTime = -1.f, uint8 DepthPriority = 0, float Thickness = 0.f) {}
-FORCEINLINE void DrawCentripetalCatmullRomSpline(const UWorld* InWorld, TConstArrayView<FVector> Points, TConstArrayView<FColor> Colors, float Alpha = 0.5f, int32 NumSamplesPerSegment = 8, bool bPersistentLines = false, float LifeTime = -1.f, uint8 DepthPriority = 0, float Thickness = 0.f);
+FORCEINLINE void DrawCentripetalCatmullRomSpline(const UWorld* InWorld, TConstArrayView<FVector> Points, TConstArrayView<FColor> Colors, float Alpha = 0.5f, int32 NumSamplesPerSegment = 8, bool bPersistentLines = false, float LifeTime = -1.f, uint8 DepthPriority = 0, float Thickness = 0.f) {}
 FORCEINLINE void FlushDebugStrings(const UWorld* InWorld) {}
 FORCEINLINE void DrawDebugSolidBox(const UWorld* InWorld, FBox const& Box, FColor const& Color, const FTransform& Transform = FTransform::Identity, bool bPersistent = false, float LifeTime = -1.f, uint8 DepthPriority = 0) {}
 FORCEINLINE void DrawDebugSolidBox(const UWorld* InWorld, FVector const& Center, FVector const& Extent, FColor const& Color, bool bPersistent = false, float LifeTime = -1.f, uint8 DepthPriority = 0) {}
