GitHub 9393 : Allow poses to be defined properly from user input
#feature added context menu option(s) allowing user to add additional _named_ poses created from either ref-pose or the current viewport pose. Did not integrate pull-request code, but added behaviour that will allow for achieving the same end result.

diff --git a/Engine/Source/Editor/Persona/Private/PoseEditorCommands.cpp b/Engine/Source/Editor/Persona/Private/PoseEditorCommands.cpp
--- a/Engine/Source/Editor/Persona/Private/PoseEditorCommands.cpp
+++ b/Engine/Source/Editor/Persona/Private/PoseEditorCommands.cpp
@@ -11,6 +11,8 @@ void FPoseEditorCommands::RegisterCommands()
 {
 	UI_COMMAND(PasteAllNames, "Paste All Pose Names", "Paste all pose names from clipboard", EUserInterfaceActionType::Button, FInputChord());
 	UI_COMMAND(UpdatePoseToCurrent, "Update Pose to Current", "Updates the selected pose to match the pose currently shown in the viewport", EUserInterfaceActionType::Button, FInputChord());
+	UI_COMMAND(AddPoseFromCurrent, "Add Pose from Current", "Adds a new pose matching the pose currently shown in the viewport", EUserInterfaceActionType::Button, FInputChord());	
+	UI_COMMAND(AddPoseFromReference, "Add Pose from Reference", "Adds a new pose matching the Skeleton its reference pose", EUserInterfaceActionType::Button, FInputChord());
 }
 
 #undef LOCTEXT_NAMESPACE
diff --git a/Engine/Source/Editor/Persona/Private/PoseEditorCommands.h b/Engine/Source/Editor/Persona/Private/PoseEditorCommands.h
--- a/Engine/Source/Editor/Persona/Private/PoseEditorCommands.h
+++ b/Engine/Source/Editor/Persona/Private/PoseEditorCommands.h
@@ -36,6 +36,9 @@ public:
 
 	/** Update selected pose to match viewport */
 	TSharedPtr< FUICommandInfo > UpdatePoseToCurrent;
+
+	TSharedPtr< FUICommandInfo > AddPoseFromCurrent;	
+	TSharedPtr< FUICommandInfo > AddPoseFromReference;
 };
 
 
diff --git a/Engine/Source/Editor/Persona/Private/SPoseEditor.cpp b/Engine/Source/Editor/Persona/Private/SPoseEditor.cpp
--- a/Engine/Source/Editor/Persona/Private/SPoseEditor.cpp
+++ b/Engine/Source/Editor/Persona/Private/SPoseEditor.cpp
@@ -302,6 +302,8 @@ void SPoseViewer::Construct(const FArguments& InArgs, const TSharedRef<IPersonaT
 	PersonaToolkitPtr = InPersonaToolkit;
 	EditableSkeletonPtr = InEditableSkeleton;
 	PoseAssetPtr = InArgs._PoseAsset;
+	
+	NewPoseName = UPoseAsset::GetUniquePoseName(&(EditableSkeletonPtr.Pin()->GetSkeleton()));
 
 	InPreviewScene->RegisterOnPreviewMeshChanged(FOnPreviewMeshChanged::CreateSP(this, &SPoseViewer::OnPreviewMeshChanged));
 
@@ -624,10 +626,54 @@ void SPoseViewer::BindCommands()
 			const TArray<TSharedPtr<FDisplayedPoseInfo>> SelectedRows = PoseListView->GetSelectedItems();
 			return SelectedRows.Num() == 1;
 		}));
+
+	CommandList.MapAction(
+		PoseEditorCommands.AddPoseFromCurrent,
+		FExecuteAction::CreateSP(this, &SPoseViewer::AddPoseWithCurrent),
+		FCanExecuteAction::CreateLambda([this]() -> bool
+		{
+			if (const UPoseAsset* PoseAsset = PoseAssetPtr.Get())
+			{
+				FText Temp;
+				return PoseAsset->SourceAnimation == nullptr && IsNewPoseNameValid(Temp);
+			}
+			return false;
+		}),
+		FGetActionCheckState(),
+		FIsActionButtonVisible::CreateLambda([this]() -> bool
+		{
+			if (const UPoseAsset* PoseAsset = PoseAssetPtr.Get())
+			{
+				return PoseAsset->SourceAnimation == nullptr;
+			}
+			return false;
+		}));
+
+	CommandList.MapAction(
+		PoseEditorCommands.AddPoseFromReference,
+		FExecuteAction::CreateSP(this, &SPoseViewer::AddPoseWithReference),
+		FCanExecuteAction::CreateLambda([this]() -> bool
+		{
+			if (const UPoseAsset* PoseAsset = PoseAssetPtr.Get())
+			{
+				FText Temp;
+				return PoseAsset->SourceAnimation == nullptr && IsNewPoseNameValid(Temp);
+			}
+			return false;
+		}),
+		FGetActionCheckState(),
+		FIsActionButtonVisible::CreateLambda([this]() -> bool
+		{
+			if (const UPoseAsset* PoseAsset = PoseAssetPtr.Get())
+			{
+				return PoseAsset->SourceAnimation == nullptr;
+			}
+			return false;
+		}));
 }
 
 
-TSharedPtr<SWidget> SPoseViewer::OnGetContextMenuContent() const
+TSharedPtr<SWidget> SPoseViewer::OnGetContextMenuContent()
 {
 	const bool bShouldCloseWindowAfterMenuSelection = true;
 	FMenuBuilder MenuBuilder(bShouldCloseWindowAfterMenuSelection, UICommandList);
@@ -637,15 +683,68 @@ TSharedPtr<SWidget> SPoseViewer::OnGetContextMenuContent() const
 	MenuBuilder.AddMenuEntry(PoseEditorCommands.PasteAllNames);
 
 	MenuBuilder.BeginSection("PoseAction", LOCTEXT("SelectedItems", "Selected Item Actions"));
-	MenuBuilder.AddMenuEntry(PoseEditorCommands.UpdatePoseToCurrent);
-	MenuBuilder.AddMenuEntry(FGenericCommands::Get().Delete, NAME_None, LOCTEXT("DeletePoseButtonLabel", "Delete"), LOCTEXT("DeletePoseButtonTooltip", "Delete the selected pose(s)"));
-	MenuBuilder.AddMenuEntry(FGenericCommands::Get().Rename, NAME_None, LOCTEXT("RenamePoseButtonLabel", "Rename"), LOCTEXT("RenamePoseButtonTooltip", "Renames the selected pose"));
-	MenuBuilder.AddMenuEntry(FGenericCommands::Get().Paste, NAME_None, LOCTEXT("PastePoseNamesButtonLabel", "Paste Selected"), LOCTEXT("PastePoseNamesButtonTooltip", "Paste the selected pose names from clipBoard"));
+	{
+		MenuBuilder.AddMenuEntry(PoseEditorCommands.UpdatePoseToCurrent);
+		MenuBuilder.AddMenuEntry(FGenericCommands::Get().Delete, NAME_None, LOCTEXT("DeletePoseButtonLabel", "Delete"), LOCTEXT("DeletePoseButtonTooltip", "Delete the selected pose(s)"));
+		MenuBuilder.AddMenuEntry(FGenericCommands::Get().Rename, NAME_None, LOCTEXT("RenamePoseButtonLabel", "Rename"), LOCTEXT("RenamePoseButtonTooltip", "Renames the selected pose"));
+		MenuBuilder.AddMenuEntry(FGenericCommands::Get().Paste, NAME_None, LOCTEXT("PastePoseNamesButtonLabel", "Paste Selected"), LOCTEXT("PastePoseNamesButtonTooltip", "Paste the selected pose names from clipboard"));
+	}
+	MenuBuilder.EndSection();
+
+	if (const UPoseAsset* PoseAsset = PoseAssetPtr.Get())
+	{
+		if (PoseAsset->SourceAnimation == nullptr)
+		{
+			MenuBuilder.BeginSection("AddPoseAction", LOCTEXT("PosesSection", "Poses"));
+			{				
+				MenuBuilder.AddSubMenu(LOCTEXT("AddPoseSubMenu", "Add Pose"), LOCTEXT("PosesSubMenuToolTip", "Adding new Pose Related Actions"), FNewMenuDelegate::CreateLambda([this]( FMenuBuilder& SubMenuBuilder)
+				{
+					SubMenuBuilder.AddMenuEntry(FPoseEditorCommands::Get().AddPoseFromCurrent, NAME_None, TAttribute<FText>(), TAttribute<FText>(), FSlateIcon(FAppStyle::GetAppStyleSetName(), "Persona.AssetClass.Animation"));
+					SubMenuBuilder.AddMenuEntry(FPoseEditorCommands::Get().AddPoseFromReference, NAME_None, TAttribute<FText>(), TAttribute<FText>(), FSlateIcon(FAppStyle::GetAppStyleSetName(), "Persona.AssetClass.Skeleton"));
+					SubMenuBuilder.AddVerifiedEditableText(LOCTEXT("NewPoseLabel", "Pose Name"), LOCTEXT("NewPoseTooltip", "Tooltip"), FSlateIcon(),
+					TAttribute<FText>::CreateLambda([this]() { return FText::FromName(NewPoseName); }),
+					FOnVerifyTextChanged::CreateLambda([this](const FText& NewText, FText& OutMessage)-> bool
+					{
+						NewPoseName = FName(*NewText.ToString());
+						return IsNewPoseNameValid(OutMessage);
+					}));
+				}));
+			}
+			MenuBuilder.EndSection();
+		}
+	}
+	
 	MenuBuilder.EndSection();
 
 	return MenuBuilder.MakeWidget();
 }
 
+bool SPoseViewer::IsNewPoseNameValid(FText& OutReason) const
+{
+	if (const UPoseAsset* PoseAsset = PoseAssetPtr.Get())
+	{
+		if (PoseAsset->ContainsPose(NewPoseName))
+		{
+			OutReason = LOCTEXT("NewPoseAlreadyExistsMessage", "Pose with this name already exists");
+			return false;
+		}
+	}
+
+	if(!NewPoseName.IsValidObjectName(OutReason))
+	{
+		return false;
+	}
+
+	if (NewPoseName == NAME_None)
+	{
+		OutReason = LOCTEXT("NameNonePoseName", "Pose name cannot be empty or None");
+		return false;
+	}
+	
+	return true;
+}
+
+
 TSharedPtr<SWidget> SPoseViewer::OnGetContextMenuContentForCurveList() const
 {
 	const bool bShouldCloseWindowAfterMenuSelection = true;
@@ -926,4 +1025,48 @@ void SPoseViewer::UpdateSelectedPoseWithCurrent()
 
 	}
 }
+
+void SPoseViewer::AddPoseWithCurrent()
+{
+	if(UPoseAsset* PoseAsset = PoseAssetPtr.Get())
+	{
+		UDebugSkelMeshComponent* PreviewComponent = PreviewScenePtr.Pin()->GetPreviewMeshComponent();
+		USkeleton* Skeleton = PoseAsset->GetSkeleton();
+		if (PreviewComponent && Skeleton)
+		{
+			FSmartName PoseCurveSmartName;
+			PoseCurveSmartName.DisplayName = NewPoseName;
+			
+			Skeleton->VerifySmartName(USkeleton::AnimCurveMappingName, PoseCurveSmartName);
+			if (PoseCurveSmartName.IsValid())
+			{
+				PoseAsset->AddOrUpdatePose(PoseCurveSmartName, PreviewComponent);
+				NewPoseName = UPoseAsset::GetUniquePoseName(&(EditableSkeletonPtr.Pin()->GetSkeleton()));
+			}
+		}
+	}
+
+}
+
+void SPoseViewer::AddPoseWithReference()
+{
+	if(UPoseAsset* PoseAsset = PoseAssetPtr.Get())
+	{
+		const UDebugSkelMeshComponent* PreviewComponent = PreviewScenePtr.Pin()->GetPreviewMeshComponent();
+		USkeleton* Skeleton = PoseAsset->GetSkeleton();
+		if (PreviewComponent && Skeleton)
+		{
+			FSmartName PoseCurveSmartName;
+			PoseCurveSmartName.DisplayName = NewPoseName;
+			
+			Skeleton->VerifySmartName(USkeleton::AnimCurveMappingName, PoseCurveSmartName);
+			if (PoseCurveSmartName.IsValid())
+			{
+				PoseAsset->AddReferencePose(PoseCurveSmartName, PreviewComponent->GetReferenceSkeleton());
+				NewPoseName = UPoseAsset::GetUniquePoseName(&(EditableSkeletonPtr.Pin()->GetSkeleton()));
+			}
+		}
+	}	
+}
+
 #undef LOCTEXT_NAMESPACE
\ No newline at end of file
diff --git a/Engine/Source/Editor/Persona/Private/SPoseEditor.h b/Engine/Source/Editor/Persona/Private/SPoseEditor.h
--- a/Engine/Source/Editor/Persona/Private/SPoseEditor.h
+++ b/Engine/Source/Editor/Persona/Private/SPoseEditor.h
@@ -304,17 +304,19 @@ public:
 private:
 	void UpdateSelectedPoseWithCurrent();
 
+	void AddPoseWithCurrent();
+	void AddPoseWithReference();
+
 	void BindCommands();
 	void RestartPreviewComponent();
 
 	/** Handler for context menus */
-	TSharedPtr<SWidget> OnGetContextMenuContent() const;
+	TSharedPtr<SWidget> OnGetContextMenuContent();
 	/** Handler for curve list context menu*/
 	TSharedPtr<SWidget> OnGetContextMenuContentForCurveList() const;
 	/** Called when list double-clicked */
-	void OnListDoubleClick(TSharedPtr<FDisplayedPoseInfo> InItem);
-
-
+	void OnListDoubleClick(TSharedPtr<FDisplayedPoseInfo> InItem);	
+	bool IsNewPoseNameValid(FText& OutReason) const;
 	/**
 	* Clears and rebuilds the table, according to an optional search string
 	*
@@ -356,6 +358,8 @@ private:
 	/** A list of animation curve. Used by the PoseListView. */
 	TArray< TSharedPtr<FDisplayedCurveInfo> > CurveList;
 
+	FName NewPoseName;
+
 	/** Current text typed into NameFilterBox */
 	FText FilterText;
 
@@ -374,6 +378,7 @@ private:
 //////////////////////////////////////////////////////////////////////////
 // SPoseEditor
 
+
 /** Overall animation sequence editing widget */
 class SPoseEditor : public SAnimEditorBase
 {
diff --git a/Engine/Source/Runtime/Engine/Classes/Animation/PoseAsset.h b/Engine/Source/Runtime/Engine/Classes/Animation/PoseAsset.h
--- a/Engine/Source/Runtime/Engine/Classes/Animation/PoseAsset.h
+++ b/Engine/Source/Runtime/Engine/Classes/Animation/PoseAsset.h
@@ -137,7 +137,7 @@ private:
 	bool DeleteCurve(FSmartName CurveName);
 	bool InsertTrack(const FName& InTrackName, USkeleton* InSkeleton, const TArray<FTransform>& RefPose);
 	
-	bool FillUpSkeletonPose(FPoseData* PoseData, USkeleton* InSkeleton);
+	bool FillUpSkeletonPose(FPoseData* PoseData, const USkeleton* InSkeleton);
 	void RetrieveSourcePoseFromExistingPose(bool bAdditive, int32 InBasePoseIndex, const TArray<FTransform>& InBasePose, const TArray<float>& InBaseCurve);
 
 	// editor features for full pose <-> additive pose
@@ -259,8 +259,9 @@ public:
 	UFUNCTION(BlueprintPure, Category=PoseAsset)
 	void GetPoseNames(TArray<FName>& PoseNames) const;
 	
-	ENGINE_API bool AddOrUpdatePoseWithUniqueName(USkeletalMeshComponent* MeshComponent, FSmartName* OutPoseName = nullptr);
-	ENGINE_API void AddOrUpdatePose(const FSmartName& PoseName, USkeletalMeshComponent* MeshComponent, bool bUpdateCurves = true);
+	ENGINE_API bool AddOrUpdatePoseWithUniqueName(const USkeletalMeshComponent* MeshComponent, FSmartName* OutPoseName = nullptr);
+	ENGINE_API void AddOrUpdatePose(const FSmartName& PoseName, const USkeletalMeshComponent* MeshComponent, bool bUpdateCurves = true);
+	ENGINE_API void AddReferencePose(const FSmartName& PoseName, const FReferenceSkeleton& ReferenceSkeleton);
 	ENGINE_API void CreatePoseFromAnimation(class UAnimSequence* AnimSequence, const TArray<FSmartName>* InPoseNames = nullptr);
 
 	/** Contained poses are re-generated from the provided Animation Sequence*/
@@ -317,6 +318,9 @@ public:
 		OnPoseListChanged.Remove(Handle);
 	}
 
+	ENGINE_API static FName GetUniquePoseName(const USkeleton* Skeleton);
+	ENGINE_API static FSmartName GetUniquePoseSmartName(USkeleton* Skeleton);
+
 protected:
 	virtual void RemapTracksToNewSkeleton(USkeleton* NewSkeleton, bool bConvertSpaces) override;
 private: 
diff --git a/Engine/Source/Runtime/Engine/Private/Animation/PoseAsset.cpp b/Engine/Source/Runtime/Engine/Private/Animation/PoseAsset.cpp
--- a/Engine/Source/Runtime/Engine/Private/Animation/PoseAsset.cpp
+++ b/Engine/Source/Runtime/Engine/Private/Animation/PoseAsset.cpp
@@ -18,29 +18,6 @@
 
 #define LOCTEXT_NAMESPACE "PoseAsset"
 
-// utility function 
-#if WITH_EDITOR
-FSmartName GetUniquePoseName(USkeleton* Skeleton)
-{
-	check(Skeleton);
-	int32 NameIndex = 0;
-
-	SmartName::UID_Type NewUID;
-	FName NewName;
-
-	do
-	{
-		NewName = FName(*FString::Printf(TEXT("Pose_%d"), NameIndex++));
-		NewUID = Skeleton->GetUIDByName(USkeleton::AnimCurveMappingName, NewName);
-	} while (NewUID != SmartName::MaxUID);
-
-	// if found, 
-	FSmartName NewPoseName;
-	Skeleton->AddSmartNameAndModify(USkeleton::AnimCurveMappingName, NewName, NewPoseName);
-
-	return NewPoseName;
-}
-#endif 
 //////////////////////////////////////////////////////////////////////////////////////////////////////////////////
 //
 // FPoseDataContainer
@@ -188,7 +165,7 @@ bool FPoseDataContainer::InsertTrack(const FName& InTrackName, USkeleton* InSkel
 	return false;
 }
 
-bool FPoseDataContainer::FillUpSkeletonPose(FPoseData* PoseData, USkeleton* InSkeleton)
+bool FPoseDataContainer::FillUpSkeletonPose(FPoseData* PoseData, const USkeleton* InSkeleton)
 {
 	if (PoseData)
 	{
@@ -1172,11 +1149,11 @@ void UPoseAsset::PostProcessData()
 	UpdateTrackBoneIndices();
 }
 
-bool UPoseAsset::AddOrUpdatePoseWithUniqueName(USkeletalMeshComponent* MeshComponent, FSmartName* OutPoseName /*= nullptr*/)
+bool UPoseAsset::AddOrUpdatePoseWithUniqueName(const USkeletalMeshComponent* MeshComponent, FSmartName* OutPoseName /*= nullptr*/)
 {
 	bool bSavedAdditivePose = bAdditivePose;
 
-	FSmartName NewPoseName = GetUniquePoseName(GetSkeleton());
+	FSmartName NewPoseName = GetUniquePoseSmartName(GetSkeleton());
 	AddOrUpdatePose(NewPoseName, MeshComponent);
 
 	if (OutPoseName)
@@ -1191,10 +1168,28 @@ bool UPoseAsset::AddOrUpdatePoseWithUniqueName(USkeletalMeshComponent* MeshCompo
 	return true;
 }
 
+void UPoseAsset::AddReferencePose(const FSmartName& PoseName, const FReferenceSkeleton& RefSkeleton)
+{
+	TArray<FTransform> BoneTransforms;
+	TArray<FName> TrackNames;
+
+	const TArray<FTransform>& ReferencePose = RefSkeleton.GetRefBonePose();
+	for (int32 BoneIndex = 0; BoneIndex < RefSkeleton.GetNum(); ++BoneIndex)
+	{
+		TrackNames.Add(RefSkeleton.GetBoneName(BoneIndex));
+		BoneTransforms.Add(ReferencePose[BoneIndex]);
+	}
+			
+	TArray<float> NewCurveValues;
+	NewCurveValues.AddZeroed(PoseContainer.Curves.Num());
 
-void UPoseAsset::AddOrUpdatePose(const FSmartName& PoseName, USkeletalMeshComponent* MeshComponent, bool bUpdateCurves)
+	AddOrUpdatePose(PoseName, TrackNames, BoneTransforms, NewCurveValues);
+	PostProcessData();
+}
+
+void UPoseAsset::AddOrUpdatePose(const FSmartName& PoseName, const USkeletalMeshComponent* MeshComponent, bool bUpdateCurves)
 {
-	USkeleton* MySkeleton = GetSkeleton();
+	const USkeleton* MySkeleton = GetSkeleton();
 	if (MySkeleton && MeshComponent && MeshComponent->GetSkeletalMeshAsset())
 	{
 		TArray<FName> TrackNames;
@@ -1228,8 +1223,8 @@ void UPoseAsset::AddOrUpdatePose(const FSmartName& PoseName, USkeletalMeshCompon
  
 			for (int32 NewCurveIndex = 0; NewCurveIndex < NewCurveValues.Num(); ++NewCurveIndex)
 			{
-				FAnimCurveBase& Curve = PoseContainer.Curves[NewCurveIndex];
-				SmartName::UID_Type CurveUID = Mapping->FindUID(Curve.Name.DisplayName);
+				const FAnimCurveBase& Curve = PoseContainer.Curves[NewCurveIndex];
+				const SmartName::UID_Type CurveUID = Mapping->FindUID(Curve.Name.DisplayName);
 				if (CurveUID != SmartName::MaxUID)
 				{
 					const float MeshCurveValue = MeshCurves.Get(CurveUID);
@@ -1247,12 +1242,13 @@ void UPoseAsset::AddOrUpdatePose(const FSmartName& PoseName, USkeletalMeshCompon
 
 void UPoseAsset::AddOrUpdatePose(const FSmartName& PoseName, const TArray<FName>& TrackNames, const TArray<FTransform>& LocalTransform, const TArray<float>& CurveValues)
 {
-	USkeleton* MySkeleton = GetSkeleton();
+	const USkeleton* MySkeleton = GetSkeleton();
 	if (MySkeleton)
 	{
 		// first combine track, we want to make sure all poses contains tracks with this
 		CombineTracks(TrackNames);
 
+		const bool bNewPose = PoseContainer.FindPoseData(PoseName) == nullptr;
 		FPoseData* PoseData = PoseContainer.FindOrAddPoseData(PoseName);
 		// now copy all transform back to it. 
 		check(PoseData);
@@ -1261,7 +1257,7 @@ void UPoseAsset::AddOrUpdatePose(const FSmartName& PoseName, const TArray<FName>
 		// but We have to add all tracks to match poses container
 		// TrackNames.Num() is subset of PoseContainer.Tracks.Num()
 		// Above CombineTracks will combine both
-		int32 TotalTracks = PoseContainer.Tracks.Num();
+		const int32 TotalTracks = PoseContainer.Tracks.Num();
 		PoseData->SourceLocalSpacePose.Reset(TotalTracks);
 		PoseData->SourceLocalSpacePose.AddUninitialized(TotalTracks);
 		PoseData->SourceLocalSpacePose.SetNumZeroed(TotalTracks, true);
@@ -1280,10 +1276,15 @@ void UPoseAsset::AddOrUpdatePose(const FSmartName& PoseName, const TArray<FName>
 			// now get poseData track index
 			const FName& TrackName = TrackNames[Index];
 			//int32 SkeletonIndex = RefSkeleton.FindBoneIndex(TrackName);
-			int32 InternalTrackIndex = PoseContainer.Tracks.Find(TrackName);
+			const int32 InternalTrackIndex = PoseContainer.Tracks.Find(TrackName);
 			// copy to the internal track index
 			PoseData->SourceLocalSpacePose[InternalTrackIndex] = LocalTransform[Index];
 		}
+
+		if (bNewPose)
+		{
+			OnPoseListChanged.Broadcast();
+		}
 	}
 }
 void UPoseAsset::PostEditChangeProperty(FPropertyChangedEvent& PropertyChangedEvent)
@@ -1389,8 +1390,11 @@ void UPoseAsset::CreatePoseFromAnimation(class UAnimSequence* AnimSequence, cons
 			// reinitialize, now we're making new pose from this animation
 			Reinitialize();
 
-			const int32 NumPoses = AnimSequence->GetNumberOfSampledKeys();
-
+			int32 NumPoses = AnimSequence->GetNumberOfSampledKeys();
+			if(InPoseNames && InPoseNames->Num() > NumPoses)
+			{
+				NumPoses=InPoseNames->Num();
+			}
 			// make sure we have more than one pose
 			if (NumPoses > 0)
 			{
@@ -1441,7 +1445,7 @@ void UPoseAsset::CreatePoseFromAnimation(class UAnimSequence* AnimSequence, cons
 				// add to skeleton UID, so that it knows the curve data
 				for (int32 PoseIndex = 0; PoseIndex < NumPoses; ++PoseIndex)
 				{
-					FSmartName NewPoseName = (InPoseNames && InPoseNames->IsValidIndex(PoseIndex))? (*InPoseNames)[PoseIndex] : GetUniquePoseName(TargetSkeleton);
+					FSmartName NewPoseName = (InPoseNames && InPoseNames->IsValidIndex(PoseIndex))? (*InPoseNames)[PoseIndex] : GetUniquePoseSmartName(TargetSkeleton);
 					// now get rawanimationdata, and each key is converted to new pose
 					for (int32 TrackIndex = 0; TrackIndex < NumTracks; ++TrackIndex)
 					{
@@ -1798,7 +1802,37 @@ void FPoseDataContainer::DeleteTrack(int32 TrackIndex)
 #endif // WITH_EDITOR
 	}
 }
+
 #if WITH_EDITOR
+FName UPoseAsset::GetUniquePoseName(const USkeleton* Skeleton)
+{
+	FName NewName = NAME_None;
+	if (Skeleton)
+	{
+		int32 NameIndex = 0;
+		SmartName::UID_Type NewUID;
+		do
+		{
+			NewName = FName(*FString::Printf(TEXT("Pose_%d"), NameIndex++));
+			NewUID = Skeleton->GetUIDByName(USkeleton::AnimCurveMappingName, NewName);
+		} while (NewUID != SmartName::MaxUID);
+	}
+
+	return NewName;
+}
+
+FSmartName UPoseAsset::GetUniquePoseSmartName(USkeleton* Skeleton)
+{
+	const FName NewName = GetUniquePoseName(Skeleton);
+	FSmartName NewPoseName;
+	if (NewName != NAME_None)
+	{
+		Skeleton->AddSmartNameAndModify(USkeleton::AnimCurveMappingName, NewName, NewPoseName);
+	}
+
+	return NewPoseName;
+}
+
 void UPoseAsset::RemapTracksToNewSkeleton(USkeleton* NewSkeleton, bool bConvertSpaces)
 {
 	Super::RemapTracksToNewSkeleton(NewSkeleton, bConvertSpaces);
