Fix for the preview mesh on a skeleton being reset if it comes from a compatible skeleton

diff --git a/Engine/Source/Runtime/Engine/Private/Animation/Skeleton.cpp b/Engine/Source/Runtime/Engine/Private/Animation/Skeleton.cpp
--- a/Engine/Source/Runtime/Engine/Private/Animation/Skeleton.cpp
+++ b/Engine/Source/Runtime/Engine/Private/Animation/Skeleton.cpp
@@ -1432,7 +1432,7 @@ USkeletalMesh* USkeleton::GetPreviewMesh(bool bFindIfNotSet/*=false*/)
 #if WITH_EDITORONLY_DATA
 	USkeletalMesh* PreviewMesh = PreviewSkeletalMesh.LoadSynchronous();
 
-	if(PreviewMesh && PreviewMesh->GetSkeleton() != this) // fix mismatched skeleton
+	if(PreviewMesh && !IsCompatible(PreviewMesh->GetSkeleton())) // fix mismatched skeleton
 	{
 		PreviewSkeletalMesh.Reset();
 		PreviewMesh = nullptr;
