Template animation blueprints can now be used on skeletal meshes standalone

Prior to this template anim BPs would not be instantiated due to their lack of skeleton. This change defaults to the skeletal meshes skeleton on re-initialization.

diff --git a/Engine/Source/Runtime/Engine/Private/Components/SkeletalMeshComponent.cpp b/Engine/Source/Runtime/Engine/Private/Components/SkeletalMeshComponent.cpp
--- a/Engine/Source/Runtime/Engine/Private/Components/SkeletalMeshComponent.cpp
+++ b/Engine/Source/Runtime/Engine/Private/Components/SkeletalMeshComponent.cpp
@@ -532,10 +532,16 @@ void USkeletalMeshComponent::FinalizeAttributeEvaluationResults(const FBoneConta
 
 bool USkeletalMeshComponent::NeedToSpawnAnimScriptInstance() const
 {
+	USkeletalMesh* MeshAsset = GetSkeletalMeshAsset();
 	IAnimClassInterface* AnimClassInterface = IAnimClassInterface::GetFromClass(AnimClass);
-	const USkeleton* AnimSkeleton = (AnimClassInterface) ? AnimClassInterface->GetTargetSkeleton() : nullptr;
-	const bool bSkeletonCompatible = (GetSkeletalMeshAsset() && AnimSkeleton && GetSkeletalMeshAsset()->GetSkeleton()) ? GetSkeletalMeshAsset()->GetSkeleton()->IsCompatible(AnimSkeleton) : false;
-	const bool bSkelMeshCompatible = (GetSkeletalMeshAsset() && AnimSkeleton) ? AnimSkeleton->IsCompatibleMesh(GetSkeletalMeshAsset(), false) : false;
+	USkeleton* AnimSkeleton = (AnimClassInterface) ? AnimClassInterface->GetTargetSkeleton() : nullptr;
+	if(AnimSkeleton == nullptr)
+	{
+		// Fall back to mesh skeleton (anim BP could have been a template)
+		AnimSkeleton = MeshAsset ? MeshAsset->GetSkeleton() : nullptr;
+	}
+	const bool bSkeletonCompatible = (MeshAsset && AnimSkeleton && MeshAsset->GetSkeleton()) ? MeshAsset->GetSkeleton()->IsCompatible(AnimSkeleton) : false;
+	const bool bSkelMeshCompatible = (MeshAsset && AnimSkeleton) ? AnimSkeleton->IsCompatibleMesh(MeshAsset, false) : false;
 	const bool bAnimSkelValid = !AnimClassInterface || (bSkeletonCompatible && bSkelMeshCompatible);
 
 	if (AnimationMode == EAnimationMode::AnimationBlueprint && AnimClass && bAnimSkelValid)
