Precache Skinned Mesh component PSOs

diff --git a/Engine/Source/Runtime/Engine/Classes/Components/SkinnedMeshComponent.h b/Engine/Source/Runtime/Engine/Classes/Components/SkinnedMeshComponent.h
--- a/Engine/Source/Runtime/Engine/Classes/Components/SkinnedMeshComponent.h
+++ b/Engine/Source/Runtime/Engine/Classes/Components/SkinnedMeshComponent.h
@@ -1071,6 +1071,9 @@ public:
 	 */
 	void SetForceWireframe(bool InForceWireframe);
 
+	/** Precache all PSOs which can be used by the component */
+	virtual void PrecachePSOs() override;
+	
 #if WITH_EDITOR
 	/** Return value of SectionIndexPreview  */
 	int32 GetSectionPreview() const { return SectionIndexPreview;  }
diff --git a/Engine/Source/Runtime/Engine/Private/Components/SkinnedMeshComponent.cpp b/Engine/Source/Runtime/Engine/Private/Components/SkinnedMeshComponent.cpp
--- a/Engine/Source/Runtime/Engine/Private/Components/SkinnedMeshComponent.cpp
+++ b/Engine/Source/Runtime/Engine/Private/Components/SkinnedMeshComponent.cpp
@@ -30,6 +30,7 @@
 #include "AnimationRuntime.h"
 #include "Animation/SkinWeightProfileManager.h"
 #include "GPUSkinCache.h"
+#include "PipelineStateCache.h"
 #include "SkeletalRender.h"
 
 DEFINE_LOG_CATEGORY_STATIC(LogSkinnedMeshComp, Log, All);
@@ -543,6 +544,77 @@ void USkinnedMeshComponent::PostLoad()
 		// Set default for new bSetMeshDeformer flag.
 		bSetMeshDeformer = MeshDeformer != nullptr;
 	}
+
+	PrecachePSOs();
+}
+
+void USkinnedMeshComponent::PrecachePSOs()
+{
+	if (GetSkinnedAsset() == nullptr ||
+		GetSkinnedAsset()->GetResourceForRendering() == nullptr ||
+		!FApp::CanEverRender() || 
+		!PipelineStateCache::IsPSOPrecachingEnabled())
+	{
+		return;
+	}
+
+	ERHIFeatureLevel::Type FeatureLevel = GetWorld() ? GetWorld()->FeatureLevel.GetValue() : GMaxRHIFeatureLevel;
+	FSkeletalMeshRenderData* SkelMeshRenderData = GetSkinnedAsset()->GetResourceForRendering();
+	int32 MinLODIndex = ComputeMinLOD();
+
+	bool bAnySectionCastsShadows = false;
+
+	struct VFsPerMaterialData
+	{
+		int16 MaterialIndex;
+		TArray<const FVertexFactoryType*, TInlineAllocator<2>> VertexFactoryTypes;
+	};
+	TArray<VFsPerMaterialData, TInlineAllocator<4>> VFsPerMaterials;
+
+	for (int32 LODIndex = MinLODIndex; LODIndex < SkelMeshRenderData->LODRenderData.Num(); LODIndex++)
+	{
+		FSkeletalMeshLODRenderData& LODRenderData = SkelMeshRenderData->LODRenderData[LODIndex];
+		
+		// Check all render sections for the used material indices
+		for (FSkelMeshRenderSection& RenderSection : LODRenderData.RenderSections)
+		{
+			bAnySectionCastsShadows |= RenderSection.bCastShadow;
+
+			uint16 MaterialIndex = RenderSection.MaterialIndex;
+
+			VFsPerMaterialData* VFsPerMaterial = VFsPerMaterials.FindByPredicate([MaterialIndex](const VFsPerMaterialData& Other) { return Other.MaterialIndex == MaterialIndex; });
+			if (VFsPerMaterial == nullptr)
+			{
+				VFsPerMaterial = &VFsPerMaterials.AddDefaulted_GetRef();
+				VFsPerMaterial->MaterialIndex = MaterialIndex;
+			}
+
+			// Find all the possible used vertex factory types needed to render this render section
+			if (bRenderStatic || ShouldCPUSkin())
+			{
+				// Force static from GPU point of view
+				VFsPerMaterial->VertexFactoryTypes.AddUnique(&FLocalVertexFactory::StaticType);
+			}
+			else if (!SkelMeshRenderData->RequiresCPUSkinning(FeatureLevel, MinLODIndex))
+			{
+				// Add all the vertex factories which can be used for gpu skinning
+				FSkeletalMeshObjectGPUSkin::GetUsedVertexFactories(LODRenderData, RenderSection, FeatureLevel, VFsPerMaterial->VertexFactoryTypes);
+			}
+		}
+	}
+
+	FPSOPrecacheParams PrecachePSOParams;
+	SetupPrecachePSOParams(PrecachePSOParams);
+	PrecachePSOParams.bCastShadow = PrecachePSOParams.bCastShadow && bAnySectionCastsShadows;
+
+	for (VFsPerMaterialData& VFsPerMaterial : VFsPerMaterials)
+	{
+		UMaterialInterface* MaterialInterface = GetMaterial(VFsPerMaterial.MaterialIndex);
+		if (MaterialInterface)
+		{
+			MaterialInterface->PrecachePSOs(VFsPerMaterial.VertexFactoryTypes, PrecachePSOParams);
+		}
+	}
 }
 
 void USkinnedMeshComponent::OnRegister()
@@ -1910,6 +1982,8 @@ void USkinnedMeshComponent::SetSkinnedAssetAndUpdate(USkinnedAsset* InSkinnedAss
 			ExternalMorphWeightData.Empty();
 		}
 
+		PrecachePSOs();
+
 		// Re-init the MeshDeformer which might come from the SkelMesh.
 		UMeshDeformer* ActiveMeshDeformer = GetActiveMeshDeformer();
 		MeshDeformerInstance = ActiveMeshDeformer != nullptr ? ActiveMeshDeformer->CreateInstance(this, MeshDeformerInstanceSettings) : nullptr;
diff --git a/Engine/Source/Runtime/Engine/Private/SkeletalRenderGPUSkin.cpp b/Engine/Source/Runtime/Engine/Private/SkeletalRenderGPUSkin.cpp
--- a/Engine/Source/Runtime/Engine/Private/SkeletalRenderGPUSkin.cpp
+++ b/Engine/Source/Runtime/Engine/Private/SkeletalRenderGPUSkin.cpp
@@ -1731,6 +1731,19 @@ public:
 	const FSkeletalMeshObjectGPUSkin::FVertexFactoryBuffers VertexBuffers;
 };
 
+static const FVertexFactoryType* GetVertexFactoryType(FSkeletalMeshLODRenderData& LODRenderData, ERHIFeatureLevel::Type FeatureLevel)
+{
+	GPUSkinBoneInfluenceType BoneInfluenceType = LODRenderData.SkinWeightVertexBuffer.GetBoneInfluenceType();
+	if (BoneInfluenceType == GPUSkinBoneInfluenceType::DefaultBoneInfluence)
+	{
+		return &TGPUSkinVertexFactory<GPUSkinBoneInfluenceType::DefaultBoneInfluence>::StaticType;
+	}
+	else
+	{
+		return &TGPUSkinVertexFactory<GPUSkinBoneInfluenceType::UnlimitedBoneInfluence>::StaticType;
+	}
+}
+
 /**
  * Creates a vertex factory entry for the given type and initialize it on the render thread
  */
@@ -1883,6 +1896,19 @@ static void UpdateVertexFactoryMorph(TArray<TUniquePtr<FGPUBaseSkinVertexFactory
 
 // APEX cloth
 
+static const FVertexFactoryType* GetVertexFactoryTypeCloth(FSkeletalMeshLODRenderData& LODRenderData, ERHIFeatureLevel::Type FeatureLevel)
+{
+	GPUSkinBoneInfluenceType BoneInfluenceType = LODRenderData.SkinWeightVertexBuffer.GetBoneInfluenceType();
+	if (BoneInfluenceType == GPUSkinBoneInfluenceType::DefaultBoneInfluence)
+	{
+		return &TGPUSkinAPEXClothVertexFactory<GPUSkinBoneInfluenceType::DefaultBoneInfluence>::StaticType;
+	}
+	else
+	{
+		return &TGPUSkinAPEXClothVertexFactory<GPUSkinBoneInfluenceType::UnlimitedBoneInfluence>::StaticType;
+	}
+}
+
 /**
  * Creates a vertex factory entry for the given type and initialize it on the render thread
  */
@@ -1947,6 +1973,25 @@ static void UpdateVertexFactoryCloth(TArray<TUniquePtr<FGPUBaseSkinAPEXClothVert
 }
 
 
+void FSkeletalMeshObjectGPUSkin::GetUsedVertexFactories(FSkeletalMeshLODRenderData& LODRenderData, FSkelMeshRenderSection& RenderSection, ERHIFeatureLevel::Type InFeatureLevel, TArray<const FVertexFactoryType*, TInlineAllocator<2>>& VertexFactoryTypes)
+{
+	// TODO: optimize types to add based on if SkinCache is used and allowed for given Mesh (UE-164762)
+	VertexFactoryTypes.AddUnique(&FGPUSkinPassthroughVertexFactory::StaticType);
+
+	// Add GPU skin cloth vertex factory type is needed
+	const EShaderPlatform ShaderPlatform = GetFeatureLevelShaderPlatform(InFeatureLevel);
+	const bool bClothEnabled = FGPUBaseSkinAPEXClothVertexFactory::IsClothEnabled(ShaderPlatform);
+	if (bClothEnabled && RenderSection.HasClothingData())
+	{
+		VertexFactoryTypes.AddUnique(GetVertexFactoryTypeCloth(LODRenderData, InFeatureLevel));
+	}
+	else
+	{
+		// Add GPU skin vertex factory type
+		VertexFactoryTypes.AddUnique(GetVertexFactoryType(LODRenderData, InFeatureLevel));
+	}
+}
+
 /**
  * Determine the current vertex buffers valid for the current LOD
  *
diff --git a/Engine/Source/Runtime/Engine/Private/SkeletalRenderGPUSkin.h b/Engine/Source/Runtime/Engine/Private/SkeletalRenderGPUSkin.h
--- a/Engine/Source/Runtime/Engine/Private/SkeletalRenderGPUSkin.h
+++ b/Engine/Source/Runtime/Engine/Private/SkeletalRenderGPUSkin.h
@@ -451,6 +451,8 @@ public:
 	virtual void RefreshClothingTransforms(const FMatrix& InNewLocalToWorld, uint32 FrameNumber) override;
 	virtual void UpdateSkinWeightBuffer(USkinnedMeshComponent* InMeshComponent) override;
 
+	static void GetUsedVertexFactories(FSkeletalMeshLODRenderData& LODRenderData, FSkelMeshRenderSection& RenderSection, ERHIFeatureLevel::Type InFeatureLevel, TArray<const FVertexFactoryType*, TInlineAllocator<2>>& VertexFactoryTypes);
+
 protected:
 	friend class FSkeletalMeshDeformerHelpers;
 
