Workaround DXC/SPIR-V bug in HLSL material translator: Make integer type explicit in ternary operator for shading model indices.

#rnx

diff --git a/Engine/Source/Runtime/Engine/Private/Materials/HLSLMaterialDerivativeAutogen.cpp b/Engine/Source/Runtime/Engine/Private/Materials/HLSLMaterialDerivativeAutogen.cpp
--- a/Engine/Source/Runtime/Engine/Private/Materials/HLSLMaterialDerivativeAutogen.cpp
+++ b/Engine/Source/Runtime/Engine/Private/Materials/HLSLMaterialDerivativeAutogen.cpp
@@ -1238,11 +1238,24 @@ int32 FMaterialDerivativeAutogen::GenerateIfFunc(FHLSLMaterialTranslator& Transl
 		}
 		else
 		{
-			CodeFinite = FString::Printf(
-				TEXT("(%s ? %s : %s)"),
-				*CompareGreaterEqual,
-				*GreaterFinite, *LessFinite
-			);
+			if (ResultType == MCT_ShadingModel)
+			{
+				// @lh-todo: Workaround SPIR-V bug in DXC: Wrong literal type is deduced during implicit type deduction from int to uint
+				// GitHub PR: [REDACTED_PROFILE_URL]/DirectXShaderCompiler/pull/4626
+				CodeFinite = FString::Printf(
+					TEXT("(%s ? (uint)%s : (uint)%s)"),
+					*CompareGreaterEqual,
+					*GreaterFinite, *LessFinite
+				);
+			}
+			else
+			{
+				CodeFinite = FString::Printf(
+					TEXT("(%s ? %s : %s)"),
+					*CompareGreaterEqual,
+					*GreaterFinite, *LessFinite
+				);
+			}
 		}
 	}
 
