[SkeletonRemapping] Fixed some crash a user had while packaging, due to a nullptr weak pointer.

diff --git a/Engine/Source/Runtime/Engine/Private/Animation/Skeleton.cpp b/Engine/Source/Runtime/Engine/Private/Animation/Skeleton.cpp
--- a/Engine/Source/Runtime/Engine/Private/Animation/Skeleton.cpp
+++ b/Engine/Source/Runtime/Engine/Private/Animation/Skeleton.cpp
@@ -186,7 +186,7 @@ FSkeletonRemapping::FSkeletonRemapping(const USkeleton* InSourceSkeleton, const
 	: SourceSkeleton(InSourceSkeleton)
 	, TargetSkeleton(InTargetSkeleton)
 {
-	if (!InSourceSkeleton || !InTargetSkeleton)
+	if (!GetSourceSkeleton().IsValid() || !GetTargetSkeleton().IsValid())
 	{
 		return;
 	}
