Added some safety checks in the external morph set handling.

diff --git a/Engine/Source/Runtime/Engine/Private/Components/SkinnedMeshComponent.cpp b/Engine/Source/Runtime/Engine/Private/Components/SkinnedMeshComponent.cpp
--- a/Engine/Source/Runtime/Engine/Private/Components/SkinnedMeshComponent.cpp
+++ b/Engine/Source/Runtime/Engine/Private/Components/SkinnedMeshComponent.cpp
@@ -698,7 +698,11 @@ void USkinnedMeshComponent::AddExternalMorphSet(int32 LOD, int32 ID, TSharedPtr<
 	{
 		ResizeExternalMorphTargetSets();
 	}
-	ExternalMorphSets[LOD].Add(ID, MorphSet);
+
+	if (ExternalMorphSets.IsValidIndex(LOD))
+	{
+		ExternalMorphSets[LOD].Add(ID, MorphSet);
+	}
 }
 
 void USkinnedMeshComponent::RemoveExternalMorphSet(int32 LOD, int32 ID)
@@ -708,13 +712,21 @@ void USkinnedMeshComponent::RemoveExternalMorphSet(int32 LOD, int32 ID)
 	{
 		ResizeExternalMorphTargetSets();
 	}
-	ExternalMorphSets[LOD].Remove(ID);
+
+	if (ExternalMorphSets.IsValidIndex(LOD))
+	{
+		ExternalMorphSets[LOD].Remove(ID);
+	}
 }
 
 bool USkinnedMeshComponent::HasExternalMorphSet(int32 LOD, int32 ID) const
 {
-	check(LOD < ExternalMorphSets.Num());
-	return (ExternalMorphSets[LOD].Find(ID) != nullptr);
+	if (ExternalMorphSets.IsValidIndex(LOD))
+	{
+		return (ExternalMorphSets[LOD].Find(ID) != nullptr);
+	}
+
+	return false;
 }
 
 void USkinnedMeshComponent::ClearExternalMorphSets(int32 LOD)
@@ -724,7 +736,10 @@ void USkinnedMeshComponent::ClearExternalMorphSets(int32 LOD)
 	{
 		ResizeExternalMorphTargetSets();
 	}
-	ExternalMorphSets[LOD].Empty();
+	if (ExternalMorphSets.IsValidIndex(LOD))
+	{
+		ExternalMorphSets[LOD].Empty();
+	}
 }
 
 void USkinnedMeshComponent::ResizeExternalMorphTargetSets()
