Preventing MetaSounds and Source Buses from showing up as options for Dialogue Waves

diff --git a/Engine/Source/Runtime/Engine/Classes/Sound/DialogueWave.h b/Engine/Source/Runtime/Engine/Classes/Sound/DialogueWave.h
--- a/Engine/Source/Runtime/Engine/Classes/Sound/DialogueWave.h
+++ b/Engine/Source/Runtime/Engine/Classes/Sound/DialogueWave.h
@@ -52,7 +52,7 @@ struct ENGINE_API FDialogueContextMapping
 	FDialogueContext Context;
 
 	/** The soundwave to play for this dialogue. */
-	UPROPERTY(EditAnywhere, Category=DialogueContextMapping )
+	UPROPERTY(EditAnywhere, Category=DialogueContextMapping, meta = (DisallowedClasses="MetaSoundSource, SoundSourceBus"))
 	TObjectPtr<USoundWave> SoundWave;
 
 	/**
