WorldPartition: Mark level packages created through UWorldPartitionLevelStreamingDynamic::LoadInEditor() as transient
* This prevents the editor from asking to save those temp packages

diff --git a/Engine/Source/Runtime/Engine/Private/WorldPartition/WorldPartitionLevelStreamingDynamic.cpp b/Engine/Source/Runtime/Engine/Private/WorldPartition/WorldPartitionLevelStreamingDynamic.cpp
--- a/Engine/Source/Runtime/Engine/Private/WorldPartition/WorldPartitionLevelStreamingDynamic.cpp
+++ b/Engine/Source/Runtime/Engine/Private/WorldPartition/WorldPartitionLevelStreamingDynamic.cpp
@@ -50,7 +50,7 @@ UWorldPartitionLevelStreamingDynamic* UWorldPartitionLevelStreamingDynamic::Load
 	check(World->WorldType == EWorldType::Editor);
 	UWorldPartitionLevelStreamingDynamic* LevelStreaming = NewObject<UWorldPartitionLevelStreamingDynamic>(World, LevelStreamingName, RF_Transient);
 	
-	FString PackageName = FString::Printf(TEXT("/Memory/%s"), *LevelStreamingName.ToString());
+	FString PackageName = FString::Printf(TEXT("/Temp/%s"), *LevelStreamingName.ToString());
 	TSoftObjectPtr<UWorld> WorldAsset(FSoftObjectPath(FString::Printf(TEXT("%s.%s"), *PackageName, *World->GetName())));
 	LevelStreaming->SetWorldAsset(WorldAsset);
 	
@@ -59,6 +59,13 @@ UWorldPartitionLevelStreamingDynamic* UWorldPartitionLevelStreamingDynamic::Load
 	LevelStreaming->SetShouldBeVisibleInEditor(true);
 	World->AddStreamingLevel(LevelStreaming);
 	World->FlushLevelStreaming();
+
+	// Mark the level package as transient to prevent the editor from asking to save it.
+	ULevel* Level = LevelStreaming->GetLoadedLevel();
+	if (Level)
+	{
+		Level->GetPackage()->SetFlags(RF_Transient);
+	}
 	
 	return LevelStreaming;
 }
