World Partition - HLOD: Improved logging of intermediate HLOD hash values

diff --git a/Engine/Source/Runtime/Engine/Private/WorldPartition/HLOD/HLODBuilder.cpp b/Engine/Source/Runtime/Engine/Private/WorldPartition/HLOD/HLODBuilder.cpp
--- a/Engine/Source/Runtime/Engine/Private/WorldPartition/HLOD/HLODBuilder.cpp
+++ b/Engine/Source/Runtime/Engine/Private/WorldPartition/HLOD/HLODBuilder.cpp
@@ -83,12 +83,13 @@ uint32 UHLODBuilder::ComputeHLODHash(const UActorComponent* InSourceComponent) c
 
 		// CRC component
 		ComponentCRC = UHLODProxy::GetCRC(StaticMeshComponent);
-		UE_LOG(LogHLODBuilder, VeryVerbose, TEXT("     - StaticMeshComponent (%s) = %x"), *StaticMeshComponent->GetName(), ComponentCRC);
+		UE_LOG(LogHLODBuilder, VeryVerbose, TEXT("     - Static Mesh Component (%s) = %x"), *StaticMeshComponent->GetName(), ComponentCRC);
 
 		// CRC static mesh
 		if (UStaticMesh* StaticMesh = StaticMeshComponent->GetStaticMesh())
 		{
 			ComponentCRC = UHLODProxy::GetCRC(StaticMesh, ComponentCRC);
+			UE_LOG(LogHLODBuilder, VeryVerbose, TEXT("     - Static Mesh (%s) = %x"), *StaticMesh->GetName(), ComponentCRC);
 		}
 
 		// CRC materials
@@ -99,12 +100,14 @@ uint32 UHLODBuilder::ComputeHLODHash(const UActorComponent* InSourceComponent) c
 			if (MaterialInterface)
 			{
 				ComponentCRC = UHLODProxy::GetCRC(MaterialInterface, ComponentCRC);
+				UE_LOG(LogHLODBuilder, VeryVerbose, TEXT("     - Material (%s) = %x"), *MaterialInterface->GetName(), ComponentCRC);
 
 				TArray<UTexture*> Textures;
 				MaterialInterface->GetUsedTextures(Textures, EMaterialQualityLevel::High, true, ERHIFeatureLevel::SM5, true);
 				for (UTexture* Texture : Textures)
 				{
 					ComponentCRC = UHLODProxy::GetCRC(Texture, ComponentCRC);
+					UE_LOG(LogHLODBuilder, VeryVerbose, TEXT("     - Texture (%s) = %x"), *Texture->GetName(), ComponentCRC);
 				}
 			}
 		}
