Fixed Material Expression Parameters graph pins state not being kept in sync across multiple nodes that reference the same Parameter.

diff --git a/Engine/Source/Runtime/Engine/Private/Materials/MaterialExpressions.cpp b/Engine/Source/Runtime/Engine/Private/Materials/MaterialExpressions.cpp
--- a/Engine/Source/Runtime/Engine/Private/Materials/MaterialExpressions.cpp
+++ b/Engine/Source/Runtime/Engine/Private/Materials/MaterialExpressions.cpp
@@ -1498,11 +1498,12 @@ void UMaterialExpression::PostEditChangeProperty(FPropertyChangedEvent& Property
 
 		if (PinIndex > -1)
 		{
+			const FString& NewDefaultValue = GetInputPinDefaultValue(PinIndex);
+
 			// Update the pin value of the expression input
 			UEdGraphPin* Pin = GraphNode->GetPinAt(PinIndex);
-			check(Pin);
 			Pin->Modify();
-			Pin->DefaultValue = GetInputPinDefaultValue(PinIndex);
+			Pin->DefaultValue = NewDefaultValue;
 
 			// If the property is linked as inline toggle to another property, both pins need updating to reflect the change.
 			bool bInlineEditConditionToggle = MemberPropertyThatChanged->HasMetaData(TEXT("InlineEditConditionToggle"));
@@ -1511,6 +1512,30 @@ void UMaterialExpression::PostEditChangeProperty(FPropertyChangedEvent& Property
 			{
 				CastChecked<UMaterialGraphNode>(GraphNode)->RecreateAndLinkNode();
 			}
+
+			// If this expression refers to a parameter, we need to keep the pin state in sync with all other nodes of the same type as this node.
+			if (IsA<UMaterialExpressionParameter>())
+			{
+				// Remember this expression parameter name.
+				const FName& ParameterName = static_cast<UMaterialExpressionParameter*>(this)->ParameterName;
+
+				// Fetch all nodes in the material that refer to a parameter.
+				TArray<UMaterialExpressionParameter*> ParameterExpressions;
+				Material->GetAllExpressionsInMaterialAndFunctionsOfType<UMaterialExpressionParameter>(ParameterExpressions);
+
+				for (UMaterialExpressionParameter* ExpressionParameter : ParameterExpressions)
+				{
+					// If the other expression type and parameter name are the same as this expression's...
+					if (ExpressionParameter->GetArchetype() == GetArchetype() && ExpressionParameter->ParameterName == ParameterName)
+					{
+						// ...modify the pin on other parameter expression node with the new value.
+						UEdGraphPin* OtherPin = ExpressionParameter->GraphNode->GetPinAt(PinIndex);
+						check(OtherPin->GetName() == Pin->GetName());
+						OtherPin->Modify();
+						OtherPin->DefaultValue = NewDefaultValue;
+					}
+				}
+			}
 		}
 	}
 
