Code change related to 21953187 to fix UE-158255. The material parameter value needed to be propagated.

diff --git a/Engine/Source/Runtime/Engine/Private/Materials/MaterialExpressions.cpp b/Engine/Source/Runtime/Engine/Private/Materials/MaterialExpressions.cpp
--- a/Engine/Source/Runtime/Engine/Private/Materials/MaterialExpressions.cpp
+++ b/Engine/Source/Runtime/Engine/Private/Materials/MaterialExpressions.cpp
@@ -1535,6 +1535,9 @@ void UMaterialExpression::PostEditChangeProperty(FPropertyChangedEvent& Property
 						OtherPin->DefaultValue = NewDefaultValue;
 					}
 				}
+
+				// Propagate he parameter value change so that it updates the other caches.
+				Material->PropagateExpressionParameterChanges(this);
 			}
 		}
 	}
