LevelInstance:
- Fix unhash crash caused by bounds changing after LevelInstance was hashed by WP
- Make sure to save new Level Instance actor after it's level so that initial bounds are valid

diff --git a/Engine/Source/Runtime/Engine/Private/LevelInstance/LevelInstanceSubsystem.cpp b/Engine/Source/Runtime/Engine/Private/LevelInstance/LevelInstanceSubsystem.cpp
--- a/Engine/Source/Runtime/Engine/Private/LevelInstance/LevelInstanceSubsystem.cpp
+++ b/Engine/Source/Runtime/Engine/Private/LevelInstance/LevelInstanceSubsystem.cpp
@@ -1170,18 +1170,18 @@ ILevelInstanceInterface* ULevelInstanceSubsystem::CreateLevelInstanceFrom(const
 
 	GetWorld()->SetCurrentLevel(LoadedLevel);
 
-	// Don't force saving of unsaved/temp packages onto the user.
-	if (!FPackageName::IsTempPackage(NewLevelInstanceActor->GetPackage()->GetName()))
-	{
-		DirtyPackages.Add(NewLevelInstanceActor->GetPackage()->GetFName());
-	}
-	
 	// Commit will always pop the actor editor context, make sure to push one here
 	UActorEditorContextSubsystem::Get()->PushContext();
 	bool bCommitted = CommitLevelInstanceInternal(StackLevelInstanceEdit.LevelInstanceEdit, /*bDiscardEdits=*/false, /*bDiscardOnFailure=*/true, &DirtyPackages);
 	check(bCommitted);
 	check(!StackLevelInstanceEdit.LevelInstanceEdit);
 
+	// Don't force saving of unsaved/temp packages onto the user. 
+	if (!FPackageName::IsTempPackage(NewLevelInstanceActor->GetPackage()->GetName()))
+	{
+		FEditorFileUtils::PromptForCheckoutAndSave({ NewLevelInstanceActor->GetPackage() }, /*bCheckDirty*/false, /*bPromptToSave*/false);
+	}
+
 	// EditorLevelUtils::CreateNewStreamingLevelForWorld deactivates all modes. Re-activate if needed
 	if (LevelInstanceEdit)
 	{
diff --git a/Engine/Source/Runtime/Engine/Private/WorldPartition/LevelInstance/LevelInstanceActorDesc.cpp b/Engine/Source/Runtime/Engine/Private/WorldPartition/LevelInstance/LevelInstanceActorDesc.cpp
--- a/Engine/Source/Runtime/Engine/Private/WorldPartition/LevelInstance/LevelInstanceActorDesc.cpp
+++ b/Engine/Source/Runtime/Engine/Private/WorldPartition/LevelInstance/LevelInstanceActorDesc.cpp
@@ -80,14 +80,6 @@ void FLevelInstanceActorDesc::UpdateBounds()
 	ContainerBounds.GetCenterAndExtents(BoundsLocation, BoundsExtent);
 }
 
-void FLevelInstanceActorDesc::OnLevelInstanceChanged(FName PackageName)
-{
-	if (LevelInstanceContainer.IsValid() && (LevelInstanceContainer->GetContainerPackage() == PackageName))
-	{
-		UpdateBounds();
-	}
-}
-
 void FLevelInstanceActorDesc::RegisterContainerInstance(UWorld* InWorld)
 {
 	if (InWorld)
@@ -98,11 +90,11 @@ void FLevelInstanceActorDesc::RegisterContainerInstance(UWorld* InWorld)
 		{
 			ULevelInstanceSubsystem* LevelInstanceSubsystem = UWorld::GetSubsystem<ULevelInstanceSubsystem>(InWorld);
 			check(LevelInstanceSubsystem);
-			LevelInstanceSubsystem->OnLevelInstanceChanged().AddRaw(this, &FLevelInstanceActorDesc::OnLevelInstanceChanged);
 
 			LevelInstanceContainer = LevelInstanceSubsystem->RegisterContainer(LevelPackage);
 			check(LevelInstanceContainer.IsValid());
 
+			// Should only be called on RegisterContainerInstance before ActorDesc is hashed
 			UpdateBounds();
 		}
 	}
@@ -115,8 +107,6 @@ void FLevelInstanceActorDesc::UnregisterContainerInstance()
 		ULevelInstanceSubsystem* LevelInstanceSubsystem = UWorld::GetSubsystem<ULevelInstanceSubsystem>(LevelInstanceContainer->GetWorld());
 		check(LevelInstanceSubsystem);
 
-		LevelInstanceSubsystem->OnLevelInstanceChanged().RemoveAll(this);
-
 		LevelInstanceSubsystem->UnregisterContainer(LevelInstanceContainer.Get());
 		LevelInstanceContainer.Reset();
 	}
diff --git a/Engine/Source/Runtime/Engine/Public/WorldPartition/LevelInstance/LevelInstanceActorDesc.h b/Engine/Source/Runtime/Engine/Public/WorldPartition/LevelInstance/LevelInstanceActorDesc.h
--- a/Engine/Source/Runtime/Engine/Public/WorldPartition/LevelInstance/LevelInstanceActorDesc.h
+++ b/Engine/Source/Runtime/Engine/Public/WorldPartition/LevelInstance/LevelInstanceActorDesc.h
@@ -46,7 +46,6 @@ protected:
 private:
 	void RegisterContainerInstance(UWorld* InWorld);
 	void UnregisterContainerInstance();
-	void OnLevelInstanceChanged(FName PackageName);
 	void UpdateBounds();
 };
 #endif
