fix to copy/paste not working on blend space axis min/max properties
explicitly disable reset to default on those feilds because it causes various issues.

https://p4-swarm.epicgames.net/reviews/15811505

diff --git a/Engine/Source/Editor/Persona/Private/Customization/BlendParameterDetails.cpp b/Engine/Source/Editor/Persona/Private/Customization/BlendParameterDetails.cpp
--- a/Engine/Source/Editor/Persona/Private/Customization/BlendParameterDetails.cpp
+++ b/Engine/Source/Editor/Persona/Private/Customization/BlendParameterDetails.cpp
@@ -32,7 +32,8 @@ void FBlendParameterDetails::CustomizeChildren(TSharedRef<IPropertyHandle> Struc
 				RangeProperties[RangeValueIndex] = ChildHandle;
 				bValidRangeValue[RangeValueIndex] = true;
 				
-				ChildBuilder.AddCustomRow(ChildHandle->GetPropertyDisplayName())
+				ChildBuilder.AddProperty(ChildHandle)
+				.CustomWidget()
 				.NameWidget
 				[
 					ChildHandle->CreatePropertyNameWidget()
diff --git a/Engine/Source/Runtime/Engine/Classes/Animation/BlendSpace.h b/Engine/Source/Runtime/Engine/Classes/Animation/BlendSpace.h
--- a/Engine/Source/Runtime/Engine/Classes/Animation/BlendSpace.h
+++ b/Engine/Source/Runtime/Engine/Classes/Animation/BlendSpace.h
@@ -100,11 +100,11 @@ struct FBlendParameter
 	FString DisplayName;
 
 	/** Minimum value for this axis range. */
-	UPROPERTY(EditAnywhere, DisplayName = "Minimum Axis Value", Category=BlendParameter)
+	UPROPERTY(EditAnywhere, DisplayName = "Minimum Axis Value", Category=BlendParameter, meta=(NoResetToDefault))
 	float Min;
 
 	/** Maximum value for this axis range. */
-	UPROPERTY(EditAnywhere, DisplayName = "Maximum Axis Value", Category=BlendParameter)
+	UPROPERTY(EditAnywhere, DisplayName = "Maximum Axis Value", Category=BlendParameter, meta=(NoResetToDefault))
 	float Max;
 
 	/** The number of grid divisions along this axis. */
