SkelMeshImportData: In addition to applying the package's CustomVersions when deserializing FReductionSkeletalMeshData and FSkeletalMeshImportData out of BulkData, also apply the UEVersion and LicenseeUEVersion.

#rnx

diff --git a/Engine/Source/Runtime/Engine/Private/Rendering/SkeletalMeshLODImporterData.cpp b/Engine/Source/Runtime/Engine/Private/Rendering/SkeletalMeshLODImporterData.cpp
--- a/Engine/Source/Runtime/Engine/Private/Rendering/SkeletalMeshLODImporterData.cpp
+++ b/Engine/Source/Runtime/Engine/Private/Rendering/SkeletalMeshLODImporterData.cpp
@@ -481,6 +481,8 @@ void FReductionBaseSkeletalMeshBulkData::Serialize(FArchive& Ar, UObject* Owner)
 	if (Ar.IsTransacting())
 	{
 		// If transacting, keep these members alive the other side of an undo, otherwise their values will get lost
+		Ar << UEVersion;
+		Ar << LicenseeUEVersion;
 		SerializeLoadingCustomVersionContainer.Serialize(Ar);
 		Ar << bUseSerializeLoadingCustomVersion;
 	}
@@ -500,6 +502,8 @@ void FReductionBaseSkeletalMeshBulkData::Serialize(FArchive& Ar, UObject* Owner)
 	if (!Ar.IsTransacting() && Ar.IsLoading())
 	{
 		//Save the custom version so we can load FReductionSkeletalMeshData later
+		UEVersion = Ar.UEVer();
+		LicenseeUEVersion = Ar.LicenseeUEVer();
 		SerializeLoadingCustomVersionContainer = BulkData.GetCustomVersions(Ar);
 		bUseSerializeLoadingCustomVersion = true;
 	}
@@ -526,6 +530,8 @@ void FReductionBaseSkeletalMeshBulkData::SaveReductionData(FSkeletalMeshLODModel
 		Ar << ReductionSkeletalMeshData;
 
 		// Preserve CustomVersions at save time so we can reuse the same ones when reloading direct from memory
+		UEVersion = Ar.UEVer();
+		LicenseeUEVersion = Ar.LicenseeUEVer();
 		SerializeLoadingCustomVersionContainer = Ar.GetCustomVersions();
 	}
 	// Unlock the bulk data
@@ -547,6 +553,8 @@ void FReductionBaseSkeletalMeshBulkData::LoadReductionData(FSkeletalMeshLODModel
 			
 			// Propagate the custom version information from the package to the bulk data, so that the MeshDescription
 			// is serialized with the same versioning.
+			Ar.SetUEVer(UEVersion);
+			Ar.SetLicenseeUEVer(LicenseeUEVersion);
 			Ar.SetCustomVersions(SerializeLoadingCustomVersionContainer);
 
 			Ar << ReductionSkeletalMeshData;
@@ -769,6 +777,8 @@ void FRawSkeletalMeshBulkData::Serialize(FArchive& Ar, UObject* Owner)
 	if (Ar.IsTransacting())
 	{
 		// If transacting, keep these members alive the other side of an undo, otherwise their values will get lost
+		Ar << UEVersion;
+		Ar << LicenseeUEVersion;
 		SerializeLoadingCustomVersionContainer.Serialize(Ar);
 		Ar << bUseSerializeLoadingCustomVersion;
 	}
@@ -806,6 +816,8 @@ void FRawSkeletalMeshBulkData::Serialize(FArchive& Ar, UObject* Owner)
 	if (!Ar.IsTransacting() && Ar.IsLoading())
 	{
 		//Save the custom version so we can load FReductionSkeletalMeshData later
+		UEVersion = Ar.UEVer();
+		LicenseeUEVersion = Ar.LicenseeUEVer();
 		SerializeLoadingCustomVersionContainer = BulkData.GetCustomVersions(Ar);
 		bUseSerializeLoadingCustomVersion = true;
 	}
@@ -830,6 +842,8 @@ void FRawSkeletalMeshBulkData::SaveRawMesh(FSkeletalMeshImportData& InMesh)
 		Ar << InMesh;
 
 		// Preserve CustomVersions at save time so we can reuse the same ones when reloading direct from memory
+		UEVersion = Ar.UEVer();
+		LicenseeUEVersion = Ar.LicenseeUEVer();
 		SerializeLoadingCustomVersionContainer = Ar.GetCustomVersions();
 	}
 
@@ -871,6 +885,8 @@ void FRawSkeletalMeshBulkData::LoadRawMesh(FSkeletalMeshImportData& OutMesh)
 
 			// Propagate the custom version information from the package to the bulk data, so that the MeshDescription
 			// is serialized with the same versioning.
+			Ar.SetUEVer(UEVersion);
+			Ar.SetLicenseeUEVer(LicenseeUEVersion);
 			Ar.SetCustomVersions(SerializeLoadingCustomVersionContainer);
 			Ar << OutMesh;
 		}
@@ -919,6 +935,8 @@ void FRawSkeletalMeshBulkData::UpdateRawMeshFormat()
 
 		// Propagate the custom version information from the package to the bulk data, so that the MeshDescription
 		// is serialized with the same versioning.
+		Reader.SetUEVer(UEVersion);
+		Reader.SetLicenseeUEVer(LicenseeUEVersion);
 		Reader.SetCustomVersions(SerializeLoadingCustomVersionContainer);
 		Reader << MeshImportData;
 
diff --git a/Engine/Source/Runtime/Engine/Public/Rendering/SkeletalMeshLODImporterData.h b/Engine/Source/Runtime/Engine/Public/Rendering/SkeletalMeshLODImporterData.h
--- a/Engine/Source/Runtime/Engine/Public/Rendering/SkeletalMeshLODImporterData.h
+++ b/Engine/Source/Runtime/Engine/Public/Rendering/SkeletalMeshLODImporterData.h
@@ -364,6 +364,7 @@ public:
 	// Alternate influence imported(i.e. FBX) data. The name is the alternate skinning profile name
 	TArray<FSkeletalMeshImportData> AlternateInfluences;
 	TArray<FString> AlternateInfluenceProfileNames;
+	
 
 	//////////////////////////////////////////////////////////////////////////
 
@@ -455,6 +456,8 @@ struct FReductionBaseSkeletalMeshBulkData
 
 	//The custom version when this was load
 	FCustomVersionContainer SerializeLoadingCustomVersionContainer;
+	FPackageFileVersion UEVersion;
+	int32 LicenseeUEVersion = 0;
 	bool bUseSerializeLoadingCustomVersion = false;
 
 	uint32 CacheLODVertexNumber = MAX_uint32;
@@ -534,6 +537,8 @@ class FRawSkeletalMeshBulkData
 
 	//The custom version when this was load
 	FCustomVersionContainer SerializeLoadingCustomVersionContainer;
+	FPackageFileVersion UEVersion;
+	int32 LicenseeUEVersion = 0;
 	bool bUseSerializeLoadingCustomVersion = false;
 
 public:
