ShaderDebugPath now respects the material quality level:

`...\Saved\ShaderDebugInfo\PCD3D_SM5\WorldGridMaterial\Default\FLocalVertexFactory\FHitProxyPS`

or

`...\Saved\ShaderDebugInfo\PCD3D_SM5\WorldGridMaterial\Epic\FLocalVertexFactory\FHitProxyPS`

- Use `LexToString(EMaterialQualityLevel::Type QualityLevel)` to convert the quality level enum to a string.
- I'm specifically appending to the DebugName the material quality at the level in the code where we "know" about the material.  I wanted to avoid changing the parameters to the function, `GlobalBeginCompileShader`, because it is material agnostic.

diff --git a/Engine/Source/Runtime/Engine/Private/Materials/MaterialShader.cpp b/Engine/Source/Runtime/Engine/Private/Materials/MaterialShader.cpp
--- a/Engine/Source/Runtime/Engine/Private/Materials/MaterialShader.cpp
+++ b/Engine/Source/Runtime/Engine/Private/Materials/MaterialShader.cpp
@@ -1313,7 +1313,7 @@ static void PrepareMaterialShaderCompileJob(EShaderPlatform Platform,
 
 	// Compile the shader environment passed in with the shader type's source code.
 	::GlobalBeginCompileShader(
-		Material->GetFriendlyName(),
+		Material->GetFriendlyName() / LexToString(Material->GetQualityLevel()),
 		nullptr,
 		ShaderType,
 		ShaderPipeline,
diff --git a/Engine/Source/Runtime/Engine/Private/Materials/MeshMaterialShader.cpp b/Engine/Source/Runtime/Engine/Private/Materials/MeshMaterialShader.cpp
--- a/Engine/Source/Runtime/Engine/Private/Materials/MeshMaterialShader.cpp
+++ b/Engine/Source/Runtime/Engine/Private/Materials/MeshMaterialShader.cpp
@@ -62,7 +62,7 @@ static void PrepareMeshMaterialShaderCompileJob(EShaderPlatform Platform,
 
 	// Compile the shader environment passed in with the shader type's source code.
 	::GlobalBeginCompileShader(
-		Material->GetFriendlyName(),
+		Material->GetFriendlyName() / LexToString(Material->GetQualityLevel()),
 		VertexFactoryType,
 		ShaderType,
 		ShaderPipeline,
