Reverting some previous KPK cleanup

diff --git a/Engine/Source/Runtime/Engine/Private/NetDriver.cpp b/Engine/Source/Runtime/Engine/Private/NetDriver.cpp
--- a/Engine/Source/Runtime/Engine/Private/NetDriver.cpp
+++ b/Engine/Source/Runtime/Engine/Private/NetDriver.cpp
@@ -1516,7 +1516,7 @@ void UNetDriver::Shutdown()
 				 {
 					 GuidsToLog.Append(ActorChannel->PendingGuidResolves);
 				 }
-				 ActorChannel->BreakAndReleaseReferences();
+				 ActorChannel->CleanupReplicators();
 			 }
 		}
 
@@ -1544,7 +1544,6 @@ void UNetDriver::Shutdown()
 		// Calls Channel[0]->Close to send a close bunch to server
 		ServerConnection->Close();
 		ServerConnection->FlushNet();
-		ServerConnection->MarkAsGarbage();
 	}
 
 	// Server closing connections with clients
diff --git a/Engine/Source/Runtime/Engine/Private/UnrealEngine.cpp b/Engine/Source/Runtime/Engine/Private/UnrealEngine.cpp
--- a/Engine/Source/Runtime/Engine/Private/UnrealEngine.cpp
+++ b/Engine/Source/Runtime/Engine/Private/UnrealEngine.cpp
@@ -13375,7 +13375,6 @@ namespace UE::Private
 				NetDriver->SetWorld(NULL);
 				NetDriver->Shutdown();
 				NetDriver->LowLevelDestroy();
-				NetDriver->MarkAsGarbage();
 
 				if (UWorld* World = Context.World())
 				{
