Virtualization fix for MetaSounds/procedural sounds set to restart not restarting if WaveInstance VolumeModifier goes to 0.

diff --git a/Engine/Source/Runtime/Engine/Private/SoundWave.cpp b/Engine/Source/Runtime/Engine/Private/SoundWave.cpp
--- a/Engine/Source/Runtime/Engine/Private/SoundWave.cpp
+++ b/Engine/Source/Runtime/Engine/Private/SoundWave.cpp
@@ -2754,11 +2754,14 @@ void USoundWave::Parse(FAudioDevice* AudioDevice, const UPTRINT NodeWaveInstance
 
 		WaveInstance = &HandleStart(ActiveSound, NodeWaveInstanceHash);
 	}
-	// If the procedural sound has notified as finished, honor it.
-	// Procedural sounds should always only have one instance max.
 	else
 	{
-		if (WaveInstance->bIsFinished && WaveInstance->bAlreadyNotifiedHook && bProcedural)
+		// If WaveInstance is finished & sound is procedurally generated and has
+		// virtualization disabled, early out to stop playback at the ActiveSound level.
+		const bool bVirtualizationDisabled = VirtualizationMode == EVirtualizationMode::Disabled;
+		const bool bInstanceIsFinished = WaveInstance->bIsFinished;
+		const bool bFinishedNotify = bInstanceIsFinished && WaveInstance->bAlreadyNotifiedHook;
+		if (bFinishedNotify && bProcedural && bVirtualizationDisabled)
 		{
 			return;
 		}
