Remove ancient pow2 requirement from UTextureRenderTarget2D::ConstructTexture2D (originally added in 2006 in CL 116345)

diff --git a/Engine/Source/Runtime/Engine/Classes/Kismet/KismetRenderingLibrary.h b/Engine/Source/Runtime/Engine/Classes/Kismet/KismetRenderingLibrary.h
--- a/Engine/Source/Runtime/Engine/Classes/Kismet/KismetRenderingLibrary.h
+++ b/Engine/Source/Runtime/Engine/Classes/Kismet/KismetRenderingLibrary.h
@@ -83,7 +83,7 @@ class UKismetRenderingLibrary : public UBlueprintFunctionLibrary
 	static ENGINE_API void DrawMaterialToRenderTarget(UObject* WorldContextObject, UTextureRenderTarget2D* TextureRenderTarget, UMaterialInterface* Material);
 
 	/**
-	* Creates a new Static Texture from a Render Target 2D. Render Target Must be power of two and use four channels.
+	* Creates a new Static Texture from a Render Target 2D.
 	* Only works in the editor
 	*/
 	UFUNCTION(BlueprintCallable, meta = (DisplayName = "Render Target Create Static Texture Editor Only", Keywords = "Create Static Texture from Render Target", UnsafeDuringActorConstruction = "true"), Category = Game)
diff --git a/Engine/Source/Runtime/Engine/Private/TextureRenderTarget2D.cpp b/Engine/Source/Runtime/Engine/Private/TextureRenderTarget2D.cpp
--- a/Engine/Source/Runtime/Engine/Private/TextureRenderTarget2D.cpp
+++ b/Engine/Source/Runtime/Engine/Private/TextureRenderTarget2D.cpp
@@ -314,18 +314,19 @@ FString UTextureRenderTarget2D::GetDesc()
 
 UTexture2D* UTextureRenderTarget2D::ConstructTexture2D(UObject* Outer, const FString& NewTexName, EObjectFlags InObjectFlags, uint32 Flags, TArray<uint8>* AlphaOverride)
 {
-	UTexture2D* Result = NULL;
+	UTexture2D* Result = nullptr;
+
 #if WITH_EDITOR
-	// Check render target size is valid and power of two.
-	const bool bIsValidSize = (SizeX != 0 && !(SizeX & (SizeX - 1)) &&
-		SizeY != 0 && !(SizeY & (SizeY - 1)));
-	// The r2t resource will be needed to read its surface contents
+	// Check render target size is valid.
+	const bool bIsValidSize = (SizeX != 0 && SizeY != 0);
+
+	// The render to texture resource will be needed to read its surface contents
 	FRenderTarget* RenderTarget = GameThread_GetRenderTargetResource();
 
 	const ETextureSourceFormat TextureFormat = GetTextureFormatForConversionToTexture2D();
 
 	// exit if source is not compatible.
-	if (bIsValidSize == false || RenderTarget == NULL || TextureFormat == TSF_Invalid)
+	if (bIsValidSize == false || RenderTarget == nullptr || TextureFormat == TSF_Invalid)
 	{
 		return Result;
 	}
@@ -369,6 +370,7 @@ UTexture2D* UTextureRenderTarget2D::ConstructTexture2D(UObject* Outer, const FSt
 	}
 	Result->PostEditChange();
 #endif
+
 	return Result;
 }
 
