Remove extraneous DeltaPosition from AActor::PostNetReceivePhysicsState

diff --git a/Engine/Source/Runtime/Engine/Private/ActorReplication.cpp b/Engine/Source/Runtime/Engine/Private/ActorReplication.cpp
--- a/Engine/Source/Runtime/Engine/Private/ActorReplication.cpp
+++ b/Engine/Source/Runtime/Engine/Private/ActorReplication.cpp
@@ -276,8 +276,6 @@ void AActor::PostNetReceivePhysicState()
 		const FRepMovement& ThisReplicatedMovement = GetReplicatedMovement();
 		FRigidBodyState NewState;
 		ThisReplicatedMovement.CopyTo(NewState, this);
-
-		FVector DeltaPos(FVector::ZeroVector);
 		RootPrimComp->SetRigidBodyReplicatedTarget(NewState, NAME_None, ThisReplicatedMovement.ServerFrame, ThisReplicatedMovement.ServerPhysicsHandle);
 	}
 }
