World Partition: Allow server to unload data layer cells even if streaming out cells is disallowed by wp.Runtime.EnableServerStreamingOut

diff --git a/Engine/Source/Runtime/Engine/Private/WorldPartition/WorldPartitionStreamingPolicy.cpp b/Engine/Source/Runtime/Engine/Private/WorldPartition/WorldPartitionStreamingPolicy.cpp
--- a/Engine/Source/Runtime/Engine/Private/WorldPartition/WorldPartitionStreamingPolicy.cpp
+++ b/Engine/Source/Runtime/Engine/Private/WorldPartition/WorldPartitionStreamingPolicy.cpp
@@ -396,13 +396,13 @@ void UWorldPartitionStreamingPolicy::UpdateStreamingState()
 	// Process cells to unload
 	auto ProcessCellsToUnload = [this, bIsServer, bCanDeactivateOrUnloadCells, &ShouldWaitForClientVisibility](TSet<const UWorldPartitionRuntimeCell*>& Cells)
 	{
-		check(bCanDeactivateOrUnloadCells || Cells.IsEmpty());
 		// Only loop for server as ShouldWaitForClientVisibility only concerns server
 		if (bIsServer)
 		{
 			for (TSet<const UWorldPartitionRuntimeCell*>::TIterator It = Cells.CreateIterator(); It; ++It)
 			{
 				const UWorldPartitionRuntimeCell* Cell = *It;
+				check(bCanDeactivateOrUnloadCells || Cell->HasDataLayers());
 				if (ShouldWaitForClientVisibility(Cell))
 				{
 					It.RemoveCurrent();
@@ -436,7 +436,6 @@ void UWorldPartitionStreamingPolicy::UpdateStreamingState()
 
 	// Determine cells to unload
 	TArray<const UWorldPartitionRuntimeCell*> ToUnloadCells;
-	if (bCanDeactivateOrUnloadCells)
 	{
 		TSet<const UWorldPartitionRuntimeCell*> ToUnloadCellsUnsorted = ActivatedCells.Union(LoadedCells).Difference(FrameActivateCells.Union(FrameLoadCells));
 		ProcessCellsToUnload(ToUnloadCellsUnsorted);
