Texture streaming and cinematic LOD
fix crash in Editor with 2052 texture set to num cinematic lods=1
NumNonStreaming  calculation is called even if NeverStream so detect pow2
make Editor behavior more consistent with cooked game behavior
don't allowing streaming mip levels in Editor that won't exist in game?

diff --git a/Engine/Source/Runtime/Engine/Private/Texture.cpp b/Engine/Source/Runtime/Engine/Private/Texture.cpp
--- a/Engine/Source/Runtime/Engine/Private/Texture.cpp
+++ b/Engine/Source/Runtime/Engine/Private/Texture.cpp
@@ -1372,45 +1372,66 @@ FStreamableRenderResourceState UTexture::GetResourcePostInitState(const FTexture
 
 	const int32 AssetMipIdxForResourceFirstMip = FMath::Max<int32>(0, PlatformData->Mips.Num() - NumMips);
 
-	bool bMakeStreamble = false;
+	bool bMakeStreamable = false;
 	int32 NumRequestedMips = 0;
+	
+	// same conditions as IsCandidateForTextureStreaming which TextureDerivedData will use :
+	bool bTextureIsStreamable =
+	  (!NeverStream && 
+		NumOfNonStreamingMips < NumMips && 
+		LODGroup != TEXTUREGROUP_UI ); 
 
 #if PLATFORM_SUPPORTS_TEXTURE_STREAMING
-	bool bWillProvideMipDataWithoutDisk = false;
-
-	// Check if any of the CustomMipData providers associated with this texture can provide mip data even without DDC or disk,
-	// if so, enable streaming for this texture
-	for (UAssetUserData* UserData : AssetUserData)
+	if ( bTextureIsStreamable )
 	{
-		UTextureMipDataProviderFactory* CustomMipDataProviderFactory = Cast<UTextureMipDataProviderFactory>(UserData);
-		if (CustomMipDataProviderFactory)
+		bool bWillProvideMipDataWithoutDisk = false;
+
+		// Check if any of the CustomMipData providers associated with this texture can provide mip data even without DDC or disk,
+		// if so, enable streaming for this texture
+		for (UAssetUserData* UserData : AssetUserData)
 		{
-			bWillProvideMipDataWithoutDisk = CustomMipDataProviderFactory->WillProvideMipDataWithoutDisk();
-			if (bWillProvideMipDataWithoutDisk)
+			UTextureMipDataProviderFactory* CustomMipDataProviderFactory = Cast<UTextureMipDataProviderFactory>(UserData);
+			if (CustomMipDataProviderFactory)
 			{
-				break;
+				bWillProvideMipDataWithoutDisk = CustomMipDataProviderFactory->WillProvideMipDataWithoutDisk();
+				if (bWillProvideMipDataWithoutDisk)
+				{
+					break;
+				}
 			}
 		}
-	}
 
-	if (!NeverStream && 
-		NumOfNonStreamingMips < NumMips && 
-		LODGroup != TEXTUREGROUP_UI && 
-		bAllowStreaming &&
-		(bSkipCanBeLoaded || PlatformData->CanBeLoaded() || bWillProvideMipDataWithoutDisk))
-	{
-		bMakeStreamble  = true;
+		if (bAllowStreaming &&
+			(bSkipCanBeLoaded || PlatformData->CanBeLoaded() || bWillProvideMipDataWithoutDisk))
+		{
+			bMakeStreamable  = true;
+		}
 	}
 #endif
 
-	if (bMakeStreamble && IStreamingManager::Get().IsRenderAssetStreamingEnabled(EStreamableRenderAssetType::Texture))
+	if ( ! bTextureIsStreamable )
+	{
+		// in Editor , NumOfNonStreamingMips may not be all mips
+		// but once we cook it will be
+		// so check this early to make behavior consistent
+		NumRequestedMips = NumMips;
+	}
+	else if (bMakeStreamable && IStreamingManager::Get().IsRenderAssetStreamingEnabled(EStreamableRenderAssetType::Texture))
 	{
 		NumRequestedMips = NumOfNonStreamingMips;
 	}
 	else
 	{
+		// we are not streaming (bMakeStreamable is false)
+		// but this may select a mip below the top mip
+		// (due to cinematic lod bias)
+		// but only if the texture itself is streamable
+
 		// Adjust CachedLODBias so that it takes into account FStreamableRenderResourceState::AssetLODBias.
 		const int32 ResourceLODBias = FMath::Max<int32>(0, GetCachedLODBias() - AssetMipIdxForResourceFirstMip);
+		//  ResourceLODBias almost always = NumCinematicMipLevels
+		//check( ResourceLODBias == NumCinematicMipLevels ); // this will be true unless you hit the MaxRuntimeMipCount clamp in NumMips
+		// if ResourceLODBias == 0 , this always selects NumRequestedMips = NumMips
 
 		// Ensure NumMipsInTail is within valid range to safeguard on the above expressions. 
 		const int32 NumMipsInTail = FMath::Clamp<int32>(PlatformData->GetNumMipsInTail(), 1, NumMips);
@@ -1428,13 +1449,17 @@ FStreamableRenderResourceState UTexture::GetResourcePostInitState(const FTexture
 		NumRequestedMips = FMath::Max(NumRequestedMips,NumOfNonStreamingMips);
 	}
 
+	// @todo Oodle : this looks like a bug; did it mean to be MinRequestMipCount <= NumMips
+	// typically MinRequestMipCount == 0
+	// the only place it's not zero is from UTexture2D::CreateResource from existing resource mem, where it is == NumMips
+	// but in that case it is ignored here because it is == NumMips, not <
 	if (NumRequestedMips < MinRequestMipCount && MinRequestMipCount < NumMips)
 	{
 		NumRequestedMips = MinRequestMipCount;
 	}
 
 	FStreamableRenderResourceState PostInitState;
-	PostInitState.bSupportsStreaming = bMakeStreamble;
+	PostInitState.bSupportsStreaming = bMakeStreamable;
 	PostInitState.NumNonStreamingLODs = (uint8)NumOfNonStreamingMips;
 	PostInitState.NumNonOptionalLODs = (uint8)NumOfNonOptionalMips;
 	PostInitState.MaxNumLODs = (uint8)NumMips;
diff --git a/Engine/Source/Runtime/Engine/Private/TextureDerivedData.cpp b/Engine/Source/Runtime/Engine/Private/TextureDerivedData.cpp
--- a/Engine/Source/Runtime/Engine/Private/TextureDerivedData.cpp
+++ b/Engine/Source/Runtime/Engine/Private/TextureDerivedData.cpp
@@ -2090,7 +2090,11 @@ int32 FTexturePlatformData::GetNumNonStreamingMips() const
 			return NumNonStreamingMips;
 		}
 	}
-	else if (Mips.Num() > 0)
+	else if (Mips.Num() <= 1)
+	{
+		return Mips.Num();
+	}
+	else
 	{
 		int32 MipCount = Mips.Num();
 		int32 NumNonStreamingMips = 1;
@@ -2103,16 +2107,26 @@ int32 FTexturePlatformData::GetNumNonStreamingMips() const
 		int32 BlockSizeY = GPixelFormats[PixelFormat].BlockSizeY;
 		if (BlockSizeX > 1 || BlockSizeY > 1)
 		{
-			NumNonStreamingMips = FMath::Max<int32>(NumNonStreamingMips, MipCount - FPlatformMath::FloorLog2(Mips[0].SizeX / BlockSizeX));
-			NumNonStreamingMips = FMath::Max<int32>(NumNonStreamingMips, MipCount - FPlatformMath::FloorLog2(Mips[0].SizeY / BlockSizeY));
+			// ensure the top non-streamed mip size is >= BlockSize (and a multiple of block size!)
+			// @todo Oodle : only do for BCN, not for ASTC
+			
+			if ( FMath::IsPowerOfTwo(Mips[0].SizeX) && FMath::IsPowerOfTwo(Mips[0].SizeY) )
+			{
+				// note: this is not right for non pow 2 (in that case you should check all streaming mips are aligned to blocksize)
+				NumNonStreamingMips = FMath::Max<int32>(NumNonStreamingMips, MipCount - FPlatformMath::FloorLog2(Mips[0].SizeX / BlockSizeX));
+				NumNonStreamingMips = FMath::Max<int32>(NumNonStreamingMips, MipCount - FPlatformMath::FloorLog2(Mips[0].SizeY / BlockSizeY));
+			}
+			else
+			{
+				// non-pow2 should be NeverStream
+				// but they call into here anyway
+				// and actually CAN stream via cinematic lod bias (in Editor, but not in cooked/game)
+				NumNonStreamingMips = MipCount;
+			}
 		}
 
 		return NumNonStreamingMips;
 	}
-	else
-	{
-		return 0;
-	}
 }
 
 int32 FTexturePlatformData::GetNumNonOptionalMips() const
@@ -2416,7 +2430,6 @@ static void SerializePlatformData(
 			const int32 Height = PlatformData->SizeY;
 			const int32 LODGroup = Texture->LODGroup;
 			const int32 LODBias = Texture->LODBias;
-			const int32 NumCinematicMipLevels = Texture->NumCinematicMipLevels;
 			const TextureMipGenSettings MipGenSetting = Texture->MipGenSettings;
 			const int32 LastMip = FMath::Max(NumMips - 1, 0);
 			check(NumMips >= (int32)PlatformData->GetNumMipsInTail());
@@ -2425,6 +2438,7 @@ static void SerializePlatformData(
 			FirstMipToSerialize = Ar.CookingTarget()->GetTextureLODSettings().CalculateLODBias(Width, Height, Texture->MaxTextureSize, LODGroup, LODBias, 0, MipGenSetting, bIsVirtual);
 			if (!bIsVirtual)
 			{
+				// NumMips is the number of mips starting from FirstMipToSerialize
 				FirstMipToSerialize = FMath::Clamp(FirstMipToSerialize, 0, PlatformData->GetNumMipsInTail() > 0 ? FirstMipTailMip : LastMip);
 				NumMips = FMath::Max(0, NumMips - FirstMipToSerialize);
 			}
@@ -2459,6 +2473,10 @@ static void SerializePlatformData(
 			int32 OptionalMips = 0; // TODO: do we need to consider platforms saving texture assets as cooked files? all the info to calculate the optional is part of the editor only data
 			bool DuplicateNonOptionalMips = false;
 		
+			// bStreamable comes from IsCandidateForTextureStreaming
+			//  if not bStreamable, all mips are written inline
+			//  so the runtime will see NumNonStreamingMips = all
+
 #if WITH_EDITORONLY_DATA
 			check(Ar.CookingTarget());
 			// This also needs to check whether the project enables texture streaming.
@@ -2469,35 +2487,46 @@ static void SerializePlatformData(
 			if (bStreamable)
 #endif
 			{
-				MinMipToInline = FMath::Max(0, NumMips - PlatformData->GetNumNonStreamingMips());
+				int32 NumNonStreamingMips = PlatformData->GetNumNonStreamingMips();
+				// NumMips has been reduced by FirstMipToSerialize (LODBias)
+				NumNonStreamingMips = FMath::Min(NumNonStreamingMips,NumMips);
+				// NumNonStreamingMips is not serialized
+				// the runtime will just use NumNonStreamingMips == NumInlineMips
+				MinMipToInline = NumMips - NumNonStreamingMips;
 #if WITH_EDITORONLY_DATA
-				const int32 Width = PlatformData->SizeX;
-				const int32 Height = PlatformData->SizeY;
 				const int32 LODGroup = Texture->LODGroup;
-				const int32 LODBias = Texture->LODBias;
-				const int32 NumCinematicMipLevels = Texture->NumCinematicMipLevels;
 
-				OptionalMips = Ar.CookingTarget()->GetTextureLODSettings().CalculateNumOptionalMips(LODGroup, Width, Height, NumMips, MinMipToInline, Texture->MipGenSettings);
+				// was a bug ? CalculateNumOptional mips is using full Width/Height ?  should be from FirstMipToSerialize? -> moot, it's not actually used
+				const int32 FirstMipWidth  = PlatformData->Mips[FirstMipToSerialize].SizeX;
+				const int32 FirstMipHeight = PlatformData->Mips[FirstMipToSerialize].SizeY;
+
+				OptionalMips = Ar.CookingTarget()->GetTextureLODSettings().CalculateNumOptionalMips(LODGroup, FirstMipWidth, FirstMipHeight, NumMips, MinMipToInline, Texture->MipGenSettings);
 				DuplicateNonOptionalMips = Ar.CookingTarget()->GetTextureLODSettings().TextureLODGroups[LODGroup].DuplicateNonOptionalMips;
+
+				// Optional mips must not overlap the non-streaming mips : (MinMipToInline ensures this)
+				check( OptionalMips + NumNonStreamingMips <= NumMips );
 #endif
 			}
 
 			for (int32 MipIndex = 0; MipIndex < NumMips && MipIndex < OptionalMips; ++MipIndex) //-V654
 			{
+				// optional (and streamed) mips
 				PlatformData->Mips[MipIndex + FirstMipToSerialize].BulkData.SetBulkDataFlags(BULKDATA_Force_NOT_InlinePayload | BULKDATA_OptionalPayload);
 			}
 
 			const uint32 AdditionalNonOptionalBulkDataFlags = DuplicateNonOptionalMips ? BULKDATA_DuplicateNonOptionalPayload : 0;
 			for (int32 MipIndex = OptionalMips; MipIndex < NumMips && MipIndex < MinMipToInline; ++MipIndex)
 			{
+				// non-optional but streamed mips
 				PlatformData->Mips[MipIndex + FirstMipToSerialize].BulkData.SetBulkDataFlags(BULKDATA_Force_NOT_InlinePayload | AdditionalNonOptionalBulkDataFlags);
 			}
 			for (int32 MipIndex = MinMipToInline; MipIndex < NumMips; ++MipIndex)
 			{
+				// non-streamed (inline) mips
 				PlatformData->Mips[MipIndex + FirstMipToSerialize].BulkData.SetBulkDataFlags(BULKDATA_ForceInlinePayload | BULKDATA_SingleUse);
 			}
 		}
-		else // bVirtual == false
+		else // bVirtual == true
 		{
 			const int32 NumChunks = PlatformData->VTData->Chunks.Num();
 			BulkDataMipFlags.AddZeroed(NumChunks);
diff --git a/Engine/Source/Runtime/Engine/Private/TextureLODSettings.cpp b/Engine/Source/Runtime/Engine/Private/TextureLODSettings.cpp
--- a/Engine/Source/Runtime/Engine/Private/TextureLODSettings.cpp
+++ b/Engine/Source/Runtime/Engine/Private/TextureLODSettings.cpp
@@ -220,6 +220,10 @@ int32 UTextureLODSettings::CalculateNumOptionalMips(int32 LODGroup, const int32
 
 	int32 OptionalLOD = FMath::Min<int32>(FMath::CeilLogTwo(LODGroupInfo.OptionalMaxLODSize) + 1, NumMips);
 
+	// "MinMipToInline" actually comes from NumNonStreaming
+	//  this ensures that optional mips are always streaming mips
+	//  also streaming mips are always compression block sized aligned, therefore optional mips are too
+
 	int32 NumOptionalMips = FMath::Min(NumMips - (OptionalLOD - LODGroupInfo.OptionalLODBias), MinMipToInline);
 	return NumOptionalMips;
 }
