Fixing path through cascade activation that can lead to unregister during register and a subsequent check();

diff --git a/Engine/Source/Runtime/Engine/Classes/Particles/ParticleSystemComponent.h b/Engine/Source/Runtime/Engine/Classes/Particles/ParticleSystemComponent.h
--- a/Engine/Source/Runtime/Engine/Classes/Particles/ParticleSystemComponent.h
+++ b/Engine/Source/Runtime/Engine/Classes/Particles/ParticleSystemComponent.h
@@ -538,6 +538,8 @@ public:
 	/** If this component was previously having it's significance managed by gameplay code. Allows us to refresh render data when this changes. */
 	uint8 bWasManagingSignificance : 1;
 
+	uint8 bIsDuringRegister : 1;
+
 	UPROPERTY()
 	uint8 bWarmingUp:1;
 
diff --git a/Engine/Source/Runtime/Engine/Private/Particles/ParticleComponents.cpp b/Engine/Source/Runtime/Engine/Private/Particles/ParticleComponents.cpp
--- a/Engine/Source/Runtime/Engine/Private/Particles/ParticleComponents.cpp
+++ b/Engine/Source/Runtime/Engine/Private/Particles/ParticleComponents.cpp
@@ -3497,6 +3497,7 @@ UParticleSystemComponent::UParticleSystemComponent(const FObjectInitializer& Obj
 	LastSignificantTime = 0.0f;
 	bIsManagingSignificance = 0;
 	bWasManagingSignificance = 0;
+	bIsDuringRegister = 0;
 
 	ManagerHandle = INDEX_NONE;
 	bPendingManagerAdd = false;
@@ -3923,6 +3924,8 @@ bool UParticleSystemComponent::ParticleLineCheck(FHitResult& Hit, AActor* Source
 
 void UParticleSystemComponent::OnRegister()
 {
+	FGuardValue_Bitfield(bIsDuringRegister, true);
+	
 	ForceAsyncWorkCompletion(STALL);
 	check(FXSystem == nullptr);
 
@@ -6230,7 +6233,7 @@ void UParticleSystemComponent::ActivateSystem(bool bFlagAsJustAttached)
 		}
 
 		//We are definitely insignificant already so set insignificant before we ever begin ticking.
-		if (bIsManagingSignificance && Template->GetHighestSignificance() < RequiredSignificance && Template->InsignificanceDelay == 0.0f)
+		if (!bIsDuringRegister && bIsManagingSignificance && Template->GetHighestSignificance() < RequiredSignificance && Template->InsignificanceDelay == 0.0f)
 		{
 			OnSignificanceChanged(false, true);
 		}
