LevelInstance: Fix regression from CL 22135039 (Crash when creating Packed Level Actor)

#rnx

diff --git a/Engine/Source/Runtime/Engine/Private/LevelInstance/LevelInstanceSubsystem.cpp b/Engine/Source/Runtime/Engine/Private/LevelInstance/LevelInstanceSubsystem.cpp
--- a/Engine/Source/Runtime/Engine/Private/LevelInstance/LevelInstanceSubsystem.cpp
+++ b/Engine/Source/Runtime/Engine/Private/LevelInstance/LevelInstanceSubsystem.cpp
@@ -1176,6 +1176,9 @@ ILevelInstanceInterface* ULevelInstanceSubsystem::CreateLevelInstanceFrom(const
 	check(bCommitted);
 	check(!StackLevelInstanceEdit.LevelInstanceEdit);
 
+	// In case Commit caused actor to be GCed (Commit can cause BP reinstancing)
+	NewLevelInstanceActor = Cast<AActor>(GetLevelInstance(NewLevelInstanceID));
+
 	// Don't force saving of unsaved/temp packages onto the user. 
 	if (!FPackageName::IsTempPackage(NewLevelInstanceActor->GetPackage()->GetName()))
 	{
