World Partition: detect duplicate actors coming from content bundles as invalid. operation.

#rnx

diff --git a/Engine/Plugins/Editor/DataValidation/Source/DataValidation/Private/WorldPartitionChangelistValidator.cpp b/Engine/Plugins/Editor/DataValidation/Source/DataValidation/Private/WorldPartitionChangelistValidator.cpp
--- a/Engine/Plugins/Editor/DataValidation/Source/DataValidation/Private/WorldPartitionChangelistValidator.cpp
+++ b/Engine/Plugins/Editor/DataValidation/Source/DataValidation/Private/WorldPartitionChangelistValidator.cpp
@@ -184,7 +184,7 @@ EDataValidationResult UWorldPartitionChangelistValidator::ValidateActorsAndDataL
 		{
 			// Find in memory failed, load the ActorDescContainer
 			ActorDescContainer = NewObject<UActorDescContainer>();
-			ActorDescContainer->Initialize(nullptr, MapPath.GetPackageName());
+			ActorDescContainer->Initialize({ nullptr, MapPath.GetPackageName() });
 		}
 
 		// Build a set of Relevant Actor Guids to scope error messages to what's contained in the CL 
diff --git a/Engine/Source/Editor/SceneOutliner/Private/ActorBrowsingMode.cpp b/Engine/Source/Editor/SceneOutliner/Private/ActorBrowsingMode.cpp
--- a/Engine/Source/Editor/SceneOutliner/Private/ActorBrowsingMode.cpp
+++ b/Engine/Source/Editor/SceneOutliner/Private/ActorBrowsingMode.cpp
@@ -1121,23 +1121,26 @@ void FActorBrowsingMode::RepairErrors() const
 			ISourceControlModule& SourceControlModule = ISourceControlModule::Get();
 			ISourceControlProvider& SourceControlProvider = SourceControlModule.GetProvider();
 
-			TArray<FAssetData> InvalidActors;
-			WorldPartition->ForEachActorDescContainer([&InvalidActors ](UActorDescContainer* ActorDescContainer)
+			TArray<FWorldPartitionActorDesc*> InvalidActorDescs;
+			WorldPartition->ForEachActorDescContainer([&InvalidActorDescs](UActorDescContainer* ActorDescContainer)
 			{
-				InvalidActors += ActorDescContainer->GetInvalidActors();
+				for (const TUniquePtr<FWorldPartitionActorDesc>& InvalidActor : ActorDescContainer->GetInvalidActors())
+				{
+					InvalidActorDescs.Add(InvalidActor.Get());
+				}
 				ActorDescContainer->ClearInvalidActors();
 			});
 
 			TArray<FString> ActorFilesToDelete;
 			TArray<FString> ActorFilesToRevert;
 			{
-				FScopedSlowTask SlowTask(InvalidActors.Num(), LOCTEXT("UpdatingSourceControlStatus", "Updating source control status..."));
+				FScopedSlowTask SlowTask(InvalidActorDescs.Num(), LOCTEXT("UpdatingSourceControlStatus", "Updating source control status..."));
 				SlowTask.MakeDialogDelayed(1.0f);
 
-				for (const FAssetData& InvalidActor : InvalidActors)
+				for (const FWorldPartitionActorDesc* InvalidActorDesc : InvalidActorDescs)
 				{
 					FPackagePath PackagePath;
-					if (FPackagePath::TryFromPackageName(InvalidActor.PackageName, PackagePath))
+					if (FPackagePath::TryFromPackageName(InvalidActorDesc->GetActorPackage(), PackagePath))
 					{
 						const FString ActorFile = PackagePath.GetLocalFullPath();
 						FSourceControlStatePtr SCState = SourceControlProvider.GetState(ActorFile, EStateCacheUsage::ForceUpdate);
diff --git a/Engine/Source/Runtime/Engine/Private/LevelInstance/LevelInstanceSubsystem.cpp b/Engine/Source/Runtime/Engine/Private/LevelInstance/LevelInstanceSubsystem.cpp
--- a/Engine/Source/Runtime/Engine/Private/LevelInstance/LevelInstanceSubsystem.cpp
+++ b/Engine/Source/Runtime/Engine/Private/LevelInstance/LevelInstanceSubsystem.cpp
@@ -1468,7 +1468,7 @@ UActorDescContainer* ULevelInstanceSubsystem::FActorDescContainerInstanceManager
 		ExistingContainerInstance->Container = ActorDescContainer;
 		
 		// This will potentially invalidate ExistingContainerInstance due to ActorDescContainers reallocation
-		ActorDescContainer->Initialize(InWorld, PackageName);
+		ActorDescContainer->Initialize({ InWorld, PackageName });
 			
 		ExistingContainerInstance = &ActorDescContainers.FindChecked(PackageName);
 		ExistingContainerInstance->UpdateBounds();
diff --git a/Engine/Source/Runtime/Engine/Private/Tests/WorldPartition/WorldPartitionHandleTests.cpp b/Engine/Source/Runtime/Engine/Private/Tests/WorldPartition/WorldPartitionHandleTests.cpp
--- a/Engine/Source/Runtime/Engine/Private/Tests/WorldPartition/WorldPartitionHandleTests.cpp
+++ b/Engine/Source/Runtime/Engine/Private/Tests/WorldPartition/WorldPartitionHandleTests.cpp
@@ -27,7 +27,7 @@ namespace WorldPartitionTests
 		auto TestFalse = [Test](const TCHAR* What, bool Value) -> bool { return Test->TestFalse(What, Value); };
 
 		UActorDescContainer* ActorDescContainer = NewObject<UActorDescContainer>(GetTransientPackage());
-		ActorDescContainer->Initialize(nullptr, TEXT("/Engine/WorldPartition/WorldPartitionUnitTest"));
+		ActorDescContainer->Initialize({ nullptr, TEXT("/Engine/WorldPartition/WorldPartitionUnitTest") });
 
 		FWorldPartitionHandle Handle = FWorldPartitionHandle(ActorDescContainer, FGuid(TEXT("5D9F93BA407A811AFDDDAAB4F1CECC6A")));
 
diff --git a/Engine/Source/Runtime/Engine/Private/WorldPartition/ActorDescContainer.cpp b/Engine/Source/Runtime/Engine/Private/WorldPartition/ActorDescContainer.cpp
--- a/Engine/Source/Runtime/Engine/Private/WorldPartition/ActorDescContainer.cpp
+++ b/Engine/Source/Runtime/Engine/Private/WorldPartition/ActorDescContainer.cpp
@@ -26,11 +26,17 @@ UActorDescContainer::UActorDescContainer(const FObjectInitializer& ObjectInitial
 
 void UActorDescContainer::Initialize(UWorld* InWorld, FName InPackageName)
 {
-	check(!World || World == InWorld);
-	World = InWorld;
+	Initialize({ InWorld, InPackageName });
+}
+
+void UActorDescContainer::Initialize(const FInitializeParams& InitParams)
+{
+	check(!World || World == InitParams.World);
+	World = InitParams.World;
+
 #if WITH_EDITOR
 	check(!bContainerInitialized);
-	ContainerPackageName = InPackageName;
+	ContainerPackageName = InitParams.PackageName;
 	TArray<FAssetData> Assets;
 
 	if (!ContainerPackageName.IsNone())
@@ -53,13 +59,13 @@ void UActorDescContainer::Initialize(UWorld* InWorld, FName InPackageName)
 	{
 		TUniquePtr<FWorldPartitionActorDesc> ActorDesc = FWorldPartitionActorDescUtils::GetActorDescriptorFromAssetData(Asset);
 
-		if (ActorDesc.IsValid())
+		if (ActorDesc.IsValid() && (!InitParams.FilterActorDesc || InitParams.FilterActorDesc(ActorDesc.Get())))
 		{
 			AddActorDescriptor(ActorDesc.Release());
 		}
 		else
 		{
-			InvalidActors.Add(Asset);
+			InvalidActors.Emplace(MoveTemp(ActorDesc));
 		}
 	}
 
diff --git a/Engine/Source/Runtime/Engine/Private/WorldPartition/WorldPartition.cpp b/Engine/Source/Runtime/Engine/Private/WorldPartition/WorldPartition.cpp
--- a/Engine/Source/Runtime/Engine/Private/WorldPartition/WorldPartition.cpp
+++ b/Engine/Source/Runtime/Engine/Private/WorldPartition/WorldPartition.cpp
@@ -1071,22 +1071,6 @@ void UWorldPartition::UnhashActorDesc(FWorldPartitionActorDesc* ActorDesc)
 	FWorldPartitionHandle ActorHandle(this, ActorDesc->GetGuid());
 	EditorHash->UnhashActor(ActorHandle);
 }
-
-void UWorldPartition::HashActorDescContainer(UActorDescContainer* InActorDescContainer)
-{
-	for (UActorDescContainer::TIterator<> It(InActorDescContainer); It; ++It)
-	{
-		HashActorDesc(*It);
-	}
-}
-
-void UWorldPartition::UnhashActorDescContainer(UActorDescContainer* InActorDescContainer)
-{
-	for (UActorDescContainer::TIterator<> It(InActorDescContainer); It; ++It)
-	{
-		UnhashActorDesc(*It);
-	}
-}
 #endif
 
 void UWorldPartition::Serialize(FArchive& Ar)
@@ -1486,7 +1470,11 @@ UActorDescContainer* UWorldPartition::RegisterActorDescContainer(const FName& Co
 	if (!Contains(ContainerPackage))
 	{
 		UActorDescContainer* ContainerToRegister = NewObject<UActorDescContainer>(this);
-		ContainerToRegister->Initialize(GetWorld(), ContainerPackage);
+
+		UActorDescContainer::FInitializeParams ContainerInitParams(GetWorld(), ContainerPackage);
+		ContainerInitParams.FilterActorDesc = [this](const FWorldPartitionActorDesc* ActorDesc) { return !GetActorDesc(ActorDesc->GetGuid()); };
+		ContainerToRegister->Initialize(ContainerInitParams);
+
 		AddContainer(ContainerToRegister);
 
 		if (IsInitialized() && EditorHash != nullptr)
@@ -1501,7 +1489,10 @@ UActorDescContainer* UWorldPartition::RegisterActorDescContainer(const FName& Co
 				}
 			}
 
-			HashActorDescContainer(ContainerToRegister);
+			for (UActorDescContainer::TIterator<> It(ContainerToRegister); It; ++It)
+			{
+				HashActorDesc(*It);
+			}
 		}
 
 		OnActorDescContainerRegistered.Broadcast(ContainerToRegister);
@@ -1533,7 +1524,10 @@ bool UWorldPartition::UnregisterActorDescContainer(UActorDescContainer* InActorD
 
 		if (IsInitialized() && EditorHash != nullptr)
 		{
-			UnhashActorDescContainer(InActorDescContainer);
+			for (UActorDescContainer::TIterator<> It(InActorDescContainer); It; ++It)
+			{
+				UnhashActorDesc(*It);
+			}
 		}
 
 		InActorDescContainer->Uninitialize();
diff --git a/Engine/Source/Runtime/Engine/Public/WorldPartition/ActorDescContainer.h b/Engine/Source/Runtime/Engine/Public/WorldPartition/ActorDescContainer.h
--- a/Engine/Source/Runtime/Engine/Public/WorldPartition/ActorDescContainer.h
+++ b/Engine/Source/Runtime/Engine/Public/WorldPartition/ActorDescContainer.h
@@ -20,7 +20,27 @@ class ENGINE_API UActorDescContainer : public UObject, public FActorDescList
 	friend class FWorldPartitionActorDesc;
 
 public:
+	/* Struct of parameters passed to Initialize function. */
+	struct ENGINE_API FInitializeParams
+	{
+		FInitializeParams(UWorld* InWorld, FName InPackageName)
+			: World(InWorld)
+			, PackageName(InPackageName)
+		{}
+
+		/* The world the actor descriptor container is associated with. */
+		UWorld* World;
+			
+		/* The long package name of the container package on disk. */
+		FName PackageName;
+
+		/* Custom filter function used to filter actors descriptors. */
+		TUniqueFunction<bool(const FWorldPartitionActorDesc*)> FilterActorDesc;
+	};
+
+	UE_DEPRECATED(5.1, "UActorDescContainer::Initialize is deprecated, UActorDescContainer::Initialize with UActorDescContainer::FInitializeParams should be used instead.")
 	void Initialize(UWorld* World, FName InPackageName);
+	void Initialize(const FInitializeParams& InitParams);
 	void Uninitialize();
 	
 	virtual UWorld* GetWorld() const override;
@@ -65,7 +85,7 @@ public:
 	const FTransform& GetInstanceTransform() const;
 
 	bool HasInvalidActors() const { return InvalidActors.Num() > 0; }
-	const TArray<FAssetData>& GetInvalidActors() const { return InvalidActors; }
+	const TArray<TUniquePtr<FWorldPartitionActorDesc>>& GetInvalidActors() const { return InvalidActors; }
 	void ClearInvalidActors() { InvalidActors.Empty(); }
 #endif
 
@@ -95,7 +115,7 @@ protected:
 	FName ContainerPackageName;
 	FGuid ContentBundleGuid;
 
-	TArray<FAssetData> InvalidActors;
+	TArray<TUniquePtr<FWorldPartitionActorDesc>> InvalidActors;
 #endif
 
 private:
