World Partition Editor: added code to detect actor descriptors that changes their bounds/spatially loaded flag before being unhashed.

#rnx

diff --git a/Engine/Source/Runtime/Engine/Private/WorldPartition/WorldPartitionEditorSpatialHash.cpp b/Engine/Source/Runtime/Engine/Private/WorldPartition/WorldPartitionEditorSpatialHash.cpp
--- a/Engine/Source/Runtime/Engine/Private/WorldPartition/WorldPartitionEditorSpatialHash.cpp
+++ b/Engine/Source/Runtime/Engine/Private/WorldPartition/WorldPartitionEditorSpatialHash.cpp
@@ -2,6 +2,7 @@
 
 #include "WorldPartition/WorldPartitionEditorSpatialHash.h"
 #include "WorldPartition/WorldPartition.h"
+#include "WorldPartition/WorldPartitionLog.h"
 #include "WorldPartition/WorldPartitionActorDesc.h"
 #include "WorldPartition/WorldPartitionActorDesc.h"
 #include "WorldPartition/WorldPartitionEditorPerProjectUserSettings.h"
@@ -112,13 +113,19 @@ void UWorldPartitionEditorSpatialHash::HashActor(FWorldPartitionHandle& InActorH
 {
 	check(InActorHandle.IsValid());
 
+	const FBox ActorBounds = InActorHandle->GetIsSpatiallyLoaded() ? InActorHandle->GetBounds() : FBox(ForceInit);
+
+#if DO_CHECK
+	check(!HashedActors.Contains(InActorHandle->GetGuid()));
+	HashedActors.Add(InActorHandle->GetGuid(), ActorBounds);
+#endif
+
 	if (!InActorHandle->GetIsSpatiallyLoaded())
 	{
 		AlwaysLoadedCell->Actors.Add(InActorHandle);
 	}
 	else
 	{
-		const FBox ActorBounds = InActorHandle->GetBounds();
 		const int32 CurrentLevel = GetLevelForBox(EditorBounds);
 		const int32 ActorLevel = GetLevelForBox(ActorBounds);
 
@@ -216,13 +223,19 @@ void UWorldPartitionEditorSpatialHash::UnhashActor(FWorldPartitionHandle& InActo
 {
 	check(InActorHandle.IsValid());
 
+	const FBox ActorBounds = InActorHandle->GetIsSpatiallyLoaded() ? InActorHandle->GetBounds() : FBox(ForceInit);
+
+#if DO_CHECK
+	FBox OldActorBounds = HashedActors.FindAndRemoveChecked(InActorHandle->GetGuid());
+	check(ActorBounds == OldActorBounds);
+#endif
+
 	if (!InActorHandle->GetIsSpatiallyLoaded())
 	{
 		AlwaysLoadedCell->Actors.Remove(InActorHandle);
 	}
 	else
 	{
-		const FBox ActorBounds = InActorHandle->GetBounds();
 		const int32 CurrentLevel = GetLevelForBox(EditorBounds);
 		const int32 ActorLevel = GetLevelForBox(ActorBounds);
 
diff --git a/Engine/Source/Runtime/Engine/Public/WorldPartition/WorldPartitionEditorSpatialHash.h b/Engine/Source/Runtime/Engine/Public/WorldPartition/WorldPartitionEditorSpatialHash.h
--- a/Engine/Source/Runtime/Engine/Public/WorldPartition/WorldPartitionEditorSpatialHash.h
+++ b/Engine/Source/Runtime/Engine/Public/WorldPartition/WorldPartitionEditorSpatialHash.h
@@ -219,7 +219,11 @@ private:
 	
 	FBox EditorBounds;
 	FBox RuntimeBounds;
-	bool bBoundsDirty;	
+	bool bBoundsDirty;
+
+#if DO_CHECK
+	TMap<FGuid, FBox> HashedActors;
+#endif
 
 public:
 	UPROPERTY(Config, meta = (AllowedClasses = "/Script/Engine.Texture2D, /Script/Engine.MaterialInterface"))
