Remove BlueprintAsset property. (Can be inferred from ClassGeneratedBy)

diff --git a/Engine/Source/Editor/LevelInstanceEditor/Private/LevelInstanceEditorModule.cpp b/Engine/Source/Editor/LevelInstanceEditor/Private/LevelInstanceEditorModule.cpp
--- a/Engine/Source/Editor/LevelInstanceEditor/Private/LevelInstanceEditorModule.cpp
+++ b/Engine/Source/Editor/LevelInstanceEditor/Private/LevelInstanceEditorModule.cpp
@@ -477,7 +477,7 @@ namespace LevelInstanceMenuUtils
 				;
 				if (APackedLevelActor* PackedLevelActor = Cast<APackedLevelActor>(ContextLevelInstance))
 				{
-					if (TSoftObjectPtr<UBlueprint> BlueprintAsset = PackedLevelActor->BlueprintAsset)
+					if (TSoftObjectPtr<UBlueprint> BlueprintAsset = PackedLevelActor->GetClass()->ClassGeneratedBy)
 					{
 						FToolUIAction UIAction;
 						UIAction.ExecuteAction.BindLambda([ContextLevelInstance, BlueprintAsset](const FToolMenuContext& MenuContext)
diff --git a/Engine/Source/Runtime/Engine/Private/PackedLevelActor/PackedLevelActor.cpp b/Engine/Source/Runtime/Engine/Private/PackedLevelActor/PackedLevelActor.cpp
--- a/Engine/Source/Runtime/Engine/Private/PackedLevelActor/PackedLevelActor.cpp
+++ b/Engine/Source/Runtime/Engine/Private/PackedLevelActor/PackedLevelActor.cpp
@@ -150,7 +150,6 @@ void APackedLevelActor::OnCommitChild(bool bChanged)
 
 			if (UBlueprint* GeneratedBy = Cast<UBlueprint>(GetClass()->ClassGeneratedBy))
 			{
-				check(GeneratedBy == BlueprintAsset.Get());
 				Builder->UpdateBlueprint(GeneratedBy);
 				return; // return here because Actor might have been reinstanced
 			}
@@ -178,7 +177,6 @@ void APackedLevelActor::OnCommit(bool bChanged)
 	{
 		if (UBlueprint* GeneratedBy = Cast<UBlueprint>(GetClass()->ClassGeneratedBy))
 		{
-			check(GeneratedBy == BlueprintAsset.Get());
 			TSharedPtr<FPackedLevelActorBuilder> Builder = FPackedLevelActorBuilder::CreateDefaultBuilder();
 			Builder->UpdateBlueprint(GeneratedBy);
 			return; // return here because Actor might have been reinstanced
diff --git a/Engine/Source/Runtime/Engine/Private/PackedLevelActor/PackedLevelActorBuilder.cpp b/Engine/Source/Runtime/Engine/Private/PackedLevelActor/PackedLevelActorBuilder.cpp
--- a/Engine/Source/Runtime/Engine/Private/PackedLevelActor/PackedLevelActorBuilder.cpp
+++ b/Engine/Source/Runtime/Engine/Private/PackedLevelActor/PackedLevelActorBuilder.cpp
@@ -117,7 +117,6 @@ UBlueprint* FPackedLevelActorBuilder::CreatePackedLevelActorBlueprint(TSoftObjec
 		if (UBlueprint* NewBP = Cast<UBlueprint>(AssetTools.CreateAsset(InBlueprintAsset.GetAssetName(), PackageDir, UBlueprint::StaticClass(), BlueprintFactory, FName("Create LevelInstance Blueprint"))))
 		{
 			APackedLevelActor* CDO = CastChecked<APackedLevelActor>(NewBP->GeneratedClass->GetDefaultObject());
-			CDO->BlueprintAsset = NewBP;
 			CDO->SetWorldAsset(InWorldAsset);
 
 			if (bInCompile)
@@ -383,8 +382,7 @@ bool FPackedLevelActorBuilder::CreateOrUpdateBlueprintFromUnpacked(ALevelInstanc
 		return false;
 	}
 
-	PackedLevelActor->BlueprintAsset = InBlueprintAsset;
-	bResult &= CreateOrUpdateBlueprintFromPacked(PackedLevelActor, PackedLevelActor->BlueprintAsset, bCheckoutAndSave, bPromptForSave);
+	bResult &= CreateOrUpdateBlueprintFromPacked(PackedLevelActor, InBlueprintAsset, bCheckoutAndSave, bPromptForSave);
 
 	return bResult;
 }
diff --git a/Engine/Source/Runtime/Engine/Public/PackedLevelActor/PackedLevelActor.h b/Engine/Source/Runtime/Engine/Public/PackedLevelActor/PackedLevelActor.h
--- a/Engine/Source/Runtime/Engine/Public/PackedLevelActor/PackedLevelActor.h
+++ b/Engine/Source/Runtime/Engine/Public/PackedLevelActor/PackedLevelActor.h
@@ -74,9 +74,6 @@ public:
 #endif
 
 #if WITH_EDITORONLY_DATA
-	UPROPERTY(VisibleAnywhere, Category = Packed)
-	TSoftObjectPtr<UBlueprint> BlueprintAsset;
-
 	UPROPERTY()
 	TArray<TSoftObjectPtr<UBlueprint>> PackedBPDependencies;
 
