Spline: Fixed an issue when computing the distance of a full spline segment that uses constant interpolation.
 Function incorrectly returned a distance of 0 for all partial distances however when computing the "full" distance with p=1 it should have returned the true distance between the points.

This fixes another issue with generation of the ReparamTable where spline keys would map to incorrect distances if a spline had any number of constant interp points.

diff --git a/Engine/Source/Runtime/Engine/Private/Components/SplineComponent.cpp b/Engine/Source/Runtime/Engine/Private/Components/SplineComponent.cpp
--- a/Engine/Source/Runtime/Engine/Private/Components/SplineComponent.cpp
+++ b/Engine/Source/Runtime/Engine/Private/Components/SplineComponent.cpp
@@ -268,7 +268,8 @@ float FSplineCurves::GetSegmentLength(const int32 Index, const float Param, bool
 	}
 	else if (StartPoint.InterpMode == CIM_Constant)
 	{
-		return 0.0f;
+		// Special case: constant interpolation acts like distance = 0 for all p in [0, 1[ but for p == 1, the distance returned is the linear distance between start and end
+		return Param == 1.f ? ((P1 - P0) * Scale3D).Size() : 0.0f;
 	}
 
 	// Cache the coefficients to be fed into the function to calculate the spline derivative at each sample point as they are constant.
