Fix for shortname MetaData in Dialogue Wave

diff --git a/Engine/Source/Runtime/Engine/Classes/Sound/DialogueWave.h b/Engine/Source/Runtime/Engine/Classes/Sound/DialogueWave.h
--- a/Engine/Source/Runtime/Engine/Classes/Sound/DialogueWave.h
+++ b/Engine/Source/Runtime/Engine/Classes/Sound/DialogueWave.h
@@ -52,7 +52,7 @@ struct ENGINE_API FDialogueContextMapping
 	FDialogueContext Context;
 
 	/** The soundwave to play for this dialogue. */
-	UPROPERTY(EditAnywhere, Category=DialogueContextMapping, meta = (DisallowedClasses="MetaSoundSource, SoundSourceBus"))
+	UPROPERTY(EditAnywhere, Category=DialogueContextMapping, meta = (DisallowedClasses= "/Script/MetasoundEngine.MetaSoundSource, /Script/Engine.SoundSourceBus"))
 	TObjectPtr<USoundWave> SoundWave;
 
 	/**
