- Add RGB / RGBA pins to dynamic parameter material node
- Show pin names and add RGBA to particle color node

diff --git a/Engine/Source/Runtime/Engine/Private/Materials/MaterialExpressions.cpp b/Engine/Source/Runtime/Engine/Private/Materials/MaterialExpressions.cpp
--- a/Engine/Source/Runtime/Engine/Private/Materials/MaterialExpressions.cpp
+++ b/Engine/Source/Runtime/Engine/Private/Materials/MaterialExpressions.cpp
@@ -10016,12 +10016,14 @@ UMaterialExpressionParticleColor::UMaterialExpressionParticleColor(const FObject
 	MenuCategories.Add(ConstructorStatics.NAME_Constants);
 
 	Outputs.Reset();
-	Outputs.Add(FExpressionOutput(TEXT(""), 1, 1, 1, 1, 0));
-	Outputs.Add(FExpressionOutput(TEXT(""), 1, 1, 0, 0, 0));
-	Outputs.Add(FExpressionOutput(TEXT(""), 1, 0, 1, 0, 0));
-	Outputs.Add(FExpressionOutput(TEXT(""), 1, 0, 0, 1, 0));
-	Outputs.Add(FExpressionOutput(TEXT(""), 1, 0, 0, 0, 1));
+	Outputs.Add(FExpressionOutput(TEXT("RGB"), 1, 1, 1, 1, 0));
+	Outputs.Add(FExpressionOutput(TEXT("R"), 1, 1, 0, 0, 0));
+	Outputs.Add(FExpressionOutput(TEXT("G"), 1, 0, 1, 0, 0));
+	Outputs.Add(FExpressionOutput(TEXT("B"), 1, 0, 0, 1, 0));
+	Outputs.Add(FExpressionOutput(TEXT("A"), 1, 0, 0, 0, 1));
+	Outputs.Add(FExpressionOutput(TEXT("RGBA"), 1, 1, 1, 1, 1));
 
+	bShowOutputNameOnPin = true;
 	bShaderInputData = true;
 #endif
 }
@@ -10148,6 +10150,8 @@ UMaterialExpressionDynamicParameter::UMaterialExpressionDynamicParameter(const F
 	Outputs.Add(FExpressionOutput(TEXT(""), 1, 0, 1, 0, 0));
 	Outputs.Add(FExpressionOutput(TEXT(""), 1, 0, 0, 1, 0));
 	Outputs.Add(FExpressionOutput(TEXT(""), 1, 0, 0, 0, 1));
+	Outputs.Add(FExpressionOutput(TEXT("RGB"), 1, 1, 1, 1, 0));
+	Outputs.Add(FExpressionOutput(TEXT("RGBA"), 1, 1, 1, 1, 1));
 
 	bShaderInputData = true;
 #endif // WITH_EDITORONLY_DATA
@@ -10173,7 +10177,6 @@ TArray<FExpressionOutput>& UMaterialExpressionDynamicParameter::GetOutputs()
 	return Outputs;
 }
 
-
 int32 UMaterialExpressionDynamicParameter::GetWidth() const
 {
 	return ME_STD_THUMBNAIL_SZ+(ME_STD_BORDER*2);
