Fix crash when fixing up vertex colors on a static mesh component which is marked as "custom override vertex color per LOD", but which has an unpainted LOD.

diff --git a/Engine/Source/Runtime/Engine/Private/Components/StaticMeshComponent.cpp b/Engine/Source/Runtime/Engine/Private/Components/StaticMeshComponent.cpp
--- a/Engine/Source/Runtime/Engine/Private/Components/StaticMeshComponent.cpp
+++ b/Engine/Source/Runtime/Engine/Private/Components/StaticMeshComponent.cpp
@@ -1604,18 +1604,19 @@ void UStaticMeshComponent::PrivateFixupOverrideColors()
 		TArray<FColor> NewOverrideColors;
 		if (bCustomOverrideVertexColorPerLOD)
 		{
-			//Since in custom we fix paint only if the component has some, the PaintedVertices should be allocate
-			check(LODInfo.PaintedVertices.Num() > 0);
-			//Use the existing LOD custom paint and remap it on the new mesh
-			RemapPaintedVertexColors(
-				LODInfo.PaintedVertices,
-				nullptr,
-				SourceRenderData.VertexBuffers.PositionVertexBuffer,
-				SourceRenderData.VertexBuffers.StaticMeshVertexBuffer,
-				CurRenderData.VertexBuffers.PositionVertexBuffer,
-				&CurRenderData.VertexBuffers.StaticMeshVertexBuffer,
-				NewOverrideColors
-				);
+			if (LODInfo.PaintedVertices.Num() > 0)
+			{
+				//Use the existing LOD custom paint and remap it on the new mesh
+				RemapPaintedVertexColors(
+					LODInfo.PaintedVertices,
+					nullptr,
+					SourceRenderData.VertexBuffers.PositionVertexBuffer,
+					SourceRenderData.VertexBuffers.StaticMeshVertexBuffer,
+					CurRenderData.VertexBuffers.PositionVertexBuffer,
+					&CurRenderData.VertexBuffers.StaticMeshVertexBuffer,
+					NewOverrideColors
+					);
+			}
 		}
 		else if(LOD0Info.PaintedVertices.Num() > 0)
 		{
