Suggested fix for Shutdown async loading crash

diff --git a/Engine/Source/Runtime/Engine/Private/SoundNodeWavePlayer.cpp b/Engine/Source/Runtime/Engine/Private/SoundNodeWavePlayer.cpp
--- a/Engine/Source/Runtime/Engine/Private/SoundNodeWavePlayer.cpp
+++ b/Engine/Source/Runtime/Engine/Private/SoundNodeWavePlayer.cpp
@@ -31,7 +31,7 @@ void USoundNodeWavePlayer::Serialize(FArchive& Ar)
 
 void USoundNodeWavePlayer::LoadAsset(bool bAddToRoot)
 {
-	if (IsAsyncLoading())
+	if (IsAsyncLoadingMultithreaded())
 	{
 		SoundWave = SoundWaveAssetPtr.Get();
 		if (SoundWave && SoundWave->HasAnyFlags(RF_NeedLoad))
