TextureCube.cpp nop comment

diff --git a/Engine/Source/Runtime/Engine/Private/TextureCube.cpp b/Engine/Source/Runtime/Engine/Private/TextureCube.cpp
--- a/Engine/Source/Runtime/Engine/Private/TextureCube.cpp
+++ b/Engine/Source/Runtime/Engine/Private/TextureCube.cpp
@@ -168,6 +168,8 @@ void UTextureCube::PostLoad()
 #endif // #if WITH_EDITOR
 
 	Super::PostLoad();
+
+	// note UTexture::PostLoad casts to UTextureCube:: and calls ->UpdateResource  (should just do that here instead))
 }
 
 void UTextureCube::GetAssetRegistryTags(TArray<FAssetRegistryTag>& OutTags) const
