fix crash: Texture Downscale is not allowed on vt

diff --git a/Engine/Source/Runtime/Engine/Private/TextureDerivedData.cpp b/Engine/Source/Runtime/Engine/Private/TextureDerivedData.cpp
--- a/Engine/Source/Runtime/Engine/Private/TextureDerivedData.cpp
+++ b/Engine/Source/Runtime/Engine/Private/TextureDerivedData.cpp
@@ -965,7 +965,10 @@ static void GetTextureBuildSettings(
 	}
 
 	OutBuildSettings.Downscale = 1.0f;
-	if (MipGenSettings == TMGS_NoMipmaps && 
+	// Downscale only allowed if NoMipMaps, 2d, and not VT
+	//	silently does nothing otherwise
+	if (! bVirtualTextureStreaming &&
+		MipGenSettings == TMGS_NoMipmaps && 
 		Texture.IsA(UTexture2D::StaticClass()))	// TODO: support more texture types
 	{
 		TextureLODSettings.GetDownscaleOptions(Texture, TargetPlatform, OutBuildSettings.Downscale, (ETextureDownscaleOptions&)OutBuildSettings.DownscaleOptions);
