Strata - Fixed shading model input of conversion to be of the correct type.

diff --git a/Engine/Source/Runtime/Engine/Private/Materials/MaterialExpressions.cpp b/Engine/Source/Runtime/Engine/Private/Materials/MaterialExpressions.cpp
--- a/Engine/Source/Runtime/Engine/Private/Materials/MaterialExpressions.cpp
+++ b/Engine/Source/Runtime/Engine/Private/Materials/MaterialExpressions.cpp
@@ -21943,7 +21943,7 @@ uint32 UMaterialExpressionStrataLegacyConversion::GetInputType(int32 InputIndex)
 	else if (InputIndex == 16) return MCT_Float3; // ColorScaleBehindWater
 	else if (InputIndex == 17) return MCT_Float3; // ClearCoatNormal
 	else if (InputIndex == 18) return MCT_Float3; // CustomTangent
-	else if (InputIndex == 19) return MCT_Float1; // ShadingModel
+	else if (InputIndex == 19) return MCT_ShadingModel;	// ShadingModel
 
 	check(false);
 	return MCT_Float1;
