Strata - fixed bad check causing crash when it can be avoided.
The user is now informed with a warning about the wrong material conversion setup process that has happened, and how to fix it.

diff --git a/Engine/Source/Runtime/Engine/Private/Materials/Material.cpp b/Engine/Source/Runtime/Engine/Private/Materials/Material.cpp
--- a/Engine/Source/Runtime/Engine/Private/Materials/Material.cpp
+++ b/Engine/Source/Runtime/Engine/Private/Materials/Material.cpp
@@ -4557,6 +4557,11 @@ void UMaterial::RebuildShadingModelField()
 	const bool bStrataEnabled = Engine_IsStrataEnabled();
 	UMaterialEditorOnlyData* EditorOnly = GetEditorOnlyData();
 
+	auto StrataCheckShadingModelFromExpressionWarning = [&]()
+	{
+		UE_LOG(LogMaterial, Warning, TEXT("%s: Material seems to have been converted with ShadingModel `FromExpression`, but now is using `%s`.\nThis indicates the material has been resaved with Strata disabled when it already had a conversion node setup.\nThe material may look wrong: please delete all nodes plugged in FrontMaterial to rerun the conversion process."), *GetName(), *GetShadingModelString(ShadingModel));
+	};
+
 	if (bStrataEnabled && EditorOnly->FrontMaterial.IsConnected())
 	{
 		FStrataMaterialInfo StrataMaterialInfo;
@@ -4583,6 +4588,8 @@ void UMaterial::RebuildShadingModelField()
 			}
 			else if (StrataMaterialInfo.HasShadingModelFromExpression())
 			{
+				StrataCheckShadingModelFromExpressionWarning();
+
 				if (BlendMode == EBlendMode::BLEND_Opaque || BlendMode == EBlendMode::BLEND_Masked)
 				{
 					// Nothing
@@ -4629,7 +4636,8 @@ void UMaterial::RebuildShadingModelField()
 		{
 			MaterialDomain = EMaterialDomain::MD_Surface;
 
-			check(ShadingModel == MSM_FromMaterialExpression);
+			StrataCheckShadingModelFromExpressionWarning();
+
 			{
 				TArray<UMaterialExpressionShadingModel*> ShadingModelExpressions;
 				GetAllExpressionsInMaterialAndFunctionsOfType(ShadingModelExpressions);
