Follow-up to 22133564
#fix make sure that when copying float curve data into AnimDataModel, the pre/post infinity extrapolation settings are copied over correctly

diff --git a/Engine/Source/Editor/CurveEditor/Public/CurveDataAbstraction.h b/Engine/Source/Editor/CurveEditor/Public/CurveDataAbstraction.h
--- a/Engine/Source/Editor/CurveEditor/Public/CurveDataAbstraction.h
+++ b/Engine/Source/Editor/CurveEditor/Public/CurveDataAbstraction.h
@@ -194,72 +194,3 @@ private:
 	float LeaveTangentWeight{};
 
 };
-
-/**
- * Structure allowing external curve data to express extended attributes
- */
-struct FCurveAttributes
-{
-	FCurveAttributes()
-	{
-		bHasPreExtrapolation  = 0;
-		bHasPostExtrapolation = 0;
-		PreExtrapolation = ERichCurveExtrapolation::RCCE_None;
-		PostExtrapolation = ERichCurveExtrapolation::RCCE_None;
-	}
-
-	/**
-	 * Check whether this curve has the specified properties
-	 */
-	bool HasPreExtrapolation() const  { return bHasPreExtrapolation;  }
-	bool HasPostExtrapolation() const { return bHasPostExtrapolation; }
-
-	/**
-	 * Access the extended properties of this curve. Must check whether the curve has such properties first
-	 */
-	ERichCurveExtrapolation GetPreExtrapolation() const  { check(bHasPreExtrapolation);  return PreExtrapolation;  }
-	ERichCurveExtrapolation GetPostExtrapolation() const { check(bHasPostExtrapolation); return PostExtrapolation; }
-
-	/**
-	 * Set the extended properties of this curve
-	 */
-	FCurveAttributes& SetPreExtrapolation(ERichCurveExtrapolation InPreExtrapolation)   { bHasPreExtrapolation = 1;  PreExtrapolation = InPreExtrapolation;   return *this; }
-	FCurveAttributes& SetPostExtrapolation(ERichCurveExtrapolation InPostExtrapolation) { bHasPostExtrapolation = 1; PostExtrapolation = InPostExtrapolation; return *this; }
-	/**
-	 * Reset the extended properties of this curve, implying it does not support such properties
-	 */
-	void UnsetPreExtrapolation()  { bHasPreExtrapolation = 0;  }
-	void UnsetPostExtrapolation() { bHasPostExtrapolation = 0; }
-
-	/**
-	 * Generate a new set of attributes that contains only those attributes common to both A and B
-	 */
-	static FCurveAttributes MaskCommon(const FCurveAttributes& A, const FCurveAttributes& B)
-	{
-		FCurveAttributes NewAttributes;
-
-		if (A.bHasPreExtrapolation && B.bHasPreExtrapolation && A.PreExtrapolation == B.PreExtrapolation)
-		{
-			NewAttributes.SetPreExtrapolation(A.PreExtrapolation);
-		}
-
-		if (A.bHasPostExtrapolation && B.bHasPostExtrapolation && A.PostExtrapolation == B.PostExtrapolation)
-		{
-			NewAttributes.SetPostExtrapolation(A.PostExtrapolation);
-		}
-
-		return NewAttributes;
-	}
-
-private:
-
-	/** true if the curve can express pre-extrapolation modes */
-	uint8 bHasPreExtrapolation : 1;
-	/** true if the curve can express post-extrapolation modes */
-	uint8 bHasPostExtrapolation : 1;
-
-	/** This curve's pre-extrapolation mode. Only valid to read if bHasPreExtrapolation is true */
-	ERichCurveExtrapolation PreExtrapolation;
-	/** This curve's post-extrapolation mode. Only valid to read if bHasPostExtrapolation is true */
-	ERichCurveExtrapolation PostExtrapolation;
-};
\ No newline at end of file
diff --git a/Engine/Source/Runtime/Engine/Classes/Curves/RealCurve.h b/Engine/Source/Runtime/Engine/Classes/Curves/RealCurve.h
--- a/Engine/Source/Runtime/Engine/Classes/Curves/RealCurve.h
+++ b/Engine/Source/Runtime/Engine/Classes/Curves/RealCurve.h
@@ -39,6 +39,76 @@ enum ERichCurveExtrapolation
 	RCCE_None UMETA(DisplayName = "None")
 };
 
+
+/**
+ * Structure allowing external curve data to express extended attributes
+ */
+struct FCurveAttributes
+{
+	FCurveAttributes()
+	{
+		bHasPreExtrapolation  = 0;
+		bHasPostExtrapolation = 0;
+		PreExtrapolation = ERichCurveExtrapolation::RCCE_None;
+		PostExtrapolation = ERichCurveExtrapolation::RCCE_None;
+	}
+
+	/**
+	 * Check whether this curve has the specified properties
+	 */
+	bool HasPreExtrapolation() const  { return bHasPreExtrapolation;  }
+	bool HasPostExtrapolation() const { return bHasPostExtrapolation; }
+
+	/**
+	 * Access the extended properties of this curve. Must check whether the curve has such properties first
+	 */
+	ERichCurveExtrapolation GetPreExtrapolation() const  { check(bHasPreExtrapolation);  return PreExtrapolation;  }
+	ERichCurveExtrapolation GetPostExtrapolation() const { check(bHasPostExtrapolation); return PostExtrapolation; }
+
+	/**
+	 * Set the extended properties of this curve
+	 */
+	FCurveAttributes& SetPreExtrapolation(ERichCurveExtrapolation InPreExtrapolation)   { bHasPreExtrapolation = 1;  PreExtrapolation = InPreExtrapolation;   return *this; }
+	FCurveAttributes& SetPostExtrapolation(ERichCurveExtrapolation InPostExtrapolation) { bHasPostExtrapolation = 1; PostExtrapolation = InPostExtrapolation; return *this; }
+	/**
+	 * Reset the extended properties of this curve, implying it does not support such properties
+	 */
+	void UnsetPreExtrapolation()  { bHasPreExtrapolation = 0;  }
+	void UnsetPostExtrapolation() { bHasPostExtrapolation = 0; }
+
+	/**
+	 * Generate a new set of attributes that contains only those attributes common to both A and B
+	 */
+	static FCurveAttributes MaskCommon(const FCurveAttributes& A, const FCurveAttributes& B)
+	{
+		FCurveAttributes NewAttributes;
+
+		if (A.bHasPreExtrapolation && B.bHasPreExtrapolation && A.PreExtrapolation == B.PreExtrapolation)
+		{
+			NewAttributes.SetPreExtrapolation(A.PreExtrapolation);
+		}
+
+		if (A.bHasPostExtrapolation && B.bHasPostExtrapolation && A.PostExtrapolation == B.PostExtrapolation)
+		{
+			NewAttributes.SetPostExtrapolation(A.PostExtrapolation);
+		}
+
+		return NewAttributes;
+	}
+
+private:
+
+	/** true if the curve can express pre-extrapolation modes */
+	uint8 bHasPreExtrapolation : 1;
+	/** true if the curve can express post-extrapolation modes */
+	uint8 bHasPostExtrapolation : 1;
+
+	/** This curve's pre-extrapolation mode. Only valid to read if bHasPreExtrapolation is true */
+	ERichCurveExtrapolation PreExtrapolation;
+	/** This curve's post-extrapolation mode. Only valid to read if bHasPostExtrapolation is true */
+	ERichCurveExtrapolation PostExtrapolation;
+};
+
 /** A rich, editable float curve */
 USTRUCT()
 struct ENGINE_API FRealCurve
diff --git a/Engine/Source/Runtime/Engine/Private/Animation/AnimSequenceHelpers.cpp b/Engine/Source/Runtime/Engine/Private/Animation/AnimSequenceHelpers.cpp
--- a/Engine/Source/Runtime/Engine/Private/Animation/AnimSequenceHelpers.cpp
+++ b/Engine/Source/Runtime/Engine/Private/Animation/AnimSequenceHelpers.cpp
@@ -150,6 +150,10 @@ void CopyCurveDataToModel(const FRawCurveTracks& CurveData, const USkeleton* Ske
 		if (CurveId.IsValid())
 		{
 			Controller.AddCurve(CurveId, FloatCurve.GetCurveTypeFlags());
+			FCurveAttributes Attributes;
+			Attributes.SetPreExtrapolation(FloatCurve.FloatCurve.PreInfinityExtrap);
+			Attributes.SetPostExtrapolation(FloatCurve.FloatCurve.PostInfinityExtrap);					
+			Controller.SetCurveAttributes(CurveId, Attributes);
 			Controller.SetCurveKeys(CurveId, FloatCurve.FloatCurve.GetConstRefOfKeys());
 		}
 	}
