Fix minor issues with async loading in editor

diff --git a/Engine/Plugins/Animation/ControlRig/Source/ControlRig/Private/Rigs/RigHierarchy.cpp b/Engine/Plugins/Animation/ControlRig/Source/ControlRig/Private/Rigs/RigHierarchy.cpp
--- a/Engine/Plugins/Animation/ControlRig/Source/ControlRig/Private/Rigs/RigHierarchy.cpp
+++ b/Engine/Plugins/Animation/ControlRig/Source/ControlRig/Private/Rigs/RigHierarchy.cpp
@@ -2274,10 +2274,9 @@ URigHierarchyController* URigHierarchy::GetController(bool bCreateIfNeeded)
 		 if(ensure(!IsGarbageCollecting()))
 		 {
 			 HierarchyController = NewObject<URigHierarchyController>(this, TEXT("HierarchyController"), RF_Transient);
-			 if(!ensure(IsInGameThread()))
-			 {
-				 HierarchyController->ClearInternalFlags(EInternalObjectFlags::Async);
-			 }
+			 // In case we create this object from async loading thread
+			 HierarchyController->ClearInternalFlags(EInternalObjectFlags::Async);
+
 			 HierarchyController->SetHierarchy(this);
 			 return HierarchyController;
 		 }
diff --git a/Engine/Source/Runtime/Engine/Private/Animation/AnimInstance.cpp b/Engine/Source/Runtime/Engine/Private/Animation/AnimInstance.cpp
--- a/Engine/Source/Runtime/Engine/Private/Animation/AnimInstance.cpp
+++ b/Engine/Source/Runtime/Engine/Private/Animation/AnimInstance.cpp
@@ -833,7 +833,7 @@ void UAnimInstance::AddNativeTransitionBinding(const FName& MachineName, const F
 
 bool UAnimInstance::HasNativeTransitionBinding(const FName& MachineName, const FName& PrevStateName, const FName& NextStateName, FName& OutBindingName)
 {
-	return GetProxyOnGameThread<FAnimInstanceProxy>().HasNativeTransitionBinding(MachineName, PrevStateName, NextStateName, OutBindingName);
+	return GetProxyOnAnyThread<FAnimInstanceProxy>().HasNativeTransitionBinding(MachineName, PrevStateName, NextStateName, OutBindingName);
 }
 
 void UAnimInstance::AddNativeStateEntryBinding(const FName& MachineName, const FName& StateName, const FOnGraphStateChanged& NativeEnteredDelegate)
