Scene Capture:  Cosmetic renaming of Lumen Surface Cache resolution scale parameter so it's clearly marked as currently only applying to Scene Capture (we could eventually extend it to the level surface cache, but that's up to the Lumen team).  Also moved it to the end of the post process settings, as it will be more rarely modified than other settings (request by [REDACTED_PERSON]).

#rnx

diff --git a/Engine/Source/Runtime/Engine/Classes/Engine/Scene.h b/Engine/Source/Runtime/Engine/Classes/Engine/Scene.h
--- a/Engine/Source/Runtime/Engine/Classes/Engine/Scene.h
+++ b/Engine/Source/Runtime/Engine/Classes/Engine/Scene.h
@@ -1231,9 +1231,6 @@ struct FPostProcessSettings
 	UPROPERTY(EditAnywhere, BlueprintReadWrite, Category = Overrides, meta = (PinHiddenByDefault, InlineEditConditionToggle))
 	uint32 bOverride_DynamicGlobalIlluminationMethod : 1;
 
-	UPROPERTY(EditAnywhere, BlueprintReadWrite, Category = Overrides, meta = (PinHiddenByDefault, InlineEditConditionToggle))
-	uint32 bOverride_LumenSurfaceCacheResolution : 1;
-
 	UPROPERTY(EditAnywhere, BlueprintReadWrite, Category = Overrides, meta = (PinHiddenByDefault, InlineEditConditionToggle))
 	uint32 bOverride_LumenSceneLightingQuality : 1;
 
@@ -1270,6 +1267,9 @@ struct FPostProcessSettings
 	UPROPERTY(EditAnywhere, BlueprintReadWrite, Category=Overrides, meta=(PinHiddenByDefault, InlineEditConditionToggle))
 	uint8 bOverride_LumenFrontLayerTranslucencyReflections:1;
 
+	UPROPERTY(EditAnywhere, BlueprintReadWrite, Category = Overrides, meta = (PinHiddenByDefault, InlineEditConditionToggle))
+	uint8 bOverride_LumenSurfaceCacheResolution : 1;
+
 	UPROPERTY(EditAnywhere, BlueprintReadWrite, Category = Overrides, meta = (PinHiddenByDefault, InlineEditConditionToggle))
 	uint32 bOverride_RayTracingGI : 1;
 
@@ -1567,10 +1567,6 @@ struct FPostProcessSettings
 	uint32 RayTracingGI_DEPRECATED : 1;
 #endif
 
-	/** Scale factor for Lumen Surface Cache resolution for Scene Captures.  Smaller values save GPU memory, at a cost in quality.  Defaults to 0.5 if not overridden, 1.0 is full resolution. */
-	UPROPERTY(EditAnywhere, BlueprintReadWrite, Category = "Global Illumination|Lumen Global Illumination", meta = (ClampMin = ".5", ClampMax = "1", editcondition = "bOverride_LumenSurfaceCacheResolution", DisplayName = "Lumen Surface Cache Resolution"))
-	float LumenSurfaceCacheResolution;
-
 	/** Scales Lumen Scene's quality.  Larger scales cause Lumen Scene to be calculated with a higher fidelity, which can be visible in reflections, but increase GPU cost. */
 	UPROPERTY(interp, EditAnywhere, BlueprintReadWrite, Category = "Global Illumination|Lumen Global Illumination", meta = (ClampMin = ".25", UIMax = "2", editcondition = "bOverride_LumenSceneLightingQuality", DisplayName = "Lumen Scene Lighting Quality"))
 	float LumenSceneLightingQuality;
@@ -1611,6 +1607,10 @@ struct FPostProcessSettings
 	UPROPERTY(interp, EditAnywhere, BlueprintReadWrite, Category = "Global Illumination|Lumen Global Illumination", AdvancedDisplay, meta = (ClampMin = ".1", UIMax = "2000", editcondition = "bOverride_LumenFullSkylightLeakingDistance", DisplayName = "Full Skylight Leaking Distance"))
 	float LumenFullSkylightLeakingDistance;
 
+	/** Scale factor for Lumen Surface Cache resolution, for Scene Capture.  Smaller values save GPU memory, at a cost in quality.  Defaults to 0.5 if not overridden. */
+	UPROPERTY(EditAnywhere, BlueprintReadWrite, Category = "Global Illumination|Lumen Global Illumination", meta = (ClampMin = ".5", ClampMax = "1", editcondition = "bOverride_LumenSurfaceCacheResolution", DisplayName = "Scene Capture Cache Resolution Scale"))
+	float LumenSurfaceCacheResolution;
+
 	/** Sets the ray tracing global illumination type. */
 	UPROPERTY(interp, EditAnywhere, BlueprintReadWrite, Category = "Global Illumination|Ray Tracing Global Illumination", meta = (editcondition = "bOverride_RayTracingGI", DisplayName = "Type"))
 	ERayTracingGlobalIlluminationType RayTracingGIType;
