Make registering new soundwave a thread-safe operation for async loading

diff --git a/Engine/Source/Runtime/Engine/Private/AudioDerivedData.cpp b/Engine/Source/Runtime/Engine/Private/AudioDerivedData.cpp
--- a/Engine/Source/Runtime/Engine/Private/AudioDerivedData.cpp
+++ b/Engine/Source/Runtime/Engine/Private/AudioDerivedData.cpp
@@ -1264,7 +1264,7 @@ struct FAudioCookInputs
 #endif
 	{
 #if WITH_EDITORONLY_DATA
-		checkf(IsInGameThread() || IsInAudioThread(), TEXT("FAudioCookInputs creation must happen on the game-thread or audio-thread as it reads from many non-thread safe properties of USoundWave"));
+		checkf(IsInGameThread() || IsInAudioThread() || IsInAsyncLoadingThread(), TEXT("FAudioCookInputs creation must happen on the game-thread or audio-thread as it reads from many non-thread safe properties of USoundWave"));
 
 #if FORCE_RESAMPLE
 		FPlatformAudioCookOverrides NewCompressionOverrides = FPlatformAudioCookOverrides();
diff --git a/Engine/Source/Runtime/Engine/Private/SoundWaveCompiler.cpp b/Engine/Source/Runtime/Engine/Private/SoundWaveCompiler.cpp
--- a/Engine/Source/Runtime/Engine/Private/SoundWaveCompiler.cpp
+++ b/Engine/Source/Runtime/Engine/Private/SoundWaveCompiler.cpp
@@ -114,19 +114,8 @@ void FSoundWaveCompilingManager::Shutdown()
 	bHasShutdown = true;
 	if (GetNumRemainingSoundWaves())
 	{
-		TArray<USoundWave*> PendingSoundWaves;
-		PendingSoundWaves.Reserve(GetNumRemainingSoundWaves());
-
-		for (TWeakObjectPtr<USoundWave>& WeakSoundWave : RegisteredSoundWaves)
-		{
-			if (USoundWave* SoundWave = WeakSoundWave.Get())
-			{
-				PendingSoundWaves.Add(SoundWave);
-			}
-		}
-
 		// Wait on SoundWaves already in progress we couldn't cancel
-		FinishCompilation(PendingSoundWaves);
+		FinishCompilation(GatherPendingSoundWaves());
 	}
 }
 
@@ -168,6 +157,7 @@ FSoundWaveCompilingManager& FSoundWaveCompilingManager::Get()
 
 int32 FSoundWaveCompilingManager::GetNumRemainingSoundWaves() const
 {
+	FReadScopeLock ScopeLock(Lock);
 	return RegisteredSoundWaves.Num();
 }
 
@@ -179,11 +169,13 @@ int32 FSoundWaveCompilingManager::GetNumRemainingAssets() const
 void FSoundWaveCompilingManager::AddSoundWaves(TArrayView<USoundWave* const> InSoundWaves)
 {
 	TRACE_CPUPROFILER_EVENT_SCOPE(FSoundWaveCompilingManager::AddSoundWaves)
-	check(IsInGameThread());
 
-	for (USoundWave* SoundWave : InSoundWaves)
 	{
-		RegisteredSoundWaves.Emplace(SoundWave);
+		FWriteScopeLock ScopeLock(Lock);
+		for (USoundWave* SoundWave : InSoundWaves)
+		{
+			RegisteredSoundWaves.Emplace(SoundWave);
+		}
 	}
 
 	TRACE_COUNTER_SET(QueuedSoundWaveCompilation, GetNumRemainingSoundWaves());
@@ -199,12 +191,14 @@ void FSoundWaveCompilingManager::FinishCompilation(TArrayView<USoundWave* const>
 	TSet<USoundWave*> PendingSoundWaves;
 	PendingSoundWaves.Reserve(InSoundWaves.Num());
 
-	int32 SoundWaveIndex = 0;
-	for (USoundWave* SoundWave : InSoundWaves)
 	{
-		if (RegisteredSoundWaves.Contains(SoundWave))
+		FReadScopeLock ScopeLock(Lock);
+		for (USoundWave* SoundWave : InSoundWaves)
 		{
-			PendingSoundWaves.Add(SoundWave);
+			if (RegisteredSoundWaves.Contains(SoundWave))
+			{
+				PendingSoundWaves.Add(SoundWave);
+			}
 		}
 	}
 
@@ -247,7 +241,10 @@ void FSoundWaveCompilingManager::FinishCompilation(TArrayView<USoundWave* const>
 				USoundWave* SoundWave = static_cast<FCompilableSoundWave*>(Object)->SoundWave.Get();
 				PostCompilation(SoundWave);
 
-				RegisteredSoundWaves.Remove(SoundWave);
+				{
+					FWriteScopeLock ScopeLock(Lock);
+					RegisteredSoundWaves.Remove(SoundWave);
+				}
 			}
 		);
 
@@ -279,6 +276,28 @@ void FSoundWaveCompilingManager::PostCompilation(TArrayView<USoundWave* const> I
 	}
 }
 
+TArray<USoundWave*> FSoundWaveCompilingManager::GatherPendingSoundWaves()
+{
+	TArray<USoundWave*> PendingSoundWaves;
+	PendingSoundWaves.Reserve(GetNumRemainingSoundWaves());
+
+	// Resolve weak pointers and get rid of stale objects at the same time
+	FWriteScopeLock ScopeLock(Lock);
+	for (auto Iterator = RegisteredSoundWaves.CreateIterator(); Iterator; ++Iterator)
+	{
+		if (USoundWave* SoundWave = Iterator->Get())
+		{
+			PendingSoundWaves.Add(SoundWave);
+		}
+		else
+		{
+			Iterator.RemoveCurrent();
+		}
+	}
+
+	return MoveTemp(PendingSoundWaves);
+}
+
 void FSoundWaveCompilingManager::FinishAllCompilation()
 {
 	check(IsInGameThread());
@@ -286,18 +305,7 @@ void FSoundWaveCompilingManager::FinishAllCompilation()
 
 	if (GetNumRemainingSoundWaves())
 	{
-		TArray<USoundWave*> PendingSoundWaves;
-		PendingSoundWaves.Reserve(GetNumRemainingSoundWaves());
-
-		for (TWeakObjectPtr<USoundWave>& SoundWave : RegisteredSoundWaves)
-		{
-			if (SoundWave.IsValid())
-			{
-				PendingSoundWaves.Add(SoundWave.Get());
-			}
-		}
-
-		FinishCompilation(PendingSoundWaves);
+		FinishCompilation(GatherPendingSoundWaves());
 	}
 }
 
@@ -313,29 +321,30 @@ void FSoundWaveCompilingManager::ProcessSoundWaves(bool bLimitExecutionTime, int
 		{
 			TRACE_CPUPROFILER_EVENT_SCOPE(ProcessFinishedSoundWaves);
 
+			TArray<USoundWave*> SoundWavesToProcess = GatherPendingSoundWaves();
+
 			double TickStartTime = FPlatformTime::Seconds();
-			TSet<TWeakObjectPtr<USoundWave>> SoundWavesToPostpone;
-			for (TWeakObjectPtr<USoundWave>& SoundWave : RegisteredSoundWaves)
+			for (USoundWave* SoundWave : SoundWavesToProcess)
 			{
-				if (SoundWave.IsValid())
+				const bool bHasTimeLeft = bLimitExecutionTime ? ((FPlatformTime::Seconds() - TickStartTime) < MaxSecondsPerFrame) : true;
+				if (bHasTimeLeft && SoundWave->IsAsyncWorkComplete())
 				{
-					const bool bHasTimeLeft = bLimitExecutionTime ? ((FPlatformTime::Seconds() - TickStartTime) < MaxSecondsPerFrame) : true;
-					if (bHasTimeLeft && SoundWave->IsAsyncWorkComplete())
-					{
-						PostCompilation(SoundWave.Get());
-						ProcessedSoundWaves.Add(SoundWave.Get());
-					}
-					else
-					{
-						SoundWavesToPostpone.Emplace(MoveTemp(SoundWave));
-					}
+					PostCompilation(SoundWave);
+					ProcessedSoundWaves.Add(SoundWave);
 				}
 			}
-
-			RegisteredSoundWaves = MoveTemp(SoundWavesToPostpone);
 		}
 
-		PostCompilation(ProcessedSoundWaves);
+		if (ProcessedSoundWaves.Num())
+		{
+			PostCompilation(ProcessedSoundWaves);
+
+			FWriteScopeLock WriteLock(Lock);
+			for (USoundWave* ProcessedSoundWave : ProcessedSoundWaves)
+			{
+				RegisteredSoundWaves.Remove(ProcessedSoundWave);
+			}
+		}
 	}
 }
 
diff --git a/Engine/Source/Runtime/Engine/Public/SoundWaveCompiler.h b/Engine/Source/Runtime/Engine/Public/SoundWaveCompiler.h
--- a/Engine/Source/Runtime/Engine/Public/SoundWaveCompiler.h
+++ b/Engine/Source/Runtime/Engine/Public/SoundWaveCompiler.h
@@ -81,6 +81,7 @@ private:
 	bool bHasShutdown = false;
 	TSet<TWeakObjectPtr<USoundWave>> RegisteredSoundWaves;
 
+	mutable FRWLock Lock;
 	FName GetAssetTypeName() const override;
 	FTextFormat GetAssetNameFormat() const override;
 	TArrayView<FName> GetDependentTypeNames() const override;
@@ -90,6 +91,7 @@ private:
 	void PostCompilation(TArrayView<USoundWave* const> InCompiledSoundWaves);
 	void PostCompilation(USoundWave* SoundWave);
 	void ProcessSoundWaves(bool bLimitExecutionTime, int32 MinBatchSize = 1);
+	TArray<USoundWave*> GatherPendingSoundWaves();
 
 	/** Notification for the amount of pending sound wave compilations */
 	FAsyncCompilationNotification Notification;
