Skeletal meshes with a leader pose component set will not tick  assigned Anim BPs by default (added a boolean option to allow ticking or not)

diff --git a/Engine/Source/Runtime/Engine/Classes/Components/SkinnedMeshComponent.h b/Engine/Source/Runtime/Engine/Classes/Components/SkinnedMeshComponent.h
--- a/Engine/Source/Runtime/Engine/Classes/Components/SkinnedMeshComponent.h
+++ b/Engine/Source/Runtime/Engine/Classes/Components/SkinnedMeshComponent.h
@@ -755,6 +755,10 @@ protected:
 	/** Whether mip callbacks have been registered and need to be removed on destroy */
 	uint8 bMipLevelCallbackRegistered:1;
 
+	/** If false, Follower components ShouldTickPose function will return false (default) */
+	UPROPERTY(Transient)
+	uint8 bFollowerShouldTickPose : 1;
+
 #if UE_ENABLE_DEBUG_DRAWING
 private:
 	/** Whether to draw this mesh's debug skeleton (regardless of showflags) */
@@ -1492,9 +1496,11 @@ public:
 	 * Set LeaderPoseComponent for this component
 	 *
 	 * @param NewLeaderBoneComponent New LeaderPoseComponent
+	 * @param bForceUpdate If false, the function will be skipped if NewLeaderBoneComponent is the same as currently setup (default)
+	 * @param bInFollowerShouldTickPose If false, Follower components will not execute TickPose (default)
 	 */
 	UFUNCTION(BlueprintCallable, Category = "Components|SkinnedMesh")
-	void SetLeaderPoseComponent(USkinnedMeshComponent* NewLeaderBoneComponent, bool bForceUpdate = false);
+	void SetLeaderPoseComponent(USkinnedMeshComponent* NewLeaderBoneComponent, bool bForceUpdate = false, bool bInFollowerShouldTickPose = false);
 
 	UE_DEPRECATED(5.1, "This method has been deprecated. Please use SetLeaderPoseComponent instead.")
 	void SetMasterPoseComponent(USkinnedMeshComponent* NewMasterBoneComponent, bool bForceUpdate = false) { SetLeaderPoseComponent(NewMasterBoneComponent, bForceUpdate); }
diff --git a/Engine/Source/Runtime/Engine/Private/Components/SkeletalMeshComponent.cpp b/Engine/Source/Runtime/Engine/Private/Components/SkeletalMeshComponent.cpp
--- a/Engine/Source/Runtime/Engine/Private/Components/SkeletalMeshComponent.cpp
+++ b/Engine/Source/Runtime/Engine/Private/Components/SkeletalMeshComponent.cpp
@@ -1355,6 +1355,11 @@ bool USkeletalMeshComponent::ShouldUpdateTransform(bool bLODHasChanged) const
 
 bool USkeletalMeshComponent::ShouldTickPose() const
 {
+	if (LeaderPoseComponent.IsValid() && !bFollowerShouldTickPose)
+	{
+		return false;
+	}
+
 	// When we stop root motion we go back to ticking after CharacterMovement. Unfortunately that means that we could tick twice that frame.
 	// So only enforce a single tick per frame.
 	const bool bAlreadyTickedThisFrame = PoseTickedThisFrame();
diff --git a/Engine/Source/Runtime/Engine/Private/Components/SkinnedMeshComponent.cpp b/Engine/Source/Runtime/Engine/Private/Components/SkinnedMeshComponent.cpp
--- a/Engine/Source/Runtime/Engine/Private/Components/SkinnedMeshComponent.cpp
+++ b/Engine/Source/Runtime/Engine/Private/Components/SkinnedMeshComponent.cpp
@@ -426,6 +426,8 @@ USkinnedMeshComponent::USkinnedMeshComponent(const FObjectInitializer& ObjectIni
 
 	bMipLevelCallbackRegistered = false;
 
+	bFollowerShouldTickPose = false;
+
 #if !(UE_BUILD_SHIPPING || UE_BUILD_TEST)
 	bDrawDebugSkeleton = false;
 #endif
@@ -1141,6 +1143,11 @@ void USkinnedMeshComponent::InitLODInfos()
 
 bool USkinnedMeshComponent::ShouldTickPose() const
 {
+	if (LeaderPoseComponent.IsValid() && !bFollowerShouldTickPose)
+	{
+		return false;
+	}
+
 	return ((VisibilityBasedAnimTickOption < EVisibilityBasedAnimTickOption::OnlyTickPoseWhenRendered) || bRecentlyRendered);
 }
 
@@ -2142,7 +2149,7 @@ void USkinnedMeshComponent::SetPhysicsAsset(class UPhysicsAsset* InPhysicsAsset,
 	PhysicsAssetOverride = InPhysicsAsset;
 }
 
-void USkinnedMeshComponent::SetLeaderPoseComponent(class USkinnedMeshComponent* NewLeaderBoneComponent, bool bForceUpdate)
+void USkinnedMeshComponent::SetLeaderPoseComponent(class USkinnedMeshComponent* NewLeaderBoneComponent, bool bForceUpdate, bool bInFollowerShouldTickPose)
 {
 	// Early out if we're already setup.
 	if (!bForceUpdate && NewLeaderBoneComponent == LeaderPoseComponent)
@@ -2150,6 +2157,8 @@ void USkinnedMeshComponent::SetLeaderPoseComponent(class USkinnedMeshComponent*
 		return;
 	}
 
+	bFollowerShouldTickPose = bInFollowerShouldTickPose;
+
 	USkinnedMeshComponent* OldLeaderPoseComponent = LeaderPoseComponent.Get();
 	USkinnedMeshComponent* ValidNewLeaderPose = NewLeaderBoneComponent;
 
