Disable UE_USE_UE4_WORLD_MAX, expanding default WORLD_MAX to ~88m km.

diff --git a/Engine/Source/Runtime/Engine/Public/EngineDefines.h b/Engine/Source/Runtime/Engine/Public/EngineDefines.h
--- a/Engine/Source/Runtime/Engine/Public/EngineDefines.h
+++ b/Engine/Source/Runtime/Engine/Public/EngineDefines.h
@@ -37,7 +37,7 @@
 #define UE_LARGE_HALF_WORLD_MAX1	(UE_LARGE_HALF_WORLD_MAX - 1)	/* LWC half maximum world size minus one */
 
 #ifndef UE_USE_UE4_WORLD_MAX
-#define UE_USE_UE4_WORLD_MAX		1						// Force UE4 WORLD_MAX for converted UE4 titles that explicitly rely on it.
+#define UE_USE_UE4_WORLD_MAX		0						// Force UE4 WORLD_MAX for converted UE4 titles that explicitly rely on it.
 #endif
 
 // Note: Modifying WORLD_MAX affects UE_LWC_RENDER_TILE_SIZE in Engine\Source\Runtime\Core\Private\Misc\LargeWorldRenderPosition.cpp and may introduce precision issues in shaders using world coordinates.
