Bind a dummy vertex buffer to PreviousMorphBuffer when there isn't a valid one.
Fixes crash opening a first person template project in Vulkan

diff --git a/Engine/Source/Runtime/Engine/Private/GPUSkinVertexFactory.cpp b/Engine/Source/Runtime/Engine/Private/GPUSkinVertexFactory.cpp
--- a/Engine/Source/Runtime/Engine/Private/GPUSkinVertexFactory.cpp
+++ b/Engine/Source/Runtime/Engine/Private/GPUSkinVertexFactory.cpp
@@ -837,10 +837,7 @@ public:
 			const FMorphVertexBuffer* MorphVertexBuffer = nullptr;
 			const auto* GPUSkinVertexFactory = (const FGPUBaseSkinVertexFactory*)VertexFactory;
 			MorphVertexBuffer = GPUSkinVertexFactory->GetMorphVertexBuffer(!View->Family->bWorldIsPaused, View->Family->FrameNumber);
-			if (MorphVertexBuffer)
-			{
-				ShaderBindings.Add(PreviousMorphBufferParameter, MorphVertexBuffer->GetSRV());
-			}
+			ShaderBindings.Add(PreviousMorphBufferParameter, MorphVertexBuffer ? MorphVertexBuffer->GetSRV() : GNullVertexBuffer.VertexBufferSRV.GetReference());
 		}
 	}
 
