Added GetCurveValueWithDefault to blueprint. If curve value isn't found, output the default instead.

diff --git a/Engine/Source/Runtime/Engine/Classes/Animation/AnimInstance.h b/Engine/Source/Runtime/Engine/Classes/Animation/AnimInstance.h
--- a/Engine/Source/Runtime/Engine/Classes/Animation/AnimInstance.h
+++ b/Engine/Source/Runtime/Engine/Classes/Animation/AnimInstance.h
@@ -1131,6 +1131,16 @@ public:
 	UFUNCTION(BlueprintPure, Category="Animation|Curves", meta=(BlueprintThreadSafe))
 	float GetCurveValue(FName CurveName) const;
 
+	/** 
+	* Returns whether a named curve was found, its value, and a default value when it's not found.
+	* @param	AnimInstance	The anim instance to find this curve value for.
+	* @param	CurveName		The name of the curve.
+	* @param	DefaultValue	Value to use when the curve is not found.
+	* @param	OutValue		The curve's value.
+	*/
+	UFUNCTION(BlueprintPure, Category="Animation|Curves", meta=(BlueprintThreadSafe))
+	bool GetCurveValueWithDefault(FName CurveName, float DefaultValue, float& OutValue);
+
 	/** This returns last up-to-date list of active curve names */
 	UFUNCTION(BlueprintPure, Category = "Animation|Curves", meta=(BlueprintThreadSafe))
 	void GetActiveCurveNames(EAnimCurveType CurveType, TArray<FName>& OutNames) const;
diff --git a/Engine/Source/Runtime/Engine/Private/Animation/AnimInstance.cpp b/Engine/Source/Runtime/Engine/Private/Animation/AnimInstance.cpp
--- a/Engine/Source/Runtime/Engine/Private/Animation/AnimInstance.cpp
+++ b/Engine/Source/Runtime/Engine/Private/Animation/AnimInstance.cpp
@@ -1675,6 +1675,12 @@ bool UAnimInstance::GetCurveValue(FName CurveName, float& OutValue) const
 	return false;
 }
 
+bool UAnimInstance::GetCurveValueWithDefault(FName CurveName, float DefaultValue, float& OutValue)
+{
+	OutValue = DefaultValue;
+	return GetCurveValue(CurveName, OutValue);
+}
+
 void UAnimInstance::GetActiveCurveNames(EAnimCurveType CurveType, TArray<FName>& OutNames) const
 {
 	TMap<FName, float> ActiveCurves;
