fix for FPoseContext::Initialize

diff --git a/Engine/Source/Runtime/Engine/Private/Animation/AnimNodeBase.cpp b/Engine/Source/Runtime/Engine/Private/Animation/AnimNodeBase.cpp
--- a/Engine/Source/Runtime/Engine/Private/Animation/AnimNodeBase.cpp
+++ b/Engine/Source/Runtime/Engine/Private/Animation/AnimNodeBase.cpp
@@ -63,6 +63,8 @@ void FAnimationBaseContext::LogMessageInternal(FName InLogType, EMessageSeverity
 
 void FPoseContext::Initialize(FAnimInstanceProxy* InAnimInstanceProxy)
 {
+	AnimInstanceProxy = InAnimInstanceProxy;
+
 	checkSlow(AnimInstanceProxy && AnimInstanceProxy->GetRequiredBones().IsValid());
 	const FBoneContainer& RequiredBone = AnimInstanceProxy->GetRequiredBones();
 	Pose.SetBoneContainer(&RequiredBone);
