Fix Strata decal conversion when decal is using material attributes.

diff --git a/Engine/Source/Runtime/Engine/Private/Materials/Material.cpp b/Engine/Source/Runtime/Engine/Private/Materials/Material.cpp
--- a/Engine/Source/Runtime/Engine/Private/Materials/Material.cpp
+++ b/Engine/Source/Runtime/Engine/Private/Materials/Material.cpp
@@ -3160,11 +3160,24 @@ void UMaterial::ConvertMaterialToStrataMaterial()
 			check(ConvertNode->ConvertedStrataMaterialInfo.CountShadingModels() == 1);
 		}
 
+		if (MaterialDomain == MD_DeferredDecal)
+		{
+			// For now we don't enforce shading model since it could be driven by expression and we don't have much 
+			// control on this, but only DefaultLit should be supported.
+
+			// Now pass through the convert to decal node, which flag the material as SSM_Decal, which will set the domain to Decal.
+			UMaterialExpressionStrataConvertToDecal* ConvertToDecalNode = NewObject<UMaterialExpressionStrataConvertToDecal>(this);
+			ConvertToDecalNode->DecalMaterial.Connect(0, ConvertNode);
+
+			EditorOnly->FrontMaterial.Connect(0, ConvertToDecalNode);
+		}
+
 		ConvertLegacyToStrataBlendMode();
 		bInvalidateShader = true;
 	}
 	else if (!bUseMaterialAttributes && !EditorOnly->FrontMaterial.IsConnected() && GetExpressions().IsEmpty())
 	{
+		// Empty material: Create by default a slab node
 		const int32 NewNodeOffsetX = -300;
 		UMaterialExpressionStrataSlabBSDF* SlabNode = NewObject<UMaterialExpressionStrataSlabBSDF>(this);
 		SlabNode->MaterialExpressionEditorX = EditorX + NewNodeOffsetX;
