Fix error messages in cooked editor related to thumbnail textures

diff --git a/Engine/Source/Runtime/Engine/Private/Texture2D.cpp b/Engine/Source/Runtime/Engine/Private/Texture2D.cpp
--- a/Engine/Source/Runtime/Engine/Private/Texture2D.cpp
+++ b/Engine/Source/Runtime/Engine/Private/Texture2D.cpp
@@ -1615,7 +1615,7 @@ bool UTexture2D::StreamIn(int32 NewMipCount, bool bHighPrio)
 		if (!CustomMipDataProvider && GUseGenericStreamingPath != 1)
 		{
 	#if WITH_EDITORONLY_DATA
-			if (FPlatformProperties::HasEditorOnlyData() && !GetOutermost()->bIsCookedForEditor)
+			if (FPlatformProperties::HasEditorOnlyData() && !GetOutermost()->bIsCookedForEditor && !GetOutermost()->HasAnyPackageFlags(PKG_Cooked))
 			{
 				if (GRHISupportsAsyncTextureCreation)
 				{
@@ -1652,7 +1652,7 @@ bool UTexture2D::StreamIn(int32 NewMipCount, bool bHighPrio)
 			FTextureMipDataProvider* DefaultMipDataProvider = nullptr;
 
 	#if WITH_EDITORONLY_DATA
-			if (FPlatformProperties::HasEditorOnlyData() && !GetOutermost()->bIsCookedForEditor)
+			if (FPlatformProperties::HasEditorOnlyData() && !GetOutermost()->bIsCookedForEditor && !GetOutermost()->HasAnyPackageFlags(PKG_Cooked))
 			{
 				DefaultMipDataProvider = new FTexture2DMipDataProvider_DDC(this);
 			}
