World Partition: replace popup by on-screen message when streaming is disabled but the world is large enough to justify enabling it.

fyi jeansebastien.guay
#rnx

diff --git a/Engine/Source/Runtime/Engine/Private/WorldPartition/WorldPartition.cpp b/Engine/Source/Runtime/Engine/Private/WorldPartition/WorldPartition.cpp
--- a/Engine/Source/Runtime/Engine/Private/WorldPartition/WorldPartition.cpp
+++ b/Engine/Source/Runtime/Engine/Private/WorldPartition/WorldPartition.cpp
@@ -67,6 +67,14 @@ static FAutoConsoleVariableRef CVarLoadingRangeBugItGo(
 	ECVF_Default
 );
 
+static int32 GWorldExtentToEnableStreaming = 400000;
+static FAutoConsoleVariableRef CVarWorldExtentToEnableStreaming(
+	TEXT("wp.Editor.WorldExtentToEnableStreaming"),
+	GWorldExtentToEnableStreaming,
+	TEXT("World extend to justify enabling streaming."),
+	ECVF_Default
+);
+
 static bool GDebugDedicatedServerStreaming = false;
 static FAutoConsoleVariableRef CVarDebugDedicatedServerStreaming(
 	TEXT("wp.Runtime.DebugDedicatedServerStreaming"),
@@ -1240,7 +1248,7 @@ void UWorldPartition::Tick(float DeltaSeconds)
 					AllActorsBounds += ActorDescIterator->GetBounds();
 
 					// Warn the user if the world becomes larger that 4km in any axis
-					if (AllActorsBounds.GetSize().GetMax() >= 400000.0f)
+					if (AllActorsBounds.GetSize().GetMax() >= GWorldExtentToEnableStreaming)
 					{
 						bEnablingStreamingJustified = true;
 						break;
