World Partition - HLOD: Change default cell size & loading range to make more sense given the default main grid size
* Cell size: 3200 -> 25600
* Loading range: 12800 -> 51200

diff --git a/Engine/Source/Runtime/Engine/Private/WorldPartition/HLOD/HLODLayer.cpp b/Engine/Source/Runtime/Engine/Private/WorldPartition/HLOD/HLODLayer.cpp
--- a/Engine/Source/Runtime/Engine/Private/WorldPartition/HLOD/HLODLayer.cpp
+++ b/Engine/Source/Runtime/Engine/Private/WorldPartition/HLOD/HLODLayer.cpp
@@ -29,8 +29,8 @@ UHLODLayer::UHLODLayer(const FObjectInitializer& ObjectInitializer)
 	: Super(ObjectInitializer)
 #if WITH_EDITORONLY_DATA
 	, bIsSpatiallyLoaded(true)
-	, CellSize(3200)
-	, LoadingRange(12800)
+	, CellSize(25600)
+	, LoadingRange(51200)
 #endif
 {
 }
