Analytic derivatives support for VirtualTextureWorldToUV and TransformBase
Moved FLWCScalarDerivX to LargeWorldCoordinates.ush

diff --git a/Engine/Shaders/Private/LargeWorldCoordinates.ush b/Engine/Shaders/Private/LargeWorldCoordinates.ush
--- a/Engine/Shaders/Private/LargeWorldCoordinates.ush
+++ b/Engine/Shaders/Private/LargeWorldCoordinates.ush
@@ -26,6 +26,34 @@ struct FLWCVector4
 	float4 Offset;
 };
 
+struct FLWCScalarDeriv
+{
+	FLWCScalar Value;
+	float Ddx;
+	float Ddy;
+};
+
+struct FLWCVector2Deriv
+{
+	FLWCVector2 Value;
+	float2 Ddx;
+	float2 Ddy;
+};
+
+struct FLWCVector3Deriv
+{
+	FLWCVector3 Value;
+	float3 Ddx;
+	float3 Ddy;
+};
+
+struct FLWCVector4Deriv
+{
+	FLWCVector4 Value;
+	float4 Ddx;
+	float4 Ddy;
+};
+
 // Transforms *to* absolute world space
 struct FLWCMatrix
 {
diff --git a/Engine/Shaders/Private/VirtualTextureCommon.ush b/Engine/Shaders/Private/VirtualTextureCommon.ush
--- a/Engine/Shaders/Private/VirtualTextureCommon.ush
+++ b/Engine/Shaders/Private/VirtualTextureCommon.ush
@@ -734,6 +734,16 @@ float2 VirtualTextureWorldToUV(in FLWCVector3 WorldPos, in FLWCVector3 O, in flo
 	return float2(dot(P, U), dot(P, V));
 }
 
+/**	Derivative-aware helper function to convert world space to virtual texture UV space. */
+FloatDeriv2 VirtualTextureWorldToUVDeriv(in FLWCVector3Deriv WorldPos, in FLWCVector3 O, in float3 U, in float3 V)
+{
+	FloatDeriv2 Result;
+	Result.Value = VirtualTextureWorldToUV(WorldPos.Value, O, U, V);
+	Result.Ddx = float2(dot(WorldPos.Ddx, U), dot(WorldPos.Ddx, V));
+	Result.Ddy = float2(dot(WorldPos.Ddy, U), dot(WorldPos.Ddy, V));
+	return Result;
+}
+
 /** Unpack color from YCoCg stored in BC3. */
 float3 VirtualTextureUnpackBaseColorYCoCg(in float4 PackedValue)
 {
diff --git a/Engine/Source/Runtime/Engine/Private/Materials/HLSLMaterialDerivativeAutogen.cpp b/Engine/Source/Runtime/Engine/Private/Materials/HLSLMaterialDerivativeAutogen.cpp
--- a/Engine/Source/Runtime/Engine/Private/Materials/HLSLMaterialDerivativeAutogen.cpp
+++ b/Engine/Source/Runtime/Engine/Private/Materials/HLSLMaterialDerivativeAutogen.cpp
@@ -1304,29 +1304,6 @@ FString FMaterialDerivativeAutogen::GenerateUsedFunctions(FHLSLMaterialTranslato
 
 	FString Ret;
 
-	// The basic structs (FloatDeriv, FloatDeriv2, FloatDeriv3, FloatDeriv4)
-	// It's not worth keeping track of all the times these are used, just make them.
-	for (int32 Index = 0; Index < NumDerivativeTypes; Index++)
-	{
-		EDerivativeType DerivType = (EDerivativeType)Index;
-		if (!IsLWCType(DerivType))
-		{
-			continue;				// Non-LWC types defined in common.ush
-		}
-
-		FString BaseName = GetDerivVectorName(DerivType);
-		FString FieldName = GetFloatVectorName(DerivType);
-		FString FieldNameDDXY = GetFloatVectorDDXYName(DerivType);
-
-		Ret += TEXT("struct ") + BaseName + LINE_TERMINATOR;
-		Ret += TEXT("{") LINE_TERMINATOR;
-		Ret += TEXT("\t") + FieldName + TEXT(" Value;") LINE_TERMINATOR;
-		Ret += TEXT("\t") + FieldNameDDXY + TEXT(" Ddx;") LINE_TERMINATOR;
-		Ret += TEXT("\t") + FieldNameDDXY + TEXT(" Ddy;") LINE_TERMINATOR;
-		Ret += TEXT("};") LINE_TERMINATOR;
-		Ret += TEXT("") LINE_TERMINATOR;
-	}
-
 	// Full FloatDerivX constructors with explicit derivatives.
 	for (int32 Index = 0; Index < NumDerivativeTypes; Index++)
 	{
diff --git a/Engine/Source/Runtime/Engine/Private/Materials/HLSLMaterialTranslator.cpp b/Engine/Source/Runtime/Engine/Private/Materials/HLSLMaterialTranslator.cpp
--- a/Engine/Source/Runtime/Engine/Private/Materials/HLSLMaterialTranslator.cpp
+++ b/Engine/Source/Runtime/Engine/Private/Materials/HLSLMaterialTranslator.cpp
@@ -6844,12 +6844,32 @@ int32 FHLSLMaterialTranslator::VirtualTextureUniform(FName ParameterName, int32
 
 int32 FHLSLMaterialTranslator::VirtualTextureWorldToUV(int32 WorldPositionIndex, int32 P0, int32 P1, int32 P2)
 {
-	if (WorldPositionIndex == INDEX_NONE || P0 == INDEX_NONE || P1 == INDEX_NONE || P2 == INDEX_NONE)
+	if (!UseVirtualTexturing(FeatureLevel, TargetPlatform) || WorldPositionIndex == INDEX_NONE || P0 == INDEX_NONE || P1 == INDEX_NONE || P2 == INDEX_NONE)
 	{
 		return INDEX_NONE;
 	}
-	FString	SampleCode(TEXT("VirtualTextureWorldToUV(%s, %s, %s, %s)"));
-	return AddInlinedCodeChunk(MCT_Float2, *SampleCode, *CoerceParameter(WorldPositionIndex, MCT_LWCVector3), *GetParameterCode(P0), *GetParameterCode(P1), *GetParameterCode(P2));
+
+	const EDerivativeStatus WorldPositionDerivStatus = GetDerivativeStatus(WorldPositionIndex);
+
+	FString CodeFinite = FString::Printf(TEXT("VirtualTextureWorldToUV(%s, %s, %s, %s)"), *CoerceParameter(WorldPositionIndex, MCT_LWCVector3), *GetParameterCode(P0), *GetParameterCode(P1), *GetParameterCode(P2));
+	if (IsAnalyticDerivEnabled() && IsDerivativeValid(WorldPositionDerivStatus))
+	{
+		if (WorldPositionDerivStatus == EDerivativeStatus::Valid)
+		{
+			WorldPositionIndex = ValidCast(WorldPositionIndex, MCT_LWCVector3);
+			FString WorldPositionDeriv = GetParameterCodeDeriv(WorldPositionIndex, CompiledPDV_Analytic);
+			FString CodeAnalytic = FString::Printf(TEXT("VirtualTextureWorldToUVDeriv(%s, %s, %s, %s)"), *WorldPositionDeriv, *GetParameterCode(P0), *GetParameterCode(P1), *GetParameterCode(P2));
+			return AddCodeChunkInnerDeriv(*CodeFinite, *CodeAnalytic, MCT_Float2, false, EDerivativeStatus::Valid);
+		}
+		else
+		{
+			return AddInlinedCodeChunkZeroDeriv(MCT_Float2, *CodeFinite);
+		}
+	}
+	else
+	{
+		return AddInlinedCodeChunk(MCT_Float2, *CodeFinite);
+	}
 }
 
 int32 FHLSLMaterialTranslator::VirtualTextureUnpack(int32 CodeIndex0, int32 CodeIndex1, int32 CodeIndex2, int32 P0, EVirtualTextureUnpackType UnpackType)
@@ -8373,6 +8393,7 @@ int32 FHLSLMaterialTranslator::TransformBase(EMaterialCommonBasis SourceCoordBas
 	}
 		
 	FString CodeStr;
+	FString CodeDerivStr;
 	EMaterialCommonBasis IntermediaryBasis = MCB_World;
 
 	switch (SourceCoordBasis)
@@ -8383,6 +8404,7 @@ int32 FHLSLMaterialTranslator::TransformBase(EMaterialCommonBasis SourceCoordBas
 			if (DestCoordBasis == MCB_World)
 			{
 				CodeStr = TEXT("mul(<A>, Parameters.TangentToWorld)");
+				CodeDerivStr = TEXT("mul(<A>, Parameters.TangentToWorld)");
 			}
 			// else use MCB_World as intermediary basis
 			break;
@@ -8392,6 +8414,7 @@ int32 FHLSLMaterialTranslator::TransformBase(EMaterialCommonBasis SourceCoordBas
 			if (DestCoordBasis == MCB_World)
 			{
 				CodeStr = TEXT("TransformLocal<TO><PREV>World(Parameters, <A>)");
+				CodeDerivStr = TEXT("TransformLocal<TO><PREV>World(Parameters, <A>)");
 			}
 			// else use MCB_World as intermediary basis
 			break;
@@ -8408,14 +8431,17 @@ int32 FHLSLMaterialTranslator::TransformBase(EMaterialCommonBasis SourceCoordBas
 				{
 					CodeStr = TEXT("<A>");
 				}
+				CodeDerivStr = TEXT("<A>");
 			}
 			else if (DestCoordBasis == MCB_Camera)
 			{
 				CodeStr = MultiplyMatrix(TEXT("<A>"), TEXT("ResolvedView.<PREV>TranslatedWorldToCameraView"), AWComponent);
+				CodeDerivStr = MultiplyMatrix(TEXT("<A>"), TEXT("ResolvedView.<PREV>TranslatedWorldToCameraView"), 0);
 			}
 			else if (DestCoordBasis == MCB_View)
 			{
 				CodeStr = MultiplyMatrix(TEXT("<A>"), TEXT("ResolvedView.<PREV>TranslatedWorldToView"), AWComponent);
+				CodeDerivStr = MultiplyMatrix(TEXT("<A>"), TEXT("ResolvedView.<PREV>TranslatedWorldToView"), 0);
 			}
 			// else use MCB_World as intermediary basis
 			break;
@@ -8425,6 +8451,7 @@ int32 FHLSLMaterialTranslator::TransformBase(EMaterialCommonBasis SourceCoordBas
 			if (DestCoordBasis == MCB_Tangent)
 			{
 				CodeStr = MultiplyTransposeMatrix(TEXT("Parameters.TangentToWorld"), TEXT("<A>"), AWComponent);
+				CodeDerivStr = MultiplyTransposeMatrix(TEXT("Parameters.TangentToWorld"), TEXT("<A>"), 0);
 			}
 			else if (DestCoordBasis == MCB_Local)
 			{
@@ -8439,6 +8466,7 @@ int32 FHLSLMaterialTranslator::TransformBase(EMaterialCommonBasis SourceCoordBas
 
 				// TODO: inconsistent with TransformLocal<TO>World with instancing
 				CodeStr = LWCMultiplyMatrix(TEXT("<A>"), TEXT("GetPrimitiveData(Parameters).<PREVIOUS>WorldToLocal"), AWComponent);
+				CodeDerivStr = LWCMultiplyMatrix(TEXT("<A>"), TEXT("GetPrimitiveData(Parameters).<PREVIOUS>WorldToLocal"), 0);
 			}
 			else if (DestCoordBasis == MCB_TranslatedWorld)
 			{
@@ -8450,15 +8478,18 @@ int32 FHLSLMaterialTranslator::TransformBase(EMaterialCommonBasis SourceCoordBas
 				{
 					CodeStr = TEXT("<A>");
 				}
+				CodeDerivStr = TEXT("<A>");
 			}
 			else if (DestCoordBasis == MCB_MeshParticle)
 			{
 				CodeStr = LWCMultiplyMatrix(TEXT("<A>"), TEXT("Parameters.Particle.WorldToParticle"), AWComponent);
+				CodeDerivStr = LWCMultiplyMatrix(TEXT("<A>"), TEXT("Parameters.Particle.WorldToParticle"), 0);
 				bUsesParticleWorldToLocal = true;
 			}
 			else if (DestCoordBasis == MCB_Instance)
 			{
 				CodeStr = LWCMultiplyMatrix(TEXT("<A>"), TEXT("GetWorldToInstance(Parameters)"), AWComponent);
+				CodeDerivStr = LWCMultiplyMatrix(TEXT("<A>"), TEXT("GetWorldToInstance(Parameters)"), 0);
 				bUsesInstanceWorldToLocalPS = ShaderFrequency == SF_Pixel;
 			}
 
@@ -8471,6 +8502,7 @@ int32 FHLSLMaterialTranslator::TransformBase(EMaterialCommonBasis SourceCoordBas
 			if (DestCoordBasis == MCB_TranslatedWorld)
 			{
 				CodeStr = MultiplyMatrix(TEXT("<A>"), TEXT("ResolvedView.<PREV>CameraViewToTranslatedWorld"), AWComponent);
+				CodeDerivStr = MultiplyMatrix(TEXT("<A>"), TEXT("ResolvedView.<PREV>CameraViewToTranslatedWorld"), 0);
 			}
 			// else use MCB_TranslatedWorld as intermediary basis
 			IntermediaryBasis = MCB_TranslatedWorld;
@@ -8481,6 +8513,7 @@ int32 FHLSLMaterialTranslator::TransformBase(EMaterialCommonBasis SourceCoordBas
 			if (DestCoordBasis == MCB_TranslatedWorld)
 			{
 				CodeStr = MultiplyMatrix(TEXT("<A>"), TEXT("ResolvedView.<PREV>ViewToTranslatedWorld"), AWComponent);
+				CodeDerivStr = MultiplyMatrix(TEXT("<A>"), TEXT("ResolvedView.<PREV>ViewToTranslatedWorld"), 0);
 			}
 			// else use MCB_TranslatedWorld as intermediary basis
 			IntermediaryBasis = MCB_TranslatedWorld;
@@ -8491,6 +8524,7 @@ int32 FHLSLMaterialTranslator::TransformBase(EMaterialCommonBasis SourceCoordBas
 			if (DestCoordBasis == MCB_World)
 			{
 				CodeStr = LWCMultiplyMatrix(TEXT("<A>"), TEXT("Parameters.Particle.ParticleToWorld"), AWComponent);
+				CodeDerivStr = LWCMultiplyMatrix(TEXT("<A>"), TEXT("Parameters.Particle.ParticleToWorld"), 0);
 				bUsesParticleLocalToWorld = true;
 			}
 			// use World as an intermediary base
@@ -8501,6 +8535,7 @@ int32 FHLSLMaterialTranslator::TransformBase(EMaterialCommonBasis SourceCoordBas
 			if (DestCoordBasis == MCB_World)
 			{
 				CodeStr = LWCMultiplyMatrix(TEXT("<A>"), TEXT("GetInstanceToWorld(Parameters)"), AWComponent);
+				CodeDerivStr = LWCMultiplyMatrix(TEXT("<A>"), TEXT("GetInstanceToWorld(Parameters)"), 0);
 				bUsesInstanceLocalToWorldPS = ShaderFrequency == SF_Pixel;
 			}
 			// use World as an intermediary base
@@ -8532,16 +8567,21 @@ int32 FHLSLMaterialTranslator::TransformBase(EMaterialCommonBasis SourceCoordBas
 	{
 		CodeStr.ReplaceInline(TEXT("<TO>"),TEXT("VectorTo"));
 	}
+	CodeDerivStr.ReplaceInline(TEXT("<TO>"), TEXT("VectorTo"));
 		
 	if (bCompilingPreviousFrame)
 	{
 		CodeStr.ReplaceInline(TEXT("<PREV>"),TEXT("Prev"));
 		CodeStr.ReplaceInline(TEXT("<PREVIOUS>"), TEXT("Previous"));
+		CodeDerivStr.ReplaceInline(TEXT("<PREV>"), TEXT("Prev"));
+		CodeDerivStr.ReplaceInline(TEXT("<PREVIOUS>"), TEXT("Previous"));
 	}
 	else
 	{
 		CodeStr.ReplaceInline(TEXT("<PREV>"),TEXT(""));
 		CodeStr.ReplaceInline(TEXT("<PREVIOUS>"), TEXT(""));
+		CodeDerivStr.ReplaceInline(TEXT("<PREV>"), TEXT(""));
+		CodeDerivStr.ReplaceInline(TEXT("<PREVIOUS>"), TEXT(""));
 	}
 
 	int32 CastA = A;
@@ -8562,10 +8602,30 @@ int32 FHLSLMaterialTranslator::TransformBase(EMaterialCommonBasis SourceCoordBas
 	}
 
 	const EMaterialValueType ResultType = (DestCoordBasis == MCB_World && AWComponent) ? MCT_LWCVector3 : MCT_Float3;
-	const int32 Result = AddCodeChunk(
-		ResultType,
-		*CodeStr
-		);
+	const EDerivativeStatus ADerivStatus = GetDerivativeStatus(CastA);
+
+	int32 Result;
+	if (IsAnalyticDerivEnabled() && IsDerivativeValid(ADerivStatus))
+	{
+		if (ADerivStatus == EDerivativeStatus::Valid)
+		{
+			FString CastADeriv = GetParameterCodeDeriv(CastA, CompiledPDV_Analytic);
+			FString CodeAnalytic = DerivativeAutogen.ConstructDeriv(
+				CodeStr,
+				CodeDerivStr.Replace(TEXT("<A>"), *(CastADeriv + TEXT(".Ddx"))),
+				CodeDerivStr.Replace(TEXT("<A>"), *(CastADeriv + TEXT(".Ddy"))), GetDerivType(ResultType));
+			Result = AddCodeChunkInnerDeriv(*CodeStr, *CodeAnalytic, ResultType, false, EDerivativeStatus::Valid);
+		}
+		else
+		{
+			Result = AddCodeChunkZeroDeriv(ResultType, *CodeStr);
+		}
+	}
+	else
+	{
+		Result = AddCodeChunk(ResultType, *CodeStr);
+	}
+	
 	return CastToNonLWCIfDisabled(Result);
 }
 
