Rendering: add IsOrthographic material function, returns 1 for orthographic views and 0 for perspective views. The test here (View.ViewToClip[3][3] < 1 == perspective) is the same as is used in GetPixelDepth.

diff --git a/Engine/Source/Runtime/Engine/Classes/Materials/MaterialExpressionIsOrthographic.h b/Engine/Source/Runtime/Engine/Classes/Materials/MaterialExpressionIsOrthographic.h
new file mode 100644
--- /dev/null
+++ b/Engine/Source/Runtime/Engine/Classes/Materials/MaterialExpressionIsOrthographic.h
@@ -0,0 +1,23 @@
+// Copyright Epic Games, Inc. All Rights Reserved.
+
+
+#pragma once
+
+#include "CoreMinimal.h"
+#include "UObject/ObjectMacros.h"
+#include "Materials/MaterialExpression.h"
+#include "MaterialExpressionIsOrthographic.generated.h"
+
+UCLASS(collapsecategories, hidecategories=Object)
+class UMaterialExpressionIsOrthographic : public UMaterialExpression
+{
+	GENERATED_UCLASS_BODY()
+
+	//~ Begin UMaterialExpression Interface
+#if WITH_EDITOR
+	virtual bool GenerateHLSLExpression(FMaterialHLSLGenerator& Generator, UE::HLSLTree::FScope& Scope, int32 OutputIndex, UE::HLSLTree::FExpression const*& OutExpression) const override;
+	virtual int32 Compile(class FMaterialCompiler* Compiler, int32 OutputIndex) override;
+	virtual void GetCaption(TArray<FString>& OutCaptions) const override;
+#endif
+	//~ End UMaterialExpression Interface
+};
diff --git a/Engine/Source/Runtime/Engine/Private/Materials/HLSLMaterialTranslator.cpp b/Engine/Source/Runtime/Engine/Private/Materials/HLSLMaterialTranslator.cpp
--- a/Engine/Source/Runtime/Engine/Private/Materials/HLSLMaterialTranslator.cpp
+++ b/Engine/Source/Runtime/Engine/Private/Materials/HLSLMaterialTranslator.cpp
@@ -4370,6 +4370,13 @@ int32 FHLSLMaterialTranslator::ViewProperty(EMaterialExposedViewProperty Propert
 	return CastToNonLWCIfDisabled(Result);
 }
 
+
+int32 FHLSLMaterialTranslator::IsOrthographic()
+{
+	return AddInlinedCodeChunkZeroDeriv(MCT_Float, TEXT("((View.ViewToClip[3][3] < 1.0f) ? 0.0f : 1.0f)"));
+}
+
+
 int32 FHLSLMaterialTranslator::GameTime(bool bPeriodic, float Period)
 {
 	if (!bPeriodic)
diff --git a/Engine/Source/Runtime/Engine/Private/Materials/HLSLMaterialTranslator.h b/Engine/Source/Runtime/Engine/Private/Materials/HLSLMaterialTranslator.h
--- a/Engine/Source/Runtime/Engine/Private/Materials/HLSLMaterialTranslator.h
+++ b/Engine/Source/Runtime/Engine/Private/Materials/HLSLMaterialTranslator.h
@@ -740,6 +740,7 @@ protected:
 	virtual int32 GenericConstant(const UE::Shader::FValue& Value) override;
 	
 	virtual int32 ViewProperty(EMaterialExposedViewProperty Property, bool InvProperty) override;
+	virtual int32 IsOrthographic() override;
 
 	virtual int32 GameTime(bool bPeriodic, float Period) override;
 	virtual int32 RealTime(bool bPeriodic, float Period) override;
diff --git a/Engine/Source/Runtime/Engine/Private/Materials/MaterialExpressionHLSL.cpp b/Engine/Source/Runtime/Engine/Private/Materials/MaterialExpressionHLSL.cpp
--- a/Engine/Source/Runtime/Engine/Private/Materials/MaterialExpressionHLSL.cpp
+++ b/Engine/Source/Runtime/Engine/Private/Materials/MaterialExpressionHLSL.cpp
@@ -25,6 +25,7 @@
 #include "Materials/MaterialExpressionCameraPositionWS.h"
 #include "Materials/MaterialExpressionCameraVectorWS.h"
 #include "Materials/MaterialExpressionViewProperty.h"
+#include "Materials/MaterialExpressionIsOrthographic.h"
 #include "Materials/MaterialExpressionTime.h"
 #include "Materials/MaterialExpressionDeltaTime.h"
 #include "Materials/MaterialExpressionScreenPosition.h"
@@ -512,6 +513,13 @@ bool UMaterialExpressionViewProperty::GenerateHLSLExpression(FMaterialHLSLGenera
 	return true;
 }
 
+bool UMaterialExpressionIsOrthographic::GenerateHLSLExpression(FMaterialHLSLGenerator& Generator, UE::HLSLTree::FScope& Scope, int32 OutputIndex, UE::HLSLTree::FExpression const*& OutExpression) const
+{
+	using namespace UE::HLSLTree::Material;
+	OutExpression = Generator.GetTree().NewExpression<FExpressionExternalInput>(EExternalInput::IsOrthographic);
+	return true;
+}
+
 bool UMaterialExpressionTime::GenerateHLSLExpression(FMaterialHLSLGenerator& Generator, UE::HLSLTree::FScope& Scope, int32 OutputIndex, UE::HLSLTree::FExpression const*& OutExpression) const
 {
 	using namespace UE::HLSLTree::Material;
diff --git a/Engine/Source/Runtime/Engine/Private/Materials/MaterialExpressions.cpp b/Engine/Source/Runtime/Engine/Private/Materials/MaterialExpressions.cpp
--- a/Engine/Source/Runtime/Engine/Private/Materials/MaterialExpressions.cpp
+++ b/Engine/Source/Runtime/Engine/Private/Materials/MaterialExpressions.cpp
@@ -240,6 +240,7 @@
 #include "Materials/MaterialExpressionVertexNormalWS.h"
 #include "Materials/MaterialExpressionVertexTangentWS.h"
 #include "Materials/MaterialExpressionViewProperty.h"
+#include "Materials/MaterialExpressionIsOrthographic.h"
 #include "Materials/MaterialExpressionViewSize.h"
 #include "Materials/MaterialExpressionVolumetricAdvancedMaterialInput.h"
 #include "Materials/MaterialExpressionVolumetricAdvancedMaterialOutput.h"
@@ -10564,6 +10565,39 @@ void UMaterialExpressionViewSize::GetCaption(TArray<FString>& OutCaptions) const
 }
 #endif // WITH_EDITOR
 
+UMaterialExpressionIsOrthographic::UMaterialExpressionIsOrthographic(const FObjectInitializer& ObjectInitializer)
+	: Super(ObjectInitializer)
+{
+#if WITH_EDITORONLY_DATA
+	// Structure to hold one-time initialization
+	struct FConstructorStatics
+	{
+		FText NAME_Constants;
+		FConstructorStatics()
+			: NAME_Constants(LOCTEXT("Constants", "Constants"))
+		{
+		}
+	};
+	static FConstructorStatics ConstructorStatics;
+
+	MenuCategories.Add(ConstructorStatics.NAME_Constants);
+
+	bShaderInputData = true;
+#endif
+}
+
+#if WITH_EDITOR
+int32 UMaterialExpressionIsOrthographic::Compile(class FMaterialCompiler* Compiler, int32 OutputIndex)
+{
+	return Compiler->IsOrthographic();
+}
+
+void UMaterialExpressionIsOrthographic::GetCaption(TArray<FString>& OutCaptions) const
+{
+	OutCaptions.Add(TEXT("IsOrthographic"));
+}
+#endif // WITH_EDITOR
+
 UMaterialExpressionDeltaTime::UMaterialExpressionDeltaTime(const FObjectInitializer& ObjectInitializer)
 	: Super(ObjectInitializer)
 {
diff --git a/Engine/Source/Runtime/Engine/Private/Materials/MaterialHLSLTree.cpp b/Engine/Source/Runtime/Engine/Private/Materials/MaterialHLSLTree.cpp
--- a/Engine/Source/Runtime/Engine/Private/Materials/MaterialHLSLTree.cpp
+++ b/Engine/Source/Runtime/Engine/Private/Materials/MaterialHLSLTree.cpp
@@ -147,6 +147,8 @@ FExternalInputDescription GetExternalInputDescription(EExternalInput Input)
 	case EExternalInput::ParticleTranslatedWorldPosition: return FExternalInputDescription(TEXT("ParticleTranslatedWorldPosition"), Shader::EValueType::Float3);
 	case EExternalInput::ParticleRadius: return FExternalInputDescription(TEXT("ParticleRadius"), Shader::EValueType::Float1);
 
+	case EExternalInput::IsOrthographic: return FExternalInputDescription(TEXT("IsOrthographic"), Shader::EValueType::Float1);
+
 	default: checkNoEntry(); return FExternalInputDescription(TEXT("Invalid"), Shader::EValueType::Void);
 	}
 }
@@ -335,6 +337,8 @@ void FExpressionExternalInput::EmitValueShader(FEmitContext& Context, FEmitScope
 		case EExternalInput::ParticleTranslatedWorldPosition: Code = TEXT("Parameters.Particle.TranslatedWorldPositionAndSize.xyz"); break;
 		case EExternalInput::ParticleRadius: Code = TEXT("Parameters.Particle.TranslatedWorldPositionAndSize.w"); break;
 
+		case EExternalInput::IsOrthographic: Code = TEXT("((View.ViewToClip[3][3] < 1.0f) ? 0.0f : 1.0f)"); break;
+
 		default:
 			checkNoEntry();
 			break;
diff --git a/Engine/Source/Runtime/Engine/Public/MaterialCompiler.h b/Engine/Source/Runtime/Engine/Public/MaterialCompiler.h
--- a/Engine/Source/Runtime/Engine/Public/MaterialCompiler.h
+++ b/Engine/Source/Runtime/Engine/Public/MaterialCompiler.h
@@ -287,6 +287,7 @@ public:
 	virtual int32 GenericConstant(const UE::Shader::FValue& Value) = 0;
 
 	virtual	int32 ViewProperty(EMaterialExposedViewProperty Property, bool InvProperty = false) = 0;
+	virtual int32 IsOrthographic() = 0;
 
 	virtual int32 GameTime(bool bPeriodic, float Period) = 0;
 	virtual int32 RealTime(bool bPeriodic, float Period) = 0;
@@ -714,6 +715,7 @@ public:
 	virtual int32 GenericConstant(const UE::Shader::FValue& Value) override { return Compiler->GenericConstant(Value); }
 
 	virtual	int32 ViewProperty(EMaterialExposedViewProperty Property, bool InvProperty) override { return Compiler->ViewProperty(Property, InvProperty); }
+	virtual int32 IsOrthographic() override { return Compiler->IsOrthographic(); }
 
 	virtual int32 GameTime(bool bPeriodic, float Period) override { return Compiler->GameTime(bPeriodic, Period); }
 	virtual int32 RealTime(bool bPeriodic, float Period) override { return Compiler->RealTime(bPeriodic, Period); }
diff --git a/Engine/Source/Runtime/Engine/Public/MaterialHLSLTree.h b/Engine/Source/Runtime/Engine/Public/MaterialHLSLTree.h
--- a/Engine/Source/Runtime/Engine/Public/MaterialHLSLTree.h
+++ b/Engine/Source/Runtime/Engine/Public/MaterialHLSLTree.h
@@ -180,6 +180,8 @@ enum class EExternalInput : uint8
 	ParticleTranslatedWorldPosition,
 	ParticleRadius,
 
+	IsOrthographic,
+
 	Num,
 };
 static constexpr int32 MaxNumTexCoords = 8;
