Don't add unused vertex elements in packaged builds when manual vertex fetch is enabled

#rnx

diff --git a/Engine/Source/Runtime/Engine/Private/LocalVertexFactory.cpp b/Engine/Source/Runtime/Engine/Private/LocalVertexFactory.cpp
--- a/Engine/Source/Runtime/Engine/Private/LocalVertexFactory.cpp
+++ b/Engine/Source/Runtime/Engine/Private/LocalVertexFactory.cpp
@@ -458,11 +458,13 @@ void FLocalVertexFactory::InitRHI()
 
 	AddPrimitiveIdStreamElement(EVertexInputStreamType::Default, Elements, 13, 8);
 
-	// Can't rely on manual vertex fetch to not add the unused elements because vertex factories created
+#if !WITH_EDITOR
+	// Can't rely on manual vertex fetch in the editor to not add the unused elements because vertex factories created
 	// with manual vertex fetch support can somehow still be used when booting up in for example ES3.1 preview mode
 	// The vertex factories are then used during mobile rendering and will cause PSO creation failure.
-	// First need to fix invalid usage of these vertex factories before this can be enabled again.
-	//if (!bUseManualVertexFetch)
+	// First need to fix invalid usage of these vertex factories before this can be enabled again. (UE-165187)
+	if (!bUseManualVertexFetch)
+#endif // WITH_EDITOR
 	{
 		// Only the tangent and normal are used by the stream; the bitangent is derived in the shader.
 		uint8 TangentBasisAttributes[2] = { 1, 2 };
