Make some delegate thread-safe to avoid potential race during constructor registration when async loading is activated

diff --git a/Engine/Source/Runtime/Engine/Classes/Engine/World.h b/Engine/Source/Runtime/Engine/Classes/Engine/World.h
--- a/Engine/Source/Runtime/Engine/Classes/Engine/World.h
+++ b/Engine/Source/Runtime/Engine/Classes/Engine/World.h
@@ -4176,7 +4176,7 @@ class ENGINE_API FWorldDelegates
 public:
 	DECLARE_MULTICAST_DELEGATE_TwoParams(FWorldInitializationEvent, UWorld* /*World*/, const UWorld::InitializationValues /*IVS*/);
 	DECLARE_MULTICAST_DELEGATE_ThreeParams(FWorldCleanupEvent, UWorld* /*World*/, bool /*bSessionEnded*/, bool /*bCleanupResources*/);
-	DECLARE_MULTICAST_DELEGATE_OneParam(FWorldEvent, UWorld* /*World*/);
+	DECLARE_TS_MULTICAST_DELEGATE_OneParam(FWorldEvent, UWorld* /*World*/);
 
 	/**
 	 * Post UWorld duplicate event.
diff --git a/Engine/Source/Runtime/Engine/Classes/GameFramework/HUD.h b/Engine/Source/Runtime/Engine/Classes/GameFramework/HUD.h
--- a/Engine/Source/Runtime/Engine/Classes/GameFramework/HUD.h
+++ b/Engine/Source/Runtime/Engine/Classes/GameFramework/HUD.h
@@ -20,7 +20,7 @@ class UFont;
 class UMaterialInterface;
 class UTexture;
 
-DECLARE_MULTICAST_DELEGATE_FiveParams(FOnShowDebugInfo, AHUD* /* HUD */, UCanvas* /* Canvas */, const FDebugDisplayInfo& /* DisplayInfo */, float& /* YL */, float& /* YPos */);
+DECLARE_TS_MULTICAST_DELEGATE_FiveParams(FOnShowDebugInfo, AHUD* /* HUD */, UCanvas* /* Canvas */, const FDebugDisplayInfo& /* DisplayInfo */, float& /* YL */, float& /* YPos */);
 DECLARE_MULTICAST_DELEGATE_TwoParams(FOnHUDPostRender, AHUD* /* HUD */, UCanvas* /* Canvas */);
 
 /** 
diff --git a/Engine/Source/Runtime/Engine/Public/AssetCompilingManager.h b/Engine/Source/Runtime/Engine/Public/AssetCompilingManager.h
--- a/Engine/Source/Runtime/Engine/Public/AssetCompilingManager.h
+++ b/Engine/Source/Runtime/Engine/Public/AssetCompilingManager.h
@@ -72,7 +72,7 @@ protected:
 	ENGINE_API virtual ~IAssetCompilingManager() {}
 };
 
-DECLARE_MULTICAST_DELEGATE_OneParam(FAssetPostCompileEvent, const TArray<FAssetCompileData>&);
+DECLARE_TS_MULTICAST_DELEGATE_OneParam(FAssetPostCompileEvent, const TArray<FAssetCompileData>&);
 
 class FAssetCompilingManager
 {
diff --git a/Engine/Source/Runtime/Engine/Public/AudioDeviceManager.h b/Engine/Source/Runtime/Engine/Public/AudioDeviceManager.h
--- a/Engine/Source/Runtime/Engine/Public/AudioDeviceManager.h
+++ b/Engine/Source/Runtime/Engine/Public/AudioDeviceManager.h
@@ -194,11 +194,11 @@ class ENGINE_API FAudioDeviceManagerDelegates
 {
 public:
 	// This delegate is called whenever an entirely new audio device is created.
-	DECLARE_MULTICAST_DELEGATE_OneParam(FOnAudioDeviceCreated, Audio::FDeviceId /* AudioDeviceId*/);
+	DECLARE_TS_MULTICAST_DELEGATE_OneParam(FOnAudioDeviceCreated, Audio::FDeviceId /* AudioDeviceId*/);
 	static FOnAudioDeviceCreated OnAudioDeviceCreated;
 
 	// This delegate is called whenever an audio device is destroyed.
-	DECLARE_MULTICAST_DELEGATE_OneParam(FOnAudioDeviceDestroyed, Audio::FDeviceId /* AudioDeviceId*/);
+	DECLARE_TS_MULTICAST_DELEGATE_OneParam(FOnAudioDeviceDestroyed, Audio::FDeviceId /* AudioDeviceId*/);
 	static FOnAudioDeviceDestroyed OnAudioDeviceDestroyed;
 };
 
