Read skeletal mesh render data from streaming bulk data so the editor can load the cooked data that is not cooked inline.
It fixes editor only renders the highest LOD when using cooked data as only the highest LOD is cooked inline.

diff --git a/Engine/Source/Runtime/Engine/Private/SkeletalMeshLODRenderData.cpp b/Engine/Source/Runtime/Engine/Private/SkeletalMeshLODRenderData.cpp
--- a/Engine/Source/Runtime/Engine/Private/SkeletalMeshLODRenderData.cpp
+++ b/Engine/Source/Runtime/Engine/Private/SkeletalMeshLODRenderData.cpp
@@ -762,6 +762,11 @@ void FSkeletalMeshLODRenderData::Serialize(FArchive& Ar, UObject* Owner, int32 I
 	{
 		return;
 	}
+
+	bool bUsingCookedEditorData = false;
+#if WITH_EDITORONLY_DATA
+	bUsingCookedEditorData = Owner->GetOutermost()->bIsCookedForEditor;
+#endif
 	
 	// Actual flags used during serialization
 	const uint8 ClassDataStripFlags = GenerateClassStripFlags(Ar, OwnerMesh, Idx);
@@ -807,7 +812,7 @@ void FSkeletalMeshLODRenderData::Serialize(FArchive& Ar, UObject* Owner, int32 I
 	{
 		if (bNeedsCPUAccess)
 		{
-			UE_LOG(LogStaticMesh, Verbose, TEXT("[%s] Skeletal Mesh is marked for CPU read."), *OwnerMesh->GetName());
+			UE_LOG(LogSkeletalMesh, Verbose, TEXT("[%s] Skeletal Mesh is marked for CPU read."), *OwnerMesh->GetName());
 		}
 	}
 
@@ -878,11 +883,27 @@ void FSkeletalMeshLODRenderData::Serialize(FArchive& Ar, UObject* Owner, int32 I
 				StreamingBulkData.Serialize(Ar, Owner, Idx, false);
 				bIsLODOptional = StreamingBulkData.IsOptional();
 
-				if (StreamingBulkData.GetBulkDataSize() == 0)
+				const int64 BulkDataSize = StreamingBulkData.GetBulkDataSize();
+				if (BulkDataSize == 0)
 				{
 					bDiscardBulkData = true;
 					BuffersSize = 0;
 				}
+#if WITH_EDITORONLY_DATA
+				else if (bUsingCookedEditorData && Ar.IsLoading())
+				{
+					// When using cooked data in editor, only the highest (lowest detailed) LOD is cooked inline, while the lower LODs are saved in StreamingBulkData.
+					// So here we load StreamingBulkData from disk into TempData, then populate the render data from TempData.
+					TArray<uint8> TempData;
+					TempData.SetNum(BulkDataSize);
+					void* Dest = TempData.GetData();
+					StreamingBulkData.GetCopy(&Dest);
+
+					FMemoryReader MemReader(TempData, true);
+					MemReader.SetByteSwapping(Ar.IsByteSwapping());
+					SerializeStreamedData(MemReader, OwnerMesh, Idx, ClassDataStripFlags, bNeedsCPUAccess, bForceKeepCPUResources);
+				}
+#endif
 			}
 
 			if (!bDiscardBulkData)
diff --git a/Engine/Source/Runtime/Engine/Private/SkeletalMeshRenderData.cpp b/Engine/Source/Runtime/Engine/Private/SkeletalMeshRenderData.cpp
--- a/Engine/Source/Runtime/Engine/Private/SkeletalMeshRenderData.cpp
+++ b/Engine/Source/Runtime/Engine/Private/SkeletalMeshRenderData.cpp
@@ -606,8 +606,20 @@ void FSkeletalMeshRenderData::Serialize(FArchive& Ar, USkinnedAsset* Owner)
 		NumNonOptionalLODs = GetNumNonOptionalLODs();
 	}
 #endif
+	ensure(LODRenderData.Num() >= NumInlinedLODs);
 	
-	CurrentFirstLODIdx = LODRenderData.Num() - NumInlinedLODs;
+#if WITH_EDITORONLY_DATA
+	const bool bUsingCookedEditorData = Owner->GetOutermost()->bIsCookedForEditor;
+	if (bUsingCookedEditorData && Ar.IsLoading())
+	{
+		CurrentFirstLODIdx = Owner->GetMinLodIdx();
+	}
+	else
+#endif
+	{
+		CurrentFirstLODIdx = LODRenderData.Num() - NumInlinedLODs;
+	}
+
 	PendingFirstLODIdx = CurrentFirstLODIdx;
 
 PRAGMA_DISABLE_DEPRECATION_WARNINGS
