Fix dumped ShaderDebugInfo files not produced during a SCW crash when retrying.

diff --git a/Engine/Source/Runtime/Engine/Private/ShaderCompiler/ShaderCompiler.cpp b/Engine/Source/Runtime/Engine/Private/ShaderCompiler/ShaderCompiler.cpp
--- a/Engine/Source/Runtime/Engine/Private/ShaderCompiler/ShaderCompiler.cpp
+++ b/Engine/Source/Runtime/Engine/Private/ShaderCompiler/ShaderCompiler.cpp
@@ -4957,6 +4957,10 @@ bool FShaderCompilingManager::HandlePotentialRetryOnError(TMap<int32, FShaderMap
 					// NOTE: Changes to MaterialTemplate.usf before retrying won't work, because the entry for Material.usf in CurrentJob.Environment.IncludeFileNameToContentsMap isn't reset
 					if (SingleJob)
 					{
+						if (GShaderCompilingManager->ShouldRecompileToDumpShaderDebugInfo(*SingleJob))
+						{
+							SingleJob->Input.DumpDebugInfoPath = GShaderCompilingManager->CreateShaderDebugInfoPath(SingleJob->Input);
+						}
 						SingleJob->Output = FShaderCompilerOutput();
 					}
 					else
