ShaderCompiler support for CookMultiprocess: Add OnMachineResourcesChanged to allow changing the number of WorkerInfos and ShaderCompilerWorkers after startup.
#rnx

diff --git a/Engine/Source/Runtime/Engine/Private/ShaderCompiler/ShaderCompiler.cpp b/Engine/Source/Runtime/Engine/Private/ShaderCompiler/ShaderCompiler.cpp
--- a/Engine/Source/Runtime/Engine/Private/ShaderCompiler/ShaderCompiler.cpp
+++ b/Engine/Source/Runtime/Engine/Private/ShaderCompiler/ShaderCompiler.cpp
@@ -2110,6 +2110,9 @@
 	/** Tracks whether all tasks issued to the worker have been received. */
 	bool bComplete;
 
+	/** Whether this worker is available for new jobs. It will be false when shutting down the worker. */
+	bool bAvailable; 
+
 	/** Time at which the worker started the most recent batch of tasks. */
 	double StartTime;
 
@@ -2123,6 +2126,7 @@
 		bIssuedTasksToWorker(false),		
 		bLaunchedWorker(false),
 		bComplete(false),
+		bAvailable(true),
 		StartTime(0)
 	{
 	}
@@ -2160,18 +2164,85 @@
 {
 	for (uint32 WorkerIndex = 0; WorkerIndex < Manager->NumShaderCompilingThreads; WorkerIndex++)
 	{
-		WorkerInfos.Add(new FShaderCompileWorkerInfo());
+		WorkerInfos.Add(MakeUnique<FShaderCompileWorkerInfo>());
 	}
 }
 
 FShaderCompileThreadRunnable::~FShaderCompileThreadRunnable()
 {
-	for (int32 Index = 0; Index < WorkerInfos.Num(); Index++)
-	{
-		delete WorkerInfos[Index];
-	}
-
-	WorkerInfos.Empty(0);
+	FScopeLock WorkerScopeLock(&WorkerInfosLock);
+	WorkerInfos.Empty();
 }
 
+void FShaderCompileThreadRunnable::OnMachineResourcesChanged()
+{
+	bool bWaitForWorkersToShutdown = false;
+	{
+		FScopeLock WorkerScopeLock(&WorkerInfosLock);
+		// Set all bAvailable flags back to true
+		for (TUniquePtr< FShaderCompileWorkerInfo>& WorkerInfo : WorkerInfos)
+		{
+			WorkerInfo->bAvailable = true;
+		}
+
+		if (Manager->NumShaderCompilingThreads >= static_cast<uint32>(WorkerInfos.Num()))
+		{
+			while (static_cast<uint32>(WorkerInfos.Num()) < Manager->NumShaderCompilingThreads)
+			{
+				WorkerInfos.Add(MakeUnique<FShaderCompileWorkerInfo>());
+			}
+		}
+		else
+		{
+			for (int32 Index = 0; Index < WorkerInfos.Num(); ++Index)
+			{
+				FShaderCompileWorkerInfo& WorkerInfo = *WorkerInfos[Index];
+				bool bReadyForShutdown = WorkerInfo.QueuedJobs.Num() == 0;
+				if (bReadyForShutdown)
+				{
+					WorkerInfos.RemoveAtSwap(Index--);
+					if (WorkerInfos.Num() == Manager->NumShaderCompilingThreads)
+					{
+						break;
+					}
+				}
+			}
+			bWaitForWorkersToShutdown = Manager->NumShaderCompilingThreads < static_cast<uint32>(WorkerInfos.Num());
+			for (int32 Index = WorkerInfos.Num() - 1;
+				static_cast<uint32>(Index) >= Manager->NumShaderCompilingThreads; --Index)
+			{
+				WorkerInfos[Index]->bAvailable = false;
+			}
+		}
+	}
+	const double StartTime = FPlatformTime::Seconds();
+	constexpr float MaxDurationToWait = 60.f;
+	const double MaxTimeToWait = StartTime + MaxDurationToWait;
+	while (bWaitForWorkersToShutdown)
+	{
+		FPlatformProcess::Sleep(0.01f);
+		const double CurrentTime = FPlatformTime::Seconds();
+		if (CurrentTime > MaxTimeToWait)
+		{
+			UE_LOG(LogShaderCompilers, Warning, TEXT("OnMachineResourcesChanged timedout waiting %.0f seconds for WorkerInfos to complete. Workers will remain allocated."),
+				(float)(CurrentTime - StartTime));
+			break;
+		}
+
+		FScopeLock WorkerScopeLock(&WorkerInfosLock);
+		for (int32 Index = WorkerInfos.Num() - 1;
+			static_cast<uint32>(Index) >= Manager->NumShaderCompilingThreads; --Index)
+		{
+			FShaderCompileWorkerInfo& WorkerInfo = *WorkerInfos[Index];
+			check(!WorkerInfos[Index]->bAvailable); // It should still be set to false from when we changed it above
+			bool bReadyForShutdown = WorkerInfo.QueuedJobs.Num() == 0;
+			if (bReadyForShutdown)
+			{
+				WorkerInfos.RemoveAtSwap(Index);
+			}
+		}
+		bWaitForWorkersToShutdown = Manager->NumShaderCompilingThreads < static_cast<uint32>(WorkerInfos.Num());
+	}
+}
+
 /** Entry point for the shader compiling thread. */
@@ -2191,6 +2262,7 @@
 {
 	TRACE_CPUPROFILER_EVENT_SCOPE(FShaderCompileThreadRunnable::PullTasksFromQueue);
 
+	FScopeLock WorkerScopeLock(&WorkerInfosLock); // Must be entered before CompileQueueSection
 	int32 NumActiveThreads = 0;
 	int32 NumJobsStarted[NumShaderCompileJobPriorities] = { 0 };
 	{
@@ -2210,7 +2282,7 @@
 				FShaderCompileWorkerInfo& CurrentWorkerInfo = *WorkerInfos[WorkerIndex];
 
 				// If this worker doesn't have any queued jobs, look for more in the input queue
-				if (CurrentWorkerInfo.QueuedJobs.Num() == 0 && WorkerIndex < NumWorkersToFeed)
+				if (CurrentWorkerInfo.QueuedJobs.Num() == 0 && CurrentWorkerInfo.bAvailable && WorkerIndex < NumWorkersToFeed)
 				{
 					check(!CurrentWorkerInfo.bComplete);
 
@@ -2273,6 +2345,8 @@
 
 void FShaderCompileThreadRunnable::PushCompletedJobsToManager()
 {
+	FScopeLock WorkerScopeLock(&WorkerInfosLock); // Must be entered before CompileQueueSection
+
 	for (int32 WorkerIndex = 0; WorkerIndex < WorkerInfos.Num(); WorkerIndex++)
 	{
 		FShaderCompileWorkerInfo& CurrentWorkerInfo = *WorkerInfos[WorkerIndex];
@@ -2348,6 +2422,7 @@
 void FShaderCompileThreadRunnable::WriteNewTasks()
 {
 	TRACE_CPUPROFILER_EVENT_SCOPE(ShaderCompiler.WriteNewTasks);
+	FScopeLock WorkerScopeLock(&WorkerInfosLock);
 
 	// first, a quick check if anything is needed just to avoid hammering the task graph
 	bool bHasTasksToWrite = false;
@@ -2369,6 +2444,8 @@
 
 	auto LoopBody = [this](int32 WorkerIndex)
 	{
+		// The calling thread holds the WorkerInfosLock and will not modify WorkerInfos, 
+		// so we can access it here without entering the lock
 		FShaderCompileWorkerInfo& CurrentWorkerInfo = *WorkerInfos[WorkerIndex];
 
 		// Only write tasks once
@@ -2491,6 +2568,7 @@
 		LastCheckForWorkersTime = CurrentTime;
 	}
 
+	FScopeLock WorkerScopeLock(&WorkerInfosLock);
 	for (int32 WorkerIndex = 0; WorkerIndex < WorkerInfos.Num(); WorkerIndex++)
 	{
 		FShaderCompileWorkerInfo& CurrentWorkerInfo = *WorkerInfos[WorkerIndex];
@@ -2570,6 +2648,7 @@
 {
 	TRACE_CPUPROFILER_EVENT_SCOPE(ShaderCompiler.ReadAvailableResults);
 	int32 NumProcessed = 0;
+	FScopeLock WorkerScopeLock(&WorkerInfosLock);
 
 	// first, a quick check if anything is needed just to avoid hammering the task graph
 	bool bHasQueuedJobs = false;
@@ -2589,6 +2668,8 @@
 
 	auto LoopBody = [this, &NumProcessed](int32 WorkerIndex)
 	{
+		// The calling thread holds the WorkerInfosLock and will not modify WorkerInfos, 
+		// so we can access it here without entering the lock
 		FShaderCompileWorkerInfo& CurrentWorkerInfo = *WorkerInfos[WorkerIndex];
 
 		// Check for available result files
@@ -2666,6 +2747,7 @@
 	// If we aren't compiling through workers, so we can just track the serial time here.
 	COOK_STAT(FScopedDurationTimer CompileTimer (ShaderCompilerCookStats::AsyncCompileTimeSec));
 
+	FScopeLock WorkerScopeLock(&WorkerInfosLock);
 	for (int32 WorkerIndex = 0; WorkerIndex < WorkerInfos.Num(); WorkerIndex++)
 	{
 		FShaderCompileWorkerInfo& CurrentWorkerInfo = *WorkerInfos[WorkerIndex];
@@ -3543,8 +3625,6 @@
 		return;
 	}
 
-	bool bForceUseSCWMemoryPressureLimits = false;
-
 	bIsEngineLoopInitialized = false;
 	FCoreDelegates::OnFEngineLoopInitComplete.AddLambda([&]() 
 		{ 
@@ -3573,44 +3653,11 @@
 		bAllowAsynchronousShaderCompiling = false;
 	}
 
-	int32 ShaderCompilerCoreCountThreshold;
-	verify(GConfig->GetInt(TEXT("DevOptions.Shaders"), TEXT("ShaderCompilerCoreCountThreshold"), ShaderCompilerCoreCountThreshold, GEngineIni));
-
-	const int32 NumVirtualCores = FPlatformMisc::NumberOfCoresIncludingHyperthreads();
-
-	int32 NumUnusedShaderCompilingThreads;
-	verify(GConfig->GetInt(TEXT("DevOptions.Shaders"), TEXT("NumUnusedShaderCompilingThreads"), NumUnusedShaderCompilingThreads, GEngineIni));
-
-	int32 NumUnusedShaderCompilingThreadsDuringGame;
-	verify(GConfig->GetInt(TEXT("DevOptions.Shaders"), TEXT("NumUnusedShaderCompilingThreadsDuringGame"), NumUnusedShaderCompilingThreadsDuringGame, GEngineIni));
-
-	// Don't reserve threads based on a percentage if we are in a commandlet or on a low core machine.
-	// In these scenarios we should try to use as many threads as possible.
-	if (!IsRunningCommandlet() && !GIsBuildMachine && NumVirtualCores > ShaderCompilerCoreCountThreshold)
-	{
-		// Reserve a percentage of the threads for general background work.
-		float PercentageUnusedShaderCompilingThreads;
-		verify(GConfig->GetFloat(TEXT("DevOptions.Shaders"), TEXT("PercentageUnusedShaderCompilingThreads"), PercentageUnusedShaderCompilingThreads, GEngineIni));
-
-		// ensure we get a valid multiplier.
-		PercentageUnusedShaderCompilingThreads = FMath::Clamp(PercentageUnusedShaderCompilingThreads, 0.0f, 100.0f) / 100.0f;
-
-		NumUnusedShaderCompilingThreads = FMath::CeilToInt(NumVirtualCores * PercentageUnusedShaderCompilingThreads);
-		NumUnusedShaderCompilingThreadsDuringGame = NumUnusedShaderCompilingThreads;
-	}
-
-	// Use all the cores on the build machines.
-	if (GForceAllCoresForShaderCompiling != 0)
-	{
-		NumUnusedShaderCompilingThreads = 0;
-	}
-
 	verify(GConfig->GetInt( TEXT("DevOptions.Shaders"), TEXT("MaxShaderJobBatchSize"), MaxShaderJobBatchSize, GEngineIni ));
 	verify(GConfig->GetBool( TEXT("DevOptions.Shaders"), TEXT("bPromptToRetryFailedShaderCompiles"), bPromptToRetryFailedShaderCompiles, GEngineIni ));
 	verify(GConfig->GetBool( TEXT("DevOptions.Shaders"), TEXT("bLogJobCompletionTimes"), bLogJobCompletionTimes, GEngineIni ));
 	GConfig->GetFloat(TEXT("DevOptions.Shaders"), TEXT("WorkerTimeToLive"), GRegularWorkerTimeToLive, GEngineIni);
 	GConfig->GetFloat(TEXT("DevOptions.Shaders"), TEXT("BuildWorkerTimeToLive"), GBuildWorkerTimeToLive, GEngineIni);
-	GConfig->GetBool(TEXT("DevOptions.Shaders"), TEXT("bForceUseSCWMemoryPressureLimits"), bForceUseSCWMemoryPressureLimits, GEngineIni);
 
 	verify(GConfig->GetFloat( TEXT("DevOptions.Shaders"), TEXT("ProcessGameThreadTargetTime"), ProcessGameThreadTargetTime, GEngineIni ));
 
@@ -3656,6 +3703,111 @@
 	FPaths::NormalizeDirectoryName(AbsoluteDebugInfoDirectory);
 	AbsoluteShaderDebugInfoDirectory = AbsoluteDebugInfoDirectory;
 
+	CalculateNumberOfCompilingThreads();
+
+	TUniquePtr<FShaderCompileThreadRunnableBase> RemoteCompileThread;
+	const bool bCanUseRemoteCompiling = bAllowCompilingThroughWorkers && ShaderCompiler::IsRemoteCompilingAllowed() && AllTargetPlatformSupportsRemoteShaderCompiling();
+	BuildDistributionController = bCanUseRemoteCompiling ? FindRemoteCompilerController() : nullptr;
+	
+	if (BuildDistributionController)
+	{
+		UE_LOG(LogShaderCompilers, Display, TEXT("Using %s for Shader Compilation."), *BuildDistributionController->GetName());
+		RemoteCompileThread = MakeUnique<FShaderCompileDistributedThreadRunnable_Interface>(this, *BuildDistributionController);
+	}
+
+	GConfig->SetBool(TEXT("/Script/UnrealEd.UnrealEdOptions"), TEXT("UsingXGE"), RemoteCompileThread.IsValid(), GEditorIni);
+
+	TUniquePtr<FShaderCompileThreadRunnableBase> LocalThread = MakeUnique<FShaderCompileThreadRunnable>(this);
+	if (RemoteCompileThread)
+	{
+		checkf(ShaderCompiler::IsRemoteCompilingAllowed(), TEXT("We have a remote compiling thread without the remote compilation being allowed"));
+
+		// Only force-local jobs are guaranteed to stay on the local machine. Going wide with High priority jobs is important for the startup times,
+		// since special materials use High priority. Possibly the partition by priority is too rigid in general.
+		RemoteCompileThread->SetPriorityRange(EShaderCompileJobPriority::Low, EShaderCompileJobPriority::High);
+		LocalThread->SetPriorityRange(EShaderCompileJobPriority::Normal, EShaderCompileJobPriority::ForceLocal);
+		Threads.Add(MoveTemp(RemoteCompileThread));
+	}
+	else
+	{
+		UE_LOG(LogShaderCompilers, Display, TEXT("Using Local Shader Compiler with %d workers."), NumShaderCompilingThreads);
+
+		if (GIsBuildMachine)
+		{
+			int32 MinSCWsToSpawnBeforeWarning = 8; // optional, default to 8
+			GConfig->GetInt(TEXT("DevOptions.Shaders"), TEXT("MinSCWsToSpawnBeforeWarning"), MinSCWsToSpawnBeforeWarning, GEngineIni);
+			if (NumShaderCompilingThreads < static_cast<uint32>(MinSCWsToSpawnBeforeWarning))
+			{
+				UE_LOG(LogShaderCompilers, Warning, TEXT("Only %d SCWs will be spawned, which will result in longer shader compile times."), NumShaderCompilingThreads);
+			}
+		}
+	}
+	Threads.Add(MoveTemp(LocalThread));
+
+	for (const auto& Thread : Threads)
+	{
+		Thread->StartThread();
+	}
+
+	FAssetCompilingManager::Get().RegisterManager(this);
+}
+
+FShaderCompilingManager::~FShaderCompilingManager()
+{
+	// we never initialized, so nothing to do
+	if (!AllowShaderCompiling())
+	{
+		return;
+	}
+
+	PrintStats(true);
+
+	for (const auto& Thread : Threads)
+	{
+		Thread->Stop();
+		Thread->WaitForCompletion();
+	}
+
+	FAssetCompilingManager::Get().UnregisterManager(this);
+}
+
+void FShaderCompilingManager::CalculateNumberOfCompilingThreads()
+{
+	const int32 NumVirtualCores = FPlatformMisc::NumberOfCoresIncludingHyperthreads();
+
+	int32 NumUnusedShaderCompilingThreads;
+	verify(GConfig->GetInt(TEXT("DevOptions.Shaders"), TEXT("NumUnusedShaderCompilingThreads"), NumUnusedShaderCompilingThreads, GEngineIni));
+
+	int32 NumUnusedShaderCompilingThreadsDuringGame;
+	verify(GConfig->GetInt(TEXT("DevOptions.Shaders"), TEXT("NumUnusedShaderCompilingThreadsDuringGame"), NumUnusedShaderCompilingThreadsDuringGame, GEngineIni));
+
+	int32 ShaderCompilerCoreCountThreshold;
+	verify(GConfig->GetInt(TEXT("DevOptions.Shaders"), TEXT("ShaderCompilerCoreCountThreshold"), ShaderCompilerCoreCountThreshold, GEngineIni));
+
+	bool bForceUseSCWMemoryPressureLimits = false;
+	GConfig->GetBool(TEXT("DevOptions.Shaders"), TEXT("bForceUseSCWMemoryPressureLimits"), bForceUseSCWMemoryPressureLimits, GEngineIni);
+
+	// Don't reserve threads based on a percentage if we are in a commandlet or on a low core machine.
+	// In these scenarios we should try to use as many threads as possible.
+	if (!IsRunningCommandlet() && !GIsBuildMachine && NumVirtualCores > ShaderCompilerCoreCountThreshold)
+	{
+		// Reserve a percentage of the threads for general background work.
+		float PercentageUnusedShaderCompilingThreads;
+		verify(GConfig->GetFloat(TEXT("DevOptions.Shaders"), TEXT("PercentageUnusedShaderCompilingThreads"), PercentageUnusedShaderCompilingThreads, GEngineIni));
+
+		// ensure we get a valid multiplier.
+		PercentageUnusedShaderCompilingThreads = FMath::Clamp(PercentageUnusedShaderCompilingThreads, 0.0f, 100.0f) / 100.0f;
+
+		NumUnusedShaderCompilingThreads = FMath::CeilToInt(NumVirtualCores * PercentageUnusedShaderCompilingThreads);
+		NumUnusedShaderCompilingThreadsDuringGame = NumUnusedShaderCompilingThreads;
+	}
+
+	// Use all the cores on the build machines.
+	if (GForceAllCoresForShaderCompiling != 0)
+	{
+		NumUnusedShaderCompilingThreads = 0;
+	}
+
 	NumShaderCompilingThreads = (bAllowCompilingThroughWorkers && NumVirtualCores > NumUnusedShaderCompilingThreads) ? (NumVirtualCores - NumUnusedShaderCompilingThreads) : 1;
 
 	// Make sure there's at least one worker allowed to be active when compiling during the game
@@ -3711,7 +3863,7 @@
 		}
 		else if (bForceUseSCWMemoryPressureLimits)
 		{
-			UE_LOG(LogShaderCompilers, Warning, TEXT("bForceUseSCWMemoryPressureLimits was set but missing one or more prerequisite setting(s): CookerMemoryUsedInGB, MemoryToLeaveForTheOSInGB, MemoryUsedPerSCWProcessInGB.  Ignoring bForceUseSCWMemoryPressureLimits") );
+			UE_LOG(LogShaderCompilers, Warning, TEXT("bForceUseSCWMemoryPressureLimits was set but missing one or more prerequisite setting(s): CookerMemoryUsedInGB, MemoryToLeaveForTheOSInGB, MemoryUsedPerSCWProcessInGB.  Ignoring bForceUseSCWMemoryPressureLimits"));
 		}
 
 		if (GIsBuildMachine)
@@ -3726,71 +3878,15 @@
 	NumShaderCompilingThreadsDuringGame = FMath::Max<int32>(1, NumShaderCompilingThreadsDuringGame);
 
 	NumShaderCompilingThreadsDuringGame = FMath::Min<int32>(NumShaderCompilingThreadsDuringGame, NumShaderCompilingThreads);
-
-	TUniquePtr<FShaderCompileThreadRunnableBase> RemoteCompileThread;
-	const bool bCanUseRemoteCompiling = bAllowCompilingThroughWorkers && ShaderCompiler::IsRemoteCompilingAllowed() && AllTargetPlatformSupportsRemoteShaderCompiling();
-	BuildDistributionController = bCanUseRemoteCompiling ? FindRemoteCompilerController() : nullptr;
-	
-	if (BuildDistributionController)
-	{
-		UE_LOG(LogShaderCompilers, Display, TEXT("Using %s for Shader Compilation."), *BuildDistributionController->GetName());
-		RemoteCompileThread = MakeUnique<FShaderCompileDistributedThreadRunnable_Interface>(this, *BuildDistributionController);
-	}
-
-	GConfig->SetBool(TEXT("/Script/UnrealEd.UnrealEdOptions"), TEXT("UsingXGE"), RemoteCompileThread.IsValid(), GEditorIni);
-
-	TUniquePtr<FShaderCompileThreadRunnableBase> LocalThread = MakeUnique<FShaderCompileThreadRunnable>(this);
-	if (RemoteCompileThread)
-	{
-		checkf(ShaderCompiler::IsRemoteCompilingAllowed(), TEXT("We have a remote compiling thread without the remote compilation being allowed"));
-
-		// Only force-local jobs are guaranteed to stay on the local machine. Going wide with High priority jobs is important for the startup times,
-		// since special materials use High priority. Possibly the partition by priority is too rigid in general.
-		RemoteCompileThread->SetPriorityRange(EShaderCompileJobPriority::Low, EShaderCompileJobPriority::High);
-		LocalThread->SetPriorityRange(EShaderCompileJobPriority::Normal, EShaderCompileJobPriority::ForceLocal);
-		Threads.Add(MoveTemp(RemoteCompileThread));
-	}
-	else
-	{
-		UE_LOG(LogShaderCompilers, Display, TEXT("Using Local Shader Compiler with %d workers."), NumShaderCompilingThreads);
-
-		if (GIsBuildMachine)
-		{
-			int32 MinSCWsToSpawnBeforeWarning = 8; // optional, default to 8
-			GConfig->GetInt(TEXT("DevOptions.Shaders"), TEXT("MinSCWsToSpawnBeforeWarning"), MinSCWsToSpawnBeforeWarning, GEngineIni);
-			if (NumShaderCompilingThreads < static_cast<uint32>(MinSCWsToSpawnBeforeWarning))
-			{
-				UE_LOG(LogShaderCompilers, Warning, TEXT("Only %d SCWs will be spawned, which will result in longer shader compile times."), NumShaderCompilingThreads);
-			}
-		}
-	}
-	Threads.Add(MoveTemp(LocalThread));
-
-	for (const auto& Thread : Threads)
-	{
-		Thread->StartThread();
-	}
-
-	FAssetCompilingManager::Get().RegisterManager(this);
 }
 
-FShaderCompilingManager::~FShaderCompilingManager()
+void FShaderCompilingManager::OnMachineResourcesChanged()
 {
-	// we never initialized, so nothbing to do
-	if (!AllowShaderCompiling())
-	{
-		return;
-	}
-
-	PrintStats(true);
-
+	CalculateNumberOfCompilingThreads();
 	for (const auto& Thread : Threads)
 	{
-		Thread->Stop();
-		Thread->WaitForCompletion();
+		Thread->OnMachineResourcesChanged();
 	}
-
-	FAssetCompilingManager::Get().UnregisterManager(this);
 }
 
 FName FShaderCompilingManager::GetStaticAssetTypeName()
diff --git a/Engine/Source/Runtime/Engine/Public/ShaderCompiler.h b/Engine/Source/Runtime/Engine/Public/ShaderCompiler.h
--- a/Engine/Source/Runtime/Engine/Public/ShaderCompiler.h
+++ b/Engine/Source/Runtime/Engine/Public/ShaderCompiler.h
@@ -555,6 +555,9 @@ public:
 
 	/** Main work loop. */
 	virtual int32 CompilingLoop() = 0;
+
+	/** Events from the manager */
+	virtual void OnMachineResourcesChanged() {}
 };
 
 /** 
@@ -567,7 +570,9 @@ class FShaderCompileThreadRunnable : public FShaderCompileThreadRunnableBase
 private:
 
 	/** Information about the active workers that this thread is tracking. */
-	TArray<struct FShaderCompileWorkerInfo*> WorkerInfos;
+	TArray<TUniquePtr<struct FShaderCompileWorkerInfo>> WorkerInfos;
+	FCriticalSection WorkerInfosLock;
+
 	/** Tracks the last time that this thread checked if the workers were still active. */
 	double LastCheckForWorkersTime;
 
@@ -605,6 +610,8 @@ private:
 
 	/** Main work loop. */
 	virtual int32 CompilingLoop() override;
+
+	virtual void OnMachineResourcesChanged() override;
 };
 
 namespace FShaderCompileUtilities
@@ -927,6 +934,9 @@ private:
 	/** Used to show a notification accompanying progress. */
 	FAsyncCompilationNotification Notification;
 
+	/** Calculate NumShaderCompilingThreads, during construction or OnMachineResourcesChanged */
+	void CalculateNumberOfCompilingThreads();
+
 	/** Launches the worker, returns the launched process handle. */
 	FProcHandle LaunchWorker(const FString& WorkingDirectory, uint32 ProcessId, uint32 ThreadId, const FString& WorkerInputFile, const FString& WorkerOutputFile);
 
@@ -965,6 +975,9 @@ public:
 	ENGINE_API FShaderCompilingManager();
 	ENGINE_API ~FShaderCompilingManager();
 
+	/** Called by external systems that have updated the number of worker threads available. */
+	ENGINE_API void OnMachineResourcesChanged();
+
 	ENGINE_API int32 GetNumPendingJobs() const;
 	ENGINE_API int32 GetNumOutstandingJobs() const;
 
