Support for 16-bit skin weights on the skelmesh.

The main change is that FSoftSkinVertex, used by FSkeletalMeshLODModel, in now stores weights as 16-bit after conversion from the import data. This increases the size of each FSoftSkinVertex from 144 bytes to 160 bytes (about 10% increase). By default render meshes still use 8-bit skin weights, with weights downshifted from the 16-bit modeling data, so no change in GPU memory consumption there. However, the vertex buffer will automatically return a 16-bit skin weights when requested from the GPU side (e.g. for CPU skinning and viewing tangents).

This change in the model data and vertex buffer CPU-side query, resulted in many changes throughout the codebase and will have an effect on licensees who are actively reading from and writing to these two storage locations.

The GPU skin cache shader has had one more permutation added when not using unlimited skin weights. The vertex factory is not affected.

diff --git a/Engine/Plugins/Animation/AnimToTexture/Source/AnimToTextureEditor/Private/AnimToTextureSkeletalMesh.cpp b/Engine/Plugins/Animation/AnimToTexture/Source/AnimToTextureEditor/Private/AnimToTextureSkeletalMesh.cpp
--- a/Engine/Plugins/Animation/AnimToTexture/Source/AnimToTextureEditor/Private/AnimToTextureSkeletalMesh.cpp
+++ b/Engine/Plugins/Animation/AnimToTexture/Source/AnimToTextureEditor/Private/AnimToTextureSkeletalMesh.cpp
@@ -389,11 +389,11 @@ void GetSkinWeightsData(const USkeletalMesh& SkeletalMesh, const int32 LODIndex,
 		// Store Weights
 		for (int32 InfluenceIndex = 0; InfluenceIndex < MAX_TOTAL_INFLUENCES; ++InfluenceIndex)
 		{
-			const uint8& BoneWeight = SkinWeightInfo.InfluenceWeights[InfluenceIndex];
+			const uint16& BoneWeight = SkinWeightInfo.InfluenceWeights[InfluenceIndex];
 			const uint16& BoneIndex = SkinWeightInfo.InfluenceBones[InfluenceIndex];
 			const uint16& MeshBoneIndex = RenderSection.BoneMap[BoneIndex];
 
-			SkinWeights[VertexIndex].BoneWeights[InfluenceIndex] = BoneWeight;
+			SkinWeights[VertexIndex].BoneWeights[InfluenceIndex] = BoneWeight >> 8;
 			SkinWeights[VertexIndex].BoneIndices[InfluenceIndex] = BoneIndex;
 			SkinWeights[VertexIndex].MeshBoneIndices[InfluenceIndex] = MeshBoneIndex;
 		}
diff --git a/Engine/Plugins/Animation/DeformerGraph/Shaders/Private/DataInterfaceSkeleton.ush b/Engine/Plugins/Animation/DeformerGraph/Shaders/Private/DataInterfaceSkeleton.ush
--- a/Engine/Plugins/Animation/DeformerGraph/Shaders/Private/DataInterfaceSkeleton.ush
+++ b/Engine/Plugins/Animation/DeformerGraph/Shaders/Private/DataInterfaceSkeleton.ush
@@ -43,14 +43,27 @@ int4 GetBlendIndices_{DataInterfaceName}(uint StreamOffset, uint ExtraInfluenceI
 
 float4 GetBlendWeights_{DataInterfaceName}(uint StreamOffset, uint ExtraInfluenceIndex)
 {
-	StreamOffset += ((({DataInterfaceName}_NumBoneInfluences + 3) / 4) + ExtraInfluenceIndex) * (1 + GPUSKIN_BONE_INDEX_UINT16);
+	StreamOffset += (({DataInterfaceName}_NumBoneInfluences + 3) / 4) * (1 + GPUSKIN_BONE_INDEX_UINT16);
+
+	float4 UnpackedBlendWeights = 0;
+#if GPUSKIN_BONE_WEIGHTS_UINT16
+	StreamOffset += ExtraInfluenceIndex * 2;
+	uint PackedBlendWeights = {DataInterfaceName}_InputWeightStream[StreamOffset];
+	UnpackedBlendWeights.x = float(PackedBlendWeights & 0xffff) / 65535.0f;
+	UnpackedBlendWeights.y = float(PackedBlendWeights >> 16 & 0xffff) / 65535.0f;
+	PackedBlendWeights = {DataInterfaceName}_InputWeightStream[StreamOffset + 1];
+	UnpackedBlendWeights.z = float(PackedBlendWeights & 0xffff) / 65535.0f;
+	UnpackedBlendWeights.w = float(PackedBlendWeights >> 16 & 0xffff) / 65535.0f;
+	
+#else
+	StreamOffset += ExtraInfluenceIndex;
 	uint PackedBlendWeights = {DataInterfaceName}_InputWeightStream[StreamOffset];
 	
-	float4 UnpackedBlendWeights = 0;
 	UnpackedBlendWeights.x = float(PackedBlendWeights & 0xff) / 255.0f;
 	UnpackedBlendWeights.y = float(PackedBlendWeights >> 8 & 0xff) / 255.0f;
 	UnpackedBlendWeights.z = float(PackedBlendWeights >> 16 & 0xff) / 255.0f;
 	UnpackedBlendWeights.w = float(PackedBlendWeights >> 24 & 0xff) / 255.0f;
+#endif	
 	return UnpackedBlendWeights;
 }
 
@@ -114,9 +127,19 @@ float ReadBoneWeight_{DataInterfaceName}(uint VertexIndex, uint BoneIndex)
 	uint BlendOffsetCount = {DataInterfaceName}_InputWeightLookupStream[VertexIndex];
 	int NumBoneInfluencesLocal = BlendOffsetCount & 0xff;
 	int StreamOffset = BlendOffsetCount >> 8;
-	int WeightsOffset = StreamOffset + ({DataInterfaceName}_InputWeightIndexSize * NumBoneInfluencesLocal);
-
-	return float({DataInterfaceName}_InputWeightStream[WeightsOffset + BoneIndex]) / 255.0f;
+	int BoneIndexSize = {DataInterfaceName}_InputWeightIndexSize & 0xff;
+	int BoneWeightSize = {DataInterfaceName}_InputWeightIndexSize >> 8;
+	int WeightsOffset = StreamOffset + (BoneIndexSize * NumBoneInfluencesLocal);
+	int BoneWeightOffset = WeightsOffset + BoneIndex * BoneWeightSize;
+	
+	if (BoneWeightSize > 1)
+	{
+		return float({DataInterfaceName}_InputWeightStream[BoneWeightOffset + 1] << 8 | {DataInterfaceName}_InputWeightStream[BoneWeightOffset]) / 65535.0;
+	}
+	else
+	{
+		return float({DataInterfaceName}_InputWeightStream[BoneWeightOffset]) / 255.0f;
+	}
 
 #else // !GPUSKIN_UNLIMITED_BONE_INFLUENCE
 	
@@ -141,17 +164,24 @@ float3x4 ReadWeightedBoneMatrix_{DataInterfaceName}(uint VertexIndex)
 	uint BlendOffsetCount = {DataInterfaceName}_InputWeightLookupStream[BufferIndex];
 	int NumBoneInfluencesLocal = BlendOffsetCount & 0xff;
 	int IndicesOffset = BlendOffsetCount >> 8;
-	int WeightsOffset = IndicesOffset + ({DataInterfaceName}_InputWeightIndexSize * NumBoneInfluencesLocal);
+	int BoneIndexSize = {DataInterfaceName}_InputWeightIndexSize & 0xff;
+	int BoneWeightSize = {DataInterfaceName}_InputWeightIndexSize >> 8;
+	int WeightsOffset = IndicesOffset + (BoneIndexSize * NumBoneInfluencesLocal);
 
 	for (int InfluenceIndex = 0; InfluenceIndex < NumBoneInfluencesLocal; InfluenceIndex++)
 	{
-		int BoneIndexOffset = IndicesOffset + ({DataInterfaceName}_InputWeightIndexSize * InfluenceIndex);
+		int BoneIndexOffset = IndicesOffset + (BoneIndexSize * InfluenceIndex);
 		int BoneIndex = {DataInterfaceName}_InputWeightStream[BoneIndexOffset];
 		if ({DataInterfaceName}_InputWeightIndexSize > 1)
 		{
 			BoneIndex = {DataInterfaceName}_InputWeightStream[BoneIndexOffset + 1] << 8 | BoneIndex;
 		}
-		float BoneWeight = float({DataInterfaceName}_InputWeightStream[WeightsOffset + InfluenceIndex]) / 255.0f;
+		int BoneWeightOffset = WeightsOffset + (BoneWeightSize * InfluenceIndex);
+		float BoneWeight;
+		if (BoneWeightSize > 1)
+			BoneWeight = float({DataInterfaceName}_InputWeightStream[BoneWeightOffset + 1] << 8 | {DataInterfaceName}_InputWeightStream[BoneWeightOffset]) / 65535.0;
+		else
+			BoneWeight = float({DataInterfaceName}_InputWeightStream[BoneWeightOffset]) / 255.0f;
 		BlendMatrix += BoneWeight * GetBoneMatrix_{DataInterfaceName}(BoneIndex);
 	}
 
diff --git a/Engine/Plugins/Animation/DeformerGraph/Source/OptimusCore/Private/DataInterfaces/OptimusDataInterfaceSkeleton.cpp b/Engine/Plugins/Animation/DeformerGraph/Source/OptimusCore/Private/DataInterfaces/OptimusDataInterfaceSkeleton.cpp
--- a/Engine/Plugins/Animation/DeformerGraph/Source/OptimusCore/Private/DataInterfaces/OptimusDataInterfaceSkeleton.cpp
+++ b/Engine/Plugins/Animation/DeformerGraph/Source/OptimusCore/Private/DataInterfaces/OptimusDataInterfaceSkeleton.cpp
@@ -82,10 +82,11 @@ void UOptimusSkeletonDataInterface::GetShaderParameters(TCHAR const* UID, FShade
 void UOptimusSkeletonDataInterface::GetPermutations(FComputeKernelPermutationVector& OutPermutationVector) const
 {
 	// Need to be able to support these permutations according to the skeletal mesh settings.
-	// todo[CF]: I think GPUSKIN_UNLIMITED_BONE_INFLUENCE and GPUSKIN_BONE_INDEX_UINT16 are mutually exclusive. So we could save permutations here.
+	// todo[CF]: I think GPUSKIN_UNLIMITED_BONE_INFLUENCE and GPUSKIN_BONE_INDEX_UINT16/GPUSKIN_BONE_WEIGHTS_UINT16 are mutually exclusive. So we could save permutations here.
 	OutPermutationVector.AddPermutation(TEXT("ENABLE_DEFORMER_BONES"), 2);
 	OutPermutationVector.AddPermutation(TEXT("GPUSKIN_UNLIMITED_BONE_INFLUENCE"), 2);
 	OutPermutationVector.AddPermutation(TEXT("GPUSKIN_BONE_INDEX_UINT16"), 2);
+	OutPermutationVector.AddPermutation(TEXT("GPUSKIN_BONE_WEIGHTS_UINT16"), 2);
 	//OutPermutationVector.AddPermutation(TEXT("MERGE_DUPLICATED_VERTICES"), 2);
 }
 
@@ -138,6 +139,7 @@ struct FSkeletonDataInterfacePermutationIds
 	uint32 EnableDeformerBones = 0;
 	uint32 UnlimitedBoneInfluence = 0;
 	uint32 BoneIndexUint16 = 0;
+	uint32 BoneWeightsUint16 = 0;
 
 	FSkeletonDataInterfacePermutationIds(FComputeKernelPermutationVector const& PermutationVector)
 	{
@@ -156,6 +158,11 @@ struct FSkeletonDataInterfacePermutationIds
 			static uint32 Hash = GetTypeHash(Name);
 			BoneIndexUint16 = PermutationVector.GetPermutationBits(Name, Hash, 1);
 		}
+		{
+			static FString Name(TEXT("GPUSKIN_BONE_WEIGHTS_UINT16"));
+			static uint32 Hash = GetTypeHash(Name);
+			BoneWeightsUint16 = PermutationVector.GetPermutationBits(Name, Hash, 1);
+		}
 	}
 };
 
@@ -192,6 +199,7 @@ void FOptimusSkeletonDataProviderProxy::GatherDispatchData(FDispatchSetup const&
 		FRHIShaderResourceView* InputWeightLookupStreamSRV = bUnlimitedBoneInfluences ? WeightBuffer->GetLookupVertexBuffer()->GetSRV() : nullptr;
 		const bool bValidBones = (BoneBufferSRV != nullptr) && (SkinWeightBufferSRV != nullptr) && (!bUnlimitedBoneInfluences || InputWeightLookupStreamSRV != nullptr);
 		const bool bUse16BitBoneIndex = WeightBuffer->Use16BitBoneIndex();
+		const bool bUse16BitBoneWeights = WeightBuffer->Use16BitBoneWeight();
 
 		FSkeletonDataInterfaceParameters* Parameters = (FSkeletonDataInterfaceParameters*)(InOutDispatchData.ParameterBuffer + InDispatchSetup.ParameterBufferOffset + InDispatchSetup.ParameterBufferStride * InvocationIndex);
 		Parameters->NumVertices = RenderSection.NumVertices;
@@ -199,7 +207,7 @@ void FOptimusSkeletonDataProviderProxy::GatherDispatchData(FDispatchSetup const&
 		Parameters->NumBoneInfluences = WeightBuffer->GetMaxBoneInfluences();
 		Parameters->InputWeightStart = (WeightBuffer->GetConstantInfluencesVertexStride() * RenderSection.GetVertexBufferIndex()) / sizeof(float);
 		Parameters->InputWeightStride = WeightBuffer->GetConstantInfluencesVertexStride();
-		Parameters->InputWeightIndexSize = WeightBuffer->GetBoneIndexByteSize();
+		Parameters->InputWeightIndexSize = WeightBuffer->GetBoneIndexByteSize() | (WeightBuffer->GetBoneWeightByteSize() << 8);
 		Parameters->BoneMatrices = BoneBufferSRV != nullptr ? BoneBufferSRV : NullSRVBinding;
 		Parameters->InputWeightStream = SkinWeightBufferSRV != nullptr ? SkinWeightBufferSRV : NullSRVBinding;
 		Parameters->InputWeightLookupStream = InputWeightLookupStreamSRV != nullptr ? InputWeightLookupStreamSRV : NullSRVBinding;
@@ -207,5 +215,6 @@ void FOptimusSkeletonDataProviderProxy::GatherDispatchData(FDispatchSetup const&
 		InOutDispatchData.PermutationId[InvocationIndex] |= (bValidBones ? PermutationIds.EnableDeformerBones : 0);
 		InOutDispatchData.PermutationId[InvocationIndex] |= (bUnlimitedBoneInfluences ? PermutationIds.UnlimitedBoneInfluence : 0);
 		InOutDispatchData.PermutationId[InvocationIndex] |= (bUse16BitBoneIndex ? PermutationIds.BoneIndexUint16 : 0);
+		InOutDispatchData.PermutationId[InvocationIndex] |= (bUse16BitBoneWeights ? PermutationIds.BoneWeightsUint16 : 0);
 	}
 }
diff --git a/Engine/Plugins/Animation/RigLogic/Source/RigLogicModule/Private/SkelMeshDNAUtils.cpp b/Engine/Plugins/Animation/RigLogic/Source/RigLogicModule/Private/SkelMeshDNAUtils.cpp
--- a/Engine/Plugins/Animation/RigLogic/Source/RigLogicModule/Private/SkelMeshDNAUtils.cpp
+++ b/Engine/Plugins/Animation/RigLogic/Source/RigLogicModule/Private/SkelMeshDNAUtils.cpp
@@ -23,6 +23,7 @@
 #include "Modules/ModuleManager.h"
 
 #include "AnimationRuntime.h"
+#include "BoneWeights.h"
 
 #include "riglogic/RigLogic.h"
 
@@ -350,7 +351,7 @@ void USkelMeshDNAUtils::UpdateSkinWeights(USkeletalMesh* SkelMesh, IDNAReader* D
 					if (BoneMapIndex != INDEX_NONE)
 					{
 						// Update influence weight.
-						EngineWeight = (uint8)(NormalizedSkinWeights[InfluenceIndex] * 255.0f);
+						EngineWeight = (uint16)(NormalizedSkinWeights[InfluenceIndex] * UE::AnimationCore::MaxRawBoneWeightFloat);
 					}
 
 					Vertex.InfluenceWeights[InfluenceIndex] = EngineWeight;
@@ -362,8 +363,8 @@ void USkelMeshDNAUtils::UpdateSkinWeights(USkeletalMesh* SkelMesh, IDNAReader* D
 						MaxInfluenceWeight = EngineWeight;
 					}
 				}
-				// Add missing fraction to fill up to 255.
-				Vertex.InfluenceWeights[MaxInfluenceIndex] += 255 - TotalWeight;
+				// Add missing fraction to fill up the bone weights.
+				Vertex.InfluenceWeights[MaxInfluenceIndex] += UE::AnimationCore::MaxRawBoneWeight - TotalWeight;
 			}
 		}
 	}
@@ -394,7 +395,7 @@ void USkelMeshDNAUtils::RebuildRenderData(USkeletalMesh* SkelMesh)
 			}
 
 			const FSkeletalMeshLODModel* LODModelPtr = &LODModelRef;
-			LODRenderData.BuildFromLODModel(LODModelPtr, 0);
+			LODRenderData.BuildFromLODModel(LODModelPtr);
 			LODIndex++;
 		}
 	}
diff --git a/Engine/Plugins/Experimental/Mutable/Source/CustomizableObject/Private/UnrealConversionUtils.cpp b/Engine/Plugins/Experimental/Mutable/Source/CustomizableObject/Private/UnrealConversionUtils.cpp
--- a/Engine/Plugins/Experimental/Mutable/Source/CustomizableObject/Private/UnrealConversionUtils.cpp
+++ b/Engine/Plugins/Experimental/Mutable/Source/CustomizableObject/Private/UnrealConversionUtils.cpp
@@ -87,14 +87,15 @@ namespace UnrealConversionUtils
 	const void* InMutableData)
 {
 	// \todo Ugly cast.
-	const_cast<FSkinWeightDataVertexBuffer*>(OutVertexWeightBuffer.GetDataVertexBuffer())->SetMaxBoneInfluences(NumBoneInfluences);
-	const_cast<FSkinWeightDataVertexBuffer*>(OutVertexWeightBuffer.GetDataVertexBuffer())->Init(NumBones, NumVertices);
+	FSkinWeightDataVertexBuffer* VertexBuffer = const_cast<FSkinWeightDataVertexBuffer*>(OutVertexWeightBuffer.GetDataVertexBuffer());
+	VertexBuffer->SetMaxBoneInfluences(NumBoneInfluences);
+	VertexBuffer->Init(NumBones, NumVertices);
 
 	if (NumVertices)
 	{
 		OutVertexWeightBuffer.SetNeedsCPUAccess(bNeedCPUAccess);
 
-		FSkinWeightInfo* Data = OutVertexWeightBuffer.GetDataVertexBuffer()->GetWeightData();
+		uint8* Data = VertexBuffer->GetWeightData();
 		FMemory::Memcpy(Data, InMutableData, OutVertexWeightBuffer.GetVertexDataSize());
 	}
 }
@@ -223,9 +224,9 @@ namespace UnrealConversionUtils
 	
 	MUTABLE_CPUPROFILER_SCOPE(UpdateSkeletalMesh_SurfaceLoop_MemCpy);
 
-	uint32 BuildFlags = OutSkeletalMesh->GetVertexBufferFlags();
-	const bool bUseFullPrecisionUVs = (BuildFlags & ESkeletalMeshVertexFlags::UseFullPrecisionUVs) != 0;
-	const bool bHasVertexColors = (BuildFlags & ESkeletalMeshVertexFlags::HasVertexColors) != 0;
+	const ESkeletalMeshVertexFlags BuildFlags = OutSkeletalMesh->GetVertexBufferFlags();
+	const bool bUseFullPrecisionUVs = EnumHasAllFlags(BuildFlags, ESkeletalMeshVertexFlags::UseFullPrecisionUVs);
+	const bool bHasVertexColors = EnumHasAllFlags(BuildFlags, ESkeletalMeshVertexFlags::HasVertexColors);
 	const int NumTexCoords = MutableMeshVertexBuffers.GetBufferChannelCount(MUTABLE_VERTEXBUFFER_TEXCOORDS);
 
 	const bool bNeedsCPUAccess = Helper_GetLODInfoArray(OutSkeletalMesh)[MeshLODIndex].bAllowCPUAccess;
diff --git a/Engine/Plugins/Experimental/SkeletalReduction/Source/Private/SkeletalMeshReductionPlugin.cpp b/Engine/Plugins/Experimental/SkeletalReduction/Source/Private/SkeletalMeshReductionPlugin.cpp
--- a/Engine/Plugins/Experimental/SkeletalReduction/Source/Private/SkeletalMeshReductionPlugin.cpp
+++ b/Engine/Plugins/Experimental/SkeletalReduction/Source/Private/SkeletalMeshReductionPlugin.cpp
@@ -34,6 +34,7 @@
 
 #include "Animation/AnimSequence.h"
 #include "BonePose.h"
+#include "BoneWeights.h"
 
 #define LOCTEXT_NAMESPACE "SkeletalMeshReduction"
 
@@ -527,17 +528,17 @@ void FQuadricSkeletalMeshReduction::ConvertToFSkinnedSkeletalMesh( const FSkelet
 	{
 		// Compute the inverse of the total bone influence for this vertex.
 		
-		float InvTotalInfluence = 1.f / 255.f;   // expected default - anything else could indicate a problem with the asset.
+		float InvTotalInfluence = UE::AnimationCore::InvMaxRawBoneWeightFloat;   // expected default - anything else could indicate a problem with the asset.
 		{
 			int32 TotalInfluence = 0;
 
 			for (int32 i = 0; i < MAX_TOTAL_INFLUENCES; ++i)
 			{
-				const uint8 BoneInfluence = Vertex.InfluenceWeights[i];
+				const uint16 BoneInfluence = Vertex.InfluenceWeights[i];
 				TotalInfluence += BoneInfluence;
 			}
 
-			if (TotalInfluence != 255) // 255 is the expected value.  This logic just allows for graceful failure.
+			if (TotalInfluence != UE::AnimationCore::MaxRawBoneWeight) // 255 is the expected value.  This logic just allows for graceful failure.
 			{
 				// Not expected value - record that.
 				bValidBoneWeights = false;
@@ -565,7 +566,7 @@ void FQuadricSkeletalMeshReduction::ConvertToFSkinnedSkeletalMesh( const FSkelet
 		for (int32 i = 0; i < MAX_TOTAL_INFLUENCES; ++i)
 		{
 			const uint16 BoneIndex    = Vertex.InfluenceBones[i];
-			const uint8 BoneInfluence = Vertex.InfluenceWeights[i];
+			const uint16 BoneInfluence = Vertex.InfluenceWeights[i];
 
 			// Accumulate the bone influence for this vert into the BlendedMatrix
 
@@ -773,8 +774,8 @@ void FQuadricSkeletalMeshReduction::ConvertToFSkinnedSkeletalMesh( const FSkelet
 				int32 localBoneId = (int32)SkinnedVertex.InfluenceBones[i];
 				const uint16 boneId = BoneMap[localBoneId];
 
-				const uint8 Influence = SkinnedVertex.InfluenceWeights[i];
-				double boneWeight = ((double)Influence) / 255.;
+				const uint16 Influence = SkinnedVertex.InfluenceWeights[i];
+				double boneWeight = ((double)Influence) / UE::AnimationCore::MaxRawBoneWeightFloat;
 
 				// For right now, only store bone weights that are greater than zero,
 				// by default the sparse data structure assumes a value of zero for 
@@ -1491,7 +1492,7 @@ void  FQuadricSkeletalMeshReduction::AddSourceModelInfluences( const FSkeletalMe
 					for (int32 b = 0; b < MAX_TOTAL_INFLUENCES; ++b)
 					{
 						FBoneIndexType LocalBoneId = SrcRawSkinWeight.InfluenceBones[b];
-						uint8 Weight = SrcRawSkinWeight.InfluenceWeights[b];
+						uint16 Weight = SrcRawSkinWeight.InfluenceWeights[b];
 
 						checkSlow(LocalBoneId < BoneMap.Num());
 
@@ -1500,7 +1501,7 @@ void  FQuadricSkeletalMeshReduction::AddSourceModelInfluences( const FSkeletalMe
 
 						if (Weight > 0)
 						{
-							SkeletalMeshImportData::FRawBoneInfluence RawBoneInfluence = { (float)Weight / 255.f,  VIdx, BoneId };
+							SkeletalMeshImportData::FRawBoneInfluence RawBoneInfluence = { (float)Weight / UE::AnimationCore::MaxRawBoneWeightFloat,  VIdx, BoneId };
 
 							RawBoneInfluences.Add(RawBoneInfluence);
 						}
@@ -1616,48 +1617,40 @@ void  FQuadricSkeletalMeshReduction::AddSourceModelInfluences( const FSkeletalMe
 
 					// Add each bone / weight. 
 					// keep track of the total weight.  should sum to 255 and the first weight is the largest
-					int32 TotalQuantizedWeight = 0;
 					int32 InfluenceIndex = 0;
+					FBoneIndexType InfluenceBones[MAX_TOTAL_INFLUENCES];
+					float InfluenceWeights[MAX_TOTAL_INFLUENCES];
 					for (const auto& Pair : BoneWeight.GetData())
 					{
-						int32 BoneId = Pair.Key;
-						double Weight = Pair.Value;
-
-						// Transform weight to quantized weight
-						uint8 QuantizedWeight = FMath::Clamp((uint8)(Weight*((double)0xFF)), (uint8)0x00, (uint8)0xFF);
-
-						WeightAndBones.InfluenceWeights[InfluenceIndex] = QuantizedWeight;
-
-						TotalQuantizedWeight += QuantizedWeight;
-
-						// Transform boneID to local boneID 
-						// Use the BoneMap to encode this bone
-						int32 LocalBoneId = BoneMap.Find(BoneId);
-						if (LocalBoneId != INDEX_NONE)
+						const int32 BoneId = Pair.Key;
+						const int32 LocalBoneId = BoneMap.Find(BoneId);
+						if (ensure(LocalBoneId != INDEX_NONE))
 						{
-							WeightAndBones.InfluenceBones[InfluenceIndex] = LocalBoneId;
+							InfluenceBones[InfluenceIndex] = LocalBoneId;
+							InfluenceWeights[InfluenceIndex] = static_cast<float>(Pair.Value);
+							InfluenceIndex++;
 						}
-						else
-						{
-							// Map to root of section
-							WeightAndBones.InfluenceBones[InfluenceIndex] = 0;
-
-							check(0); // should never hit this
-						}
-						InfluenceIndex++;
 					}
-					if (InfluenceIndex > MaxNumInfluences)
+					
+					using namespace UE::AnimationCore;
+					FBoneWeights QuantizedWeights = FBoneWeights::Create(InfluenceBones, InfluenceWeights, InfluenceIndex);
+					
+					if (QuantizedWeights.Num() > MaxNumInfluences)
 					{
-						MaxNumInfluences = InfluenceIndex;
+						MaxNumInfluences = QuantizedWeights.Num();
 						if (MaxNumInfluences > Section.GetMaxBoneInfluences())
 						{
 							Section.MaxBoneInfluences = MaxNumInfluences;
 						}
 					}
-					//Use the same code has the build where we modify the index 0 to have a sum of 255 for all influence per skin vertex
-					int32 ExcessQuantizedWeight = 255 - TotalQuantizedWeight;
 
-					WeightAndBones.InfluenceWeights[0] += ExcessQuantizedWeight;
+					InfluenceIndex = 0;
+					for (FBoneWeight QuantizedWeight: QuantizedWeights)
+					{
+						WeightAndBones.InfluenceBones[InfluenceIndex] = QuantizedWeight.GetBoneIndex();
+						WeightAndBones.InfluenceWeights[InfluenceIndex] = QuantizedWeight.GetRawWeight();
+						InfluenceIndex++;
+					}
 				}
 			}
 			
diff --git a/Engine/Plugins/Importers/AlembicImporter/Source/AlembicLibrary/Private/AbcImporter.cpp b/Engine/Plugins/Importers/AlembicImporter/Source/AlembicLibrary/Private/AbcImporter.cpp
--- a/Engine/Plugins/Importers/AlembicImporter/Source/AlembicLibrary/Private/AbcImporter.cpp
+++ b/Engine/Plugins/Importers/AlembicImporter/Source/AlembicLibrary/Private/AbcImporter.cpp
@@ -2,6 +2,8 @@
 
 #include "AbcImporter.h"
 
+#include "BoneWeights.h"
+
 #if PLATFORM_WINDOWS
 #include "Windows/WindowsHWrapper.h"
 #endif
@@ -1759,7 +1761,7 @@ bool FAbcImporter::BuildSkeletalMesh( FSkeletalMeshLODModel& LODModel, const FRe
 				// Set up bone influence (only using one bone so maxed out weight)
 				FMemory::Memzero(NewVertex.InfluenceBones);
 				FMemory::Memzero(NewVertex.InfluenceWeights);
-				NewVertex.InfluenceWeights[0] = 255;
+				NewVertex.InfluenceWeights[0] = UE::AnimationCore::MaxRawBoneWeight;
 				
 				int32 FinalVertexIndex = INDEX_NONE;
 				if (DuplicateVertexIndices.Num())
diff --git a/Engine/Plugins/Importers/USDImporter/Source/USDUtilities/Private/USDSkeletalDataConversion.cpp b/Engine/Plugins/Importers/USDImporter/Source/USDUtilities/Private/USDSkeletalDataConversion.cpp
--- a/Engine/Plugins/Importers/USDImporter/Source/USDUtilities/Private/USDSkeletalDataConversion.cpp
+++ b/Engine/Plugins/Importers/USDImporter/Source/USDUtilities/Private/USDSkeletalDataConversion.cpp
@@ -23,6 +23,7 @@
 #include "Animation/AnimCurveTypes.h"
 #include "AnimationRuntime.h"
 #include "AnimEncoding.h"
+#include "BoneWeights.h"
 #include "ControlRig.h"
 #include "Evaluation/MovieSceneSequenceTransform.h"
 #include "IMovieScenePlayer.h"
@@ -848,7 +849,7 @@ namespace UnrealToUsdImpl
 								int32 BoneIndex = Section.BoneMap[ Vertex.InfluenceBones[ InfluenceIndex ] ];
 
 								JointIndices.push_back( BoneIndex );
-								JointWeights.push_back( Vertex.InfluenceWeights[ InfluenceIndex ] / 255.0f );
+								JointWeights.push_back( Vertex.InfluenceWeights[ InfluenceIndex ] / UE::AnimationCore::MaxRawBoneWeightFloat );
 							}
 						}
 					}
diff --git a/Engine/Plugins/Interchange/Runtime/Source/Import/Private/Mesh/InterchangeSkeletalMeshFactory.cpp b/Engine/Plugins/Interchange/Runtime/Source/Import/Private/Mesh/InterchangeSkeletalMeshFactory.cpp
--- a/Engine/Plugins/Interchange/Runtime/Source/Import/Private/Mesh/InterchangeSkeletalMeshFactory.cpp
+++ b/Engine/Plugins/Interchange/Runtime/Source/Import/Private/Mesh/InterchangeSkeletalMeshFactory.cpp
@@ -582,8 +582,6 @@ namespace UE
 				int32	InfluenceCount = 0;
 
 				float TotalWeight = 0.f;
-				const float MINWEIGHT = 0.01f;
-
 				int MaxVertexInfluence = 0;
 				float MaxIgnoredWeight = 0.0f;
 
@@ -666,8 +664,9 @@ namespace UE
 						LastVertexIndex = Influences[i].VertexIndex;
 					}
 
-					// if less than min weight, or it's more than 8, then we clear it to use weight
-					if (Influences[i].Weight > MINWEIGHT && InfluenceCount < MAX_TOTAL_INFLUENCES)
+					// if less than min weight, or it's more than 12, then we clear it to use weight
+					if (Influences[i].Weight >= UE::AnimationCore::BoneWeightThreshold &&
+						InfluenceCount < MAX_TOTAL_INFLUENCES)
 					{
 						LastNewInfluenceIndex = NewInfluences.Add(Influences[i]);
 						InfluenceCount++;
diff --git a/Engine/Plugins/Tests/EditorTests/Source/EditorTests/Private/UnrealEd/FbxAutomationTests.cpp b/Engine/Plugins/Tests/EditorTests/Source/EditorTests/Private/UnrealEd/FbxAutomationTests.cpp
--- a/Engine/Plugins/Tests/EditorTests/Source/EditorTests/Private/UnrealEd/FbxAutomationTests.cpp
+++ b/Engine/Plugins/Tests/EditorTests/Source/EditorTests/Private/UnrealEd/FbxAutomationTests.cpp
@@ -1869,7 +1869,7 @@ bool FFbxImportAssetsAutomationTest::RunTest(const FString& Parameters)
 							FoundBoneIndex = true;
 							if (Mesh->GetImportedModel() && Mesh->GetImportedModel()->LODModels.IsValidIndex(0))
 							{
-								auto IncrementInfluence = [&VertexSkinByBoneNumber, &BoneIndex](const FSkelMeshSection& Section, const FBoneIndexType (&InfluenceBones)[MAX_TOTAL_INFLUENCES], const uint8 (&InfluenceWeights)[MAX_TOTAL_INFLUENCES])
+								auto IncrementInfluence = [&VertexSkinByBoneNumber, &BoneIndex](const FSkelMeshSection& Section, const FBoneIndexType (&InfluenceBones)[MAX_TOTAL_INFLUENCES], const uint16 (&InfluenceWeights)[MAX_TOTAL_INFLUENCES])
 								{
 									for (int32 InfluenceIndex = 0; InfluenceIndex < MAX_TOTAL_INFLUENCES; ++InfluenceIndex)
 									{
diff --git a/Engine/Plugins/Tests/InterchangeTests/Source/InterchangeTests/Private/ImportTestFunctions/SkeletalMeshImportTestFunctions.cpp b/Engine/Plugins/Tests/InterchangeTests/Source/InterchangeTests/Private/ImportTestFunctions/SkeletalMeshImportTestFunctions.cpp
--- a/Engine/Plugins/Tests/InterchangeTests/Source/InterchangeTests/Private/ImportTestFunctions/SkeletalMeshImportTestFunctions.cpp
+++ b/Engine/Plugins/Tests/InterchangeTests/Source/InterchangeTests/Private/ImportTestFunctions/SkeletalMeshImportTestFunctions.cpp
@@ -292,7 +292,7 @@ FInterchangeTestFunctionResult USkeletalMeshImportTestFunctions::CheckSkinnedVer
 			int32 SkinnedVerticesForBone = 0;
 			auto IncrementInfluence = [&SkinnedVerticesForBone, &BoneIndex](const FSkelMeshSection& Section,
 																			const FBoneIndexType(&InfluenceBones)[MAX_TOTAL_INFLUENCES],
-																			const uint8(&InfluenceWeights)[MAX_TOTAL_INFLUENCES])
+																			const uint16(&InfluenceWeights)[MAX_TOTAL_INFLUENCES])
 			{
 				for (int32 InfluenceIndex = 0; InfluenceIndex < MAX_TOTAL_INFLUENCES; ++InfluenceIndex)
 				{
diff --git a/Engine/Shaders/Private/GpuSkinCacheComputeShader.usf b/Engine/Shaders/Private/GpuSkinCacheComputeShader.usf
--- a/Engine/Shaders/Private/GpuSkinCacheComputeShader.usf
+++ b/Engine/Shaders/Private/GpuSkinCacheComputeShader.usf
@@ -144,7 +144,8 @@ struct FVertexFactoryIntermediates
 };
 
 #if GPUSKIN_UNLIMITED_BONE_INFLUENCE
-uint InputWeightIndexSize; // bone index size: 1 - 8bit or 2 - 16bit
+// Bits 0-7 => Size of the bone weight index in bytes / bits 8-15 => Size of the bone weight weights value in bytes
+uint InputWeightIndexSize;	
 Buffer<uint> InputWeightLookupStream;
 #endif
 Buffer<uint> InputWeightStream;
@@ -155,24 +156,43 @@ FBoneMatrix CalcBoneMatrix( FVertexUnpacked VertexUnpacked, uint VertexIndex )
 	uint BlendOffsetCount = InputWeightLookupStream[VertexIndex];
 	int NumBoneInfluences = BlendOffsetCount & 0xff;
 	int StreamOffset = BlendOffsetCount >> 8;
-	int WeightsOffset = StreamOffset + (InputWeightIndexSize * NumBoneInfluences);
+	int BoneIndexSize = InputWeightIndexSize & 0xff;
+	int BoneWeightSize = InputWeightIndexSize >> 8;
+	int WeightsOffset = StreamOffset + (BoneIndexSize * NumBoneInfluences);
 
 	FBoneMatrix BoneMatrix = FBoneMatrix(float4(0,0,0,0), float4(0,0,0,0), float4(0,0,0,0));
 	for (int InfluenceIdx = 0; InfluenceIdx < NumBoneInfluences; InfluenceIdx++)
 	{
-		int BoneIndexOffset = StreamOffset + (InputWeightIndexSize * InfluenceIdx);
+		int BoneIndexOffset = StreamOffset + (BoneIndexSize * InfluenceIdx);
+		int BoneWeightOffset = WeightsOffset + InfluenceIdx * BoneWeightSize;
 		int BoneIndex = InputWeightStream[BoneIndexOffset];
-		if (InputWeightIndexSize > 1)
+		if (BoneIndexSize > 1)
 		{
 			BoneIndex = InputWeightStream[BoneIndexOffset + 1] << 8 | BoneIndex;
 			//@todo-lh: Workaround to fix issue in SPIRVEmitter of DXC; this block must be inside the if branch
-			float BoneWeight = float(InputWeightStream[WeightsOffset + InfluenceIdx]) / 255.0f;
+			float BoneWeight;
+			if (BoneWeightSize > 1)
+			{
+				BoneWeight = float(InputWeightStream[BoneWeightOffset + 1] << 8 | InputWeightStream[BoneWeightOffset]) / 65535.0;
+			}
+			else
+			{
+				BoneWeight = float(InputWeightStream[BoneWeightOffset]) / 255.0f;
+			}
 			BoneMatrix += BoneWeight * FBoneMatrix(BoneMatrices[BoneIndex*3], BoneMatrices[BoneIndex*3+1], BoneMatrices[BoneIndex*3+2]);
 		}
 		else
 		{
 			//@todo-lh: Workaround to fix issue in SPIRVEmitter of DXC; this block must be inside the if branch
-			float BoneWeight = float(InputWeightStream[WeightsOffset + InfluenceIdx]) / 255.0f;
+			float BoneWeight;
+			if (BoneWeightSize > 1)
+			{
+				BoneWeight = float(InputWeightStream[BoneWeightOffset + 1] << 8 | InputWeightStream[BoneWeightOffset]) / 65535.0;
+			}
+			else
+			{
+				BoneWeight = float(InputWeightStream[BoneWeightOffset]) / 255.0f;
+			}
 			BoneMatrix += BoneWeight * FBoneMatrix(BoneMatrices[BoneIndex*3], BoneMatrices[BoneIndex*3+1], BoneMatrices[BoneIndex*3+2]);
 		}
 	}
@@ -242,6 +262,22 @@ FVertexUnpacked UnpackedVertex(uint MorphIndex, uint ClothIndex, uint ClothIndex
 	Unpacked.BlendIndices2.w = BlendIndices >> 24 & 0xff;
 #endif
 #endif // GPUSKIN_BONE_INDEX_UINT16
+#if GPUSKIN_BONE_WEIGHTS_UINT16
+	uint BlendWeights = InputWeightStream[WeightOffset + GPUSKIN_VB_OFFSET_INFLUENCEWEIGHTS];
+	Unpacked.BlendWeights.x = float(BlendWeights & 0xffff) / 65535.0f;
+	Unpacked.BlendWeights.y = float(BlendWeights >> 16 & 0xffff) / 65535.0f;
+	BlendWeights = InputWeightStream[WeightOffset + GPUSKIN_VB_OFFSET_INFLUENCEWEIGHTS + 1];
+	Unpacked.BlendWeights.z = float(BlendWeights & 0xffff) / 65535.0f;
+	Unpacked.BlendWeights.w = float(BlendWeights >> 16 & 0xffff) / 65535.0f;
+#if GPUSKIN_USE_EXTRA_INFLUENCES
+	BlendWeights = InputWeightStream[WeightOffset + GPUSKIN_VB_OFFSET_INFLUENCEWEIGHTS + 2];
+	Unpacked.BlendWeights2.x = float(BlendWeights & 0xffff) / 65535.0f;
+	Unpacked.BlendWeights2.y = float(BlendWeights >> 16 & 0xffff) / 65535.0f;
+	BlendWeights = InputWeightStream[WeightOffset + GPUSKIN_VB_OFFSET_INFLUENCEWEIGHTS + 3];
+	Unpacked.BlendWeights2.z = float(BlendWeights & 0xffff) / 65535.0f;
+	Unpacked.BlendWeights2.w = float(BlendWeights >> 16 & 0xffff) / 65535.0f;
+#endif
+#else
 	uint BlendWeights = InputWeightStream[WeightOffset + GPUSKIN_VB_OFFSET_INFLUENCEWEIGHTS];
 	Unpacked.BlendWeights.x = float(BlendWeights & 0xff) / 255.0f;
 	Unpacked.BlendWeights.y = float(BlendWeights >> 8 & 0xff) / 255.0f;
@@ -254,6 +290,8 @@ FVertexUnpacked UnpackedVertex(uint MorphIndex, uint ClothIndex, uint ClothIndex
 	Unpacked.BlendWeights2.z = float(BlendWeights >> 16 & 0xff) / 255.0f;
 	Unpacked.BlendWeights2.w = float(BlendWeights >> 24 & 0xff) / 255.0f;
 #endif
+#endif // GPUSKIN_BONE_WEIGHTS_UINT16
+
 #endif // GPUSKIN_UNLIMITED_BONE_INFLUENCE
 
 #if GPUSKIN_MORPH_BLEND
diff --git a/Engine/Shaders/Private/GpuSkinCommon.ush b/Engine/Shaders/Private/GpuSkinCommon.ush
--- a/Engine/Shaders/Private/GpuSkinCommon.ush
+++ b/Engine/Shaders/Private/GpuSkinCommon.ush
@@ -31,6 +31,10 @@
 #define GPUSKIN_BONE_INDEX_UINT16			0	// 16-bit bone index if 1, otherwise 8-bit bone index
 #endif
 
+#ifndef GPUSKIN_BONE_WEIGHTS_UINT16
+#define GPUSKIN_BONE_WEIGHTS_UINT16			0	// 16-bit bone weights if 1, otherwise 8-bit bone weights
+#endif
+
 #ifndef GPUSKIN_UNLIMITED_BONE_INFLUENCE
 #define GPUSKIN_UNLIMITED_BONE_INFLUENCE	0
 #endif
diff --git a/Engine/Shaders/Private/GpuSkinVertexFactory.ush b/Engine/Shaders/Private/GpuSkinVertexFactory.ush
--- a/Engine/Shaders/Private/GpuSkinVertexFactory.ush
+++ b/Engine/Shaders/Private/GpuSkinVertexFactory.ush
@@ -70,7 +70,8 @@ float4x4 PreviousClothToLocal;
 #endif// #if GPUSKIN_APEX_CLOTH
 
 #if GPUSKIN_UNLIMITED_BONE_INFLUENCE
-uint InputWeightIndexSize; // bone index size: 1 - 8bit or 2 - 16bit
+// Bits 0-7 => Size of the bone weight index in bytes / bits 8-15 => Size of the bone weight weights value in bytes
+uint InputWeightIndexSize;
 Buffer<uint> InputWeightStream;
 #endif
 
@@ -536,24 +537,44 @@ FBoneMatrix CalcBoneMatrix( FVertexFactoryInput Input )
 #if GPUSKIN_UNLIMITED_BONE_INFLUENCE
 	int NumBoneInfluences = Input.BlendOffsetCount & 0xff;
 	int StreamOffset = Input.BlendOffsetCount >> 8;
-	int WeightsOffset = StreamOffset + (InputWeightIndexSize * NumBoneInfluences);
+	int BoneIndexSize = InputWeightIndexSize & 0xff;
+	int BoneWeightSize = InputWeightIndexSize >> 8;
+	int WeightsOffset = StreamOffset + (BoneIndexSize * NumBoneInfluences);
 
 	FBoneMatrix BoneMatrix = FBoneMatrix(float4(0,0,0,0), float4(0,0,0,0), float4(0,0,0,0));
 	for (int InfluenceIdx = 0; InfluenceIdx < NumBoneInfluences; InfluenceIdx++)
 	{
-		int BoneIndexOffset = StreamOffset + (InputWeightIndexSize * InfluenceIdx);
+		int BoneIndexOffset = StreamOffset + (BoneIndexSize * InfluenceIdx);
+		int BoneWeightOffset = WeightsOffset + InfluenceIdx * BoneWeightSize;
 		int BoneIndex = InputWeightStream[BoneIndexOffset];
-		if (InputWeightIndexSize > 1)
+		if (BoneIndexSize > 1)
 		{
 			BoneIndex = InputWeightStream[BoneIndexOffset + 1] << 8 | BoneIndex;
 			//@todo-lh: Workaround to fix issue in SPIRVEmitter of DXC; this block must be inside the if branch
-			float BoneWeight = float(InputWeightStream[WeightsOffset + InfluenceIdx]) / 255.0f;
+			// Q: Is this workaround still necessary?
+			float BoneWeight;
+			if (BoneWeightSize > 1)
+			{
+				BoneWeight = float(InputWeightStream[BoneWeightOffset + 1] << 8 | InputWeightStream[BoneWeightOffset]) / 65535.0;
+			}
+			else
+			{
+				BoneWeight = float(InputWeightStream[BoneWeightOffset]) / 255.0f;
+			}
 			BoneMatrix += BoneWeight * GetBoneMatrix(BoneIndex);
 		}
 		else
 		{
 			//@todo-lh: Workaround to fix issue in SPIRVEmitter of DXC; this block must be inside the if branch
-			float BoneWeight = float(InputWeightStream[WeightsOffset + InfluenceIdx]) / 255.0f;
+			float BoneWeight;
+			if (BoneWeightSize > 1)
+			{
+				BoneWeight = float(InputWeightStream[BoneWeightOffset + 1] << 8 | InputWeightStream[BoneWeightOffset]) / 65535.0;
+			}
+			else
+			{
+				BoneWeight = float(InputWeightStream[BoneWeightOffset]) / 255.0f;
+			}
 			BoneMatrix += BoneWeight * GetBoneMatrix(BoneIndex);
 		}
 	}
@@ -593,24 +614,43 @@ FBoneMatrix CalcPreviousBoneMatrix( FVertexFactoryInput Input )
 #if GPUSKIN_UNLIMITED_BONE_INFLUENCE
 	int NumBoneInfluences = Input.BlendOffsetCount & 0xff;
 	int StreamOffset = Input.BlendOffsetCount >> 8;
-	int WeightsOffset = StreamOffset + (InputWeightIndexSize * NumBoneInfluences);
+	int BoneIndexSize = InputWeightIndexSize & 0xff;
+	int BoneWeightSize = InputWeightIndexSize >> 8;
+	int WeightsOffset = StreamOffset + (BoneIndexSize * NumBoneInfluences);
 
 	FBoneMatrix BoneMatrix = FBoneMatrix(float4(0,0,0,0), float4(0,0,0,0), float4(0,0,0,0));
 	for (int InfluenceIdx = 0; InfluenceIdx < NumBoneInfluences; InfluenceIdx++)
 	{
-		int BoneIndexOffset = StreamOffset + (InputWeightIndexSize * InfluenceIdx);
+		int BoneIndexOffset = StreamOffset + (BoneIndexSize * InfluenceIdx);
+		int BoneWeightOffset = WeightsOffset + InfluenceIdx * BoneWeightSize;
 		int BoneIndex = InputWeightStream[BoneIndexOffset];
-		if (InputWeightIndexSize > 1)
+		if (BoneIndexSize > 1)
 		{
 			BoneIndex = InputWeightStream[BoneIndexOffset + 1] << 8 | BoneIndex;
 			//@todo-lh: Workaround to fix issue in SPIRVEmitter of DXC; this block must be inside the if branch
-			float BoneWeight = float(InputWeightStream[WeightsOffset + InfluenceIdx]) / 255.0f;
+			float BoneWeight;
+			if (BoneWeightSize > 1)
+			{
+				BoneWeight = float(InputWeightStream[BoneWeightOffset + 1] << 8 | InputWeightStream[BoneWeightOffset]) / 65535.0;
+			}
+			else
+			{
+				BoneWeight = float(InputWeightStream[BoneWeightOffset]) / 255.0f;
+			}
 			BoneMatrix += BoneWeight * GetPreviousBoneMatrix(BoneIndex);
 		}
 		else
 		{
 			//@todo-lh: Workaround to fix issue in SPIRVEmitter of DXC; this block must be inside the if branch
-			float BoneWeight = float(InputWeightStream[WeightsOffset + InfluenceIdx]) / 255.0f;
+			float BoneWeight;
+			if (BoneWeightSize > 1)
+			{
+				BoneWeight = float(InputWeightStream[BoneWeightOffset + 1] << 8 | InputWeightStream[BoneWeightOffset]) / 65535.0;
+			}
+			else
+			{
+				BoneWeight = float(InputWeightStream[BoneWeightOffset]) / 255.0f;
+			}
 			BoneMatrix += BoneWeight * GetPreviousBoneMatrix(BoneIndex);
 		}
 	}
diff --git a/Engine/Source/Developer/MeshBoneReduction/Private/MeshBoneReduction.cpp b/Engine/Source/Developer/MeshBoneReduction/Private/MeshBoneReduction.cpp
--- a/Engine/Source/Developer/MeshBoneReduction/Private/MeshBoneReduction.cpp
+++ b/Engine/Source/Developer/MeshBoneReduction/Private/MeshBoneReduction.cpp
@@ -12,6 +12,7 @@
 #include "Rendering/SkeletalMeshModel.h"
 #include "Rendering/SkeletalMeshRenderData.h"
 #include "AnimationBlueprintLibrary.h"
+#include "BoneWeights.h"
 #include "Async/ParallelFor.h"
 #include "Animation/AnimSequence.h"
 #include "Animation/AnimSequenceHelpers.h"
@@ -223,7 +224,7 @@ public:
 				{
 					FSoftSkinVertex & Vert = Section.SoftVertices[VertIndex];
 
-					auto RemapBoneInfluenceVertexIndex = [&BoneMapRemapTable](FBoneIndexType InfluenceBones[MAX_TOTAL_INFLUENCES], uint8 InfluenceWeights[MAX_TOTAL_INFLUENCES])
+					auto RemapBoneInfluenceVertexIndex = [&BoneMapRemapTable](FBoneIndexType InfluenceBones[MAX_TOTAL_INFLUENCES], uint16 InfluenceWeights[MAX_TOTAL_INFLUENCES])
 					{
 						bool ShouldRenormalize = false;
 
@@ -464,7 +465,8 @@ public:
 				FSkelMeshSection& Section = NewModel->Sections[SectionIndex];
 				if (bBakePoseToRemovedInfluences)
 				{
-					const float InfluenceMultiplier = 1.0f / 255.0f;
+					using UE::AnimationCore::InvMaxRawBoneWeightFloat;
+					
 					for (FSoftSkinVertex& Vertex : Section.SoftVertices)
 					{
 						FVector TangentX = (FVector)Vertex.TangentX;
@@ -478,11 +480,11 @@ public:
 								const int32 ArrayIndex = BoneIndices.IndexOfByKey(Section.BoneMap[Vertex.InfluenceBones[InfluenceIndex]]);
 								if (ArrayIndex != INDEX_NONE)
 								{
-									Position += (FVector(RemovedBoneMatrices[ArrayIndex].TransformPosition((FVector)Vertex.Position)) - FVector(Vertex.Position)) * ((float)Vertex.InfluenceWeights[InfluenceIndex] * InfluenceMultiplier);
+									Position += (FVector(RemovedBoneMatrices[ArrayIndex].TransformPosition((FVector)Vertex.Position)) - FVector(Vertex.Position)) * ((float)Vertex.InfluenceWeights[InfluenceIndex] * InvMaxRawBoneWeightFloat);
 
-									TangentX += (FVector(RemovedBoneMatrices[ArrayIndex].TransformVector((FVector)Vertex.TangentX)) - FVector(Vertex.TangentX)) * ((float)Vertex.InfluenceWeights[InfluenceIndex] * InfluenceMultiplier);
-									TangentY += (FVector(RemovedBoneMatrices[ArrayIndex].TransformVector((FVector)Vertex.TangentY)) - FVector(Vertex.TangentY)) * ((float)Vertex.InfluenceWeights[InfluenceIndex] * InfluenceMultiplier);
-									TangentZ += (FVector(RemovedBoneMatrices[ArrayIndex].TransformVector((FVector)Vertex.TangentZ)) - FVector(Vertex.TangentZ)) * ((float)Vertex.InfluenceWeights[InfluenceIndex] * InfluenceMultiplier);
+									TangentX += (FVector(RemovedBoneMatrices[ArrayIndex].TransformVector((FVector)Vertex.TangentX)) - FVector(Vertex.TangentX)) * ((float)Vertex.InfluenceWeights[InfluenceIndex] * InvMaxRawBoneWeightFloat);
+									TangentY += (FVector(RemovedBoneMatrices[ArrayIndex].TransformVector((FVector)Vertex.TangentY)) - FVector(Vertex.TangentY)) * ((float)Vertex.InfluenceWeights[InfluenceIndex] * InvMaxRawBoneWeightFloat);
+									TangentZ += (FVector(RemovedBoneMatrices[ArrayIndex].TransformVector((FVector)Vertex.TangentZ)) - FVector(Vertex.TangentZ)) * ((float)Vertex.InfluenceWeights[InfluenceIndex] * InvMaxRawBoneWeightFloat);
 								}
 							}
 				        }
diff --git a/Engine/Source/Developer/MeshUtilities/Private/MeshUtilities.cpp b/Engine/Source/Developer/MeshUtilities/Private/MeshUtilities.cpp
--- a/Engine/Source/Developer/MeshUtilities/Private/MeshUtilities.cpp
+++ b/Engine/Source/Developer/MeshUtilities/Private/MeshUtilities.cpp
@@ -109,6 +109,7 @@
 #include "UnrealEdMisc.h"
 #include "Subsystems/AssetEditorSubsystem.h"
 #endif
+#include "BoneWeights.h"
 #include "MaterialBakingStructures.h"
 #include "IMaterialBakingModule.h"
 #include "MaterialOptions.h"
@@ -4070,27 +4071,43 @@ public:
 						InfluenceCount = EXTRA_BONE_INFLUENCES;
 					}
 
+					using namespace UE::AnimationCore;
+
 					// Setup the vertex influences.
-					Vertex.InfluenceBones[0] = 0;
-					Vertex.InfluenceWeights[0] = 255;
-					for (uint32 i = 1; i < MAX_TOTAL_INFLUENCES; i++)
-					{
-						Vertex.InfluenceBones[i] = 0;
-						Vertex.InfluenceWeights[i] = 0;
-					}
+					FBoneIndexType InfluenceBones[MAX_TOTAL_INFLUENCES];
+					float InfluenceWeights[MAX_TOTAL_INFLUENCES];
 
-					uint32	TotalInfluenceWeight = 0;
-					for (uint32 i = 0; i < InfluenceCount; i++)
+					int32 ActualInfluenceCount = 0; 
+					for (uint32 Index = 0; Index < InfluenceCount; Index++)
 					{
-						FBoneIndexType BoneIndex = (FBoneIndexType)BuildData.Influences[InfIdx + i].BoneIndex;
+						const FBoneIndexType BoneIndex = BuildData.Influences[InfIdx + Index].BoneIndex;
 						if (BoneIndex >= BuildData.RefSkeleton.GetRawBoneNum())
 							continue;
+					
+						InfluenceBones[ActualInfluenceCount] = BoneIndex;
+						InfluenceWeights[ActualInfluenceCount] = (BuildData.Influences[InfIdx + Index].Weight);
+						ActualInfluenceCount++;
+					}
 
-						Vertex.InfluenceBones[i] = BoneIndex;
-						Vertex.InfluenceWeights[i] = (uint8)(BuildData.Influences[InfIdx + i].Weight * 255.0f);
-						TotalInfluenceWeight += Vertex.InfluenceWeights[i];
+					const FBoneWeights BoneWeights = FBoneWeights::Create(InfluenceBones, InfluenceWeights, ActualInfluenceCount);
+				
+					FMemory::Memzero(Vertex.InfluenceBones);
+					FMemory::Memzero(Vertex.InfluenceWeights);
+				
+					if (BoneWeights.Num() == 0)
+					{
+						Vertex.InfluenceWeights[0] = std::numeric_limits<uint16>::max();
+					}
+					else
+					{
+						int32 Index = 0;
+						for (FBoneWeight BoneWeight: BoneWeights)
+						{
+							Vertex.InfluenceBones[Index] = BoneWeight.GetBoneIndex();
+							Vertex.InfluenceWeights[Index] = BoneWeight.GetRawWeight();
+							Index++;
+						}
 					}
-					Vertex.InfluenceWeights[0] += 255 - TotalInfluenceWeight;
 				}
 
 				// Add the vertex as well as its original index in the points array
@@ -4809,6 +4826,7 @@ bool FMeshUtilities::BuildSkeletalMesh_Legacy(FSkeletalMeshLODModel& LODModel
 
 			{
 				// Count the influences.
+				using namespace UE::AnimationCore;
 
 				int32 InfIdx = WedgeInfluenceIndices[Face.iWedge[VertexIndex]];
 				int32 LookIdx = InfIdx;
@@ -4822,27 +4840,40 @@ bool FMeshUtilities::BuildSkeletalMesh_Legacy(FSkeletalMeshLODModel& LODModel
 				InfluenceCount = FMath::Min<uint32>(InfluenceCount, MAX_TOTAL_INFLUENCES);
 
 				// Setup the vertex influences.
+				FBoneIndexType InfluenceBones[MAX_TOTAL_INFLUENCES];
+				float InfluenceWeights[MAX_TOTAL_INFLUENCES];
 
-				Vertex.InfluenceBones[0] = 0;
-				Vertex.InfluenceWeights[0] = 255;
-				for (uint32 i = 1; i < MAX_TOTAL_INFLUENCES; i++)
+				int32 ActualInfluenceCount = 0; 
+				for (uint32 Index = 0; Index < InfluenceCount; Index++)
 				{
-					Vertex.InfluenceBones[i] = 0;
-					Vertex.InfluenceWeights[i] = 0;
-				}
-
-				uint32	TotalInfluenceWeight = 0;
-				for (uint32 i = 0; i < InfluenceCount; i++)
-				{
-					FBoneIndexType BoneIndex = (FBoneIndexType)Influences[InfIdx + i].BoneIndex;
+					const FBoneIndexType BoneIndex = Influences[InfIdx + Index].BoneIndex;
 					if (BoneIndex >= RefSkeleton.GetRawBoneNum())
 						continue;
+					
+					InfluenceBones[ActualInfluenceCount] = BoneIndex;
+					InfluenceWeights[ActualInfluenceCount] = (Influences[InfIdx + Index].Weight);
+					ActualInfluenceCount++;
+				}
 
-					Vertex.InfluenceBones[i] = BoneIndex;
-					Vertex.InfluenceWeights[i] = (uint8)(Influences[InfIdx + i].Weight * 255.0f);
-					TotalInfluenceWeight += Vertex.InfluenceWeights[i];
+				const FBoneWeights BoneWeights = FBoneWeights::Create(InfluenceBones, InfluenceWeights, ActualInfluenceCount);
+				
+				FMemory::Memzero(Vertex.InfluenceBones);
+				FMemory::Memzero(Vertex.InfluenceWeights);
+				
+				if (BoneWeights.Num() == 0)
+				{
+					Vertex.InfluenceWeights[0] = std::numeric_limits<uint16>::max();
+				}
+				else
+				{
+					int32 Index = 0;
+					for (FBoneWeight BoneWeight: BoneWeights)
+					{
+						Vertex.InfluenceBones[Index] = BoneWeight.GetBoneIndex();
+						Vertex.InfluenceWeights[Index] = BoneWeight.GetRawWeight();
+						Index++;
+					}
 				}
-				Vertex.InfluenceWeights[0] += 255 - TotalInfluenceWeight;
 			}
 
 			// Add the vertex as well as its original index in the points array
@@ -5957,7 +5988,7 @@ void FMeshUtilities::CreateImportDataFromLODModel(USkeletalMesh* SkeletalMesh) c
 
 				for (int32 InfluenceIndex = 0; InfluenceIndex < MAX_TOTAL_INFLUENCES; ++InfluenceIndex)
 				{
-					float Weight = static_cast<float>(Vertex.InfluenceWeights[InfluenceIndex]) / 255.0f;
+					float Weight = static_cast<float>(Vertex.InfluenceWeights[InfluenceIndex]) / 65535.0f;
 					if (FMath::IsNearlyZero(Weight))
 					{
 						break;
@@ -5965,7 +5996,7 @@ void FMeshUtilities::CreateImportDataFromLODModel(USkeletalMesh* SkeletalMesh) c
 					SkeletalMeshImportData::FRawBoneInfluence& Influence = ImportData.Influences.AddZeroed_GetRef();
 					Influence.VertexIndex = PointIndex;
 					Influence.BoneIndex = Section.BoneMap[Vertex.InfluenceBones[InfluenceIndex]];
-					Influence.Weight = static_cast<float>(Vertex.InfluenceWeights[InfluenceIndex]) / 255.0f;
+					Influence.Weight = static_cast<float>(Vertex.InfluenceWeights[InfluenceIndex]) / 65535.0f;
 				}
 			}
 		}
@@ -6057,7 +6088,7 @@ void FMeshUtilities::CreateImportDataFromLODModel(USkeletalMesh* SkeletalMesh) c
 					const FRawSkinWeight& SkinWeight = ImportedSkinWeightProfileData.SkinWeights[VerticeIndex];
 					for (int32 InfluenceIndex = 0; InfluenceIndex < MAX_TOTAL_INFLUENCES; ++InfluenceIndex)
 					{
-						float Weight = static_cast<float>(SkinWeight.InfluenceWeights[InfluenceIndex]) / 255.0f;
+						float Weight = static_cast<float>(SkinWeight.InfluenceWeights[InfluenceIndex]) / 65535.0f;
 						if (FMath::IsNearlyZero(Weight))
 						{
 							break;
diff --git a/Engine/Source/Developer/MeshUtilities/Private/SkeletalMeshTools.h b/Engine/Source/Developer/MeshUtilities/Private/SkeletalMeshTools.h
--- a/Engine/Source/Developer/MeshUtilities/Private/SkeletalMeshTools.h
+++ b/Engine/Source/Developer/MeshUtilities/Private/SkeletalMeshTools.h
@@ -31,7 +31,7 @@ struct FSoftSkinBuildVertex
 	FVector2f		UVs[MAX_TEXCOORDS]; // UVs
 	FColor			Color;		// VertexColor
 	FBoneIndexType	InfluenceBones[MAX_TOTAL_INFLUENCES];
-	uint8			InfluenceWeights[MAX_TOTAL_INFLUENCES];
+	uint16			InfluenceWeights[MAX_TOTAL_INFLUENCES];
 	uint32 PointWedgeIdx;
 };
 
diff --git a/Engine/Source/Developer/MeshUtilitiesEngine/Private/MeshUtilitiesEngine.cpp b/Engine/Source/Developer/MeshUtilitiesEngine/Private/MeshUtilitiesEngine.cpp
--- a/Engine/Source/Developer/MeshUtilitiesEngine/Private/MeshUtilitiesEngine.cpp
+++ b/Engine/Source/Developer/MeshUtilitiesEngine/Private/MeshUtilitiesEngine.cpp
@@ -11,8 +11,8 @@
 // Find the most dominant bone for each vertex
 int32 GetDominantBoneIndex(FSoftSkinVertex* SoftVert)
 {
-	uint8 MaxWeightBone = 0;
-	uint8 MaxWeightWeight = 0;
+	uint16 MaxWeightBone = 0;
+	uint16 MaxWeightWeight = 0;
 
 	for (int32 i = 0; i < MAX_TOTAL_INFLUENCES; i++)
 	{
diff --git a/Engine/Source/Developer/PhysicsUtilities/Private/SkinnedBoneTriangleCache.cpp b/Engine/Source/Developer/PhysicsUtilities/Private/SkinnedBoneTriangleCache.cpp
--- a/Engine/Source/Developer/PhysicsUtilities/Private/SkinnedBoneTriangleCache.cpp
+++ b/Engine/Source/Developer/PhysicsUtilities/Private/SkinnedBoneTriangleCache.cpp
@@ -85,7 +85,7 @@ void FSkinnedBoneTriangleCache::AddIndexToInfluencerBoneSets(const FSkinnedVerte
 
 	for (int32 InfluenceIndex = 0; InfluenceIndex < MAX_TOTAL_INFLUENCES; ++InfluenceIndex)
 	{
-		const uint8 InfluenceWeight = SoftVert.InfluenceWeights[InfluenceIndex];
+		const uint16 InfluenceWeight = SoftVert.InfluenceWeights[InfluenceIndex];
 
 		if ( InfluenceHeuristic == EVW_DominantWeight )
 		{
@@ -102,7 +102,7 @@ void FSkinnedBoneTriangleCache::AddIndexToInfluencerBoneSets(const FSkinnedVerte
 			}
 		}
 
-		const uint8 BoneMapIndex = SoftVert.InfluenceBones[InfluenceIndex];
+		const FBoneIndexType BoneMapIndex = SoftVert.InfluenceBones[InfluenceIndex];
 		const int32 ActualBoneIndex = Section.BoneMap[BoneMapIndex];
 
 		FInfluencedVerticesSet& InfluencedVertexIndicesForBone = BoneIndexToInfluencedVertices.FindOrAdd(ActualBoneIndex);
diff --git a/Engine/Source/Developer/SkeletalMeshUtilitiesCommon/Private/LODUtilities.cpp b/Engine/Source/Developer/SkeletalMeshUtilitiesCommon/Private/LODUtilities.cpp
--- a/Engine/Source/Developer/SkeletalMeshUtilitiesCommon/Private/LODUtilities.cpp
+++ b/Engine/Source/Developer/SkeletalMeshUtilitiesCommon/Private/LODUtilities.cpp
@@ -2,6 +2,8 @@
 
 #include "LODUtilities.h"
 
+#include "BoneWeights.h"
+
 #if WITH_EDITOR
 
 #include "Misc/MessageDialog.h"
@@ -47,7 +49,6 @@ IMPLEMENT_MODULE(FDefaultModuleImpl, SkeletalMeshUtilitiesCommon)
 
 DEFINE_LOG_CATEGORY_STATIC(LogLODUtilities, Log, All);
 
-
 /**
 * Process and update the vertex Influences using the predefined wedges
 *
@@ -79,7 +80,6 @@ void FLODUtilities::ProcessImportMeshInfluences(const int32 WedgeCount, TArray<S
 	int32	InfluenceCount = 0;
 
 	float TotalWeight = 0.f;
-	const float MINWEIGHT = 0.01f;
 
 	int MaxVertexInfluence = 0;
 	float MaxIgnoredWeight = 0.0f;
@@ -163,7 +163,8 @@ void FLODUtilities::ProcessImportMeshInfluences(const int32 WedgeCount, TArray<S
 		}
 
 		// if less than min weight, or it's more than 8, then we clear it to use weight
-		if (Influences[i].Weight > MINWEIGHT && InfluenceCount < MAX_TOTAL_INFLUENCES)
+		if (Influences[i].Weight >= UE::AnimationCore::BoneWeightThreshold &&
+			InfluenceCount < MAX_TOTAL_INFLUENCES)
 		{
 			LastNewInfluenceIndex = NewInfluences.Add(Influences[i]);
 			InfluenceCount++;
@@ -2632,12 +2633,6 @@ void FLODUtilities::GenerateImportedSkinWeightProfileData(FSkeletalMeshLODModel&
 		const int32 VertexIndex = LODModelDest.MeshToImportVertexMap[VertexInstanceIndex];
 		check(VertexIndex >= 0 && VertexIndex <= LODModelDest.MaxImportVertex);
 		FRawSkinWeight& SkinWeight = SkinWeights.AddDefaulted_GetRef();
-		//Zero out all value
-		for (int32 InfluenceIndex = 0; InfluenceIndex < MAX_TOTAL_INFLUENCES; ++InfluenceIndex)
-		{
-			SkinWeight.InfluenceBones[InfluenceIndex] = 0;
-			SkinWeight.InfluenceWeights[InfluenceIndex] = 0;
-		}
 
 		TMap<FBoneIndexType, float> WeightForBone;
 		for (const SkeletalMeshImportData::FVertInfluence& VertInfluence : ImportedProfileData.SourceModelInfluences)
@@ -2656,32 +2651,51 @@ void FLODUtilities::GenerateImportedSkinWeightProfileData(FSkeletalMeshLODModel&
 			}
 		}
 
+		using namespace UE::AnimationCore;
 
 		//Add the prepared alternate influences for this skin vertex
-		uint32	TotalInfluenceWeight = 0;
-		int32 InfluenceBoneIndex = 0;
+		int32 InfluenceBoneCount = 0;
+		
+		FBoneIndexType InfluenceBones[MAX_TOTAL_INFLUENCES];
+		float InfluenceWeights[MAX_TOTAL_INFLUENCES];
+		
 		for (auto Kvp : WeightForBone)
 		{
-			SkinWeight.InfluenceBones[InfluenceBoneIndex] = Kvp.Key;
-			SkinWeight.InfluenceWeights[InfluenceBoneIndex] = FMath::Clamp((uint8)(Kvp.Value*((float)0xFF)), (uint8)0x00, (uint8)0xFF);
-			TotalInfluenceWeight += SkinWeight.InfluenceWeights[InfluenceBoneIndex];
-			InfluenceBoneIndex++;
-			if (InfluenceBoneIndex >= MaxInfluenceCount)
+			InfluenceBones[InfluenceBoneCount] = Kvp.Key;
+			InfluenceWeights[InfluenceBoneCount] = Kvp.Value;
+			InfluenceBoneCount++;
+		}
+
+		const FBoneWeights BoneWeights = FBoneWeights::Create(InfluenceBones, InfluenceWeights, InfluenceBoneCount);
+		FMemory::Memzero(SkinWeight.InfluenceBones);
+		FMemory::Memzero(SkinWeight.InfluenceWeights);
+	
+		if (BoneWeights.Num() == 0)
+		{
+			SkinWeight.InfluenceWeights[0] = std::numeric_limits<uint16>::max();
+			InfluenceBoneCount = 1;
+		}
+		else
+		{
+			int32 Index = 0;
+			for (FBoneWeight BoneWeight: BoneWeights)
 			{
-				break;
+				SkinWeight.InfluenceBones[Index] = BoneWeight.GetBoneIndex();
+				SkinWeight.InfluenceWeights[Index] = BoneWeight.GetRawWeight();
+				Index++;
 			}
+			InfluenceBoneCount = BoneWeights.Num(); 
 		}
+
 		//Adjust section influence count if the alternate influence bone count is greater
-		if (InfluenceBoneIndex > MaxNumInfluences)
+		if (InfluenceBoneCount > MaxNumInfluences)
 		{
-			MaxNumInfluences = InfluenceBoneIndex;
+			MaxNumInfluences = InfluenceBoneCount;
 			if (MaxNumInfluences > Section.GetMaxBoneInfluences())
 			{
 				Section.MaxBoneInfluences = MaxNumInfluences;
 			}
 		}
-		//Use the same code has the build where we modify the index 0 to have a sum of 255 for all influence per skin vertex
-		SkinWeight.InfluenceWeights[0] += 255 - TotalInfluenceWeight;
 	}
 }
 
@@ -3826,4 +3840,4 @@ bool FLODUtilities::StripLODGeometry(USkeletalMesh* SkeletalMesh, const int32 LO
 
 #undef LOCTEXT_NAMESPACE // "LODUtilities"
 
-#endif //WITH_EDITOR
\ No newline at end of file
+#endif //WITH_EDITOR
diff --git a/Engine/Source/Editor/ClothingSystemEditor/Private/ClothingAssetFactory.cpp b/Engine/Source/Editor/ClothingSystemEditor/Private/ClothingAssetFactory.cpp
--- a/Engine/Source/Editor/ClothingSystemEditor/Private/ClothingAssetFactory.cpp
+++ b/Engine/Source/Editor/ClothingSystemEditor/Private/ClothingAssetFactory.cpp
@@ -3,6 +3,7 @@
 #include "ClothingAssetFactory.h"
 
 #include "BoneIndices.h"
+#include "BoneWeights.h"
 #include "ClothLODData.h"
 #include "ClothPhysicalMeshData.h"
 #include "ClothVertBoneData.h"
@@ -440,7 +441,7 @@ bool UClothingAssetFactory::ImportToLodInternal(
 			{
 				FName BoneName = SourceMesh->GetRefSkeleton().GetBoneName(SourceIndex);
 				BoneData.BoneIndices[InfluenceIndex] = DestAsset->UsedBoneNames.AddUnique(BoneName);
-				BoneData.BoneWeights[InfluenceIndex] = (float)SourceVert.InfluenceWeights[InfluenceIndex] / 255.0f;
+				BoneData.BoneWeights[InfluenceIndex] = (float)SourceVert.InfluenceWeights[InfluenceIndex] / UE::AnimationCore::MaxRawBoneWeightFloat;
 			}
 		}
 	}
diff --git a/Engine/Source/Editor/Persona/Private/PersonaMeshDetails.cpp b/Engine/Source/Editor/Persona/Private/PersonaMeshDetails.cpp
--- a/Engine/Source/Editor/Persona/Private/PersonaMeshDetails.cpp
+++ b/Engine/Source/Editor/Persona/Private/PersonaMeshDetails.cpp
@@ -1604,6 +1604,24 @@ void FSkeletalMeshBuildSettingsLayout::GenerateChildContent(IDetailChildrenBuild
 		];
 	}
 
+	{
+		ChildrenBuilder.AddCustomRow(LOCTEXT("UseHighPrecisionSkinWeights", "Use High Precision Skin Weights"))
+		.NameContent()
+		[
+			SNew(STextBlock)
+			.Font(IDetailLayoutBuilder::GetDetailFont())
+			.Text(LOCTEXT("UseHighPrecisionSkinWeights", "Use High Precision Skin Weights"))
+			.IsEnabled(this, &FSkeletalMeshBuildSettingsLayout::IsBuildEnabled)
+		]
+		.ValueContent()
+		[
+			SNew(SCheckBox)
+			.IsChecked(this, &FSkeletalMeshBuildSettingsLayout::ShouldUseHighPrecisionSkinWeights)
+			.OnCheckStateChanged(this, &FSkeletalMeshBuildSettingsLayout::OnUseHighPrecisionSkinWeightsChanged)
+			.IsEnabled(this, &FSkeletalMeshBuildSettingsLayout::IsBuildEnabled)
+		];
+	}
+
 	{
 		ChildrenBuilder.AddCustomRow( LOCTEXT("UseFullPrecisionUVs", "Use Full Precision UVs") )
 		.NameContent()
@@ -1849,6 +1867,10 @@ ECheckBoxState FSkeletalMeshBuildSettingsLayout::ShouldUseHighPrecisionTangentBa
 {
 	return BuildSettings.bUseHighPrecisionTangentBasis ? ECheckBoxState::Checked : ECheckBoxState::Unchecked;
 }
+ECheckBoxState FSkeletalMeshBuildSettingsLayout::ShouldUseHighPrecisionSkinWeights() const
+{
+	return BuildSettings.bUseHighPrecisionSkinWeights ? ECheckBoxState::Checked : ECheckBoxState::Unchecked;
+}
 ECheckBoxState FSkeletalMeshBuildSettingsLayout::ShouldUseFullPrecisionUVs() const
 {
 	return BuildSettings.bUseFullPrecisionUVs ? ECheckBoxState::Checked : ECheckBoxState::Unchecked;
@@ -1931,6 +1953,20 @@ void FSkeletalMeshBuildSettingsLayout::OnUseHighPrecisionTangentBasisChanged(ECh
 		BuildSettings.bUseHighPrecisionTangentBasis = bUseHighPrecisionTangents;
 	}
 }
+
+void FSkeletalMeshBuildSettingsLayout::OnUseHighPrecisionSkinWeightsChanged(ECheckBoxState NewState)
+{
+	const bool bUseHighPrecisionSkinWeights = (NewState == ECheckBoxState::Checked) ? true : false;
+	if (BuildSettings.bUseHighPrecisionSkinWeights != bUseHighPrecisionSkinWeights)
+	{
+		FText TransactionText = FText::Format(LOCTEXT("PersonaChangedOnHighPrecisionSkinWeightsLOD", "LOD{0} build settings: use high precision skin weights changed"), LODIndex);
+		FScopedTransaction Transaction(TransactionText);
+		ModifyMeshLODSettingsDelegate.ExecuteIfBound(LODIndex);
+
+		BuildSettings.bUseHighPrecisionSkinWeights = bUseHighPrecisionSkinWeights;
+	}
+}
+
 void FSkeletalMeshBuildSettingsLayout::OnUseFullPrecisionUVsChanged(ECheckBoxState NewState)
 {
 	const bool bUseFullPrecisionUVs = (NewState == ECheckBoxState::Checked) ? true : false;
diff --git a/Engine/Source/Editor/Persona/Private/PersonaMeshDetails.h b/Engine/Source/Editor/Persona/Private/PersonaMeshDetails.h
--- a/Engine/Source/Editor/Persona/Private/PersonaMeshDetails.h
+++ b/Engine/Source/Editor/Persona/Private/PersonaMeshDetails.h
@@ -292,6 +292,7 @@ private:
 	ECheckBoxState ShouldComputeWeightedNormals() const;
 	ECheckBoxState ShouldRemoveDegenerates() const;
 	ECheckBoxState ShouldUseHighPrecisionTangentBasis() const;
+	ECheckBoxState ShouldUseHighPrecisionSkinWeights() const;
 	ECheckBoxState ShouldUseFullPrecisionUVs() const;
 	ECheckBoxState ShouldUseBackwardsCompatibleF16TruncUVs() const;
 
@@ -301,6 +302,7 @@ private:
 	void OnComputeWeightedNormalsChanged(ECheckBoxState NewState);
 	void OnRemoveDegeneratesChanged(ECheckBoxState NewState);
 	void OnUseHighPrecisionTangentBasisChanged(ECheckBoxState NewState);
+	void OnUseHighPrecisionSkinWeightsChanged(ECheckBoxState NewState);
 	void OnUseFullPrecisionUVsChanged(ECheckBoxState NewState);
 	void OnUseBackwardsCompatibleF16TruncUVsChanged(ECheckBoxState NewState);
 
@@ -321,6 +323,7 @@ private:
 	TArray<FSliderStateData> SliderStateDataArray;
 };
 
+
 class FPersonaMeshDetails : public IDetailCustomization
 {
 public:
diff --git a/Engine/Source/Editor/UnrealEd/Private/Fbx/FbxSkeletalMeshExport.cpp b/Engine/Source/Editor/UnrealEd/Private/Fbx/FbxSkeletalMeshExport.cpp
--- a/Engine/Source/Editor/UnrealEd/Private/Fbx/FbxSkeletalMeshExport.cpp
+++ b/Engine/Source/Editor/UnrealEd/Private/Fbx/FbxSkeletalMeshExport.cpp
@@ -5,6 +5,7 @@
 =============================================================================*/
 
 #include "CoreMinimal.h"
+#include "BoneWeights.h"
 #include "GPUSkinPublicDefs.h"
 #include "Components/SkeletalMeshComponent.h"
 #include "Animation/AnimSequence.h"
@@ -450,8 +451,8 @@ void FFbxExporter::BindMeshToSkeleton(const USkeletalMesh* SkelMesh, FbxNode* Me
 
 				for(int32 InfluenceIndex = 0; InfluenceIndex < MAX_TOTAL_INFLUENCES; ++InfluenceIndex)
 				{
-					int32 InfluenceBone		= Section.BoneMap[ Vert.InfluenceBones[InfluenceIndex] ];
-					float InfluenceWeight	= Vert.InfluenceWeights[InfluenceIndex] / 255.f;
+					const int32 InfluenceBone		= Section.BoneMap[ Vert.InfluenceBones[InfluenceIndex] ];
+					const float InfluenceWeight	= Vert.InfluenceWeights[InfluenceIndex] / UE::AnimationCore::MaxRawBoneWeightFloat;
 
 					if(InfluenceBone == BoneIndex && InfluenceWeight > 0.f)
 					{
diff --git a/Engine/Source/Runtime/AnimationCore/Private/Tests/BoneWeightsTests.cpp b/Engine/Source/Runtime/AnimationCore/Private/Tests/BoneWeightsTests.cpp
--- a/Engine/Source/Runtime/AnimationCore/Private/Tests/BoneWeightsTests.cpp
+++ b/Engine/Source/Runtime/AnimationCore/Private/Tests/BoneWeightsTests.cpp
@@ -358,7 +358,7 @@ bool FBoneWeightsTestCreate::RunTest(const FString& Parameters)
 
 	{
 		static const FBoneIndexType Bones[MaxInlineBoneWeightCount] = { 0, 1, 2, 3, 0, 0, 0, 0, 0, 0, 0, 0 };
-		static const uint8 Influences[MaxInlineBoneWeightCount] = { 255, 31, 63, 127, 0, 0, 0, 0, 0, 0, 0, 0 };
+		static const uint16 Influences[MaxInlineBoneWeightCount] = { 65535, 8191, 16383, 32767, 0, 0, 0, 0, 0, 0, 0, 0 };
 
 		{
 			const TCHAR* Name = TEXT("FBoneWeights::Create[FSoftSkinVertex, Threshold]");
diff --git a/Engine/Source/Runtime/AnimationCore/Public/BoneWeights.h b/Engine/Source/Runtime/AnimationCore/Public/BoneWeights.h
--- a/Engine/Source/Runtime/AnimationCore/Public/BoneWeights.h
+++ b/Engine/Source/Runtime/AnimationCore/Public/BoneWeights.h
@@ -31,6 +31,19 @@ The maximum number of inline bone weights.
 */
 static constexpr int32 MaxInlineBoneWeightCount = MAX_TOTAL_INFLUENCES;
 
+/** The maximum raw weight value */
+static constexpr uint16 MaxRawBoneWeight = std::numeric_limits<uint16>::max();
+
+/** The maximum raw weight value as a float value */
+static constexpr float MaxRawBoneWeightFloat = static_cast<float>(std::numeric_limits<uint16>::max());
+
+/** The inverse of the maximum raw weight value as a float value. Used for scaling. */
+static constexpr float InvMaxRawBoneWeightFloat = 1.0f / MaxRawBoneWeightFloat;
+
+/** The threshold value at or above which a 0-1 normalized bone weight value will be stored as a non-zero value after
+ *  scaling and quantizing to a 16-bit integer */
+static constexpr float BoneWeightThreshold = InvMaxRawBoneWeightFloat;
+
 class FBoneWeight
 {
 public:
@@ -389,7 +402,7 @@ public:
 	inline typename std::enable_if<!std::is_const<CT>::value, void>::type
 	SetBoneWeights(
 		const FBoneIndexType InBones[MaxInlineBoneWeightCount],
-		const uint8 InInfluences[MaxInlineBoneWeightCount],
+		const uint16 InInfluences[MaxInlineBoneWeightCount],
 		const FBoneWeightsSettings& InSettings = {});
 
 	template<typename CT = ContainerType>
@@ -602,7 +615,7 @@ public:
 	/** A helper to create a FBoneWeights container from FSoftSkinVertex data structure. */
 	static inline FBoneWeights Create(
 	    const FBoneIndexType InBones[MaxInlineBoneWeightCount],
-	    const uint8 InWeights[MaxInlineBoneWeightCount],
+	    const uint16 InWeights[MaxInlineBoneWeightCount],
 	    const FBoneWeightsSettings& InSettings = {});
 
 	/** A helper to create a FBoneWeights container from separated bone index and weight arrays.
@@ -781,7 +794,7 @@ template<typename CT>
 typename std::enable_if<!std::is_const<CT>::value, void>::type
 TBoneWeights<ContainerAdapter>::SetBoneWeights(
 	const FBoneIndexType InBones[MaxInlineBoneWeightCount],
-	const uint8 InInfluences[MaxInlineBoneWeightCount],
+	const uint16 InInfluences[MaxInlineBoneWeightCount],
 	const FBoneWeightsSettings& InSettings /*= {}*/)
 {
 	// The weights are valid until the first zero influence.
@@ -1262,7 +1275,7 @@ void FBoneWeights::Renormalize(const FBoneWeightsSettings& InSettings /*= {}*/)
 
 FBoneWeights FBoneWeights::Create(
     const FBoneIndexType InBones[MaxInlineBoneWeightCount],
-    const uint8 InInfluences[MaxInlineBoneWeightCount],
+    const uint16 InInfluences[MaxInlineBoneWeightCount],
     const FBoneWeightsSettings& InSettings /*= {} */
 )
 {
diff --git a/Engine/Source/Runtime/Core/Private/UObject/UE5MainStreamObjectVersion.cpp b/Engine/Source/Runtime/Core/Private/UObject/UE5MainStreamObjectVersion.cpp
--- a/Engine/Source/Runtime/Core/Private/UObject/UE5MainStreamObjectVersion.cpp
+++ b/Engine/Source/Runtime/Core/Private/UObject/UE5MainStreamObjectVersion.cpp
@@ -15,7 +15,7 @@ TMap<FGuid, FGuid> FUE5MainStreamObjectVersion::GetSystemGuids()
 	SystemGuids.Add(DevGuids.NANITE_DERIVEDDATA_VER, FGuid("7A3EA93546624597BEBE2ADFDE3A7628"));
 	SystemGuids.Add(DevGuids.NIAGARASHADERMAP_DERIVEDDATA_VER, FGuid("8A37C45D24F2423CBE5F8F371DE33575"));
 	SystemGuids.Add(DevGuids.Niagara_LatestScriptCompileVersion, FGuid("E5A56F8E17E7614C9040A7F8F168F0C0"));
-	SystemGuids.Add(DevGuids.SkeletalMeshDerivedDataVersion, FGuid("FB7B2FB546DE49B2A2F88B7F69763CEE"));
+	SystemGuids.Add(DevGuids.SkeletalMeshDerivedDataVersion, FGuid("4084BE94509A8C3FA7CF410DDF8ABF04"));
 	SystemGuids.Add(DevGuids.STATICMESH_DERIVEDDATA_VER, FGuid("98F7D79A5811013825E75BBCC41ED3E9"));
 	SystemGuids.Add(DevGuids.POSESEARCHDB_DERIVEDDATA_VER, FGuid("4E595C2AC5E947D6BA9ABC874353E5BC"));
 
diff --git a/Engine/Source/Runtime/Core/Public/UObject/UE5MainStreamObjectVersion.h b/Engine/Source/Runtime/Core/Public/UObject/UE5MainStreamObjectVersion.h
--- a/Engine/Source/Runtime/Core/Public/UObject/UE5MainStreamObjectVersion.h
+++ b/Engine/Source/Runtime/Core/Public/UObject/UE5MainStreamObjectVersion.h
@@ -283,6 +283,9 @@ struct CORE_API FUE5MainStreamObjectVersion
 		// Reintroducing AnimationDataModelInterface_BackedOut changes
 		ReintroduceAnimationDataModelInterface,
 
+		// Support 16-bit skin weights on SkeletalMesh
+		IncreasedSkinWeightPrecision,
+
 		// -----<new versions can be added above this line>-------------------------------------------------
 		VersionPlusOne,
 		LatestVersion = VersionPlusOne - 1
diff --git a/Engine/Source/Runtime/Engine/Classes/Engine/EngineTypes.h b/Engine/Source/Runtime/Engine/Classes/Engine/EngineTypes.h
--- a/Engine/Source/Runtime/Engine/Classes/Engine/EngineTypes.h
+++ b/Engine/Source/Runtime/Engine/Classes/Engine/EngineTypes.h
@@ -2604,6 +2604,10 @@ struct FSkeletalMeshBuildSettings
 	UPROPERTY(EditAnywhere, BlueprintReadWrite, Category = BuildSettings)
 	uint8 bUseHighPrecisionTangentBasis:1;
 
+	/** Use 16-bit precision for rendering skin weights, instead of 8-bit precision. */
+	UPROPERTY(EditAnywhere, BlueprintReadWrite, Category = BuildSettings)
+	uint8 bUseHighPrecisionSkinWeights:1;
+
 	/** If true, UVs will be stored at full floating point precision. */
 	UPROPERTY(EditAnywhere, BlueprintReadWrite, Category=BuildSettings)
 	uint8 bUseFullPrecisionUVs:1;
@@ -2636,6 +2640,7 @@ struct FSkeletalMeshBuildSettings
 		, bComputeWeightedNormals(false)
 		, bRemoveDegenerates(true)
 		, bUseHighPrecisionTangentBasis(false)
+		, bUseHighPrecisionSkinWeights(false)
 		, bUseFullPrecisionUVs(false)
 		, bUseBackwardsCompatibleF16TruncUVs(false)
 		, ThresholdPosition(0.00002)
@@ -2653,6 +2658,7 @@ struct FSkeletalMeshBuildSettings
 			&& bComputeWeightedNormals == Other.bComputeWeightedNormals
 			&& bRemoveDegenerates == Other.bRemoveDegenerates
 			&& bUseHighPrecisionTangentBasis == Other.bUseHighPrecisionTangentBasis
+			&& bUseHighPrecisionSkinWeights == Other.bUseHighPrecisionSkinWeights
 			&& bUseFullPrecisionUVs == Other.bUseFullPrecisionUVs
 			&& bUseBackwardsCompatibleF16TruncUVs == Other.bUseBackwardsCompatibleF16TruncUVs
 			&& ThresholdPosition == Other.ThresholdPosition
diff --git a/Engine/Source/Runtime/Engine/Classes/Engine/SkinnedAsset.h b/Engine/Source/Runtime/Engine/Classes/Engine/SkinnedAsset.h
--- a/Engine/Source/Runtime/Engine/Classes/Engine/SkinnedAsset.h
+++ b/Engine/Source/Runtime/Engine/Classes/Engine/SkinnedAsset.h
@@ -208,7 +208,7 @@ public:
 	virtual void SetSkinWeightProfilesData(int32 LODIndex, struct FSkinWeightProfilesData& SkinWeightProfilesData) {}
 
 	/** Computes flags for building vertex buffers. */
-	virtual uint32 GetVertexBufferFlags() const
+	virtual ESkeletalMeshVertexFlags GetVertexBufferFlags() const
 	{ return GetHasVertexColors() ? ESkeletalMeshVertexFlags::HasVertexColors : ESkeletalMeshVertexFlags::None; }
 
 	/**
diff --git a/Engine/Source/Runtime/Engine/Private/Animation/SkinWeightProfile.cpp b/Engine/Source/Runtime/Engine/Private/Animation/SkinWeightProfile.cpp
--- a/Engine/Source/Runtime/Engine/Private/Animation/SkinWeightProfile.cpp
+++ b/Engine/Source/Runtime/Engine/Private/Animation/SkinWeightProfile.cpp
@@ -188,6 +188,7 @@ FArchive& operator<<(FArchive& Ar, FImportedSkinWeightProfileData& ProfileData)
 FArchive& operator<<(FArchive& Ar, FRawSkinWeight& OverrideEntry)
 {
 	Ar.UsingCustomVersion(FAnimObjectVersion::GUID);
+	Ar.UsingCustomVersion(FUE5MainStreamObjectVersion::GUID);
 
 	if (Ar.IsLoading())
 	{
@@ -213,6 +214,16 @@ FArchive& operator<<(FArchive& Ar, FRawSkinWeight& OverrideEntry)
 			Ar << OverrideEntry.InfluenceWeights[InfluenceIndex];
 		}
 	}
+	else if (Ar.CustomVer(FUE5MainStreamObjectVersion::GUID) < FUE5MainStreamObjectVersion::IncreasedSkinWeightPrecision)
+	{
+		for (int32 InfluenceIndex = 0; InfluenceIndex < MAX_TOTAL_INFLUENCES; ++InfluenceIndex)
+		{
+			uint8 Weight = 0;
+			Ar << OverrideEntry.InfluenceBones[InfluenceIndex];
+			Ar << Weight;
+			OverrideEntry.InfluenceWeights[InfluenceIndex] = (static_cast<uint16>(Weight) << 8) | Weight;
+		}
+	}
 	else
 	{
 		for (int32 InfluenceIndex = 0; InfluenceIndex < MAX_TOTAL_INFLUENCES; ++InfluenceIndex)
diff --git a/Engine/Source/Runtime/Engine/Private/Components/SkinnedMeshComponent.cpp b/Engine/Source/Runtime/Engine/Private/Components/SkinnedMeshComponent.cpp
--- a/Engine/Source/Runtime/Engine/Private/Components/SkinnedMeshComponent.cpp
+++ b/Engine/Source/Runtime/Engine/Private/Components/SkinnedMeshComponent.cpp
@@ -28,6 +28,7 @@
 #include "Animation/MeshDeformerInstance.h"
 #include "Animation/MorphTarget.h"
 #include "AnimationRuntime.h"
+#include "BoneWeights.h"
 #include "Animation/SkinWeightProfileManager.h"
 #include "GPUSkinCache.h"
 #include "PipelineStateCache.h"
@@ -4066,7 +4067,8 @@ void CreateSectionSkinWeightsArray(
 				TargetWeight.InfluenceWeights[InfIndex] = 0;
 
 				// if we have a valid weight, see if we have a valid bone mapping for desired bone
-				uint8 InfWeight = SrcWeight.Weights[InfIndex];
+				// Map from 8-bit weights to 16-bit by aligning the min/max range.
+				const uint16 InfWeight = (SrcWeight.Weights[InfIndex] << 8) | SrcWeight.Weights[InfIndex];
 				if (InfWeight > 0)
 				{
 					const int32 SkelBoneIndex = SrcWeight.Bones[InfIndex];
@@ -4092,7 +4094,7 @@ void CreateSectionSkinWeightsArray(
 			bWeightUnderrun = true;
 
 			TargetWeight.InfluenceBones[0] = 0;
-			TargetWeight.InfluenceWeights[0] = 255;
+			TargetWeight.InfluenceWeights[0] = UE::AnimationCore::MaxRawBoneWeight;
 
 			for (int32 InfIndex = 1; InfIndex < MAX_TOTAL_INFLUENCES; InfIndex++)
 			{
diff --git a/Engine/Source/Runtime/Engine/Private/GPUSkinCache.cpp b/Engine/Source/Runtime/Engine/Private/GPUSkinCache.cpp
--- a/Engine/Source/Runtime/Engine/Private/GPUSkinCache.cpp
+++ b/Engine/Source/Runtime/Engine/Private/GPUSkinCache.cpp
@@ -472,7 +472,8 @@ public:
 	{
 		FSkinWeightVertexBuffer* WeightBuffer = GPUSkin->GetSkinWeightVertexBuffer(LOD);
 		bUse16BitBoneIndex = WeightBuffer->Use16BitBoneIndex();
-		InputWeightIndexSize = WeightBuffer->GetBoneIndexByteSize();
+		bUse16BitBoneWeight = WeightBuffer->Use16BitBoneWeight();
+		InputWeightIndexSize = WeightBuffer->GetBoneIndexByteSize() | (WeightBuffer->GetBoneWeightByteSize() << 8);
 		InputWeightStride = WeightBuffer->GetConstantInfluencesVertexStride();
 		InputWeightStreamSRV = WeightBuffer->GetDataVertexBuffer()->GetSRV();
 		InputWeightLookupStreamSRV = WeightBuffer->GetLookupVertexBuffer()->GetSRV();
@@ -616,6 +617,7 @@ protected:
 	FSkeletalMeshObjectGPUSkin* GPUSkin;
 	int BoneInfluenceType;
 	bool bUse16BitBoneIndex;
+	bool bUse16BitBoneWeight;
 	uint32 InputWeightIndexSize;
 	uint32 InputWeightStride;
 	FShaderResourceViewRHIRef InputWeightStreamSRV;
@@ -767,12 +769,14 @@ private:
 };
 
 /** Compute shader that skins a batch of vertices. */
-// @param SkinType 0:normal, 1:with morph targets calculated outside the cache, 2: with cloth, 3:with morph target calculated insde the cache (not yet implemented)
+// @param SkinType 0:normal, 1:with morph targets calculated outside the cache, 2: with cloth, 3:with morph target calculated inside the cache (not yet implemented)
 //        BoneInfluenceType 0:normal, 1:extra bone influences, 2:unlimited bone influences
 //        BoneIndex16 0: 8-bit indices, 1: 16-bit indices
+//        BoneWeights16 0: 8-bit weights, 1: 16-bit weights
 template <int Permutation>
 class TGPUSkinCacheCS : public FBaseGPUSkinCacheCS
 {
+	constexpr static bool bBoneWeights16 = (32 == (Permutation & 32));
 	constexpr static bool bBoneIndex16 = (16 == (Permutation & 16));
 	constexpr static bool bUnlimitedBoneInfluence = (8 == (Permutation & 12));
 	constexpr static bool bUseExtraBoneInfluencesT = (4 == (Permutation & 12));
@@ -790,18 +794,14 @@ public:
 	static void ModifyCompilationEnvironment(const FGlobalShaderPermutationParameters& Parameters, FShaderCompilerEnvironment& OutEnvironment)
 	{
 		FGlobalShader::ModifyCompilationEnvironment(Parameters, OutEnvironment);
-		const uint32 BoneIndex16 = bBoneIndex16;
-		const uint32 UnlimitedBoneInfluence = bUnlimitedBoneInfluence;
-		const uint32 UseExtraBoneInfluences = bUseExtraBoneInfluencesT;
-		const uint32 MorphBlend = bMorphBlend;
-		const uint32 ApexCloth = bApexCloth;
-		OutEnvironment.SetDefine(TEXT("GPUSKIN_UNLIMITED_BONE_INFLUENCE"), UnlimitedBoneInfluence);
-		OutEnvironment.SetDefine(TEXT("GPUSKIN_USE_EXTRA_INFLUENCES"), UseExtraBoneInfluences);
-		OutEnvironment.SetDefine(TEXT("GPUSKIN_MORPH_BLEND"), MorphBlend);
-		OutEnvironment.SetDefine(TEXT("GPUSKIN_APEX_CLOTH"), ApexCloth);
+		OutEnvironment.SetDefine(TEXT("GPUSKIN_UNLIMITED_BONE_INFLUENCE"), bUnlimitedBoneInfluence);
+		OutEnvironment.SetDefine(TEXT("GPUSKIN_USE_EXTRA_INFLUENCES"), bUseExtraBoneInfluencesT);
+		OutEnvironment.SetDefine(TEXT("GPUSKIN_MORPH_BLEND"), bMorphBlend);
+		OutEnvironment.SetDefine(TEXT("GPUSKIN_APEX_CLOTH"), bApexCloth);
 		OutEnvironment.SetDefine(TEXT("GPUSKIN_RWBUFFER_OFFSET_TANGENT_X"), FGPUSkinCache::RWTangentXOffsetInFloats);
 		OutEnvironment.SetDefine(TEXT("GPUSKIN_RWBUFFER_OFFSET_TANGENT_Z"), FGPUSkinCache::RWTangentZOffsetInFloats);
-		OutEnvironment.SetDefine(TEXT("GPUSKIN_BONE_INDEX_UINT16"), BoneIndex16);
+		OutEnvironment.SetDefine(TEXT("GPUSKIN_BONE_INDEX_UINT16"), bBoneIndex16);
+		OutEnvironment.SetDefine(TEXT("GPUSKIN_BONE_WEIGHTS_UINT16"), bBoneWeights16);
 	}
 
 	TGPUSkinCacheCS(const ShaderMetaType::CompiledShaderInitializerType& Initializer)
@@ -814,21 +814,29 @@ public:
 	}
 };
 
-IMPLEMENT_SHADER_TYPE(template<>, TGPUSkinCacheCS<0>, TEXT("/Engine/Private/GpuSkinCacheComputeShader.usf"), TEXT("SkinCacheUpdateBatchCS"), SF_Compute);	// 16bit_0, BoneInfluenceType_0, SkinType_0 
-IMPLEMENT_SHADER_TYPE(template<>, TGPUSkinCacheCS<1>, TEXT("/Engine/Private/GpuSkinCacheComputeShader.usf"), TEXT("SkinCacheUpdateBatchCS"), SF_Compute);	// 16bit_0, BoneInfluenceType_0, SkinType_1 
-IMPLEMENT_SHADER_TYPE(template<>, TGPUSkinCacheCS<2>, TEXT("/Engine/Private/GpuSkinCacheComputeShader.usf"), TEXT("SkinCacheUpdateBatchCS"), SF_Compute);	// 16bit_0, BoneInfluenceType_0, SkinType_2 
-IMPLEMENT_SHADER_TYPE(template<>, TGPUSkinCacheCS<4>, TEXT("/Engine/Private/GpuSkinCacheComputeShader.usf"), TEXT("SkinCacheUpdateBatchCS"), SF_Compute);	// 16bit_0, BoneInfluenceType_1, SkinType_0 
-IMPLEMENT_SHADER_TYPE(template<>, TGPUSkinCacheCS<5>, TEXT("/Engine/Private/GpuSkinCacheComputeShader.usf"), TEXT("SkinCacheUpdateBatchCS"), SF_Compute);	// 16bit_0, BoneInfluenceType_1, SkinType_1 
-IMPLEMENT_SHADER_TYPE(template<>, TGPUSkinCacheCS<6>, TEXT("/Engine/Private/GpuSkinCacheComputeShader.usf"), TEXT("SkinCacheUpdateBatchCS"), SF_Compute);	// 16bit_0, BoneInfluenceType_1, SkinType_2 
-IMPLEMENT_SHADER_TYPE(template<>, TGPUSkinCacheCS<8>, TEXT("/Engine/Private/GpuSkinCacheComputeShader.usf"), TEXT("SkinCacheUpdateBatchCS"), SF_Compute);	// 16bit_0, BoneInfluenceType_2, SkinType_0 
-IMPLEMENT_SHADER_TYPE(template<>, TGPUSkinCacheCS<9>, TEXT("/Engine/Private/GpuSkinCacheComputeShader.usf"), TEXT("SkinCacheUpdateBatchCS"), SF_Compute);	// 16bit_0, BoneInfluenceType_2, SkinType_1 
-IMPLEMENT_SHADER_TYPE(template<>, TGPUSkinCacheCS<10>, TEXT("/Engine/Private/GpuSkinCacheComputeShader.usf"), TEXT("SkinCacheUpdateBatchCS"), SF_Compute);	// 16bit_0, BoneInfluenceType_2, SkinType_2 
-IMPLEMENT_SHADER_TYPE(template<>, TGPUSkinCacheCS<16>, TEXT("/Engine/Private/GpuSkinCacheComputeShader.usf"), TEXT("SkinCacheUpdateBatchCS"), SF_Compute);  // 16bit_1, BoneInfluenceType_0, SkinType_0 
-IMPLEMENT_SHADER_TYPE(template<>, TGPUSkinCacheCS<17>, TEXT("/Engine/Private/GpuSkinCacheComputeShader.usf"), TEXT("SkinCacheUpdateBatchCS"), SF_Compute);  // 16bit_1, BoneInfluenceType_0, SkinType_1 
-IMPLEMENT_SHADER_TYPE(template<>, TGPUSkinCacheCS<18>, TEXT("/Engine/Private/GpuSkinCacheComputeShader.usf"), TEXT("SkinCacheUpdateBatchCS"), SF_Compute);  // 16bit_1, BoneInfluenceType_0, SkinType_2 
-IMPLEMENT_SHADER_TYPE(template<>, TGPUSkinCacheCS<20>, TEXT("/Engine/Private/GpuSkinCacheComputeShader.usf"), TEXT("SkinCacheUpdateBatchCS"), SF_Compute);  // 16bit_1, BoneInfluenceType_1, SkinType_0 
-IMPLEMENT_SHADER_TYPE(template<>, TGPUSkinCacheCS<21>, TEXT("/Engine/Private/GpuSkinCacheComputeShader.usf"), TEXT("SkinCacheUpdateBatchCS"), SF_Compute);  // 16bit_1, BoneInfluenceType_1, SkinType_1 
-IMPLEMENT_SHADER_TYPE(template<>, TGPUSkinCacheCS<22>, TEXT("/Engine/Private/GpuSkinCacheComputeShader.usf"), TEXT("SkinCacheUpdateBatchCS"), SF_Compute);  // 16bit_1, BoneInfluenceType_1, SkinType_2 
+#define SKIN_CACHE_SHADER(_WEIGHT16_, _INDEX16_, _INFLUENCE_TYPE_, _SKIN_TYPE_) \
+	IMPLEMENT_SHADER_TYPE(template<>, TGPUSkinCacheCS<_WEIGHT16_ * 32 + _INDEX16_ * 16 + _INFLUENCE_TYPE_ * 4 + _SKIN_TYPE_>, TEXT("/Engine/Private/GpuSkinCacheComputeShader.usf"), TEXT("SkinCacheUpdateBatchCS"), SF_Compute) 
+
+#define SKIN_CACHE_SHADER_ALL_SKIN_TYPES(_WEIGHT16_, _INDEX16_, _INFLUENCE_TYPE_) \
+	SKIN_CACHE_SHADER(_WEIGHT16_, _INDEX16_, _INFLUENCE_TYPE_, 0); \
+	SKIN_CACHE_SHADER(_WEIGHT16_, _INDEX16_, _INFLUENCE_TYPE_, 1); \
+	SKIN_CACHE_SHADER(_WEIGHT16_, _INDEX16_, _INFLUENCE_TYPE_, 2)
+
+// NOTE: Bone influence type 2 (multiple) does not require a 16-bit index or weight permutation.
+#define SKIN_CACHE_SHADER_ALL() \
+	SKIN_CACHE_SHADER_ALL_SKIN_TYPES(0, 0, 0); \
+	SKIN_CACHE_SHADER_ALL_SKIN_TYPES(0, 0, 1); \
+	SKIN_CACHE_SHADER_ALL_SKIN_TYPES(0, 1, 0); \
+	SKIN_CACHE_SHADER_ALL_SKIN_TYPES(0, 1, 1); \
+	SKIN_CACHE_SHADER_ALL_SKIN_TYPES(1, 0, 0); \
+	SKIN_CACHE_SHADER_ALL_SKIN_TYPES(1, 0, 1); \
+	SKIN_CACHE_SHADER_ALL_SKIN_TYPES(1, 1, 0); \
+	SKIN_CACHE_SHADER_ALL_SKIN_TYPES(1, 1, 1); \
+	SKIN_CACHE_SHADER_ALL_SKIN_TYPES(0, 0, 2);
+
+SKIN_CACHE_SHADER_ALL()
+
+#undef SKIN_CACHE_SHADER
 
 FGPUSkinCache::FGPUSkinCache(ERHIFeatureLevel::Type InFeatureLevel, bool bInRequiresMemoryLimit, UWorld* InWorld)
 	: UsedMemoryInBytes(0)
@@ -1853,81 +1861,40 @@ void FGPUSkinCache::DispatchUpdateSkinning(FRHICommandListImmediate& RHICmdList,
 	const FString RayTracingTag = (Entry->Mode == EGPUSkinCacheEntryMode::RayTracing ? TEXT("[RT]") : TEXT(""));
 	
 	SCOPED_DRAW_EVENTF(RHICmdList, SkinCacheDispatch,
-		TEXT("%sSkinning%d%d%d Mesh=%s LOD=%d Chunk=%d InStreamStart=%d OutStart=%d Vert=%d Morph=%d/%d"),
-		*RayTracingTag, (int32)Entry->bUse16BitBoneIndex, (int32)Entry->BoneInfluenceType, DispatchData.SkinType, *GetSkeletalMeshObjectName(Entry->GPUSkin), Entry->LOD,
+		TEXT("%sSkinning%d%d%d%d Mesh=%s LOD=%d Chunk=%d InStreamStart=%d OutStart=%d Vert=%d Morph=%d/%d"),
+		*RayTracingTag, (int32)Entry->bUse16BitBoneIndex, (int32)Entry->bUse16BitBoneWeight, (int32)Entry->BoneInfluenceType, DispatchData.SkinType, *GetSkeletalMeshObjectName(Entry->GPUSkin), Entry->LOD,
 		DispatchData.SectionIndex, DispatchData.InputStreamStart, DispatchData.OutputStreamStart, DispatchData.NumVertices, Entry->MorphBuffer != 0, DispatchData.MorphBufferOffset);
 	auto* GlobalShaderMap = GetGlobalShaderMap(GetFeatureLevel());
-	TShaderMapRef<TGPUSkinCacheCS<0>> SkinCacheCS000(GlobalShaderMap);		// 16bit_0, BoneInfluenceType_0, SkinType_0
-	TShaderMapRef<TGPUSkinCacheCS<1>> SkinCacheCS001(GlobalShaderMap);		// 16bit_0, BoneInfluenceType_0, SkinType_1
-	TShaderMapRef<TGPUSkinCacheCS<2>> SkinCacheCS002(GlobalShaderMap);		// 16bit_0, BoneInfluenceType_0, SkinType_2
-	TShaderMapRef<TGPUSkinCacheCS<4>> SkinCacheCS010(GlobalShaderMap);		// 16bit_0, BoneInfluenceType_1, SkinType_0
-	TShaderMapRef<TGPUSkinCacheCS<5>> SkinCacheCS011(GlobalShaderMap);		// 16bit_0, BoneInfluenceType_1, SkinType_1
-	TShaderMapRef<TGPUSkinCacheCS<6>> SkinCacheCS012(GlobalShaderMap);		// 16bit_0, BoneInfluenceType_1, SkinType_2
-	TShaderMapRef<TGPUSkinCacheCS<8>> SkinCacheCS020(GlobalShaderMap);		// 16bit_0, BoneInfluenceType_2, SkinType_0
-	TShaderMapRef<TGPUSkinCacheCS<9>> SkinCacheCS021(GlobalShaderMap);		// 16bit_0, BoneInfluenceType_2, SkinType_1
-	TShaderMapRef<TGPUSkinCacheCS<10>> SkinCacheCS022(GlobalShaderMap);		// 16bit_0, BoneInfluenceType_2, SkinType_2
-	TShaderMapRef<TGPUSkinCacheCS<16>>  SkinCacheCS100(GlobalShaderMap);	// 16bit_1, BoneInfluenceType_0, SkinType_0
-	TShaderMapRef<TGPUSkinCacheCS<17>>  SkinCacheCS101(GlobalShaderMap);	// 16bit_1, BoneInfluenceType_0, SkinType_1
-	TShaderMapRef<TGPUSkinCacheCS<18>>  SkinCacheCS102(GlobalShaderMap);	// 16bit_1, BoneInfluenceType_0, SkinType_2
-	TShaderMapRef<TGPUSkinCacheCS<20>>  SkinCacheCS110(GlobalShaderMap);	// 16bit_1, BoneInfluenceType_1, SkinType_0
-	TShaderMapRef<TGPUSkinCacheCS<21>>  SkinCacheCS111(GlobalShaderMap);	// 16bit_1, BoneInfluenceType_1, SkinType_1
-	TShaderMapRef<TGPUSkinCacheCS<22>>  SkinCacheCS112(GlobalShaderMap);	// 16bit_1, BoneInfluenceType_1, SkinType_2
-
-	TShaderRef<FBaseGPUSkinCacheCS> Shader;
-	switch (DispatchData.SkinType)
-	{
-	case 0:
-		if (Entry->BoneInfluenceType == 0)
-		{
-			if (Entry->bUse16BitBoneIndex) Shader = SkinCacheCS100;
-			else Shader = SkinCacheCS000;
-		}
-		else if (Entry->BoneInfluenceType == 1)
-		{
-			if (Entry->bUse16BitBoneIndex) Shader = SkinCacheCS110;
-			else Shader = SkinCacheCS010;
-		}
-		else
-		{
-			Shader = SkinCacheCS020;
-		}
-		break;
-	case 1:
-		if (Entry->BoneInfluenceType == 0)
-		{
-			if (Entry->bUse16BitBoneIndex) Shader = SkinCacheCS101;
-			else Shader = SkinCacheCS001;
-		}
-		else if (Entry->BoneInfluenceType == 1)
-		{
-			if (Entry->bUse16BitBoneIndex) Shader = SkinCacheCS111;
-			else Shader = SkinCacheCS011;
-		}
-		else
-		{
-			Shader = SkinCacheCS021;
-		}
-		break;
-	case 2:
-		// Cloth
-		if (Entry->BoneInfluenceType == 0)
-		{
-			if (Entry->bUse16BitBoneIndex) Shader = SkinCacheCS102;
-			else Shader = SkinCacheCS002;
-		}
-		else if (Entry->BoneInfluenceType == 1)
-		{
-			if (Entry->bUse16BitBoneIndex) Shader = SkinCacheCS112;
-			else Shader = SkinCacheCS012;
-		}
-		else
-		{
-			Shader = SkinCacheCS022;
-		}
-		break;
-	default:
-		check(0);
+
+	TMap<int32, TShaderRef<FBaseGPUSkinCacheCS>> SkinCacheCSMap;
+
+#define SKIN_CACHE_SHADER(_WEIGHT16_, _INDEX16_, _INFLUENCE_TYPE_, _SKIN_TYPE_) \
+	{ \
+		constexpr int32 Index = _WEIGHT16_ * 32 + _INDEX16_ * 16 + _INFLUENCE_TYPE_ * 4 + _SKIN_TYPE_; \
+		SkinCacheCSMap.Add(Index, TShaderMapRef<TGPUSkinCacheCS<Index>>(GlobalShaderMap)); \
 	}
+
+	SKIN_CACHE_SHADER_ALL()
+
+	// For 'unlimited' bone indexes, we pass in the index and weight sizes via a shader parameter and so we
+	// can re-use the same shader permutation as for 8-bit indexes.
+	bool bUse16BitBoneIndex = Entry->bUse16BitBoneIndex;
+	bool bUse16BitBoneWeight = Entry->bUse16BitBoneWeight;
+	if (Entry->BoneInfluenceType == 2)
+	{
+		bUse16BitBoneIndex = bUse16BitBoneWeight = false;
+	}
+
+#undef SKIN_CACHE_SHADER
+	int32 ShaderIndex =
+		DispatchData.SkinType |
+		(Entry->BoneInfluenceType * 4) |
+		(static_cast<int32>(bUse16BitBoneIndex) * 16) |
+		(static_cast<int32>(bUse16BitBoneWeight) * 32) ;
+
+	check(SkinCacheCSMap.Contains(ShaderIndex));
+	
+	TShaderRef<FBaseGPUSkinCacheCS> Shader = SkinCacheCSMap[ShaderIndex];
 	check(Shader.IsValid());
 
 	const FVertexBufferAndSRV& BoneBuffer = ShaderData.GetBoneBufferForReading(false);
@@ -2333,4 +2300,8 @@ void FGPUSkinCache::DrawVisualizationInfoText(const FName& GPUSkinCacheVisualiza
 	#endif
 		}
 	}
-}
\ No newline at end of file
+}
+
+#undef IMPLEMENT_SKIN_CACHE_SHADER_CLOTH
+#undef IMPLEMENT_SKIN_CACHE_SHADER_ALL_SKIN_TYPES
+#undef IMPLEMENT_SKIN_CACHE_SHADER
diff --git a/Engine/Source/Runtime/Engine/Private/Rendering/SkinWeightVertexBuffer.cpp b/Engine/Source/Runtime/Engine/Private/Rendering/SkinWeightVertexBuffer.cpp
--- a/Engine/Source/Runtime/Engine/Private/Rendering/SkinWeightVertexBuffer.cpp
+++ b/Engine/Source/Runtime/Engine/Private/Rendering/SkinWeightVertexBuffer.cpp
@@ -191,10 +191,11 @@ FSkinWeightDataVertexBuffer::FSkinWeightDataVertexBuffer()
 ,	bVariableBonesPerVertex(false)
 ,	MaxBoneInfluences(MAX_INFLUENCES_PER_STREAM)
 ,	bUse16BitBoneIndex(false)
+,	bUse16BitBoneWeight(false)
 ,	WeightData(nullptr)
 ,	Data(nullptr)
 ,	NumVertices(0)
-,	NumBones(0)
+,	NumBoneWeights(0)
 {
 }
 
@@ -203,10 +204,11 @@ FSkinWeightDataVertexBuffer::FSkinWeightDataVertexBuffer( const FSkinWeightDataV
 	, bVariableBonesPerVertex(Other.bVariableBonesPerVertex)
 	, MaxBoneInfluences(Other.MaxBoneInfluences)
 	, bUse16BitBoneIndex(Other.bUse16BitBoneIndex)
+	, bUse16BitBoneWeight(Other.bUse16BitBoneWeight)
 	, WeightData(nullptr)
 	, Data(nullptr)
 	, NumVertices(0)
-	, NumBones(0)
+	, NumBoneWeights(0)
 {
 	
 }
@@ -223,6 +225,7 @@ FSkinWeightDataVertexBuffer& FSkinWeightDataVertexBuffer::operator=(const FSkinW
 	bVariableBonesPerVertex = Other.bVariableBonesPerVertex;
 	MaxBoneInfluences = Other.MaxBoneInfluences;
 	bUse16BitBoneIndex = Other.bUse16BitBoneIndex;
+	bUse16BitBoneWeight = Other.bUse16BitBoneWeight;
 	return *this;
 }
 
@@ -231,23 +234,23 @@ void FSkinWeightDataVertexBuffer::CleanUp()
 	if (WeightData)
 	{
 		delete WeightData;
-		WeightData = NULL;
+		WeightData = nullptr;
 	}
 }
 
-void FSkinWeightDataVertexBuffer::Init(uint32 InNumBones, uint32 InNumVertices)
+void FSkinWeightDataVertexBuffer::Init(uint32 InNumWeights, uint32 InNumVertices)
 {
 	AllocateData();
-	ResizeBuffer(InNumBones, InNumVertices);
+	ResizeBuffer(InNumWeights, InNumVertices);
 }
 
-void FSkinWeightDataVertexBuffer::ResizeBuffer(uint32 InNumBones, uint32 InNumVertices)
+void FSkinWeightDataVertexBuffer::ResizeBuffer(uint32 InNumWeights, uint32 InNumVertices)
 {
-	NumBones = InNumBones;
+	NumBoneWeights = InNumWeights;
 	NumVertices = InNumVertices;
 	WeightData->ResizeBuffer(GetVertexDataSize());
 
-	if (NumBones > 0)
+	if (NumBoneWeights > 0)
 	{
 		Data = WeightData->GetDataPointer();
 	}
@@ -260,7 +263,7 @@ FArchive& operator<<(FArchive& Ar, FSkinWeightDataVertexBuffer& VertexBuffer)
 	Ar.UsingCustomVersion(FAnimObjectVersion::GUID);
 	VertexBuffer.SerializeMetaData(Ar);
 
-	if (Ar.IsLoading() || VertexBuffer.WeightData == NULL)
+	if (Ar.IsLoading() || VertexBuffer.WeightData == nullptr)
 	{
 		// If we're loading, or we have no valid buffer, allocate container.
 		VertexBuffer.AllocateData();
@@ -319,7 +322,7 @@ FArchive& operator<<(FArchive& Ar, FSkinWeightDataVertexBuffer& VertexBuffer)
 			if (!Ar.IsCountingMemory())
 			{
 				// update cached buffer info
-				VertexBuffer.Data = (VertexBuffer.NumBones > 0 && VertexBuffer.WeightData->GetResourceArray()->GetResourceDataSize()) ? VertexBuffer.WeightData->GetDataPointer() : nullptr;
+				VertexBuffer.Data = (VertexBuffer.NumBoneWeights > 0 && VertexBuffer.WeightData->GetResourceArray()->GetResourceDataSize()) ? VertexBuffer.WeightData->GetDataPointer() : nullptr;
 			}
 		}
 	}
@@ -331,6 +334,7 @@ void FSkinWeightDataVertexBuffer::SerializeMetaData(FArchive& Ar)
 {
 	Ar.UsingCustomVersion(FSkeletalMeshCustomVersion::GUID);
 	Ar.UsingCustomVersion(FAnimObjectVersion::GUID);
+	Ar.UsingCustomVersion(FUE5MainStreamObjectVersion::GUID);
 
 	if (Ar.IsLoading() && Ar.CustomVer(FAnimObjectVersion::GUID) < FAnimObjectVersion::UnlimitedBoneInfluences)
 	{
@@ -345,12 +349,12 @@ void FSkinWeightDataVertexBuffer::SerializeMetaData(FArchive& Ar)
 			Ar << bExtraBoneInfluences << Stride << NumVertices;
 		}
 		MaxBoneInfluences = bExtraBoneInfluences ? EXTRA_BONE_INFLUENCES : MAX_INFLUENCES_PER_STREAM;
-		NumBones = MaxBoneInfluences * NumVertices;
+		NumBoneWeights = MaxBoneInfluences * NumVertices;
 		bVariableBonesPerVertex = false;
 	}
 	else
 	{
-		Ar << bVariableBonesPerVertex << MaxBoneInfluences << NumBones << NumVertices;
+		Ar << bVariableBonesPerVertex << MaxBoneInfluences << NumBoneWeights << NumVertices;
 	}
 	
 	// bUse16BitBoneIndex doesn't exist before version IncreaseBoneIndexLimitPerChunk
@@ -358,6 +362,10 @@ void FSkinWeightDataVertexBuffer::SerializeMetaData(FArchive& Ar)
 	{
 		Ar << bUse16BitBoneIndex;
 	}
+	if (Ar.CustomVer(FUE5MainStreamObjectVersion::GUID) >= FUE5MainStreamObjectVersion::IncreasedSkinWeightPrecision)
+	{
+		Ar << bUse16BitBoneWeight;
+	}
 }
 
 void FSkinWeightDataVertexBuffer::CopyMetaData(const FSkinWeightDataVertexBuffer& Other)
@@ -365,14 +373,15 @@ void FSkinWeightDataVertexBuffer::CopyMetaData(const FSkinWeightDataVertexBuffer
 	bVariableBonesPerVertex = Other.bVariableBonesPerVertex;
 	MaxBoneInfluences = Other.MaxBoneInfluences;
 	bUse16BitBoneIndex = Other.bUse16BitBoneIndex;
-	NumBones = Other.NumBones;
+	bUse16BitBoneWeight = Other.bUse16BitBoneWeight;
+	NumBoneWeights = Other.NumBoneWeights;
 }
 
 template <bool bRenderThread>
 FBufferRHIRef FSkinWeightDataVertexBuffer::CreateRHIBuffer_Internal()
 {
 	// BUF_ShaderResource is needed for support of the SkinCache (we could make is dependent on GEnableGPUSkinCacheShaders or are there other users?)
-	return CreateRHIBuffer<bRenderThread>(WeightData, NumBones, BUF_Static | BUF_ShaderResource | BUF_SourceCopy, TEXT("FSkinWeightDataVertexBuffer"));
+	return CreateRHIBuffer<bRenderThread>(WeightData, NumBoneWeights, BUF_Static | BUF_ShaderResource | BUF_SourceCopy, TEXT("FSkinWeightDataVertexBuffer"));
 }
 
 FBufferRHIRef FSkinWeightDataVertexBuffer::CreateRHIBuffer_RenderThread()
@@ -513,13 +522,22 @@ void FSkinWeightDataVertexBuffer::SetBoneIndex(uint32 VertexWeightOffset, uint32
 	}
 }
 
-uint8 FSkinWeightDataVertexBuffer::GetBoneWeight(uint32 VertexWeightOffset, uint32 VertexInfluenceCount, uint32 InfluenceIndex) const
+uint16 FSkinWeightDataVertexBuffer::GetBoneWeight(uint32 VertexWeightOffset, uint32 VertexInfluenceCount, uint32 InfluenceIndex) const
 {
 	if (InfluenceIndex < VertexInfluenceCount)
 	{
-		uint8* BoneData = Data + VertexWeightOffset;
-		uint32 BoneWeightOffset = GetBoneIndexByteSize() * VertexInfluenceCount;
-		return BoneData[BoneWeightOffset + InfluenceIndex];
+		const uint32 BoneWeightOffset = GetBoneIndexByteSize() * VertexInfluenceCount;
+		const uint8* BoneData = Data + VertexWeightOffset + BoneWeightOffset;
+		const uint8* BoneVertexData = &BoneData[InfluenceIndex * GetBoneWeightByteSize()];  
+
+		if (Use16BitBoneWeight())
+		{
+			return *reinterpret_cast<const uint16*>(BoneVertexData);
+		}
+		else
+		{
+			return (static_cast<uint16>(*BoneVertexData) << 8) | *BoneVertexData;
+		}
 	}
 	else
 	{
@@ -527,13 +545,23 @@ uint8 FSkinWeightDataVertexBuffer::GetBoneWeight(uint32 VertexWeightOffset, uint
 	}
 }
 
-void FSkinWeightDataVertexBuffer::SetBoneWeight(uint32 VertexWeightOffset, uint32 VertexInfluenceCount, uint32 InfluenceIndex, uint8 BoneWeight)
+void FSkinWeightDataVertexBuffer::SetBoneWeight(uint32 VertexWeightOffset, uint32 VertexInfluenceCount, uint32 InfluenceIndex, uint16 BoneWeight)
 {
 	if (InfluenceIndex < VertexInfluenceCount)
 	{
-		uint8* BoneData = Data + VertexWeightOffset;
-		uint32 BoneWeightOffset = GetBoneIndexByteSize() * VertexInfluenceCount;
-		BoneData[BoneWeightOffset + InfluenceIndex] = BoneWeight;
+		const uint32 BoneWeightOffset = GetBoneIndexByteSize() * VertexInfluenceCount;
+		uint8* BoneData = Data + VertexWeightOffset + BoneWeightOffset;
+		uint8* BoneVertexData = &BoneData[InfluenceIndex * GetBoneWeightByteSize()];  
+		
+		if (Use16BitBoneWeight())
+		{
+			*reinterpret_cast<uint16*>(BoneVertexData) = BoneWeight;
+		}
+		else
+		{
+			*BoneVertexData = static_cast<uint8>(BoneWeight >> 8);
+		}
+		
 	}
 	else
 	{
@@ -545,18 +573,7 @@ void FSkinWeightDataVertexBuffer::ResetVertexBoneWeights(uint32 VertexWeightOffs
 {
 	if (VertexInfluenceCount > 0)
 	{
-		uint8* BoneData = Data + VertexWeightOffset;
-		if (BoneData)
-		{
-			if (Use16BitBoneIndex())
-			{
-				FMemory::Memzero(BoneData, (sizeof(FBoneIndex16) + sizeof(uint8)) * VertexInfluenceCount);
-			}
-			else
-			{
-				FMemory::Memzero(BoneData, (sizeof(FBoneIndex8) + sizeof(uint8)) * VertexInfluenceCount);
-			}
-		}
+		FMemory::Memzero(Data + VertexWeightOffset, GetBoneIndexAndWeightByteSize() * VertexInfluenceCount);
 	}
 }
 
@@ -582,6 +599,7 @@ FSkinWeightVertexBuffer::FSkinWeightVertexBuffer(const FSkinWeightVertexBuffer&
 
 	SetMaxBoneInfluences(Other.GetMaxBoneInfluences());
 	SetUse16BitBoneIndex(Other.Use16BitBoneIndex());
+	SetUse16BitBoneWeight(Other.Use16BitBoneWeight());
 }
 
 FSkinWeightVertexBuffer::~FSkinWeightVertexBuffer()
@@ -599,6 +617,7 @@ FSkinWeightVertexBuffer& FSkinWeightVertexBuffer::operator=(const FSkinWeightVer
 
 	SetMaxBoneInfluences(Other.GetMaxBoneInfluences());
 	SetUse16BitBoneIndex(Other.Use16BitBoneIndex());
+	SetUse16BitBoneWeight(Other.Use16BitBoneWeight());
 
 	return *this;
 }
@@ -618,13 +637,13 @@ FArchive& operator<<(FArchive& Ar, FSkinWeightVertexBuffer& VertexBuffer)
 	{
 		// LookupVertexBuffer doesn't exist before this version, so construct its content from scratch
 		uint32 MaxBoneInfluences = VertexBuffer.DataVertexBuffer.GetMaxBoneInfluences();
-		uint32 NumVertices = VertexBuffer.DataVertexBuffer.GetNumBones() / VertexBuffer.DataVertexBuffer.GetMaxBoneInfluences();
+		uint32 NumVertices = VertexBuffer.DataVertexBuffer.GetNumBoneWeights() / VertexBuffer.DataVertexBuffer.GetMaxBoneInfluences();
 		VertexBuffer.LookupVertexBuffer.Init(NumVertices);
 
 		uint32 WeightOffset = 0;
 		for (uint32 VertIdx = 0; VertIdx < NumVertices; VertIdx++)
 		{
-			VertexBuffer.LookupVertexBuffer.SetWeightOffsetAndInfluenceCount(VertIdx, WeightOffset * (VertexBuffer.GetBoneIndexByteSize() + sizeof(uint8)), MaxBoneInfluences);
+			VertexBuffer.LookupVertexBuffer.SetWeightOffsetAndInfluenceCount(VertIdx, WeightOffset * VertexBuffer.GetBoneIndexAndWeightByteSize(), MaxBoneInfluences);
 			WeightOffset += MaxBoneInfluences;
 		}
 	}
@@ -713,7 +732,7 @@ void FSkinWeightVertexBuffer::SetBoneIndex(uint32 VertexIndex, uint32 InfluenceI
 	DataVertexBuffer.SetBoneIndex(VertexWeightOffset, VertexInfluenceCount, InfluenceIndex, BoneIndex);
 }
 
-uint8 FSkinWeightVertexBuffer::GetBoneWeight(uint32 VertexIndex, uint32 InfluenceIndex) const
+uint16 FSkinWeightVertexBuffer::GetBoneWeight(uint32 VertexIndex, uint32 InfluenceIndex) const
 {
 	uint32 VertexWeightOffset = 0;
 	uint32 VertexInfluenceCount = 0;
@@ -721,7 +740,7 @@ uint8 FSkinWeightVertexBuffer::GetBoneWeight(uint32 VertexIndex, uint32 Influenc
 	return DataVertexBuffer.GetBoneWeight(VertexWeightOffset, VertexInfluenceCount, InfluenceIndex);
 }
 
-void FSkinWeightVertexBuffer::SetBoneWeight(uint32 VertexIndex, uint32 InfluenceIndex, uint8 BoneWeight)
+void FSkinWeightVertexBuffer::SetBoneWeight(uint32 VertexIndex, uint32 InfluenceIndex, uint16 BoneWeight)
 {
 	uint32 VertexWeightOffset = 0;
 	uint32 VertexInfluenceCount = 0;
@@ -759,6 +778,9 @@ void FSkinWeightVertexBuffer::ReleaseResources()
 #if WITH_EDITOR
 void FSkinWeightVertexBuffer::Init(const TArray<FSoftSkinVertex>& InVertices)
 {
+	static_assert(sizeof(FSoftSkinVertex::InfluenceBones) == sizeof(FSkinWeightInfo::InfluenceBones));
+	static_assert(sizeof(FSoftSkinVertex::InfluenceWeights) == sizeof(FSkinWeightInfo::InfluenceWeights));
+	
 	TArray<FSkinWeightInfo> VertexWeightInfos;
 	VertexWeightInfos.AddUninitialized(InVertices.Num());
 
@@ -767,7 +789,7 @@ void FSkinWeightVertexBuffer::Init(const TArray<FSoftSkinVertex>& InVertices)
 		const FSoftSkinVertex& SrcVertex = InVertices[VertIdx];
 		FSkinWeightInfo& DstVertex = VertexWeightInfos[VertIdx];
 		FMemory::Memcpy(DstVertex.InfluenceBones, SrcVertex.InfluenceBones, MAX_TOTAL_INFLUENCES * sizeof(FBoneIndexType));
-		FMemory::Memcpy(DstVertex.InfluenceWeights, SrcVertex.InfluenceWeights, MAX_TOTAL_INFLUENCES * sizeof(uint8));
+		FMemory::Memcpy(DstVertex.InfluenceWeights, SrcVertex.InfluenceWeights, MAX_TOTAL_INFLUENCES * sizeof(uint16));
 
 	}
 
@@ -794,7 +816,7 @@ FSkinWeightVertexBuffer& FSkinWeightVertexBuffer::operator=(const TArray<FSkinWe
 				}
 			}
 
-			LookupVertexBuffer.SetWeightOffsetAndInfluenceCount(VertIdx, TotalNumUsedBones * (GetBoneIndexByteSize() + sizeof(uint8)), NumUsedBones);
+			LookupVertexBuffer.SetWeightOffsetAndInfluenceCount(VertIdx, TotalNumUsedBones * GetBoneIndexAndWeightByteSize(), NumUsedBones);
 			TotalNumUsedBones += NumUsedBones;
 		}
 	}
diff --git a/Engine/Source/Runtime/Engine/Private/SkeletalMeshLODModel.cpp b/Engine/Source/Runtime/Engine/Private/SkeletalMeshLODModel.cpp
--- a/Engine/Source/Runtime/Engine/Private/SkeletalMeshLODModel.cpp
+++ b/Engine/Source/Runtime/Engine/Private/SkeletalMeshLODModel.cpp
@@ -122,24 +122,30 @@ FArchive& operator<<(FArchive& Ar, FSoftSkinVertex& V)
 		}
 	}
 
-	for (uint32 InfluenceIndex = 0; InfluenceIndex < MAX_INFLUENCES_PER_STREAM; InfluenceIndex++)
-	{
-		Ar << V.InfluenceWeights[InfluenceIndex];
-	}
-
-	if (Ar.UEVer() >= VER_UE4_SUPPORT_8_BONE_INFLUENCES_SKELETAL_MESHES)
+	if (!Ar.IsLoading() || Ar.CustomVer(FUE5MainStreamObjectVersion::GUID) >= FUE5MainStreamObjectVersion::IncreasedSkinWeightPrecision)
 	{
-		for (uint32 InfluenceIndex = MAX_INFLUENCES_PER_STREAM; InfluenceIndex < EXTRA_BONE_INFLUENCES; InfluenceIndex++)
+		for (uint32 InfluenceIndex = 0; InfluenceIndex < MAX_TOTAL_INFLUENCES; InfluenceIndex++)
 		{
 			Ar << V.InfluenceWeights[InfluenceIndex];
 		}
 	}
-
-	if (Ar.CustomVer(FAnimObjectVersion::GUID) >= FAnimObjectVersion::UnlimitedBoneInfluences)
+	else
 	{
-		for (uint32 InfluenceIndex = EXTRA_BONE_INFLUENCES; InfluenceIndex < MAX_TOTAL_INFLUENCES; InfluenceIndex++)
+		uint32 MaxInfluences = MAX_INFLUENCES_PER_STREAM;
+		if (Ar.CustomVer(FAnimObjectVersion::GUID) >= FAnimObjectVersion::UnlimitedBoneInfluences)
 		{
-			Ar << V.InfluenceWeights[InfluenceIndex];
+			MaxInfluences = MAX_TOTAL_INFLUENCES;
+		}
+		else if (Ar.UEVer() >= VER_UE4_SUPPORT_8_BONE_INFLUENCES_SKELETAL_MESHES)
+		{
+			MaxInfluences = EXTRA_BONE_INFLUENCES;
+		}
+		
+		uint8 OldInfluence = 0;
+		for (uint32 InfluenceIndex = 0; InfluenceIndex < MaxInfluences; InfluenceIndex++)
+		{
+			Ar << OldInfluence;
+			V.InfluenceWeights[InfluenceIndex] = (static_cast<uint16>(OldInfluence) << 8) | OldInfluence;
 		}
 	}
 
@@ -152,7 +158,7 @@ bool FSoftSkinVertex::GetRigidWeightBone(FBoneIndexType& OutBoneIndex) const
 
 	for (int32 WeightIdx = 0; WeightIdx < MAX_TOTAL_INFLUENCES; WeightIdx++)
 	{
-		if (InfluenceWeights[WeightIdx] == 255)
+		if (InfluenceWeights[WeightIdx] == std::numeric_limits<uint16>::max())
 		{
 			bIsRigid = true;
 			OutBoneIndex = InfluenceBones[WeightIdx];
@@ -163,13 +169,13 @@ bool FSoftSkinVertex::GetRigidWeightBone(FBoneIndexType& OutBoneIndex) const
 	return bIsRigid;
 }
 
-uint8 FSoftSkinVertex::GetMaximumWeight() const
+uint16 FSoftSkinVertex::GetMaximumWeight() const
 {
-	uint8 MaxInfluenceWeight = 0;
+	uint16 MaxInfluenceWeight = 0;
 
 	for (int32 Index = 0; Index < MAX_TOTAL_INFLUENCES; Index++)
 	{
-		const uint8 Weight = InfluenceWeights[Index];
+		const uint16 Weight = InfluenceWeights[Index];
 
 		if (Weight > MaxInfluenceWeight)
 		{
@@ -239,7 +245,7 @@ struct FLegacyRigidSkinVertex
 
 		DestVertex.Color = Color;
 		DestVertex.InfluenceBones[0] = Bone;
-		DestVertex.InfluenceWeights[0] = 255;
+		DestVertex.InfluenceWeights[0] = std::numeric_limits<uint16>::max();
 		for (int32 InfluenceIndex = 1; InfluenceIndex < MAX_TOTAL_INFLUENCES; InfluenceIndex++)
 		{
 			DestVertex.InfluenceBones[InfluenceIndex] = 0;
@@ -1402,7 +1408,7 @@ void FSkeletalMeshLODModel::GetMeshDescription(FMeshDescription& MeshDescription
 	}
 
 	const TArray<FSkeletalMaterial>& Materials = Owner->GetMaterials();
-	const bool bHasVertexColors = (Owner->GetVertexBufferFlags() & ESkeletalMeshVertexFlags::HasVertexColors);
+	const bool bHasVertexColors = EnumHasAllFlags(Owner->GetVertexBufferFlags(), ESkeletalMeshVertexFlags::HasVertexColors);
 
 	// Convert sections to polygon groups, each with their own material.
 	for (int32 SectionIndex = 0; SectionIndex < Sections.Num(); SectionIndex++)
diff --git a/Engine/Source/Runtime/Engine/Private/SkeletalMeshLODRenderData.cpp b/Engine/Source/Runtime/Engine/Private/SkeletalMeshLODRenderData.cpp
--- a/Engine/Source/Runtime/Engine/Private/SkeletalMeshLODRenderData.cpp
+++ b/Engine/Source/Runtime/Engine/Private/SkeletalMeshLODRenderData.cpp
@@ -317,12 +317,13 @@ void FSkeletalMeshLODRenderData::DecrementMemoryStats()
 }
 
 #if WITH_EDITOR
-void FSkeletalMeshLODRenderData::BuildFromLODModel(const FSkeletalMeshLODModel* ImportedModel, uint32 BuildFlags)
+void FSkeletalMeshLODRenderData::BuildFromLODModel(const FSkeletalMeshLODModel* ImportedModel, ESkeletalMeshVertexFlags BuildFlags)
 {
-	bool bUseFullPrecisionUVs = (BuildFlags & ESkeletalMeshVertexFlags::UseFullPrecisionUVs) != 0;
-	bool bUseHighPrecisionTangentBasis = (BuildFlags & ESkeletalMeshVertexFlags::UseHighPrecisionTangentBasis) != 0;
-	bool bHasVertexColors = (BuildFlags & ESkeletalMeshVertexFlags::HasVertexColors) != 0;
-	bool bUseBackwardsCompatibleF16TruncUVs = (BuildFlags & ESkeletalMeshVertexFlags::UseBackwardsCompatibleF16TruncUVs) != 0;
+	const bool bUseFullPrecisionUVs = EnumHasAllFlags(BuildFlags, ESkeletalMeshVertexFlags::UseFullPrecisionUVs);
+	const bool bUseHighPrecisionTangentBasis = EnumHasAllFlags(BuildFlags, ESkeletalMeshVertexFlags::UseHighPrecisionTangentBasis);
+	const bool bHasVertexColors = EnumHasAllFlags(BuildFlags, ESkeletalMeshVertexFlags::HasVertexColors);
+	const bool bUseBackwardsCompatibleF16TruncUVs = EnumHasAllFlags(BuildFlags, ESkeletalMeshVertexFlags::UseBackwardsCompatibleF16TruncUVs);
+	const bool bUseHighPrecisionWeights = EnumHasAllFlags(BuildFlags, ESkeletalMeshVertexFlags::UseHighPrecisionWeights);
 
 	// Copy required info from source sections
 	RenderSections.Empty();
@@ -371,10 +372,45 @@ void FSkeletalMeshLODRenderData::BuildFromLODModel(const FSkeletalMeshLODModel*
 		}
 	}
 
+	int32 MaxBoneInfluences = ImportedModel->GetMaxBoneInfluences();
+	if (!bUseHighPrecisionWeights)
+	{
+		// Re-normalize the weights for 8-bits to ensure that they are distributed using the old algorithm
+		// from MeshUtilities.cpp
+		uint8 InfluenceWeights[MAX_TOTAL_INFLUENCES];
+
+		MaxBoneInfluences = 0;
+		for (FSoftSkinVertex& Vertex: Vertices)
+		{
+			uint32	TotalInfluenceWeight = 0;
+			int32	MaxInfluenceIndex = 0;
+			for (; Vertex.InfluenceWeights[MaxInfluenceIndex] && MaxInfluenceIndex < MAX_TOTAL_INFLUENCES; MaxInfluenceIndex++)
+			{
+				InfluenceWeights[MaxInfluenceIndex] = static_cast<uint8>(Vertex.InfluenceWeights[MaxInfluenceIndex] >> 8);
+				if (InfluenceWeights[MaxInfluenceIndex] == 0)
+				{
+					break;
+				}
+				TotalInfluenceWeight += InfluenceWeights[MaxInfluenceIndex];
+			}
+			InfluenceWeights[0] += 255 - TotalInfluenceWeight;
+
+			MaxBoneInfluences = FMath::Max(MaxBoneInfluences, MaxInfluenceIndex);
+
+			FMemory::Memzero(Vertex.InfluenceWeights);
+			for (int32 Index = 0; Index < MaxInfluenceIndex; Index++)
+			{
+				// Map from 8-bit range to 16-bit range.
+				Vertex.InfluenceWeights[Index] = (static_cast<uint16>(InfluenceWeights[Index]) << 8) | InfluenceWeights[Index];
+			}
+		}
+	}
+
 	// Init skin weight buffer
 	SkinWeightVertexBuffer.SetNeedsCPUAccess(true);
-	SkinWeightVertexBuffer.SetMaxBoneInfluences(ImportedModel->GetMaxBoneInfluences());
+	SkinWeightVertexBuffer.SetMaxBoneInfluences(MaxBoneInfluences);
 	SkinWeightVertexBuffer.SetUse16BitBoneIndex(ImportedModel->DoSectionsUse16BitBoneIndex());
+	SkinWeightVertexBuffer.SetUse16BitBoneWeight(bUseHighPrecisionWeights);
 	SkinWeightVertexBuffer.Init(Vertices);
 
 	// Init the color buffer if this mesh has vertex colors.
diff --git a/Engine/Source/Runtime/Engine/Private/SkeletalMeshMerge.cpp b/Engine/Source/Runtime/Engine/Private/SkeletalMeshMerge.cpp
--- a/Engine/Source/Runtime/Engine/Private/SkeletalMeshMerge.cpp
+++ b/Engine/Source/Runtime/Engine/Private/SkeletalMeshMerge.cpp
@@ -671,7 +671,7 @@ void FSkeletalMeshMerge::GenerateLODModel( int32 LODIdx )
 					if (DestWeight.InfluenceWeights[Idx] > 0)
 					{
 						checkSlow(MergeSectionInfo.BoneMapToMergedBoneMap.IsValidIndex(DestWeight.InfluenceBones[Idx]));
-						DestWeight.InfluenceBones[Idx] = (uint8)MergeSectionInfo.BoneMapToMergedBoneMap[DestWeight.InfluenceBones[Idx]];
+						DestWeight.InfluenceBones[Idx] = (FBoneIndexType)MergeSectionInfo.BoneMapToMergedBoneMap[DestWeight.InfluenceBones[Idx]];
 					}
 				}
 			}
diff --git a/Engine/Source/Runtime/Engine/Private/SkeletalMeshRenderData.cpp b/Engine/Source/Runtime/Engine/Private/SkeletalMeshRenderData.cpp
--- a/Engine/Source/Runtime/Engine/Private/SkeletalMeshRenderData.cpp
+++ b/Engine/Source/Runtime/Engine/Private/SkeletalMeshRenderData.cpp
@@ -345,11 +345,12 @@ void FSkeletalMeshRenderData::Cache(const ITargetPlatform* TargetPlatform, USkin
 				LODRenderData.Add(LODData);
 				
 				//Get the UVs and tangents precision build settings flag specific for this LOD index
-				uint32 VertexBufferBuildFlags = Owner->GetVertexBufferFlags();
+				ESkeletalMeshVertexFlags VertexBufferBuildFlags = Owner->GetVertexBufferFlags();
 				{
-					bool bUseFullPrecisionUVs = LODInfo->BuildSettings.bUseFullPrecisionUVs;
-					bool bUseHighPrecisionTangentBasis = LODInfo->BuildSettings.bUseHighPrecisionTangentBasis;
-					bool bUseBackwardsCompatibleF16TruncUVs = LODInfo->BuildSettings.bUseBackwardsCompatibleF16TruncUVs;
+					const bool bUseFullPrecisionUVs = LODInfo->BuildSettings.bUseFullPrecisionUVs;
+					const bool bUseHighPrecisionTangentBasis = LODInfo->BuildSettings.bUseHighPrecisionTangentBasis;
+					const bool bUseBackwardsCompatibleF16TruncUVs = LODInfo->BuildSettings.bUseBackwardsCompatibleF16TruncUVs;
+					const bool bUseHighPrecisionWeights = LODInfo->BuildSettings.bUseHighPrecisionSkinWeights;
 					if (bUseFullPrecisionUVs || !GVertexElementTypeSupport.IsSupported(VET_Half2))
 					{
 						VertexBufferBuildFlags |= ESkeletalMeshVertexFlags::UseFullPrecisionUVs;
@@ -362,6 +363,10 @@ void FSkeletalMeshRenderData::Cache(const ITargetPlatform* TargetPlatform, USkin
 					{
 						VertexBufferBuildFlags |= ESkeletalMeshVertexFlags::UseBackwardsCompatibleF16TruncUVs;
 					}
+					if (bUseHighPrecisionWeights)
+					{
+						VertexBufferBuildFlags |= ESkeletalMeshVertexFlags::UseHighPrecisionWeights;
+					}
 				}
 				LODData->BuildFromLODModel(LODModel, VertexBufferBuildFlags);
 			}
diff --git a/Engine/Source/Runtime/Engine/Private/SkeletalRenderCPUSkin.cpp b/Engine/Source/Runtime/Engine/Private/SkeletalRenderCPUSkin.cpp
--- a/Engine/Source/Runtime/Engine/Private/SkeletalRenderCPUSkin.cpp
+++ b/Engine/Source/Runtime/Engine/Private/SkeletalRenderCPUSkin.cpp
@@ -745,6 +745,8 @@ static void SkinVertexSection(
 	float ClothBlendWeight, 
 	const FMatrix& WorldToLocal)
 {
+	static constexpr VectorRegister VECTOR_INV_65535 = MakeVectorRegisterDoubleConstant(1.0 / 65535, 1.0 / 65535, 1.0 / 65535, 1.0 / 65535);
+	
 	// VertexCopy for morph. Need to allocate right struct
 	// To avoid re-allocation, create 2 statics, and assign right struct
 	VertexType  VertexCopy;
@@ -795,23 +797,23 @@ static void SkinVertexSection(
 			}
 
 			const FBoneIndexType* RESTRICT BoneIndices = SrcWeights.InfluenceBones;
-			const uint8* RESTRICT BoneWeights = SrcWeights.InfluenceWeights;
+			const uint16* RESTRICT BoneWeights = SrcWeights.InfluenceWeights;
 
 			static VectorRegister	SrcNormals[3];
 			VectorRegister			DstNormals[3];
 			SrcNormals[0] = VectorLoadFloat3_W1( &MorphedVertex->Position);
 			SrcNormals[1] = Unpack3( &MorphedVertex->TangentX.Vector.Packed );
 			SrcNormals[2] = Unpack4( &MorphedVertex->TangentZ.Vector.Packed );
-			VectorRegister Weights = VectorMultiply( VectorLoadByte4(BoneWeights), VECTOR_INV_255 );
+			VectorRegister Weights = VectorMultiply( VectorLoadURGBA16N(BoneWeights), VECTOR_INV_65535 );
 			VectorRegister ExtraWeights = MakeVectorRegister(0.f, 0.f, 0.f, 0.f);
 			VectorRegister ExtraWeights2 = MakeVectorRegister(0.f, 0.f, 0.f, 0.f);
 			if (MaxSectionBoneInfluences > 4)
 			{
-				ExtraWeights = VectorMultiply( VectorLoadByte4(&BoneWeights[MAX_INFLUENCES_PER_STREAM]), VECTOR_INV_255 );
+				ExtraWeights = VectorMultiply( VectorLoadURGBA16N(&BoneWeights[MAX_INFLUENCES_PER_STREAM]), VECTOR_INV_65535 );
 			}
 			if (MaxSectionBoneInfluences > 8)
 			{
-				ExtraWeights2 = VectorMultiply(VectorLoadByte4(&BoneWeights[EXTRA_BONE_INFLUENCES]), VECTOR_INV_255);
+				ExtraWeights2 = VectorMultiply(VectorLoadURGBA16N(&BoneWeights[EXTRA_BONE_INFLUENCES]), VECTOR_INV_65535);
 			}
 			VectorResetFloatRegisters(); // Need to call this to be able to use regular floating point registers again after Unpack and VectorLoadByte4.
 
@@ -1121,8 +1123,6 @@ FVector4 GetTangetToColor(FPackedNormal Tangent)
  */
 static FORCEINLINE void CalculateSectionBoneWeights(FFinalSkinVertex*& DestVertex, FSkinWeightVertexBuffer& SkinWeightVertexBuffer, FSkelMeshRenderSection& Section, const TArray<int32>& BonesOfInterest)
 {
-	const float INV255 = 1.f/255.f;
-
 	int32 VertexBufferBaseIndex = 0;
 
 	//array of bone mapping
@@ -1138,14 +1138,14 @@ static FORCEINLINE void CalculateSectionBoneWeights(FFinalSkinVertex*& DestVerte
 		DestVertex->TextureCoordinates[0].Y = 0.0f;
 
 		const FBoneIndexType* RESTRICT BoneIndices = SrcWeight.InfluenceBones;
-		const uint8* RESTRICT BoneWeights = SrcWeight.InfluenceWeights;
+		const uint16* RESTRICT BoneWeights = SrcWeight.InfluenceWeights;
 
 		for (uint32 i = 0; i < SkinWeightVertexBuffer.GetMaxBoneInfluences(); i++)
 		{
 			if (BonesOfInterest.Contains(BoneMap[BoneIndices[i]]))
 			{
-				DestVertex->TextureCoordinates[0].X += BoneWeights[i] * INV255; 
-				DestVertex->TextureCoordinates[0].Y += BoneWeights[i] * INV255;
+				DestVertex->TextureCoordinates[0].X += BoneWeights[i] / 65535.0; 
+				DestVertex->TextureCoordinates[0].Y += BoneWeights[i] / 65535.0;
 			}
 		}
 	}
diff --git a/Engine/Source/Runtime/Engine/Private/SkeletalRenderGPUSkin.cpp b/Engine/Source/Runtime/Engine/Private/SkeletalRenderGPUSkin.cpp
--- a/Engine/Source/Runtime/Engine/Private/SkeletalRenderGPUSkin.cpp
+++ b/Engine/Source/Runtime/Engine/Private/SkeletalRenderGPUSkin.cpp
@@ -1631,28 +1631,30 @@ void InitGPUSkinVertexFactoryComponents(FGPUSkinDataType* VertexFactoryData, con
 	VertexBuffers.StaticVertexBuffers->StaticMeshVertexBuffer.BindTangentVertexBuffer(VertexFactory, *VertexFactoryData);
 	VertexBuffers.StaticVertexBuffers->StaticMeshVertexBuffer.BindPackedTexCoordVertexBuffer(VertexFactory, *VertexFactoryData, MAX_TEXCOORDS);
 
-	bool bUse16BitBoneIndex = VertexBuffers.SkinWeightVertexBuffer->Use16BitBoneIndex();
+	const FSkinWeightVertexBuffer* WeightBuffer = VertexBuffers.SkinWeightVertexBuffer; 
+	const bool bUse16BitBoneIndex = WeightBuffer->Use16BitBoneIndex();
+	const bool bUse16BitBoneWeight = WeightBuffer->Use16BitBoneWeight();
 	VertexFactoryData->bUse16BitBoneIndex = bUse16BitBoneIndex;
-	VertexFactoryData->NumBoneInfluences = VertexBuffers.SkinWeightVertexBuffer->GetMaxBoneInfluences();
+	VertexFactoryData->NumBoneInfluences = WeightBuffer->GetMaxBoneInfluences();
 
-	GPUSkinBoneInfluenceType BoneInfluenceType = VertexBuffers.SkinWeightVertexBuffer->GetBoneInfluenceType();
+	GPUSkinBoneInfluenceType BoneInfluenceType = WeightBuffer->GetBoneInfluenceType();
 	if (BoneInfluenceType == GPUSkinBoneInfluenceType::UnlimitedBoneInfluence)
 	{
 		FGPUBaseSkinVertexFactory::FShaderDataType& ShaderData = VertexFactory->GetShaderData();
-		ShaderData.InputWeightIndexSize = VertexBuffers.SkinWeightVertexBuffer->GetBoneIndexByteSize();
-		ShaderData.InputWeightStream = VertexBuffers.SkinWeightVertexBuffer->GetDataVertexBuffer()->GetSRV();
+		ShaderData.InputWeightIndexSize = WeightBuffer->GetBoneIndexByteSize() | (WeightBuffer->GetBoneWeightByteSize() << 8);
+		ShaderData.InputWeightStream = WeightBuffer->GetDataVertexBuffer()->GetSRV();
 
-		const FSkinWeightLookupVertexBuffer* LookupVertexBuffer = VertexBuffers.SkinWeightVertexBuffer->GetLookupVertexBuffer();
+		const FSkinWeightLookupVertexBuffer* LookupVertexBuffer = WeightBuffer->GetLookupVertexBuffer();
 		VertexFactoryData->BlendOffsetCount = FVertexStreamComponent(LookupVertexBuffer, 0, LookupVertexBuffer->GetStride(), VET_UInt);
 	}
 	else
 	{
 		// bone indices & weights
-		const FSkinWeightDataVertexBuffer* WeightDataVertexBuffer = VertexBuffers.SkinWeightVertexBuffer->GetDataVertexBuffer();
-		const uint32 Stride = VertexBuffers.SkinWeightVertexBuffer->GetConstantInfluencesVertexStride();
-		const uint32 WeightsOffset = VertexBuffers.SkinWeightVertexBuffer->GetConstantInfluencesBoneWeightsOffset();
+		const FSkinWeightDataVertexBuffer* WeightDataVertexBuffer = WeightBuffer->GetDataVertexBuffer();
+		const uint32 Stride = WeightBuffer->GetConstantInfluencesVertexStride();
+		const uint32 WeightsOffset = WeightBuffer->GetConstantInfluencesBoneWeightsOffset();
 		VertexFactoryData->BoneIndices = FVertexStreamComponent(WeightDataVertexBuffer, 0, Stride, bUse16BitBoneIndex ? VET_UShort4 : VET_UByte4);
-		VertexFactoryData->BoneWeights = FVertexStreamComponent(WeightDataVertexBuffer, WeightsOffset, Stride, VET_UByte4N);
+		VertexFactoryData->BoneWeights = FVertexStreamComponent(WeightDataVertexBuffer, WeightsOffset, Stride, bUse16BitBoneWeight ? VET_UShort4N : VET_UByte4N);
 
 		if (VertexFactoryData->NumBoneInfluences > MAX_INFLUENCES_PER_STREAM)
 		{
@@ -1660,7 +1662,7 @@ void InitGPUSkinVertexFactoryComponents(FGPUSkinDataType* VertexFactoryData, con
 			VertexFactoryData->ExtraBoneIndices = FVertexStreamComponent(
 				WeightDataVertexBuffer, 4 * VertexBuffers.SkinWeightVertexBuffer->GetBoneIndexByteSize(), Stride, bUse16BitBoneIndex ? VET_UShort4 : VET_UByte4);
 			VertexFactoryData->ExtraBoneWeights = FVertexStreamComponent(
-				WeightDataVertexBuffer, WeightsOffset + 4, Stride, VET_UByte4N);
+				WeightDataVertexBuffer, WeightsOffset + 4 * VertexBuffers.SkinWeightVertexBuffer->GetBoneWeightByteSize(), Stride, bUse16BitBoneWeight ? VET_UShort4N : VET_UByte4N);
 		}
 	}
 
diff --git a/Engine/Source/Runtime/Engine/Private/SkinnedAssetCommon.cpp b/Engine/Source/Runtime/Engine/Private/SkinnedAssetCommon.cpp
--- a/Engine/Source/Runtime/Engine/Private/SkinnedAssetCommon.cpp
+++ b/Engine/Source/Runtime/Engine/Private/SkinnedAssetCommon.cpp
@@ -64,6 +64,7 @@ static void SerializeBuildSettingsForDDC(FArchive& Ar, FSkeletalMeshBuildSetting
 	FArchive_Serialize_BitfieldBool(Ar, BuildSettings.bUseFullPrecisionUVs);
 	FArchive_Serialize_BitfieldBool(Ar, BuildSettings.bUseBackwardsCompatibleF16TruncUVs);
 	FArchive_Serialize_BitfieldBool(Ar, BuildSettings.bUseHighPrecisionTangentBasis);
+	FArchive_Serialize_BitfieldBool(Ar, BuildSettings.bUseHighPrecisionSkinWeights);
 	Ar << BuildSettings.ThresholdPosition;
 	Ar << BuildSettings.ThresholdTangentNormal;
 	Ar << BuildSettings.ThresholdUV;
@@ -98,6 +99,7 @@ FGuid FSkeletalMeshLODInfo::ComputeDeriveDataCacheKey(const FSkeletalMeshLODGrou
 		FString BakePoseOverridePath = BakePoseOverride->GetFullName();
 		Ar << BakePoseOverridePath;
 	}
+	
 	FArchive_Serialize_BitfieldBool(Ar, bAllowCPUAccess);
 	FArchive_Serialize_BitfieldBool(Ar, bSupportUniformlyDistributedSampling);
 
diff --git a/Engine/Source/Runtime/Engine/Public/Animation/SkinWeightProfile.h b/Engine/Source/Runtime/Engine/Public/Animation/SkinWeightProfile.h
--- a/Engine/Source/Runtime/Engine/Public/Animation/SkinWeightProfile.h
+++ b/Engine/Source/Runtime/Engine/Public/Animation/SkinWeightProfile.h
@@ -61,7 +61,7 @@ struct FRawSkinWeight
 {
 	// MAX_TOTAL_INFLUENCES for now
 	FBoneIndexType InfluenceBones[MAX_TOTAL_INFLUENCES];
-	uint8 InfluenceWeights[MAX_TOTAL_INFLUENCES];
+	uint16 InfluenceWeights[MAX_TOTAL_INFLUENCES];
 
 	friend FArchive& operator<<(FArchive& Ar, FRawSkinWeight& OverrideEntry);
 };
diff --git a/Engine/Source/Runtime/Engine/Public/Rendering/SkeletalMeshLODModel.h b/Engine/Source/Runtime/Engine/Public/Rendering/SkeletalMeshLODModel.h
--- a/Engine/Source/Runtime/Engine/Public/Rendering/SkeletalMeshLODModel.h
+++ b/Engine/Source/Runtime/Engine/Public/Rendering/SkeletalMeshLODModel.h
@@ -37,13 +37,13 @@ struct FSoftSkinVertex
 	// VertexColor
 	FColor			Color;
 	FBoneIndexType	InfluenceBones[MAX_TOTAL_INFLUENCES];
-	uint8			InfluenceWeights[MAX_TOTAL_INFLUENCES];
+	uint16			InfluenceWeights[MAX_TOTAL_INFLUENCES];
 
 	/** If this vert is rigidly weighted to a bone, return true and the bone index. Otherwise return false. */
 	ENGINE_API bool GetRigidWeightBone(FBoneIndexType& OutBoneIndex) const;
 
 	/** Returns the maximum weight of any bone that influences this vertex. */
-	ENGINE_API uint8 GetMaximumWeight() const;
+	ENGINE_API uint16 GetMaximumWeight() const;
 
 	/**
 	* Serializer
diff --git a/Engine/Source/Runtime/Engine/Public/Rendering/SkeletalMeshLODRenderData.h b/Engine/Source/Runtime/Engine/Public/Rendering/SkeletalMeshLODRenderData.h
--- a/Engine/Source/Runtime/Engine/Public/Rendering/SkeletalMeshLODRenderData.h
+++ b/Engine/Source/Runtime/Engine/Public/Rendering/SkeletalMeshLODRenderData.h
@@ -231,7 +231,7 @@ public:
 	 * Initialize render data (e.g. vertex buffers) from model info
 	 * @param BuildFlags See ESkeletalMeshVertexFlags.
 	 */
-	void ENGINE_API BuildFromLODModel(const FSkeletalMeshLODModel* LODModel, uint32 BuildFlags);
+	void ENGINE_API BuildFromLODModel(const FSkeletalMeshLODModel* LODModel, ESkeletalMeshVertexFlags BuildFlags = ESkeletalMeshVertexFlags::None);
 #endif // WITH_EDITOR
 
 	uint32 GetNumVertices() const
diff --git a/Engine/Source/Runtime/Engine/Public/Rendering/SkinWeightVertexBuffer.h b/Engine/Source/Runtime/Engine/Public/Rendering/SkinWeightVertexBuffer.h
--- a/Engine/Source/Runtime/Engine/Public/Rendering/SkinWeightVertexBuffer.h
+++ b/Engine/Source/Runtime/Engine/Public/Rendering/SkinWeightVertexBuffer.h
@@ -59,7 +59,7 @@ struct TLegacySkinWeightInfo
 struct FSkinWeightInfo
 {
 	FBoneIndexType	InfluenceBones[MAX_TOTAL_INFLUENCES];
-	uint8			InfluenceWeights[MAX_TOTAL_INFLUENCES];
+	uint16			InfluenceWeights[MAX_TOTAL_INFLUENCES];
 };
 
 struct FSkinWeightRHIInfo
@@ -218,7 +218,7 @@ public:
 	/** Delete existing resources */
 	ENGINE_API void CleanUp();
 
-	ENGINE_API void Init(uint32 InNumBones, uint32 InNumVertices);
+	ENGINE_API void Init(uint32 InNumWeights, uint32 InNumVertices);
 
 	friend FArchive& operator<<(FArchive& Ar, FSkinWeightDataVertexBuffer& VertexBuffer);
 
@@ -236,17 +236,27 @@ public:
 	FORCEINLINE uint32 GetNumVertices() const
 	{ return NumVertices; }
 
-	/** @return number of bones in this vertex buffer */
-	FORCEINLINE uint32 GetNumBones() const
-	{ return NumBones; }
+	/** @return number of bone weights in this vertex buffer */
+	FORCEINLINE uint32 GetNumBoneWeights() const
+	{ return NumBoneWeights; }
 
 	/** @return byte size of each bone index */
 	FORCEINLINE uint32 GetBoneIndexByteSize() const
 	{ return (Use16BitBoneIndex() ? sizeof(FBoneIndex16) : sizeof(FBoneIndex8)); }
 
+	FORCEINLINE uint32 GetBoneWeightByteSize() const
+	{
+		return Use16BitBoneWeight() ? sizeof(uint16) : sizeof(uint8);
+	}
+
+	FORCEINLINE uint32 GetBoneIndexAndWeightByteSize() const
+	{
+		return GetBoneIndexByteSize() + GetBoneWeightByteSize(); 
+	}
+
 	/** @return vertex stride for when using constant number of bones per vertex buffer */
 	FORCEINLINE uint32 GetConstantInfluencesVertexStride() const
-	{ return (GetBoneIndexByteSize() + sizeof(uint8)) * MaxBoneInfluences; }
+	{ return GetBoneIndexAndWeightByteSize() * MaxBoneInfluences; }
 
 	/** @return offset position for bone weights data for each vertex */
 	FORCEINLINE uint32 GetConstantInfluencesBoneWeightsOffset() const
@@ -254,7 +264,7 @@ public:
 
 	/** @return total size of data in resource array */
 	FORCEINLINE uint32 GetVertexDataSize() const
-	{ return NumBones * (GetBoneIndexByteSize() + sizeof(uint8)); }
+	{ return NumBoneWeights * GetBoneIndexAndWeightByteSize(); }
 
 	// @param guaranteed only to be valid if the vertex buffer is valid
 	FRHIShaderResourceView* GetSRV() const
@@ -281,13 +291,21 @@ public:
 
 	FORCEINLINE bool Use16BitBoneIndex() const
 	{ return bUse16BitBoneIndex; }
+	
+	void SetUse16BitBoneWeight(bool bInUse16BitBoneWeight)
+	{ bUse16BitBoneWeight = bInUse16BitBoneWeight; }
+
+	bool Use16BitBoneWeight() const
+	{ return bUse16BitBoneWeight; }
+	
+	
 
 	ENGINE_API GPUSkinBoneInfluenceType GetBoneInfluenceType() const;
 	ENGINE_API bool GetRigidWeightBone(uint32 VertexWeightOffset, uint32 VertexInfluenceCount, int32& OutBoneIndex) const;
 	ENGINE_API uint32 GetBoneIndex(uint32 VertexWeightOffset, uint32 VertexInfluenceCount, uint32 InfluenceIndex) const;
 	ENGINE_API void SetBoneIndex(uint32 VertexWeightOffset, uint32 VertexInfluenceCount, uint32 InfluenceIndex, uint32 BoneIndex);
-	ENGINE_API uint8 GetBoneWeight(uint32 VertexWeightOffset, uint32 VertexInfluenceCount, uint32 InfluenceIndex) const;
-	ENGINE_API void SetBoneWeight(uint32 VertexWeightOffset, uint32 VertexInfluenceCount, uint32 InfluenceIndex, uint8 BoneWeight);
+	ENGINE_API uint16 GetBoneWeight(uint32 VertexWeightOffset, uint32 VertexInfluenceCount, uint32 InfluenceIndex) const;
+	ENGINE_API void SetBoneWeight(uint32 VertexWeightOffset, uint32 VertexInfluenceCount, uint32 InfluenceIndex, uint16 BoneWeight);
 	ENGINE_API void ResetVertexBoneWeights(uint32 VertexWeightOffset, uint32 VertexInfluenceCount);
 
 	ENGINE_API void CopyDataFromBuffer(const TArrayView<const FSkinWeightInfo>& SkinWeightData);
@@ -323,17 +341,22 @@ public:
 
 	bool IsWeightDataValid() const;
 
-	FSkinWeightInfo* GetWeightData() const
+	const uint8* GetWeightData() const
 	{
-		return (FSkinWeightInfo*)Data;
+		return Data;
 	}
 
+	uint8* GetWeightData()
+	{
+		return Data;
+	}
+	
 protected:
 	// guaranteed only to be valid if the vertex buffer is valid
 	FShaderResourceViewRHIRef SRVValue;
 
 private:
-	void ResizeBuffer(uint32 InNumBones, uint32 InNumVertices);
+	void ResizeBuffer(uint32 InNumWeights, uint32 InNumVertices);
 
 	EPixelFormat GetPixelFormat() const
 	{ return bVariableBonesPerVertex ? PF_R8_UINT : PF_R32_UINT; }
@@ -352,6 +375,9 @@ private:
 	/** Use 16 bit bone index instead of 8 bit */
 	bool bUse16BitBoneIndex;
 
+	/** Use 16 bit bone weight instead of 8 bit */
+	bool bUse16BitBoneWeight;
+
 	/** The vertex data storage type */
 	FStaticMeshVertexDataInterface* WeightData;
 
@@ -361,8 +387,8 @@ private:
 	/** The cached number of vertices. */
 	uint32 NumVertices;
 
-	/** The cached number of bones. */
-	uint32 NumBones;
+	/** Total number of bone weights across all vertices */
+	uint32 NumBoneWeights;
 
 	/** Allocates the vertex data storage type. */
 	ENGINE_API void AllocateData();
@@ -436,9 +462,21 @@ public:
 	bool Use16BitBoneIndex() const
 	{ return DataVertexBuffer.Use16BitBoneIndex(); }
 
+	void SetUse16BitBoneWeight(bool bInUse16BitBoneWeight)
+	{ DataVertexBuffer.SetUse16BitBoneWeight(bInUse16BitBoneWeight); }
+	
+	bool Use16BitBoneWeight() const
+	{ return DataVertexBuffer.Use16BitBoneWeight(); }
+	
 	uint32 GetBoneIndexByteSize() const
 	{ return DataVertexBuffer.GetBoneIndexByteSize(); }
 
+	uint32 GetBoneWeightByteSize() const
+	{ return DataVertexBuffer.GetBoneWeightByteSize(); }
+	
+	uint32 GetBoneIndexAndWeightByteSize() const
+	{ return DataVertexBuffer.GetBoneIndexAndWeightByteSize(); }
+	
 	bool GetVariableBonesPerVertex() const
 	{ return DataVertexBuffer.GetVariableBonesPerVertex(); }
 
@@ -462,8 +500,8 @@ public:
 	ENGINE_API bool GetRigidWeightBone(uint32 VertexIndex, int32& OutBoneIndex) const;
 	ENGINE_API uint32 GetBoneIndex(uint32 VertexIndex, uint32 InfluenceIndex) const;
 	ENGINE_API void SetBoneIndex(uint32 VertexIndex, uint32 InfluenceIndex, uint32 BoneIndex);
-	ENGINE_API uint8 GetBoneWeight(uint32 VertexIndex, uint32 InfluenceIndex) const;
-	ENGINE_API void SetBoneWeight(uint32 VertexIndex, uint32 InfluenceIndex, uint8 BoneWeight);
+	ENGINE_API uint16 GetBoneWeight(uint32 VertexIndex, uint32 InfluenceIndex) const;
+	ENGINE_API void SetBoneWeight(uint32 VertexIndex, uint32 InfluenceIndex, uint16 BoneWeight);
 	ENGINE_API void ResetVertexBoneWeights(uint32 VertexIndex);
 
 	ENGINE_API void BeginInitResources();
diff --git a/Engine/Source/Runtime/Engine/Public/SkeletalMeshTypes.h b/Engine/Source/Runtime/Engine/Public/SkeletalMeshTypes.h
--- a/Engine/Source/Runtime/Engine/Public/SkeletalMeshTypes.h
+++ b/Engine/Source/Runtime/Engine/Public/SkeletalMeshTypes.h
@@ -22,17 +22,16 @@ class FSkeletalMeshRenderData;
 class FSkeletalMeshLODRenderData;
 
 /** Flags used when building vertex buffers. */
-struct ESkeletalMeshVertexFlags
+enum class ESkeletalMeshVertexFlags
 {
-	enum
-	{
-		None = 0x0,
-		UseFullPrecisionUVs = 0x1,
-		HasVertexColors = 0x2,
-		UseHighPrecisionTangentBasis = 0x4,
-		UseBackwardsCompatibleF16TruncUVs = 0x8
-	};
+	None								= 0x0,
+	UseFullPrecisionUVs					= 0x1,
+	HasVertexColors						= 0x2,
+	UseHighPrecisionTangentBasis		= 0x4,
+	UseBackwardsCompatibleF16TruncUVs	= 0x8,
+	UseHighPrecisionWeights				= 0x10,
 };
+ENUM_CLASS_FLAGS(ESkeletalMeshVertexFlags);
 
 /** Name of vertex color channels, used by recompute tangents */
 enum class ESkinVertexColorChannel : uint8
