- Optionally reset Niagara & Cascade when crossing too many LWC tiles as it will result in rendering artifacts

diff --git a/Engine/Plugins/FX/Niagara/Source/Niagara/Private/NiagaraSystemInstance.cpp b/Engine/Plugins/FX/Niagara/Source/Niagara/Private/NiagaraSystemInstance.cpp
--- a/Engine/Plugins/FX/Niagara/Source/Niagara/Private/NiagaraSystemInstance.cpp
+++ b/Engine/Plugins/FX/Niagara/Source/Niagara/Private/NiagaraSystemInstance.cpp
@@ -22,7 +22,6 @@
 #include "GameFramework/PlayerController.h"
 #include "Templates/AlignmentTemplates.h"
 
-
 DECLARE_CYCLE_STAT(TEXT("System Activate [GT]"), STAT_NiagaraSystemActivate, STATGROUP_Niagara);
 DECLARE_CYCLE_STAT(TEXT("System Deactivate [GT]"), STAT_NiagaraSystemDeactivate, STATGROUP_Niagara);
 DECLARE_CYCLE_STAT(TEXT("System Complete [GT]"), STAT_NiagaraSystemComplete, STATGROUP_Niagara);
@@ -2379,6 +2378,21 @@ void FNiagaraSystemInstance::Tick_GameThread(float DeltaSeconds)
 		return;
 	}
 
+	// Has the actor position changed to the point where we need to reset the LWC tile
+	if (GetSystem()->SupportsLargeWorldCoordinates())
+	{
+		if (USceneComponent* SceneComponent = AttachComponent.Get())
+		{
+			if (UFXSystemComponent::RequiresLWCTileRecache(LWCTile, SceneComponent->GetComponentLocation()))
+			{
+				//-OPT: For safety we reset everything, but if everything is local space we may not need to, or we could rebase.
+				UE_LOG(LogNiagara, Warning, TEXT("NiagaraComponent(%s - %s) required LWC tile recache and was reset."), *GetFullNameSafe(SceneComponent), *GetFullNameSafe(Asset.Get()));
+				Reset(EResetMode::ResetAll);
+				return;
+			}
+		}
+	}
+
 	CachedDeltaSeconds = DeltaSeconds;
 	FixedBounds_CNC = FixedBounds_GT;
 
diff --git a/Engine/Source/Runtime/Engine/Classes/Particles/ParticleSystemComponent.h b/Engine/Source/Runtime/Engine/Classes/Particles/ParticleSystemComponent.h
--- a/Engine/Source/Runtime/Engine/Classes/Particles/ParticleSystemComponent.h
+++ b/Engine/Source/Runtime/Engine/Classes/Particles/ParticleSystemComponent.h
@@ -405,7 +405,6 @@ public:
 	UFUNCTION(BlueprintCallable, Category = "Effects|Components|ParticleSystem")
 	virtual void SetEmitterEnable(FName EmitterName, bool bNewEnableState) {};
 
-
 	/**
 	 * Set AutoAttachParent, AutoAttachSocketName, AutoAttachLocationRule, AutoAttachRotationRule, AutoAttachScaleRule to the specified parameters. Does not change bAutoManageAttachment; that must be set separately.
 	 * @param  Parent			Component to attach to.
@@ -444,6 +443,9 @@ public:
 	/** Forces component to deactivate immediately. */
 	virtual void DeactivateImmediate() {}
 
+	/** Returns true if we have crossed LWC tiles to the point that we may introduce artifacts. */
+	static bool RequiresLWCTileRecache(const FVector3f CurrentTile, const FVector CurrentLocation);
+
 #if WITH_PER_COMPONENT_PARTICLE_PERF_STATS
 	mutable FParticlePerfStats* ParticlePerfStats = nullptr;
 #endif
diff --git a/Engine/Source/Runtime/Engine/Private/Particles/ParticleComponents.cpp b/Engine/Source/Runtime/Engine/Private/Particles/ParticleComponents.cpp
--- a/Engine/Source/Runtime/Engine/Private/Particles/ParticleComponents.cpp
+++ b/Engine/Source/Runtime/Engine/Private/Particles/ParticleComponents.cpp
@@ -180,6 +180,14 @@ static TAutoConsoleVariable<float> CVarPruneEmittersOnCookByDetailMode(
 	TEXT("This will only work if scalability settings affecting detail mode can not be changed at runtime (depends on platform).\n"),
 	ECVF_ReadOnly);
 
+float GFXLWCTileRecache = 2;
+FAutoConsoleVariableRef CVarFXLWCTileRecache(
+	TEXT("fx.LWCTileRecache"),
+	GFXLWCTileRecache,
+	TEXT("When we cross this number of LWC tiles from where we started the FX we need to recache the LWC tile to avoid artifacts.\n")
+	TEXT("When this occurs the system may need to reset, cull particles too far away, or do some additional processing to handle it.\n")
+	TEXT("Setting this value to 0 will remove this behavior but could introduce rendering & simulation artifacts.\n"),
+	ECVF_Default);
 
 int32 OldDetailModeToBitmask(int32 OldMode)
 {
@@ -3447,6 +3455,19 @@ UFXSystemComponent::UFXSystemComponent(const FObjectInitializer& ObjectInitializ
 	: Super(ObjectInitializer)
 {}
 
+bool UFXSystemComponent::RequiresLWCTileRecache(const FVector3f CurrentTile, const FVector CurrentLocation)
+{
+	bool bNeedsRecache = false;
+	const float TileRecache = GFXLWCTileRecache;
+	if (TileRecache > 0.0f)
+	{
+		const FVector3f ActorTile = FLargeWorldRenderScalar::GetTileFor(CurrentLocation);
+		const float MaxMovement = (CurrentTile - ActorTile).GetAbs().GetMax();
+		bNeedsRecache = MaxMovement >= TileRecache;
+	}
+	return bNeedsRecache;
+}
+
 FOnSystemPreActivationChange UParticleSystemComponent::OnSystemPreActivationChange;
 
 UParticleSystemComponent::UParticleSystemComponent(const FObjectInitializer& ObjectInitializer)
@@ -5272,6 +5293,15 @@ void UParticleSystemComponent::TickComponent(float DeltaTime, enum ELevelTick Ti
 		return;
 	} 
 	
+	// Has the actor position changed to the point where we need to reset the LWC tile
+	if (RequiresLWCTileRecache(LWCTile, GetComponentLocation()))
+	{
+		//-OPT: We may be able to narrow down when a reset is required, like having a GPU emitter, having world space emitters, etc.
+		//      Cascade generally operates at double precision so it may only be GPU emitters that require a reset.
+		UE_LOG(LogParticles, Warning, TEXT("PSC(%s - %s) required LWC tile recache and was reset."), *GetFullNameSafe(this), *GetFullNameSafe(Template));
+		bRequiresReset = true;
+	}
+
 	if (bRequiresReset)
 	{
 		ForceReset();
