BodySetup: Renamed PhysX to BodySetup in stats

STAT_PhysXCooking has not been renamed to maintain backward compatibility in case it is referenced.

#rnx

diff --git a/Engine/Source/Runtime/Engine/Classes/PhysicsEngine/BodySetup.h b/Engine/Source/Runtime/Engine/Classes/PhysicsEngine/BodySetup.h
--- a/Engine/Source/Runtime/Engine/Classes/PhysicsEngine/BodySetup.h
+++ b/Engine/Source/Runtime/Engine/Classes/PhysicsEngine/BodySetup.h
@@ -39,7 +39,7 @@ namespace Chaos
 template<typename T, int d>
 class FChaosDerivedDataReader;
 
-DECLARE_CYCLE_STAT_EXTERN(TEXT("PhysX Cooking"), STAT_PhysXCooking, STATGROUP_Physics, );
+DECLARE_CYCLE_STAT_EXTERN(TEXT("BodySetup Cooking"), STAT_PhysXCooking, STATGROUP_Physics, );
 
 
 /** UV information for BodySetup, only created if UPhysicsSettings::bSupportUVFromHitResults */
diff --git a/Engine/Source/Runtime/Engine/Private/PhysicsEngine/BodySetup.cpp b/Engine/Source/Runtime/Engine/Private/PhysicsEngine/BodySetup.cpp
--- a/Engine/Source/Runtime/Engine/Private/PhysicsEngine/BodySetup.cpp
+++ b/Engine/Source/Runtime/Engine/Private/PhysicsEngine/BodySetup.cpp
@@ -70,12 +70,12 @@ UBodySetup::~UBodySetup() = default;
 
 
 #if ENABLE_COOK_STATS
-namespace PhysXBodySetupCookStats
+namespace BodySetupCookStats
 {
 	static FCookStats::FDDCResourceUsageStats UsageStats;
 	static FCookStatsManager::FAutoRegisterCallback RegisterCookStats([](FCookStatsManager::AddStatFuncRef AddStat)
 	{
-		UsageStats.LogStats(AddStat, TEXT("PhysX.Usage"), TEXT("BodySetup"));
+		UsageStats.LogStats(AddStat, TEXT("BodySetup.Usage"), TEXT(""));
 	});
 }
 #endif
@@ -1310,7 +1310,7 @@ void GetDDCBuiltData(FByteBulkData* OutResult, DDCBuilderType& InBuilder, UBodyS
 
 	if(InBuilder.CanBuild())
 	{
-		COOK_STAT(FCookStats::FScopedStatsCounter Timer = PhysXBodySetupCookStats::UsageStats.TimeSyncWork());
+		COOK_STAT(FCookStats::FScopedStatsCounter Timer = BodySetupCookStats::UsageStats.TimeSyncWork());
 
 		// Debugging switch, force builder to always run
 		bool bSkipDDC = false;
