Check PipelineStateCache::IsPSOPrecachingEnabled() first before trying to call GetRenderData which will wait on async compiles and block editor loads (workaround fix)

diff --git a/Engine/Source/Runtime/Engine/Private/Components/SkinnedMeshComponent.cpp b/Engine/Source/Runtime/Engine/Private/Components/SkinnedMeshComponent.cpp
--- a/Engine/Source/Runtime/Engine/Private/Components/SkinnedMeshComponent.cpp
+++ b/Engine/Source/Runtime/Engine/Private/Components/SkinnedMeshComponent.cpp
@@ -553,10 +553,10 @@ void USkinnedMeshComponent::PostLoad()
 
 void USkinnedMeshComponent::PrecachePSOs()
 {
-	if (GetSkinnedAsset() == nullptr ||
-		GetSkinnedAsset()->GetResourceForRendering() == nullptr ||
-		!FApp::CanEverRender() || 
-		!PipelineStateCache::IsPSOPrecachingEnabled())
+	if (!FApp::CanEverRender() ||
+		!PipelineStateCache::IsPSOPrecachingEnabled() ||
+		GetSkinnedAsset() == nullptr ||
+		GetSkinnedAsset()->GetResourceForRendering() == nullptr)
 	{
 		return;
 	}
diff --git a/Engine/Source/Runtime/Engine/Private/Components/StaticMeshComponent.cpp b/Engine/Source/Runtime/Engine/Private/Components/StaticMeshComponent.cpp
--- a/Engine/Source/Runtime/Engine/Private/Components/StaticMeshComponent.cpp
+++ b/Engine/Source/Runtime/Engine/Private/Components/StaticMeshComponent.cpp
@@ -1538,7 +1538,7 @@ void UStaticMeshComponent::InitResources()
 
 void UStaticMeshComponent::PrecachePSOs()
 {
-	if (StaticMesh == nullptr || StaticMesh->GetRenderData() == nullptr || !FApp::CanEverRender() || !PipelineStateCache::IsPSOPrecachingEnabled())
+	if (!FApp::CanEverRender() || !PipelineStateCache::IsPSOPrecachingEnabled() || StaticMesh == nullptr || StaticMesh->GetRenderData() == nullptr)
 	{
 		return;
 	}
diff --git a/Engine/Source/Runtime/Engine/Private/InstancedStaticMesh.cpp b/Engine/Source/Runtime/Engine/Private/InstancedStaticMesh.cpp
--- a/Engine/Source/Runtime/Engine/Private/InstancedStaticMesh.cpp
+++ b/Engine/Source/Runtime/Engine/Private/InstancedStaticMesh.cpp
@@ -4481,7 +4481,7 @@ void UInstancedStaticMeshComponent::PostLoad()
 
 void UInstancedStaticMeshComponent::PrecachePSOs()
 {	
-	if (GetStaticMesh() == nullptr || GetStaticMesh()->GetRenderData() == nullptr || !FApp::CanEverRender() || !PipelineStateCache::IsPSOPrecachingEnabled())
+	if (!FApp::CanEverRender() || !PipelineStateCache::IsPSOPrecachingEnabled() || GetStaticMesh() == nullptr || GetStaticMesh()->GetRenderData() == nullptr)
 	{
 		return;
 	}
