Add descriptions on EDataLayerRuntimeState to help users (Blueprint, Sequencer or Data Layer itself).

#rnx

diff --git a/Engine/Source/Runtime/Engine/Public/WorldPartition/DataLayer/DataLayerInstance.h b/Engine/Source/Runtime/Engine/Public/WorldPartition/DataLayer/DataLayerInstance.h
--- a/Engine/Source/Runtime/Engine/Public/WorldPartition/DataLayer/DataLayerInstance.h
+++ b/Engine/Source/Runtime/Engine/Public/WorldPartition/DataLayer/DataLayerInstance.h
@@ -17,8 +17,13 @@ class IStreamingGenerationErrorHandler;
 UENUM(BlueprintType)
 enum class EDataLayerRuntimeState : uint8
 {
+	// Unloaded
 	Unloaded,
+	
+	// Loaded (meaning loaded but not visible)
 	Loaded,
+
+	// Activated (meaning loaded and visible)
 	Activated
 };
 
