Texture streaming factor out UTexture::IsPossibleToStream logic
add bool arg to GetNumNonStreamingMips so it knows if its being called on non-streamable textures
add more asserts for sanity checking

diff --git a/Engine/Source/Runtime/Engine/Classes/Engine/Texture.h b/Engine/Source/Runtime/Engine/Classes/Engine/Texture.h
--- a/Engine/Source/Runtime/Engine/Classes/Engine/Texture.h
+++ b/Engine/Source/Runtime/Engine/Classes/Engine/Texture.h
@@ -994,7 +994,7 @@ public:
 #endif
 
 	/** Return the number of mips that are not streamable. */
-	int32 GetNumNonStreamingMips() const;
+	int32 GetNumNonStreamingMips(bool bIsStreamingPossible) const;
 	/** Return the number of mips that streamable but not optional. */
 	int32 GetNumNonOptionalMips() const;
 	/** Return true if at least one mip can be loaded either from DDC or disk. */
@@ -1787,9 +1787,14 @@ public:
 	/** Sets the minimum number of mips that must be resident in memory (cannot be streamed). */
 	static void SetMinTextureResidentMipCount(int32 InMinTextureResidentMipCount);
 
+	/** Do the Texture properties make it a possible candidate for streaming.  Can be called by editor or runtime.
+	If true, does not mean it will actually be streamed, but if false will definitely not be streamed. */
+	bool IsPossibleToStream() const;
+
 #if WITH_EDITOR
-	/** Called by ULevel::MarkNoStreamableTexturesPrimitiveComponents when cooking level. */
-	bool IsCandidateForTextureStreaming(const ITargetPlatform* InTargetPlatform) const;
+	/** Called by ULevel::MarkNoStreamableTexturesPrimitiveComponents when cooking level.
+	Checks bCookedIsStreamable.  Return false for VT. */
+	bool IsCandidateForTextureStreamingOnPlatformDuringCook(const ITargetPlatform* InTargetPlatform) const;
 
 	/** Get the largest allowed dimension of non-VT texture
 	* this is not for the current RHI (which may have a lower limit), this is for a Texture in general
diff --git a/Engine/Source/Runtime/Engine/Private/Components/PrimitiveComponent.cpp b/Engine/Source/Runtime/Engine/Private/Components/PrimitiveComponent.cpp
--- a/Engine/Source/Runtime/Engine/Private/Components/PrimitiveComponent.cpp
+++ b/Engine/Source/Runtime/Engine/Private/Components/PrimitiveComponent.cpp
@@ -1416,7 +1416,7 @@ void UPrimitiveComponent::PreSave(FObjectPreSaveContext ObjectSaveContext)
 		bool bHasStreamableTextures = false;
 		for (const UTexture* Texture : Textures)
 		{
-			if (Texture && Texture->IsCandidateForTextureStreaming(ObjectSaveContext.GetTargetPlatform()))
+			if (Texture && Texture->IsCandidateForTextureStreamingOnPlatformDuringCook(ObjectSaveContext.GetTargetPlatform()))
 			{
 				bHasStreamableTextures = true;
 				break;
diff --git a/Engine/Source/Runtime/Engine/Private/Materials/MaterialInterface.cpp b/Engine/Source/Runtime/Engine/Private/Materials/MaterialInterface.cpp
--- a/Engine/Source/Runtime/Engine/Private/Materials/MaterialInterface.cpp
+++ b/Engine/Source/Runtime/Engine/Private/Materials/MaterialInterface.cpp
@@ -1245,7 +1245,7 @@ void UMaterialInterface::SortTextureStreamingData(bool bForceSort, bool bFinalSo
 			// Note that because we are cooking, UStreamableRenderAsset::bIsStreamable is not reliable here.
 			for (const UTexture* UnfilteredTexture : UnfilteredUsedTextures)
 			{
-				if (UnfilteredTexture && !UnfilteredTexture->NeverStream && UnfilteredTexture->LODGroup != TEXTUREGROUP_UI && UnfilteredTexture->MipGenSettings != TMGS_NoMipmaps)
+				if (UnfilteredTexture && UnfilteredTexture->IsPossibleToStream() )
 				{
 					UsedTextures.Add(UnfilteredTexture);
 				}
diff --git a/Engine/Source/Runtime/Engine/Private/Streaming/StreamingTexture.cpp b/Engine/Source/Runtime/Engine/Private/Streaming/StreamingTexture.cpp
--- a/Engine/Source/Runtime/Engine/Private/Streaming/StreamingTexture.cpp
+++ b/Engine/Source/Runtime/Engine/Private/Streaming/StreamingTexture.cpp
@@ -218,6 +218,9 @@ void FStreamingRenderAsset::UpdateDynamicData(const int32* NumStreamedMips, int3
 		}
 		NoRefLODBias = LocalNoRefBias;
 
+		check( ResourceState.MaxNumLODs >= ResourceState.NumNonOptionalLODs );
+		check( ResourceState.NumNonOptionalLODs >= ResourceState.NumNonStreamingLODs );
+
 		// If the optional mips are not available, or if we shouldn't load them now, clamp the possible mips requested. 
 		// (when the non-optional mips are not yet loaded, loading optional mips generates cross files requests).
 		// This is not bullet proof though since the texture/mesh could have a pending stream-out request.
@@ -230,6 +233,8 @@ void FStreamingRenderAsset::UpdateDynamicData(const int32* NumStreamedMips, int3
 			MaxAllowedMips = FMath::Clamp<int32>(ResourceState.MaxNumLODs - LODBias, ResourceState.NumNonStreamingLODs, ResourceState.MaxNumLODs);
 		}
 	
+		check( MaxAllowedMips >= ResourceState.NumNonStreamingLODs );
+
 		check(LODGroup < NumLODGroups);
 		if (NumStreamedMips[LODGroup] > 0)
 		{
diff --git a/Engine/Source/Runtime/Engine/Private/Streaming/TextureStreamingBuild.cpp b/Engine/Source/Runtime/Engine/Private/Streaming/TextureStreamingBuild.cpp
--- a/Engine/Source/Runtime/Engine/Private/Streaming/TextureStreamingBuild.cpp
+++ b/Engine/Source/Runtime/Engine/Private/Streaming/TextureStreamingBuild.cpp
@@ -165,7 +165,12 @@ bool GetTextureIsStreamable(const UTexture& Texture)
 bool GetTextureIsStreamableOnPlatform(const UTexture& Texture, const ITargetPlatform& TargetPlatform)
 {
 	const bool bPlatformSupportsTextureStreaming = TargetPlatform.SupportsFeature(ETargetPlatformFeatures::TextureStreaming);
-	
+	if ( ! bPlatformSupportsTextureStreaming )
+	{
+		// IsCandidateForTextureStreamingOnPlatformDuringCook also checks this
+		return false;
+	}
+
 	bool bStreamable = false;
 	if (Texture.IsA(UTexture2DArray::StaticClass()))
 	{
@@ -180,7 +185,7 @@ bool GetTextureIsStreamableOnPlatform(const UTexture& Texture, const ITargetPlat
 		bStreamable = true;
 	}
 
-	bStreamable &= Texture.IsCandidateForTextureStreaming(&TargetPlatform);
+	bStreamable &= Texture.IsCandidateForTextureStreamingOnPlatformDuringCook(&TargetPlatform);
 	return bStreamable;
 }
 
diff --git a/Engine/Source/Runtime/Engine/Private/Texture.cpp b/Engine/Source/Runtime/Engine/Private/Texture.cpp
--- a/Engine/Source/Runtime/Engine/Private/Texture.cpp
+++ b/Engine/Source/Runtime/Engine/Private/Texture.cpp
@@ -366,6 +366,10 @@ bool UTexture::CanEditChange(const FProperty* InProperty) const
 }
 
 // we're in WITH_EDITOR but not sure that's right, maybe move out?
+// Beware: while ValidateSettingsAfterImportOrEdit should have been called on all Textures,
+//	 it is not called on load at runtime
+//   and it is not always called from dynamically generated textures
+//	so you must not rely on the rules it sets up being true!
 void UTexture::ValidateSettingsAfterImportOrEdit(bool * pRequiresNotifyMaterials)
 {
 	bool bRequiresNotifyMaterialsDummy = false;
@@ -380,101 +384,105 @@ void UTexture::ValidateSettingsAfterImportOrEdit(bool * pRequiresNotifyMaterials
 
 	if (Source.IsValid()) // we can have an empty source if the last source in a texture2d array is removed via the editor.
 	{
-	if ( MipGenSettings == TMGS_LeaveExistingMips && PowerOfTwoMode != ETexturePowerOfTwoSetting::None )
-	{
-		// power of 2 pads not allowed with LeaveExistingMips
-		UE_LOG(LogTexture, Display, TEXT("Power of 2 padding cannot be used with LeaveExistingMips, disabled. (%s)"), *GetName());
+		if ( MipGenSettings == TMGS_LeaveExistingMips && PowerOfTwoMode != ETexturePowerOfTwoSetting::None )
+		{
+			// power of 2 pads not allowed with LeaveExistingMips
+			UE_LOG(LogTexture, Display, TEXT("Power of 2 padding cannot be used with LeaveExistingMips, disabled. (%s)"), *GetName());
 
-		PowerOfTwoMode = ETexturePowerOfTwoSetting::None;
-	}
+			PowerOfTwoMode = ETexturePowerOfTwoSetting::None;
+		}
 
-	// IsPowerOfTwo only checks XY :
-	bool bIsPowerOfTwo = Source.IsPowerOfTwo();
-	if ( ! FMath::IsPowerOfTwo(Source.GetVolumeSizeZ()) )
-	{
-		bIsPowerOfTwo = false;
-	}
-	if ( PowerOfTwoMode != ETexturePowerOfTwoSetting::None )
-	{
-		bIsPowerOfTwo = true;
-	}
+		// IsPowerOfTwo only checks XY :
+		bool bIsPowerOfTwo = Source.IsPowerOfTwo();
+		if ( ! FMath::IsPowerOfTwo(Source.GetVolumeSizeZ()) )
+		{
+			bIsPowerOfTwo = false;
+		}
+		if ( PowerOfTwoMode != ETexturePowerOfTwoSetting::None )
+		{
+			bIsPowerOfTwo = true;
+		}
 
-	if ( ! bIsPowerOfTwo )
-	{
-		// streaming only supports power of 2 mips
-		// due to failure to compensate for the GPU row pitch
-		// it only works for mips that naturally have the required 256 pitch
-		// so mip levels >= 256 and power of 2 only
-		// (this used to be in Texture2D.cpp)
-		// this could be fixed and then streaming nonpow2 could be allowed
-		// see WarnRequiresTightPackedMip
-		NeverStream = true;	
-	}
-	
-	int32 MaxDimension = FMath::Max( Source.GetSizeX() , Source.GetSizeY() );
-	bool bLargeTextureMustBeVT = MaxDimension > GetMaximumDimensionOfNonVT();
+		// Downscale can violate IsPow2, but it only acts when NoMipMaps, so it's moot
 
-	if ( bLargeTextureMustBeVT && ! VirtualTextureStreaming && MaxTextureSize == 0 )
-	{
-		static const auto CVarVirtualTexturesEnabled = IConsoleManager::Get().FindTConsoleVariableDataInt(TEXT("r.VirtualTextures"));
-		check( CVarVirtualTexturesEnabled != nullptr );
+		if ( ! bIsPowerOfTwo )
+		{
+			// streaming only supports power of 2 mips
+			// due to failure to compensate for the GPU row pitch
+			// it only works for mips that naturally have the required 256 pitch
+			// so mip levels >= 256 and power of 2 only
+			// (this used to be in Texture2D.cpp)
+			// see WarnRequiresTightPackedMip
+			// there are other issues with streaming non-pow2
+			//  all streamable levels must be valid textures, so block-of-4 alignment for BCN
+			//  is easiest to guarantee if the source is pow2
+			NeverStream = true;	
+		}
+	
+		int32 MaxDimension = FMath::Max( Source.GetSizeX() , Source.GetSizeY() );
+		bool bLargeTextureMustBeVT = MaxDimension > GetMaximumDimensionOfNonVT();
 
-		if ( CVarVirtualTexturesEnabled->GetValueOnAnyThread() )
+		if ( bLargeTextureMustBeVT && ! VirtualTextureStreaming && MaxTextureSize == 0 )
 		{
-			if ( GetTextureClass() == ETextureClass::TwoD )
+			static const auto CVarVirtualTexturesEnabled = IConsoleManager::Get().FindTConsoleVariableDataInt(TEXT("r.VirtualTextures"));
+			check( CVarVirtualTexturesEnabled != nullptr );
+
+			if ( CVarVirtualTexturesEnabled->GetValueOnAnyThread() )
 			{
+				if ( GetTextureClass() == ETextureClass::TwoD )
+				{
 				UE_LOG(LogTexture, Display, TEXT("Large Texture %s Dimension=%d changed to VT; to disable VT set MaxTextureSize first"), *GetName(),MaxDimension);
-				VirtualTextureStreaming = true;
-				bRequiresNotifyMaterials = true;
+					VirtualTextureStreaming = true;
+					bRequiresNotifyMaterials = true;
+				}
+				else
+				{
+					UE_LOG(LogTexture, Warning, TEXT("Large Texture %s Dimension=%d needs to be VT but is not 2d, changing MaxTextureSize"), *GetName(),MaxDimension);
+				
+					// GetMaximumDimension is the max size for this texture type on the current RHI
+					MaxTextureSize = GetMaximumDimension();
+				}
 			}
 			else
 			{
-				UE_LOG(LogTexture, Warning, TEXT("Large Texture %s Dimension=%d needs to be VT but is not 2d, changing MaxTextureSize"), *GetName(),MaxDimension);
-				
+				UE_LOG(LogTexture, Warning, TEXT("Large Texture %s Dimension=%d must be VT but VirtualTextures are disabled, changing MaxTextureSize"), *GetName(),MaxDimension);
+
 				// GetMaximumDimension is the max size for this texture type on the current RHI
 				MaxTextureSize = GetMaximumDimension();
 			}
 		}
-		else
-		{
-			UE_LOG(LogTexture, Warning, TEXT("Large Texture %s Dimension=%d must be VT but VirtualTextures are disabled, changing MaxTextureSize"), *GetName(),MaxDimension);
-
-			// GetMaximumDimension is the max size for this texture type on the current RHI
-			MaxTextureSize = GetMaximumDimension();
-		}
-	}
 	
-	if (VirtualTextureStreaming)
-	{
-		if (!bIsPowerOfTwo)
+		if (VirtualTextureStreaming)
 		{
-			if ( bLargeTextureMustBeVT )
+			if (!bIsPowerOfTwo)
 			{
-				UE_LOG(LogTexture, Warning, TEXT("Large VT \"%s\", must be padded to power-of-2 for VT support (%dx%d)"), *GetName(), Source.GetSizeX(),Source.GetSizeY());
-				// VT nonpow2 will fail to build
-				// force it into a state that will succeed? or just let it fail?
-				// you can either pad to pow2 or set MaxTextureSize and turn off VT
+				if ( bLargeTextureMustBeVT )
+				{
+					UE_LOG(LogTexture, Warning, TEXT("Large VT \"%s\", must be padded to power-of-2 for VT support (%dx%d)"), *GetName(), Source.GetSizeX(),Source.GetSizeY());
+					// VT nonpow2 will fail to build
+					// force it into a state that will succeed? or just let it fail?
+					// you can either pad to pow2 or set MaxTextureSize and turn off VT
 				PowerOfTwoMode = ETexturePowerOfTwoSetting::PadToPowerOfTwo;
-			}
-			else
-			{
-				UE_LOG(LogTexture, Warning, TEXT("VirtualTextureStreaming not supported for \"%s\", texture size is not a power-of-2"), *GetName());
-				VirtualTextureStreaming = false;
-				bRequiresNotifyMaterials = true;
+				}
+				else
+				{
+					UE_LOG(LogTexture, Warning, TEXT("VirtualTextureStreaming not supported for \"%s\", texture size is not a power-of-2"), *GetName());
+					VirtualTextureStreaming = false;
+					bRequiresNotifyMaterials = true;
+				}
 			}
 		}
-	}
 
-	// Make sure settings are correct for LUT textures.
-	if(LODGroup == TEXTUREGROUP_ColorLookupTable)
-	{
-		if ( MipGenSettings != TMGS_NoMipmaps || SRGB != false )
+		// Make sure settings are correct for LUT textures.
+		if(LODGroup == TEXTUREGROUP_ColorLookupTable)
 		{
-			MipGenSettings = TMGS_NoMipmaps;
-			SRGB = false;
-			bRequiresNotifyMaterials = true;
+			if ( MipGenSettings != TMGS_NoMipmaps || SRGB != false )
+			{
+				MipGenSettings = TMGS_NoMipmaps;
+				SRGB = false;
+				bRequiresNotifyMaterials = true;
+			}
 		}
-	}
 	} // end if valid source
 #endif // #if WITH_EDITORONLY_DATA
 
@@ -1318,19 +1326,63 @@ const TArray<UAssetUserData*>* UTexture::GetAssetUserDataArray() const
 	return &ToRawPtrTArrayUnsafe(AssetUserData);
 }
 
+
+bool UTexture::IsPossibleToStream() const
+{
+	if ( NeverStream || LODGroup == TEXTUREGROUP_UI )
+	{
+		return false;
+	}
+
+	#if WITH_EDITORONLY_DATA
+	if ( MipGenSettings == TMGS_NoMipmaps )
+	{
+		return false;
+	}
+
+	// VirtualTextureStreaming can be true here and we will still stream if VT is disabled
+	
+	if ( Source.IsValid() )
+	{
+		// should have set NeverStream for nonpow2
+		// note: this is not the case for all old textures
+		// ValidateSettingsAfterImportOrEdit makes sure this is true
+		//check( Source.IsPowerOfTwo() || PowerOfTwoMode != ETexturePowerOfTwoSetting::None )
+		
+		// duplicate the checks done for NeverStream :
+
+		// IsPowerOfTwo only checks XY :
+		bool bIsPowerOfTwo = Source.IsPowerOfTwo();
+		if ( ! FMath::IsPowerOfTwo(Source.GetVolumeSizeZ()) )
+		{
+			bIsPowerOfTwo = false;
+		}
+		if ( PowerOfTwoMode != ETexturePowerOfTwoSetting::None )
+		{
+			bIsPowerOfTwo = true;
+		}
+		if ( ! bIsPowerOfTwo )
+		{
+			// NeverStream should have been set
+			return false;
+		}
+	}
+	#endif
+
+	return true;
+}
+
 #if WITH_EDITOR
 // Based on target platform, returns wether texture is candidate to be streamed.
 // This method is used to decide if PrimitiveComponent's bHasNoStreamableTextures flag can be set to true.
 // See ULevel::MarkNoStreamableTexturesPrimitiveComponents for details.
-bool UTexture::IsCandidateForTextureStreaming(const ITargetPlatform* InTargetPlatform) const
+bool UTexture::IsCandidateForTextureStreamingOnPlatformDuringCook(const ITargetPlatform* InTargetPlatform) const
 {
 	const bool bIsVirtualTextureStreaming = InTargetPlatform->SupportsFeature(ETargetPlatformFeatures::VirtualTextureStreaming) ? VirtualTextureStreaming : false;
 	const bool bIsCandidateForTextureStreaming = InTargetPlatform->SupportsFeature(ETargetPlatformFeatures::TextureStreaming) && !bIsVirtualTextureStreaming;
 
 	if (bIsCandidateForTextureStreaming &&
-		!NeverStream &&
-		LODGroup != TEXTUREGROUP_UI &&
-		MipGenSettings != TMGS_NoMipmaps)
+		IsPossibleToStream())
 	{
 		// If bCookedIsStreamable flag was previously computed, use it.
 		// This is computed _after_ the derived data is cached. It's only Reset() and used in the SerializeCookedPlatformData
@@ -1364,23 +1416,24 @@ FStreamableRenderResourceState UTexture::GetResourcePostInitState(const FTexture
 			return FMath::Min<int32>(PlatformData->Mips.Num() - ExpectedAssetLODBias, MaxRuntimeMipCount);
 		}
 	}();
+	
+	bool bTextureIsStreamable = IsPossibleToStream();
 
 	const int32 NumOfNonOptionalMips = FMath::Min<int32>(NumMips, PlatformData->GetNumNonOptionalMips());
-	const int32 NumOfNonStreamingMips = FMath::Min<int32>(NumMips, PlatformData->GetNumNonStreamingMips());
+	const int32 NumOfNonStreamingMips = FMath::Min<int32>(NumMips, PlatformData->GetNumNonStreamingMips(bTextureIsStreamable));
 	// Optional mips must be streaming mips :
 	check( NumOfNonOptionalMips >= NumOfNonStreamingMips );
 
+	if ( NumOfNonStreamingMips == NumMips )
+	{
+		bTextureIsStreamable = false;
+	}
+
 	const int32 AssetMipIdxForResourceFirstMip = FMath::Max<int32>(0, PlatformData->Mips.Num() - NumMips);
 
 	bool bMakeStreamable = false;
 	int32 NumRequestedMips = 0;
 	
-	// same conditions as IsCandidateForTextureStreaming which TextureDerivedData will use :
-	bool bTextureIsStreamable =
-	  (!NeverStream && 
-		NumOfNonStreamingMips < NumMips && 
-		LODGroup != TEXTUREGROUP_UI ); 
-
 #if PLATFORM_SUPPORTS_TEXTURE_STREAMING
 	if ( bTextureIsStreamable )
 	{
@@ -3135,9 +3188,9 @@ FName GetDefaultTextureFormatName( const ITargetPlatform* TargetPlatform, const
 			else
 			{
 				// @todo Oodle : this is not the best possible selection of output format from SourceFormat
-				// R16F and R32F is available but not used here even if their TC_ would have chosen them!
-				TextureFormatName = NameRGBA16F;
-			}
+			// R16F and R32F is available but not used here even if their TC_ would have chosen them!
+			TextureFormatName = NameRGBA16F;
+		}
 		}
 		else if (SourceFormat == TSF_G16)
 		{
diff --git a/Engine/Source/Runtime/Engine/Private/TextureDerivedData.cpp b/Engine/Source/Runtime/Engine/Private/TextureDerivedData.cpp
--- a/Engine/Source/Runtime/Engine/Private/TextureDerivedData.cpp
+++ b/Engine/Source/Runtime/Engine/Private/TextureDerivedData.cpp
@@ -2065,7 +2065,7 @@ bool FTexturePlatformData::TryLoadMipsWithSizes(int32 FirstMipToLoad, void** Out
 	return true;
 }
 
-int32 FTexturePlatformData::GetNumNonStreamingMips() const
+int32 FTexturePlatformData::GetNumNonStreamingMips(bool bIsStreamingPossible) const
 {
 	if (CanUseCookedDataPath())
 	{
@@ -2090,10 +2090,15 @@ int32 FTexturePlatformData::GetNumNonStreamingMips() const
 		}
 		else
 		{
+			if ( ! bIsStreamingPossible )
+			{
+				check( NumNonStreamingMips == Mips.Num() );
+			}
+
 			return NumNonStreamingMips;
 		}
 	}
-	else if (Mips.Num() <= 1)
+	else if (Mips.Num() <= 1 || !bIsStreamingPossible)
 	{
 		return Mips.Num();
 	}
@@ -2101,6 +2106,7 @@ int32 FTexturePlatformData::GetNumNonStreamingMips() const
 	{
 		int32 MipCount = Mips.Num();
 		int32 NumNonStreamingMips = 1;
+		// MipCount >= 2 and bIsStreamingPossible
 
 		// Take in to account the min resident limit.
 		NumNonStreamingMips = FMath::Max(NumNonStreamingMips, (int32)GetNumMipsInTail());
@@ -2113,17 +2119,15 @@ int32 FTexturePlatformData::GetNumNonStreamingMips() const
 			// ensure the top non-streamed mip size is >= BlockSize (and a multiple of block size!)
 			// @todo Oodle : only do for BCN, not for ASTC
 			
+			// note: this is not right for non pow 2; NeverStream should set !bIsStreamingPossible in that case
 			if ( FMath::IsPowerOfTwo(Mips[0].SizeX) && FMath::IsPowerOfTwo(Mips[0].SizeY) )
 			{
-				// note: this is not right for non pow 2 (in that case you should check all streaming mips are aligned to blocksize)
 				NumNonStreamingMips = FMath::Max<int32>(NumNonStreamingMips, MipCount - FPlatformMath::FloorLog2(Mips[0].SizeX / BlockSizeX));
 				NumNonStreamingMips = FMath::Max<int32>(NumNonStreamingMips, MipCount - FPlatformMath::FloorLog2(Mips[0].SizeY / BlockSizeY));
 			}
 			else
 			{
-				// non-pow2 should be NeverStream
-				// but they call into here anyway
-				// and actually CAN stream via cinematic lod bias (in Editor, but not in cooked/game)
+				// should have !bIsStreamingPossible, so we should not hit this branch, but it's not reliable
 				NumNonStreamingMips = MipCount;
 			}
 		}
@@ -2490,7 +2494,7 @@ static void SerializePlatformData(
 			if (bStreamable)
 #endif
 			{
-				int32 NumNonStreamingMips = PlatformData->GetNumNonStreamingMips();
+				int32 NumNonStreamingMips = PlatformData->GetNumNonStreamingMips(true);
 				// NumMips has been reduced by FirstMipToSerialize (LODBias)
 				NumNonStreamingMips = FMath::Min(NumNonStreamingMips,NumMips);
 				// NumNonStreamingMips is not serialized
diff --git a/Engine/Source/Runtime/Engine/Private/TextureDerivedDataTask.cpp b/Engine/Source/Runtime/Engine/Private/TextureDerivedDataTask.cpp
--- a/Engine/Source/Runtime/Engine/Private/TextureDerivedDataTask.cpp
+++ b/Engine/Source/Runtime/Engine/Private/TextureDerivedDataTask.cpp
@@ -1484,9 +1484,9 @@ void FTextureCacheDerivedDataWorker::DoWork()
 			}
 			else
 			{
-				const bool bDisableStreaming = (Texture.NeverStream || Texture.LODGroup == TEXTUREGROUP_UI);
+				const bool bDisableStreaming = ! Texture.IsPossibleToStream();
 				const int32 FirstMipToLoad = FirstResidentMipIndex;
-				const int32 FirstNonStreamingMipIndex = DerivedData->Mips.Num() - DerivedData->GetNumNonStreamingMips();
+				const int32 FirstNonStreamingMipIndex = DerivedData->Mips.Num() - DerivedData->GetNumNonStreamingMips(!bDisableStreaming);
 				const int32 FirstMipToPrefetch = IsInGameThread() ? FirstMipToLoad : bDisableStreaming ? 0 : FirstNonStreamingMipIndex;
 				bSucceeded = TryCacheStreamingMips(FirstMipToLoad, FirstMipToPrefetch);
 				if (!bSucceeded)
