Silence a static analysis warning by zero-initializing the array.

#rnx

diff --git a/Engine/Source/Runtime/Engine/Private/SkeletalMeshLODRenderData.cpp b/Engine/Source/Runtime/Engine/Private/SkeletalMeshLODRenderData.cpp
--- a/Engine/Source/Runtime/Engine/Private/SkeletalMeshLODRenderData.cpp
+++ b/Engine/Source/Runtime/Engine/Private/SkeletalMeshLODRenderData.cpp
@@ -377,7 +377,7 @@ void FSkeletalMeshLODRenderData::BuildFromLODModel(const FSkeletalMeshLODModel*
 	{
 		// Re-normalize the weights for 8-bits to ensure that they are distributed using the old algorithm
 		// from MeshUtilities.cpp
-		uint8 InfluenceWeights[MAX_TOTAL_INFLUENCES];
+		uint8 InfluenceWeights[MAX_TOTAL_INFLUENCES] = {0};
 
 		MaxBoneInfluences = 0;
 		for (FSoftSkinVertex& Vertex: Vertices)
