Fix CopyProperties from broadcasting object change for object still being async loaded

diff --git a/Engine/Source/Runtime/Engine/Private/UnrealEngine.cpp b/Engine/Source/Runtime/Engine/Private/UnrealEngine.cpp
--- a/Engine/Source/Runtime/Engine/Private/UnrealEngine.cpp
+++ b/Engine/Source/Runtime/Engine/Private/UnrealEngine.cpp
@@ -16788,9 +16788,11 @@ void UEngine::CopyPropertiesForUnrelatedObjects(UObject* OldObject, UObject* New
 		DumpObject(*FString::Printf(TEXT("CopyPropertiesForUnrelatedObjects: New (%s)"), *NewObject->GetFullName()), NewObject);
 	}
 
-	// Now notify any tools that aren't already updated via the FArchiveReplaceObjectRef path
-	if (Params.bNotifyObjectReplacement && GEngine != nullptr)
+	// Now notify any tools that aren't already updated via the FArchiveReplaceObjectRef path unless the OldObject is still being async loaded
+	if (Params.bNotifyObjectReplacement && GEngine != nullptr &&
+		!OldObject->HasAnyInternalFlags(EInternalObjectFlags::Async | EInternalObjectFlags::AsyncLoading))
 	{
+		check(IsInGameThread());
 		GEngine->NotifyToolsOfObjectReplacement(ReferenceReplacementMap);
 	}
 }
