Make EdGraphPin globals thread-safe

diff --git a/Engine/Source/Runtime/Engine/Private/EdGraph/EdGraphPin.cpp b/Engine/Source/Runtime/Engine/Private/EdGraph/EdGraphPin.cpp
--- a/Engine/Source/Runtime/Engine/Private/EdGraph/EdGraphPin.cpp
+++ b/Engine/Source/Runtime/Engine/Private/EdGraph/EdGraphPin.cpp
@@ -45,21 +45,27 @@ public:
 	static void Add(UEdGraphPin* PinToDelete)
 	{
 		FPinDeletionQueue* PinDeletionQueue = Get();
-	
+		FWriteScopeLock WriteLock(PinDeletionQueue->PinsToDeleteLock);
 		PinDeletionQueue->PinsToDelete.Add(PinToDelete);
 	}
 
 	virtual void Tick(float DeltaTime) override
 	{
-		for (UEdGraphPin* Pin : PinsToDelete)
+		TArray<UEdGraphPin*> LocalPinsToDelete;
+		{
+			FWriteScopeLock ReadLock(PinsToDeleteLock);
+			Swap(LocalPinsToDelete, PinsToDelete);
+		}
+
+		for (UEdGraphPin* Pin : LocalPinsToDelete)
 		{
 			delete Pin;
 		}
-		PinsToDelete.Reset();
 	}
 
 	virtual bool IsTickable() const override
 	{
+		FReadScopeLock ReadLock(PinsToDeleteLock);
 		return (PinsToDelete.Num() > 0);
 	}
 
@@ -74,7 +80,7 @@ public:
 private:
 
 	FPinDeletionQueue() = default;
-
+	mutable FRWLock PinsToDeleteLock;
 	TArray<UEdGraphPin*> PinsToDelete;
 
 };
@@ -456,11 +462,9 @@ struct FPinResolveId
 // PinHelpers
 namespace PinHelpers
 {
-	static TMap<FPinResolveId, TArray<FUnresolvedPinData>> UnresolvedPins;
-	static TMap<TWeakObjectPtr<UEdGraphPin_Deprecated>, FGuid> DeprecatedPinToNewPinGUIDMap;
-	static TMap<TWeakObjectPtr<UEdGraphPin_Deprecated>, TArray<FUnresolvedPinData>> UnresolvedDeprecatedPins;
-
-	static uint64 NumPinsInMemory = 0;
+	static thread_local TMap<FPinResolveId, TArray<FUnresolvedPinData>> UnresolvedPins;
+	static thread_local TMap<TWeakObjectPtr<UEdGraphPin_Deprecated>, FGuid> DeprecatedPinToNewPinGUIDMap;
+	static thread_local TMap<TWeakObjectPtr<UEdGraphPin_Deprecated>, TArray<FUnresolvedPinData>> UnresolvedDeprecatedPins;
 
 	static const TCHAR ExportTextPropDelimiter = ',';
 
@@ -1402,9 +1406,6 @@ void UEdGraphPin::MarkAsGarbage()
 void UEdGraphPin::ShutdownVerification()
 {
 	Purge();
-	// There's a static UEdGraphPin in UEdGraphPin::ExportTextItem so if that code
-	// has run we'll have a single pin 'in memory' on shutdown
-	ensure(PinHelpers::NumPinsInMemory == 0 || PinHelpers::NumPinsInMemory == 1);
 }
 
 void UEdGraphPin::Purge()
@@ -1448,7 +1449,6 @@ UEdGraphPin::UEdGraphPin(UEdGraphNode* InOwningNode, const FGuid& PinIdGuid)
 	, PersistentGuid()
 #endif
 {
-	PinHelpers::NumPinsInMemory++;
 #ifdef TRACK_PINS
 	PinAllocationTracking.Emplace(this, InOwningNode ? InOwningNode->GetName() : FString(TEXT("UNOWNED")));
 #endif //TRACK_PINS
@@ -1456,8 +1456,6 @@ UEdGraphPin::UEdGraphPin(UEdGraphNode* InOwningNode, const FGuid& PinIdGuid)
 
 UEdGraphPin::~UEdGraphPin()
 {
-	check(PinHelpers::NumPinsInMemory > 0);
-	PinHelpers::NumPinsInMemory--;
 #ifdef TRACK_PINS
 	PinAllocationTracking.RemoveAll([this](const TPair<UEdGraphPin*, FString>& Entry) { return Entry.Key == this; });
 #endif //TRACK_PINS
