Only compute PSO when static meshes are ready

diff --git a/Engine/Source/Runtime/Engine/Private/Components/StaticMeshComponent.cpp b/Engine/Source/Runtime/Engine/Private/Components/StaticMeshComponent.cpp
--- a/Engine/Source/Runtime/Engine/Private/Components/StaticMeshComponent.cpp
+++ b/Engine/Source/Runtime/Engine/Private/Components/StaticMeshComponent.cpp
@@ -2124,8 +2124,12 @@ void UStaticMeshComponent::PostLoad()
 	// Initialize the resources for the freshly loaded component.
 	InitResources();
 
-	// Precache PSOs for the used materials
-	PrecachePSOs();
+	// If currently compiling, those will be called once the static mesh compilation has finished
+	if (GetStaticMesh() && !GetStaticMesh()->IsCompiling())
+	{
+		// Precache PSOs for the used materials
+		PrecachePSOs();
+	}
 }
 
 bool UStaticMeshComponent::IsPostLoadThreadSafe() const
@@ -2989,6 +2993,8 @@ void UStaticMeshComponent::PostStaticMeshCompilation()
 
 	FixupOverrideColorsIfNecessary(true);
 
+	PrecachePSOs();
+
 	UpdateCollisionFromStaticMesh();
 
 	RecreatePhysicsState();
