Properly support morph vertex buffers when render data is streamed

diff --git a/Engine/Source/Runtime/Engine/Private/Rendering/MorphTargetVertexInfoBuffers.cpp b/Engine/Source/Runtime/Engine/Private/Rendering/MorphTargetVertexInfoBuffers.cpp
--- a/Engine/Source/Runtime/Engine/Private/Rendering/MorphTargetVertexInfoBuffers.cpp
+++ b/Engine/Source/Runtime/Engine/Private/Rendering/MorphTargetVertexInfoBuffers.cpp
@@ -465,3 +465,40 @@ bool FMorphTargetVertexInfoBuffers::IsPlatformShaderSupported(EShaderPlatform Sh
 {
 	return IsFeatureLevelSupported(ShaderPlatform, ERHIFeatureLevel::SM5);
 }
+
+template <bool bRenderThread>
+FBufferRHIRef FMorphTargetVertexInfoBuffers::CreateMorphRHIBuffer_Internal()
+{
+	uint32 SizeInBytes = MorphData.Num() * sizeof(uint32);
+
+	if(SizeInBytes > 0)
+	{
+		const EBufferUsageFlags BufferFlags = BUF_Static | BUF_ByteAddressBuffer | BUF_ShaderResource;
+
+		// Create the index buffer.
+		FRHIResourceCreateInfo CreateInfo(TEXT("MorphData"), &MorphData);
+		if (bRenderThread)
+		{
+			FBufferRHIRef BufferRHIRef;
+			BufferRHIRef = RHICreateStructuredBuffer(sizeof(uint32), SizeInBytes, BufferFlags, ERHIAccess::SRVMask, CreateInfo);
+
+			return BufferRHIRef;
+		}
+		else
+		{
+			FRHIAsyncCommandList CommandList;
+			return CommandList->CreateBuffer(SizeInBytes, BufferFlags, sizeof(uint32), ERHIAccess::SRVMask, CreateInfo);
+		}
+	}
+	return nullptr;
+}
+
+FBufferRHIRef FMorphTargetVertexInfoBuffers::CreateMorphRHIBuffer_RenderThread()
+{
+	return CreateMorphRHIBuffer_Internal<true>();
+}
+
+FBufferRHIRef FMorphTargetVertexInfoBuffers::CreateMorphRHIBuffer_Async()
+{
+	return CreateMorphRHIBuffer_Internal<false>();
+}
diff --git a/Engine/Source/Runtime/Engine/Private/SkeletalMeshLODRenderData.cpp b/Engine/Source/Runtime/Engine/Private/SkeletalMeshLODRenderData.cpp
--- a/Engine/Source/Runtime/Engine/Private/SkeletalMeshLODRenderData.cpp
+++ b/Engine/Source/Runtime/Engine/Private/SkeletalMeshLODRenderData.cpp
@@ -30,13 +30,6 @@ static FAutoConsoleVariableRef CVarStripSkeletalMeshLodsBelowMinLod(
 	TEXT("If set will strip skeletal mesh LODs under the minimum renderable LOD for the target platform during cooking.")
 );
 
-int32 GAllowSkinnedMorphDataStreaming = 0;
-static FAutoConsoleVariableRef CVarAllowSkinnedMorphDataStreaming(
-	TEXT("r.Skinned.AllowSkinnedMorphDataStreaming"),
-	GAllowSkinnedMorphDataStreaming,
-	TEXT("Call BeginInitResource when the morph resources are finish streaming (DoFinishUpdate)")
-);
-
 namespace
 {
 	struct FReverseOrderBitArraysBySetBits
@@ -166,14 +159,6 @@ FArchive& operator<<(FArchive& Ar, FSkelMeshRenderSection& S)
 	return Ar;
 }
 
-void FSkeletalMeshLODRenderData::InitMorphResources()
-{
-	if (GAllowSkinnedMorphDataStreaming && !MorphTargetVertexInfoBuffers.IsRHIIntialized() && MorphTargetVertexInfoBuffers.IsMorphCPUDataValid() && MorphTargetVertexInfoBuffers.NumTotalBatches > 0)
-	{
-		BeginInitResource(&MorphTargetVertexInfoBuffers);
-	}
-}
-
 void FSkeletalMeshLODRenderData::InitResources(bool bNeedsVertexColors, int32 LODIndex, TArray<UMorphTarget*>& InMorphTargets, USkinnedAsset* Owner)
 {
 	if (bStreamedDataInlined)
@@ -231,7 +216,7 @@ void FSkeletalMeshLODRenderData::InitResources(bool bNeedsVertexColors, int32 LO
 		MorphTargetVertexInfoBuffers.InitMorphResources(GMaxRHIShaderPlatform, RenderSections, Owner->GetMorphTargets(), StaticVertexBuffers.StaticMeshVertexBuffer.GetNumVertices(), LODIndex, SkeletalMeshLODInfo->MorphTargetPositionErrorTolerance);
 	}
 	
-	if (!MorphTargetVertexInfoBuffers.IsRHIIntialized() && MorphTargetVertexInfoBuffers.IsMorphCPUDataValid() && MorphTargetVertexInfoBuffers.NumTotalBatches > 0)
+	if (!MorphTargetVertexInfoBuffers.IsRHIIntialized() && MorphTargetVertexInfoBuffers.IsMorphCPUDataValid() && (MorphTargetVertexInfoBuffers.NumTotalBatches > 0 || !bStreamedDataInlined))
 	{
 		BeginInitResource(&MorphTargetVertexInfoBuffers);
 	}
diff --git a/Engine/Source/Runtime/Engine/Private/Streaming/SkeletalMeshUpdate.cpp b/Engine/Source/Runtime/Engine/Private/Streaming/SkeletalMeshUpdate.cpp
--- a/Engine/Source/Runtime/Engine/Private/Streaming/SkeletalMeshUpdate.cpp
+++ b/Engine/Source/Runtime/Engine/Private/Streaming/SkeletalMeshUpdate.cpp
@@ -60,6 +60,7 @@ void FSkeletalMeshStreamIn::FIntermediateBuffers::CreateFromCPUData_RenderThread
 	SkinWeightVertexBuffer = LODResource.SkinWeightVertexBuffer.CreateRHIBuffer_RenderThread();
 	ClothVertexBuffer = LODResource.ClothVertexBuffer.CreateRHIBuffer_RenderThread();
 	IndexBuffer = LODResource.MultiSizeIndexContainer.CreateRHIBuffer_RenderThread();
+	MorphBuffer = LODResource.MorphTargetVertexInfoBuffers.CreateMorphRHIBuffer_RenderThread();
 }
 
 void FSkeletalMeshStreamIn::FIntermediateBuffers::CreateFromCPUData_Async(FSkeletalMeshLODRenderData& LODResource)
@@ -73,6 +74,7 @@ void FSkeletalMeshStreamIn::FIntermediateBuffers::CreateFromCPUData_Async(FSkele
 	SkinWeightVertexBuffer = LODResource.SkinWeightVertexBuffer.CreateRHIBuffer_Async();
 	ClothVertexBuffer = LODResource.ClothVertexBuffer.CreateRHIBuffer_Async();
 	IndexBuffer = LODResource.MultiSizeIndexContainer.CreateRHIBuffer_Async();
+	MorphBuffer = LODResource.MorphTargetVertexInfoBuffers.CreateMorphRHIBuffer_Async();
 }
 
 void FSkeletalMeshStreamIn::FIntermediateBuffers::SafeRelease()
@@ -86,6 +88,7 @@ void FSkeletalMeshStreamIn::FIntermediateBuffers::SafeRelease()
 	ClothVertexBuffer.SafeRelease();
 	IndexBuffer.SafeRelease();
 	AltSkinWeightVertexBuffers.Empty();
+	MorphBuffer.SafeRelease();
 }
 
 template <uint32 MaxNumUpdates>
@@ -99,6 +102,7 @@ void FSkeletalMeshStreamIn::FIntermediateBuffers::TransferBuffers(FSkeletalMeshL
 	LODResource.ClothVertexBuffer.InitRHIForStreaming(ClothVertexBuffer, Batcher);
 	LODResource.MultiSizeIndexContainer.InitRHIForStreaming(IndexBuffer, Batcher);
 	LODResource.SkinWeightProfilesData.InitRHIForStreaming(AltSkinWeightVertexBuffers, Batcher);
+	LODResource.MorphTargetVertexInfoBuffers.InitRHIForStreaming(MorphBuffer, Batcher);
 	SafeRelease();
 }
 
@@ -112,7 +116,8 @@ void FSkeletalMeshStreamIn::FIntermediateBuffers::CheckIsNull() const
 		&& !SkinWeightVertexBuffer.LookupVertexBufferRHI
 		&& !ClothVertexBuffer
 		&& !IndexBuffer
-		&& !AltSkinWeightVertexBuffers.Num());
+		&& !AltSkinWeightVertexBuffers.Num()
+		&& !MorphBuffer);
 }
 
 FSkeletalMeshStreamIn::FSkeletalMeshStreamIn(const USkeletalMesh* InMesh)
@@ -198,7 +203,6 @@ void FSkeletalMeshStreamIn::DoFinishUpdate(const FContext& Context)
 			{
 				FSkeletalMeshLODRenderData& LODResource = *Context.LODResourcesView[LODIndex];
 				LODResource.IncrementMemoryStats(Mesh->GetHasVertexColors());
-				LODResource.InitMorphResources();
 				IntermediateBuffersArray[LODIndex].TransferBuffers(LODResource, Batcher);
 			}
 		}
@@ -357,6 +361,7 @@ void FSkeletalMeshStreamOut::ReleaseBuffers(const FContext& Context)
 			LODResource.ClothVertexBuffer.ReleaseRHIForStreaming(Batcher);
 			LODResource.MultiSizeIndexContainer.ReleaseRHIForStreaming(Batcher);
 			LODResource.SkinWeightProfilesData.ReleaseRHIForStreaming(Batcher);
+			LODResource.MorphTargetVertexInfoBuffers.ReleaseRHIForStreaming(Batcher);
 
 			if (!FPlatformProperties::HasEditorOnlyData())
 			{
diff --git a/Engine/Source/Runtime/Engine/Private/Streaming/SkeletalMeshUpdate.h b/Engine/Source/Runtime/Engine/Private/Streaming/SkeletalMeshUpdate.h
--- a/Engine/Source/Runtime/Engine/Private/Streaming/SkeletalMeshUpdate.h
+++ b/Engine/Source/Runtime/Engine/Private/Streaming/SkeletalMeshUpdate.h
@@ -80,6 +80,7 @@ protected:
 		FBufferRHIRef ClothVertexBuffer;
 		FBufferRHIRef IndexBuffer;
 		TArray<TPair<FName, FSkinWeightRHIInfo>> AltSkinWeightVertexBuffers;
+		FBufferRHIRef MorphBuffer;
 
 		void CreateFromCPUData_RenderThread(FSkeletalMeshLODRenderData& LODResource);
 		void CreateFromCPUData_Async(FSkeletalMeshLODRenderData& LODResource);
diff --git a/Engine/Source/Runtime/Engine/Public/Rendering/MorphTargetVertexInfoBuffers.h b/Engine/Source/Runtime/Engine/Public/Rendering/MorphTargetVertexInfoBuffers.h
--- a/Engine/Source/Runtime/Engine/Public/Rendering/MorphTargetVertexInfoBuffers.h
+++ b/Engine/Source/Runtime/Engine/Public/Rendering/MorphTargetVertexInfoBuffers.h
@@ -79,6 +79,32 @@ public:
 	FBufferRHIRef MorphDataBuffer;
 	FShaderResourceViewRHIRef MorphDataSRV;
 
+	/** Create an RHI vertex buffer with CPU data. CPU data may be discarded after creation (see TResourceArray::Discard) */
+	FBufferRHIRef CreateMorphRHIBuffer_RenderThread();
+	FBufferRHIRef CreateMorphRHIBuffer_Async();
+
+	/** Similar to Init/ReleaseRHI but only update existing SRV so references to the SRV stays valid */
+	template <uint32 MaxNumUpdates>
+	void InitRHIForStreaming(
+		FRHIBuffer* IntermediatMorphTargetBuffer,
+		TRHIResourceUpdateBatcher<MaxNumUpdates>& Batcher)
+	{
+		check(MorphDataBuffer);
+		if (MorphDataBuffer && IntermediatMorphTargetBuffer)
+		{
+			Batcher.QueueUpdateRequest(MorphDataBuffer, IntermediatMorphTargetBuffer);
+			Batcher.QueueUpdateRequest(MorphDataSRV, MorphDataBuffer);
+		}
+	}
+
+	template<uint32 MaxNumUpdates>
+	void ReleaseRHIForStreaming(TRHIResourceUpdateBatcher<MaxNumUpdates>& Batcher)
+	{
+		check(MorphDataBuffer);
+		Batcher.QueueUpdateRequest(MorphDataBuffer, nullptr);
+		Batcher.QueueUpdateRequest(MorphDataSRV, nullptr, 0, 0);
+	}
+
 protected:
 	void ResetCPUData()
 	{
@@ -98,7 +124,10 @@ protected:
 	void Serialize(FArchive& Ar);
 
 	// Transient data. Gets deleted as soon as the GPU resource has been initialized (unless excplicitly disabled by bEmptyMorphCPUDataOnInitRHI).
-	TArray<uint32> MorphData;
+	TResourceArray<uint32> MorphData;
+
+	
+
 
 	//x,y,y separate for position and shared w for tangent
 	TArray<FVector4f> MaximumValuePerMorph;
@@ -117,6 +146,9 @@ protected:
 
 	friend class FSkeletalMeshLODRenderData;
 	ENGINE_API friend FArchive& operator<<(FArchive& Ar, FMorphTargetVertexInfoBuffers& MorphTargetVertexInfoBuffers);
+private:
+	template <bool bRenderThread>
+	FBufferRHIRef CreateMorphRHIBuffer_Internal();
 };
 
 ENGINE_API FArchive& operator<<(FArchive& Ar, FMorphTargetVertexInfoBuffers& MorphTargetVertexInfoBuffers);
diff --git a/Engine/Source/Runtime/Engine/Public/Rendering/SkeletalMeshLODRenderData.h b/Engine/Source/Runtime/Engine/Public/Rendering/SkeletalMeshLODRenderData.h
--- a/Engine/Source/Runtime/Engine/Public/Rendering/SkeletalMeshLODRenderData.h
+++ b/Engine/Source/Runtime/Engine/Public/Rendering/SkeletalMeshLODRenderData.h
@@ -177,8 +177,6 @@ public:
 	*/
 	void InitResources(bool bNeedsVertexColors, int32 LODIndex, TArray<class UMorphTarget*>& InMorphTargets, USkinnedAsset* Owner);
 
-	void InitMorphResources();
-
 	/**
 	* Releases the LOD's render resources.
 	*/
