Fix typo naming mismatch between read and write lock

diff --git a/Engine/Source/Runtime/Engine/Private/EdGraph/EdGraphPin.cpp b/Engine/Source/Runtime/Engine/Private/EdGraph/EdGraphPin.cpp
--- a/Engine/Source/Runtime/Engine/Private/EdGraph/EdGraphPin.cpp
+++ b/Engine/Source/Runtime/Engine/Private/EdGraph/EdGraphPin.cpp
@@ -45,7 +45,7 @@ public:
 	static void Add(UEdGraphPin* PinToDelete)
 	{
 		FPinDeletionQueue* PinDeletionQueue = Get();
-		FWriteScopeLock WriteLock(PinDeletionQueue->PinsToDeleteLock);
+		FWriteScopeLock ScopeLock(PinDeletionQueue->PinsToDeleteLock);
 		PinDeletionQueue->PinsToDelete.Add(PinToDelete);
 	}
 
@@ -53,7 +53,7 @@ public:
 	{
 		TArray<UEdGraphPin*> LocalPinsToDelete;
 		{
-			FWriteScopeLock ReadLock(PinsToDeleteLock);
+			FWriteScopeLock ScopeLock(PinsToDeleteLock);
 			Swap(LocalPinsToDelete, PinsToDelete);
 		}
 
@@ -65,7 +65,7 @@ public:
 
 	virtual bool IsTickable() const override
 	{
-		FReadScopeLock ReadLock(PinsToDeleteLock);
+		FReadScopeLock ScopeLock(PinsToDeleteLock);
 		return (PinsToDelete.Num() > 0);
 	}
 
