Returning the value of a Transform as a String now has a colon after Scale

#rnx

diff --git a/Engine/Source/Runtime/Engine/Private/KismetStringLibrary.cpp b/Engine/Source/Runtime/Engine/Private/KismetStringLibrary.cpp
--- a/Engine/Source/Runtime/Engine/Private/KismetStringLibrary.cpp
+++ b/Engine/Source/Runtime/Engine/Private/KismetStringLibrary.cpp
@@ -115,7 +115,7 @@ FString UKismetStringLibrary::Conv_RotatorToString(FRotator InRot)
 
 FString UKismetStringLibrary::Conv_TransformToString(const FTransform& InTrans)
 {
-	return FString::Printf(TEXT("Translation: %s Rotation: %s Scale %s"), *InTrans.GetTranslation().ToString(), *InTrans.Rotator().ToString(), *InTrans.GetScale3D().ToString());
+	return FString::Printf(TEXT("Translation: %s Rotation: %s Scale: %s"), *InTrans.GetTranslation().ToString(), *InTrans.Rotator().ToString(), *InTrans.GetScale3D().ToString());
 }
 
 FString UKismetStringLibrary::Conv_ObjectToString(class UObject* InObj)
