Deletes ICustomStaticScreenPercentage

Motivations:
1) Force compatibilty of third-party temporal upscalers with dyn-res
2) Allows multiple temporal upscaler plugins to live concurrently without collision on a single global
3) Reduce support overhead from licensee due to inconsistent internal rendering resolution across different temporal upscalers

Notes:
1) ITemporalUpscaler is now meant to be setup by the plugin's FSceneViewExtensionBase::BeginRenderViewFamily() that obeys the screen percentage driver;
2) The licensee's game user settings can just enable disable each plugin respectively based on their implementation specialization of ApplyNonResolutionSettings();
3) ITemporalUpscaler now has a lifetime tied to a FSceneViewFamily;
4) FSceneViewFamily offers a generic way to bind other data tied to its life time with FSceneViewFamily::GetOrCreateExtentionData() for instance for setting customise behavior of FSceneViewExtensionBase::BeginRenderViewFamily() before FRendererModule::BeginRenderingViewFamily().

diff --git a/Engine/Plugins/MovieScene/MovieRenderPipeline/Source/MovieRenderPipelineRenderPasses/Private/MoviePipelineImagePassBase.cpp b/Engine/Plugins/MovieScene/MovieRenderPipeline/Source/MovieRenderPipelineRenderPasses/Private/MoviePipelineImagePassBase.cpp
--- a/Engine/Plugins/MovieScene/MovieRenderPipeline/Source/MovieRenderPipelineRenderPasses/Private/MoviePipelineImagePassBase.cpp
+++ b/Engine/Plugins/MovieScene/MovieRenderPipeline/Source/MovieRenderPipelineRenderPasses/Private/MoviePipelineImagePassBase.cpp
@@ -285,7 +285,7 @@ TSharedPtr<FSceneViewFamilyContext> UMoviePipelineImagePassBase::CalculateViewFa
 		}
 	}
 
-	OutViewFamily->ViewExtensions = GEngine->ViewExtensions->GatherActiveExtensions(FSceneViewExtensionContext(GetWorld()->Scene));
+	OutViewFamily->ViewExtensions.Append(GEngine->ViewExtensions->GatherActiveExtensions(FSceneViewExtensionContext(GetWorld()->Scene)));
 
 	AddViewExtensions(*OutViewFamily, InOutSampleState);
 
@@ -584,8 +584,8 @@ FSceneView* UMoviePipelineImagePassBase::GetSceneViewForSampleState(FSceneViewFa
 			// ToDo: Does orthographic support tiling in the same way or do I need to modify the values before creating the ortho view.
 		}
 	}
-	// BaseProjMatrix may be perspective or orthographic.
-	ViewInitOptions.ProjectionMatrix = BaseProjMatrix;
+		// BaseProjMatrix may be perspective or orthographic.
+		ViewInitOptions.ProjectionMatrix = BaseProjMatrix;
 
 	ViewInitOptions.SceneViewStateInterface = GetSceneViewStateInterface(OptPayload);
 	ViewInitOptions.FOV = ViewFOV;
@@ -712,7 +712,7 @@ void UMoviePipelineImagePassBase::ModifyProjectionMatrixForTiling(const FMoviePi
 	}
 	else
 	{
-		InOutProjectionMatrix.M[2][0] += OffsetX / PadRatioX;
+	InOutProjectionMatrix.M[2][0] += OffsetX / PadRatioX;
 		InOutProjectionMatrix.M[2][1] += OffsetY / PadRatioY;
 	}
 }
diff --git a/Engine/Plugins/Runtime/nDisplay/Source/DisplayClusterLightCardEditor/Private/Viewport/DisplayClusterLightCardEditorViewportClient.cpp b/Engine/Plugins/Runtime/nDisplay/Source/DisplayClusterLightCardEditor/Private/Viewport/DisplayClusterLightCardEditorViewportClient.cpp
--- a/Engine/Plugins/Runtime/nDisplay/Source/DisplayClusterLightCardEditor/Private/Viewport/DisplayClusterLightCardEditorViewportClient.cpp
+++ b/Engine/Plugins/Runtime/nDisplay/Source/DisplayClusterLightCardEditor/Private/Viewport/DisplayClusterLightCardEditorViewportClient.cpp
@@ -27,7 +27,6 @@
 #include "Components/PostProcessComponent.h"
 #include "Components/SkyLightComponent.h"
 #include "Components/PrimitiveComponent.h"
-#include "CustomEditorStaticScreenPercentage.h"
 #include "Debug/DebugDrawService.h"
 #include "EngineUtils.h"
 #include "EditorModes.h"
diff --git a/Engine/Source/Editor/LevelEditor/Private/SLevelViewportToolBar.cpp b/Engine/Source/Editor/LevelEditor/Private/SLevelViewportToolBar.cpp
--- a/Engine/Source/Editor/LevelEditor/Private/SLevelViewportToolBar.cpp
+++ b/Engine/Source/Editor/LevelEditor/Private/SLevelViewportToolBar.cpp
@@ -44,8 +44,6 @@
 #include "SEditorViewportViewMenuContext.h"
 #include "Bookmarks/IBookmarkTypeTools.h"
 #include "ToolMenu.h"
-#include "CustomStaticScreenPercentage.h"
-#include "CustomEditorStaticScreenPercentage.h"
 #include "WorldPartition/WorldPartitionSubsystem.h"
 #include "SLevelViewport.h"
 
@@ -718,28 +716,7 @@ void SLevelViewportToolBar::FillOptionsMenu(UToolMenu* Menu)
 				Section.AddEntry(FToolMenuEntry::InitWidget("FarViewPlane", GenerateFarViewPlaneMenu(), LOCTEXT("FarViewPlane", "Far View Plane")));
 			}
 
-
-			// for the TemporalUpscaler plugin to inject its UI
-
-			bool bAddDefaultScreenPercentageSlider = true;
-
-			if (GCustomEditorStaticScreenPercentage)
-			{
-				ICustomEditorStaticScreenPercentage::FViewportMenuEntryArguments Arguments
-				{
-					&Section,
-					this,
-					Viewport.Pin()
-				};
-
-				bAddDefaultScreenPercentageSlider = !GCustomEditorStaticScreenPercentage->GenerateEditorViewportOptionsMenuEntry(Arguments);
-			}
-
-
-			if (bAddDefaultScreenPercentageSlider)
-			{
-				Section.AddEntry(FToolMenuEntry::InitWidget("ScreenPercentage", GenerateScreenPercentageMenu(), LOCTEXT("ScreenPercentage", "Screen Percentage")));
-			}
+			Section.AddEntry(FToolMenuEntry::InitWidget("ScreenPercentage", GenerateScreenPercentageMenu(), LOCTEXT("ScreenPercentage", "Screen Percentage")));
 		}
 
 		{
diff --git a/Engine/Source/Editor/UnrealEd/Private/EditorViewportClient.cpp b/Engine/Source/Editor/UnrealEd/Private/EditorViewportClient.cpp
--- a/Engine/Source/Editor/UnrealEd/Private/EditorViewportClient.cpp
+++ b/Engine/Source/Editor/UnrealEd/Private/EditorViewportClient.cpp
@@ -63,11 +63,8 @@
 #include "ProfilingDebugging/MiscTrace.h"
 #include "ProfilingDebugging/TraceScreenshot.h"
 #include "GenericPlatform/GenericPlatformInputDeviceMapper.h"
-#include "CustomEditorStaticScreenPercentage.h"
 #include "IImageWrapperModule.h"
 
-ICustomEditorStaticScreenPercentage* GCustomEditorStaticScreenPercentage = nullptr;
-
 #define LOCTEXT_NAMESPACE "EditorViewportClient"
 
 const EViewModeIndex FEditorViewportClient::DefaultPerspectiveViewMode = VMI_Lit;
@@ -4000,12 +3997,6 @@ void FEditorViewportClient::Draw(FViewport* InViewport, FCanvas* Canvas)
 		{
 			ViewFamily.SecondaryViewFraction = GetDPIDerivedResolutionFraction();
 		}
-
-		// Setup custom upscaler and screen percentage.
-		if (GCustomEditorStaticScreenPercentage && ViewFamily.ViewMode == EViewModeIndex::VMI_Lit && ViewFamily.Scene && ViewFamily.Scene->GetShadingPath() == EShadingPath::Deferred && ViewFamily.bRealtimeUpdate)
-		{
-			GCustomEditorStaticScreenPercentage->SetupEditorViewFamily(ViewFamily, this);
-		}
 	}
 
 	FSceneView* View = nullptr;
diff --git a/Engine/Source/Editor/UnrealEd/Public/CustomEditorStaticScreenPercentage.h b/Engine/Source/Editor/UnrealEd/Public/CustomEditorStaticScreenPercentage.h
deleted file mode 100644
diff --git a/Engine/Source/Runtime/Engine/Private/CustomStaticScreenPercentage.cpp b/Engine/Source/Runtime/Engine/Private/CustomStaticScreenPercentage.cpp
deleted file mode 100644
diff --git a/Engine/Source/Runtime/Engine/Private/GameViewportClient.cpp b/Engine/Source/Runtime/Engine/Private/GameViewportClient.cpp
--- a/Engine/Source/Runtime/Engine/Private/GameViewportClient.cpp
+++ b/Engine/Source/Runtime/Engine/Private/GameViewportClient.cpp
@@ -66,7 +66,6 @@
 #include "IImageWrapper.h"
 #include "IImageWrapperModule.h"
 #include "HAL/PlatformApplicationMisc.h"
-#include "CustomStaticScreenPercentage.h"
 #include "ObjectTrace.h"
 #include "DynamicResolutionState.h"
 
@@ -1429,11 +1428,6 @@ void UGameViewportClient::Draw(FViewport* InViewport, FCanvas* SceneCanvas)
 		}
 		#endif
 
-		if (GCustomStaticScreenPercentage && ViewFamily.ViewMode == EViewModeIndex::VMI_Lit && ViewFamily.Scene->GetShadingPath() == EShadingPath::Deferred && ViewFamily.bRealtimeUpdate)
-		{
-			GCustomStaticScreenPercentage->SetupMainGameViewFamily(ViewFamily);
-		}
-
 		bFinalScreenPercentageShowFlag = ViewFamily.EngineShowFlags.ScreenPercentage;
 	}
 
diff --git a/Engine/Source/Runtime/Engine/Private/SceneView.cpp b/Engine/Source/Runtime/Engine/Private/SceneView.cpp
--- a/Engine/Source/Runtime/Engine/Private/SceneView.cpp
+++ b/Engine/Source/Runtime/Engine/Private/SceneView.cpp
@@ -911,10 +911,6 @@ PRAGMA_ENABLE_DEPRECATION_WARNINGS
 	{
 		PrimaryScreenPercentageMethod = EPrimaryScreenPercentageMethod::TemporalUpscale;
 	}
-	else if (Family && Family->GetTemporalUpscalerInterface() != nullptr)
-	{
-		PrimaryScreenPercentageMethod = EPrimaryScreenPercentageMethod::TemporalUpscale;
-	}
 
 	if (Family)
 	{
@@ -1009,11 +1005,7 @@ void FSceneView::SetupAntiAliasingMethod()
 			}
 		}
 
-		// Overides the anti aliasing method to temporal AA when using a custom temporal upscaler.
-		if (Family->GetTemporalUpscalerInterface() != nullptr)
-		{
-			AntiAliasingMethod = AAM_TemporalAA;
-		}
+		checkf(Family->GetTemporalUpscalerInterface() == nullptr, TEXT("ITemporalUpscaler should be set up in FSceneViewExtensionBase::BeginRenderViewFamily()"));
 	}
 
 	// TemporalAA requires view state for history.
@@ -2948,6 +2940,12 @@ FSceneViewFamily::~FSceneViewFamily()
 		delete ScreenPercentageInterface;
 	}
 
+	if (TemporalUpscalerInterface)
+	{
+		// ITemporalUpscaler* is only defined in renderer's private header because no backward compatibility is provided between major version.
+		delete reinterpret_cast<ISceneViewFamilyExtention*>(TemporalUpscalerInterface);
+	}
+
 	if (PrimarySpatialUpscalerInterface)
 	{
 		// ISpatialUpscaler* is only defined in renderer's private header because no backward compatibility is provided between major version.
diff --git a/Engine/Source/Runtime/Engine/Public/CustomStaticScreenPercentage.h b/Engine/Source/Runtime/Engine/Public/CustomStaticScreenPercentage.h
deleted file mode 100644
diff --git a/Engine/Source/Runtime/Engine/Public/SceneView.h b/Engine/Source/Runtime/Engine/Public/SceneView.h
--- a/Engine/Source/Runtime/Engine/Public/SceneView.h
+++ b/Engine/Source/Runtime/Engine/Public/SceneView.h
@@ -1544,7 +1544,6 @@ struct FDisplayInternalsData
 
 //////////////////////////////////////////////////////////////////////////
 
-
 /**
  * Generic plugin extension that have a lifetime of the FSceneViewFamily
  */
@@ -1560,6 +1559,16 @@ protected:
 	friend class FSceneViewFamily;
 };
 
+/**
+ * Generic plugin extension that have a lifetime of the FSceneViewFamily that can contain arbitrary data to passdown from game thread to render thread.
+ */
+class ENGINE_API ISceneViewFamilyExtentionData : public ISceneViewFamilyExtention
+{
+public:
+	/** Returns a const TCHAR* to uniquely identify what implementation ISceneViewFamilyExtentionData is. */
+	virtual const TCHAR* GetSubclassIdentifier() const = 0;
+};
+
 
 /*
  * Game thread and render thread interface that takes care of a FSceneViewFamily's screen percentage.
@@ -1937,8 +1946,33 @@ public:
 		check(SecondarySpatialUpscalerInterface == nullptr);
 	}
 
+	template<typename TExtensionData> TExtensionData* GetExtentionData()
+	{
+		static_assert(TIsDerivedFrom<TExtensionData, ISceneViewFamilyExtentionData>::Value, "TExtensionData is not derived from ISceneViewFamilyExtentionData.");
+
+		for (TSharedRef<class ISceneViewFamilyExtentionData, ESPMode::ThreadSafe>& ViewExtensionData : ViewExtentionDatas)
+		{
+			if (ViewExtensionData->GetSubclassIdentifier() == TExtensionData::GSubclassIdentifier)
+			{
+				return static_cast<TExtensionData*>(&ViewExtensionData.Get());
+			}
+		}
+		return nullptr;
+	}
+
+	template<typename TExtensionData> TExtensionData* GetOrCreateExtentionData()
+	{
+		TExtensionData* ViewExtensionData = GetExtentionData<TExtensionData>();
+		if (!ViewExtensionData)
+		{
+			ViewExtentionDatas.Push(MakeShared<TExtensionData>());
+			ViewExtensionData = static_cast<TExtensionData*>(&ViewExtentionDatas.Last().Get());
+			check(ViewExtensionData->GetSubclassIdentifier() == TExtensionData::GSubclassIdentifier);
+		}
+		return ViewExtensionData;
+	}
 
-	FORCEINLINE void SetTemporalUpscalerInterface(const ITemporalUpscaler* InTemporalUpscalerInterface)
+	FORCEINLINE void SetTemporalUpscalerInterface(ITemporalUpscaler* InTemporalUpscalerInterface)
 	{
 		check(InTemporalUpscalerInterface);
 		checkf(TemporalUpscalerInterface == nullptr, TEXT("View family already had a temporal upscaler assigned."));
@@ -1978,10 +2012,14 @@ private:
 	/** Interface to handle screen percentage of the views of the family. */
 	ISceneViewFamilyScreenPercentage* ScreenPercentageInterface;
 
-	const ITemporalUpscaler* TemporalUpscalerInterface;
+	/** Renderer private interfaces, automatically have same lifetime as FSceneViewFamily. */
+	ITemporalUpscaler* TemporalUpscalerInterface;
 	ISpatialUpscaler* PrimarySpatialUpscalerInterface;
 	ISpatialUpscaler* SecondarySpatialUpscalerInterface;
 
+	/** Arrays of standalone data that have safe lifetime as FSceneViewFamily. */
+	TArray<TSharedRef<class ISceneViewFamilyExtentionData, ESPMode::ThreadSafe> > ViewExtentionDatas;
+
 	/** whether the translucency are allowed to render after DOF, if not they will be rendered in standard translucency. */
 	bool bAllowTranslucencyAfterDOF;
 
diff --git a/Engine/Source/Runtime/Renderer/Private/PostProcess/TemporalAA.cpp b/Engine/Source/Runtime/Renderer/Private/PostProcess/TemporalAA.cpp
--- a/Engine/Source/Runtime/Renderer/Private/PostProcess/TemporalAA.cpp
+++ b/Engine/Source/Runtime/Renderer/Private/PostProcess/TemporalAA.cpp
@@ -1111,6 +1111,12 @@ public:
 	{
 		return ISceneViewFamilyScreenPercentage::kMaxTAAUpsampleResolutionFraction;
 	}
+
+	virtual ITemporalUpscaler* Fork_GameThread(const class FSceneViewFamily& ViewFamily) const override
+	{
+		return nullptr;
+	}
+
 };
 
 // static
diff --git a/Engine/Source/Runtime/Renderer/Private/PostProcess/TemporalAA.h b/Engine/Source/Runtime/Renderer/Private/PostProcess/TemporalAA.h
--- a/Engine/Source/Runtime/Renderer/Private/PostProcess/TemporalAA.h
+++ b/Engine/Source/Runtime/Renderer/Private/PostProcess/TemporalAA.h
@@ -173,7 +173,7 @@ extern RENDERER_API FTAAOutputs AddTemporalAAPass(
 	FTemporalAAHistory* OutputHistory);
 
 /** Interface for the main temporal upscaling algorithm. */
-class RENDERER_API ITemporalUpscaler
+class RENDERER_API ITemporalUpscaler : public ISceneViewFamilyExtention
 {
 public:
 
@@ -214,6 +214,8 @@ public:
 	virtual float GetMinUpsampleResolutionFraction() const = 0;
 	virtual float GetMaxUpsampleResolutionFraction() const = 0;
 
+	virtual ITemporalUpscaler* Fork_GameThread(const class FSceneViewFamily& ViewFamily) const = 0;
+
 	static const ITemporalUpscaler* GetDefaultTemporalUpscaler();
 
 	static EMainTAAPassConfig GetMainTAAPassConfig(const FViewInfo& View);
diff --git a/Engine/Source/Runtime/Renderer/Private/SceneRendering.cpp b/Engine/Source/Runtime/Renderer/Private/SceneRendering.cpp
--- a/Engine/Source/Runtime/Renderer/Private/SceneRendering.cpp
+++ b/Engine/Source/Runtime/Renderer/Private/SceneRendering.cpp
@@ -2574,6 +2574,19 @@ FSceneRenderer::FSceneRenderer(const FSceneViewFamily* InViewFamily, FHitProxyCo
 			ViewFamily.SetScreenPercentageInterface(InViewFamily->ScreenPercentageInterface->Fork_GameThread(ViewFamily));
 		}
 
+		if (ViewFamily.TemporalUpscalerInterface)
+		{
+			ViewFamily.TemporalUpscalerInterface = nullptr;
+			ViewFamily.SetTemporalUpscalerInterface(InViewFamily->TemporalUpscalerInterface->Fork_GameThread(ViewFamily));
+
+			for (FViewInfo& View : Views)
+			{
+				View.AntiAliasingMethod = AAM_TemporalAA;
+				View.PrimaryScreenPercentageMethod = EPrimaryScreenPercentageMethod::TemporalUpscale;
+				check(View.VerifyMembersChecks());
+			}
+		}
+
 		if (ViewFamily.PrimarySpatialUpscalerInterface)
 		{
 			ViewFamily.PrimarySpatialUpscalerInterface = nullptr;
@@ -4598,7 +4611,8 @@ void FRendererModule::BeginRenderingViewFamilies(FCanvas* Canvas, TArrayView<FSc
 		}
 #endif // !(UE_BUILD_SHIPPING || UE_BUILD_TEST)
 
-		// Force the spatial upscaler to be set no earlier than ISceneViewExtension::BeginRenderViewFamily();
+		// Force the upscalers to be set no earlier than ISceneViewExtension::BeginRenderViewFamily();
+		check(ViewFamily->GetTemporalUpscalerInterface() == nullptr);
 		check(ViewFamily->GetPrimarySpatialUpscalerInterface() == nullptr);
 		check(ViewFamily->GetSecondarySpatialUpscalerInterface() == nullptr);
 
@@ -4606,6 +4620,9 @@ void FRendererModule::BeginRenderingViewFamilies(FCanvas* Canvas, TArrayView<FSc
 		{
 			ViewFamily->ViewExtensions[ViewExt]->BeginRenderViewFamily(*ViewFamily);
 		}
+
+		checkf(!(ViewFamily->GetTemporalUpscalerInterface() != nullptr && ViewFamily->GetPrimarySpatialUpscalerInterface() != nullptr),
+			TEXT("Conflict setting up a third party primary spatial upscaler or temporal upscaler."));
 	}
 
 	if (Scene)
