Fixed superflluous ifdef WITH_EDITOR

#rnx

diff --git a/Engine/Source/Runtime/Engine/Private/Texture.cpp b/Engine/Source/Runtime/Engine/Private/Texture.cpp
--- a/Engine/Source/Runtime/Engine/Private/Texture.cpp
+++ b/Engine/Source/Runtime/Engine/Private/Texture.cpp
@@ -523,9 +523,6 @@ void UTexture::PostEditChangeProperty(FPropertyChangedEvent& PropertyChangedEven
 
 	Super::PostEditChangeProperty(PropertyChangedEvent);
 
-	// this whole function is in WITH_EDITOR so these ifs don't make much sense
-
-#if WITH_EDITOR
 	ON_SCOPE_EXIT
 	{
 		// PostEditChange is the last step in the import sequence (PreEditChange/PostEditImport/PostEditChange, called twice : see further details) so reset the import-related flags here:
@@ -539,7 +536,6 @@ void UTexture::PostEditChangeProperty(FPropertyChangedEvent& PropertyChangedEven
 	{
 		return;
 	}
-#endif // WITH_EDITOR
 
 	// assume there was a change that needs a new lighting guid :
 	SetLightingGuid();
