Make sure shader creation for PSO precaching happens on RT (UE-164471).

- OpenGL crashes if we're creating a shader on GT when RT already exists.

diff --git a/Engine/Source/Runtime/Engine/Private/ShaderCompiler/ShaderCompiler.cpp b/Engine/Source/Runtime/Engine/Private/ShaderCompiler/ShaderCompiler.cpp
--- a/Engine/Source/Runtime/Engine/Private/ShaderCompiler/ShaderCompiler.cpp
+++ b/Engine/Source/Runtime/Engine/Private/ShaderCompiler/ShaderCompiler.cpp
@@ -7030,6 +7030,8 @@ void PrecacheComputePipelineStatesForGlobalShaders(EShaderPlatform Platform, con
 
 	FGlobalShaderMap* GlobalShaderMap = GetGlobalShaderMap(Platform);
 
+	// some RHIs (OpenGL) can only create shaders on the Render thread. Queue the creation instead of doing it here.
+	TArray<TShaderRef<FShader>> ComputeShadersToPrecache;
 	for (TLinkedList<FShaderType*>::TIterator ShaderTypeIt(FShaderType::GetTypeList()); ShaderTypeIt; ShaderTypeIt.Next())
 	{
 		FGlobalShaderType* GlobalShaderType = ShaderTypeIt->GetGlobalShaderType();
@@ -7045,12 +7047,24 @@ void PrecacheComputePipelineStatesForGlobalShaders(EShaderPlatform Platform, con
 				TShaderRef<FShader> GlobalShader = GlobalShaderMap->GetShader(GlobalShaderType, PermutationId);
 				if (GlobalShader.IsValid() && GlobalShader->GetFrequency() == SF_Compute)
 				{
-					FRHIComputeShader* RHIComputeShader = GlobalShader.GetComputeShader();
-					PipelineStateCache::PrecacheComputePipelineState(RHIComputeShader);
+					ComputeShadersToPrecache.Add(GlobalShader);
 				}
 			}
 		}
 	}
+
+	if (ComputeShadersToPrecache.Num() > 0)
+	{
+		ENQUEUE_RENDER_COMMAND(PrecachePSOsForGlobalShaders)(
+			[ComputeShadersToPrecache](FRHICommandListImmediate& RHICmdList)
+			{
+				for (TShaderRef<FShader> GlobalShader : ComputeShadersToPrecache)
+				{
+					FRHIComputeShader* RHIComputeShader = GlobalShader.GetComputeShader();
+					PipelineStateCache::PrecacheComputePipelineState(RHIComputeShader);
+				}
+			});
+	}
 }
 
 #include "Misc/PreLoadFile.h"
