Adding SetAffectDistanceFieldLighting blueprint callable function to UPrimitiveComponent

diff --git a/Engine/Source/Runtime/Engine/Classes/Components/PrimitiveComponent.h b/Engine/Source/Runtime/Engine/Classes/Components/PrimitiveComponent.h
--- a/Engine/Source/Runtime/Engine/Classes/Components/PrimitiveComponent.h
+++ b/Engine/Source/Runtime/Engine/Classes/Components/PrimitiveComponent.h
@@ -1673,6 +1673,10 @@ public:
 	UFUNCTION(BlueprintCallable, Category = "Rendering")
 	void SetTranslucencySortDistanceOffset(float NewTranslucencySortDistanceOffset);
 
+	/** Changes the value of Affect Distance Field Lighting */
+	UFUNCTION(BlueprintCallable, Category = "Rendering")
+	void SetAffectDistanceFieldLighting(bool NewAffectDistanceFieldLighting);
+
 	/** Changes the value of bReceivesDecals. */
 	UFUNCTION(BlueprintCallable, Category = "Rendering")
 	void SetReceivesDecals(bool bNewReceivesDecals);
diff --git a/Engine/Source/Runtime/Engine/Private/Components/PrimitiveComponent.cpp b/Engine/Source/Runtime/Engine/Private/Components/PrimitiveComponent.cpp
--- a/Engine/Source/Runtime/Engine/Private/Components/PrimitiveComponent.cpp
+++ b/Engine/Source/Runtime/Engine/Private/Components/PrimitiveComponent.cpp
@@ -1709,6 +1709,15 @@ void UPrimitiveComponent::SetTranslucentSortPriority(int32 NewTranslucentSortPri
 	}
 }
 
+void UPrimitiveComponent::SetAffectDistanceFieldLighting(bool NewAffectDistanceFieldLighting)
+{
+	if(NewAffectDistanceFieldLighting != bAffectDistanceFieldLighting)
+	{
+		bAffectDistanceFieldLighting = NewAffectDistanceFieldLighting;
+		MarkRenderStateDirty();
+	}
+}
+
 void UPrimitiveComponent::SetTranslucencySortDistanceOffset(float NewTranslucencySortDistanceOffset)
 {
 	if ( !FMath::IsNearlyEqual(NewTranslucencySortDistanceOffset, TranslucencySortDistanceOffset) )
