Fix loading of really old assets with alternate skin profiles.

#rnx

diff --git a/Engine/Source/Runtime/Engine/Private/Animation/SkinWeightProfile.cpp b/Engine/Source/Runtime/Engine/Private/Animation/SkinWeightProfile.cpp
--- a/Engine/Source/Runtime/Engine/Private/Animation/SkinWeightProfile.cpp
+++ b/Engine/Source/Runtime/Engine/Private/Animation/SkinWeightProfile.cpp
@@ -200,7 +200,7 @@ FArchive& operator<<(FArchive& Ar, FRawSkinWeight& OverrideEntry)
 	{
 		for (int32 InfluenceIndex = 0; InfluenceIndex < EXTRA_BONE_INFLUENCES; ++InfluenceIndex)
 		{
-			if (Ar.IsLoading() && Ar.CustomVer(FAnimObjectVersion::GUID) < FAnimObjectVersion::IncreaseBoneIndexLimitPerChunk)
+			if (Ar.CustomVer(FAnimObjectVersion::GUID) < FAnimObjectVersion::IncreaseBoneIndexLimitPerChunk)
 			{
 				uint8 BoneIndex = 0;
 				Ar << BoneIndex;
@@ -211,7 +211,9 @@ FArchive& operator<<(FArchive& Ar, FRawSkinWeight& OverrideEntry)
 				Ar << OverrideEntry.InfluenceBones[InfluenceIndex];
 			}
 
-			Ar << OverrideEntry.InfluenceWeights[InfluenceIndex];
+			uint8 Weight = 0;
+			Ar << Weight;
+			OverrideEntry.InfluenceWeights[InfluenceIndex] = (static_cast<uint16>(Weight) << 8) | Weight;
 		}
 	}
 	else if (Ar.CustomVer(FUE5MainStreamObjectVersion::GUID) < FUE5MainStreamObjectVersion::IncreasedSkinWeightPrecision)
