Fixes for linked anim graphs

Fix for concurrently accessing the linked anim layers array (if layers are linked/unlinked from a ForEachInstance call)
Added a warning when attempting to link anim layers from the initialization callback, as 'self' layers have not been initialized at that point.
Partially merges PR #9481.

PR #9481: Fixing linked anim graphs not updating any more if they call LinkAnimClassLayers (Contributed by [REDACTED_PERSON])

diff --git a/Engine/Source/Runtime/Engine/Classes/Animation/AnimInstance.h b/Engine/Source/Runtime/Engine/Classes/Animation/AnimInstance.h
--- a/Engine/Source/Runtime/Engine/Classes/Animation/AnimInstance.h
+++ b/Engine/Source/Runtime/Engine/Classes/Animation/AnimInstance.h
@@ -438,6 +438,9 @@ private:
 	uint8 bQueueMontageEvents : 1;
 
 #if DO_CHECK
+	/** Guard flag used for checking whether we are in user callbacks for initialization */
+	uint8 bInitializing : 1;
+
 	/** Used to guard against recursive calls to UpdateAnimation */
 	bool bPostUpdatingAnimation;
 
diff --git a/Engine/Source/Runtime/Engine/Private/Animation/AnimInstance.cpp b/Engine/Source/Runtime/Engine/Private/Animation/AnimInstance.cpp
--- a/Engine/Source/Runtime/Engine/Private/Animation/AnimInstance.cpp
+++ b/Engine/Source/Runtime/Engine/Private/Animation/AnimInstance.cpp
@@ -124,6 +124,10 @@ UAnimInstance::UAnimInstance(const FObjectInitializer& ObjectInitializer)
 	bPropagateNotifiesToLinkedInstances = false;
 	bUseMainInstanceMontageEvaluationData = false;
 
+#if DO_CHECK
+	bInitializing = false;
+#endif
+
 #if WITH_EDITOR
 	if(!HasAnyFlags(RF_ClassDefaultObject))
 	{
@@ -251,9 +255,16 @@ void UAnimInstance::InitializeAnimation(bool bInDeferRootNodeInitialization)
 
 	GetProxyOnGameThread<FAnimInstanceProxy>().Initialize(this);
 
-	NativeInitializeAnimation();
-	BlueprintInitializeAnimation();
-
+#if DO_CHECK
+	{
+		// Allow us to validate callbacks within user code
+		FGuardValue_Bitfield(bInitializing, true);
+		
+		NativeInitializeAnimation();
+		BlueprintInitializeAnimation();
+	}
+#endif
+	
 	GetProxyOnGameThread<FAnimInstanceProxy>().InitializeRootNode(bInDeferRootNodeInitialization);
 
 	// we can bind rules & events now the graph has been initialized
@@ -2883,6 +2894,13 @@ void UAnimInstance::LinkAnimGraphByTag(FName InTag, TSubclassOf<UAnimInstance> I
 
 void UAnimInstance::PerformLinkedLayerOverlayOperation(TSubclassOf<UAnimInstance> InClass, TFunctionRef<UClass*(UClass* InClass, FAnimNode_LinkedAnimLayer*)> InClassSelectorFunction, bool bInDeferSubGraphInitialization)
 {
+#if DO_CHECK
+	if(bInitializing)
+	{
+		UE_LOG(LogAnimation, Warning, TEXT("Performing linked layer operations in initialization may produce unexpected results, please use the BlueprintLinkedAnimationLayersInitialized event to perform linked layer operations."));
+	}
+#endif
+	
 	if (IAnimClassInterface* AnimBlueprintClass = IAnimClassInterface::GetFromClass(GetClass()))
 	{
 		UClass* NewClass = InClass.Get();
diff --git a/Engine/Source/Runtime/Engine/Private/Components/SkeletalMeshComponent.cpp b/Engine/Source/Runtime/Engine/Private/Components/SkeletalMeshComponent.cpp
--- a/Engine/Source/Runtime/Engine/Private/Components/SkeletalMeshComponent.cpp
+++ b/Engine/Source/Runtime/Engine/Private/Components/SkeletalMeshComponent.cpp
@@ -3171,7 +3171,8 @@ void USkeletalMeshComponent::ForEachAnimInstance(TFunctionRef<void(UAnimInstance
 		InFunction(AnimScriptInstance);
 	}
 
-	for(UAnimInstance* LinkedInstance : LinkedInstances)
+	// Copy LinkedInstances because the array can be concurrently modified inside this loop 
+	for(UAnimInstance* LinkedInstance : TArray<TObjectPtr<UAnimInstance>, TInlineAllocator<8>>(LinkedInstances))
 	{
 		if(LinkedInstance)
 		{
