Fix for Montages not handling section loop under very specific timing conditions, where there is no time left on the sub-stepper

diff --git a/Engine/Source/Runtime/Engine/Private/Animation/AnimMontage.cpp b/Engine/Source/Runtime/Engine/Private/Animation/AnimMontage.cpp
--- a/Engine/Source/Runtime/Engine/Private/Animation/AnimMontage.cpp
+++ b/Engine/Source/Runtime/Engine/Private/Animation/AnimMontage.cpp
@@ -38,6 +38,16 @@ namespace MontageFNames
 	static FName TimeStretchCurveName(TEXT("MontageTimeStretchCurve"));
 }
 
+// CVars
+namespace MontageCVars
+{
+	static bool bEndSectionRequiresTimeRemaining = false;
+	static FAutoConsoleVariableRef CVarMontageEndSectionRequiresTimeRemaining(
+		TEXT("a.Montage.EndSectionRequiresTimeRemaining"),
+		bEndSectionRequiresTimeRemaining,
+		TEXT("Montage EndOfSection is only checked if there is remaining time (default false)."));
+} // end namespace MontageCVars
+
 ///////////////////////////////////////////////////////////////////////////
 //
 UAnimMontage::UAnimMontage(const FObjectInitializer& ObjectInitializer)
@@ -2449,7 +2459,9 @@ void FAnimMontageInstance::Advance(float DeltaTime, struct FRootMotionMovementPa
 					}
 				}
 
-				if(MontageSubStepper.HasTimeRemaining())
+				// Note that we have to check this even if there is no time remaining, in order to correctly handle loops
+				// CVar allows reverting to old behavior, in case a project relies on it
+				if (MontageCVars::bEndSectionRequiresTimeRemaining == false || MontageSubStepper.HasTimeRemaining())
 				{
 					// if we reached end of section, and we were not processing a branching point, and no events has messed with out current position..
 					// .. Move to next section.
