Fix crash when we somehow end up with no sync scope in asset players

Added fallback to ensure we get a scope added before we enter an anim graph update

diff --git a/Engine/Source/Runtime/Engine/Private/Animation/AnimInstanceProxy.cpp b/Engine/Source/Runtime/Engine/Private/Animation/AnimInstanceProxy.cpp
--- a/Engine/Source/Runtime/Engine/Private/Animation/AnimInstanceProxy.cpp
+++ b/Engine/Source/Runtime/Engine/Private/Animation/AnimInstanceProxy.cpp
@@ -1241,7 +1241,10 @@ void FAnimInstanceProxy::UpdateAnimation_WithRoot(const FAnimationUpdateContext&
 
 		// Anything syncing within this scope is subject to sync groups.
 		// We only enable syncing here for the main instance or post process instance
-		const bool bEnableSyncScope = GetAnimInstanceObject() == GetSkelMeshComponent()->GetAnimInstance() || GetAnimInstanceObject() == GetSkelMeshComponent()->GetPostProcessInstance();
+		// We also fall back to enabling this sync scope if there is not one already enabled (there must always be one)
+		const bool bEnableSyncScope =	GetAnimInstanceObject() == GetSkelMeshComponent()->GetAnimInstance() ||
+										GetAnimInstanceObject() == GetSkelMeshComponent()->GetPostProcessInstance() || 
+										InContext.GetMessage<UE::Anim::FAnimSyncGroupScope>() == nullptr;
 		UE::Anim::TOptionalScopedGraphMessage<UE::Anim::FAnimSyncGroupScope> Message(bEnableSyncScope, InContext, InContext);
 
 		// update all nodes
