Check material pin is valid before accessing it

diff --git a/Engine/Source/Runtime/Engine/Private/Materials/MaterialExpressions.cpp b/Engine/Source/Runtime/Engine/Private/Materials/MaterialExpressions.cpp
--- a/Engine/Source/Runtime/Engine/Private/Materials/MaterialExpressions.cpp
+++ b/Engine/Source/Runtime/Engine/Private/Materials/MaterialExpressions.cpp
@@ -1503,8 +1503,11 @@ void UMaterialExpression::PostEditChangeProperty(FPropertyChangedEvent& Property
 
 			// Update the pin value of the expression input
 			UEdGraphPin* Pin = GraphNode->GetPinAt(PinIndex);
-			Pin->Modify();
-			Pin->DefaultValue = NewDefaultValue;
+			if (Pin)
+			{
+				Pin->Modify();
+				Pin->DefaultValue = NewDefaultValue;
+			}
 
 			// If the property is linked as inline toggle to another property, both pins need updating to reflect the change.
 			bool bInlineEditConditionToggle = MemberPropertyThatChanged->HasMetaData(TEXT("InlineEditConditionToggle"));
