Shared Linear Encoding for DDC1 - disabled by default, enable by cvar (r.SharedLinearTextureEncoding 1)

diff --git a/Engine/Source/Developer/TextureCompressor/Private/TextureCompressorModule.cpp b/Engine/Source/Developer/TextureCompressor/Private/TextureCompressorModule.cpp
--- a/Engine/Source/Developer/TextureCompressor/Private/TextureCompressorModule.cpp
+++ b/Engine/Source/Developer/TextureCompressor/Private/TextureCompressorModule.cpp
@@ -2951,7 +2951,7 @@ static bool ApplyCompositeTexture(FImage& DestRoughness, const FImage& SourceNor
 	Image Compression.
 ------------------------------------------------------------------------------*/
 
-void FTextureBuildSettings::GetEncodedTextureDescription(FEncodedTextureDescription* OutTextureDescription, const ITextureFormat* InTextureFormat, int32 InEncodedMip0SizeX, int32 InEncodedMip0SizeY, int32 InEncodedMip0NumSlices, int32 InMipCount, bool bInImageHasAlphaChannel) const
+void FTextureBuildSettings::GetEncodedTextureDescriptionWithPixelFormat(FEncodedTextureDescription* OutTextureDescription, EPixelFormat InEncodedPixelFormat, int32 InEncodedMip0SizeX, int32 InEncodedMip0SizeY, int32 InEncodedMip0NumSlices, int32 InMipCount) const
 {
 	FEncodedTextureDescription& TextureDescription = *OutTextureDescription;
 	TextureDescription = FEncodedTextureDescription();
@@ -2959,7 +2959,7 @@ void FTextureBuildSettings::GetEncodedTextureDescription(FEncodedTextureDescript
 	TextureDescription.bTextureArray = bTextureArray;
 	TextureDescription.bVolumeTexture = bVolume;
 	TextureDescription.NumMips = InMipCount;
-	TextureDescription.PixelFormat = InTextureFormat->GetEncodedPixelFormat(*this, bInImageHasAlphaChannel);
+	TextureDescription.PixelFormat = InEncodedPixelFormat;
 
 	TextureDescription.TopMipSizeX = InEncodedMip0SizeX;
 	TextureDescription.TopMipSizeY = InEncodedMip0SizeY;
@@ -2981,6 +2981,12 @@ void FTextureBuildSettings::GetEncodedTextureDescription(FEncodedTextureDescript
 	}
 }
 
+void FTextureBuildSettings::GetEncodedTextureDescription(FEncodedTextureDescription* OutTextureDescription, const ITextureFormat* InTextureFormat, int32 InEncodedMip0SizeX, int32 InEncodedMip0SizeY, int32 InEncodedMip0NumSlices, int32 InMipCount, bool bInImageHasAlphaChannel) const
+{
+	EPixelFormat EncodedPixelFormat = InTextureFormat->GetEncodedPixelFormat(*this, bInImageHasAlphaChannel);
+	GetEncodedTextureDescriptionWithPixelFormat(OutTextureDescription, EncodedPixelFormat, InEncodedMip0SizeX, InEncodedMip0SizeY, InEncodedMip0NumSlices, InMipCount);
+}
+
 // compress mip-maps in InMipChain and add mips to Texture, might alter the source content
 // MipChain FImage payloads are freed by this function (RawData.Empty() is called)
 static bool CompressMipChain(
diff --git a/Engine/Source/Developer/TextureCompressor/Public/TextureCompressorModule.h b/Engine/Source/Developer/TextureCompressor/Public/TextureCompressorModule.h
--- a/Engine/Source/Developer/TextureCompressor/Public/TextureCompressorModule.h
+++ b/Engine/Source/Developer/TextureCompressor/Public/TextureCompressorModule.h
@@ -12,6 +12,7 @@ struct FEncodedTextureDescription;
 class ITextureFormat;
 class ITextureTiler;
 struct FTextureEngineParameters;
+enum EPixelFormat : uint8;
 
 /**
  * Compressed image data.
@@ -318,10 +319,9 @@ struct FTextureBuildSettings
 	/*
 	* Convert the build settings to an actual texture description containing enough information to describe the texture
 	* to hardware APIs.
-	* 
-	* At the moment, this only generates valid results _after_ some setup code in FTextureCompressorModule::BuildTexture()
 	*/
 	TEXTURECOMPRESSOR_API void GetEncodedTextureDescription(FEncodedTextureDescription* OutTextureDescription, const ITextureFormat* InTextureFormat, int32 InEncodedMip0SizeX, int32 InEncodedMip0SizeY, int32 InEncodedMip0NumSlices, int32 InMipCount, bool bInImageHasAlphaChannel) const;
+	TEXTURECOMPRESSOR_API void GetEncodedTextureDescriptionWithPixelFormat(FEncodedTextureDescription* OutTextureDescription, EPixelFormat InEncodedPixelFormat, int32 InEncodedMip0SizeX, int32 InEncodedMip0SizeY, int32 InEncodedMip0NumSlices, int32 InMipCount) const;
 };
 
 /**
diff --git a/Engine/Source/Developer/TextureFormat/Public/ChildTextureFormat.h b/Engine/Source/Developer/TextureFormat/Public/ChildTextureFormat.h
--- a/Engine/Source/Developer/TextureFormat/Public/ChildTextureFormat.h
+++ b/Engine/Source/Developer/TextureFormat/Public/ChildTextureFormat.h
@@ -113,7 +113,7 @@ protected:
 	{
 		return ObjView.FindView("BaseTextureFormatConfig").AsObjectView();
 	}
-
+public:
 	FTextureBuildSettings GetBaseTextureBuildSettings(const FTextureBuildSettings& BuildSettings) const
 	{
 		FTextureBuildSettings BaseSettings = BuildSettings;
@@ -121,6 +121,7 @@ protected:
 		BaseSettings.FormatConfigOverride = GetBaseFormatConfigOverride(BuildSettings.FormatConfigOverride);
 		return BaseSettings;
 	}
+protected:
 
 	/**
 	 * The final version is a combination of parent and child formats, 8 bits for each
diff --git a/Engine/Source/Runtime/Core/Private/HAL/ConsoleManager.cpp b/Engine/Source/Runtime/Core/Private/HAL/ConsoleManager.cpp
--- a/Engine/Source/Runtime/Core/Private/HAL/ConsoleManager.cpp
+++ b/Engine/Source/Runtime/Core/Private/HAL/ConsoleManager.cpp
@@ -3140,7 +3140,7 @@ static TAutoConsoleVariable<float> CVarSkeletalMeshLODRadiusScale(
 
 static TAutoConsoleVariable<int32> CVarSharedLinearTextureEncoding(
 	TEXT("r.SharedLinearTextureEncoding"),
-	1,
+	0,
 	TEXT("If set to 1, textures for platforms that tile will reuse a host linear texture instead of reencoding."),
 	ECVF_RenderThreadSafe);
 
diff --git a/Engine/Source/Runtime/Engine/Private/TextureDerivedData.cpp b/Engine/Source/Runtime/Engine/Private/TextureDerivedData.cpp
--- a/Engine/Source/Runtime/Engine/Private/TextureDerivedData.cpp
+++ b/Engine/Source/Runtime/Engine/Private/TextureDerivedData.cpp
@@ -625,7 +625,7 @@ static void FinalizeBuildSettingsForLayer(
 					OutBuildResultMetadata->bSupportsEncodeSpeed = bSupportsEncodeSpeed;
 				}
 			
-				if (IsUsingNewDerivedData()) // shared linear encoding only hooked up to ddc2 atm.
+				
 				{
 					static auto CVarSharedLinearTextureEncoding = IConsoleManager::Get().FindTConsoleVariableDataInt(TEXT("r.SharedLinearTextureEncoding"));
 					if (CVarSharedLinearTextureEncoding->GetValueOnAnyThread())
@@ -636,10 +636,10 @@ static void FinalizeBuildSettingsForLayer(
 						if (ChildTextureFormat && ChildTextureFormat->GetBaseFormatObject(OutSettings.TextureFormatName)->SupportsTiling() == false)
 						{
 							OutSettings.Tiler = ChildTextureFormat->GetTiler();
-							if (OutSettings.Tiler)
+							if (OutSettings.Tiler && IsUsingNewDerivedData())
 							{
-								// Since the tiler will do the encoding, skip the child texture format entirely and just look for the
-								// base texture and then run the tiler after.
+								// New derived data wants to treat everything as the base format and then have a separate tiling function
+								// afterwards.
 								OutSettings.TextureFormatName = ChildTextureFormat->GetBaseFormatName(OutSettings.TextureFormatName);
 							}
 						}
diff --git a/Engine/Source/Runtime/Engine/Private/TextureDerivedDataTask.cpp b/Engine/Source/Runtime/Engine/Private/TextureDerivedDataTask.cpp
--- a/Engine/Source/Runtime/Engine/Private/TextureDerivedDataTask.cpp
+++ b/Engine/Source/Runtime/Engine/Private/TextureDerivedDataTask.cpp
@@ -528,8 +528,105 @@ private:
 
 } // UE::TextureDerivedData
 
+static void DDC1_StoreClassicTextureInDerivedData(
+	TArray<FCompressedImage2D>& CompressedMips, FTexturePlatformData* DerivedData, bool bVolume, bool bTextureArray, bool bCubemap, uint32 NumMipsInTail,
+	uint32 ExtData, bool bReplaceExistingDDC, const FString& TexturePathName, const FString& KeySuffix, int64& BytesCached
+	)
+{
+	const int32 MipCount = CompressedMips.Num();
+
+	// VT can be bigger than (1<<(MAX_TEXTURE_MIP_COUNT-1)) , but doesn't actually make all those mips
+	// bForVirtualTextureStreamingBuild is false in this branch
+	check(MipCount <= MAX_TEXTURE_MIP_COUNT);
+
+	for (int32 MipIndex = 0; MipIndex < MipCount; ++MipIndex)
+	{
+		const FCompressedImage2D& CompressedImage = CompressedMips[MipIndex];
+		FTexture2DMipMap* NewMip = new FTexture2DMipMap();
+		DerivedData->Mips.Add(NewMip);
+		NewMip->SizeX = CompressedImage.SizeX;
+		NewMip->SizeY = CompressedImage.SizeY;
+		NewMip->SizeZ = CompressedImage.SizeZ;
+		NewMip->FileRegionType = FFileRegion::SelectType(EPixelFormat(CompressedImage.PixelFormat));
+		check(NewMip->SizeZ == 1 || bVolume || bTextureArray); // Only volume & arrays can have SizeZ != 1
+
+		check(CompressedImage.RawData.GetTypeSize() == 1);
+		int64 CompressedDataSize = CompressedImage.RawData.Num();
+
+		// check CompressedDataSize against int32_max , except if VT
+		// bForVirtualTextureStreamingBuild is false in this branch
+		check(CompressedDataSize < MAX_int32);
+
+		// CompressedImage sizes were padded up to multiple of 4 for d3d, no longer
+		UE_LOG(LogTextureUpload, Verbose, TEXT("Compressed Mip %d PF=%d : %dx%dx%d : %d ; up4 %dx%d=%d"),
+			MipIndex, (int)CompressedImage.PixelFormat,
+			CompressedImage.SizeX, CompressedImage.SizeY, CompressedImage.SizeZ, (int)CompressedDataSize,
+			(CompressedImage.SizeX + 3) & (~3),
+			(CompressedImage.SizeY + 3) & (~3),
+			((CompressedImage.SizeX + 3) & (~3)) * ((CompressedImage.SizeY + 3) & (~3)));
+
+		NewMip->BulkData.Lock(LOCK_READ_WRITE);
+		void* NewMipData = NewMip->BulkData.Realloc(CompressedDataSize);
+		FMemory::Memcpy(NewMipData, CompressedImage.RawData.GetData(), CompressedDataSize);
+		NewMip->BulkData.Unlock();
+
+		if (MipIndex == 0)
+		{
+			DerivedData->SizeX = CompressedImage.SizeX;
+			DerivedData->SizeY = CompressedImage.SizeY;
+			DerivedData->PixelFormat = (EPixelFormat)CompressedImage.PixelFormat;
+
+			// it would be better if CompressedImage just stored NumSlices, rather than recomputing it here
+			if (bVolume || bTextureArray)
+			{
+				DerivedData->SetNumSlices(CompressedImage.SizeZ);
+			}
+			else if (bCubemap)
+			{
+				DerivedData->SetNumSlices(6);
+			}
+			else
+			{
+				DerivedData->SetNumSlices(1);
+			}
+			DerivedData->SetIsCubemap(bCubemap);
+		}
+		else
+		{
+			check(CompressedImage.PixelFormat == DerivedData->PixelFormat);
+		}
+	}
+
+	FOptTexturePlatformData OptData;
+	OptData.NumMipsInTail = NumMipsInTail;
+	OptData.ExtData = ExtData;
+	DerivedData->SetOptData(OptData);
+
+	// Store it in the cache.
+	// @todo: This will remove the streaming bulk data, which we immediately reload below!
+	// Should ideally avoid this redundant work, but it only happens when we actually have 
+	// to build the texture, which should only ever be once.
+	BytesCached = PutDerivedDataInCache(DerivedData, KeySuffix, TexturePathName, bCubemap || (bVolume && !GSupportsVolumeTextureStreaming) || (bTextureArray && !GSupportsTexture2DArrayStreaming), bReplaceExistingDDC);
+}
+
 // Synchronous DDC1 texture build function
-void FTextureCacheDerivedDataWorker::BuildTexture(TArray<FTextureBuildSettings>& InBuildSettingsPerLayer, bool bReplaceExistingDDC)
+static void DDC1_BuildTexture(
+	ITextureCompressorModule* Compressor,
+	IImageWrapperModule* ImageWrapper,
+	const UTexture& Texture, // should be able to get rid of this and just check CompositeTextureData.IsValid()
+	const FString& TexturePathName,
+	ETextureCacheFlags CacheFlags,
+	FTextureSourceData& TextureData,
+	FTextureSourceData& CompositeTextureData,
+	const TArrayView<FTextureBuildSettings>& InBuildSettingsPerLayer,
+	const FString& KeySuffix,
+	bool bReplaceExistingDDC,
+	int64 RequiredMemoryEstimate,
+
+	FTexturePlatformData* DerivedData,
+	int64& BytesCached,
+	bool& bSucceeded	// currently completely ignored in the main build path, but used in the synchronous path
+	)
 {
 	TRACE_CPUPROFILER_EVENT_SCOPE(FTextureCacheDerivedDataWorker::BuildTexture);
 
@@ -586,7 +683,7 @@ void FTextureCacheDerivedDataWorker::BuildTexture(TArray<FTextureBuildSettings>&
 			// @todo: This will remove the streaming bulk data, which we immediately reload below!
 			// Should ideally avoid this redundant work, but it only happens when we actually have 
 			// to build the texture, which should only ever be once.
-			this->BytesCached = PutDerivedDataInCache(DerivedData, KeySuffix, TexturePathName, InBuildSettingsPerLayer[0].bCubemap || InBuildSettingsPerLayer[0].bVolume || InBuildSettingsPerLayer[0].bTextureArray, bReplaceExistingDDC);
+			BytesCached = PutDerivedDataInCache(DerivedData, KeySuffix, TexturePathName, InBuildSettingsPerLayer[0].bCubemap || InBuildSettingsPerLayer[0].bVolume || InBuildSettingsPerLayer[0].bTextureArray, bReplaceExistingDDC);
 
 			if (DerivedData->VTData->Chunks.Num())
 			{
@@ -627,7 +724,8 @@ void FTextureCacheDerivedDataWorker::BuildTexture(TArray<FTextureBuildSettings>&
 		DerivedData->SetIsCubemap(false);
 		DerivedData->VTData = nullptr;
 
-		FOptTexturePlatformData OptData;
+		uint32 NumMipsInTail;
+		uint32 ExtData;
 
 		// Compress the texture by calling texture compressor directly.
 		bool bImageHasAlphaChannel = false;
@@ -637,95 +735,16 @@ void FTextureCacheDerivedDataWorker::BuildTexture(TArray<FTextureBuildSettings>&
 			InBuildSettingsPerLayer[0],
 			TexturePathName,
 			CompressedMips,
-			OptData.NumMipsInTail,
-			OptData.ExtData,
+			NumMipsInTail,
+			ExtData,
 			&bImageHasAlphaChannel))
 		{
 			check(CompressedMips.Num());
 
-			// Build the derived data.
-			const int32 MipCount = CompressedMips.Num();
-			
-			// VT can be bigger than (1<<(MAX_TEXTURE_MIP_COUNT-1)) , but doesn't actually make all those mips
-			// bForVirtualTextureStreamingBuild is false in this branch
-			check( MipCount <= MAX_TEXTURE_MIP_COUNT );
-
-			for (int32 MipIndex = 0; MipIndex < MipCount; ++MipIndex)
-			{
-				const FCompressedImage2D& CompressedImage = CompressedMips[MipIndex];
-				FTexture2DMipMap* NewMip = new FTexture2DMipMap();
-				DerivedData->Mips.Add(NewMip);
-				NewMip->SizeX = CompressedImage.SizeX;
-				NewMip->SizeY = CompressedImage.SizeY;
-				NewMip->SizeZ = CompressedImage.SizeZ;
-				NewMip->FileRegionType = FFileRegion::SelectType(EPixelFormat(CompressedImage.PixelFormat));
-				check(NewMip->SizeZ == 1 || InBuildSettingsPerLayer[0].bVolume || InBuildSettingsPerLayer[0].bTextureArray); // Only volume & arrays can have SizeZ != 1
-				
-				check(CompressedImage.RawData.GetTypeSize() == 1);
-				int64 CompressedDataSize = CompressedImage.RawData.Num();
-
-				// check CompressedDataSize against int32_max , except if VT
-				// bForVirtualTextureStreamingBuild is false in this branch
-				check( CompressedDataSize < MAX_int32 );
-
-				// CompressedImage sizes were padded up to multiple of 4 for d3d, no longer
-				UE_LOG(LogTextureUpload,Verbose,TEXT("Compressed Mip %d PF=%d : %dx%dx%d : %d ; up4 %dx%d=%d"),
-					MipIndex,(int)CompressedImage.PixelFormat,
-					CompressedImage.SizeX,CompressedImage.SizeY,CompressedImage.SizeZ,(int)CompressedDataSize,
-					(CompressedImage.SizeX+3)&(~3),
-					(CompressedImage.SizeY+3)&(~3),
-					((CompressedImage.SizeX+3)&(~3))*((CompressedImage.SizeY+3)&(~3)));
-
-				NewMip->BulkData.Lock(LOCK_READ_WRITE);
-				void* NewMipData = NewMip->BulkData.Realloc(CompressedDataSize);
-				FMemory::Memcpy(NewMipData, CompressedImage.RawData.GetData(), CompressedDataSize);
-				NewMip->BulkData.Unlock();
-
-				if (MipIndex == 0)
-				{
-					DerivedData->SizeX = CompressedImage.SizeX;
-					DerivedData->SizeY = CompressedImage.SizeY;
-					DerivedData->PixelFormat = (EPixelFormat)CompressedImage.PixelFormat;
-
-					// it would be better if CompressedImage just stored NumSlices, rather than recomputing it here
-					if (InBuildSettingsPerLayer[0].bVolume || InBuildSettingsPerLayer[0].bTextureArray)
-					{
-						DerivedData->SetNumSlices(CompressedImage.SizeZ);
-					}
-					else if (InBuildSettingsPerLayer[0].bCubemap)
-					{
-						DerivedData->SetNumSlices(6);
-					}
-					else
-					{
-						DerivedData->SetNumSlices(1);
-					}
-					DerivedData->SetIsCubemap(InBuildSettingsPerLayer[0].bCubemap);
-				}
-				else
-				{
-					check(CompressedImage.PixelFormat == DerivedData->PixelFormat);
-				}
-			}
-
-			DerivedData->SetOptData(OptData);
-
-			// Store it in the cache.
-			// @todo: This will remove the streaming bulk data, which we immediately reload below!
-			// Should ideally avoid this redundant work, but it only happens when we actually have 
-			// to build the texture, which should only ever be once.
-			this->BytesCached = PutDerivedDataInCache(DerivedData, KeySuffix, TexturePathName, InBuildSettingsPerLayer[0].bCubemap || (InBuildSettingsPerLayer[0].bVolume && !GSupportsVolumeTextureStreaming) || (InBuildSettingsPerLayer[0].bTextureArray && !GSupportsTexture2DArrayStreaming), bReplaceExistingDDC);
-		}
-		else
-		{
-			// BuildTexture failed
-			// will log below
-			check( DerivedData->Mips.Num() == 0 );
-			DerivedData->Mips.Empty();
-		}
+			DDC1_StoreClassicTextureInDerivedData(
+				CompressedMips, DerivedData, InBuildSettingsPerLayer[0].bVolume, InBuildSettingsPerLayer[0].bTextureArray, InBuildSettingsPerLayer[0].bCubemap, 
+				NumMipsInTail, ExtData, bReplaceExistingDDC, TexturePathName, KeySuffix, BytesCached);
 
-		if (DerivedData->Mips.Num())
-		{
 			bSucceeded = true;
 
 			const bool bInlineMips = EnumHasAnyFlags(CacheFlags, ETextureCacheFlags::InlineMips);
@@ -740,6 +759,11 @@ void FTextureCacheDerivedDataWorker::BuildTexture(TArray<FTextureBuildSettings>&
 		}
 		else
 		{
+			// BuildTexture failed
+			// will log below
+			check( DerivedData->Mips.Num() == 0 );
+			DerivedData->Mips.Empty();
+
 			UE_LOG(LogTexture, Warning, TEXT("BuildTexture failed to build %s derived data for %s"), *InBuildSettingsPerLayer[0].TextureFormatName.GetPlainNameString(), *TexturePathName);
 		}
 	}
@@ -1248,7 +1272,8 @@ FTextureCacheDerivedDataWorker::FTextureCacheDerivedDataWorker(
 	}
 }
 
-bool FTextureCacheDerivedDataWorker::TryCacheStreamingMips(int32 FirstMipToLoad, int32 FirstMipToPrefetch)
+// Currently only used for prefetching (pulling data down from shared ddc to local ddc).
+static bool TryCacheStreamingMips(const FString& TexturePathName, int32 FirstMipToLoad, int32 FirstMipToPrefetch, FTexturePlatformData* DerivedData)
 {
 	using namespace UE::DerivedData;
 	check(DerivedData->DerivedDataKey.IsType<FString>());
@@ -1278,7 +1303,7 @@ bool FTextureCacheDerivedDataWorker::TryCacheStreamingMips(int32 FirstMipToLoad,
 
 	bool bOk = true;
 	FRequestOwner BlockingOwner(EPriority::Blocking);
-	GetCache().GetValue(MipRequests, BlockingOwner, [this, &bOk](FCacheGetValueResponse&& Response)
+	GetCache().GetValue(MipRequests, BlockingOwner, [DerivedData, &bOk](FCacheGetValueResponse&& Response)
 	{
 		bOk &= Response.Status == EStatus::Ok;
 		if (const FSharedBuffer MipBuffer = Response.Value.GetData().Decompress())
@@ -1294,37 +1319,42 @@ bool FTextureCacheDerivedDataWorker::TryCacheStreamingMips(int32 FirstMipToLoad,
 	return bOk;
 }
 
-// DDC1 primary fetch/build work function
-void FTextureCacheDerivedDataWorker::DoWork()
+static void DDC1_FetchAndFillDerivedData(
+	/* inputs */
+	const UTexture& Texture,
+	const FString& TexturePathName,
+	ETextureCacheFlags CacheFlags,
+	const TArrayView<FTextureBuildSettings>& BuildSettingsPerLayerFetchFirst,
+	const FTexturePlatformData::FTextureEncodeResultMetadata& FetchFirstMetadata,
+
+	const TArrayView<FTextureBuildSettings>& BuildSettingsPerLayerFetchOrBuild,
+	const FTexturePlatformData::FTextureEncodeResultMetadata& FetchOrBuildMetadata,
+
+	/* outputs */
+	FTexturePlatformData* DerivedData,
+	FString& KeySuffix,
+	bool& bSucceeded,
+	bool& bLoadedFromDDC,
+	bool& bInvalidVirtualTextureCompression,
+	int64& BytesCached
+	)
 {
 	using namespace UE;
 	using namespace UE::DerivedData;
 
-	TRACE_CPUPROFILER_EVENT_SCOPE(FTextureCacheDerivedDataWorker::DoWork);
-
-	const bool bForceRebuild = EnumHasAnyFlags(CacheFlags, ETextureCacheFlags::ForceRebuild);
-	const bool bAllowAsyncBuild = EnumHasAnyFlags(CacheFlags, ETextureCacheFlags::AllowAsyncBuild);
-	const bool bAllowAsyncLoading = EnumHasAnyFlags(CacheFlags, ETextureCacheFlags::AllowAsyncLoading);
-	const bool bForVirtualTextureStreamingBuild = EnumHasAnyFlags(CacheFlags, ETextureCacheFlags::ForVirtualTextureStreamingBuild);
-	const bool bValidateVirtualTextureCompression = CVarVTValidateCompressionOnLoad.GetValueOnAnyThread() != 0;
-	bool bInvalidVirtualTextureCompression = false;
+	bool bForceRebuild = EnumHasAnyFlags(CacheFlags, ETextureCacheFlags::ForceRebuild);
+	bool bForVirtualTextureStreamingBuild = EnumHasAnyFlags(CacheFlags, ETextureCacheFlags::ForVirtualTextureStreamingBuild);
 
 	FSharedBuffer RawDerivedData;
 
-	// Can't have a fetch first if we are rebuilding
-	if (bForceRebuild)
-	{
-		BuildSettingsPerLayerFetchFirst.Empty();
-	}
-
 	FString LocalDerivedDataKeySuffix;
 	FString LocalDerivedDataKey;
-	
+
 	FString FetchOrBuildKeySuffix;
 	GetTextureDerivedDataKeySuffix(Texture, BuildSettingsPerLayerFetchOrBuild.GetData(), FetchOrBuildKeySuffix);
 
 	bool bUsedFetchFirst = false;
-	if (BuildSettingsPerLayerFetchFirst.Num())
+	if (BuildSettingsPerLayerFetchFirst.Num() && !bForceRebuild)
 	{
 		FString FetchFirstKeySuffix;
 		GetTextureDerivedDataKeySuffix(Texture, BuildSettingsPerLayerFetchFirst.GetData(), FetchFirstKeySuffix);
@@ -1336,7 +1366,7 @@ void FTextureCacheDerivedDataWorker::DoWork()
 			GetTextureDerivedDataKeyFromSuffix(FetchFirstKeySuffix, FetchFirstKey);
 
 			FRequestOwner BlockingOwner(EPriority::Blocking);
-			GetCache().GetValue({{{TexturePathName}, ConvertLegacyCacheKey(FetchFirstKey)}}, BlockingOwner, [&RawDerivedData](FCacheGetValueResponse&& Response)
+			GetCache().GetValue({ {{TexturePathName}, ConvertLegacyCacheKey(FetchFirstKey)} }, BlockingOwner, [&RawDerivedData](FCacheGetValueResponse&& Response)
 			{
 				RawDerivedData = Response.Value.GetData().Decompress();
 			});
@@ -1361,7 +1391,7 @@ void FTextureCacheDerivedDataWorker::DoWork()
 		if (bForceRebuild == false)
 		{
 			FRequestOwner BlockingOwner(EPriority::Blocking);
-			GetCache().GetValue({{{TexturePathName}, ConvertLegacyCacheKey(LocalDerivedDataKey)}}, BlockingOwner, [&RawDerivedData](FCacheGetValueResponse&& Response)
+			GetCache().GetValue({ {{TexturePathName}, ConvertLegacyCacheKey(LocalDerivedDataKey)} }, BlockingOwner, [&RawDerivedData](FCacheGetValueResponse&& Response)
 			{
 				RawDerivedData = Response.Value.GetData().Decompress();
 			});
@@ -1373,7 +1403,6 @@ void FTextureCacheDerivedDataWorker::DoWork()
 
 	KeySuffix = LocalDerivedDataKeySuffix;
 	DerivedData->DerivedDataKey.Emplace<FString>(LocalDerivedDataKey);
-	TArray<FTextureBuildSettings>& ActiveBuildSettingsPerLayer = bUsedFetchFirst ? BuildSettingsPerLayerFetchFirst : BuildSettingsPerLayerFetchOrBuild;
 	DerivedData->ResultMetadata = bUsedFetchFirst ? FetchFirstMetadata : FetchOrBuildMetadata;
 
 	if (bLoadedFromDDC)
@@ -1462,7 +1491,7 @@ void FTextureCacheDerivedDataWorker::DoWork()
 		}
 		else if (bInlineMips)
 		{
-			bSucceeded = DerivedData->TryInlineMipData(ActiveBuildSettingsPerLayer[0].LODBiasWithCinematicMips, TexturePathName);
+			bSucceeded = DerivedData->TryInlineMipData(BuildSettingsPerLayerFetchOrBuild[0].LODBiasWithCinematicMips, TexturePathName);
 
 			if (!bSucceeded)
 			{
@@ -1473,9 +1502,9 @@ void FTextureCacheDerivedDataWorker::DoWork()
 		{
 			if (bForVirtualTextureStreamingBuild)
 			{
-				bSucceeded =	DerivedData->VTData != nullptr &&
-								DerivedData->VTData->IsInitialized() &&
-								DerivedData->AreDerivedVTChunksAvailable(TexturePathName);
+				bSucceeded = DerivedData->VTData != nullptr &&
+					DerivedData->VTData->IsInitialized() &&
+					DerivedData->AreDerivedVTChunksAvailable(TexturePathName);
 
 				if (!bSucceeded)
 				{
@@ -1488,7 +1517,7 @@ void FTextureCacheDerivedDataWorker::DoWork()
 				const int32 FirstMipToLoad = FirstResidentMipIndex;
 				const int32 FirstNonStreamingMipIndex = DerivedData->Mips.Num() - DerivedData->GetNumNonStreamingMips(!bDisableStreaming);
 				const int32 FirstMipToPrefetch = IsInGameThread() ? FirstMipToLoad : bDisableStreaming ? 0 : FirstNonStreamingMipIndex;
-				bSucceeded = TryCacheStreamingMips(FirstMipToLoad, FirstMipToPrefetch);
+				bSucceeded = TryCacheStreamingMips(TexturePathName, FirstMipToLoad, FirstMipToPrefetch, DerivedData);
 				if (!bSucceeded)
 				{
 					UE_LOG(LogTexture, Display, TEXT("Texture %s is missing derived mips. The texture will be rebuilt."), *TexturePathName);
@@ -1520,41 +1549,328 @@ void FTextureCacheDerivedDataWorker::DoWork()
 			bLoadedFromDDC = false;
 		}
 	}
-	
-	if (!bSucceeded && bAllowAsyncBuild)
+
+}
+
+static bool DDC1_IsTextureDataValid(const FTextureSourceData& TextureData, const FTextureSourceData& CompositeTextureData)
+{
+	bool bTextureDataValid = TextureData.Blocks.Num() && TextureData.Blocks[0].MipsPerLayer.Num() && TextureData.Blocks[0].MipsPerLayer[0].Num();
+	if (CompositeTextureData.IsValid()) // Says IsValid, but means whether or not we _need_ composite texture data.
 	{
-		bool bHasTextureSourceMips = false;
-		if (TextureData.IsValid() && Texture.Source.IsBulkDataLoaded())
-		{
-			TextureData.GetSourceMips(Texture.Source, ImageWrapper);
-			bHasTextureSourceMips = true;
-		}
+		// here we know we _need_ composite texture data, so we actually check if the stuff we loaded is valid.
+		bool bCompositeDataValid = CompositeTextureData.Blocks.Num() && CompositeTextureData.Blocks[0].MipsPerLayer.Num() && CompositeTextureData.Blocks[0].MipsPerLayer[0].Num();
+		return bTextureDataValid && bCompositeDataValid;
+	}
+	return bTextureDataValid;
+}
 
-		bool bHasCompositeTextureSourceMips = false;
-		if (CompositeTextureData.IsValid() && Texture.CompositeTexture && Texture.CompositeTexture->Source.IsBulkDataLoaded())
-		{
-			CompositeTextureData.GetSourceMips(Texture.CompositeTexture->Source, ImageWrapper);
-			bHasCompositeTextureSourceMips = true;
-		}
+// Tries to get the source texture data resident for building the texture.
+static bool DDC1_LoadAndValidateTextureData(
+	UTexture& Texture,
+	FTextureSourceData& TextureData,
+	FTextureSourceData& CompositeTextureData,
+	IImageWrapperModule* ImageWrapper,
+	bool bAllowAsyncLoading
+)
+{
+	bool bHasTextureSourceMips = false;
+	if (TextureData.IsValid() && Texture.Source.IsBulkDataLoaded())
+	{
+		TextureData.GetSourceMips(Texture.Source, ImageWrapper);
+		bHasTextureSourceMips = true;
+	}
+
+	bool bHasCompositeTextureSourceMips = false;
+	if (CompositeTextureData.IsValid() && Texture.CompositeTexture && Texture.CompositeTexture->Source.IsBulkDataLoaded())
+	{
+		CompositeTextureData.GetSourceMips(Texture.CompositeTexture->Source, ImageWrapper);
+		bHasCompositeTextureSourceMips = true;
+	}
 
-		if (bAllowAsyncLoading && !bHasTextureSourceMips)
+	if (bAllowAsyncLoading)
+	{
+		if (!bHasTextureSourceMips)
 		{
 			TextureData.GetAsyncSourceMips(ImageWrapper);
 			TextureData.AsyncSource.RemoveBulkData();
 		}
 
-		if (bAllowAsyncLoading && !bHasCompositeTextureSourceMips)
+		if (!bHasCompositeTextureSourceMips)
 		{
 			CompositeTextureData.GetAsyncSourceMips(ImageWrapper);
 			CompositeTextureData.AsyncSource.RemoveBulkData();
 		}
+	}
+
+	return DDC1_IsTextureDataValid(TextureData, CompositeTextureData);
+}
+
+
+bool DDC1_BuildTiledClassicTexture(
+	ITextureCompressorModule* Compressor,
+	IImageWrapperModule* ImageWrapper,
+	UTexture& Texture,
+	const FString& TexturePathName,
+	const TArrayView<FTextureBuildSettings> BuildSettingsPerLayerFetchFirst,
+	const TArrayView<FTextureBuildSettings> BuildSettingsPerLayerFetchOrBuild,
+	const FTexturePlatformData::FTextureEncodeResultMetadata& FetchFirstMetadata,
+	const FTexturePlatformData::FTextureEncodeResultMetadata& FetchOrBuildMetadata,
+	FTextureSourceData& TextureData,
+	FTextureSourceData& CompositeTextureData,
+	ETextureCacheFlags CacheFlags,
+	int32 RequiredMemoryEstimate,
+	const FString& KeySuffix,
+	// outputs
+	FTexturePlatformData* DerivedData,
+	int64& BytesCached
+)
+{
+	// we know we are a child format if we have a tiler.
+	const ITextureTiler* Tiler = BuildSettingsPerLayerFetchOrBuild[0].Tiler;
+	const FChildTextureFormat* ChildFormat = GetTextureFormatManager()->FindTextureFormat(BuildSettingsPerLayerFetchOrBuild[0].TextureFormatName)->GetChildFormat();
+
+	// NonVT textures only have one layer.
+	// We need to get the linear texture, which means we have to create the settings for it.
+	TArray<FTextureBuildSettings, TInlineAllocator<1>> LinearSettingsPerLayerFetchFirst;
+	if (BuildSettingsPerLayerFetchFirst.Num())
+	{
+		LinearSettingsPerLayerFetchFirst.Add(ChildFormat->GetBaseTextureBuildSettings(BuildSettingsPerLayerFetchFirst[0]));
+	}
+
+	TArray<FTextureBuildSettings, TInlineAllocator<1>> LinearSettingsPerLayerFetchOrBuild;
+	LinearSettingsPerLayerFetchOrBuild.Add(ChildFormat->GetBaseTextureBuildSettings(BuildSettingsPerLayerFetchOrBuild[0]));
+
+	// Now try and fetch.
+	FTexturePlatformData LinearDerivedData;
+	FString LinearKeySuffix;
+	int64 LinearBytesCached = 0;
+	bool bLinearDDCCorrupted = false;
+	bool bLinearSucceeded = false;
+	bool bLinearLoadedFromDDC = false;
+	DDC1_FetchAndFillDerivedData(
+		Texture, TexturePathName, CacheFlags,
+		LinearSettingsPerLayerFetchFirst, FetchFirstMetadata, 
+		LinearSettingsPerLayerFetchOrBuild, FetchOrBuildMetadata,
+		&LinearDerivedData, LinearKeySuffix, bLinearSucceeded, bLinearLoadedFromDDC, bLinearDDCCorrupted, LinearBytesCached);
+
+	BytesCached = LinearBytesCached;
+	bool bHasLinearDerivedData = bLinearSucceeded && bLinearLoadedFromDDC;
+
+	if (bHasLinearDerivedData == false)
+	{
+		// Linear data didn't exist, need to build it.
+		bool bGotSourceTextureData = DDC1_LoadAndValidateTextureData(Texture, TextureData, CompositeTextureData, ImageWrapper, EnumHasAnyFlags(CacheFlags, ETextureCacheFlags::AllowAsyncLoading));
+		if (bGotSourceTextureData)
+		{
+			// We know we want all the mips for tiling, so force inline
+			ETextureCacheFlags LinearCacheFlags = CacheFlags;
+			EnumAddFlags(LinearCacheFlags, ETextureCacheFlags::InlineMips);
+
+			// Note that this will update the DDC with the linear texture if we end up building _before_ the linear platforms!
+			DDC1_BuildTexture(
+				Compressor,
+				ImageWrapper,
+				Texture,
+				TexturePathName,
+				CacheFlags,
+				TextureData,
+				CompositeTextureData,
+				LinearSettingsPerLayerFetchOrBuild,
+				LinearKeySuffix,
+				bLinearDDCCorrupted,
+				RequiredMemoryEstimate,
+				&LinearDerivedData,
+				LinearBytesCached,
+				bHasLinearDerivedData);
+		}
+	}
+
+	if (bHasLinearDerivedData == false)
+	{
+		UE_LOG(LogTexture, Error, TEXT("Tiling texture build was unable to fetch or build the linear texture source: %s"), *TexturePathName);
+		return false;
+	}
+
+	check(LinearDerivedData.GetNumMipsInTail() == 0);
+
+	// The linear bits are built - need to fetch
+	void* LinearMipData[MAX_TEXTURE_MIP_COUNT] = {};
+	int64 LinearMipSizes[MAX_TEXTURE_MIP_COUNT];
+	if (LinearDerivedData.TryLoadMipsWithSizes(0, LinearMipData, LinearMipSizes, TexturePathName) == false)
+	{
+		UE_LOG(LogTexture, Error, TEXT("Tiling texture build was unable to reload the linear texture mips after building/fetching: %s"), *TexturePathName);
+		return false;
+	}
+
+	// Have all the data - do some sanity checks as we convert to the metadata format the tiler expects.
+	TArray<FMemoryView, TInlineAllocator<MAX_TEXTURE_MIP_COUNT>> InputTextureMipViews;
+	FEncodedTextureDescription TextureDescription;
+	FEncodedTextureExtendedData TextureExtendedData;
+	int32 OutputTextureNumStreamingMips;
+	{
+#if DO_CHECK
+		TArray<FImage>& SourceMips = TextureData.Blocks[0].MipsPerLayer[0];
+		if (SourceMips.Num()) // if we're tiling existing linear data then the texture data isn't actually loaded so we can't do this sanity check.
+		{
+			int32 LinearMip0SizeX, LinearMip0SizeY, LinearMip0NumSlices;
+			int32 LinearMipCount = Compressor->GetMipCountForBuildSettings(SourceMips[0].SizeX, SourceMips[0].SizeY, SourceMips[0].NumSlices, SourceMips.Num(), LinearSettingsPerLayerFetchOrBuild[0],
+				LinearMip0SizeX,
+				LinearMip0SizeY,
+				LinearMip0NumSlices);
+			check(LinearDerivedData.Mips[0].SizeX == LinearMip0SizeX);
+			check(LinearDerivedData.Mips[0].SizeY == LinearMip0SizeY);
+			check(LinearDerivedData.Mips.Num() == LinearMipCount);
+			check(LinearDerivedData.GetNumSlices() == LinearMip0NumSlices);
+		}
+#endif
+
+		LinearSettingsPerLayerFetchOrBuild[0].GetEncodedTextureDescriptionWithPixelFormat(
+			&TextureDescription,
+			LinearDerivedData.PixelFormat, LinearDerivedData.Mips[0].SizeX, LinearDerivedData.Mips[0].SizeY, LinearDerivedData.GetNumSlices(), LinearDerivedData.Mips.Num());
+
+		for (int32 MipIndex = 0; MipIndex < TextureDescription.NumMips; MipIndex++)
+		{
+			check(LinearMipSizes[MipIndex] == TextureDescription.GetMipSizeInBytes(MipIndex));
+		}
+
+		TextureExtendedData = Tiler->GetExtendedDataForTexture(TextureDescription);
+		OutputTextureNumStreamingMips = TextureDescription.GetNumStreamingMips(&TextureExtendedData, GenerateTextureEngineParameters());
+
+		for (int32 MipIndex = 0; MipIndex < TextureDescription.NumMips; MipIndex++)
+		{
+			InputTextureMipViews.Add(FMemoryView(LinearMipData[MipIndex], LinearMipSizes[MipIndex]));
+		}
+	}
+
+	TArray<FCompressedImage2D> TiledMips;
+	TiledMips.AddDefaulted(TextureDescription.NumMips);
+
+	// If the platform packs mip tails, we need to pass all the relevant mip buffers at once.
+	int32 MipTailIndex, MipsInTail;
+	TextureDescription.GetEncodedMipIterators(&TextureExtendedData, MipTailIndex, MipsInTail);
+
+	UE_LOG(LogTexture, Display, TEXT("Tiling %s"), *TexturePathName);
+
+	// Do the actual tiling.
+	for (int32 EncodedMipIndex = 0; EncodedMipIndex < MipTailIndex + 1; EncodedMipIndex++)
+	{
+		int32 MipsRepresentedThisIndex = EncodedMipIndex == MipTailIndex ? MipsInTail : 1;
+
+		TArrayView<FMemoryView> MipsForLevel = MakeArrayView(InputTextureMipViews.GetData() + EncodedMipIndex, MipsRepresentedThisIndex);
+
+		FSharedBuffer MipData = Tiler->ProcessMipLevel(TextureDescription, TextureExtendedData, MipsForLevel, EncodedMipIndex);
+		FIntVector3 MipDims = TextureDescription.GetMipDimensions(EncodedMipIndex);
+
+		// Make sure we got the size we advertised prior to the build. If this ever fires then we
+		// have a critical mismatch!
+		check(TextureExtendedData.MipSizesInBytes[EncodedMipIndex] == MipData.GetSize());
+
+		FCompressedImage2D& TiledMip = TiledMips[EncodedMipIndex];
+		TiledMip.PixelFormat = LinearDerivedData.PixelFormat;
+		TiledMip.SizeX = MipDims.X;
+		TiledMip.SizeY = MipDims.Y;
+		TiledMip.SizeZ = TextureDescription.GetNumSlices_WithDepth(EncodedMipIndex);
+		if (TextureDescription.bCubeMap && !TextureDescription.bTextureArray)
+		{
+			TiledMip.SizeZ = 1; // compressed image 2d wants it this way...
+		}
+
+		// \todo try and Move this data rather than copying? We use FSharedBuffer as that's the future way,
+		// but we're interacting with older systems that didn't have it, and we can't Move() from an FSharedBuffer.
+		TiledMip.RawData.AddUninitialized(MipData.GetSize());
+		FMemory::Memcpy(TiledMip.RawData.GetData(), MipData.GetData(), MipData.GetSize());
+	} // end for each mip
+
+	for (int32 MipIndex = 0; MipIndex < TextureDescription.NumMips; ++MipIndex)
+	{
+		FMemory::Free(LinearMipData[MipIndex]);
+	}
+
+	// The derived data expects to have mips (with no data) for the packed tail, if there is one
+	for (int32 MipIndex = MipTailIndex + 1; MipIndex < TextureDescription.NumMips; ++MipIndex)
+	{
+		FCompressedImage2D& PrevMip = TiledMips[MipIndex - 1];
+		FCompressedImage2D& DestMip = TiledMips[MipIndex];
+		DestMip.SizeX = FMath::Max(1, PrevMip.SizeX >> 1);
+		DestMip.SizeY = FMath::Max(1, PrevMip.SizeY >> 1);
+		DestMip.SizeZ = TextureDescription.bVolumeTexture ? FMath::Max(1, PrevMip.SizeZ >> 1) : PrevMip.SizeZ;
+		DestMip.PixelFormat = PrevMip.PixelFormat;
+	}
+
+	// We now have the final (tiled) data, and need to fill out the actual build output
+	int64 TiledBytesCached;
+	DDC1_StoreClassicTextureInDerivedData(TiledMips, DerivedData, TextureDescription.bVolumeTexture, TextureDescription.bTextureArray, TextureDescription.bCubeMap,
+		TextureExtendedData.NumMipsInTail, TextureExtendedData.ExtData, false, TexturePathName, KeySuffix, TiledBytesCached);
+	
+	BytesCached += TiledBytesCached;
+
+	// Do we need to reload streaming mips (evicted during DDC store)
+	if (EnumHasAnyFlags(CacheFlags, ETextureCacheFlags::InlineMips))
+	{
+		if (DerivedData->TryInlineMipData(LinearSettingsPerLayerFetchOrBuild[0].LODBiasWithCinematicMips, TexturePathName) == false)
+		{
+			UE_LOG(LogTexture, Display, TEXT("Tiled texture build failed to put and then read back tiled mipmap data from DDC for %s"), *TexturePathName);
+		}
+	}
+
+	return true;
+}
+
+
+// DDC1 primary fetch/build work function
+void FTextureCacheDerivedDataWorker::DoWork()
+{
+	using namespace UE;
+	using namespace UE::DerivedData;
+
+	TRACE_CPUPROFILER_EVENT_SCOPE(FTextureCacheDerivedDataWorker::DoWork);
+
+	const bool bAllowAsyncBuild = EnumHasAnyFlags(CacheFlags, ETextureCacheFlags::AllowAsyncBuild);
+	const bool bAllowAsyncLoading = EnumHasAnyFlags(CacheFlags, ETextureCacheFlags::AllowAsyncLoading);
+	const bool bForVirtualTextureStreamingBuild = EnumHasAnyFlags(CacheFlags, ETextureCacheFlags::ForVirtualTextureStreamingBuild);
+	bool bInvalidVirtualTextureCompression = false;
+
+	bSucceeded = false;
+	bLoadedFromDDC = false;
+
+	DDC1_FetchAndFillDerivedData(
+		/* inputs */ Texture, TexturePathName, CacheFlags, BuildSettingsPerLayerFetchFirst, FetchFirstMetadata, BuildSettingsPerLayerFetchOrBuild, FetchOrBuildMetadata, 
+		/* outputs */ DerivedData, KeySuffix, bSucceeded, bLoadedFromDDC, bInvalidVirtualTextureCompression, BytesCached);
+
+	if (BuildSettingsPerLayerFetchOrBuild[0].Tiler && !bForVirtualTextureStreamingBuild)
+	{
+		if (CVarForceRetileTextures.GetValueOnAnyThread())
+		{
+			// We do this after the fetch so it can fill out the metadata nad key suffix that gets used.
+			bSucceeded = false;
+			bLoadedFromDDC = false;
+
+			DerivedData->Mips.Empty();
+			delete DerivedData->VTData;
+			DerivedData->VTData = nullptr;
+		}
+	}
 
-		if (TextureData.Blocks.Num() && TextureData.Blocks[0].MipsPerLayer.Num() && TextureData.Blocks[0].MipsPerLayer[0].Num() && 
-			(!CompositeTextureData.IsValid() || (CompositeTextureData.Blocks.Num() && CompositeTextureData.Blocks[0].MipsPerLayer.Num() && CompositeTextureData.Blocks[0].MipsPerLayer[0].Num())))
+	if (!bSucceeded && bAllowAsyncBuild)
+	{
+		if (DDC1_LoadAndValidateTextureData(Texture, TextureData, CompositeTextureData, ImageWrapper, bAllowAsyncLoading))
 		{
 			// Replace any existing DDC data, if corrupt compression was detected
 			const bool bReplaceExistingDDC = bInvalidVirtualTextureCompression;
-			BuildTexture(ActiveBuildSettingsPerLayer, bReplaceExistingDDC);
+
+			if (BuildSettingsPerLayerFetchOrBuild[0].Tiler && !bForVirtualTextureStreamingBuild)
+			{
+				bSucceeded = DDC1_BuildTiledClassicTexture(
+					Compressor, ImageWrapper, Texture, TexturePathName, BuildSettingsPerLayerFetchFirst, BuildSettingsPerLayerFetchOrBuild, FetchFirstMetadata, FetchOrBuildMetadata,
+					TextureData, CompositeTextureData, CacheFlags, RequiredMemoryEstimate, KeySuffix, DerivedData, BytesCached);
+			}
+			else
+			{
+				DDC1_BuildTexture(
+					Compressor, ImageWrapper, Texture, TexturePathName, CacheFlags, TextureData, CompositeTextureData, BuildSettingsPerLayerFetchOrBuild, KeySuffix, bReplaceExistingDDC,
+					RequiredMemoryEstimate, DerivedData, BytesCached, bSucceeded);
+			}
+
 			if (bInvalidVirtualTextureCompression && DerivedData->VTData)
 			{
 				// If we loaded data that turned out to be corrupt, flag it here so we can also recreate the VT data cached to local /DerivedDataCache/VT/ directory
@@ -1562,9 +1878,9 @@ void FTextureCacheDerivedDataWorker::DoWork()
 				{
 					Chunk.bCorruptDataLoadedFromDDC = true;
 				}
-
 			}
 
+			// this always ignores whether the build failed?
 			bSucceeded = true;
 		}
 		else
@@ -1588,12 +1904,12 @@ void FTextureCacheDerivedDataWorker::DoWork()
 
 void FTextureCacheDerivedDataWorker::Finalize()
 {
-	// if we couldn't get from the DDC or didn't build synchronously, then we have to build now. 
-	// This is a super edge case that should rarely happen.
-	// (update) this always happens with a ForceRebuildPlatformData, as its a synchronous build that
-	//		uses this worker class. The DoWork never does work if async is not set unless its a clean
-	//		get, so we land here. Note that we never _fetch_ here, it's only ever a full build, so
-	//		we can ignore FetchFirst.
+	// Building happens here whenever the ddc is missed and async builds aren't allowed.
+	// This generally doesn't happen, but does in a few cases:
+	// --	always happens with a ForceRebuildPlatformData, which is called whenever mip data is requested
+	//		in the editor and is missing for some reason.
+	// --	always with a lighting build, as the async light/shadowmap tasks will disallow async builds. 
+	// --	if the texture compiler cvar disallows async texture compilation
 	if (!bSucceeded)
 	{
 		TextureData.GetSourceMips(Texture.Source, ImageWrapper);
@@ -1602,7 +1918,38 @@ void FTextureCacheDerivedDataWorker::Finalize()
 			CompositeTextureData.GetSourceMips(Texture.CompositeTexture->Source, ImageWrapper);
 		}
 
-		BuildTexture(BuildSettingsPerLayerFetchOrBuild);
+		if (DDC1_IsTextureDataValid(TextureData, CompositeTextureData) == false)
+		{
+			UE_LOG(LogTexture, Error, TEXT("Unable to get texture source data for synchronous build of %s"), *TexturePathName);
+		}
+		else
+		{
+			if (BuildSettingsPerLayerFetchOrBuild[0].Tiler && !EnumHasAnyFlags(CacheFlags, ETextureCacheFlags::ForVirtualTextureStreamingBuild))
+			{
+				bSucceeded = DDC1_BuildTiledClassicTexture(
+					Compressor,
+					ImageWrapper,
+					Texture,
+					TexturePathName,
+					BuildSettingsPerLayerFetchFirst,
+					BuildSettingsPerLayerFetchOrBuild,
+					FetchFirstMetadata,
+					FetchOrBuildMetadata,
+					TextureData,
+					CompositeTextureData,
+					CacheFlags,
+					RequiredMemoryEstimate,
+					KeySuffix,
+					DerivedData,
+					BytesCached);
+			}
+			else
+			{
+				DDC1_BuildTexture(Compressor, ImageWrapper, Texture, TexturePathName, CacheFlags,
+					TextureData, CompositeTextureData, BuildSettingsPerLayerFetchOrBuild, KeySuffix, false /* currently corrupt vt data is not routed out of DoWork() */,
+					RequiredMemoryEstimate, DerivedData, BytesCached, bSucceeded);
+			}
+		}
 	}
 		
 	if (bSucceeded && BuildSettingsPerLayerFetchOrBuild[0].bVirtualStreamable) // Texture.VirtualTextureStreaming is more a hint that might be overruled by the buildsettings
diff --git a/Engine/Source/Runtime/Engine/Private/TextureDerivedDataTask.h b/Engine/Source/Runtime/Engine/Private/TextureDerivedDataTask.h
--- a/Engine/Source/Runtime/Engine/Private/TextureDerivedDataTask.h
+++ b/Engine/Source/Runtime/Engine/Private/TextureDerivedDataTask.h
@@ -161,12 +161,6 @@ class FTextureCacheDerivedDataWorker : public FNonAbandonableTask
 	bool bSucceeded = false;
 	/** true if the derived data was pulled from DDC */
 	bool bLoadedFromDDC = false;
-
-	/** Check for existence of the streaming mips in the cache, and optionally prefetch or load them. */
-	bool TryCacheStreamingMips(int32 FirstMipToLoad, int32 FirstMipToPrefetch);
-
-	/** Build the texture. This function is safe to call from any thread. */
-	void BuildTexture(TArray<FTextureBuildSettings>& InBuildSettingsPerLayer, bool bReplaceExistingDDC = false);
 public:
 
 	/** Initialization constructor. */
