Revert changes from CL 19124885, so a ticking NetDriver can shutdown another one

diff --git a/Engine/Source/Runtime/Engine/Private/NetDriver.cpp b/Engine/Source/Runtime/Engine/Private/NetDriver.cpp
--- a/Engine/Source/Runtime/Engine/Private/NetDriver.cpp
+++ b/Engine/Source/Runtime/Engine/Private/NetDriver.cpp
@@ -805,7 +805,7 @@ void UNetDriver::TickFlush(float DeltaSeconds)
 {
 	LLM_SCOPE_BYTAG(NetDriver);
 
-	bInTick = true;
+	TGuardValue<bool> GuardInNetTick(bInTick, true);
 
 	CSV_SCOPED_TIMING_STAT_EXCLUSIVE(NetworkOutgoing);
 	SCOPE_CYCLE_COUNTER(STAT_NetTickFlush);
@@ -1255,7 +1255,6 @@ void UNetDriver::PostTickFlush()
 		UOnlineEngineInterface::Get()->ClearVoicePackets(World);
 	}
 
-	bInTick = false;
 	if (bPendingDestruction)
 	{
 		if (World)
@@ -1696,7 +1695,7 @@ struct FRPCCSVTracker
 
 void UNetDriver::TickDispatch( float DeltaTime )
 {
-	bInTick = true;
+	TGuardValue<bool> GuardInNetTick(bInTick, true);
 
 	SendCycles=0;
 
@@ -1810,7 +1809,6 @@ void UNetDriver::PostTickDispatch()
 		GReceiveRPCTimingEnabled = false;
 	}
 
-	bInTick = false;
 	if (bPendingDestruction)
 	{
 		if (World)
