World Partition: actor descriptors can now be runtime only, which means that they won't be loaded in the editor. They still can be pinned in the world outliner though for debugging purposes.
Made WorldPartitionHLOD and LandscapeSplineMeshesActor runtime only actors.

#rnx

diff --git a/Engine/Source/Runtime/Core/Public/UObject/FortniteMainBranchObjectVersion.h b/Engine/Source/Runtime/Core/Public/UObject/FortniteMainBranchObjectVersion.h
--- a/Engine/Source/Runtime/Core/Public/UObject/FortniteMainBranchObjectVersion.h
+++ b/Engine/Source/Runtime/Core/Public/UObject/FortniteMainBranchObjectVersion.h
@@ -247,6 +247,9 @@ struct CORE_API FFortniteMainBranchObjectVersion
 		// Serialize ContentBundleGuid in WorldPartitionActorDesc
 		WorldPartitionActorDescSerializeContentBundleGuid,
 
+		// Serialize IsActorRuntimeOnly in WorldPartitionActorDesc
+		WorldPartitionActorDescSerializeActorIsRuntimeOnly,
+
 		// -----<new versions can be added above this line>-------------------------------------------------
 		VersionPlusOne,
 		LatestVersion = VersionPlusOne - 1
diff --git a/Engine/Source/Runtime/Engine/Classes/GameFramework/Actor.h b/Engine/Source/Runtime/Engine/Classes/GameFramework/Actor.h
--- a/Engine/Source/Runtime/Engine/Classes/GameFramework/Actor.h
+++ b/Engine/Source/Runtime/Engine/Classes/GameFramework/Actor.h
@@ -2025,6 +2025,7 @@ public:
 	virtual bool CanBeInCluster() const override;
 	static void AddReferencedObjects(UObject* InThis, FReferenceCollector& Collector);
 	virtual bool IsEditorOnly() const override;
+	virtual bool IsRuntimeOnly() const { return false; }
 	virtual bool IsAsset() const override;
 
 	PRAGMA_DISABLE_DEPRECATION_WARNINGS // Suppress compiler warning on override of deprecated function
diff --git a/Engine/Source/Runtime/Engine/Private/WorldPartition/HLOD/HLODActorDesc.cpp b/Engine/Source/Runtime/Engine/Private/WorldPartition/HLOD/HLODActorDesc.cpp
--- a/Engine/Source/Runtime/Engine/Private/WorldPartition/HLOD/HLODActorDesc.cpp
+++ b/Engine/Source/Runtime/Engine/Private/WorldPartition/HLOD/HLODActorDesc.cpp
@@ -8,6 +8,7 @@
 #include "Hash/CityHashHelpers.h"
 #include "UObject/UE5MainStreamObjectVersion.h"
 #include "UObject/UE5ReleaseStreamObjectVersion.h"
+#include "UObject/FortniteMainBranchObjectVersion.h"
 #include "WorldPartition/HLOD/HLODActor.h"
 #include "WorldPartition/HLOD/HLODLayer.h"
 
@@ -38,6 +39,7 @@ void FHLODActorDesc::Serialize(FArchive& Ar)
 {
 	Ar.UsingCustomVersion(FUE5MainStreamObjectVersion::GUID);
 	Ar.UsingCustomVersion(FUE5ReleaseStreamObjectVersion::GUID);
+	Ar.UsingCustomVersion(FFortniteMainBranchObjectVersion::GUID);
 
 	FWorldPartitionActorDesc::Serialize(Ar);
 
@@ -61,6 +63,11 @@ void FHLODActorDesc::Serialize(FArchive& Ar)
 	{
 		Ar << CellHash;
 	}
+
+	if (Ar.CustomVer(FFortniteMainBranchObjectVersion::GUID) < FFortniteMainBranchObjectVersion::WorldPartitionActorDescSerializeActorIsRuntimeOnly)
+	{
+		bActorIsRuntimeOnly = true;
+	}
 }
 
 bool FHLODActorDesc::Equals(const FWorldPartitionActorDesc* Other) const
diff --git a/Engine/Source/Runtime/Engine/Private/WorldPartition/WorldPartitionActorDesc.cpp b/Engine/Source/Runtime/Engine/Private/WorldPartition/WorldPartitionActorDesc.cpp
--- a/Engine/Source/Runtime/Engine/Private/WorldPartition/WorldPartitionActorDesc.cpp
+++ b/Engine/Source/Runtime/Engine/Private/WorldPartition/WorldPartitionActorDesc.cpp
@@ -63,6 +63,7 @@ void FWorldPartitionActorDesc::Init(const AActor* InActor)
 	RuntimeGrid = InActor->GetRuntimeGrid();
 	bIsSpatiallyLoaded = InActor->GetIsSpatiallyLoaded();
 	bActorIsEditorOnly = InActor->IsEditorOnly();
+	bActorIsRuntimeOnly = InActor->IsRuntimeOnly();
 	bActorIsHLODRelevant = InActor->IsHLODRelevant();
 	HLODLayer = InActor->GetHLODLayer() ? FName(InActor->GetHLODLayer()->GetPathName()) : FName();
 	
@@ -184,6 +185,7 @@ bool FWorldPartitionActorDesc::Equals(const FWorldPartitionActorDesc* Other) con
 		BoundsExtent.Equals(Other->BoundsExtent, 0.1f) &&
 		RuntimeGrid == Other->RuntimeGrid &&
 		bActorIsEditorOnly == Other->bActorIsEditorOnly &&
+		bActorIsRuntimeOnly == Other->bActorIsRuntimeOnly &&
 		bActorIsHLODRelevant == Other->bActorIsHLODRelevant &&
 		bIsUsingDataLayerAsset == Other->bIsUsingDataLayerAsset &&
 		HLODLayer == Other->HLODLayer &&
@@ -253,7 +255,7 @@ FString FWorldPartitionActorDesc::ToString(EToStringMode Mode) const
 	if (Mode >= EToStringMode::Compact)
 	{
 		Result.Appendf(
-			TEXT(" BaseClass:%s NativeClass:%s Name:%s Label:%s SpatiallyLoaded:%s Bounds:%s RuntimeGrid:%s EditorOnly:%s HLODRelevant:%s"),
+			TEXT(" BaseClass:%s NativeClass:%s Name:%s Label:%s SpatiallyLoaded:%s Bounds:%s RuntimeGrid:%s EditorOnly:%s RuntimeOnly:%s HLODRelevant:%s"),
 			*BaseClass.ToString(), 
 			*NativeClass.ToString(), 
 			*GetActorName().ToString(),
@@ -262,6 +264,7 @@ FString FWorldPartitionActorDesc::ToString(EToStringMode Mode) const
 			*GetBounds().ToString(),
 			*RuntimeGrid.ToString(),
 			GetBoolStr(bActorIsEditorOnly),
+			GetBoolStr(bActorIsRuntimeOnly),
 			GetBoolStr(bActorIsHLODRelevant)
 		);
 
@@ -379,6 +382,15 @@ PRAGMA_ENABLE_DEPRECATION_WARNINGS
 		
 	Ar << RuntimeGrid << bActorIsEditorOnly;
 
+	if (Ar.CustomVer(FFortniteMainBranchObjectVersion::GUID) >= FFortniteMainBranchObjectVersion::WorldPartitionActorDescSerializeActorIsRuntimeOnly)
+	{
+		Ar << bActorIsRuntimeOnly;
+	}
+	else
+	{
+		bActorIsRuntimeOnly = false;
+	}
+
 	if (Ar.CustomVer(FUE5MainStreamObjectVersion::GUID) < FUE5MainStreamObjectVersion::WorldPartitionActorDescRemoveBoundsRelevantSerialization)
 	{
 		bool bLevelBoundsRelevant;
diff --git a/Engine/Source/Runtime/Engine/Private/WorldPartition/WorldPartitionActorDescView.cpp b/Engine/Source/Runtime/Engine/Private/WorldPartition/WorldPartitionActorDescView.cpp
--- a/Engine/Source/Runtime/Engine/Private/WorldPartition/WorldPartitionActorDescView.cpp
+++ b/Engine/Source/Runtime/Engine/Private/WorldPartition/WorldPartitionActorDescView.cpp
@@ -58,6 +58,11 @@ bool FWorldPartitionActorDescView::GetActorIsEditorOnly() const
 	return ActorDesc->GetActorIsEditorOnly();
 }
 
+bool FWorldPartitionActorDescView::GetActorIsRuntimeOnly() const
+{
+	return ActorDesc->GetActorIsRuntimeOnly();
+}
+
 bool FWorldPartitionActorDescView::GetIsSpatiallyLoaded() const
 {
 	return bIsForcedNonSpatiallyLoaded ? false : ActorDesc->GetIsSpatiallyLoaded();
diff --git a/Engine/Source/Runtime/Engine/Private/WorldPartition/WorldPartitionActorLoaderInterface.cpp b/Engine/Source/Runtime/Engine/Private/WorldPartition/WorldPartitionActorLoaderInterface.cpp
--- a/Engine/Source/Runtime/Engine/Private/WorldPartition/WorldPartitionActorLoaderInterface.cpp
+++ b/Engine/Source/Runtime/Engine/Private/WorldPartition/WorldPartitionActorLoaderInterface.cpp
@@ -119,7 +119,7 @@ void IWorldPartitionActorLoaderInterface::ILoaderAdapter::UnregisterDelegates()
 
 bool IWorldPartitionActorLoaderInterface::ILoaderAdapter::PassActorDescFilter(const FWorldPartitionHandle& Actor) const
 {
-	return Actor->ShouldBeLoadedByEditor();
+	return !Actor->GetActorIsRuntimeOnly();
 }
 
 bool IWorldPartitionActorLoaderInterface::ILoaderAdapter::PassDataLayersFilter(const FWorldPartitionHandle& Actor) const
diff --git a/Engine/Source/Runtime/Engine/Public/WorldPartition/HLOD/HLODActor.h b/Engine/Source/Runtime/Engine/Public/WorldPartition/HLOD/HLODActor.h
--- a/Engine/Source/Runtime/Engine/Public/WorldPartition/HLOD/HLODActor.h
+++ b/Engine/Source/Runtime/Engine/Public/WorldPartition/HLOD/HLODActor.h
@@ -79,6 +79,7 @@ protected:
 	//~ Begin AActor Interface.
 	virtual void BeginPlay() override;
 	virtual void EndPlay(const EEndPlayReason::Type EndPlayReason) override;
+	virtual bool IsRuntimeOnly() const override { return true; }
 #if WITH_EDITOR
 	virtual TUniquePtr<class FWorldPartitionActorDesc> CreateClassActorDesc() const override;
 
diff --git a/Engine/Source/Runtime/Engine/Public/WorldPartition/HLOD/HLODActorDesc.h b/Engine/Source/Runtime/Engine/Public/WorldPartition/HLOD/HLODActorDesc.h
--- a/Engine/Source/Runtime/Engine/Public/WorldPartition/HLOD/HLODActorDesc.h
+++ b/Engine/Source/Runtime/Engine/Public/WorldPartition/HLOD/HLODActorDesc.h
@@ -31,7 +31,6 @@ protected:
 	virtual void Init(const AActor* InActor) override;
 	virtual bool Equals(const FWorldPartitionActorDesc* Other) const override;
 	virtual void Serialize(FArchive& Ar) override;
-	virtual bool ShouldBeLoadedByEditor() const override { return false; }
 	virtual bool IsRuntimeRelevant(const FActorContainerID& InContainerID) const override { return !bIsForcedNonSpatiallyLoaded; }
 	//~ End FWorldPartitionActorDesc Interface.
 
diff --git a/Engine/Source/Runtime/Engine/Public/WorldPartition/WorldPartitionActorDesc.h b/Engine/Source/Runtime/Engine/Public/WorldPartition/WorldPartitionActorDesc.h
--- a/Engine/Source/Runtime/Engine/Public/WorldPartition/WorldPartitionActorDesc.h
+++ b/Engine/Source/Runtime/Engine/Public/WorldPartition/WorldPartitionActorDesc.h
@@ -90,6 +90,7 @@ PRAGMA_ENABLE_DEPRECATION_WARNINGS
 	inline bool GetIsSpatiallyLoaded() const { return bIsForcedNonSpatiallyLoaded ? false : bIsSpatiallyLoaded; }
 	inline bool GetIsSpatiallyLoadedRaw() const { return bIsSpatiallyLoaded; }
 	inline bool GetActorIsEditorOnly() const { return bActorIsEditorOnly; }
+	inline bool GetActorIsRuntimeOnly() const { return bActorIsRuntimeOnly; }
 
 	UE_DEPRECATED(5.1, "GetLevelBoundsRelevant is deprecated.")
 	inline bool GetLevelBoundsRelevant() const { return false; }
@@ -201,10 +202,8 @@ public:
 	AActor* Load() const;
 	virtual void Unload();
 
-	UE_DEPRECATED(5.1, "ShouldBeLoadedByEditorCells is deprecated, ShouldBeLoadedByEditor should be used instead.")
-	bool ShouldBeLoadedByEditorCells() const { return ShouldBeLoadedByEditor(); }
-
-	virtual bool ShouldBeLoadedByEditor() const { return true; }
+	UE_DEPRECATED(5.1, "ShouldBeLoadedByEditorCells is deprecated, GetActorIsRuntimeOnly should be used instead.")
+	bool ShouldBeLoadedByEditorCells() const { return !GetActorIsRuntimeOnly(); }
 
 	virtual void Init(const AActor* InActor);
 	virtual void Init(const FWorldPartitionActorDescInitData& DescData);
@@ -249,6 +248,7 @@ protected:
 	FName							RuntimeGrid;
 	bool							bIsSpatiallyLoaded;
 	bool							bActorIsEditorOnly;
+	bool							bActorIsRuntimeOnly;
 	bool							bActorIsHLODRelevant;
 	bool							bIsUsingDataLayerAsset; // Used to know if DataLayers array represents DataLayers Asset paths or the FNames of the deprecated version of Data Layers
 	FName							HLODLayer;
diff --git a/Engine/Source/Runtime/Engine/Public/WorldPartition/WorldPartitionActorDescView.h b/Engine/Source/Runtime/Engine/Public/WorldPartition/WorldPartitionActorDescView.h
--- a/Engine/Source/Runtime/Engine/Public/WorldPartition/WorldPartitionActorDescView.h
+++ b/Engine/Source/Runtime/Engine/Public/WorldPartition/WorldPartitionActorDescView.h
@@ -34,6 +34,7 @@ public:
 	FName GetRuntimeGrid() const;
 	bool GetIsSpatiallyLoaded() const;
 	bool GetActorIsEditorOnly() const;
+	bool GetActorIsRuntimeOnly() const;
 	bool GetActorIsHLODRelevant() const;
 	FName GetHLODLayer() const;
 	const TArray<FName>& GetDataLayers() const;
diff --git a/Engine/Source/Runtime/Landscape/Classes/LandscapeSplineMeshesActor.h b/Engine/Source/Runtime/Landscape/Classes/LandscapeSplineMeshesActor.h
--- a/Engine/Source/Runtime/Landscape/Classes/LandscapeSplineMeshesActor.h
+++ b/Engine/Source/Runtime/Landscape/Classes/LandscapeSplineMeshesActor.h
@@ -29,6 +29,7 @@ public:
 	//~ Begin AActor Interface
 	virtual void CheckForErrors() override;
 	virtual bool GetReferencedContentObjects(TArray<UObject*>& Objects) const override;
+	virtual bool IsRuntimeOnly() const override { return true; }
 	//~ End AActor Interface
 
 	//~ Begin APartitionActor Interface
