Fix non-unity build.

#rnx

diff --git a/Engine/Source/Runtime/Engine/Private/WorldPartition/WorldPartitionEditorHash.cpp b/Engine/Source/Runtime/Engine/Private/WorldPartition/WorldPartitionEditorHash.cpp
--- a/Engine/Source/Runtime/Engine/Private/WorldPartition/WorldPartitionEditorHash.cpp
+++ b/Engine/Source/Runtime/Engine/Private/WorldPartition/WorldPartitionEditorHash.cpp
@@ -3,6 +3,10 @@
 #include "WorldPartition/WorldPartitionEditorHash.h"
 #include "WorldPartition/IWorldPartitionEditorModule.h"
 
+#if WITH_EDITOR
+#include "Modules/ModuleManager.h"
+#endif
+
 #include UE_INLINE_GENERATED_CPP_BY_NAME(WorldPartitionEditorHash)
 
 UWorldPartitionEditorHash::UWorldPartitionEditorHash(const FObjectInitializer& ObjectInitializer)
