LoadMap: Handle PIE Prefixed URL (udn)

diff --git a/Engine/Source/Runtime/Engine/Private/UnrealEngine.cpp b/Engine/Source/Runtime/Engine/Private/UnrealEngine.cpp
--- a/Engine/Source/Runtime/Engine/Private/UnrealEngine.cpp
+++ b/Engine/Source/Runtime/Engine/Private/UnrealEngine.cpp
@@ -14294,6 +14294,9 @@ bool UEngine::LoadMap( FWorldContext& WorldContext, FURL URL, class UPendingNetG
 
 	LLM_SCOPE(ELLMTag::LoadMapMisc);
 
+	const FString OriginalURLMap = URL.Map;
+	URL.Map = UWorld::RemovePIEPrefix(URL.Map);
+
 	FDisableHitchDetectorScope SuspendHitchDetector;
 
 	NETWORK_PROFILER(GNetworkProfiler.TrackSessionChange(true,URL));
@@ -14542,9 +14545,9 @@ bool UEngine::LoadMap( FWorldContext& WorldContext, FURL URL, class UPendingNetG
 			// a server.
 			if (WorldContextFromList.World() != WorldContext.World())
 			{
-				if (!WorldContextFromList.PIEPrefix.IsEmpty() && URL.Map.Contains(WorldContextFromList.PIEPrefix))
+				if (!WorldContextFromList.PIEPrefix.IsEmpty() && OriginalURLMap.Contains(WorldContextFromList.PIEPrefix))
 				{
-					FString SourceWorldPackage = UWorld::RemovePIEPrefix(URL.Map);
+					FString SourceWorldPackage = URL.Map;
 
 					// We are loading a new world for this context, so clear out PIE fixups that might be lingering.
 					// (note we dont want to do this in DuplicateWorldForPIE, since that is also called on streaming worlds.
