WorldPartitionFoliageBuilder: Fix failed check regression

- Add ActorGuids param to the FForEachActorWithLoadingParams to iterate only on this subset of actors (for iterations that can modify the actordesc list)

#rnx

diff --git a/Engine/Source/Editor/UnrealEd/Private/WorldPartition/WorldPartitionFoliageBuilder.cpp b/Engine/Source/Editor/UnrealEd/Private/WorldPartition/WorldPartitionFoliageBuilder.cpp
--- a/Engine/Source/Editor/UnrealEd/Private/WorldPartition/WorldPartitionFoliageBuilder.cpp
+++ b/Engine/Source/Editor/UnrealEd/Private/WorldPartition/WorldPartitionFoliageBuilder.cpp
@@ -64,7 +64,14 @@ bool UWorldPartitionFoliageBuilder::RunInternal(UWorld* World, const FCellInfo&
 		
 	FPackageSourceControlHelper SCCHelper;
 	TArray<FGuid> ExistingActors;
-	
+			
+	// Snapshot existing ActorDescs
+	FWorldPartitionHelpers::ForEachActorDesc(WorldPartition, AInstancedFoliageActor::StaticClass(), [&ExistingActors](const FWorldPartitionActorDesc* ActorDesc)
+	{
+		ExistingActors.Add(ActorDesc->GetGuid());
+		return true;
+	});
+
 	TMap<UActorPartitionSubsystem::FCellCoord, FGuid> NewActors;
 	UActorPartitionSubsystem* ActorPartitionSubsystem = World->GetSubsystem<UActorPartitionSubsystem>();
 	int32 NumInstances = 0;
@@ -80,11 +87,10 @@ bool UWorldPartitionFoliageBuilder::RunInternal(UWorld* World, const FCellInfo&
 	// - Repeat for next Existing IFA
 	FWorldPartitionHelpers::FForEachActorWithLoadingParams ForEachActorWithLoadingParams;
 	ForEachActorWithLoadingParams.ActorClasses = { AInstancedFoliageActor::StaticClass() };
+	ForEachActorWithLoadingParams.ActorGuids.Append(ExistingActors);
 
-	FWorldPartitionHelpers::ForEachActorWithLoading(WorldPartition, [this, World, WorldPartition, ActorPartitionSubsystem, &SCCHelper, &NewActors, &NumInstances, &NumInstancesProcessed, &ExistingActors](const FWorldPartitionActorDesc* ActorDesc)
+	FWorldPartitionHelpers::ForEachActorWithLoading(WorldPartition, [this, World, WorldPartition, ActorPartitionSubsystem, &SCCHelper, &NewActors, &NumInstances, &NumInstancesProcessed](const FWorldPartitionActorDesc* ActorDesc)
 	{
-		ExistingActors.Add(ActorDesc->GetGuid());
-
 		TMap<UFoliageType*, TArray<FFoliageInstance>> FoliageToAdd;
 		FBox InstanceBounds(ForceInit);
 				
diff --git a/Engine/Source/Runtime/Engine/Private/WorldPartition/WorldPartitionHelpers.cpp b/Engine/Source/Runtime/Engine/Private/WorldPartition/WorldPartitionHelpers.cpp
--- a/Engine/Source/Runtime/Engine/Private/WorldPartition/WorldPartitionHelpers.cpp
+++ b/Engine/Source/Runtime/Engine/Private/WorldPartition/WorldPartitionHelpers.cpp
@@ -133,9 +133,8 @@ void FWorldPartitionHelpers::ForEachActorWithLoading(UWorldPartition* WorldParti
 		DoCollectGarbage();
 	};
 
-	for (FActorDescContainerCollection::TConstIterator<> ActorDescIterator(WorldPartition); ActorDescIterator; ++ActorDescIterator)
+	auto ForEachActorWithLoadingImpl = [&](const FWorldPartitionActorDesc* ActorDesc)
 	{
-		const FWorldPartitionActorDesc* ActorDesc = *ActorDescIterator;
 		if (Algo::AnyOf(Params.ActorClasses, [ActorDesc](UClass* ActorClass) { return IsActorDescClassCompatibleWith(ActorDesc, ActorClass); }))
 		{
 			if (!Params.FilterActorDesc || Params.FilterActorDesc(ActorDesc))
@@ -145,7 +144,7 @@ void FWorldPartitionHelpers::ForEachActorWithLoading(UWorldPartition* WorldParti
 				FWorldPartitionReference ActorReference(WorldPartition, ActorDesc->GetGuid());
 				if (!Func(ActorReference.Get()))
 				{
-					break;
+					return false;
 				}
 
 				if (!Params.bKeepReferences && (Params.bGCPerActor || FWorldPartitionHelpers::HasExceededMaxMemory()))
@@ -154,6 +153,36 @@ void FWorldPartitionHelpers::ForEachActorWithLoading(UWorldPartition* WorldParti
 				}
 			}
 		}
+
+		return true;
+	};
+
+	if (Params.ActorGuids.IsEmpty())
+	{
+		for (FActorDescContainerCollection::TConstIterator<> ActorDescIterator(WorldPartition); ActorDescIterator; ++ActorDescIterator)
+		{
+			if (const FWorldPartitionActorDesc* ActorDesc = *ActorDescIterator)
+			{
+				if (!ForEachActorWithLoadingImpl(ActorDesc))
+				{
+					break;
+				}
+			}
+			
+		}
+	}
+	else
+	{
+		for (const FGuid& ActorGuid : Params.ActorGuids)
+		{
+			if (const FWorldPartitionActorDesc* ActorDesc = WorldPartition->GetActorDesc(ActorGuid))
+			{
+				if (!ForEachActorWithLoadingImpl(ActorDesc))
+				{
+					break;
+				}
+			}
+		}
 	}
 
 	if (!Params.bKeepReferences)
diff --git a/Engine/Source/Runtime/Engine/Public/WorldPartition/WorldPartitionHelpers.h b/Engine/Source/Runtime/Engine/Public/WorldPartition/WorldPartitionHelpers.h
--- a/Engine/Source/Runtime/Engine/Public/WorldPartition/WorldPartitionHelpers.h
+++ b/Engine/Source/Runtime/Engine/Public/WorldPartition/WorldPartitionHelpers.h
@@ -45,6 +45,9 @@ public:
 		/* The classes used to filter actors loading. */
 		TArray<TSubclassOf<AActor>> ActorClasses;
 
+		/* If not empty, iteration will be done only on those actors. */
+		TArray<FGuid> ActorGuids;
+
 		/* Custom filter function used to filter actors loading. */
 		TUniqueFunction<bool(const FWorldPartitionActorDesc*)> FilterActorDesc;
 
