Make decal components blueprintable.

diff --git a/Engine/Source/Runtime/Engine/Classes/Components/DecalComponent.h b/Engine/Source/Runtime/Engine/Classes/Components/DecalComponent.h
--- a/Engine/Source/Runtime/Engine/Classes/Components/DecalComponent.h
+++ b/Engine/Source/Runtime/Engine/Classes/Components/DecalComponent.h
@@ -18,7 +18,7 @@ class UMaterialInterface;
  * @see https://docs.unrealengine.com/latest/INT/Engine/Actors/DecalActor
  * @see UDecalActor
  */
-UCLASS(hidecategories=(Collision, Object, Physics, SceneComponent, Activation, "Components|Activation", Mobility), ClassGroup=Rendering, meta=(BlueprintSpawnableComponent))
+UCLASS(Blueprintable, editinlinenew, hidecategories=(Collision, Object, Physics, SceneComponent, Activation, "Components|Activation", Mobility), ClassGroup=Rendering, meta=(BlueprintSpawnableComponent))
 class ENGINE_API UDecalComponent : public USceneComponent
 {
 	GENERATED_UCLASS_BODY()
