[Backout] - CL22206377
Original CL Desc
-----------------------------------------------------------------
Properly support morph vertex buffers when render data is streamed

diff --git a/Engine/Source/Runtime/Engine/Private/Rendering/MorphTargetVertexInfoBuffers.cpp b/Engine/Source/Runtime/Engine/Private/Rendering/MorphTargetVertexInfoBuffers.cpp
--- a/Engine/Source/Runtime/Engine/Private/Rendering/MorphTargetVertexInfoBuffers.cpp
+++ b/Engine/Source/Runtime/Engine/Private/Rendering/MorphTargetVertexInfoBuffers.cpp
@@ -465,40 +465,3 @@ bool FMorphTargetVertexInfoBuffers::IsPlatformShaderSupported(EShaderPlatform Sh
 {
 	return IsFeatureLevelSupported(ShaderPlatform, ERHIFeatureLevel::SM5);
 }
-
-template <bool bRenderThread>
-FBufferRHIRef FMorphTargetVertexInfoBuffers::CreateMorphRHIBuffer_Internal()
-{
-	uint32 SizeInBytes = MorphData.Num() * sizeof(uint32);
-
-	if(SizeInBytes > 0)
-	{
-		const EBufferUsageFlags BufferFlags = BUF_Static | BUF_ByteAddressBuffer | BUF_ShaderResource;
-
-		// Create the index buffer.
-		FRHIResourceCreateInfo CreateInfo(TEXT("MorphData"), &MorphData);
-		if (bRenderThread)
-		{
-			FBufferRHIRef BufferRHIRef;
-			BufferRHIRef = RHICreateStructuredBuffer(sizeof(uint32), SizeInBytes, BufferFlags, ERHIAccess::SRVMask, CreateInfo);
-
-			return BufferRHIRef;
-		}
-		else
-		{
-			FRHIAsyncCommandList CommandList;
-			return CommandList->CreateBuffer(SizeInBytes, BufferFlags, sizeof(uint32), ERHIAccess::SRVMask, CreateInfo);
-		}
-	}
-	return nullptr;
-}
-
-FBufferRHIRef FMorphTargetVertexInfoBuffers::CreateMorphRHIBuffer_RenderThread()
-{
-	return CreateMorphRHIBuffer_Internal<true>();
-}
-
-FBufferRHIRef FMorphTargetVertexInfoBuffers::CreateMorphRHIBuffer_Async()
-{
-	return CreateMorphRHIBuffer_Internal<false>();
-}
diff --git a/Engine/Source/Runtime/Engine/Private/SkeletalMeshLODRenderData.cpp b/Engine/Source/Runtime/Engine/Private/SkeletalMeshLODRenderData.cpp
--- a/Engine/Source/Runtime/Engine/Private/SkeletalMeshLODRenderData.cpp
+++ b/Engine/Source/Runtime/Engine/Private/SkeletalMeshLODRenderData.cpp
@@ -30,6 +30,13 @@ static FAutoConsoleVariableRef CVarStripSkeletalMeshLodsBelowMinLod(
 	TEXT("If set will strip skeletal mesh LODs under the minimum renderable LOD for the target platform during cooking.")
 );
 
+int32 GAllowSkinnedMorphDataStreaming = 0;
+static FAutoConsoleVariableRef CVarAllowSkinnedMorphDataStreaming(
+	TEXT("r.Skinned.AllowSkinnedMorphDataStreaming"),
+	GAllowSkinnedMorphDataStreaming,
+	TEXT("Call BeginInitResource when the morph resources are finish streaming (DoFinishUpdate)")
+);
+
 namespace
 {
 	struct FReverseOrderBitArraysBySetBits
@@ -159,6 +166,14 @@ FArchive& operator<<(FArchive& Ar, FSkelMeshRenderSection& S)
 	return Ar;
 }
 
+void FSkeletalMeshLODRenderData::InitMorphResources()
+{
+	if (GAllowSkinnedMorphDataStreaming && !MorphTargetVertexInfoBuffers.IsRHIIntialized() && MorphTargetVertexInfoBuffers.IsMorphCPUDataValid() && MorphTargetVertexInfoBuffers.NumTotalBatches > 0)
+	{
+		BeginInitResource(&MorphTargetVertexInfoBuffers);
+	}
+}
+
 void FSkeletalMeshLODRenderData::InitResources(bool bNeedsVertexColors, int32 LODIndex, TArray<UMorphTarget*>& InMorphTargets, USkinnedAsset* Owner)
 {
 	if (bStreamedDataInlined)
@@ -216,7 +231,7 @@ void FSkeletalMeshLODRenderData::InitResources(bool bNeedsVertexColors, int32 LO
 		MorphTargetVertexInfoBuffers.InitMorphResources(GMaxRHIShaderPlatform, RenderSections, Owner->GetMorphTargets(), StaticVertexBuffers.StaticMeshVertexBuffer.GetNumVertices(), LODIndex, SkeletalMeshLODInfo->MorphTargetPositionErrorTolerance);
 	}
 	
-	if (!MorphTargetVertexInfoBuffers.IsRHIIntialized() && MorphTargetVertexInfoBuffers.IsMorphCPUDataValid() && (MorphTargetVertexInfoBuffers.NumTotalBatches > 0 || !bStreamedDataInlined))
+	if (!MorphTargetVertexInfoBuffers.IsRHIIntialized() && MorphTargetVertexInfoBuffers.IsMorphCPUDataValid() && MorphTargetVertexInfoBuffers.NumTotalBatches > 0)
 	{
 		BeginInitResource(&MorphTargetVertexInfoBuffers);
 	}
diff --git a/Engine/Source/Runtime/Engine/Private/Streaming/SkeletalMeshUpdate.cpp b/Engine/Source/Runtime/Engine/Private/Streaming/SkeletalMeshUpdate.cpp
--- a/Engine/Source/Runtime/Engine/Private/Streaming/SkeletalMeshUpdate.cpp
+++ b/Engine/Source/Runtime/Engine/Private/Streaming/SkeletalMeshUpdate.cpp
@@ -60,7 +60,6 @@ void FSkeletalMeshStreamIn::FIntermediateBuffers::CreateFromCPUData_RenderThread
 	SkinWeightVertexBuffer = LODResource.SkinWeightVertexBuffer.CreateRHIBuffer_RenderThread();
 	ClothVertexBuffer = LODResource.ClothVertexBuffer.CreateRHIBuffer_RenderThread();
 	IndexBuffer = LODResource.MultiSizeIndexContainer.CreateRHIBuffer_RenderThread();
-	MorphBuffer = LODResource.MorphTargetVertexInfoBuffers.CreateMorphRHIBuffer_RenderThread();
 }
 
 void FSkeletalMeshStreamIn::FIntermediateBuffers::CreateFromCPUData_Async(FSkeletalMeshLODRenderData& LODResource)
@@ -74,7 +73,6 @@ void FSkeletalMeshStreamIn::FIntermediateBuffers::CreateFromCPUData_Async(FSkele
 	SkinWeightVertexBuffer = LODResource.SkinWeightVertexBuffer.CreateRHIBuffer_Async();
 	ClothVertexBuffer = LODResource.ClothVertexBuffer.CreateRHIBuffer_Async();
 	IndexBuffer = LODResource.MultiSizeIndexContainer.CreateRHIBuffer_Async();
-	MorphBuffer = LODResource.MorphTargetVertexInfoBuffers.CreateMorphRHIBuffer_Async();
 }
 
 void FSkeletalMeshStreamIn::FIntermediateBuffers::SafeRelease()
@@ -88,7 +86,6 @@ void FSkeletalMeshStreamIn::FIntermediateBuffers::SafeRelease()
 	ClothVertexBuffer.SafeRelease();
 	IndexBuffer.SafeRelease();
 	AltSkinWeightVertexBuffers.Empty();
-	MorphBuffer.SafeRelease();
 }
 
 template <uint32 MaxNumUpdates>
@@ -102,7 +99,6 @@ void FSkeletalMeshStreamIn::FIntermediateBuffers::TransferBuffers(FSkeletalMeshL
 	LODResource.ClothVertexBuffer.InitRHIForStreaming(ClothVertexBuffer, Batcher);
 	LODResource.MultiSizeIndexContainer.InitRHIForStreaming(IndexBuffer, Batcher);
 	LODResource.SkinWeightProfilesData.InitRHIForStreaming(AltSkinWeightVertexBuffers, Batcher);
-	LODResource.MorphTargetVertexInfoBuffers.InitRHIForStreaming(MorphBuffer, Batcher);
 	SafeRelease();
 }
 
@@ -116,8 +112,7 @@ void FSkeletalMeshStreamIn::FIntermediateBuffers::CheckIsNull() const
 		&& !SkinWeightVertexBuffer.LookupVertexBufferRHI
 		&& !ClothVertexBuffer
 		&& !IndexBuffer
-		&& !AltSkinWeightVertexBuffers.Num()
-		&& !MorphBuffer);
+		&& !AltSkinWeightVertexBuffers.Num());
 }
 
 FSkeletalMeshStreamIn::FSkeletalMeshStreamIn(const USkeletalMesh* InMesh)
@@ -203,6 +198,7 @@ void FSkeletalMeshStreamIn::DoFinishUpdate(const FContext& Context)
 			{
 				FSkeletalMeshLODRenderData& LODResource = *Context.LODResourcesView[LODIndex];
 				LODResource.IncrementMemoryStats(Mesh->GetHasVertexColors());
+				LODResource.InitMorphResources();
 				IntermediateBuffersArray[LODIndex].TransferBuffers(LODResource, Batcher);
 			}
 		}
@@ -361,7 +357,6 @@ void FSkeletalMeshStreamOut::ReleaseBuffers(const FContext& Context)
 			LODResource.ClothVertexBuffer.ReleaseRHIForStreaming(Batcher);
 			LODResource.MultiSizeIndexContainer.ReleaseRHIForStreaming(Batcher);
 			LODResource.SkinWeightProfilesData.ReleaseRHIForStreaming(Batcher);
-			LODResource.MorphTargetVertexInfoBuffers.ReleaseRHIForStreaming(Batcher);
 
 			if (!FPlatformProperties::HasEditorOnlyData())
 			{
diff --git a/Engine/Source/Runtime/Engine/Private/Streaming/SkeletalMeshUpdate.h b/Engine/Source/Runtime/Engine/Private/Streaming/SkeletalMeshUpdate.h
--- a/Engine/Source/Runtime/Engine/Private/Streaming/SkeletalMeshUpdate.h
+++ b/Engine/Source/Runtime/Engine/Private/Streaming/SkeletalMeshUpdate.h
@@ -80,7 +80,6 @@ protected:
 		FBufferRHIRef ClothVertexBuffer;
 		FBufferRHIRef IndexBuffer;
 		TArray<TPair<FName, FSkinWeightRHIInfo>> AltSkinWeightVertexBuffers;
-		FBufferRHIRef MorphBuffer;
 
 		void CreateFromCPUData_RenderThread(FSkeletalMeshLODRenderData& LODResource);
 		void CreateFromCPUData_Async(FSkeletalMeshLODRenderData& LODResource);
diff --git a/Engine/Source/Runtime/Engine/Public/Rendering/MorphTargetVertexInfoBuffers.h b/Engine/Source/Runtime/Engine/Public/Rendering/MorphTargetVertexInfoBuffers.h
--- a/Engine/Source/Runtime/Engine/Public/Rendering/MorphTargetVertexInfoBuffers.h
+++ b/Engine/Source/Runtime/Engine/Public/Rendering/MorphTargetVertexInfoBuffers.h
@@ -79,32 +79,6 @@ public:
 	FBufferRHIRef MorphDataBuffer;
 	FShaderResourceViewRHIRef MorphDataSRV;
 
-	/** Create an RHI vertex buffer with CPU data. CPU data may be discarded after creation (see TResourceArray::Discard) */
-	FBufferRHIRef CreateMorphRHIBuffer_RenderThread();
-	FBufferRHIRef CreateMorphRHIBuffer_Async();
-
-	/** Similar to Init/ReleaseRHI but only update existing SRV so references to the SRV stays valid */
-	template <uint32 MaxNumUpdates>
-	void InitRHIForStreaming(
-		FRHIBuffer* IntermediatMorphTargetBuffer,
-		TRHIResourceUpdateBatcher<MaxNumUpdates>& Batcher)
-	{
-		check(MorphDataBuffer);
-		if (MorphDataBuffer && IntermediatMorphTargetBuffer)
-		{
-			Batcher.QueueUpdateRequest(MorphDataBuffer, IntermediatMorphTargetBuffer);
-			Batcher.QueueUpdateRequest(MorphDataSRV, MorphDataBuffer);
-		}
-	}
-
-	template<uint32 MaxNumUpdates>
-	void ReleaseRHIForStreaming(TRHIResourceUpdateBatcher<MaxNumUpdates>& Batcher)
-	{
-		check(MorphDataBuffer);
-		Batcher.QueueUpdateRequest(MorphDataBuffer, nullptr);
-		Batcher.QueueUpdateRequest(MorphDataSRV, nullptr, 0, 0);
-	}
-
 protected:
 	void ResetCPUData()
 	{
@@ -124,10 +98,7 @@ protected:
 	void Serialize(FArchive& Ar);
 
 	// Transient data. Gets deleted as soon as the GPU resource has been initialized (unless excplicitly disabled by bEmptyMorphCPUDataOnInitRHI).
-	TResourceArray<uint32> MorphData;
-
-	
-
+	TArray<uint32> MorphData;
 
 	//x,y,y separate for position and shared w for tangent
 	TArray<FVector4f> MaximumValuePerMorph;
@@ -146,9 +117,6 @@ protected:
 
 	friend class FSkeletalMeshLODRenderData;
 	ENGINE_API friend FArchive& operator<<(FArchive& Ar, FMorphTargetVertexInfoBuffers& MorphTargetVertexInfoBuffers);
-private:
-	template <bool bRenderThread>
-	FBufferRHIRef CreateMorphRHIBuffer_Internal();
 };
 
 ENGINE_API FArchive& operator<<(FArchive& Ar, FMorphTargetVertexInfoBuffers& MorphTargetVertexInfoBuffers);
diff --git a/Engine/Source/Runtime/Engine/Public/Rendering/SkeletalMeshLODRenderData.h b/Engine/Source/Runtime/Engine/Public/Rendering/SkeletalMeshLODRenderData.h
--- a/Engine/Source/Runtime/Engine/Public/Rendering/SkeletalMeshLODRenderData.h
+++ b/Engine/Source/Runtime/Engine/Public/Rendering/SkeletalMeshLODRenderData.h
@@ -177,6 +177,8 @@ public:
 	*/
 	void InitResources(bool bNeedsVertexColors, int32 LODIndex, TArray<class UMorphTarget*>& InMorphTargets, USkinnedAsset* Owner);
 
+	void InitMorphResources();
+
 	/**
 	* Releases the LOD's render resources.
 	*/
