WorldPartitionBlueprintLibrary: added actors package and path to the exposed actor descriptor properties.

#rnx

diff --git a/Engine/Source/Runtime/Engine/Private/WorldPartitionBlueprintLibrary.cpp b/Engine/Source/Runtime/Engine/Private/WorldPartitionBlueprintLibrary.cpp
--- a/Engine/Source/Runtime/Engine/Private/WorldPartitionBlueprintLibrary.cpp
+++ b/Engine/Source/Runtime/Engine/Private/WorldPartitionBlueprintLibrary.cpp
@@ -37,7 +37,9 @@ FActorDesc::FActorDesc(const FWorldPartitionActorDesc& InActorDesc, const FTrans
 	Bounds = InActorDesc.GetBounds().TransformBy(InTransform);
 	RuntimeGrid = InActorDesc.GetRuntimeGrid();
 	bIsSpatiallyLoaded = InActorDesc.GetIsSpatiallyLoaded();
-	bActorIsEditorOnly = InActorDesc.GetActorIsEditorOnly();	
+	bActorIsEditorOnly = InActorDesc.GetActorIsEditorOnly();
+	ActorPackage = InActorDesc.GetActorPackage();
+	ActorPath = InActorDesc.GetActorPath();
 }
 
 TMap<UWorldPartition*, TUniquePtr<FLoaderAdapterActorList>> UWorldPartitionBlueprintLibrary::LoaderAdapterActorListMap;
@@ -60,9 +62,7 @@ UWorldPartition* UWorldPartitionBlueprintLibrary::GetWorldPartition()
 
 void UWorldPartitionBlueprintLibrary::OnWorldPartitionUninitialized(UWorldPartition* InWorldPartition)
 {
-	verify(LoaderAdapterActorListMap.Remove(InWorldPartition));
-
-	if (LoaderAdapterActorListMap.IsEmpty())
+	if (LoaderAdapterActorListMap.Remove(InWorldPartition) && LoaderAdapterActorListMap.IsEmpty())
 	{
 		InWorldPartition->GetWorld()->OnWorldPartitionUninitialized().Remove(OnWorldPartitionUninitializedHandle);
 	}
diff --git a/Engine/Source/Runtime/Engine/Public/WorldPartition/WorldPartitionBlueprintLibrary.h b/Engine/Source/Runtime/Engine/Public/WorldPartition/WorldPartitionBlueprintLibrary.h
--- a/Engine/Source/Runtime/Engine/Public/WorldPartition/WorldPartitionBlueprintLibrary.h
+++ b/Engine/Source/Runtime/Engine/Public/WorldPartition/WorldPartitionBlueprintLibrary.h
@@ -56,6 +56,14 @@ struct FActorDesc
 	/** Actor's editor-only property. */
 	UPROPERTY(BlueprintReadOnly, VisibleAnywhere, Category=Actor)
 	bool bActorIsEditorOnly;
+
+	/** Actor's package name. */
+	UPROPERTY(BlueprintReadOnly, VisibleAnywhere, Category=Actor)
+	FName ActorPackage;
+
+	/** Actor's path name. */
+	UPROPERTY(BlueprintReadOnly, VisibleAnywhere, Category=Actor)
+	FName ActorPath;
 };
 
 UCLASS(MinimalAPI)
