Explicitly disable FXSystem GPU sorting for ES3.1 feature level

diff --git a/Engine/Source/Runtime/Engine/Private/Particles/FXSystem.cpp b/Engine/Source/Runtime/Engine/Private/Particles/FXSystem.cpp
--- a/Engine/Source/Runtime/Engine/Private/Particles/FXSystem.cpp
+++ b/Engine/Source/Runtime/Engine/Private/Particles/FXSystem.cpp
@@ -222,11 +222,17 @@ FFXSystem::FFXSystem(ERHIFeatureLevel::Type InFeatureLevel, EShaderPlatform InSh
 	: ParticleSimulationResources(NULL)
 	, FeatureLevel(InFeatureLevel)
 	, ShaderPlatform(InShaderPlatform)
-	, GPUSortManager(InGPUSortManager)
+	, GPUSortManager(nullptr)
 #if WITH_EDITOR
 	, bSuspended(false)
 #endif // #if WITH_EDITOR
 {
+	// FXSystem GPU sorting is disabled for mobile, see FParticleSortKeyGenCS
+	if (FeatureLevel >= ERHIFeatureLevel::SM5)
+	{
+		GPUSortManager = InGPUSortManager;
+	}
+	
 	InitGPUSimulation();
 
 	// Register the callback in the GPUSortManager. 
diff --git a/Engine/Source/Runtime/Engine/Private/Particles/ParticleSortingGPU.cpp b/Engine/Source/Runtime/Engine/Private/Particles/ParticleSortingGPU.cpp
--- a/Engine/Source/Runtime/Engine/Private/Particles/ParticleSortingGPU.cpp
+++ b/Engine/Source/Runtime/Engine/Private/Particles/ParticleSortingGPU.cpp
@@ -55,6 +55,12 @@ public:
 		OutEnvironment.SetDefine( TEXT("TEXTURE_SIZE_Y"), GParticleSimulationTextureSizeY );
 	}
 
+	static bool ShouldCompilePermutation(const FGlobalShaderPermutationParameters& Parameters)
+	{
+		// this shader writes to a 2 component RWBuffer, which is not supported in GLES
+		return !IsMobilePlatform(Parameters.Platform);
+	}
+
 	/** Default constructor. */
 	FParticleSortKeyGenCS()
 	{
@@ -167,7 +173,7 @@ int32 GenerateParticleSortKeys(
 	int32 BatchId
 	)
 {
-	check(FeatureLevel >= ERHIFeatureLevel::SM5 || FeatureLevel == ERHIFeatureLevel::ES3_1);
+	check(FeatureLevel >= ERHIFeatureLevel::SM5);
 
 	FParticleKeyGenParameters KeyGenParameters;
 	FParticleKeyGenUniformBufferRef KeyGenUniformBuffer;
