Disable deprecation around use of GetActorPath in new world partition code

diff --git a/Engine/Source/Runtime/Engine/Private/WorldPartitionBlueprintLibrary.cpp b/Engine/Source/Runtime/Engine/Private/WorldPartitionBlueprintLibrary.cpp
--- a/Engine/Source/Runtime/Engine/Private/WorldPartitionBlueprintLibrary.cpp
+++ b/Engine/Source/Runtime/Engine/Private/WorldPartitionBlueprintLibrary.cpp
@@ -39,7 +39,10 @@ FActorDesc::FActorDesc(const FWorldPartitionActorDesc& InActorDesc, const FTrans
 	bIsSpatiallyLoaded = InActorDesc.GetIsSpatiallyLoaded();
 	bActorIsEditorOnly = InActorDesc.GetActorIsEditorOnly();
 	ActorPackage = InActorDesc.GetActorPackage();
+	
+	PRAGMA_DISABLE_DEPRECATION_WARNINGS
 	ActorPath = InActorDesc.GetActorPath();
+	PRAGMA_ENABLE_DEPRECATION_WARNINGS
 }
 
 TMap<UWorldPartition*, TUniquePtr<FLoaderAdapterActorList>> UWorldPartitionBlueprintLibrary::LoaderAdapterActorListMap;
