Skeletal Mesh: Make UpdateClothSimulationData() an abstract method of USkinnedMeshComponent.

diff --git a/Engine/Source/Runtime/Engine/Classes/Components/SkeletalMeshComponent.h b/Engine/Source/Runtime/Engine/Classes/Components/SkeletalMeshComponent.h
--- a/Engine/Source/Runtime/Engine/Classes/Components/SkeletalMeshComponent.h
+++ b/Engine/Source/Runtime/Engine/Classes/Components/SkeletalMeshComponent.h
@@ -1594,8 +1594,12 @@ public:
 	void TickClothing(float DeltaTime, FTickFunction& ThisTickFunction);
 
 	/** Store cloth simulation data into OutClothSimData */
+	UE_DEPRECATED(5.2, "Use GetUpdateClothSimulationData_AnyThread instead.")
 	void GetUpdateClothSimulationData(TMap<int32, FClothSimulData>& OutClothSimData, USkeletalMeshComponent* OverrideLocalRootComponent = nullptr);
 
+	/** Store cloth simulation data into OutClothSimulData. Override USkinnedMeshComponent. */
+	virtual void GetUpdateClothSimulationData_AnyThread(TMap<int32, FClothSimulData>& OutClothSimulData, FMatrix& OutLocalToWorld, float& OutClothBlendWeight) override;
+
 	/** Remove clothing actors from their simulation */
 	void RemoveAllClothingActors();
 
diff --git a/Engine/Source/Runtime/Engine/Classes/Components/SkinnedMeshComponent.h b/Engine/Source/Runtime/Engine/Classes/Components/SkinnedMeshComponent.h
--- a/Engine/Source/Runtime/Engine/Classes/Components/SkinnedMeshComponent.h
+++ b/Engine/Source/Runtime/Engine/Classes/Components/SkinnedMeshComponent.h
@@ -17,6 +17,7 @@
 #include "LODSyncInterface.h"
 #include "BoneContainer.h"
 #include "Rendering/MorphTargetVertexInfoBuffers.h"
+#include "ClothingSystemRuntimeTypes.h"
 #include "SkinnedMeshComponent.generated.h"
 
 enum class ESkinCacheUsage : uint8;
@@ -814,6 +815,9 @@ public:
 	/** Gets the skeletal mesh resource used for rendering the component. */
 	FSkeletalMeshRenderData* GetSkeletalMeshRenderData() const;
 
+	/** Override this function to pass on cloth simulation data (or any deformable using the cloth pipeline) to the skeletal renderer. */
+	virtual void GetUpdateClothSimulationData_AnyThread(TMap<int32, FClothSimulData>& OutClothSimulData, FMatrix& OutLocalToWorld, float& OutClothBlendWeight);
+
 	/** 
 	 * Override the Physics Asset of the mesh. It uses SkeletalMesh.PhysicsAsset, but if you'd like to override use this function
 	 * 
diff --git a/Engine/Source/Runtime/Engine/Private/Components/SkinnedMeshComponent.cpp b/Engine/Source/Runtime/Engine/Private/Components/SkinnedMeshComponent.cpp
--- a/Engine/Source/Runtime/Engine/Private/Components/SkinnedMeshComponent.cpp
+++ b/Engine/Source/Runtime/Engine/Private/Components/SkinnedMeshComponent.cpp
@@ -2079,6 +2079,12 @@ FSkeletalMeshRenderData* USkinnedMeshComponent::GetSkeletalMeshRenderData() cons
 	}
 }
 
+void USkinnedMeshComponent::GetUpdateClothSimulationData_AnyThread(TMap<int32, FClothSimulData>& OutClothSimulData, FMatrix& OutLocalToWorld, float& OutClothBlendWeight)
+{
+	OutClothSimulData.Reset();
+	OutClothBlendWeight = 0.f;
+}
+
 bool USkinnedMeshComponent::AllocateTransformData()
 {
 	LLM_SCOPE_BYNAME(TEXT("SkeletalMesh/TransformData"));
diff --git a/Engine/Source/Runtime/Engine/Private/SkeletalMeshComponentPhysics.cpp b/Engine/Source/Runtime/Engine/Private/SkeletalMeshComponentPhysics.cpp
--- a/Engine/Source/Runtime/Engine/Private/SkeletalMeshComponentPhysics.cpp
+++ b/Engine/Source/Runtime/Engine/Private/SkeletalMeshComponentPhysics.cpp
@@ -3844,6 +3844,33 @@ void USkeletalMeshComponent::GetUpdateClothSimulationData(TMap<int32, FClothSimu
 	}
 }
 
+void USkeletalMeshComponent::GetUpdateClothSimulationData_AnyThread(TMap<int32, FClothSimulData>& OutClothSimulData, FMatrix& OutLocalToWorld, float& OutClothBlendWeight)
+{
+	OutLocalToWorld = GetComponentToWorld().ToMatrixWithScale();
+
+	const USkeletalMeshComponent* const LeaderPoseSkeletalMeshComponent = Cast<USkeletalMeshComponent>(LeaderPoseComponent.Get());
+	if (LeaderPoseSkeletalMeshComponent && bBindClothToLeaderComponent)
+	{
+		OutClothBlendWeight = ClothBlendWeight;
+		OutClothSimulData = LeaderPoseSkeletalMeshComponent->GetCurrentClothingData_AnyThread();
+	}
+	else if (!bDisableClothSimulation && !bBindClothToLeaderComponent)
+	{
+		OutClothBlendWeight = ClothBlendWeight;
+		OutClothSimulData = GetCurrentClothingData_AnyThread();
+	}
+	else
+	{
+		OutClothSimulData.Reset();
+	}
+
+	// Blend cloth out whenever the simulation data is invalid
+	if (!OutClothSimulData.Num())
+	{
+		OutClothBlendWeight = 0.0f;
+	}
+}
+
 void USkeletalMeshComponent::DebugDrawClothing(FPrimitiveDrawInterface* PDI)
 {
 #if WITH_EDITOR && ENABLE_DRAW_DEBUG
diff --git a/Engine/Source/Runtime/Engine/Private/SkeletalRenderCPUSkin.cpp b/Engine/Source/Runtime/Engine/Private/SkeletalRenderCPUSkin.cpp
--- a/Engine/Source/Runtime/Engine/Private/SkeletalRenderCPUSkin.cpp
+++ b/Engine/Source/Runtime/Engine/Private/SkeletalRenderCPUSkin.cpp
@@ -586,7 +586,13 @@ FDynamicSkelMeshObjectDataCPUSkin::FDynamicSkelMeshObjectDataCPUSkin(
 	UpdateRefToLocalMatrices( ReferenceToLocal, InMeshComponent, InSkelMeshRenderData, LODIndex );
 
 	// Update the clothing simulation mesh positions and normals
-	UpdateClothSimulationData(InMeshComponent);
+	FMatrix LocalToWorld;
+	InMeshComponent->GetUpdateClothSimulationData_AnyThread(ClothSimulUpdateData, LocalToWorld, ClothBlendWeight);
+	WorldToLocal = LocalToWorld.InverseFast();
+	if (!IsSkeletalMeshClothBlendEnabled())
+	{
+		ClothBlendWeight = 0.f;
+	}
 
 #if !(UE_BUILD_SHIPPING || UE_BUILD_TEST)
 	MeshComponentSpaceTransforms = InMeshComponent->GetComponentSpaceTransforms();
@@ -1213,7 +1219,9 @@ bool FDynamicSkelMeshObjectDataCPUSkin::UpdateClothSimulationData(USkinnedMeshCo
 
 		WorldToLocal = SrcComponent->GetRenderMatrix().InverseFast();
 		ClothBlendWeight = IsSkeletalMeshClothBlendEnabled() ? SrcComponent->ClothBlendWeight : 0.0f;
+PRAGMA_DISABLE_DEPRECATION_WARNINGS
 		SimMeshComponent->GetUpdateClothSimulationData(ClothSimulUpdateData, SrcComponent);
+PRAGMA_ENABLE_DEPRECATION_WARNINGS
 
 		return true;
 	}
@@ -1222,7 +1230,9 @@ bool FDynamicSkelMeshObjectDataCPUSkin::UpdateClothSimulationData(USkinnedMeshCo
 	{
 		WorldToLocal = SimMeshComponent->GetRenderMatrix().InverseFast();
 		ClothBlendWeight = IsSkeletalMeshClothBlendEnabled() ? SimMeshComponent->ClothBlendWeight : 0.0f;
+PRAGMA_DISABLE_DEPRECATION_WARNINGS
 		SimMeshComponent->GetUpdateClothSimulationData(ClothSimulUpdateData);
+PRAGMA_ENABLE_DEPRECATION_WARNINGS
 		return true;
 	}
 	return false;
diff --git a/Engine/Source/Runtime/Engine/Private/SkeletalRenderCPUSkin.h b/Engine/Source/Runtime/Engine/Private/SkeletalRenderCPUSkin.h
--- a/Engine/Source/Runtime/Engine/Private/SkeletalRenderCPUSkin.h
+++ b/Engine/Source/Runtime/Engine/Private/SkeletalRenderCPUSkin.h
@@ -81,6 +81,7 @@ public:
  	}
 
 	/** Update Simulated Positions & Normals from Clothing actor */
+	UE_DEPRECATED(5.2, "Use USkinnedMeshComponent::GetUpdateClothSimulationData_AnyThread() instead.")
 	bool UpdateClothSimulationData(USkinnedMeshComponent* InMeshComponent);
 };
 
diff --git a/Engine/Source/Runtime/Engine/Private/SkeletalRenderGPUSkin.cpp b/Engine/Source/Runtime/Engine/Private/SkeletalRenderGPUSkin.cpp
--- a/Engine/Source/Runtime/Engine/Private/SkeletalRenderGPUSkin.cpp
+++ b/Engine/Source/Runtime/Engine/Private/SkeletalRenderGPUSkin.cpp
@@ -2544,7 +2544,12 @@ void FDynamicSkelMeshObjectDataGPUSkin::InitDynamicSkelMeshObjectDataGPUSkin(
 	LocalToWorld = InMeshComponent ? InMeshComponent->GetComponentTransform().ToMatrixWithScale() : FMatrix::Identity;
 
 	// Update the clothing simulation mesh positions and normals
-	UpdateClothSimulationData(InMeshComponent);
+	InMeshComponent->GetUpdateClothSimulationData_AnyThread(ClothingSimData, ClothObjectLocalToWorld, ClothBlendWeight);
+
+	if (!IsSkeletalMeshClothBlendEnabled())
+	{
+		ClothBlendWeight = 0.f;
+	}
 
 	bIsSkinCacheAllowed = InMeshComponent ? InMeshComponent->IsSkinCacheAllowed(InLODIndex) : false;
 	bHasMeshDeformer = InMeshComponent ? InMeshComponent->HasMeshDeformer() : false;
diff --git a/Engine/Source/Runtime/Engine/Private/SkeletalRenderGPUSkin.h b/Engine/Source/Runtime/Engine/Private/SkeletalRenderGPUSkin.h
--- a/Engine/Source/Runtime/Engine/Private/SkeletalRenderGPUSkin.h
+++ b/Engine/Source/Runtime/Engine/Private/SkeletalRenderGPUSkin.h
@@ -134,6 +134,7 @@ public:
 	}
 
 	/** Update Simulated Positions & Normals from APEX Clothing actor */
+	UE_DEPRECATED(5.2, "Use USkinnedMeshComponent::GetUpdateClothSimulationData_AnyThread() instead.")
 	bool UpdateClothSimulationData(USkinnedMeshComponent* InMeshComponent);
 
 	// Whether this LOD is allowed to use the skin cache feature
