CIS - Fix warning warning V595: The 'InMeshComponent' pointer was utilized before it was verified against nullptr.

diff --git a/Engine/Source/Runtime/Engine/Private/Components/SkinnedMeshComponent.cpp b/Engine/Source/Runtime/Engine/Private/Components/SkinnedMeshComponent.cpp
--- a/Engine/Source/Runtime/Engine/Private/Components/SkinnedMeshComponent.cpp
+++ b/Engine/Source/Runtime/Engine/Private/Components/SkinnedMeshComponent.cpp
@@ -2082,6 +2082,7 @@ FSkeletalMeshRenderData* USkinnedMeshComponent::GetSkeletalMeshRenderData() cons
 void USkinnedMeshComponent::GetUpdateClothSimulationData_AnyThread(TMap<int32, FClothSimulData>& OutClothSimulData, FMatrix& OutLocalToWorld, float& OutClothBlendWeight)
 {
 	OutClothSimulData.Reset();
+	OutLocalToWorld = FMatrix::Identity;
 	OutClothBlendWeight = 0.f;
 }
 
diff --git a/Engine/Source/Runtime/Engine/Private/SkeletalRenderCPUSkin.cpp b/Engine/Source/Runtime/Engine/Private/SkeletalRenderCPUSkin.cpp
--- a/Engine/Source/Runtime/Engine/Private/SkeletalRenderCPUSkin.cpp
+++ b/Engine/Source/Runtime/Engine/Private/SkeletalRenderCPUSkin.cpp
@@ -587,7 +587,17 @@ FDynamicSkelMeshObjectDataCPUSkin::FDynamicSkelMeshObjectDataCPUSkin(
 
 	// Update the clothing simulation mesh positions and normals
 	FMatrix LocalToWorld;
-	InMeshComponent->GetUpdateClothSimulationData_AnyThread(ClothSimulUpdateData, LocalToWorld, ClothBlendWeight);
+	if (InMeshComponent)
+	{
+		InMeshComponent->GetUpdateClothSimulationData_AnyThread(ClothSimulUpdateData, LocalToWorld, ClothBlendWeight);
+	}
+	else
+	{
+		ClothSimulUpdateData.Reset();
+		LocalToWorld = FMatrix::Identity;
+		ClothBlendWeight = 0.f;
+	}
+
 	WorldToLocal = LocalToWorld.InverseFast();
 	if (!IsSkeletalMeshClothBlendEnabled())
 	{
diff --git a/Engine/Source/Runtime/Engine/Private/SkeletalRenderGPUSkin.cpp b/Engine/Source/Runtime/Engine/Private/SkeletalRenderGPUSkin.cpp
--- a/Engine/Source/Runtime/Engine/Private/SkeletalRenderGPUSkin.cpp
+++ b/Engine/Source/Runtime/Engine/Private/SkeletalRenderGPUSkin.cpp
@@ -2544,7 +2544,16 @@ void FDynamicSkelMeshObjectDataGPUSkin::InitDynamicSkelMeshObjectDataGPUSkin(
 	LocalToWorld = InMeshComponent ? InMeshComponent->GetComponentTransform().ToMatrixWithScale() : FMatrix::Identity;
 
 	// Update the clothing simulation mesh positions and normals
-	InMeshComponent->GetUpdateClothSimulationData_AnyThread(ClothingSimData, ClothObjectLocalToWorld, ClothBlendWeight);
+	if (InMeshComponent)
+	{
+		InMeshComponent->GetUpdateClothSimulationData_AnyThread(ClothingSimData, ClothObjectLocalToWorld, ClothBlendWeight);
+	}
+	else
+	{
+		ClothingSimData.Reset();
+		ClothObjectLocalToWorld = FMatrix::Identity;
+		ClothBlendWeight = 0.f;
+	}
 
 	if (!IsSkeletalMeshClothBlendEnabled())
 	{
