Split ISubmixBufferListener from AudioDevice to improve compile times. AudioDevice.h is an expensive header to include and several files just want the ISubmixBufferListener interface.

diff --git a/Engine/Plugins/Experimental/MotoSynth/Source/MotoSynth/Private/MotoSynthSourceAsset.cpp b/Engine/Plugins/Experimental/MotoSynth/Source/MotoSynth/Private/MotoSynthSourceAsset.cpp
--- a/Engine/Plugins/Experimental/MotoSynth/Source/MotoSynth/Private/MotoSynthSourceAsset.cpp
+++ b/Engine/Plugins/Experimental/MotoSynth/Source/MotoSynth/Private/MotoSynthSourceAsset.cpp
@@ -1,12 +1,13 @@
 // Copyright Epic Games, Inc. All Rights Reserved.
 
 #include "MotoSynthSourceAsset.h"
+
 #include "MotoSynthModule.h"
 #include "MotoSynthDataManager.h"
 
-#include UE_INLINE_GENERATED_CPP_BY_NAME(MotoSynthSourceAsset)
-
 #if WITH_EDITOR
+#include "AudioDeviceManager.h"
+#include "AudioDevice.h"
 #include "DSP/Filter.h"
 #include "DSP/DynamicsProcessor.h"
 #include "DSP/Dsp.h"
@@ -14,6 +15,8 @@
 #include "DSP/SampleRateConverter.h"
 #endif
 
+#include UE_INLINE_GENERATED_CPP_BY_NAME(MotoSynthSourceAsset)
+
 ///////////////////////////////////////////////////////////
 // RPM estimation implementation
 ///////////////////////////////////////////////////////////
diff --git a/Engine/Plugins/Experimental/MotoSynth/Source/MotoSynth/Public/MotoSynthPreset.h b/Engine/Plugins/Experimental/MotoSynth/Source/MotoSynth/Public/MotoSynthPreset.h
--- a/Engine/Plugins/Experimental/MotoSynth/Source/MotoSynth/Public/MotoSynthPreset.h
+++ b/Engine/Plugins/Experimental/MotoSynth/Source/MotoSynth/Public/MotoSynthPreset.h
@@ -10,9 +10,13 @@
 #include "Async/AsyncWork.h"
 #include "DSP/BufferVectorOperations.h"
 #include "Curves/CurveFloat.h"
-#include "AudioDevice.h"
+#include "ISubmixBufferListener.h"
 #include "MotoSynthPreset.generated.h"
 
+#if UE_ENABLE_INCLUDE_ORDER_DEPRECATED_IN_5_1
+#include "AudioDevice.h"
+#endif
+
 class UMotoSynthSource;
 class FMotoSynthEngine;
 
diff --git a/Engine/Plugins/Experimental/MotoSynth/Source/MotoSynth/Public/MotoSynthSourceAsset.h b/Engine/Plugins/Experimental/MotoSynth/Source/MotoSynth/Public/MotoSynthSourceAsset.h
--- a/Engine/Plugins/Experimental/MotoSynth/Source/MotoSynth/Public/MotoSynthSourceAsset.h
+++ b/Engine/Plugins/Experimental/MotoSynth/Source/MotoSynth/Public/MotoSynthSourceAsset.h
@@ -11,13 +11,17 @@
 #include "DSP/BufferVectorOperations.h"
 #include "Curves/CurveFloat.h"
 #include "Containers/SortedMap.h"
-#include "AudioDevice.h"
 #include "DSP/Osc.h"
 #include "DSP/Filter.h"
 #include "DSP/DelayStereo.h"
+#include "ISubmixBufferListener.h"
 #include "MotoSynthPreset.h"
 #include "MotoSynthSourceAsset.generated.h"
 
+#if UE_ENABLE_INCLUDE_ORDER_DEPRECATED_IN_5_1
+#include "AudioDevice.h"
+#endif
+
 class UMotoSynthSource;
 
 USTRUCT()
diff --git a/Engine/Plugins/FX/Niagara/Source/Niagara/Classes/NiagaraDataInterfaceAudio.h b/Engine/Plugins/FX/Niagara/Source/Niagara/Classes/NiagaraDataInterfaceAudio.h
--- a/Engine/Plugins/FX/Niagara/Source/Niagara/Classes/NiagaraDataInterfaceAudio.h
+++ b/Engine/Plugins/FX/Niagara/Source/Niagara/Classes/NiagaraDataInterfaceAudio.h
@@ -3,13 +3,16 @@
 
 #include "CoreMinimal.h"
 #include "NiagaraDataInterface.h"
-#include "AudioDevice.h"
 #include "AudioDeviceManager.h"
-#include "Sound/SoundSubmix.h"
+#include "ISubmixBufferListener.h"
 #include "DSP/MultithreadedPatching.h"
 #include "UObject/WeakObjectPtrTemplates.h"
 #include "NiagaraDataInterfaceAudio.generated.h"
 
+#if UE_ENABLE_INCLUDE_ORDER_DEPRECATED_IN_5_1
+#include "AudioDevice.h"
+#endif
+
 /** Class used to to capture the audio stream of an arbitrary submix. */
 class NIAGARA_API FNiagaraSubmixListener : public ISubmixBufferListener
 {
diff --git a/Engine/Plugins/FX/Niagara/Source/Niagara/Private/NiagaraDataInterfaceAudio.cpp b/Engine/Plugins/FX/Niagara/Source/Niagara/Private/NiagaraDataInterfaceAudio.cpp
--- a/Engine/Plugins/FX/Niagara/Source/Niagara/Private/NiagaraDataInterfaceAudio.cpp
+++ b/Engine/Plugins/FX/Niagara/Source/Niagara/Private/NiagaraDataInterfaceAudio.cpp
@@ -2,6 +2,8 @@
 
 #include "NiagaraDataInterfaceAudio.h"
 
+#include "AudioDevice.h"
+
 #include UE_INLINE_GENERATED_CPP_BY_NAME(NiagaraDataInterfaceAudio)
 
 FNiagaraSubmixListener::FNiagaraSubmixListener(Audio::FPatchMixer& InMixer, int32 InNumSamplesToBuffer, Audio::FDeviceId InDeviceId, USoundSubmix* InSoundSubmix)
diff --git a/Engine/Plugins/Media/MediaIOFramework/Source/MediaIOCore/Public/MediaIOCoreAudioOutput.h b/Engine/Plugins/Media/MediaIOFramework/Source/MediaIOCore/Public/MediaIOCoreAudioOutput.h
--- a/Engine/Plugins/Media/MediaIOFramework/Source/MediaIOCore/Public/MediaIOCoreAudioOutput.h
+++ b/Engine/Plugins/Media/MediaIOFramework/Source/MediaIOCore/Public/MediaIOCoreAudioOutput.h
@@ -2,10 +2,14 @@
 
 #pragma once
 
-#include "AudioDevice.h"
 #include "DSP/BufferVectorOperations.h"
+#include "ISubmixBufferListener.h"
 #include "Math/NumericLimits.h"
 
+#if UE_ENABLE_INCLUDE_ORDER_DEPRECATED_IN_5_1
+#include "AudioDevice.h"
+#endif
+
 namespace UE::MediaIoCoreModule::Private
 {
 	template <typename OutputType>
diff --git a/Engine/Plugins/Media/PixelStreaming/Source/PixelStreaming/Private/AudioCapturer.h b/Engine/Plugins/Media/PixelStreaming/Source/PixelStreaming/Private/AudioCapturer.h
--- a/Engine/Plugins/Media/PixelStreaming/Source/PixelStreaming/Private/AudioCapturer.h
+++ b/Engine/Plugins/Media/PixelStreaming/Source/PixelStreaming/Private/AudioCapturer.h
@@ -3,8 +3,8 @@
 #pragma once
 
 #include "WebRTCIncludes.h"
-#include "AudioDevice.h"
 #include "HAL/CriticalSection.h"
+#include "ISubmixBufferListener.h"
 
 namespace UE::PixelStreaming 
 {
diff --git a/Engine/Plugins/Media/PixelStreaming/Source/PixelStreaming/Private/AudioSubmixCapturer.h b/Engine/Plugins/Media/PixelStreaming/Source/PixelStreaming/Private/AudioSubmixCapturer.h
--- a/Engine/Plugins/Media/PixelStreaming/Source/PixelStreaming/Private/AudioSubmixCapturer.h
+++ b/Engine/Plugins/Media/PixelStreaming/Source/PixelStreaming/Private/AudioSubmixCapturer.h
@@ -2,7 +2,7 @@
 
 #pragma once
 
-#include "AudioDevice.h"
+#include "ISubmixBufferListener.h"
 
 namespace webrtc
 {
diff --git a/Engine/Source/Runtime/AudioLink/AudioLinkEngine/Private/BufferedSubmixListener.h b/Engine/Source/Runtime/AudioLink/AudioLinkEngine/Private/BufferedSubmixListener.h
--- a/Engine/Source/Runtime/AudioLink/AudioLinkEngine/Private/BufferedSubmixListener.h
+++ b/Engine/Source/Runtime/AudioLink/AudioLinkEngine/Private/BufferedSubmixListener.h
@@ -3,6 +3,7 @@
 #pragma once
 
 #include "BufferedListenerBase.h"
+#include "ISubmixBufferListener.h"
 
 /** Concrete Submix Buffer Listener  
 */
diff --git a/Engine/Source/Runtime/Engine/Public/AudioDevice.h b/Engine/Source/Runtime/Engine/Public/AudioDevice.h
--- a/Engine/Source/Runtime/Engine/Public/AudioDevice.h
+++ b/Engine/Source/Runtime/Engine/Public/AudioDevice.h
@@ -10,6 +10,7 @@
 #include "EngineGlobals.h"
 #include "HAL/LowLevelMemStats.h"
 #include "IAudioExtensionPlugin.h"
+#include "ISubmixBufferListener.h"
 #include "AudioDynamicParameter.h"
 #include "Sound/AudioSettings.h"
 #include "Sound/AudioVolume.h"
@@ -407,22 +408,6 @@ public:
 	virtual void OnDefaultDeviceChanged() = 0;
 };
 
-/** Abstract interface for receiving audio data from a given submix. */
-class ENGINE_API ISubmixBufferListener
-{
-public:
-	/**
-	Called when a new buffer has been rendered for a given submix
-	@param OwningSubmix	The submix object which has rendered a new buffer
-	@param AudioData		Ptr to the audio buffer
-	@param NumSamples		The number of audio samples in the audio buffer
-	@param NumChannels		The number of channels of audio in the buffer (e.g. 2 for stereo, 6 for 5.1, etc)
-	@param SampleRate		The sample rate of the audio buffer
-	@param AudioClock		Double audio clock value, from start of audio rendering.
-	*/
-	virtual void OnNewSubmixBuffer(const USoundSubmix* OwningSubmix, float* AudioData, int32 NumSamples, int32 NumChannels, const int32 SampleRate, double AudioClock) = 0;
-};
-
 struct FAudioDeviceRenderInfo
 {
 	int32 NumFrames = 0;
diff --git a/Engine/Source/Runtime/Engine/Public/ISubmixBufferListener.h b/Engine/Source/Runtime/Engine/Public/ISubmixBufferListener.h
new file mode 100644
--- /dev/null
+++ b/Engine/Source/Runtime/Engine/Public/ISubmixBufferListener.h
@@ -0,0 +1,20 @@
+// Copyright Epic Games, Inc. All Rights Reserved.
+#pragma once
+
+#include "Sound/SoundSubmix.h"
+
+/** Abstract interface for receiving audio data from a given submix. */
+class ENGINE_API ISubmixBufferListener
+{
+public:
+	/**
+	Called when a new buffer has been rendered for a given submix
+	@param OwningSubmix	The submix object which has rendered a new buffer
+	@param AudioData		Ptr to the audio buffer
+	@param NumSamples		The number of audio samples in the audio buffer
+	@param NumChannels		The number of channels of audio in the buffer (e.g. 2 for stereo, 6 for 5.1, etc)
+	@param SampleRate		The sample rate of the audio buffer
+	@param AudioClock		Double audio clock value, from start of audio rendering.
+	*/
+	virtual void OnNewSubmixBuffer(const USoundSubmix* OwningSubmix, float* AudioData, int32 NumSamples, int32 NumChannels, const int32 SampleRate, double AudioClock) = 0;
+};
\ No newline at end of file
