Fixed: Crash on uploading reflection captures when project has mobile deferred shading enabled

diff --git a/Engine/Source/Runtime/Engine/Private/Components/ReflectionCaptureComponent.cpp b/Engine/Source/Runtime/Engine/Private/Components/ReflectionCaptureComponent.cpp
--- a/Engine/Source/Runtime/Engine/Private/Components/ReflectionCaptureComponent.cpp
+++ b/Engine/Source/Runtime/Engine/Private/Components/ReflectionCaptureComponent.cpp
@@ -563,16 +563,15 @@ public:
 	FReflectionTextureCubeResource() :
 		Size(0),
 		NumMips(0),
-		Format(PF_Unknown),
-		SourceData(nullptr)
+		Format(PF_Unknown)
 	{}
 
-	void SetupParameters(int32 InSize, int32 InNumMips, EPixelFormat InFormat, TArray<uint8>* InSourceData)
+	void SetupParameters(int32 InSize, int32 InNumMips, EPixelFormat InFormat, TArray<uint8>&& InSourceData)
 	{
 		Size = InSize;
 		NumMips = InNumMips;
 		Format = InFormat;
-		SourceData = InSourceData;
+		SourceData = MoveTemp(InSourceData);
 	}
 
 	virtual void InitRHI() override
@@ -585,10 +584,8 @@ public:
 		TextureCubeRHI = RHICreateTexture(Desc);
 		TextureRHI = TextureCubeRHI;
 
-		if (SourceData)
+		if (SourceData.Num())
 		{
-			check(SourceData->Num() > 0);
-
 			const int32 BlockBytes = GPixelFormats[Format].BlockBytes;
 			int32 MipBaseIndex = 0;
 
@@ -605,7 +602,7 @@ public:
 					uint8* DestBuffer = (uint8*)RHILockTextureCubeFace(TextureCubeRHI, CubeFace, 0, MipIndex, RLM_WriteOnly, DestPitch, false);
 					
 					const int32 SourceIndex = MipBaseIndex + CubeFace * CubeFaceBytes;
-					const uint8* SourceBuffer = &(*SourceData)[SourceIndex];
+					const uint8* SourceBuffer = &SourceData[SourceIndex];
 
 					if (SrcPitch == DestPitch)
 					{
@@ -631,12 +628,7 @@ public:
 				MipBaseIndex += CubeFaceBytes * CubeFace_MAX;
 			}
 
-			if (!GIsEditor)
-			{
-				// Toss the source data now that we've created the cubemap
-				// Note: can't do this if we ever use this texture resource in the editor and want to save the data later
-				SourceData->Empty();
-			}
+			SourceData.Empty();
 		}
 
 		// Create the sampler state RHI resource.
@@ -681,7 +673,7 @@ private:
 	EPixelFormat Format;
 	FTextureCubeRHIRef TextureCubeRHI;
 
-	TArray<uint8>* SourceData;
+	TArray<uint8> SourceData;
 };
 
 TArray<UReflectionCaptureComponent*> UReflectionCaptureComponent::ReflectionCapturesToUpdate;
@@ -744,8 +736,9 @@ void UReflectionCaptureComponent::SafeReleaseEncodedHDRCubemapTexture()
 
 void UReflectionCaptureComponent::OnRegister()
 {
-	const bool bMobileDeferredShading = IsMobileDeferredShadingEnabled(GMaxRHIShaderPlatform);
-	const bool bEncodedHDRCubemapTextureRequired = (GIsEditor || GMaxRHIFeatureLevel == ERHIFeatureLevel::ES3_1) && !MobileForwardEnableClusteredReflections(GMaxRHIShaderPlatform);
+	const bool bEncodedHDRCubemapTextureRequired = (GIsEditor || GMaxRHIFeatureLevel == ERHIFeatureLevel::ES3_1) 
+		// mobile forward renderer or translucensy in mobile deferred need encoded reflection texture when clustered reflections disabled
+		&& !MobileForwardEnableClusteredReflections(GMaxRHIShaderPlatform);
 
 	if (bEncodedHDRCubemapTextureRequired)
 	{
@@ -757,20 +750,29 @@ void UReflectionCaptureComponent::OnRegister()
 			if (!EncodedHDRCubemapTexture)
 			{
 				EncodedHDRCubemapTexture = new FReflectionTextureCubeResource();
-				TArray<uint8>* EncodedHDRCapturedData = nullptr;
+				TArray<uint8> EncodedHDRCapturedData;
 				EPixelFormat EncodedHDRCubemapTextureFormat = PF_Unknown;
-				if (bMobileDeferredShading)
+				if (IsMobileDeferredShadingEnabled(GMaxRHIShaderPlatform))
 				{
-					EncodedHDRCapturedData = &MapBuildData->FullHDRCapturedData;
+					// make a copy, FullHDR data will still be needed later
+					EncodedHDRCapturedData = MapBuildData->FullHDRCapturedData;
 					EncodedHDRCubemapTextureFormat = PF_FloatRGBA;
 				}
 				else
 				{
-					EncodedHDRCapturedData = &MapBuildData->EncodedHDRCapturedData;
+					if (GIsEditor)
+					{ 
+						EncodedHDRCapturedData = MapBuildData->EncodedHDRCapturedData;
+					}
+					else
+					{
+						// EncodedHDR data is no longer needed
+						Swap(EncodedHDRCapturedData, MapBuildData->EncodedHDRCapturedData);
+					}
 					EncodedHDRCubemapTextureFormat = PF_FloatR11G11B10;
 				}
 
-				EncodedHDRCubemapTexture->SetupParameters(MapBuildData->CubemapSize, FMath::CeilLogTwo(MapBuildData->CubemapSize) + 1, EncodedHDRCubemapTextureFormat, EncodedHDRCapturedData);
+				EncodedHDRCubemapTexture->SetupParameters(MapBuildData->CubemapSize, FMath::CeilLogTwo(MapBuildData->CubemapSize) + 1, EncodedHDRCubemapTextureFormat, MoveTemp(EncodedHDRCapturedData));
 				BeginInitResource(EncodedHDRCubemapTexture);
 			}
 
