Use the LODMaterialMap to retrieve the actual used material index if available during PSO precaching of a skinned component

#rnx

diff --git a/Engine/Source/Runtime/Engine/Private/Components/SkinnedMeshComponent.cpp b/Engine/Source/Runtime/Engine/Private/Components/SkinnedMeshComponent.cpp
--- a/Engine/Source/Runtime/Engine/Private/Components/SkinnedMeshComponent.cpp
+++ b/Engine/Source/Runtime/Engine/Private/Components/SkinnedMeshComponent.cpp
@@ -577,14 +577,25 @@ void USkinnedMeshComponent::PrecachePSOs()
 	for (int32 LODIndex = MinLODIndex; LODIndex < SkelMeshRenderData->LODRenderData.Num(); LODIndex++)
 	{
 		FSkeletalMeshLODRenderData& LODRenderData = SkelMeshRenderData->LODRenderData[LODIndex];
+		const FSkeletalMeshLODInfo& Info = *(GetSkinnedAsset()->GetLODInfo(LODIndex));
 		
 		// Check all render sections for the used material indices
-		for (FSkelMeshRenderSection& RenderSection : LODRenderData.RenderSections)
+		for (int32 SectionIndex = 0; SectionIndex < LODRenderData.RenderSections.Num(); SectionIndex++)
 		{
+			FSkelMeshRenderSection& RenderSection = LODRenderData.RenderSections[SectionIndex];
+
 			bAnySectionCastsShadows |= RenderSection.bCastShadow;
 
+			// By default use the material index of the render section
 			uint16 MaterialIndex = RenderSection.MaterialIndex;
 
+			// Can be remapped by the LODInfo
+			if (SectionIndex < Info.LODMaterialMap.Num() && GetSkinnedAsset()->IsValidMaterialIndex(Info.LODMaterialMap[SectionIndex]))
+			{
+				MaterialIndex = Info.LODMaterialMap[SectionIndex];
+				MaterialIndex = FMath::Clamp(MaterialIndex, 0, GetSkinnedAsset()->GetNumMaterials());
+			}
+
 			VFsPerMaterialData* VFsPerMaterial = VFsPerMaterials.FindByPredicate([MaterialIndex](const VFsPerMaterialData& Other) { return Other.MaterialIndex == MaterialIndex; });
 			if (VFsPerMaterial == nullptr)
 			{
