Fix static analysis and change check to ensure to not crash material editor

diff --git a/Engine/Source/Runtime/Engine/Private/Materials/MaterialExpressions.cpp b/Engine/Source/Runtime/Engine/Private/Materials/MaterialExpressions.cpp
--- a/Engine/Source/Runtime/Engine/Private/Materials/MaterialExpressions.cpp
+++ b/Engine/Source/Runtime/Engine/Private/Materials/MaterialExpressions.cpp
@@ -1534,9 +1534,11 @@ void UMaterialExpression::PostEditChangeProperty(FPropertyChangedEvent& Property
 					{
 						// ...modify the pin on other parameter expression node with the new value.
 						UEdGraphPin* OtherPin = ExpressionParameter->GraphNode->GetPinAt(PinIndex);
-						check(OtherPin->GetName() == Pin->GetName());
-						OtherPin->Modify();
-						OtherPin->DefaultValue = NewDefaultValue;
+						if (ensure(OtherPin && Pin && OtherPin->GetName() == Pin->GetName()))
+						{
+							OtherPin->Modify();
+							OtherPin->DefaultValue = NewDefaultValue;
+						}
 					}
 				}
 
