Remove Strata nodes & material functions from the material contextual menu when Strata is disabled.

diff --git a/Engine/Source/Editor/MaterialEditor/Private/MaterialExpressionClasses.cpp b/Engine/Source/Editor/MaterialEditor/Private/MaterialExpressionClasses.cpp
--- a/Engine/Source/Editor/MaterialEditor/Private/MaterialExpressionClasses.cpp
+++ b/Engine/Source/Editor/MaterialEditor/Private/MaterialExpressionClasses.cpp
@@ -24,6 +24,7 @@
 #include "Materials/MaterialExpressionForLoop.h"
 #include "Materials/MaterialExpressionGetLocal.h"
 #include "Materials/MaterialExpressionSetLocal.h"
+#include "Materials/MaterialExpressionStrata.h"
 
 ///////////////////////////////////////////////////////////////////////////////////////////////////////////////////////
 //
@@ -80,6 +81,12 @@ FCategorizedMaterialExpressionNode* MaterialExpressionClasses::GetCategoryNode(c
 	return NULL;
 }
 
+bool MaterialEditor_IsStrataEnabled()
+{
+	static const auto CVar = IConsoleManager::Get().FindTConsoleVariableDataInt(TEXT("r.Strata"));
+	return CVar && CVar->GetValueOnAnyThread() > 0;
+}
+
 void MaterialExpressionClasses::InitMaterialExpressionClasses()
 {
 	if(!bInitialized)
@@ -173,6 +180,13 @@ void MaterialExpressionClasses::InitMaterialExpressionClasses()
 								}
 							}
 
+							// Skip adding Strata node to the context menu if Strata is disabled
+							// STRATA_TODO: remove this when Strata becomes the only shading path
+							if ((Class->IsChildOf(UMaterialExpressionStrataBSDF::StaticClass()) || Class->IsChildOf(UMaterialExpressionStrataUtilityBase::StaticClass())) && !MaterialEditor_IsStrataEnabled())
+							{
+								continue;
+							}
+
 							// Category fill...
 							if (TempObject)
 							{
diff --git a/Engine/Source/Editor/UnrealEd/Private/MaterialGraphSchema.cpp b/Engine/Source/Editor/UnrealEd/Private/MaterialGraphSchema.cpp
--- a/Engine/Source/Editor/UnrealEd/Private/MaterialGraphSchema.cpp
+++ b/Engine/Source/Editor/UnrealEd/Private/MaterialGraphSchema.cpp
@@ -908,6 +908,12 @@ TSharedPtr<FEdGraphSchemaAction> UMaterialGraphSchema::GetCreateCommentAction()
 	return TSharedPtr<FEdGraphSchemaAction>(static_cast<FEdGraphSchemaAction*>(new FMaterialGraphSchemaAction_NewComment));
 }
 
+bool UnrealEd_IsStrataEnabled()
+{
+	static const auto CVar = IConsoleManager::Get().FindTConsoleVariableDataInt(TEXT("r.Strata"));
+	return CVar && CVar->GetValueOnAnyThread() > 0;
+}
+
 void UMaterialGraphSchema::GetMaterialFunctionActions(FGraphActionMenuBuilder& ActionMenuBuilder) const
 {
 	// Get type of dragged pin
@@ -983,6 +989,22 @@ void UMaterialGraphSchema::GetMaterialFunctionActions(FGraphActionMenuBuilder& A
 					{
 						LibraryCategoriesText.Add( LOCTEXT("UncategorizedMaterialFunction", "Uncategorized") );
 					}
+
+					// When Strata is disabled, skip all material function related to Strata
+					// STRATA_TODO: remove this when Strata becomes the only shading path
+					bool bSkipMaterialFunction = false;
+					for (const FText& Category : LibraryCategoriesText)
+					{
+						if (Category.ToString().Contains("Strata"))
+						{
+							bSkipMaterialFunction = !UnrealEd_IsStrataEnabled();
+							break;
+						}
+					}
+					if (bSkipMaterialFunction)
+					{
+						continue;
+					}
 				}
 
 				// Extract the object name from the path
diff --git a/Engine/Source/Runtime/Engine/Classes/Materials/MaterialExpressionStrata.h b/Engine/Source/Runtime/Engine/Classes/Materials/MaterialExpressionStrata.h
--- a/Engine/Source/Runtime/Engine/Classes/Materials/MaterialExpressionStrata.h
+++ b/Engine/Source/Runtime/Engine/Classes/Materials/MaterialExpressionStrata.h
@@ -974,8 +974,14 @@ class UMaterialExpressionStrataThinFilm : public UMaterialExpressionStrataBSDF
 ///////////////////////////////////////////////////////////////////////////////
 // Utilities
 
+UCLASS(MinimalAPI, collapsecategories, hidecategories = Object, Abstract, DisplayName = "Strata Utility Base Class")
+class UMaterialExpressionStrataUtilityBase : public UMaterialExpression
+{
+	GENERATED_UCLASS_BODY()
+};
+
 UCLASS(MinimalAPI, collapsecategories, hidecategories = Object, DisplayName = "Strata Transmittance-To-MeanFreePath")
-class UMaterialExpressionStrataTransmittanceToMFP : public UMaterialExpression
+class UMaterialExpressionStrataTransmittanceToMFP : public UMaterialExpressionStrataUtilityBase
 {
 	GENERATED_UCLASS_BODY()
 
@@ -1005,7 +1011,7 @@ class UMaterialExpressionStrataTransmittanceToMFP : public UMaterialExpression
 };
 
 UCLASS(MinimalAPI, collapsecategories, hidecategories = Object, DisplayName = "Strata Metalness-To-DiffuseColorF0")
-class UMaterialExpressionStrataMetalnessToDiffuseAlbedoF0 : public UMaterialExpression
+class UMaterialExpressionStrataMetalnessToDiffuseAlbedoF0 : public UMaterialExpressionStrataUtilityBase
 {
 	GENERATED_UCLASS_BODY()
 
@@ -1040,7 +1046,7 @@ class UMaterialExpressionStrataMetalnessToDiffuseAlbedoF0 : public UMaterialExpr
 };
 
 UCLASS(MinimalAPI, collapsecategories, hidecategories = Object, DisplayName = "Strata Haziness-To-Secondary-Roughness")
-class UMaterialExpressionStrataHazinessToSecondaryRoughness : public UMaterialExpression
+class UMaterialExpressionStrataHazinessToSecondaryRoughness : public UMaterialExpressionStrataUtilityBase
 {
 	GENERATED_UCLASS_BODY()
 
diff --git a/Engine/Source/Runtime/Engine/Private/Materials/MaterialExpressions.cpp b/Engine/Source/Runtime/Engine/Private/Materials/MaterialExpressions.cpp
--- a/Engine/Source/Runtime/Engine/Private/Materials/MaterialExpressions.cpp
+++ b/Engine/Source/Runtime/Engine/Private/Materials/MaterialExpressions.cpp
@@ -24322,7 +24322,10 @@ FStrataOperator* UMaterialExpressionStrataThinFilm::StrataGenerateMaterialTopolo
 }
 #endif // WITH_EDITOR
 
-
+UMaterialExpressionStrataUtilityBase::UMaterialExpressionStrataUtilityBase(const FObjectInitializer& ObjectInitializer)
+	: Super(ObjectInitializer)
+{
+}
 
 UMaterialExpressionStrataTransmittanceToMFP::UMaterialExpressionStrataTransmittanceToMFP(const FObjectInitializer& ObjectInitializer)
 	: Super(ObjectInitializer)
