Fix crash happening when exiting world

diff --git a/Engine/Source/Runtime/Engine/Private/WorldPartition/ContentBundle/ContentBundleWorldSubsystem.cpp b/Engine/Source/Runtime/Engine/Private/WorldPartition/ContentBundle/ContentBundleWorldSubsystem.cpp
--- a/Engine/Source/Runtime/Engine/Private/WorldPartition/ContentBundle/ContentBundleWorldSubsystem.cpp
+++ b/Engine/Source/Runtime/Engine/Private/WorldPartition/ContentBundle/ContentBundleWorldSubsystem.cpp
@@ -44,6 +44,7 @@ void UContentBundleManager::Deinitialize()
 	{
 		ContentBundleContainer->Deinitialize();
 	}
+	ContentBundleContainers.Empty();
 
 #if WITH_EDITOR
 	PIEDuplicateHelper = nullptr;
