World Partition - HLOD: Ensure that UHLODLayer::DuplicateHLODLayersSetup() is creating public & standalone assets

diff --git a/Engine/Source/Runtime/Engine/Private/WorldPartition/HLOD/HLODLayer.cpp b/Engine/Source/Runtime/Engine/Private/WorldPartition/HLOD/HLODLayer.cpp
--- a/Engine/Source/Runtime/Engine/Private/WorldPartition/HLOD/HLODLayer.cpp
+++ b/Engine/Source/Runtime/Engine/Private/WorldPartition/HLOD/HLODLayer.cpp
@@ -131,8 +131,11 @@ UHLODLayer* UHLODLayer::DuplicateHLODLayersSetup(UHLODLayer* HLODLayer, const FS
 		// In case Package already exists setting this flag will allow overwriting it
 		Package->MarkAsFullyLoaded();
 
-		FString NewHLODLayerName = Prefix + TEXT("_") + CurrentHLODLayer->GetName();
-		UHLODLayer* NewHLODLayer = CastChecked<UHLODLayer>(StaticDuplicateObject(CurrentHLODLayer, Package, *NewHLODLayerName));
+		FObjectDuplicationParameters ObjParameters(CurrentHLODLayer, Package);
+		ObjParameters.DestName = FName(Prefix + TEXT("_") + CurrentHLODLayer->GetName());
+		ObjParameters.ApplyFlags = RF_Public | RF_Standalone;
+
+		UHLODLayer* NewHLODLayer = CastChecked<UHLODLayer>(StaticDuplicateObjectEx(ObjParameters));
 		check(NewHLODLayer);
 
 		if (LastHLODLayer)
