World Partition Unit Tests: added actor descriptor serialization to the test suite, by loading all actor descriptors from the unit test map, resaving their descriptors and comparing against loaded one.

#rnx

diff --git a/Engine/Source/Runtime/Engine/Private/Tests/WorldPartition/WorldPartitionHandleTests.cpp b/Engine/Source/Runtime/Engine/Private/Tests/WorldPartition/WorldPartitionHandleTests.cpp
--- a/Engine/Source/Runtime/Engine/Private/Tests/WorldPartition/WorldPartitionHandleTests.cpp
+++ b/Engine/Source/Runtime/Engine/Private/Tests/WorldPartition/WorldPartitionHandleTests.cpp
@@ -260,8 +260,34 @@ namespace WorldPartitionTests
 	bool FWorldPartitionSoftRefTest::RunTest(const FString& Parameters)
 	{
 #if WITH_EDITOR
+		// Handle tests
 		WorldPartitionTests::PerformTests<FWorldPartitionLoadingContext::FImmediate>(this);
 		WorldPartitionTests::PerformTests<FWorldPartitionLoadingContext::FDeferred>(this);
+
+		// Serialization tests
+		UActorDescContainer* ActorDescContainer = NewObject<UActorDescContainer>(GetTransientPackage());
+		ActorDescContainer->Initialize({ nullptr, TEXT("/Engine/WorldPartition/WorldPartitionUnitTest") });
+
+		for (FActorDescList::TIterator<> ActorDescIterator(ActorDescContainer); ActorDescIterator; ++ActorDescIterator)
+		{
+			UClass* ActorNativeClass = UClass::TryFindTypeSlow<UClass>(ActorDescIterator->GetNativeClass().ToString(), EFindFirstObjectOptions::ExactClass);
+			TestFalse(TEXT("Actor Descriptor Serialization"), !ActorNativeClass);
+
+			if (ActorNativeClass)
+			{
+				TUniquePtr<FWorldPartitionActorDesc> NewActorDesc(AActor::StaticCreateClassActorDesc(ActorNativeClass));
+
+				FWorldPartitionActorDescInitData ActorDescInitData;
+				ActorDescInitData.NativeClass = ActorNativeClass;
+				ActorDescInitData.PackageName = ActorDescIterator->GetActorPackage();
+				ActorDescInitData.ActorPath = ActorDescIterator->GetActorPath();
+			
+				ActorDescIterator->SerializeTo(ActorDescInitData.SerializedData);
+				NewActorDesc->Init(ActorDescInitData);
+
+				TestTrue(TEXT("Actor Descriptor Serialization"), NewActorDesc->Equals(*ActorDescIterator));
+			}
+		}
 #endif
 		return true;
 	}
