Fix crash after trying to render job with not opened World Partition level as a target map in MRQ.

diff --git a/Engine/Source/Runtime/Engine/Private/WorldPartition/WorldPartition.cpp b/Engine/Source/Runtime/Engine/Private/WorldPartition/WorldPartition.cpp
--- a/Engine/Source/Runtime/Engine/Private/WorldPartition/WorldPartition.cpp
+++ b/Engine/Source/Runtime/Engine/Private/WorldPartition/WorldPartition.cpp
@@ -1096,15 +1096,13 @@ void UWorldPartition::Serialize(FArchive& Ar)
 	if (Ar.GetPortFlags() & PPF_DuplicateForPIE)
 	{
 		Ar << StreamingPolicy;
+
 #if WITH_EDITORONLY_DATA
 		Ar << GeneratedStreamingPackageNames;
 #endif
 
 #if WITH_EDITOR
-		if (Ar.IsLoading())
-		{
-			bIsPIE = true;
-		}
+		Ar << bIsPIE;
 #endif
 	}
 	else if (Ar.CustomVer(FUE5MainStreamObjectVersion::GUID) >= FUE5MainStreamObjectVersion::WorldPartitionSerializeStreamingPolicyOnCook)
