Fix deprecation warnings caused by merge.

diff --git a/Engine/Source/Runtime/Engine/Private/Tests/WorldPartition/WorldPartitionHandleTests.cpp b/Engine/Source/Runtime/Engine/Private/Tests/WorldPartition/WorldPartitionHandleTests.cpp
--- a/Engine/Source/Runtime/Engine/Private/Tests/WorldPartition/WorldPartitionHandleTests.cpp
+++ b/Engine/Source/Runtime/Engine/Private/Tests/WorldPartition/WorldPartitionHandleTests.cpp
@@ -280,7 +280,7 @@ namespace WorldPartitionTests
 				FWorldPartitionActorDescInitData ActorDescInitData;
 				ActorDescInitData.NativeClass = ActorNativeClass;
 				ActorDescInitData.PackageName = ActorDescIterator->GetActorPackage();
-				ActorDescInitData.ActorPath = ActorDescIterator->GetActorPath();
+				ActorDescInitData.ActorPath = ActorDescIterator->GetActorSoftPath();
 			
 				ActorDescIterator->SerializeTo(ActorDescInitData.SerializedData);
 				NewActorDesc->Init(ActorDescInitData);
