Remove dynamic MatchingDPStatus_ key from engine data crash context.

diff --git a/Engine/Source/Runtime/Engine/Private/DeviceProfiles/DeviceProfileMatching.cpp b/Engine/Source/Runtime/Engine/Private/DeviceProfiles/DeviceProfileMatching.cpp
--- a/Engine/Source/Runtime/Engine/Private/DeviceProfiles/DeviceProfileMatching.cpp
+++ b/Engine/Source/Runtime/Engine/Private/DeviceProfiles/DeviceProfileMatching.cpp
@@ -628,7 +628,7 @@ static FString LoadAndProcessMatchingRulesFromConfig(const FString& ParentDP, ID
 	UE_CLOG(DPMatchingErrorOutput.NumErrors > 0, LogInit, Fatal, TEXT("DeviceProfileMatching: %d Error(s) encountered while processing MatchedRules for %s"), DPMatchingErrorOutput.NumErrors, *ParentDP);
 #endif
 
-	FGenericCrashContext::SetEngineData(TEXT("MatchingDPStatus_")+ParentDP, DPMatchingErrorOutput.NumErrors > 0 ? TEXT("Error") : TEXT("No errors"));
+	FGenericCrashContext::SetEngineData(TEXT("MatchingDPStatus"), ParentDP + (DPMatchingErrorOutput.NumErrors > 0 ? TEXT("Error") : TEXT("No errors")));
 
 	return SelectedFragments;
 }
@@ -671,7 +671,7 @@ TArray<FSelectedFragmentProperties> UDeviceProfileManager::FindMatchingFragments
 #endif
 		IDeviceProfileSelectorModule* DPSelector = PreviewDPSelector ? PreviewDPSelector : GetDeviceProfileSelectorModule();
 		SelectedFragments = LoadAndProcessMatchingRulesFromConfig(ParentDP, DPSelector, ConfigSystem);
-	
+		
 	}
 	SelectedFragments = RemoveAllWhiteSpace(SelectedFragments);
 
