Optimus: Recreate deformer bindings on skinned mesh BeginPlay().
This deals with case where binding components can be deleted during blueprint re-init.

diff --git a/Engine/Source/Runtime/Engine/Classes/Components/SkinnedMeshComponent.h b/Engine/Source/Runtime/Engine/Classes/Components/SkinnedMeshComponent.h
--- a/Engine/Source/Runtime/Engine/Classes/Components/SkinnedMeshComponent.h
+++ b/Engine/Source/Runtime/Engine/Classes/Components/SkinnedMeshComponent.h
@@ -1024,6 +1024,7 @@ protected:
 	//~ Begin UActorComponent Interface
 	virtual void OnRegister() override;
 	virtual void OnUnregister() override;
+	virtual void BeginPlay() override;
 	virtual void CreateRenderState_Concurrent(FRegisterComponentContext* Context) override;
 	virtual void SendRenderDynamicData_Concurrent() override;
 	virtual void DestroyRenderState_Concurrent() override;
diff --git a/Engine/Source/Runtime/Engine/Private/Components/SkinnedMeshComponent.cpp b/Engine/Source/Runtime/Engine/Private/Components/SkinnedMeshComponent.cpp
--- a/Engine/Source/Runtime/Engine/Private/Components/SkinnedMeshComponent.cpp
+++ b/Engine/Source/Runtime/Engine/Private/Components/SkinnedMeshComponent.cpp
@@ -705,6 +705,15 @@ void USkinnedMeshComponent::OnUnregister()
 	MeshDeformerInstance = nullptr;
 }
 
+void USkinnedMeshComponent::BeginPlay()
+{
+	Super::BeginPlay();
+
+	// Recreate the deformer instance to collect any changes for component bindings during startup.
+	UMeshDeformer* ActiveMeshDeformer = GetActiveMeshDeformer();
+	MeshDeformerInstance = ActiveMeshDeformer != nullptr ? ActiveMeshDeformer->CreateInstance(this, MeshDeformerInstanceSettings) : nullptr;
+}
+
 void USkinnedMeshComponent::AddExternalMorphSet(int32 LOD, int32 ID, TSharedPtr<FExternalMorphSet> MorphSet)
 {
 	const int32 NumLODs = GetNumLODs();
