Iris - SetPIEInstanceID in Replication system to match the NetDriver's world

diff --git a/Engine/Source/Runtime/Engine/Private/NetDriver.cpp b/Engine/Source/Runtime/Engine/Private/NetDriver.cpp
--- a/Engine/Source/Runtime/Engine/Private/NetDriver.cpp
+++ b/Engine/Source/Runtime/Engine/Private/NetDriver.cpp
@@ -614,14 +614,24 @@ void UNetDriver::LoadChannelDefinitions()
 
 void UNetDriver::NotifyGameInstanceUpdated()
 {
-#if UE_NET_TRACE_ENABLED
 	const UWorld* LocalWorld = GetWorld();
+
+#if UE_NET_TRACE_ENABLED
 	if (GetNetTraceId() != NetTraceInvalidGameInstanceId && LocalWorld && LocalWorld->WorldType)
 	{
 		FString InstanceName = FString::Printf(TEXT("%s (%s)"), *GetNameSafe(LocalWorld->GetGameInstance()), LexToString(LocalWorld->WorldType.GetValue()));
 		UE_NET_TRACE_UPDATE_INSTANCE(GetNetTraceId(), IsServer(), *InstanceName);
 	}
 #endif
+
+#if WITH_EDITOR
+#if UE_WITH_IRIS
+	if (ReplicationSystem && LocalWorld)
+	{
+		ReplicationSystem->SetPIEInstanceID(World->GetPackage()->GetPIEInstanceID());
+	}
+#endif
+#endif
 }
 
 void UNetDriver::AssertValid()
@@ -5773,7 +5783,6 @@ void UNetDriver::SetWorld(class UWorld* InWorld)
 		Notify = InWorld;
 		RegisterTickEvents(InWorld);
 
-
 #if UE_WITH_IRIS
 		if (IsUsingIrisReplication())
 		{
