Check if the HLOD subsystem is valid before trying to use it when using the console command wp.Runtime.HLOD

diff --git a/Engine/Source/Runtime/Engine/Private/WorldPartition/HLOD/HLODSubsystem.cpp b/Engine/Source/Runtime/Engine/Private/WorldPartition/HLOD/HLODSubsystem.cpp
--- a/Engine/Source/Runtime/Engine/Private/WorldPartition/HLOD/HLODSubsystem.cpp
+++ b/Engine/Source/Runtime/Engine/Private/WorldPartition/HLOD/HLODSubsystem.cpp
@@ -100,20 +100,21 @@ FAutoConsoleCommand UHLODSubsystem::EnableHLODCommand(
 			UWorld* World = Context.World();
 			if (World && World->IsGameWorld())
 			{
-				UHLODSubsystem* HLODSubSystem = World->GetSubsystem<UHLODSubsystem>();
-				for (const auto& KeyValuePair : HLODSubSystem->WorldPartitionsHLODRuntimeData)
+				if (UHLODSubsystem* HLODSubSystem = World->GetSubsystem<UHLODSubsystem>()) 
 				{
-					for (const auto& CellHLODMapping : KeyValuePair.Value.CellsData)
+					for (const auto& KeyValuePair : HLODSubSystem->WorldPartitionsHLODRuntimeData)
 					{
-						const FCellData& CellData = CellHLODMapping.Value;
-						bool bIsHLODVisible = UHLODSubsystem::WorldPartitionHLODEnabled && !CellData.bIsCellVisible;
-						for (AWorldPartitionHLOD* HLODActor : CellData.LoadedHLODs)
+						for (const auto& CellHLODMapping : KeyValuePair.Value.CellsData)
 						{
-							HLODActor->SetVisibility(bIsHLODVisible);
+							const FCellData& CellData = CellHLODMapping.Value;
+							bool bIsHLODVisible = UHLODSubsystem::WorldPartitionHLODEnabled && !CellData.bIsCellVisible;
+							for (AWorldPartitionHLOD* HLODActor : CellData.LoadedHLODs)
+							{
+								HLODActor->SetVisibility(bIsHLODVisible);
+							}
 						}
 					}
 				}
-				
 			}
 		}
 	})
