Modifying animation sequence duration in montage causes decision window connections to break.
#fix deprecated LinkMontage/LinkAnimSequence API in favour of ::Link which handles the types internally

diff --git a/Engine/Source/Editor/Persona/Private/AnimTimeline/AnimModel_AnimMontage.cpp b/Engine/Source/Editor/Persona/Private/AnimTimeline/AnimModel_AnimMontage.cpp
--- a/Engine/Source/Editor/Persona/Private/AnimTimeline/AnimModel_AnimMontage.cpp
+++ b/Engine/Source/Editor/Persona/Private/AnimTimeline/AnimModel_AnimMontage.cpp
@@ -135,7 +135,7 @@ void FAnimModel_AnimMontage::OnSetEditableTime(int32 TimeIndex, double Time, boo
 	
 			FCompositeSection& Section = AnimMontage->CompositeSections[TimeIndex];
 			Section.SetTime(Time);
-			Section.LinkMontage(AnimMontage, Time);
+			Section.Link(AnimMontage, Time);
 
 			SortSections();
 			RefreshNotifyTriggerOffsets();
diff --git a/Engine/Source/Editor/Persona/Private/SAnimMontagePanel.cpp b/Engine/Source/Editor/Persona/Private/SAnimMontagePanel.cpp
--- a/Engine/Source/Editor/Persona/Private/SAnimMontagePanel.cpp
+++ b/Engine/Source/Editor/Persona/Private/SAnimMontagePanel.cpp
@@ -1219,7 +1219,7 @@ void SAnimMontagePanel::SetSectionTime(int32 SectionIndex, float NewTime)
 	
 		FCompositeSection& Section = Montage->CompositeSections[SectionIndex];
 		Section.SetTime(NewTime);
-		Section.LinkMontage(Montage, NewTime);
+		Section.Link(Montage, NewTime);
 
 		SortAndUpdateMontage();
 	}
diff --git a/Engine/Source/Runtime/Engine/Classes/Animation/AnimLinkableElement.h b/Engine/Source/Runtime/Engine/Classes/Animation/AnimLinkableElement.h
--- a/Engine/Source/Runtime/Engine/Classes/Animation/AnimLinkableElement.h
+++ b/Engine/Source/Runtime/Engine/Classes/Animation/AnimLinkableElement.h
@@ -61,23 +61,25 @@ struct FAnimLinkableElement
 	ENGINE_API void Update();
 
 	/** Link this element to an animation object (Sequence or Montage)
-	 *	@param AnimObject The object to link to
+	 *	@param AnimSequenceBase The object to link to
 	 *	@param AbsTime The absolute time to place this element
 	 *	@param InSlotIndex Slot index for montages (ignored otherwise)
 	 */
-	ENGINE_API void Link(UAnimSequenceBase* AnimObject, float AbsTime, int32 InSlotIndex = 0);
+	ENGINE_API void Link(UAnimSequenceBase* AnimSequenceBase, float AbsTime, int32 InSlotIndex = 0);
 
 	/** Link this element to a montage 
 	 * @param Montage The montage to link to
 	 * @param AbsMontageTime The time in the montage that this element should be placed at
 	 * @param InSlotIndex The slot in the montage to detect segments in
 	 */
+	UE_DEPRECATED(5.1, "LinkMontage has been deprecated, use Link instead")
 	ENGINE_API void LinkMontage(UAnimMontage* Montage, float AbsMontageTime, int32 InSlotIndex = 0);
 
 	/** Link this element to a Sequence, Just setting basic data as sequences don't need full linking 
 	 * @param Sequence The sequence to link to
 	 * @param AbsSequenceTime The time in the sequence that this element should be placed at
 	 */
+	UE_DEPRECATED(5.1, "LinkSequence has been deprecated, use Link instead")
 	ENGINE_API void LinkSequence(UAnimSequenceBase* Sequence, float AbsSequenceTime);
 
 	/** Clear the linking information from this element, leaves montage link intact */
diff --git a/Engine/Source/Runtime/Engine/Private/Animation/AnimLinkableElement.cpp b/Engine/Source/Runtime/Engine/Private/Animation/AnimLinkableElement.cpp
--- a/Engine/Source/Runtime/Engine/Private/Animation/AnimLinkableElement.cpp
+++ b/Engine/Source/Runtime/Engine/Private/Animation/AnimLinkableElement.cpp
@@ -7,51 +7,12 @@
 
 void FAnimLinkableElement::LinkMontage(UAnimMontage* Montage, float AbsMontageTime, int32 InSlotIndex)
 {
-	if(Montage && Montage->SlotAnimTracks.Num() > 0)
-	{
-		LinkedMontage = Montage;
-
-		SlotIndex = InSlotIndex;
-		FSlotAnimationTrack& Slot = Montage->SlotAnimTracks[SlotIndex];
-
-		SegmentIndex = Slot.AnimTrack.GetSegmentIndexAtTime(AbsMontageTime);
-		if(SegmentIndex != INDEX_NONE)
-		{
-			FAnimSegment& Segment = Slot.AnimTrack.AnimSegments[SegmentIndex];
-			LinkedSequence = Segment.GetAnimReference();
-			SegmentBeginTime = Segment.StartPos;
-			SegmentLength = Segment.GetLength();
-
-			SetTime_Internal(AbsMontageTime);
-		}
-		else
-		{
-			// Nothing to link to
-			// We have no segment to link to, we need to clear our the segment data
-			// and give ourselves an absolute time
-			LinkValue = AbsMontageTime;
-			LinkedSequence = nullptr;
-			SegmentBeginTime = -1.0f;
-			SegmentLength = -1.0f;
-			LinkMethod = EAnimLinkMethod::Absolute;
-			CachedLinkMethod = LinkMethod;
-		}
-	}
+	Link(Montage, AbsMontageTime, InSlotIndex);
 }
 
 void FAnimLinkableElement::LinkSequence(UAnimSequenceBase* Sequence, float AbsSequenceTime)
 {
-	if(Sequence && Sequence->GetPlayLength() > 0)
-	{
-		LinkedMontage = nullptr;
-		LinkedSequence = Sequence;
-		SegmentIndex = 0;
-
-		SegmentBeginTime = 0.0f;
-		SegmentLength = Sequence->GetPlayLength();
-
-		SetTime(AbsSequenceTime);
-	}
+	Link(Sequence, AbsSequenceTime);
 }
 
 void FAnimLinkableElement::Clear()
@@ -332,7 +293,7 @@ void FAnimLinkableElement::ChangeSlotIndex(int32 NewSlotIndex)
 {
 	if(LinkedMontage)
 	{
-		LinkMontage(LinkedMontage, GetTime(), NewSlotIndex);
+		Link(LinkedMontage, GetTime(), NewSlotIndex);
 	}
 }
 
@@ -390,26 +351,62 @@ bool FAnimLinkableElement::ConditionalRelink()
 	{
 		if(LinkedMontage)
 		{
-			LinkMontage(LinkedMontage, CurrentAbsTime, SlotIndex);
+			Link(LinkedMontage, CurrentAbsTime, SlotIndex);
 		}
 		else if(LinkedSequence)
 		{
-			LinkSequence(LinkedSequence, CurrentAbsTime);
+			Link(LinkedSequence, CurrentAbsTime);
 		}
 	}
 
 	return bRequiresRelink;
 }
 
-void FAnimLinkableElement::Link(UAnimSequenceBase* AnimObject, float AbsTime, int32 InSlotIndex /*= 0*/)
+void FAnimLinkableElement::Link(UAnimSequenceBase* AnimSequenceBase, float AbsTime, int32 InSlotIndex /*= 0*/)
 {
-	if(UAnimMontage* Montage = Cast<UAnimMontage>(AnimObject))
+	if(UAnimMontage* Montage = Cast<UAnimMontage>(AnimSequenceBase))
 	{
-		LinkMontage(Montage, AbsTime, InSlotIndex);
+		if(Montage->SlotAnimTracks.Num() > 0)
+		{
+			LinkedMontage = Montage;
+
+			SlotIndex = InSlotIndex;
+			FSlotAnimationTrack& Slot = Montage->SlotAnimTracks[SlotIndex];
+
+			SegmentIndex = Slot.AnimTrack.GetSegmentIndexAtTime(AbsTime);
+			if(SegmentIndex != INDEX_NONE)
+			{
+				const FAnimSegment& Segment = Slot.AnimTrack.AnimSegments[SegmentIndex];
+				LinkedSequence = Segment.GetAnimReference();
+				SegmentBeginTime = Segment.StartPos;
+				SegmentLength = Segment.GetLength();
+
+				SetTime_Internal(AbsTime);
+			}
+			else
+			{
+				// Nothing to link to
+				// We have no segment to link to, we need to clear our the segment data
+				// and give ourselves an absolute time
+				LinkValue = AbsTime;
+				LinkedSequence = nullptr;
+				SegmentBeginTime = -1.0f;
+				SegmentLength = -1.0f;
+				LinkMethod = EAnimLinkMethod::Absolute;
+				CachedLinkMethod = LinkMethod;
+			}
+		}
 	}
-	else
-	{
-		LinkSequence(AnimObject, AbsTime);
+	else if (AnimSequenceBase && AnimSequenceBase->GetPlayLength() > 0)
+	{		
+		LinkedMontage = nullptr;
+		LinkedSequence = AnimSequenceBase;
+		SegmentIndex = 0;
+
+		SegmentBeginTime = 0.0f;
+		SegmentLength = AnimSequenceBase->GetPlayLength();
+
+		SetTime(AbsTime);
 	}
 }
 
diff --git a/Engine/Source/Runtime/Engine/Private/Animation/AnimMontage.cpp b/Engine/Source/Runtime/Engine/Private/Animation/AnimMontage.cpp
--- a/Engine/Source/Runtime/Engine/Private/Animation/AnimMontage.cpp
+++ b/Engine/Source/Runtime/Engine/Private/Animation/AnimMontage.cpp
@@ -304,7 +304,7 @@ int32 UAnimMontage::AddAnimCompositeSection(FName InSectionName, float StartTime
 		return INDEX_NONE;
 	}
 
-	NewSection.LinkMontage(this, StartTime);
+	NewSection.Link(this, StartTime);
 
 	// we'd like to sort them in the order of time
 	int32 NewSectionIndex = CompositeSections.Add(NewSection);
@@ -405,12 +405,12 @@ void UAnimMontage::PostLoad()
 		if(Composite.StartTime_DEPRECATED != 0.0f)
 		{
 			Composite.Clear();
-			Composite.LinkMontage(this, Composite.StartTime_DEPRECATED);
+			Composite.Link(this, Composite.StartTime_DEPRECATED);
 		}
 		else
 		{
 			Composite.RefreshSegmentOnLoad();
-			Composite.LinkMontage(this, Composite.GetTime());
+			Composite.Link(this, Composite.GetTime());
 		}
 	}
 
@@ -487,16 +487,16 @@ void UAnimMontage::PostLoad()
 		if(Notify.DisplayTime_DEPRECATED != 0.0f)
 		{
 			Notify.Clear();
-			Notify.LinkMontage(this, Notify.DisplayTime_DEPRECATED);
+			Notify.Link(this, Notify.DisplayTime_DEPRECATED);
 		}
 		else
 		{
-			Notify.LinkMontage(this, Notify.GetTime());
+			Notify.Link(this, Notify.GetTime());
 		}
 
 		if(Notify.Duration != 0.0f)
 		{
-			Notify.EndLink.LinkMontage(this, Notify.GetTime() + Notify.Duration);
+			Notify.EndLink.Link(this, Notify.GetTime() + Notify.Duration);
 		}
 	}
 
@@ -533,11 +533,11 @@ void UAnimMontage::ConvertBranchingPointsToAnimNotifies()
 			if (BranchingPoint.DisplayTime_DEPRECATED != 0.0f)
 			{
 				BranchingPoint.Clear();
-				BranchingPoint.LinkMontage(this, BranchingPoint.DisplayTime_DEPRECATED);
+				BranchingPoint.Link(this, BranchingPoint.DisplayTime_DEPRECATED);
 			}
 			else
 			{
-				BranchingPoint.LinkMontage(this, BranchingPoint.GetTime());
+				BranchingPoint.Link(this, BranchingPoint.GetTime());
 			}
 		}
 
@@ -561,7 +561,7 @@ void UAnimMontage::ConvertBranchingPointsToAnimNotifies()
 			NewEvent.NotifyName = BranchingPoint.EventName;
 
 			float TriggerTime = BranchingPoint.GetTriggerTime();
-			NewEvent.LinkMontage(this, TriggerTime);
+			NewEvent.Link(this, TriggerTime);
 #if WITH_EDITOR
 			NewEvent.TriggerTimeOffset = GetTriggerTimeOffsetForType(CalculateOffsetForNotify(TriggerTime));
 #endif
@@ -2889,7 +2889,7 @@ UAnimMontage* UAnimMontage::CreateSlotAnimationAsDynamicMontage_WithBlendSetting
 
 	FCompositeSection NewSection;
 	NewSection.SectionName = TEXT("Default");
-	NewSection.LinkSequence(Asset, Asset->GetPlayLength());
+	NewSection.Link(Asset, Asset->GetPlayLength());
 	NewSection.SetTime(0.0f);
 
 	// add new section
diff --git a/Engine/Source/Runtime/Engine/Private/Animation/AnimSequence.cpp b/Engine/Source/Runtime/Engine/Private/Animation/AnimSequence.cpp
--- a/Engine/Source/Runtime/Engine/Private/Animation/AnimSequence.cpp
+++ b/Engine/Source/Runtime/Engine/Private/Animation/AnimSequence.cpp
@@ -960,16 +960,16 @@ void UAnimSequence::PostLoad()
 		if(Notify.DisplayTime_DEPRECATED != 0.0f)
 		{
 			Notify.Clear();
-			Notify.LinkSequence(this, Notify.DisplayTime_DEPRECATED);
+			Notify.Link(this, Notify.DisplayTime_DEPRECATED);
 		}
 		else
 		{
-			Notify.LinkSequence(this, Notify.GetTime());
+			Notify.Link(this, Notify.GetTime());
 		}
 	
 		if(Notify.Duration != 0.0f)
 		{
-			Notify.EndLink.LinkSequence(this, Notify.GetTime() + Notify.Duration);
+			Notify.EndLink.Link(this, Notify.GetTime() + Notify.Duration);
 		}
 	}
 
diff --git a/Engine/Source/Runtime/Engine/Private/Animation/AnimSequenceBase.cpp b/Engine/Source/Runtime/Engine/Private/Animation/AnimSequenceBase.cpp
--- a/Engine/Source/Runtime/Engine/Private/Animation/AnimSequenceBase.cpp
+++ b/Engine/Source/Runtime/Engine/Private/Animation/AnimSequenceBase.cpp
@@ -1124,8 +1124,7 @@ void UAnimSequenceBase::OnModelModified(const EAnimDataModelNotifyType& NotifyTy
 				// Clamps against the sequence length, ensuring that the notify does not end up outside of the playback bounds
 				const float ClampedCurrentTime = FMath::Clamp(CurrentTime, 0.f, NewLength);
 
-				Notify.LinkSequence(this, ClampedCurrentTime);
-
+				Notify.Link(this, ClampedCurrentTime);
 				Notify.SetDuration(NewDuration);
 
 				if (ClampedCurrentTime == 0.f)
@@ -1185,7 +1184,7 @@ void UAnimSequenceBase::OnModelModified(const EAnimDataModelNotifyType& NotifyTy
 
 				const float ClampedCurrentTime = FMath::Clamp(CurrentTime, 0.f, NewLength);
 
-				Notify.LinkSequence(this, ClampedCurrentTime);
+				Notify.Link(this, ClampedCurrentTime);
 				Notify.SetDuration(NewDuration);
 
 				if (ClampedCurrentTime == 0.f)
diff --git a/Engine/Source/Runtime/Engine/Private/Animation/AnimSequenceHelpers.cpp b/Engine/Source/Runtime/Engine/Private/Animation/AnimSequenceHelpers.cpp
--- a/Engine/Source/Runtime/Engine/Private/Animation/AnimSequenceHelpers.cpp
+++ b/Engine/Source/Runtime/Engine/Private/Animation/AnimSequenceHelpers.cpp
@@ -363,7 +363,7 @@ bool CopyNotifies(const UAnimSequenceBase* SourceAnimSeq, UAnimSequenceBase* Des
 		}
 
 		// Copy notify timing
-		NotifyEvent.LinkSequence(DestAnimSeq, SrcNotifyEvent.GetTriggerTime());
+		NotifyEvent.Link(DestAnimSeq, SrcNotifyEvent.GetTriggerTime());
 		NotifyEvent.TriggerTimeOffset = GetTriggerTimeOffsetForType(DestAnimSeq->CalculateOffsetForNotify(NotifyEvent.GetTriggerTime()));
 
 		// Make sure editor knows we've changed something.
