Add ability to choose the framepro and nprof file names when starting a recording

diff --git a/Engine/Source/Runtime/Core/Private/FramePro/FrameProProfiler.cpp b/Engine/Source/Runtime/Core/Private/FramePro/FrameProProfiler.cpp
--- a/Engine/Source/Runtime/Core/Private/FramePro/FrameProProfiler.cpp
+++ b/Engine/Source/Runtime/Core/Private/FramePro/FrameProProfiler.cpp
@@ -283,6 +283,7 @@ void FFrameProProfiler::Initialize()
 
 void FFrameProProfiler::TearDown()
 {
+	GFrameProEnabled = false;
 	FRAMEPRO_SHUTDOWN();
 }
 
@@ -401,10 +402,17 @@ void FFrameProProfiler::StartFrameProRecordingFromCommand(const TArray< FString
 		FilenameRoot = Args[0];
 	}
 
-	StartFrameProRecording(FilenameRoot, ScopeMinTimeMicroseconds);
+	bool bAppendDateTime = true;
+	if (Args.Num() > 1)
+	{
+		// If someone wants to use the full filename they provided, they need to set the 2nd arg to false.
+		LexTryParseString(bAppendDateTime, *Args[1]);
+	}
+
+	StartFrameProRecording(FilenameRoot, ScopeMinTimeMicroseconds, bAppendDateTime);
 }
 
-FString FFrameProProfiler::StartFrameProRecording(const FString& FilenameRoot, int32 MinScopeTime)
+FString FFrameProProfiler::StartFrameProRecording(const FString& FilenameRoot, int32 MinScopeTime, bool bAppendDateTime)
 {
 	if (GFrameProIsRecording)
 	{
@@ -414,7 +422,15 @@ FString FFrameProProfiler::StartFrameProRecording(const FString& FilenameRoot, i
 	FString RelPathName = FPaths::ProfilingDir() + TEXT("FramePro/");
 	bool bSuccess = IFileManager::Get().MakeDirectory(*RelPathName, true); // ensure folder exists
 
-	FString Filename = FString::Printf(TEXT("%s(%s).framepro_recording"), *FilenameRoot, *FDateTime::Now().ToString(TEXT("%Y%m%d_%H%M%S")));
+	FString Filename;
+	if (bAppendDateTime)
+	{
+		Filename = FString::Printf(TEXT("%s(%s).framepro_recording"), *FilenameRoot, *FDateTime::Now().ToString(TEXT("%Y%m%d_%H%M%S")));
+	}
+	else
+	{
+		Filename = FilenameRoot;
+	}
 	FString OutputFilename = RelPathName + Filename;
 
 	UE_LOG(LogFramePro, Log, TEXT("--- Start Recording To File: %s"), *OutputFilename);
diff --git a/Engine/Source/Runtime/Core/Public/FramePro/FrameProProfiler.h b/Engine/Source/Runtime/Core/Public/FramePro/FrameProProfiler.h
--- a/Engine/Source/Runtime/Core/Public/FramePro/FrameProProfiler.h
+++ b/Engine/Source/Runtime/Core/Public/FramePro/FrameProProfiler.h
@@ -27,7 +27,7 @@ public:
 	static void PopEvent(const ANSICHAR* Override);
 
 	static void StartFrameProRecordingFromCommand(const TArray< FString >& Args);
-	static FString StartFrameProRecording(const FString& FilenameRoot, int32 MinScopeTime);
+	static FString StartFrameProRecording(const FString& FilenameRoot, int32 MinScopeTime, bool bAppendDateTime=true);
 	static void StopFrameProRecording();
 
 	static bool IsFrameProRecording();
diff --git a/Engine/Source/Runtime/Engine/Private/NetworkProfiler.cpp b/Engine/Source/Runtime/Engine/Private/NetworkProfiler.cpp
--- a/Engine/Source/Runtime/Engine/Private/NetworkProfiler.cpp
+++ b/Engine/Source/Runtime/Engine/Private/NetworkProfiler.cpp
@@ -712,7 +712,18 @@ void FNetworkProfiler::TrackSessionChange( bool bShouldContinueTracking, const F
 
 			static int32 Salt = 0;
 			Salt++;		// Use a salt to solve the issue where this function is called so fast it produces the same time (seems to happen during seamless travel)
-			const FString FinalFileName = FPaths::ProfilingDir() + FApp::GetProjectName() + FString::Printf(TEXT("-Pid%i"), FPlatformProcess::GetCurrentProcessId()) + TEXT( "-" ) + FDateTime::Now().ToString() + FString::Printf( TEXT( "[%i]" ), Salt ) + TEXT( ".nprof" );
+
+			FString FileName;
+
+			if (NextFileName.IsEmpty())
+			{
+				FileName = FApp::GetProjectName() + FString::Printf(TEXT("-Pid%i"), FPlatformProcess::GetCurrentProcessId()) + TEXT("-") + FDateTime::Now().ToString() + FString::Printf(TEXT("[%i]"), Salt) + TEXT(".nprof");
+			}
+			else
+			{
+				FileName = MoveTemp(NextFileName);
+			}
+			const FString FinalFileName = FPaths::ProfilingDir() + FileName;
 
 			IFileManager::Get().MakeDirectory( *FPaths::GetPath( FinalFileName ) );
 			FileWriter = IFileManager::Get().CreateFileWriter( *FinalFileName, FILEWRITE_EvenIfReadOnly );
@@ -833,6 +844,11 @@ void FNetworkProfiler::TrackWritePropertyHandle( uint16 NumBits, UNetConnection*
 	}
 }
 
+void FNetworkProfiler::SetNextFileName(const FString& FileName)
+{
+	NextFileName = FileName;
+}
+
 /**
  * Processes any network profiler specific exec commands
  *
@@ -846,11 +862,11 @@ bool FNetworkProfiler::Exec( UWorld * InWorld, const TCHAR* Cmd, FOutputDevice &
 {
 	if (FParse::Command(&Cmd, TEXT("ENABLE")))
 	{
-		EnableTracking( true );
+		EnableTracking(true);
 	}
 	else if (FParse::Command(&Cmd, TEXT("DISABLE")))
 	{
-		EnableTracking( false );
+		EnableTracking(false);
 	} 
 	else if (FParse::Command(&Cmd, TEXT("AUTOSTOP")))
 	{
diff --git a/Engine/Source/Runtime/Engine/Public/Net/NetworkProfiler.h b/Engine/Source/Runtime/Engine/Public/Net/NetworkProfiler.h
--- a/Engine/Source/Runtime/Engine/Public/Net/NetworkProfiler.h
+++ b/Engine/Source/Runtime/Engine/Public/Net/NetworkProfiler.h
@@ -196,6 +196,9 @@ private:
 
 	void AutoStopTracking();
 
+	/** The filename to use next time a file writer is opened.  If empty then an autogenerated name will be used. */
+	FString NextFileName;
+
 public:
 	/**
 	 * Constructor, initializing members.
@@ -209,6 +212,9 @@ public:
 	 */
 	void EnableTracking( bool bShouldEnableTracking );
 
+	/** Force the next profiling file to use the given filename. If none are set an autogenerated name will be used instead */
+	void SetNextFileName(const FString& FileName);
+
 	/**
 	 * Marks the beginning of a frame.
 	 */
