Assert if PropertyChangedTracker if not valid

diff --git a/Engine/Source/Runtime/Engine/Private/Actor.cpp b/Engine/Source/Runtime/Engine/Private/Actor.cpp
--- a/Engine/Source/Runtime/Engine/Private/Actor.cpp
+++ b/Engine/Source/Runtime/Engine/Private/Actor.cpp
@@ -1486,7 +1486,7 @@ void AActor::CallPreReplication(UNetDriver* NetDriver)
 	const bool bPreReplication = ShouldCallPreReplication();
 	const bool bPreReplicationForReplay = ShouldCallPreReplicationForReplay();
 
-	IRepChangedPropertyTracker* ActorChangedPropertyTracker = nullptr;
+	TSharedPtr<FRepChangedPropertyTracker> ActorChangedPropertyTracker;
 
 	if (bPreReplication)
 	{
@@ -1497,12 +1497,8 @@ void AActor::CallPreReplication(UNetDriver* NetDriver)
 		// In that case we call PreReplication on the locally controlled Character as well.
 		if ((LocalRole == ROLE_Authority) || ((LocalRole == ROLE_AutonomousProxy) && World && World->IsRecordingClientReplay()))
 		{
-			if (ActorChangedPropertyTracker == nullptr)
-			{
-				ActorChangedPropertyTracker = NetDriver->FindOrCreateRepChangedPropertyTracker(this).Get();
-			}
-
-			PreReplication(*ActorChangedPropertyTracker);
+			ActorChangedPropertyTracker = NetDriver->FindOrCreateRepChangedPropertyTracker(this);
+			PreReplication(*(ActorChangedPropertyTracker.Get()));
 		}
 	}
 
@@ -1511,12 +1507,12 @@ void AActor::CallPreReplication(UNetDriver* NetDriver)
 		// If we're recording a replay, call this for everyone (includes SimulatedProxies).
 		if (Cast<UDemoNetDriver>(NetDriver) || NetDriver->HasReplayConnection())
 		{
-			if (ActorChangedPropertyTracker == nullptr)
+			if (!ActorChangedPropertyTracker.IsValid())
 			{
-				ActorChangedPropertyTracker = NetDriver->FindOrCreateRepChangedPropertyTracker(this).Get();
+				ActorChangedPropertyTracker = NetDriver->FindOrCreateRepChangedPropertyTracker(this);
 			}
 
-			PreReplicationForReplay(*ActorChangedPropertyTracker);
+			PreReplicationForReplay(*(ActorChangedPropertyTracker.Get()));
 		}
 	}
 
@@ -1528,7 +1524,8 @@ void AActor::CallPreReplication(UNetDriver* NetDriver)
 			// Only call on components that aren't pending kill
 			if (IsValid(Component))
 			{
-				Component->PreReplication(*NetDriver->FindOrCreateRepChangedPropertyTracker(Component).Get());
+				TSharedPtr<FRepChangedPropertyTracker> ComponentChangedPropertyTracker = NetDriver->FindOrCreateRepChangedPropertyTracker(Component);
+				Component->PreReplication(*(ComponentChangedPropertyTracker.Get()));
 			}
 		}
 	}
diff --git a/Engine/Source/Runtime/Engine/Private/NetDriver.cpp b/Engine/Source/Runtime/Engine/Private/NetDriver.cpp
--- a/Engine/Source/Runtime/Engine/Private/NetDriver.cpp
+++ b/Engine/Source/Runtime/Engine/Private/NetDriver.cpp
@@ -6225,7 +6225,10 @@ TSharedPtr<FRepChangedPropertyTracker> UNetDriver::FindOrCreateRepChangedPropert
 {
 	check(IsServer() || MaySendProperties());
 
-	return UE::Net::Private::FNetPropertyConditionManager::Get().FindOrCreatePropertyTracker(Obj);
+	const FObjectKey ObjectKey = Obj;
+	checkf(ObjectKey.ResolveObjectPtr() != nullptr, TEXT("Unresolveable object %s received in FindOrCreateRepChangedPropertyTracker"), *GetNameSafe(Obj));
+
+	return UE::Net::Private::FNetPropertyConditionManager::Get().FindOrCreatePropertyTracker(ObjectKey);
 }
 
 TSharedPtr<FRepChangedPropertyTracker> UNetDriver::FindRepChangedPropertyTracker(UObject* Obj)
