Network Physics under a cvar

diff --git a/Engine/Source/Runtime/Engine/Classes/GameFramework/PlayerController.h b/Engine/Source/Runtime/Engine/Classes/GameFramework/PlayerController.h
--- a/Engine/Source/Runtime/Engine/Classes/GameFramework/PlayerController.h
+++ b/Engine/Source/Runtime/Engine/Classes/GameFramework/PlayerController.h
@@ -212,6 +212,16 @@ protected:
 	virtual void ApplyInputMode(FReply& SlateOperations, class UGameViewportClient& GameViewportClient) const override;
 };
 
+USTRUCT()
+struct FAsyncPhysicsTimestamp
+{
+	GENERATED_BODY()
+	UPROPERTY()
+	int32 ServerFrame = INDEX_NONE;
+
+	UPROPERTY()
+	int32 LocalFrame = INDEX_NONE;
+};
 
 /**
  * PlayerControllers are used by human players to control Pawns.
@@ -2244,4 +2254,57 @@ private:
 	FInputCmdBuffer InputBuffer;
 	FClientFrameInfo ClientFrameInfo;
 	FServerFrameInfo ServerFrameInfo;
+
+	/** The estimated offset between the local async physics tick frame number and the server's
+	*	This is used to synchronize events that happen in the async physics tick */
+	int32 LocalToServerAsyncPhysicsTickOffset;
+	
+	/** The latest server step we've received an offset correction for. This allows us to ignore out of order corrections that arrive late */
+	int32 ClientLatestCorrectedOffsetServerStep = INDEX_NONE;
+
+	/** The latest physics step we've sent to the server. Due to async we need to avoid duplicate sends */
+	int32 ClientLatestAsyncPhysicsStepSent = INDEX_NONE;
+
+	/** The latest server step we've received a time dilation for. Needed for out of order updates */
+	int32 ClientLatestTimeDilationServerStep = INDEX_NONE;
+
+	/** The server tells the client to speed up or slow down in order to keep its buffer full */
+	float ServerAsyncPhysicsTimeDilationToSend = 1.f;
+
+	/** The server records the latest timestamp it has to correct. This is used to update client (which may not happen on every physics step) */
+	FAsyncPhysicsTimestamp ServerLatestTimestampToCorrect;
+
+	/** The set of timestamps the client has sent to the server. Sorted by local frame number */
+	TArray<FAsyncPhysicsTimestamp> ServerPendingTimestamps;
+
+	/** Update the tick ofsset in between the local client and the server */
+	void UpdateServerAsyncPhysicsTickOffset();
+
+	UFUNCTION(Server, Unreliable)
+	void ServerSendLatestAsyncPhysicsTimestamp(FAsyncPhysicsTimestamp Timestamp);
+
+	UFUNCTION(Client, Unreliable)
+	void ClientCorrectionAsyncPhysicsTimestamp(FAsyncPhysicsTimestamp Timestamp);
+
+	UFUNCTION(Client, Unreliable)
+	void ClientAckTimeDilation(float TimeDilation, int32 ServerStep);
+
+	/** Update the tick offset in between the local client and the server */
+	virtual void AsyncPhysicsTickActor(float DeltaTime, float SimTime) override;
+
+public:
+
+	/** Enqueues a command to run at the time specified by AsyncPhysicsTimestamp. Note that if the time specified was missed the command is triggered as soon as possible as part of the async tick.
+		These commands are all run on the game thread. If you want to run on the physics thread see FPhysicsSolverBase::RegisterSimOneShotCallback
+		If OwningObject is not null this command will only fire as long as the object is still alive. This allows a lambda to still use the owning object's data for read/write (assuming data is designed for async physics tick)
+	*/
+	void ExecuteAsyncPhysicsCommand(const FAsyncPhysicsTimestamp& AsyncPhysicsTimestamp, UObject* OwningObject, const TFunction<void()>& Command);
+
+	/** Generates a timestamp for the upcoming physics step (plus any pending time). Useful for synchronizing client and server events on a specific physics step */
+	FAsyncPhysicsTimestamp GetAsyncPhysicsTimestamp(float DeltaSeconds = 0.f);
+
+	/** Returns the current estimated offset between the local async physics step and the server. This is useful for dealing with low level synchronization.
+		In general it's recommended to use GetAsyncPhysicsTimestamp which accounts for the offset automatically*/
+	int32 GetLocalToServerAsyncPhysicsTickOffset() const { return LocalToServerAsyncPhysicsTickOffset; }
+	
 };
diff --git a/Engine/Source/Runtime/Engine/Private/PhysicsEngine/AsyncPhysicsInputComponent.cpp b/Engine/Source/Runtime/Engine/Private/PhysicsEngine/AsyncPhysicsInputComponent.cpp
--- a/Engine/Source/Runtime/Engine/Private/PhysicsEngine/AsyncPhysicsInputComponent.cpp
+++ b/Engine/Source/Runtime/Engine/Private/PhysicsEngine/AsyncPhysicsInputComponent.cpp
@@ -101,20 +101,28 @@ struct FAsyncPhysicsInputRewindCallback : public Chaos::IRewindCallback
 
 	virtual void InjectInputs_External(int32 PhysicsStep, int32 NumSteps) override
 	{
-		int32 LocalOffset = 0;
+		int32 ServerFrame = PhysicsStep;
 		if (World->GetNetMode() == NM_Client)
 		{
 			if (APlayerController* PC = World->GetFirstPlayerController())
 			{
-				APlayerController::FClientFrameInfo& ClientFrameInfo = PC->GetClientFrameInfo();
-
-				if (ClientFrameInfo.LastProcessedInputFrame != INDEX_NONE)
+				static IConsoleVariable* EnableNetworkPhysicsCVar = IConsoleManager::Get().FindConsoleVariable(TEXT("np2.EnableNetworkPhysicsPrediction"));
+				if (EnableNetworkPhysicsCVar && EnableNetworkPhysicsCVar->GetInt() == 1)
 				{
-					LocalOffset = ClientFrameInfo.GetLocalFrameOffset();
+					const int32 LocalToServerOffset = PC->GetLocalToServerAsyncPhysicsTickOffset();
+					ServerFrame = PhysicsStep + LocalToServerOffset;
+				}
+				else
+				{
+					APlayerController::FClientFrameInfo& ClientFrameInfo = PC->GetClientFrameInfo();
+                    if (ClientFrameInfo.LastProcessedInputFrame != INDEX_NONE)
+                    {
+                    	ServerFrame = PhysicsStep - ClientFrameInfo.GetLocalFrameOffset();
+                    }
 				}
 			}
 		}
-		DispatchPhysicsTick.Broadcast(PhysicsStep, NumSteps, PhysicsStep - LocalOffset);
+		DispatchPhysicsTick.Broadcast(PhysicsStep, NumSteps, ServerFrame);
 	}
 
 	virtual void ProcessInputs_External(int32 PhysicsStep, const TArray<Chaos::FSimCallbackInputAndObject>& SimCallbackInputs)
@@ -410,7 +418,7 @@ void UAsyncPhysicsInputComponent::OnDispatchPhysicsTick(int32 PhysicsStep, int32
 	{
 		return;
 	}
-
+	
 	// It would be better if we only registered while we had a local controller,
 	// but this is simpler to implement
 
diff --git a/Engine/Source/Runtime/Engine/Private/PhysicsEngine/Experimental/PhysScene_Chaos.cpp b/Engine/Source/Runtime/Engine/Private/PhysicsEngine/Experimental/PhysScene_Chaos.cpp
--- a/Engine/Source/Runtime/Engine/Private/PhysicsEngine/Experimental/PhysScene_Chaos.cpp
+++ b/Engine/Source/Runtime/Engine/Private/PhysicsEngine/Experimental/PhysScene_Chaos.cpp
@@ -70,6 +70,14 @@ DECLARE_CYCLE_STAT(TEXT("Update Kinematics On Deferred SkelMeshes"), STAT_Update
 #include "Editor.h"
 #endif
 
+struct FPendingAsyncPhysicsCommand
+{
+	int32 PhysicsStep;
+	TWeakObjectPtr<UObject> OwningObject;
+	TFunction<void()> Command;
+};
+
+
 class FAsyncPhysicsTickCallback : public Chaos::TSimCallbackObject<>
 {
 public:
@@ -81,9 +89,31 @@ public:
 
 	TSet<UActorComponent*> AsyncPhysicsTickComponents;
 	TSet<AActor*> AsyncPhysicsTickActors;
+	TArray<FPendingAsyncPhysicsCommand> PendingCommands;
+
 	virtual void OnPreSimulate_Internal() override
 	{
 		using namespace Chaos;
+
+		for (int32 Idx = 0; Idx < PendingCommands.Num(); ++Idx)
+		{
+			const int32 CurrentFrame = static_cast<FPBDRigidsSolver*>(GetSolver())->GetCurrentFrame();
+
+			FPendingAsyncPhysicsCommand* PendingCommand = &PendingCommands[Idx];
+			bool bRemove = PendingCommand->OwningObject.IsStale();
+			if (!bRemove && PendingCommands[Idx].PhysicsStep <= CurrentFrame)
+			{
+				PendingCommands[Idx].Command();
+				bRemove = true;
+			}
+
+			if (bRemove)
+			{
+				PendingCommands.RemoveAt(Idx);	//Need to keep functions in order. If this is slow we could try going in reverse order, but expecting number of commands to be low per frame
+				--Idx;
+			}
+		}
+
 		const FReal DeltaTime = GetDeltaTime_Internal();
 		const FReal SimTime = GetSimTime_Internal();
 		//TODO: handle case where callbacks modify AsyncPhysicsTickComponents or AsyncPhysicsTickActors
@@ -1915,5 +1945,11 @@ void FPhysScene_Chaos::UnregisterAsyncPhysicsTickActor(AActor* Actor)
 	}
 }
 
+void FPhysScene_Chaos::EnqueueAsyncPhysicsCommand(int32 PhysicsStep, UObject* OwningObject, const TFunction<void()>& Command)
+{
+	EnableAsyncPhysicsTickCallback();
+	AsyncPhysicsTickCallback->PendingCommands.Add({ PhysicsStep, TWeakObjectPtr<UObject>(OwningObject), Command });
+}
+
 TSharedPtr<IPhysicsReplicationFactory> FPhysScene_Chaos::PhysicsReplicationFactory;
 
diff --git a/Engine/Source/Runtime/Engine/Private/PhysicsEngine/PhysicsReplication.cpp b/Engine/Source/Runtime/Engine/Private/PhysicsEngine/PhysicsReplication.cpp
--- a/Engine/Source/Runtime/Engine/Private/PhysicsEngine/PhysicsReplication.cpp
+++ b/Engine/Source/Runtime/Engine/Private/PhysicsEngine/PhysicsReplication.cpp
@@ -518,7 +518,11 @@ void FPhysicsReplication::OnTick(float DeltaSeconds, TMap<TWeakObjectPtr<UPrimit
 				Chaos::FPhysicsSolver* Solver = PhysScene->GetSolver();
 				check(Solver);
 
-				//LocalFrameOffset = PlayerController->GetLocalToServerAsyncPhysicsTickOffset();
+				static IConsoleVariable* EnableNetworkPhysicsCVar = IConsoleManager::Get().FindConsoleVariable(TEXT("np2.EnableNetworkPhysicsPrediction"));
+				if (EnableNetworkPhysicsCVar && EnableNetworkPhysicsCVar->GetInt() == 1)
+				{
+					LocalFrameOffset = PlayerController->GetLocalToServerAsyncPhysicsTickOffset();
+				}
 				//TODO: as we send physics updates down we need to record latest seen
 				//NumPredictedFrames = Solver->GetCurrentFrame() - ClientFrameInfo.LastProcessedInputFrame;
 			}
@@ -675,7 +679,7 @@ void FPhysicsReplication::ApplyAsyncDesiredState(const float DeltaSeconds, const
 
 				const FVector NewPos = FMath::Lerp(FVector(CurrentState.Position), TargetPos, PositionLerp);
 				const FQuat NewAng = FQuat::Slerp(CurrentState.Quaternion, TargetQuat, AngleLerp);
-
+				
 				Handle->SetX(NewPos);
 				Handle->SetR(NewAng);
 				Handle->SetV(NewLinVel);
diff --git a/Engine/Source/Runtime/Engine/Private/PlayerController.cpp b/Engine/Source/Runtime/Engine/Private/PlayerController.cpp
--- a/Engine/Source/Runtime/Engine/Private/PlayerController.cpp
+++ b/Engine/Source/Runtime/Engine/Private/PlayerController.cpp
@@ -71,6 +71,7 @@
 #include "WorldPartition/WorldPartitionStreamingSource.h"
 #include "Physics/AsyncPhysicsInputComponent.h"
 #include "GenericPlatform/GenericPlatformInputDeviceMapper.h"
+#include "PBDRigidsSolver.h"
 
 #if UE_WITH_IRIS
 #include "Iris/ReplicationSystem/ReplicationSystem.h"
@@ -102,6 +103,23 @@ namespace PlayerControllerCVars
 		ECVF_Default);
 }
 
+namespace NetworkPhysicsCvars
+{
+	int32 NumRedundantCmds = 3;
+	FAutoConsoleVariableRef CVarNumRedundantCmds(TEXT("np2.NumRedundantCmds"), NumRedundantCmds, TEXT("Number of redundant user cmds to send per frame"));
+
+#if (UE_BUILD_SHIPPING || UE_BUILD_TEST)
+	int32 EnableDebugRPC = 0;
+#else
+	int32 EnableDebugRPC = 1;
+#endif
+	FAutoConsoleVariableRef CVarEnableDebugRPC(TEXT("np2.EnableDebugRPC"), EnableDebugRPC, TEXT("Sends extra debug information to clients about server side input buffering"));
+
+	int32 EnableNetworkPhysicsPrediction = 0;
+	FAutoConsoleVariableRef CVarEnableNetworkPhysicsPrediction(TEXT("np2.EnableNetworkPhysicsPrediction"), EnableNetworkPhysicsPrediction, TEXT("Enables network physics prediction"));
+}
+
+
 const float RetryClientRestartThrottleTime = 0.5f;
 const float RetryServerAcknowledgeThrottleTime = 0.25f;
 const float RetryServerCheckSpectatorThrottleTime = 0.25f;
@@ -166,6 +184,11 @@ APlayerController::APlayerController(const FObjectInitializer& ObjectInitializer
 		// We want to drive rotation with ControlRotation regardless of attachment state.
 		RootComponent->SetUsingAbsoluteRotation(true);
 	}
+
+	if(NetworkPhysicsCvars::EnableNetworkPhysicsPrediction == 1)
+	{
+		bAsyncPhysicsTickEnabled = true;
+	}
 }
 
 float APlayerController::GetNetPriority(const FVector& ViewPos, const FVector& ViewDir, AActor* Viewer, AActor* ViewTarget, UActorChannel* InChannel, float Time, bool bLowBandwidth)
@@ -1654,35 +1677,36 @@ void APlayerController::ClientSetCameraFade_Implementation(bool bEnableFading, F
 
 /// @endcond
 
-namespace UE_NETWORK_PHYSICS
-{
-	int32 NumRedundantCmds=3;
-	FAutoConsoleVariableRef CVarNumRedundantCmds(TEXT("np2.NumRedundantCmds"), NumRedundantCmds, TEXT("Number of redundant user cmds to send per frame"));
-
-#if (UE_BUILD_SHIPPING || UE_BUILD_TEST)
-	int32 EnableDebugRPC=0;
-#else
-	int32 EnableDebugRPC=1;
-#endif
-	FAutoConsoleVariableRef CVarEnableDebugRPC(TEXT("np2.EnableDebugRPC"), EnableDebugRPC, TEXT("Sends extra debug information to clients about server side input buffering"));
-}
-
 void APlayerController::SendClientAdjustment()
 {
-	if (ServerFrameInfo.LastProcessedInputFrame != INDEX_NONE && ServerFrameInfo.LastProcessedInputFrame != ServerFrameInfo.LastSentLocalFrame)
+	if(!NetworkPhysicsCvars::EnableNetworkPhysicsPrediction)
 	{
-		ServerFrameInfo.LastSentLocalFrame = ServerFrameInfo.LastProcessedInputFrame;		
-		ClientRecvServerAckFrame(ServerFrameInfo.LastProcessedInputFrame, ServerFrameInfo.LastLocalFrame, ServerFrameInfo.QuantizedTimeDilation);
-
-		if (UE_NETWORK_PHYSICS::EnableDebugRPC)
+		if (ServerFrameInfo.LastProcessedInputFrame != INDEX_NONE && ServerFrameInfo.LastProcessedInputFrame != ServerFrameInfo.LastSentLocalFrame)
 		{
-			ClientRecvServerAckFrameDebug(InputBuffer.HeadFrame() - ServerFrameInfo.LastProcessedInputFrame, ServerFrameInfo.TargetNumBufferedCmds);
+			ServerFrameInfo.LastSentLocalFrame = ServerFrameInfo.LastProcessedInputFrame;		
+			ClientRecvServerAckFrame(ServerFrameInfo.LastProcessedInputFrame, ServerFrameInfo.LastLocalFrame, ServerFrameInfo.QuantizedTimeDilation);
+
+			if (NetworkPhysicsCvars::EnableDebugRPC)
+			{
+				ClientRecvServerAckFrameDebug(InputBuffer.HeadFrame() - ServerFrameInfo.LastProcessedInputFrame, ServerFrameInfo.TargetNumBufferedCmds);
+			}
 		}
 	}
-
-	if (AcknowledgedPawn != GetPawn() && !GetSpectatorPawn())
+	else
 	{
-		return;
+		if (ServerLatestTimestampToCorrect.ServerFrame == INDEX_NONE)
+		{
+			//Nothing to correct so do not update client
+			return;
+		}
+
+		ClientCorrectionAsyncPhysicsTimestamp(ServerLatestTimestampToCorrect);
+		ServerLatestTimestampToCorrect.ServerFrame = INDEX_NONE;
+
+		if (AcknowledgedPawn != GetPawn() && !GetSpectatorPawn())
+		{
+			return;
+		}
 	}
 
 	// Server sends updates.
@@ -1703,7 +1727,7 @@ void APlayerController::PushClientInput(int32 InRecvClientInputFrame, TArray<uin
 	InputBuffer.Write(InRecvClientInputFrame) = MoveTemp(Data);
 
 	// Do the RPC right here, including the redundant send. This should probably be time based and managed somewhere else like in Tick eventually
-	for (int32 Frame = FMath::Max(1, InRecvClientInputFrame-UE_NETWORK_PHYSICS::NumRedundantCmds+1); Frame <= InRecvClientInputFrame; ++Frame)
+	for (int32 Frame = FMath::Max(1, InRecvClientInputFrame-NetworkPhysicsCvars::NumRedundantCmds+1); Frame <= InRecvClientInputFrame; ++Frame)
 	{
 		ServerRecvClientInputFrame(Frame, InputBuffer.Get(Frame));
 	}
@@ -5068,6 +5092,11 @@ void APlayerController::TickActor( float DeltaSeconds, ELevelTick TickType, FAct
 	// Clear old axis inputs since we are done with them. 
 	RotationInput = FRotator::ZeroRotator;
 
+	if(!!NetworkPhysicsCvars::EnableNetworkPhysicsPrediction && GetLocalRole() == ROLE_AutonomousProxy && bIsClient)
+	{
+		UpdateServerAsyncPhysicsTickOffset();
+	}
+
 #if !(UE_BUILD_SHIPPING || UE_BUILD_TEST)
 	if (CheatManager != nullptr)
 	{
@@ -5915,6 +5944,157 @@ const UAsyncPhysicsData* APlayerController::GetAsyncPhysicsDataToConsume() const
 	return AsyncPhysicsDataComponent ? AsyncPhysicsDataComponent->GetDataToConsume() : nullptr;
 }
 
+void APlayerController::ExecuteAsyncPhysicsCommand(const FAsyncPhysicsTimestamp& AsyncPhysicsTimestamp, UObject* OwningObject, const TFunction<void()>& Command)
+{
+	if(UWorld* World = GetWorld())
+	{
+		if(FPhysScene* PhysScene = World->GetPhysicsScene())
+		{
+			const int32 PhysicsStep = IsLocalController() ? AsyncPhysicsTimestamp.LocalFrame : AsyncPhysicsTimestamp.ServerFrame;
+			PhysScene->EnqueueAsyncPhysicsCommand(PhysicsStep, OwningObject, Command);
+		}
+	}
+}
+
+FAsyncPhysicsTimestamp APlayerController::GetAsyncPhysicsTimestamp(float DeltaSeconds)
+{
+	using namespace Chaos;
+
+	FAsyncPhysicsTimestamp Timestamp;
+
+	if(UWorld* World = GetWorld())
+	{
+		if (FPhysScene* PhysScene = World->GetPhysicsScene())
+		{
+			if (FPBDRigidsSolver* Solver = static_cast<FPBDRigidsSolver*>(PhysScene->GetSolver()))
+			{
+				const FReal DeltaTime = Solver->GetAsyncDeltaTime();
+				const int32 PendingSteps = (DeltaTime > 0.0) ? DeltaSeconds / DeltaTime : 0;
+
+				int32 LocalPhysicsStep;
+				if(Solver->IsGameThreadFrozen())
+				{
+					//We are calling inside the async tick, so use the current frame
+					LocalPhysicsStep = Solver->GetCurrentFrame();
+				}
+				else
+				{
+					//We are on the GT so give the upcoming async tick
+					LocalPhysicsStep = Solver->GetMarshallingManager().GetExternalTimestamp_External();
+				}
+				
+				LocalPhysicsStep += PendingSteps;	//Add any pending steps user wants to wait on
+				Timestamp.ServerFrame = LocalPhysicsStep;
+				Timestamp.LocalFrame = LocalPhysicsStep;
+
+				if (IsLocalController())
+				{
+					//If local controller we update server frame based on our estimate
+					Timestamp.ServerFrame = LocalPhysicsStep + LocalToServerAsyncPhysicsTickOffset;
+				}
+			}
+		}
+	}
+
+	return Timestamp;
+}
+
+void APlayerController::UpdateServerAsyncPhysicsTickOffset()
+{
+	FAsyncPhysicsTimestamp Timestamp = GetAsyncPhysicsTimestamp();
+	if(ClientLatestAsyncPhysicsStepSent == Timestamp.LocalFrame)
+	{
+		//If GT is running faster than physics sim the physics timestep will not have changed, so no need to send another update to server
+		//This ensures monotonic increase
+		return;
+	}
+
+	ClientLatestAsyncPhysicsStepSent = Timestamp.LocalFrame;
+	ServerSendLatestAsyncPhysicsTimestamp(Timestamp);
+}
+
+void APlayerController::ServerSendLatestAsyncPhysicsTimestamp_Implementation(FAsyncPhysicsTimestamp Timestamp)
+{
+	//This tells the server how the client thinks the async physics tick will line up.
+	//Timestamps could be out of order due to networking, so we make sure they are sorted using the client's local frame which is monotonically increasing
+	int32 Idx;
+	for(Idx = ServerPendingTimestamps.Num() - 1; Idx >= 0; --Idx)
+	{
+		ensureMsgf(ServerPendingTimestamps[Idx].LocalFrame != Timestamp.LocalFrame, TEXT("Client should never send duplicate timestamps, something is wrong"));
+		if(ServerPendingTimestamps[Idx].LocalFrame < Timestamp.LocalFrame)
+		{
+			break;
+		}
+	}
+	ServerPendingTimestamps.Insert(Timestamp, Idx + 1);
+}
+
+void APlayerController::ClientCorrectionAsyncPhysicsTimestamp_Implementation(FAsyncPhysicsTimestamp Timestamp)
+{
+	//This tells the client that a timestamp it sent out was wrong (for example we ran locally on step 5 and expected server to run on step 10, but it actually ran on step 11).
+	//The error can only be later. That is, it can never be that we expect to run on server step 10 but actually ran on sever step 9
+	//Once a timestamp has been corrected, any earlier timestamps can be ignored (these can be out of order because of networking)
+
+	ensureMsgf(Timestamp.ServerFrame != ClientLatestCorrectedOffsetServerStep, TEXT("Server only sends at most one correction per timestamp, duplicate corrections means something is wrong"));
+	if(Timestamp.ServerFrame < ClientLatestCorrectedOffsetServerStep)
+	{
+		//already corrected after this so do nothing
+		return;
+	}
+
+	const int32 NewOffset = Timestamp.ServerFrame - Timestamp.LocalFrame; //The new offset as reported by the server
+	ensureMsgf(NewOffset >= LocalToServerAsyncPhysicsTickOffset, TEXT("The offset between client and server can only ever increase"));
+	LocalToServerAsyncPhysicsTickOffset = NewOffset;
+}
+
+void APlayerController::ClientAckTimeDilation_Implementation(float TimeDilation, int32 ServerStep)
+{
+	if(ServerStep < ClientLatestTimeDilationServerStep)
+	{
+		//Stale ack so do nothing
+		return;
+	}
+
+	ensureMsgf(ClientLatestTimeDilationServerStep != ServerStep, TEXT("Server should send at most one time dilation per step, duplicate means something is wrong"));
+	ClientLatestTimeDilationServerStep = ServerStep;
+	if(UWorld* World = GetWorld())
+	{
+		World->GetPhysicsScene()->SetNetworkDeltaTimeScale(TimeDilation);
+	}
+}
+
+void APlayerController::AsyncPhysicsTickActor(float DeltaTime, float SimTime)
+{
+	Super::AsyncPhysicsTickActor(DeltaTime, SimTime);
+
+	if(NetworkPhysicsCvars::EnableNetworkPhysicsPrediction)
+	{
+		//TODO: only kick this off if server and using this feature
+		if (IsLocalController()) { return; }
+
+		if (ServerPendingTimestamps.Num() == 0)
+		{
+			//TODO: starved for input user needs to speed up
+			return;
+		}
+
+		const FAsyncPhysicsTimestamp ActualTimestamp = GetAsyncPhysicsTimestamp();
+
+		//If the pending timestamp is bigger than the server frame we simple wait, otherwise make sure they aren't too early
+		if (ServerPendingTimestamps[0].ServerFrame <= ActualTimestamp.ServerFrame)
+		{
+			if (ServerPendingTimestamps[0].ServerFrame < ActualTimestamp.ServerFrame)
+			{
+				//The earliest pending client timestamp is too early, so their offset must be bigger than they thought
+				ServerLatestTimestampToCorrect.ServerFrame = ActualTimestamp.ServerFrame;
+				ServerLatestTimestampToCorrect.LocalFrame = ServerPendingTimestamps[0].LocalFrame;	//The client's local frame stays the same there's just a mismatch on the server frame
+			}
+
+			ServerPendingTimestamps.RemoveAt(0);	//We've updated the client if needed, so can discard
+			//NOTE: we purposely don't fixup any future pending client timestamps. This is because the RPC is unreliable so it's best to send the error correction redundantly over multiple frames if needed
+		}
+	}
+}
 
 #undef LOCTEXT_NAMESPACE
 
diff --git a/Engine/Source/Runtime/Engine/Public/Physics/Experimental/PhysScene_Chaos.h b/Engine/Source/Runtime/Engine/Public/Physics/Experimental/PhysScene_Chaos.h
--- a/Engine/Source/Runtime/Engine/Public/Physics/Experimental/PhysScene_Chaos.h
+++ b/Engine/Source/Runtime/Engine/Public/Physics/Experimental/PhysScene_Chaos.h
@@ -119,6 +119,9 @@ public:
 	void RegisterAsyncPhysicsTickActor(AActor* Actor);
 	void UnregisterAsyncPhysicsTickActor(AActor* Actor);
 
+	void EnqueueAsyncPhysicsCommand(int32 PhysicsStep, UObject* OwningObject, const TFunction<void()>& Command);
+
+
 	/**
 	 * Called during creation of the physics state for gamethread objects to pass off an object to the physics thread
 	 */
diff --git a/Engine/Source/Runtime/Experimental/Chaos/Public/Chaos/Framework/PhysicsSolverBase.h b/Engine/Source/Runtime/Experimental/Chaos/Public/Chaos/Framework/PhysicsSolverBase.h
--- a/Engine/Source/Runtime/Experimental/Chaos/Public/Chaos/Framework/PhysicsSolverBase.h
+++ b/Engine/Source/Runtime/Experimental/Chaos/Public/Chaos/Framework/PhysicsSolverBase.h
@@ -392,6 +392,8 @@ namespace Chaos
 
 		void EnableAsyncMode(FReal FixedDt);
 		
+		FReal GetAsyncDeltaTime() const { return AsyncDt; } 
+		
 		void DisableAsyncMode();
 		
 		virtual void ConditionalApplyRewind_Internal(){}
