Decal materials calculate full material relevance structure.
Previously we didn't fill anything but bDecal, but that meant that WPO relevance was missing.
Could add just a few things like WPO/PDO, but I think it's best to gather all relevance, and if some parts of relevance are irrelevant, or shouldn't be set then we can spot fix that and enforce on the UI side.

diff --git a/Engine/Source/Runtime/Engine/Private/Materials/MaterialInterface.cpp b/Engine/Source/Runtime/Engine/Private/Materials/MaterialInterface.cpp
--- a/Engine/Source/Runtime/Engine/Private/Materials/MaterialInterface.cpp
+++ b/Engine/Source/Runtime/Engine/Private/Materials/MaterialInterface.cpp
@@ -385,60 +385,54 @@ FMaterialRelevance UMaterialInterface::GetRelevance_Internal(const UMaterial* Ma
 
 		MaterialRelevance.ShadingModelMask = GetShadingModels().GetShadingModelField();
 		MaterialRelevance.CustomDepthStencilUsageMask = MaterialResource->GetCustomDepthStencilUsageMask_GameThread();
+		MaterialRelevance.bDecal = bDecal;
 
-		if(bDecal)
+		// Check whether the material can be drawn in the separate translucency pass as per FMaterialResource::IsTranslucencyAfterDOFEnabled and IsMobileSeparateTranslucencyEnabled
+		EMaterialTranslucencyPass TranslucencyPass = MTP_BeforeDOF;
+		const bool bSupportsSeparateTranslucency = Material->MaterialDomain != MD_UI && Material->MaterialDomain != MD_DeferredDecal;
+		if (bIsTranslucent && bSupportsSeparateTranslucency)
 		{
-			MaterialRelevance.bDecal = bDecal;
-			// we rely on FMaterialRelevance defaults are 0
-		}
-		else
-		{
-			// Check whether the material can be drawn in the separate translucency pass as per FMaterialResource::IsTranslucencyAfterDOFEnabled and IsMobileSeparateTranslucencyEnabled
-			EMaterialTranslucencyPass TranslucencyPass = MTP_BeforeDOF;
-			const bool bSupportsSeparateTranslucency = Material->MaterialDomain != MD_UI && Material->MaterialDomain != MD_DeferredDecal;
-			if (bIsTranslucent && bSupportsSeparateTranslucency)
+			if (bIsMobile)
 			{
-				if (bIsMobile)
+				if (Material->bEnableMobileSeparateTranslucency)
 				{
-					if (Material->bEnableMobileSeparateTranslucency)
-					{
-						TranslucencyPass = MTP_AfterDOF;
-					}
+					TranslucencyPass = MTP_AfterDOF;
 				}
-				else
-				{
-					TranslucencyPass = Material->TranslucencyPass;
-				}
-			}			
-
-			// If dual blending is supported, and we are rendering post-DOF translucency, then we also need to render a second pass to the modulation buffer.
-			// The modulation buffer can also be used for regular modulation shaders after DoF.
-			const bool bMaterialSeparateModulation = MaterialResource->IsDualBlendingEnabled(GShaderPlatformForFeatureLevel[InFeatureLevel]) || BlendMode == BLEND_Modulate || StrataBlendMode == SBM_ColoredTransmittanceOnly;
-
-			MaterialRelevance.bOpaque = !bIsTranslucent;
-			MaterialRelevance.bMasked = IsMasked();
-			MaterialRelevance.bDistortion = MaterialResource->IsDistorted();
-			MaterialRelevance.bHairStrands = IsCompatibleWithHairStrands(MaterialResource, InFeatureLevel);
-			MaterialRelevance.bTwoSided = MaterialResource->IsTwoSided();
-			MaterialRelevance.bSeparateTranslucency = bIsTranslucent && (TranslucencyPass == MTP_AfterDOF);
-			MaterialRelevance.bTranslucencyModulate = bMaterialSeparateModulation;
-			MaterialRelevance.bPostMotionBlurTranslucency = (TranslucencyPass == MTP_AfterMotionBlur);
-			MaterialRelevance.bNormalTranslucency = bIsTranslucent && (TranslucencyPass == MTP_BeforeDOF);
-			MaterialRelevance.bDisableDepthTest = bIsTranslucent && Material->bDisableDepthTest;		
-			MaterialRelevance.bUsesSceneColorCopy = bIsTranslucent && MaterialResource->RequiresSceneColorCopy_GameThread();
-			MaterialRelevance.bOutputsTranslucentVelocity = Material->IsTranslucencyWritingVelocity();
-			MaterialRelevance.bUsesGlobalDistanceField = MaterialResource->UsesGlobalDistanceField_GameThread();
-			MaterialRelevance.bUsesWorldPositionOffset = MaterialResource->MaterialUsesWorldPositionOffset_GameThread();
-			MaterialRelevance.bUsesPixelDepthOffset = MaterialResource->MaterialUsesPixelDepthOffset_GameThread();
-			ETranslucencyLightingMode TranslucencyLightingMode = MaterialResource->GetTranslucencyLightingMode();
-			MaterialRelevance.bTranslucentSurfaceLighting = bIsTranslucent && (TranslucencyLightingMode == TLM_SurfacePerPixelLighting || TranslucencyLightingMode == TLM_Surface);
-			MaterialRelevance.bUsesSceneDepth = MaterialResource->MaterialUsesSceneDepthLookup_GameThread();
-			MaterialRelevance.bHasVolumeMaterialDomain = MaterialResource->IsVolumetricPrimitive();
-			MaterialRelevance.bUsesDistanceCullFade = MaterialResource->MaterialUsesDistanceCullFade_GameThread();
-			MaterialRelevance.bUsesSkyMaterial = Material->bIsSky;
-			MaterialRelevance.bUsesSingleLayerWaterMaterial = bUsesSingleLayerWaterMaterial;
-			MaterialRelevance.bUsesAnisotropy = bUsesAnisotropy;
-		}
+			}
+			else
+			{
+				TranslucencyPass = Material->TranslucencyPass;
+			}
+		}			
+
+		// If dual blending is supported, and we are rendering post-DOF translucency, then we also need to render a second pass to the modulation buffer.
+		// The modulation buffer can also be used for regular modulation shaders after DoF.
+		const bool bMaterialSeparateModulation = MaterialResource->IsDualBlendingEnabled(GShaderPlatformForFeatureLevel[InFeatureLevel]) || BlendMode == BLEND_Modulate || StrataBlendMode == SBM_ColoredTransmittanceOnly;
+
+		MaterialRelevance.bOpaque = !bIsTranslucent;
+		MaterialRelevance.bMasked = IsMasked();
+		MaterialRelevance.bDistortion = MaterialResource->IsDistorted();
+		MaterialRelevance.bHairStrands = IsCompatibleWithHairStrands(MaterialResource, InFeatureLevel);
+		MaterialRelevance.bTwoSided = MaterialResource->IsTwoSided();
+		MaterialRelevance.bSeparateTranslucency = bIsTranslucent && (TranslucencyPass == MTP_AfterDOF);
+		MaterialRelevance.bTranslucencyModulate = bMaterialSeparateModulation;
+		MaterialRelevance.bPostMotionBlurTranslucency = (TranslucencyPass == MTP_AfterMotionBlur);
+		MaterialRelevance.bNormalTranslucency = bIsTranslucent && (TranslucencyPass == MTP_BeforeDOF);
+		MaterialRelevance.bDisableDepthTest = bIsTranslucent && Material->bDisableDepthTest;		
+		MaterialRelevance.bUsesSceneColorCopy = bIsTranslucent && MaterialResource->RequiresSceneColorCopy_GameThread();
+		MaterialRelevance.bOutputsTranslucentVelocity = Material->IsTranslucencyWritingVelocity();
+		MaterialRelevance.bUsesGlobalDistanceField = MaterialResource->UsesGlobalDistanceField_GameThread();
+		MaterialRelevance.bUsesWorldPositionOffset = MaterialResource->MaterialUsesWorldPositionOffset_GameThread();
+		MaterialRelevance.bUsesPixelDepthOffset = MaterialResource->MaterialUsesPixelDepthOffset_GameThread();
+		ETranslucencyLightingMode TranslucencyLightingMode = MaterialResource->GetTranslucencyLightingMode();
+		MaterialRelevance.bTranslucentSurfaceLighting = bIsTranslucent && (TranslucencyLightingMode == TLM_SurfacePerPixelLighting || TranslucencyLightingMode == TLM_Surface);
+		MaterialRelevance.bUsesSceneDepth = MaterialResource->MaterialUsesSceneDepthLookup_GameThread();
+		MaterialRelevance.bHasVolumeMaterialDomain = MaterialResource->IsVolumetricPrimitive();
+		MaterialRelevance.bUsesDistanceCullFade = MaterialResource->MaterialUsesDistanceCullFade_GameThread();
+		MaterialRelevance.bUsesSkyMaterial = Material->bIsSky;
+		MaterialRelevance.bUsesSingleLayerWaterMaterial = bUsesSingleLayerWaterMaterial;
+		MaterialRelevance.bUsesAnisotropy = bUsesAnisotropy;
+
 		return MaterialRelevance;
 	}
 	else
