Chaos: Don't force recompute convex indices

diff --git a/Engine/Source/Runtime/Engine/Classes/PhysicsEngine/ConvexElem.h b/Engine/Source/Runtime/Engine/Classes/PhysicsEngine/ConvexElem.h
--- a/Engine/Source/Runtime/Engine/Classes/PhysicsEngine/ConvexElem.h
+++ b/Engine/Source/Runtime/Engine/Classes/PhysicsEngine/ConvexElem.h
@@ -52,6 +52,15 @@ private:
 
 public:
 
+	enum class EConvexDataUpdateMethod
+	{
+		/** Update convex index and vertex data from the chaos convex object only if they are missing */
+		UpdateConvexDataOnlyIfMissing = 0,
+
+		/** Always convex recompute index and vertex data from the set chaos convex object */
+		AlwaysUpdateConvexData
+	};
+
 	ENGINE_API FKConvexElem();
 	ENGINE_API FKConvexElem(const FKConvexElem& Other);
 	ENGINE_API ~FKConvexElem();
@@ -96,7 +105,14 @@ public:
 		return ChaosConvex;
 	}
 
-	ENGINE_API void SetChaosConvexMesh(TSharedPtr<Chaos::FConvex, ESPMode::ThreadSafe>&& InChaosConvex);
+	/** 
+	* Set the chaos convex mesh 
+	* Note : This will by default invalidate the convex data and recompute them by calling ComputeChaosConvexIndices
+	* Only set ConvexDataUpdateMethod to UpdateConvexDataOnlyIfMissing if you know the data is up to date with the chaos convex object or if you plan to call ComputeChaosConvexIndices later
+	* @param InChaosConvex	Chaos convex mesh to set 
+	* @param ConvexDataUpdateMethod	method to use to update internal convex data
+	*/
+	ENGINE_API void SetChaosConvexMesh(TSharedPtr<Chaos::FConvex, ESPMode::ThreadSafe>&& InChaosConvex, EConvexDataUpdateMethod ConvexDataUpdateMethod = EConvexDataUpdateMethod::AlwaysUpdateConvexData);
 
 	ENGINE_API void ResetChaosConvexMesh();
 
diff --git a/Engine/Source/Runtime/Engine/Private/PhysicsEngine/BodySetup.cpp b/Engine/Source/Runtime/Engine/Private/PhysicsEngine/BodySetup.cpp
--- a/Engine/Source/Runtime/Engine/Private/PhysicsEngine/BodySetup.cpp
+++ b/Engine/Source/Runtime/Engine/Private/PhysicsEngine/BodySetup.cpp
@@ -596,7 +596,8 @@ void UBodySetup::FinishCreatingPhysicsMeshes_Chaos(TArray<TSharedPtr<Chaos::FCon
 			if(CHAOS_ENSURE((ElementIndex < ConvexImplicits.Num())
 			   && ConvexImplicits[ElementIndex]->IsValidGeometry()))
 			{
-				ConvexElem.SetChaosConvexMesh(MoveTemp(ConvexImplicits[ElementIndex]));
+				// Optimization: Only update internal convex data because we just deserialized the convex and the FKConvexElem convex information matches the chaos implicit one
+				ConvexElem.SetChaosConvexMesh(MoveTemp(ConvexImplicits[ElementIndex]), FKConvexElem::EConvexDataUpdateMethod::UpdateConvexDataOnlyIfMissing);
 
 #if TRACK_CHAOS_GEOMETRY
 				ConvexElem.GetChaosConvexMesh()->Track(Chaos::MakeSerializable(ConvexElem.GetChaosConvexMesh()), FullName);
@@ -1744,11 +1745,11 @@ FVector::FReal FKConvexElem::GetScaledVolume(const FVector& Scale) const
 	return Volume;
 }
 
-void FKConvexElem::SetChaosConvexMesh(TSharedPtr<Chaos::FConvex, ESPMode::ThreadSafe>&& InChaosConvex)
+void FKConvexElem::SetChaosConvexMesh(TSharedPtr<Chaos::FConvex, ESPMode::ThreadSafe>&& InChaosConvex, EConvexDataUpdateMethod ConvexDataUpdateMethod /* = EConvexDataUpdateMethod::AlwaysUpdateConvexData */)
 {
 	ChaosConvex = MoveTemp(InChaosConvex);
 	
-	const bool bForceCompute = true;
+	const bool bForceCompute = (ConvexDataUpdateMethod == EConvexDataUpdateMethod::AlwaysUpdateConvexData);
 	ComputeChaosConvexIndices(bForceCompute);
 }
 
