World Partition: unload location volumes when deleting them.

#rnx

diff --git a/Engine/Source/Runtime/Engine/Private/LocationVolume.cpp b/Engine/Source/Runtime/Engine/Private/LocationVolume.cpp
--- a/Engine/Source/Runtime/Engine/Private/LocationVolume.cpp
+++ b/Engine/Source/Runtime/Engine/Private/LocationVolume.cpp
@@ -53,6 +53,18 @@ void ALocationVolume::PostRegisterAllComponents()
 	}
 }
 
+void ALocationVolume::UnregisterAllComponents(bool bForReregister)
+{
+	Super::UnregisterAllComponents(bForReregister);
+
+	if (!bForReregister && WorldPartitionActorLoader)
+	{
+		WorldPartitionActorLoader->Unload();
+		delete WorldPartitionActorLoader;
+		WorldPartitionActorLoader = nullptr;
+	}
+}
+
 void ALocationVolume::BeginDestroy()
 {
 	if (WorldPartitionActorLoader)
diff --git a/Engine/Source/Runtime/Engine/Private/WorldPartition/WorldPartition.cpp b/Engine/Source/Runtime/Engine/Private/WorldPartition/WorldPartition.cpp
--- a/Engine/Source/Runtime/Engine/Private/WorldPartition/WorldPartition.cpp
+++ b/Engine/Source/Runtime/Engine/Private/WorldPartition/WorldPartition.cpp
@@ -1388,7 +1388,7 @@ void UWorldPartition::SavePerUserSettings()
 		TArray<FName> EditorLoadedLocationVolumes;
 		for (FActorDescContainerCollection::TConstIterator<> ActorDescIterator(this); ActorDescIterator; ++ActorDescIterator)
 		{
-			if (ALocationVolume* LocationVolume = Cast<ALocationVolume>(ActorDescIterator->GetActor()))
+			if (ALocationVolume* LocationVolume = Cast<ALocationVolume>(ActorDescIterator->GetActor()); IsValid(LocationVolume))
 			{
 				check(LocationVolume->GetClass()->ImplementsInterface(UWorldPartitionActorLoaderInterface::StaticClass()));
 
diff --git a/Engine/Source/Runtime/Engine/Public/LocationVolume.h b/Engine/Source/Runtime/Engine/Public/LocationVolume.h
--- a/Engine/Source/Runtime/Engine/Public/LocationVolume.h
+++ b/Engine/Source/Runtime/Engine/Public/LocationVolume.h
@@ -29,6 +29,7 @@ public:
 	virtual bool IsEditorOnly() const override { return !bIsRuntime; }
 #if WITH_EDITOR
 	virtual void PostRegisterAllComponents();
+	virtual void UnregisterAllComponents(bool bForReregister) override;
 	virtual bool CanChangeIsSpatiallyLoadedFlag() const override { return false; }
 	//~ End AActor Interface
 
