WorldPartitionCells are no longer always outered to WorldPartiton
Adds IWorldPartitionCellOwner which cells are now outered to
Allows for cells to not be packaged within the main world

#test pie, cooked, ancient_game, hlod, asteria_terrain

diff --git a/Engine/Source/Runtime/Engine/Private/WorldPartition/WorldPartitionLevelStreamingDynamic.cpp b/Engine/Source/Runtime/Engine/Private/WorldPartition/WorldPartitionLevelStreamingDynamic.cpp
--- a/Engine/Source/Runtime/Engine/Private/WorldPartition/WorldPartitionLevelStreamingDynamic.cpp
+++ b/Engine/Source/Runtime/Engine/Private/WorldPartition/WorldPartitionLevelStreamingDynamic.cpp
@@ -110,8 +110,8 @@ void UWorldPartitionLevelStreamingDynamic::Initialize(const UWorldPartitionRunti
 	StreamingPriority = 0;
 	UnsavedActorsContainer = InCell.UnsavedActorsContainer;
 
-	UWorld* OuterWorld = InCell.GetOuterUWorldPartition()->GetTypedOuter<UWorld>();
-	Initialize(OuterWorld, InCell.GetPackages());
+	IWorldPartitionRuntimeCellOwner* CellOwner = InCell.GetCellOwner();
+	Initialize(CellOwner->GetOuterWorld(), InCell.GetPackages());
 }
 
 /**
diff --git a/Engine/Source/Runtime/Engine/Private/WorldPartition/WorldPartitionRuntimeCell.cpp b/Engine/Source/Runtime/Engine/Private/WorldPartition/WorldPartitionRuntimeCell.cpp
--- a/Engine/Source/Runtime/Engine/Private/WorldPartition/WorldPartitionRuntimeCell.cpp
+++ b/Engine/Source/Runtime/Engine/Private/WorldPartition/WorldPartitionRuntimeCell.cpp
@@ -4,6 +4,7 @@
 #include "WorldPartition/WorldPartition.h"
 #include "WorldPartition/WorldPartitionDebugHelper.h"
 #include "WorldPartition/WorldPartitionRuntimeSpatialHash.h"
+#include "WorldPartition/WorldPartitionRuntimeCellOwner.h"
 #include "WorldPartition/WorldPartitionLevelStreamingPolicy.h"
 #include "WorldPartition/DataLayer/DataLayerSubsystem.h"
 #include "WorldPartition/DataLayer/DataLayersID.h"
@@ -23,7 +24,9 @@ UWorldPartitionRuntimeCell::UWorldPartitionRuntimeCell(const FObjectInitializer&
 #if !UE_BUILD_SHIPPING
 	, DebugStreamingPriority(-1.f)
 #endif	
-{}
+{
+	check(HasAnyFlags(RF_ClassDefaultObject) || GetOuter()->Implements<UWorldPartitionRuntimeCellOwner>());
+}
 
 #if WITH_EDITOR
 bool UWorldPartitionRuntimeCell::NeedsActorToCellRemapping() const
@@ -68,7 +71,7 @@ void UWorldPartitionRuntimeCell::UpdateDebugName()
 	Builder += FString::Printf(TEXT("L%d_X%d_Y%d"), DebugInfo.CoordZ, DebugInfo.CoordX, DebugInfo.CoordY);
 	int32 DataLayerCount = DataLayers.Num();
 
-	const UDataLayerSubsystem* DataLayerSubsystem = UWorld::GetSubsystem<UDataLayerSubsystem>(GetOuterUWorldPartition()->GetWorld());
+	const UDataLayerSubsystem* DataLayerSubsystem = UWorld::GetSubsystem<UDataLayerSubsystem>(GetCellOwner()->GetOuterWorld());
 	TArray<const UDataLayerInstance*> DataLayerObjects;
 	if (DataLayerSubsystem && DataLayerCount > 0)
 	{
diff --git a/Engine/Source/Runtime/Engine/Private/WorldPartition/WorldPartitionRuntimeCellOwner.cpp b/Engine/Source/Runtime/Engine/Private/WorldPartition/WorldPartitionRuntimeCellOwner.cpp
new file mode 100644
--- /dev/null
+++ b/Engine/Source/Runtime/Engine/Private/WorldPartition/WorldPartitionRuntimeCellOwner.cpp
@@ -0,0 +1,7 @@
+// Copyright Epic Games, Inc. All Rights Reserved.
+
+#include "WorldPartition/WorldPartitionRuntimeCellOwner.h"
+
+UWorldPartitionRuntimeCellOwner::UWorldPartitionRuntimeCellOwner(const FObjectInitializer& ObjectInitializer)
+	: Super(ObjectInitializer)
+{}
\ No newline at end of file
diff --git a/Engine/Source/Runtime/Engine/Private/WorldPartition/WorldPartitionRuntimeLevelStreamingCell.cpp b/Engine/Source/Runtime/Engine/Private/WorldPartition/WorldPartitionRuntimeLevelStreamingCell.cpp
--- a/Engine/Source/Runtime/Engine/Private/WorldPartition/WorldPartitionRuntimeLevelStreamingCell.cpp
+++ b/Engine/Source/Runtime/Engine/Private/WorldPartition/WorldPartitionRuntimeLevelStreamingCell.cpp
@@ -113,9 +113,9 @@ UWorldPartitionLevelStreamingDynamic* UWorldPartitionRuntimeLevelStreamingCell::
 {
 	if (GetActorCount() > 0)
 	{
-		const UWorldPartition* WorldPartition = GetOuterUWorldPartition();
-		UWorld* OuterWorld = WorldPartition->GetTypedOuter<UWorld>();
-		UWorld* OwningWorld = WorldPartition->GetWorld();
+		const UWorldPartition* WorldPartition = GetCellOwner()->GetWorldPartition();
+		UWorld* OuterWorld = GetCellOwner()->GetOuterWorld();
+		UWorld* OwningWorld = GetCellOwner()->GetOwningWorld();
 		
 		const FName LevelStreamingName = FName(*FString::Printf(TEXT("WorldPartitionLevelStreaming_%s"), *GetName()));
 		
@@ -159,8 +159,8 @@ bool UWorldPartitionRuntimeLevelStreamingCell::PopulateGeneratorPackageForCook(T
 	{
 		FWorldPartitionLevelHelper::FPackageReferencer PackageReferencer;
 		const bool bLoadAsync = false;
-		UWorldPartition* WorldPartition = GetOuterUWorldPartition();
-		UWorld* OuterWorld = WorldPartition->GetTypedOuter<UWorld>();
+		UWorld* OuterWorld = GetCellOwner()->GetOuterWorld();
+		UWorldPartition* WorldPartition = OuterWorld->GetWorldPartition();
 
 		// Don't do SoftObjectPath remapping for PersistentLevel actors because references can end up in different cells
 		const bool bSoftObjectRemappingEnabled = false;
@@ -209,8 +209,8 @@ bool UWorldPartitionRuntimeLevelStreamingCell::PopulateGeneratedPackageForCook(U
 			return false;
 		}
 
-		UWorldPartition* WorldPartition = GetOuterUWorldPartition();
-		UWorld* OuterWorld = WorldPartition->GetTypedOuter<UWorld>();
+		UWorld* OuterWorld = GetCellOwner()->GetOuterWorld();
+		UWorldPartition* WorldPartition = OuterWorld->GetWorldPartition();
 
 		// Load cell Actors
 		FWorldPartitionLevelHelper::FPackageReferencer PackageReferencer;
@@ -238,9 +238,7 @@ int32 UWorldPartitionRuntimeLevelStreamingCell::GetActorCount() const
 
 FString UWorldPartitionRuntimeLevelStreamingCell::GetPackageNameToCreate() const
 {
-	const UWorldPartition* WorldPartition = GetOuterUWorldPartition();
-	UWorld* OuterWorld = WorldPartition->GetTypedOuter<UWorld>();
-	return UWorldPartitionLevelStreamingPolicy::GetCellPackagePath(GetFName(), OuterWorld);
+	return UWorldPartitionLevelStreamingPolicy::GetCellPackagePath(GetFName(), GetCellOwner()->GetOuterWorld());
 }
 
 void UWorldPartitionRuntimeLevelStreamingCell::DumpStateLog(FHierarchicalLogArchive& Ar)
@@ -278,8 +276,8 @@ UWorldPartitionLevelStreamingDynamic* UWorldPartitionRuntimeLevelStreamingCell::
 	if (LevelStreaming)
 	{
 		// Setup pre-created LevelStreaming's outer to the WorldPartition owning world
-		const UWorldPartition* WorldPartition = GetOuterUWorldPartition();
-		UWorld* OwningWorld = WorldPartition->GetWorld();
+		const UWorldPartition* WorldPartition = GetCellOwner()->GetWorldPartition();
+		UWorld* OwningWorld = GetCellOwner()->GetOwningWorld();
 		if (LevelStreaming->GetWorld() != OwningWorld)
 		{
 			LevelStreaming->Rename(nullptr, OwningWorld);
diff --git a/Engine/Source/Runtime/Engine/Private/WorldPartition/WorldPartitionRuntimeSpatialHash.cpp b/Engine/Source/Runtime/Engine/Private/WorldPartition/WorldPartitionRuntimeSpatialHash.cpp
--- a/Engine/Source/Runtime/Engine/Private/WorldPartition/WorldPartitionRuntimeSpatialHash.cpp
+++ b/Engine/Source/Runtime/Engine/Private/WorldPartition/WorldPartitionRuntimeSpatialHash.cpp
@@ -795,9 +795,8 @@ void UWorldPartitionRuntimeSpatialHash::DrawPreview() const
 URuntimeHashExternalStreamingObjectBase* UWorldPartitionRuntimeSpatialHash::StoreToExternalStreamingObject(UObject* StreamingObjectOuter, FName StreamingObjectName)
 {
 	check(!GetWorld()->IsGameWorld());
-
-	FName StreamingObjectUniqueName = MakeUniqueObjectName(StreamingObjectOuter, URuntimeSpatialHashExternalStreamingObject::StaticClass(), StreamingObjectName);
-	URuntimeSpatialHashExternalStreamingObject* StreamingObject = NewObject<URuntimeSpatialHashExternalStreamingObject>(StreamingObjectOuter, StreamingObjectUniqueName);
+	URuntimeSpatialHashExternalStreamingObject* StreamingObject = NewObject<URuntimeSpatialHashExternalStreamingObject>(StreamingObjectOuter, StreamingObjectName, RF_Public);
+	StreamingObject->Initialize(GetWorld(), GetTypedOuter<UWorld>());
 	StreamingObject->StreamingGrids = MoveTemp(StreamingGrids);
 
 	// Rename grids so the name of external streaming objects do not clashes with the main streaming grids.
@@ -1081,7 +1080,7 @@ bool UWorldPartitionRuntimeSpatialHash::CreateStreamingGrid(const FSpatialHashRu
 				verify(PartionedActors.GetCellGlobalCoords(FGridCellCoord(CellCoordX, CellCoordY, Level), CellGlobalCoords));
 				FString CellName = GetCellNameString(CurrentStreamingGrid.GridName, CellGlobalCoords, GridCellDataChunk.GetDataLayersID(), GridCellDataChunk.GetContentBundleID());
 
-				UWorldPartitionRuntimeSpatialHashCell* StreamingCell = NewObject<UWorldPartitionRuntimeSpatialHashCell>(WorldPartition, StreamingPolicy->GetRuntimeCellClass(), FName(CellName));
+				UWorldPartitionRuntimeSpatialHashCell* StreamingCell = NewObject<UWorldPartitionRuntimeSpatialHashCell>(this, StreamingPolicy->GetRuntimeCellClass(), FName(CellName));
 
 				int32 LayerCellIndex;
 				int32* LayerCellIndexPtr = GridLevel.LayerCellsMapping.Find(CellIndex);
diff --git a/Engine/Source/Runtime/Engine/Public/WorldPartition/WorldPartitionRuntimeCell.h b/Engine/Source/Runtime/Engine/Public/WorldPartition/WorldPartitionRuntimeCell.h
--- a/Engine/Source/Runtime/Engine/Public/WorldPartition/WorldPartitionRuntimeCell.h
+++ b/Engine/Source/Runtime/Engine/Public/WorldPartition/WorldPartitionRuntimeCell.h
@@ -8,6 +8,7 @@
 #include "WorldPartition/WorldPartitionStreamingSource.h"
 #include "WorldPartition/WorldPartitionActorContainerID.h"
 #include "WorldPartition/WorldPartitionRuntimeCellInterface.h"
+#include "WorldPartition/WorldPartitionRuntimeCellOwner.h"
 #include "ProfilingDebugging/ProfilingHelpers.h"
 #include "Misc/HierarchicalLogArchive.h"
 #include "Algo/AnyOf.h"
@@ -146,7 +147,7 @@ static_assert(EWorldPartitionRuntimeCellState::Unloaded < EWorldPartitionRuntime
 /**
  * Represents a PIE/Game streaming cell which points to external actor/data chunk packages
  */
-UCLASS(Abstract, Within = WorldPartition)
+UCLASS(Abstract)
 class ENGINE_API UWorldPartitionRuntimeCell : public UObject, public IWorldPartitionCell
 {
 	GENERATED_UCLASS_BODY()
@@ -173,6 +174,8 @@ class ENGINE_API UWorldPartitionRuntimeCell : public UObject, public IWorldParti
 	virtual FName GetGridName() const { return DebugInfo.GridName; }
 	virtual FGuid const& GetContentBundleID() const { return ContentBundleID; }
 
+	IWorldPartitionRuntimeCellOwner* GetCellOwner() const { return CastChecked<IWorldPartitionRuntimeCellOwner>(GetOuter()); }
+
 	/** Caches information on streaming source that will be used later on to sort cell. */
 	bool ShouldResetStreamingSourceInfo() const;
 	virtual void ResetStreamingSourceInfo() const;
diff --git a/Engine/Source/Runtime/Engine/Public/WorldPartition/WorldPartitionRuntimeCellOwner.h b/Engine/Source/Runtime/Engine/Public/WorldPartition/WorldPartitionRuntimeCellOwner.h
new file mode 100644
--- /dev/null
+++ b/Engine/Source/Runtime/Engine/Public/WorldPartition/WorldPartitionRuntimeCellOwner.h
@@ -0,0 +1,30 @@
+// Copyright Epic Games, Inc. All Rights Reserved.
+#pragma once
+
+#include "CoreMinimal.h"
+#include "UObject/ObjectMacros.h"
+#include "UObject/Interface.h"
+
+#include "Engine/World.h"
+
+#include "WorldPartitionRuntimeCellOwner.generated.h"
+
+class UWorld;
+class UWorldPartition;
+
+UINTERFACE()
+class ENGINE_API UWorldPartitionRuntimeCellOwner : public UInterface
+{
+	GENERATED_UINTERFACE_BODY()
+};
+
+class IWorldPartitionRuntimeCellOwner
+{
+	GENERATED_IINTERFACE_BODY()
+
+public:
+	virtual UWorld* GetOwningWorld() const = 0;
+	virtual UWorld* GetOuterWorld() const = 0;
+
+	UWorldPartition* GetWorldPartition() const { return GetOuterWorld()->GetWorldPartition(); }
+};
diff --git a/Engine/Source/Runtime/Engine/Public/WorldPartition/WorldPartitionRuntimeHash.h b/Engine/Source/Runtime/Engine/Public/WorldPartition/WorldPartitionRuntimeHash.h
--- a/Engine/Source/Runtime/Engine/Public/WorldPartition/WorldPartitionRuntimeHash.h
+++ b/Engine/Source/Runtime/Engine/Public/WorldPartition/WorldPartitionRuntimeHash.h
@@ -9,6 +9,7 @@
 #include "WorldPartition/ActorDescList.h"
 #include "WorldPartition/WorldPartitionHandle.h"
 #include "WorldPartition/WorldPartitionRuntimeCell.h"
+#include "WorldPartition/WorldPartitionRuntimeCellOwner.h"
 #include "WorldPartition/WorldPartitionActorDescViewProxy.h"
 #include "WorldPartition/WorldPartitionStreamingGeneration.h"
 #if WITH_EDITOR
@@ -30,16 +31,32 @@ enum class EWorldPartitionStreamingPerformance : uint8
 };
 
 UCLASS(Abstract)
-class ENGINE_API URuntimeHashExternalStreamingObjectBase : public UObject
+class ENGINE_API URuntimeHashExternalStreamingObjectBase : public UObject, public IWorldPartitionRuntimeCellOwner
 {
 	GENERATED_BODY()
 
 public:
+	virtual void Initialize(UWorld* InOwningWorld, UWorld* InOuterWorld) 
+	{ 
+		OwningWorld = InOwningWorld;  
+		OuterWorld = InOuterWorld;
+	}
 
+	//~ Begin IWorldPartitionRuntimeCellOwner Interface
+	virtual UWorld* GetOwningWorld() const override { return OwningWorld.Get(); }
+	virtual UWorld* GetOuterWorld() const override { return OuterWorld.Get(); }
+	//~ End IWorldPartitionRuntimeCellOwner Interface
+
+private:
+	UPROPERTY();
+	TSoftObjectPtr<UWorld> OwningWorld;
+
+	UPROPERTY();
+	TSoftObjectPtr<UWorld> OuterWorld;
 };
 
 UCLASS(Abstract, Config=Engine, AutoExpandCategories=(WorldPartition), Within = WorldPartition)
-class ENGINE_API UWorldPartitionRuntimeHash : public UObject
+class ENGINE_API UWorldPartitionRuntimeHash : public UObject, public IWorldPartitionRuntimeCellOwner
 {
 	GENERATED_UCLASS_BODY()
 
@@ -88,6 +105,11 @@ public:
 	UE_DEPRECATED(5.1, "GetStreamingCells is deprecated, use ForEachStreamingCells instead.")
 	bool GetStreamingCells(const TArray<FWorldPartitionStreamingSource>& Sources, UWorldPartitionRuntimeHash::FStreamingSourceCells& OutActivateCells, UWorldPartitionRuntimeHash::FStreamingSourceCells& OutLoadCells) const;
 
+	//~ Begin IWorldPartitionRuntimeCellOwner Interface
+	virtual UWorld* GetOwningWorld() const override { return GetWorld(); }
+	virtual UWorld* GetOuterWorld() const override { return GetTypedOuter<UWorld>(); }
+	//~ End IWorldPartitionRuntimeCellOwner Interface
+
 	// Streaming interface
 	virtual void ForEachStreamingCells(TFunctionRef<bool(const UWorldPartitionRuntimeCell*)> Func) const {}
 	virtual void ForEachStreamingCellsQuery(const FWorldPartitionStreamingQuerySource& QuerySource, TFunctionRef<bool(const UWorldPartitionRuntimeCell*)> Func) const {}
@@ -96,6 +118,7 @@ public:
 	bool IsCellRelevantFor(bool bClientOnlyVisible) const;
 	EWorldPartitionStreamingPerformance GetStreamingPerformance(const TSet<const UWorldPartitionRuntimeCell*>& CellToActivate) const;
 
+	virtual URuntimeHashExternalStreamingObjectBase* StoreToExternalStreamingObject(UObject* Outer) { return nullptr; }
 	virtual bool InjectExternalStreamingObject(URuntimeHashExternalStreamingObjectBase* ExternalStreamingObject) { return false; }
 	virtual bool RemoveExternalStreamingObject(URuntimeHashExternalStreamingObjectBase* ExternalStreamingObject) { return false; }
 
diff --git a/Engine/Source/Runtime/Engine/Public/WorldPartition/WorldPartitionRuntimeSpatialHashCell.h b/Engine/Source/Runtime/Engine/Public/WorldPartition/WorldPartitionRuntimeSpatialHashCell.h
--- a/Engine/Source/Runtime/Engine/Public/WorldPartition/WorldPartitionRuntimeSpatialHashCell.h
+++ b/Engine/Source/Runtime/Engine/Public/WorldPartition/WorldPartitionRuntimeSpatialHashCell.h
@@ -7,7 +7,7 @@
 
 class UActorContainer;
 
-UCLASS(Abstract, Within = WorldPartition)
+UCLASS(Abstract)
 class UWorldPartitionRuntimeSpatialHashCell : public UWorldPartitionRuntimeCell
 {
 	GENERATED_UCLASS_BODY()
