UWorld - change bIsWorldInitialized to be reset to false when CleanupWorld is called. Delete the old bInitializedAndNeedsCleanup since bIsWorldInitialized now means the same thing. Add the new bHasEverBeenInitialized to provide the information bIsWorldInitialized used to provide (used in FinishDestroy and in a few other places).
This fixes cooking with EditorDomain, which can save the package multiple times and thereby call InitWorld/CleanupWorld multiple times.
#rnx

diff --git a/Engine/Source/Editor/UnrealEd/Private/EditorEngine.cpp b/Engine/Source/Editor/UnrealEd/Private/EditorEngine.cpp
--- a/Engine/Source/Editor/UnrealEd/Private/EditorEngine.cpp
+++ b/Engine/Source/Editor/UnrealEd/Private/EditorEngine.cpp
@@ -1259,7 +1259,7 @@ void UEditorEngine::PreExit()
 		// Cleanup worlds that were initialized through UEditorEngine::InitializeNewlyCreatedInactiveWorld before final GC Purge
 		for (TObjectIterator<UWorld> It; It; ++It)
 		{
-			if (UWorld* World = *It; World && World->WorldType == EWorldType::Inactive && World->IsInitializedAndNeedsCleanup())
+			if (UWorld* World = *It; World && World->WorldType == EWorldType::Inactive && World->IsInitialized())
 			{
 				// GWorld shouldn't be an Inactive World
 				check(World != GWorld);
diff --git a/Engine/Source/Editor/UnrealEd/Private/EditorServer.cpp b/Engine/Source/Editor/UnrealEd/Private/EditorServer.cpp
--- a/Engine/Source/Editor/UnrealEd/Private/EditorServer.cpp
+++ b/Engine/Source/Editor/UnrealEd/Private/EditorServer.cpp
@@ -2565,7 +2565,7 @@ bool UEditorEngine::Map_Load(const TCHAR* Str, FOutputDevice& Ar)
 					// All other worlds need to be reloaded from disk, as a world is only initialized correctly as part of the level loading process.
 					auto IsWorldValidForReuse = [](UWorld* WorldToConsider)
 					{
-						return !WorldToConsider->bIsWorldInitialized;
+						return !WorldToConsider->HasEverBeenInitialized();
 					};
 
 					// If we are loading the same world again (reloading) then we must not specify that we want to keep this world in memory.
@@ -2770,11 +2770,11 @@ bool UEditorEngine::Map_Load(const TCHAR* Str, FOutputDevice& Ar)
 				FeatureLevelIndex = FMath::Clamp(FeatureLevelIndex, 0, (int32)ERHIFeatureLevel::Num);
 				ERHIFeatureLevel::Type FeatureLevel = (ERHIFeatureLevel::Type)FeatureLevelIndex;
 
-				if (World->bIsWorldInitialized)
+				if (World->HasEverBeenInitialized())
 				{
 					// We do not handle reinitializing. If the World was initialized during Load, we require that it is still fully initialized
 					// This should have been guaranteed by UWorld::KeepInitializedDuringLoadTag
-					check(World->IsInitializedAndNeedsCleanup());
+					check(World->IsInitialized());
 					World->ChangeFeatureLevel(FeatureLevel);
 				}
 				else
diff --git a/Engine/Source/Runtime/Engine/Classes/Engine/World.h b/Engine/Source/Runtime/Engine/Classes/Engine/World.h
--- a/Engine/Source/Runtime/Engine/Classes/Engine/World.h
+++ b/Engine/Source/Runtime/Engine/Classes/Engine/World.h
@@ -1272,7 +1272,7 @@ public:
 	 */
 	uint8 bPostTickComponentUpdate:1;
 
-	/** Whether world object has been initialized via Init()																	*/
+	/** Whether world object has been initialized via Init and has not yet had CleanupWorld called								*/
 	uint8 bIsWorldInitialized:1;
 
 	/** Is level streaming currently frozen?																					*/
@@ -1388,8 +1388,8 @@ private:
 	/** Is there at least one material parameter collection instance waiting for a deferred update?								*/
 	uint8 bMaterialParameterCollectionInstanceNeedsDeferredUpdate : 1;
 
-	/** Whether world object has been initialized via Init and has not yet had CleanupWorld called								*/
-	uint8 bInitializedAndNeedsCleanup : 1;
+	/** Whether InitWorld was ever called on this world since its creation. Not cleared to false during CleanupWorld			*/
+	uint8 bHasEverBeenInitialized: 1;
 
 	/** Whether the world is currently in a BlockTillLevelStreamingCompleted() call */
 	uint32 IsInBlockTillLevelStreamingCompleted;
@@ -3129,12 +3129,13 @@ public:
 	 * operations that need to call InitWorld during the Load of the World's package should break that rule and not call CleanupWorld.
 	 */
 	static const FName KeepInitializedDuringLoadTag;
-	/**
-	 * Returns whether InitWorld has been called without yet calling CleanupWorld.
-	 * Note this is different from bIsWorldInitialized flag, which is set to true InitWorld and never set back to false.
-	 */
-	bool IsInitializedAndNeedsCleanup() const { return bInitializedAndNeedsCleanup; }
+	UE_DEPRECATED(5.2, "Call IsInitialized instead.")
+	bool IsInitializedAndNeedsCleanup() const { return bIsWorldInitialized; }
+	/** Returns whether InitWorld has ever been called since this World was created.  */
+	bool HasEverBeenInitialized() const { return bHasEverBeenInitialized; }
 #endif
+	/** Returns whether InitWorld has been called without yet calling CleanupWorld.  */
+	bool IsInitialized() const { return bIsWorldInitialized; }
 
 	/**
 	 * Initializes a newly created world.
diff --git a/Engine/Source/Runtime/Engine/Private/Tests/WorldPartition/WorldPartitionHandleTests.cpp b/Engine/Source/Runtime/Engine/Private/Tests/WorldPartition/WorldPartitionHandleTests.cpp
--- a/Engine/Source/Runtime/Engine/Private/Tests/WorldPartition/WorldPartitionHandleTests.cpp
+++ b/Engine/Source/Runtime/Engine/Private/Tests/WorldPartition/WorldPartitionHandleTests.cpp
@@ -243,7 +243,7 @@ namespace WorldPartitionTests
 			UPackage* Package = FindPackage(NULL, *ActorDescContainer->GetContainerPackage().ToString());
 			check(Package);
 			UWorld* World = UWorld::FindWorldInPackage(Package);
-			check(World->IsInitializedAndNeedsCleanup());
+			check(World->IsInitialized());
 			World->CleanupWorld();
 			
 			CollectGarbage(GARBAGE_COLLECTION_KEEPFLAGS);
diff --git a/Engine/Source/Runtime/Engine/Private/World.cpp b/Engine/Source/Runtime/Engine/Private/World.cpp
--- a/Engine/Source/Runtime/Engine/Private/World.cpp
+++ b/Engine/Source/Runtime/Engine/Private/World.cpp
@@ -394,12 +394,13 @@ UWorld::UWorld( const FObjectInitializer& ObjectInitializer )
 , bAllowDeferredPhysicsStateCreation(false)
 , FeatureLevel(GMaxRHIFeatureLevel)
 , bIsBuilt(false)
+, bIsWorldInitialized(false)
 #if WITH_EDITOR
 , bDebugFrameStepExecutedThisFrame(false)
 , bToggledBetweenPIEandSIEThisFrame(false)
 #endif
 , bShouldTick(true)
-, bInitializedAndNeedsCleanup(false)
+, bHasEverBeenInitialized(false)
 , ActiveLevelCollectionIndex(INDEX_NONE)
 , AudioDeviceHandle()
 #if WITH_EDITOR
@@ -1021,7 +1022,7 @@ void UWorld::BeginDestroy()
 #if WITH_EDITOR
 	// Make sure we catch worlds that where initialized through UEditorEngine::OnAssetLoaded/OnAssetCreated
 	// (This can happen if World was loaded, then its RF_Standalone flag was removed and a GC happened) 
-	if (WorldType == EWorldType::Inactive && IsInitializedAndNeedsCleanup())
+	if (WorldType == EWorldType::Inactive && IsInitialized())
 	{
 		CleanupWorld();
 	}
@@ -1082,14 +1083,14 @@ void UWorld::ReleasePhysicsScene()
 
 void UWorld::FinishDestroy()
 {
-	if (bInitializedAndNeedsCleanup)
+	if (bIsWorldInitialized)
 	{
 		UE_LOG(LogWorld, Warning, TEXT("UWorld::FinishDestroy called after InitWorld without calling CleanupWorld first."));
 	}
 
 	// Avoid cleanup if the world hasn't been initialized, e.g., the default object or a world object that got loaded
 	// due to level streaming.
-	if (bIsWorldInitialized)
+	if (bHasEverBeenInitialized)
 	{
 		FWorldDelegates::OnPreWorldFinishDestroy.Broadcast(this);
 
@@ -1950,7 +1951,7 @@ void UWorld::InitWorld(const InitializationValues IVS)
 
 	// We're initialized now.
 	bIsWorldInitialized = true;
-	bInitializedAndNeedsCleanup = true;
+	bHasEverBeenInitialized = true;
 
 	// Call the general post initialization delegates
 	FWorldDelegates::OnPostWorldInitialization.Broadcast(this, IVS);
@@ -5019,7 +5020,7 @@ bool UWorld::IsNavigationRebuilt() const
 void UWorld::CleanupWorld(bool bSessionEnded, bool bCleanupResources, UWorld* NewWorld)
 {
     CleanupWorldGlobalTag++;
-	if (!bInitializedAndNeedsCleanup)
+	if (!bIsWorldInitialized)
 	{
 		// Only issue the warning for the TopLevelWorld. It is currently valid to call CleanupWorld on the UWorld of
 		// streaming sublevels, and they never call InitWorld. (this is done by PrivateDestroyLevel when removing a
@@ -5029,7 +5030,7 @@ void UWorld::CleanupWorld(bool bSessionEnded, bool bCleanupResources, UWorld* Ne
 	}
 	const bool bWorldChanged = NewWorld != this;
 	CleanupWorldInternal(bSessionEnded, bCleanupResources, bWorldChanged);
-	bInitializedAndNeedsCleanup = false;
+	bIsWorldInitialized = false;
 }
 
 void UWorld::CleanupWorldInternal(bool bSessionEnded, bool bCleanupResources, bool bWorldChanged)
