Chaos - fix knockon RBAN joint enable/disable, broken by fix to joint inital active state

diff --git a/Engine/Source/Runtime/Engine/Private/PhysicsEngine/ImmediatePhysics/ImmediatePhysicsChaos/ImmediatePhysicsActorHandle_Chaos.cpp b/Engine/Source/Runtime/Engine/Private/PhysicsEngine/ImmediatePhysics/ImmediatePhysicsChaos/ImmediatePhysicsActorHandle_Chaos.cpp
--- a/Engine/Source/Runtime/Engine/Private/PhysicsEngine/ImmediatePhysics/ImmediatePhysicsChaos/ImmediatePhysicsActorHandle_Chaos.cpp
+++ b/Engine/Source/Runtime/Engine/Private/PhysicsEngine/ImmediatePhysics/ImmediatePhysicsChaos/ImmediatePhysicsActorHandle_Chaos.cpp
@@ -319,10 +319,16 @@ namespace ImmediatePhysics_Chaos
 		return ParticleHandle;
 	}
 
+	bool FActorHandle::GetEnabled() const
+	{
+		const Chaos::FConstGenericParticleHandle Particle = ParticleHandle;
+		return !Particle->Disabled();
+	}
+
 	void FActorHandle::SetEnabled(bool bEnabled)
 	{
-		auto* Dynamic = ParticleHandle->CastToRigidParticle();
-		if(Dynamic && Dynamic->ObjectState() == Chaos::EObjectStateType::Dynamic)
+		Chaos::FPBDRigidParticleHandle* Dynamic = ParticleHandle->CastToRigidParticle();
+		if (Dynamic && Dynamic->ObjectState() == Chaos::EObjectStateType::Dynamic)
 		{
 			Dynamic->Disabled() = !bEnabled;
 		}
diff --git a/Engine/Source/Runtime/Engine/Private/PhysicsEngine/ImmediatePhysics/ImmediatePhysicsChaos/ImmediatePhysicsSimulation_Chaos.cpp b/Engine/Source/Runtime/Engine/Private/PhysicsEngine/ImmediatePhysics/ImmediatePhysicsChaos/ImmediatePhysicsSimulation_Chaos.cpp
--- a/Engine/Source/Runtime/Engine/Private/PhysicsEngine/ImmediatePhysics/ImmediatePhysicsChaos/ImmediatePhysicsSimulation_Chaos.cpp
+++ b/Engine/Source/Runtime/Engine/Private/PhysicsEngine/ImmediatePhysics/ImmediatePhysicsChaos/ImmediatePhysicsSimulation_Chaos.cpp
@@ -707,6 +707,18 @@ namespace ImmediatePhysics_Chaos
 		}
 	}
 
+	void FSimulation::EnableDisableJoints()
+	{
+		for (FJointHandle* Joint : Implementation->JointHandles)
+		{
+			const bool bActorEnabled0 = Joint->GetActorHandles()[0]->GetEnabled();
+			const bool bActorEnabled1 = Joint->GetActorHandles()[1]->GetEnabled();
+			const bool bJointEnabled = bActorEnabled0 && bActorEnabled1;
+			Joint->GetConstraint()->SetConstraintEnabled(bJointEnabled);
+		}
+	}
+
+
 	void FSimulation::InitSimulationSpace(
 		const FTransform& Transform)
 	{
@@ -908,6 +920,7 @@ namespace ImmediatePhysics_Chaos
 		{
 			PackCollidingPairs();
 			UpdateActivePotentiallyCollidingPairs();
+			EnableDisableJoints();
 			Implementation->bActorsDirty = false;
 		}
 
diff --git a/Engine/Source/Runtime/Engine/Public/Physics/ImmediatePhysics/ImmediatePhysicsChaos/ImmediatePhysicsActorHandle_Chaos.h b/Engine/Source/Runtime/Engine/Public/Physics/ImmediatePhysics/ImmediatePhysicsChaos/ImmediatePhysicsActorHandle_Chaos.h
--- a/Engine/Source/Runtime/Engine/Public/Physics/ImmediatePhysics/ImmediatePhysicsChaos/ImmediatePhysicsActorHandle_Chaos.h
+++ b/Engine/Source/Runtime/Engine/Public/Physics/ImmediatePhysics/ImmediatePhysicsChaos/ImmediatePhysicsActorHandle_Chaos.h
@@ -18,6 +18,8 @@ namespace ImmediatePhysics_Chaos
 		void SetName(const FName& InName) { Name = InName; }
 		const FName& GetName() const { return Name; }
 
+		bool GetEnabled() const;
+
 		void SetEnabled(bool bEnabled);
 
 		/** Sets the world transform, zeroes velocity, etc.*/
diff --git a/Engine/Source/Runtime/Engine/Public/Physics/ImmediatePhysics/ImmediatePhysicsChaos/ImmediatePhysicsSimulation_Chaos.h b/Engine/Source/Runtime/Engine/Public/Physics/ImmediatePhysics/ImmediatePhysicsChaos/ImmediatePhysicsSimulation_Chaos.h
--- a/Engine/Source/Runtime/Engine/Public/Physics/ImmediatePhysics/ImmediatePhysicsChaos/ImmediatePhysicsSimulation_Chaos.h
+++ b/Engine/Source/Runtime/Engine/Public/Physics/ImmediatePhysics/ImmediatePhysicsChaos/ImmediatePhysicsSimulation_Chaos.h
@@ -94,6 +94,7 @@ namespace ImmediatePhysics_Chaos
 		void UpdateInertiaConditioning(const FVector& Gravity);
 		void PackCollidingPairs();
 		void UpdateActivePotentiallyCollidingPairs();
+		void EnableDisableJoints();
 		FReal UpdateStepTime(const FReal DeltaTime, const FReal MaxStepTime);
 
 		void UpdateStatCounters();
