Improvements to tracking/handling of replay/streamer errors

diff --git a/Engine/Source/Runtime/Engine/Classes/Engine/DemoNetDriver.h b/Engine/Source/Runtime/Engine/Classes/Engine/DemoNetDriver.h
--- a/Engine/Source/Runtime/Engine/Classes/Engine/DemoNetDriver.h
+++ b/Engine/Source/Runtime/Engine/Classes/Engine/DemoNetDriver.h
@@ -16,6 +16,8 @@
 #include "NetworkReplayStreaming.h"
 #include "Engine/DemoNetConnection.h"
 #include "Net/RepLayout.h"
+#include "Net/Core/Connection/NetResult.h"
+#include "Net/ReplayResult.h"
 #include "Templates/Atomic.h"
 #include "Net/UnrealNetwork.h"
 #include "ReplayHelper.h"
@@ -30,7 +32,11 @@ DECLARE_MULTICAST_DELEGATE(FOnGotoTimeMCDelegate);
 DECLARE_DELEGATE_OneParam(FOnGotoTimeDelegate, const bool /* bWasSuccessful */);
 
 DECLARE_MULTICAST_DELEGATE_OneParam(FOnDemoStartedDelegate, UDemoNetDriver* /* DemoNetDriver */);
+
+PRAGMA_DISABLE_DEPRECATION_WARNINGS
+UE_DEPRECATED(5.2, "No longer used")
 DECLARE_MULTICAST_DELEGATE_TwoParams(FOnDemoFailedToStartDelegate, UDemoNetDriver* /* DemoNetDriver */, EDemoPlayFailure::Type /* FailureType*/);
+PRAGMA_ENABLE_DEPRECATION_WARNINGS
 
 DECLARE_MULTICAST_DELEGATE(FOnDemoFinishPlaybackDelegate);
 DECLARE_MULTICAST_DELEGATE(FOnDemoFinishRecordingDelegate);
@@ -261,7 +267,7 @@ protected:
 	uint8 bIsWaitingForHeaderDownload : 1;
 	uint8 bIsWaitingForStream : 1;
 
-	TOptional<EDemoRecordFailure> PendingRecordFailure;
+	TOptional<UE::Net::TNetResult<EReplayResult>> PendingRecordFailure;
 
 private:
 	TArray<FNetworkGUID> NonQueuedGUIDsForScrubbing;
@@ -694,8 +700,13 @@ protected:
 
 	bool DemoReplicateActor(AActor* Actor, UNetConnection* Connection, bool bMustReplicate);
 
+	PRAGMA_DISABLE_DEPRECATION_WARNINGS
+	UE_DEPRECATED(5.2, "Please use NotifyDemoPlaybackError instead")
 	void NotifyDemoPlaybackFailure(EDemoPlayFailure::Type FailureType);
-	void NotifyDemoRecordFailure(EDemoRecordFailure FailureType);
+	PRAGMA_ENABLE_DEPRECATION_WARNINGS
+	
+	void NotifyDemoPlaybackError(const UE::Net::TNetResult<EReplayResult>& Result);
+	void NotifyDemoRecordFailure(const UE::Net::TNetResult<EReplayResult>& Result);
 
 	TArray<FQueuedDemoPacket> QueuedPacketsBeforeTravel;
 
diff --git a/Engine/Source/Runtime/Engine/Classes/Engine/EngineBaseTypes.h b/Engine/Source/Runtime/Engine/Classes/Engine/EngineBaseTypes.h
--- a/Engine/Source/Runtime/Engine/Classes/Engine/EngineBaseTypes.h
+++ b/Engine/Source/Runtime/Engine/Classes/Engine/EngineBaseTypes.h
@@ -629,7 +629,7 @@ enum ETravelType
 UENUM(BlueprintType)
 namespace EDemoPlayFailure
 {
-	enum Type
+	enum UE_DEPRECATED(5.2, "No longer used in favor of EReplayResult") Type
 	{
 		/** A Generic failure. */
 		Generic,
@@ -654,6 +654,8 @@ namespace EDemoPlayFailure
 
 namespace EDemoPlayFailure
 {
+	PRAGMA_DISABLE_DEPRECATION_WARNINGS
+	UE_DEPRECATED(5.2, "EDemoPlayFailure is now deprecated")
 	inline const TCHAR* ToString(EDemoPlayFailure::Type FailureType)
 	{
 		switch (FailureType)
@@ -680,6 +682,7 @@ namespace EDemoPlayFailure
 
 		return TEXT("Unknown EDemoPlayFailure error occurred.");
 	}
+	PRAGMA_ENABLE_DEPRECATION_WARNINGS
 }
 
 //URL structure.
diff --git a/Engine/Source/Runtime/Engine/Classes/Engine/GameInstance.h b/Engine/Source/Runtime/Engine/Classes/Engine/GameInstance.h
--- a/Engine/Source/Runtime/Engine/Classes/Engine/GameInstance.h
+++ b/Engine/Source/Runtime/Engine/Classes/Engine/GameInstance.h
@@ -429,11 +429,16 @@ public:
 
 	void CleanupGameViewport();
 
+	PRAGMA_DISABLE_DEPRECATION_WARNINGS
 	/** Called when demo playback fails for any reason */
+	UE_DEPRECATED(5.2, "Now takes a EReplayResult instead.")
 	virtual void HandleDemoPlaybackFailure(EDemoPlayFailure::Type FailureType, const FString& ErrorString = TEXT("")) { }
+	PRAGMA_ENABLE_DEPRECATION_WARNINGS
+
+	virtual void HandleDemoPlaybackFailure(const UE::Net::TNetResult<EReplayResult>& Result) { }
 
 	/** Called when demo recording fails for any reason */
-	virtual void HandleDemoRecordFailure(EDemoRecordFailure FailureType) { }
+	virtual void HandleDemoRecordFailure(const UE::Net::TNetResult<EReplayResult>& Result) { }
 
 	/** This gets called when the player scrubs in a replay to a different level */
 	virtual void OnSeamlessTravelDuringReplay() { }
diff --git a/Engine/Source/Runtime/Engine/Private/DemoNetDriver.cpp b/Engine/Source/Runtime/Engine/Private/DemoNetDriver.cpp
--- a/Engine/Source/Runtime/Engine/Private/DemoNetDriver.cpp
+++ b/Engine/Source/Runtime/Engine/Private/DemoNetDriver.cpp
@@ -856,8 +856,8 @@ bool UDemoNetDriver::InitBase(bool bInitAsClient, FNetworkNotify* InNotify, cons
 		ReplayHelper.CheckpointSaveMaxMSPerFrame = CheckpointSaveMaxMSPerFrame;
 
 		// if the helper encounters an error, stop the presses
-		ReplayHelper.OnReplayRecordError.AddUObject(this, &UDemoNetDriver::StopDemo);
-		ReplayHelper.OnReplayPlaybackError.AddUObject(this, &UDemoNetDriver::NotifyDemoPlaybackFailure);
+		ReplayHelper.OnReplayRecordError.BindUObject(this, &UDemoNetDriver::NotifyDemoRecordFailure);
+		ReplayHelper.OnReplayPlaybackError.BindUObject(this, &UDemoNetDriver::NotifyDemoPlaybackError);
 
 		return true;
 	}
@@ -879,8 +879,8 @@ void UDemoNetDriver::FinishDestroy()
 	CleanUpSplitscreenConnections(true);
 	FCoreUObjectDelegates::PostLoadMapWithWorld.RemoveAll(this);
 
-	ReplayHelper.OnReplayRecordError.RemoveAll(this);
-	ReplayHelper.OnReplayPlaybackError.RemoveAll(this);
+	ReplayHelper.OnReplayRecordError.Unbind();
+	ReplayHelper.OnReplayPlaybackError.Unbind();
 	
 	if (World)
 	{
@@ -917,22 +917,12 @@ void UDemoNetDriver::ResetDemoState()
 
 bool UDemoNetDriver::InitConnect(FNetworkNotify* InNotify, const FURL& ConnectURL, FString& Error)
 {
-	if (World == nullptr)
-	{
-		UE_LOG(LogDemo, Error, TEXT("World == nullptr"));
-		return false;
-	}
-
-	if (World->GetGameInstance() == nullptr)
-	{
-		UE_LOG(LogDemo, Error, TEXT("World->GetGameInstance() == nullptr"));
-		return false;
-	}
+	ReplayHelper.bRecording = false;
 
 	// handle default initialization
-	if (!InitBase(true, InNotify, ConnectURL, false, Error))
+	if (!World || !World->GetGameInstance() || !InitBase(true, InNotify, ConnectURL, false, Error))
 	{
-		World->GetGameInstance()->HandleDemoPlaybackFailure(EDemoPlayFailure::InitBase, FString(TEXT("InitBase FAILED")));
+		ReplayHelper.NotifyReplayError(EReplayResult::InitConnect);
 		return false;
 	}
 
@@ -1049,7 +1039,7 @@ bool UDemoNetDriver::InitConnectInternal(FString& Error)
 
 			if (!GEngine->MakeSureMapNameIsValid(LocalDemoURL.Map))
 			{
-				NotifyDemoPlaybackFailure(EDemoPlayFailure::LoadMap);
+				ReplayHelper.NotifyReplayError(EReplayResult::LoadMap);
 				return false;
 			}
 
@@ -1057,11 +1047,9 @@ bool UDemoNetDriver::InitConnectInternal(FString& Error)
 
 			if (WorldContext == nullptr)
 			{
-				UGameInstance* GameInstance = World->GetGameInstance();
-
 				Error = FString::Printf(TEXT("No world context"));
 				UE_LOG(LogDemo, Error, TEXT("UDemoNetDriver::InitConnect: %s"), *Error);
-				GameInstance->HandleDemoPlaybackFailure(EDemoPlayFailure::Generic, FString(TEXT("No world context")));
+				ReplayHelper.NotifyReplayError(EReplayResult::InitConnect);
 				return false;
 			}
 
@@ -1089,6 +1077,8 @@ bool UDemoNetDriver::InitConnectInternal(FString& Error)
 
 bool UDemoNetDriver::InitListen(FNetworkNotify* InNotify, FURL& ListenURL, bool bReuseAddressAndPort, FString& Error)
 {
+	ReplayHelper.bRecording = true;
+
 	if (!InitBase(false, InNotify, ListenURL, bReuseAddressAndPort, Error))
 	{
 		return false;
@@ -1271,7 +1261,7 @@ void UDemoNetDriver::PostTickFlush()
 		{
 			if (UGameInstance* GameInstance = World->GetGameInstance())
 			{
-				GameInstance->HandleDemoRecordFailure(PendingRecordFailure.GetValue());
+				GameInstance->HandleDemoRecordFailure(PendingRecordFailure.GetValue().GetResult());
 			}
 		}
 
@@ -1317,10 +1307,9 @@ void UDemoNetDriver::TickFlushInternal(float DeltaSeconds)
 
 	TSharedPtr<INetworkReplayStreamer> Streamer = GetReplayStreamer();
 
-	if (Streamer->GetLastError() != ENetworkReplayError::None)
+	if (UE::Net::EHandleNetResult::Handled == Streamer->HandleLastError(ReplayHelper.ResultManager))
 	{
-		UE_LOG(LogDemo, Error, TEXT("UDemoNetDriver::TickFlush: ReplayStreamer ERROR: %s"), ENetworkReplayError::ToString(Streamer->GetLastError()));
-		NotifyDemoRecordFailure(EDemoRecordFailure::Generic);
+		ReplayHelper.NotifyReplayError(EReplayResult::StreamerError);
 		return;
 	}
 
@@ -1334,7 +1323,7 @@ void UDemoNetDriver::TickFlushInternal(float DeltaSeconds)
 	if (FileAr == nullptr)
 	{
 		UE_LOG(LogDemo, Error, TEXT("UDemoNetDriver::TickFlush: FileAr == nullptr"));
-		NotifyDemoRecordFailure(EDemoRecordFailure::StreamingArchive);
+		ReplayHelper.NotifyReplayError(EReplayResult::MissingArchive);
 		return;
 	}
 
@@ -1398,19 +1387,17 @@ void UDemoNetDriver::TickDispatch(float DeltaSeconds)
 		return;
 	}
 
-	if (GetReplayStreamer()->GetLastError() != ENetworkReplayError::None)
+	if (UE::Net::EHandleNetResult::Handled == GetReplayStreamer()->HandleLastError(ReplayHelper.ResultManager))
 	{
-		UE_LOG(LogDemo, Error, TEXT("UDemoNetDriver::TickDispatch: ReplayStreamer ERROR: %s"), ENetworkReplayError::ToString(GetReplayStreamer()->GetLastError()));
-		NotifyDemoPlaybackFailure(EDemoPlayFailure::ReplayStreamerInternal);
+		ReplayHelper.NotifyReplayError(EReplayResult::StreamerError);
 		return;
 	}
 
 	FArchive* FileAr = GetReplayStreamer()->GetStreamingArchive();
-
 	if (FileAr == nullptr)
 	{
 		UE_LOG(LogDemo, Error, TEXT("UDemoNetDriver::TickDispatch: FileAr == nullptr"));
-		NotifyDemoPlaybackFailure(EDemoPlayFailure::ReplayStreamerInternal);
+		ReplayHelper.NotifyReplayError(EReplayResult::MissingArchive);
 		return;
 	}
 
@@ -2457,17 +2444,20 @@ void UDemoNetDriver::OnRefreshHeaderCompletePrivate(const FDownloadHeaderResult&
 			}
 			else
 			{
-				World->GetGameInstance()->HandleDemoPlaybackFailure(EDemoPlayFailure::Corrupt, FString::Printf(TEXT("UDemoNetDriver::OnDownloadHeaderComplete: LevelIndex %d not in range of level names of size: %d"), LevelIndex, ReplayHelper.PlaybackDemoHeader.LevelNamesAndTimes.Num()));
+				UE_LOG(LogDemo, Error, TEXT("UDemoNetDriver::OnDownloadHeaderComplete: LevelIndex %d not in range of level names of size: %d"), LevelIndex, ReplayHelper.PlaybackDemoHeader.LevelNamesAndTimes.Num());
+				ReplayHelper.NotifyReplayError(EReplayResult::Corrupt);
 			}
 		}
 		else
 		{
-			World->GetGameInstance()->HandleDemoPlaybackFailure(EDemoPlayFailure::Corrupt, FString::Printf(TEXT("UDemoNetDriver::OnDownloadHeaderComplete: ReadPlaybackDemoHeader header failed with error %s."), *Error));
+			UE_LOG(LogDemo, Error, TEXT("UDemoNetDriver::OnDownloadHeaderComplete: ReadPlaybackDemoHeader header failed with error %s."), *Error);
+			// ReadPlaybackDemoHeader will have already called NotifyPlaybackError internally on failure
 		}
 	}
 	else
 	{
-		World->GetGameInstance()->HandleDemoPlaybackFailure(EDemoPlayFailure::Corrupt, FString::Printf(TEXT("UDemoNetDriver::OnDownloadHeaderComplete: Downloading header failed.")));
+		UE_LOG(LogDemo, Error, TEXT("UDemoNetDriver::OnDownloadHeaderComplete: Downloading header failed."));
+		ReplayHelper.NotifyReplayError(EReplayResult::Corrupt);
 	}
 }
 
@@ -2578,7 +2568,7 @@ bool UDemoNetDriver::ProcessPacket(const uint8* Data, int32 Count)
 	{
 		// Something we received resulted in the demo being stopped
 		UE_LOG(LogDemo, Error, TEXT("UDemoNetDriver::ProcessPacket: ReceivedRawPacket closed connection"));
-		NotifyDemoPlaybackFailure(EDemoPlayFailure::Generic);
+		ReplayHelper.NotifyReplayError(EReplayResult::ConnectionClosed);
 		return false;
 	}
 
@@ -3167,15 +3157,13 @@ void UDemoNetDriver::ReplayStreamingReady(const FStartStreamingResult& Result)
 
 	if (!bWasStartStreamingSuccessful)
 	{
-		UE_LOG(LogDemo, Warning, TEXT("UDemoNetDriver::ReplayStreamingReady: Failed. %s"), Result.bRecording ? TEXT("") : EDemoPlayFailure::ToString(EDemoPlayFailure::DemoNotFound));
-
 		if (Result.bRecording)
 		{
-			NotifyDemoRecordFailure(EDemoRecordFailure::StartStreaming);
+			ReplayHelper.NotifyReplayError(EReplayResult::StreamerError);
 		}
 		else
 		{
-			NotifyDemoPlaybackFailure(EDemoPlayFailure::DemoNotFound);
+			ReplayHelper.NotifyReplayError(EReplayResult::ReplayNotFound);
 		}
 		return;
 	}
@@ -3689,7 +3677,7 @@ bool UDemoNetDriver::FastForwardLevels(const FGotoResult& GotoResult)
 				if (!ReadPacketsHelper.ReadPackets(*CheckpointArchive) && ReadPacketsHelper.IsError())
 				{
 					UE_LOG(LogDemo, Warning, TEXT("UDemoNetDriver::FastForwardLevels: Failed to read packets from Checkpoint."));
-					NotifyDemoPlaybackFailure(EDemoPlayFailure::Generic);
+					ReplayHelper.NotifyReplayError(EReplayResult::Unknown);
 					return false;
 				}
 
@@ -3722,7 +3710,7 @@ bool UDemoNetDriver::FastForwardLevels(const FGotoResult& GotoResult)
 		if (ReadPacketsHelper.IsError())
 		{
 			UE_LOG(LogDemo, Warning, TEXT("UDemoNetDriver::FastForwardLevels: Failed to read packets from Stream."));
-			NotifyDemoPlaybackFailure(EDemoPlayFailure::Serialization);
+			ReplayHelper.NotifyReplayError(EReplayResult::Serialization);
 			return false;
 		}
 
@@ -5294,40 +5282,59 @@ UChannel* UDemoNetDriver::InternalCreateChannelByName(const FName& ChName)
 	return NewChannel;
 }
 
+PRAGMA_DISABLE_DEPRECATION_WARNINGS
 void UDemoNetDriver::NotifyDemoPlaybackFailure(EDemoPlayFailure::Type FailureType)
 {
-	UE_LOG(LogDemo, Warning, TEXT("Demo playback failure: '%s'"), EDemoPlayFailure::ToString(FailureType));
+	UE_LOG(LogDemo, Warning, TEXT("NotifyDemoPlaybackFailure: %s"), EDemoPlayFailure::ToString(FailureType));
 
 	const bool bIsPlaying = IsPlaying();
 
 	// fire delegate
 	FNetworkReplayDelegates::OnReplayStartFailure.Broadcast(World, FailureType);
+	FNetworkReplayDelegates::OnReplayPlaybackFailure.Broadcast(World, EReplayResult::Unknown);
 
 	StopDemo();
 
-	if (bIsPlaying)
+	if (bIsPlaying && World)
 	{
-		if (World)
+		if (UGameInstance* GameInstance = World->GetGameInstance())
 		{
-			if (UGameInstance* GameInstance = World->GetGameInstance())
-			{
-				GameInstance->HandleDemoPlaybackFailure(FailureType, FString(EDemoPlayFailure::ToString(FailureType)));
-			}
+			GameInstance->HandleDemoPlaybackFailure(EReplayResult::Unknown);
 		}
 	}
 }
+PRAGMA_ENABLE_DEPRECATION_WARNINGS
 
-void UDemoNetDriver::NotifyDemoRecordFailure(EDemoRecordFailure FailureType)
+void UDemoNetDriver::NotifyDemoPlaybackError(const UE::Net::TNetResult<EReplayResult>& Result)
 {
-	UE_LOG(LogDemo, Warning, TEXT("Demo record failure: '%s'"), LexToString(FailureType));
+	const bool bIsPlaying = IsPlaying();
+
+	// fire delegate
+	PRAGMA_DISABLE_DEPRECATION_WARNINGS
+	FNetworkReplayDelegates::OnReplayStartFailure.Broadcast(World, EDemoPlayFailure::Generic);
+	PRAGMA_ENABLE_DEPRECATION_WARNINGS
+	FNetworkReplayDelegates::OnReplayPlaybackFailure.Broadcast(World, Result);
+
+	StopDemo();
+
+	if (bIsPlaying && World)
+	{
+		if (UGameInstance* GameInstance = World->GetGameInstance())
+		{
+			GameInstance->HandleDemoPlaybackFailure(Result);
+		}
+	}
+}
 
+void UDemoNetDriver::NotifyDemoRecordFailure(const UE::Net::TNetResult<EReplayResult>& Result)
+{
 	const bool bIsRecording = IsRecording();
 
 	StopDemo();
 
 	if (bIsRecording && !PendingRecordFailure.IsSet())
 	{
-		PendingRecordFailure.Emplace(FailureType);
+		PendingRecordFailure.Emplace(Result);
 	}
 }
 
@@ -5470,9 +5477,10 @@ bool UDemoPendingNetGame::LoadMapCompleted(UEngine* Engine, FWorldContext& Conte
 		GEngine->BrowseToDefaultMap(Context);
 
 		UE_LOG(LogDemo, Error, TEXT("UDemoPendingNetGame::HandlePostLoadMap: LoadMap failed: %s"), *LoadMapError);
+
 		if (Context.OwningGameInstance)
 		{
-			Context.OwningGameInstance->HandleDemoPlaybackFailure(EDemoPlayFailure::LoadMap, FString(TEXT("LoadMap failed")));
+			Context.OwningGameInstance->HandleDemoPlaybackFailure(EReplayResult::LoadMap);
 		}
 		// we already handled the error so don't need the TickPendingNetGame to handle it 
 		return true;
diff --git a/Engine/Source/Runtime/Engine/Private/Net/ReplayPlaylistTracker.cpp b/Engine/Source/Runtime/Engine/Private/Net/ReplayPlaylistTracker.cpp
--- a/Engine/Source/Runtime/Engine/Private/Net/ReplayPlaylistTracker.cpp
+++ b/Engine/Source/Runtime/Engine/Private/Net/ReplayPlaylistTracker.cpp
@@ -50,14 +50,14 @@ void FReplayPlaylistTracker::Reset()
 	{
 		DemoNetDriver.Reset();
 
-		FNetworkReplayDelegates::OnReplayStartFailure.Remove(OnDemoFailedToStartHandle);
+		FNetworkReplayDelegates::OnReplayPlaybackFailure.Remove(OnDemoPlaybackFailedHandle);
 		FNetworkReplayDelegates::OnReplayPlaybackComplete.Remove(OnDemoPlaybackFinishedHandle);
 		FNetworkReplayDelegates::OnReplayRecordingComplete.Remove(OnDemoStoppedHandle);
 
 		LocalDemoNetDriver->SetPlayingPlaylist(nullptr);
 	}
 
-	OnDemoFailedToStartHandle.Reset();
+	OnDemoPlaybackFailedHandle.Reset();
 	OnDemoPlaybackFinishedHandle.Reset();
 	OnDemoStoppedHandle.Reset();
 
@@ -93,7 +93,7 @@ bool FReplayPlaylistTracker::Start()
 			{
 				DemoNetDriver = LocalDemoNetDriver;
 
-				OnDemoFailedToStartHandle = FNetworkReplayDelegates::OnReplayStartFailure.AddSP(AsShared(), &ThisClass::OnDemoFailedToStart);
+				OnDemoPlaybackFailedHandle = FNetworkReplayDelegates::OnReplayPlaybackFailure.AddSP(AsShared(), &ThisClass::OnDemoPlaybackFailed);
 				OnDemoPlaybackFinishedHandle = FNetworkReplayDelegates::OnReplayPlaybackComplete.AddSP(AsShared() , &ThisClass::OnDemoPlaybackFinished);
 				OnDemoStoppedHandle = FNetworkReplayDelegates::OnReplayRecordingComplete.AddSP(AsShared(), &ThisClass::OnDemoStopped);
 
@@ -122,7 +122,7 @@ bool FReplayPlaylistTracker::Restart()
 	return Start();
 }
 
-void FReplayPlaylistTracker::OnDemoFailedToStart(UWorld* InWorld, EDemoPlayFailure::Type FailureType)
+void FReplayPlaylistTracker::OnDemoPlaybackFailed(UWorld* InWorld, const UE::Net::TNetResult<EReplayResult>& Result)
 {
 	UWorld* World = WorldOverride.Get();
 	World = World ? World : GameInstance->GetWorld();
diff --git a/Engine/Source/Runtime/Engine/Private/Net/ReplayResult.cpp b/Engine/Source/Runtime/Engine/Private/Net/ReplayResult.cpp
new file mode 100644
--- /dev/null
+++ b/Engine/Source/Runtime/Engine/Private/Net/ReplayResult.cpp
@@ -0,0 +1,26 @@
+// Copyright Epic Games, Inc. All Rights Reserved.
+
+#include "Net/ReplayResult.h"
+
+#define CASE_EREPLAYRESULT_TO_TEXT_RET(txt) case txt: ReturnVal = TEXT(#txt); break;
+
+const TCHAR* LexToString(EReplayResult Result)
+{
+	const TCHAR* ReturnVal = TEXT("::Invalid");
+
+	switch (Result)
+	{
+		FOREACH_ENUM_EREPLAYRESULT(CASE_EREPLAYRESULT_TO_TEXT_RET)
+	}
+
+	while (*ReturnVal != ':')
+	{
+		ReturnVal++;
+	}
+
+	ReturnVal += 2;
+
+	return ReturnVal;
+}
+
+#undef CASE_EREPLAYRESULT_TO_TEXT_RET
\ No newline at end of file
diff --git a/Engine/Source/Runtime/Engine/Private/ReplayHelper.cpp b/Engine/Source/Runtime/Engine/Private/ReplayHelper.cpp
--- a/Engine/Source/Runtime/Engine/Private/ReplayHelper.cpp
+++ b/Engine/Source/Runtime/Engine/Private/ReplayHelper.cpp
@@ -51,6 +51,7 @@ FReplayHelper::FReplayHelper()
 	, CheckpointSaveMaxMSPerFrame(0)
 	, NumLevelsAddedThisFrame(0)
 	, bPendingCheckpointRequest(false)
+	, bRecording(false)
 {
 }
 
@@ -78,6 +79,10 @@ TSharedPtr<INetworkReplayStreamer> FReplayHelper::Init(const FURL& URL)
 		}
 	}
 
+	TUniquePtr<FReplayResultHandler> ReplayHandler = MakeUnique<FReplayResultHandler>();
+	ReplayHandler->InitResultHandler(this);
+
+	ResultManager.AddResultHandler(MoveTemp(ReplayHandler), UE::Net::EAddResultHandlerPos::First);
 
 	FWorldDelegates::LevelRemovedFromWorld.AddRaw(this, &FReplayHelper::OnLevelRemovedFromWorld);
 	FWorldDelegates::LevelAddedToWorld.AddRaw(this, &FReplayHelper::OnLevelAddedToWorld);
@@ -118,7 +123,7 @@ void FReplayHelper::OnStartRecordingComplete(const FStartStreamingResult& Result
 	if (!bWasStartRecordingSuccessful)
 	{
 		UE_LOG(LogDemo, Warning, TEXT("FReplayRecordHelper::StartRecordingComplete: Failed"));
-		OnReplayRecordError.Broadcast();
+		NotifyReplayError(EReplayResult::StreamerError);
 		return;
 	}
 
@@ -396,7 +401,7 @@ bool FReplayHelper::ReadPlaybackDemoHeader(FString& Error)
 	{
 		Error = FString::Printf(TEXT("Couldn't open demo file %s for reading"), *DemoURL.Map);
 		UE_LOG(LogDemo, Error, TEXT("FReplayHelper::ReadPlaybackDemoHeader: %s"), *Error);
-		GameInstance->HandleDemoPlaybackFailure(EDemoPlayFailure::DemoNotFound, FString(EDemoPlayFailure::ToString(EDemoPlayFailure::DemoNotFound)));
+		NotifyReplayError(EReplayResult::ReplayNotFound);
 		return false;
 	}
 
@@ -406,7 +411,7 @@ bool FReplayHelper::ReadPlaybackDemoHeader(FString& Error)
 	{
 		Error = FString(TEXT("Demo file is corrupt"));
 		UE_LOG(LogDemo, Error, TEXT("FReplayHelper::ReadPlaybackDemoHeader: %s"), *Error);
-		GameInstance->HandleDemoPlaybackFailure(EDemoPlayFailure::Corrupt, Error);
+		NotifyReplayError(EReplayResult::Corrupt);
 		return false;
 	}
 
@@ -418,7 +423,7 @@ bool FReplayHelper::ReadPlaybackDemoHeader(FString& Error)
 	if (HasDeltaCheckpoints() && !ReplayStreamer->IsCheckpointTypeSupported(EReplayCheckpointType::Delta))
 	{
 		UE_LOG(LogDemo, Error, TEXT("FReplayHelper::ReadPlaybackDemoHeader: Replay has delta checkpoints but streamer does not support them."));
-		GameInstance->HandleDemoPlaybackFailure(EDemoPlayFailure::Generic, FString(EDemoPlayFailure::ToString(EDemoPlayFailure::Generic)));
+		NotifyReplayError(EReplayResult::UnsupportedCheckpoint);
 		return false;
 	}
 
@@ -436,7 +441,7 @@ bool FReplayHelper::ReadPlaybackDemoHeader(FString& Error)
 	if (!Error.IsEmpty())
 	{
 		UE_LOG(LogDemo, Error, TEXT("FReplayHelper::ReadPlaybackDemoHeader: (Game Specific) %s"), *Error);
-		GameInstance->HandleDemoPlaybackFailure(EDemoPlayFailure::GameSpecificHeader, Error);
+		NotifyReplayError(EReplayResult::GameSpecific);
 		return false;
 	}
 
@@ -451,10 +456,9 @@ void FReplayHelper::TickRecording(float DeltaSeconds, UNetConnection* Connection
 		return;
 	}
 
-	if (ReplayStreamer->GetLastError() != ENetworkReplayError::None)
+	if (UE::Net::EHandleNetResult::Handled == ReplayStreamer->HandleLastError(ResultManager))
 	{
-		UE_LOG(LogDemo, Error, TEXT("FReplayHelper::TickRecording: ReplayStreamer ERROR: %s"), ENetworkReplayError::ToString(ReplayStreamer->GetLastError()));
-		OnReplayRecordError.Broadcast();
+		NotifyReplayError(EReplayResult::StreamerError);
 		return;
 	}
 
@@ -468,7 +472,7 @@ void FReplayHelper::TickRecording(float DeltaSeconds, UNetConnection* Connection
 	if (FileAr == nullptr)
 	{
 		UE_LOG(LogDemo, Error, TEXT("FReplayHelper::TickRecording: FileAr == nullptr"));
-		OnReplayRecordError.Broadcast();
+		NotifyReplayError(EReplayResult::MissingArchive);
 		return;
 	}
 
@@ -1871,7 +1875,7 @@ bool FReplayHelper::ReadDemoFrame(UNetConnection* Connection, FArchive& Ar, TArr
 	if (Ar.IsError())
 	{
 		UE_LOG(LogDemo, Error, TEXT("FReplayHelper::ReadDemoFrame: Archive Error"));
-		OnReplayPlaybackError.Broadcast(EDemoPlayFailure::Serialization);
+		NotifyReplayError(EReplayResult::Serialization);
 		return false;
 	}
 
@@ -1880,10 +1884,9 @@ bool FReplayHelper::ReadDemoFrame(UNetConnection* Connection, FArchive& Ar, TArr
 		return false;
 	}
 
-	if (ReplayStreamer->GetLastError() != ENetworkReplayError::None)
+	if (UE::Net::EHandleNetResult::Handled == ReplayStreamer->HandleLastError(ResultManager))
 	{
-		UE_LOG(LogDemo, Error, TEXT("FReplayHelper::ReadDemoFrame: ReplayStreamer ERROR: %s"), ENetworkReplayError::ToString(ReplayStreamer->GetLastError()));
-		OnReplayPlaybackError.Broadcast(EDemoPlayFailure::ReplayStreamerInternal);
+		NotifyReplayError(EReplayResult::StreamerError);
 		return false;
 	}
 
@@ -1999,7 +2002,7 @@ bool FReplayHelper::ReadDemoFrame(UNetConnection* Connection, FArchive& Ar, TArr
 	if (Ar.IsError())
 	{
 		UE_LOG(LogDemo, Error, TEXT("FReplayHelper::ReadDemoFrame: Failed to read demo ServerDeltaTime"));
-		OnReplayPlaybackError.Broadcast(EDemoPlayFailure::Serialization);
+		NotifyReplayError(EReplayResult::Serialization);
 		return false;
 	}
 
@@ -2062,7 +2065,7 @@ bool FReplayHelper::ReadDemoFrame(UNetConnection* Connection, FArchive& Ar, TArr
 			case EReadPacketState::Error:
 			{
 				UE_LOG(LogDemo, Error, TEXT("FReplayHelper::ReadDemoFrame: ReadPacket failed."));
-				OnReplayPlaybackError.Broadcast(EDemoPlayFailure::Serialization);
+				NotifyReplayError(EReplayResult::Serialization);
 				return false;
 			}
 
@@ -2291,22 +2294,6 @@ const TCHAR* LexToString(EReplayHeaderFlags Flag)
 	}
 }
 
-const TCHAR* LexToString(EDemoRecordFailure FailureType)
-{
-	switch (FailureType)
-	{
-	case EDemoRecordFailure::Generic:
-		return TEXT("Generic");
-	case EDemoRecordFailure::StreamingArchive:
-		return TEXT("StreamingArchive");
-	case EDemoRecordFailure::StartStreaming:
-		return TEXT("StartStreaming");
-	default:
-		check(false);
-		return TEXT("Unknown");
-	}
-}
-
 void FReplayHelper::RequestCheckpoint()
 {
 	bPendingCheckpointRequest = true;
@@ -2395,3 +2382,42 @@ void FReplayHelper::NotifyActorDestroyed(UNetConnection* Connection, AActor* Act
 		}
 	}
 }
+
+void FReplayHelper::NotifyReplayError(UE::Net::TNetResult<EReplayResult>&& Result)
+{
+	ResultManager.HandleNetResult(MoveTemp(Result));
+}
+
+void FReplayResultHandler::InitResultHandler(FReplayHelper* InReplayHelper)
+{
+	ReplayHelper = InReplayHelper;
+}
+
+UE::Net::EHandleNetResult FReplayResultHandler::HandleNetResult(UE::Net::FNetResult&& InResult)
+{
+	using namespace UE::Net;
+
+	UE_LOG(LogDemo, Error, TEXT("FReplayResultHandler::HandleNetResult:"));
+
+	for (FNetResult::FConstIterator It(InResult); It; ++It)
+	{
+		UE_LOG(LogDemo, Error, TEXT(" - %s"), ToCStr(It->DynamicToString()));
+	}
+
+	TNetResult<EReplayResult>* CastedResult = Cast<EReplayResult>(&InResult);
+
+	if (CastedResult && ReplayHelper)
+	{
+		if (ReplayHelper->bRecording)
+		{
+			ReplayHelper->OnReplayRecordError.ExecuteIfBound(*CastedResult);
+		}
+		else
+		{
+			ReplayHelper->OnReplayPlaybackError.ExecuteIfBound(*CastedResult);
+		}
+	}
+
+	return EHandleNetResult::Handled;
+}
+
diff --git a/Engine/Source/Runtime/Engine/Private/ReplayNetConnection.cpp b/Engine/Source/Runtime/Engine/Private/ReplayNetConnection.cpp
--- a/Engine/Source/Runtime/Engine/Private/ReplayNetConnection.cpp
+++ b/Engine/Source/Runtime/Engine/Private/ReplayNetConnection.cpp
@@ -38,6 +38,7 @@ void UReplayNetConnection::InitConnection(UNetDriver* InDriver, EConnectionState
 	InitSendBuffer();
 
 	ReplayHelper.Init(InURL);
+	ReplayHelper.bRecording = true;
 }
 
 void UReplayNetConnection::CleanUp()
diff --git a/Engine/Source/Runtime/Engine/Private/UnrealNetwork.cpp b/Engine/Source/Runtime/Engine/Private/UnrealNetwork.cpp
--- a/Engine/Source/Runtime/Engine/Private/UnrealNetwork.cpp
+++ b/Engine/Source/Runtime/Engine/Private/UnrealNetwork.cpp
@@ -11,7 +11,10 @@ FOnProcessGameSpecificFrameData                      FNetworkReplayDelegates::On
 FGetOverridableVersionDataForDemoHeaderReadDelegate  FNetworkReplayDelegates::GetOverridableVersionDataForHeaderRead;
 FGetOverridableVersionDataForDemoHeaderWriteDelegate FNetworkReplayDelegates::GetOverridableVersionDataForHeaderWrite;
 FOnReplayStartedDelegate                             FNetworkReplayDelegates::OnReplayStarted;
+PRAGMA_DISABLE_DEPRECATION_WARNINGS
 FOnReplayStartFailureDelegate                        FNetworkReplayDelegates::OnReplayStartFailure;
+PRAGMA_ENABLE_DEPRECATION_WARNINGS
+FOnReplayPlaybackFailureDelegate                     FNetworkReplayDelegates::OnReplayPlaybackFailure;
 FOnReplayScrubCompleteDelegate                       FNetworkReplayDelegates::OnReplayScrubComplete;
 FOnReplayPlaybackCompleteDelegate                    FNetworkReplayDelegates::OnReplayPlaybackComplete;
 FOnReplayRecordingStartAttemptDelegate               FNetworkReplayDelegates::OnReplayRecordingStartAttempt;
diff --git a/Engine/Source/Runtime/Engine/Private/World.cpp b/Engine/Source/Runtime/Engine/Private/World.cpp
--- a/Engine/Source/Runtime/Engine/Private/World.cpp
+++ b/Engine/Source/Runtime/Engine/Private/World.cpp
@@ -4675,7 +4675,7 @@ bool UWorld::HandleDemoPlayCommand( const TCHAR* Cmd, FOutputDevice& Ar, UWorld*
 
 		if (GetGameInstance() != nullptr)
 		{
-			GetGameInstance()->HandleDemoPlaybackFailure(EDemoPlayFailure::Generic, FString(ErrorString));
+			GetGameInstance()->HandleDemoPlaybackFailure(EReplayResult::Unknown);
 		}
 	}
 	else
diff --git a/Engine/Source/Runtime/Engine/Public/Net/ReplayPlaylistTracker.h b/Engine/Source/Runtime/Engine/Public/Net/ReplayPlaylistTracker.h
--- a/Engine/Source/Runtime/Engine/Public/Net/ReplayPlaylistTracker.h
+++ b/Engine/Source/Runtime/Engine/Public/Net/ReplayPlaylistTracker.h
@@ -6,6 +6,8 @@
 #include "Engine/EngineBaseTypes.h"
 #include "Templates/UnrealTemplate.h"
 #include "Templates/SharedPointer.h"
+#include "Net/Core/Connection/NetResult.h"
+#include "Net/ReplayResult.h"
 
 /**
  * Delegate that is called prior to starting each replay in a playlist,
@@ -113,7 +115,7 @@ private:
 	void PlayNextReplay();
 
 	/** Called if an error occurs *after* the replay was requested to start (successfully), but before the replay actually starts. */
-	void OnDemoFailedToStart(UWorld* InWorld, EDemoPlayFailure::Type FailureType);
+	void OnDemoPlaybackFailed(UWorld* InWorld, const UE::Net::TNetResult<EReplayResult>& Result);
 
 	/** Called when the DemoNetDriver hits a Hard Stop, like an error or a user requesting a stop. */
 	void OnDemoStopped(UWorld* InWorld);
@@ -129,7 +131,7 @@ private:
 	TWeakObjectPtr<class UGameInstance> GameInstance;
 	TWeakObjectPtr<class UDemoNetDriver> DemoNetDriver;
 
-	FDelegateHandle OnDemoFailedToStartHandle;
+	FDelegateHandle OnDemoPlaybackFailedHandle;
 	FDelegateHandle OnDemoPlaybackFinishedHandle;
 	FDelegateHandle OnDemoStoppedHandle;
 
diff --git a/Engine/Source/Runtime/Engine/Public/Net/ReplayResult.h b/Engine/Source/Runtime/Engine/Public/Net/ReplayResult.h
new file mode 100644
--- /dev/null
+++ b/Engine/Source/Runtime/Engine/Public/Net/ReplayResult.h
@@ -0,0 +1,27 @@
+// Copyright Epic Games, Inc. All Rights Reserved.
+
+#include "Net/Core/Connection/NetResult.h"
+#include "ReplayResult.generated.h"
+
+#pragma once
+
+UENUM()
+enum class EReplayResult : uint32
+{
+	Success,
+	ReplayNotFound,
+	Corrupt,
+	UnsupportedCheckpoint,
+	GameSpecific,
+	InitConnect,
+	LoadMap,
+	Serialization,
+	StreamerError,
+	ConnectionClosed,
+	MissingArchive,
+	Unknown,
+};
+
+DECLARE_NETRESULT_ENUM(EReplayResult);
+
+ENGINE_API const TCHAR* LexToString(EReplayResult Result);
\ No newline at end of file
diff --git a/Engine/Source/Runtime/Engine/Public/Net/UnrealNetwork.h b/Engine/Source/Runtime/Engine/Public/Net/UnrealNetwork.h
--- a/Engine/Source/Runtime/Engine/Public/Net/UnrealNetwork.h
+++ b/Engine/Source/Runtime/Engine/Public/Net/UnrealNetwork.h
@@ -16,7 +16,9 @@
 #include "EngineLogs.h"
 #include "UObject/UnrealType.h"
 #include "Engine/EngineBaseTypes.h"
+#include "Net/Core/Connection/NetResult.h"
 #include "Net/Core/PropertyConditions/PropertyConditions.h"
+#include "Net/ReplayResult.h"
 
 class AActor;
 
@@ -60,7 +62,14 @@ DECLARE_DELEGATE_OneParam(FGetOverridableVersionDataForDemoHeaderReadDelegate, F
 DECLARE_DELEGATE_OneParam(FGetOverridableVersionDataForDemoHeaderWriteDelegate, FOverridableReplayVersionData& /*OveridableReplayVersionData*/);
 
 DECLARE_MULTICAST_DELEGATE_OneParam(FOnReplayStartedDelegate, UWorld* /*World*/);
+
+PRAGMA_DISABLE_DEPRECATION_WARNINGS
+UE_DEPRECATED(5.2, "Deprecated in favor of FOnReplayPlaybackFailureDelegate")
 DECLARE_MULTICAST_DELEGATE_TwoParams(FOnReplayStartFailureDelegate, UWorld* /*World*/, EDemoPlayFailure::Type /*Error*/);
+PRAGMA_ENABLE_DEPRECATION_WARNINGS
+
+DECLARE_MULTICAST_DELEGATE_TwoParams(FOnReplayPlaybackFailureDelegate, UWorld* /*World*/, const UE::Net::TNetResult<EReplayResult>& /*Error*/);
+
 DECLARE_MULTICAST_DELEGATE_OneParam(FOnReplayScrubCompleteDelegate, UWorld* /*World*/);
 DECLARE_MULTICAST_DELEGATE_OneParam(FOnReplayPlaybackCompleteDelegate, UWorld* /*World*/);
 DECLARE_MULTICAST_DELEGATE_OneParam(FOnReplayRecordingStartAttemptDelegate, UWorld* /*World*/);
@@ -89,8 +98,14 @@ struct ENGINE_API FNetworkReplayDelegates
 	/** Public delegate for external systems to be notified when a replay begins */
 	static FOnReplayStartedDelegate OnReplayStarted;
 
+	PRAGMA_DISABLE_DEPRECATION_WARNINGS
 	/** Public delegate to be notified when a replay failed to start */
+	UE_DEPRECATED(5.2, "Deprecated in favor of OnReplayPlaybackFailure")
 	static FOnReplayStartFailureDelegate OnReplayStartFailure;
+	PRAGMA_ENABLE_DEPRECATION_WARNINGS
+
+	/** Public delegate to be notified when replay playback failed */
+	static FOnReplayPlaybackFailureDelegate OnReplayPlaybackFailure;
 
 	/** Public delegate for external systems to be notified when scrubbing is complete. Only called for successful scrub. */
 	static FOnReplayScrubCompleteDelegate OnReplayScrubComplete;
diff --git a/Engine/Source/Runtime/Engine/Public/ReplayHelper.h b/Engine/Source/Runtime/Engine/Public/ReplayHelper.h
--- a/Engine/Source/Runtime/Engine/Public/ReplayHelper.h
+++ b/Engine/Source/Runtime/Engine/Public/ReplayHelper.h
@@ -6,20 +6,35 @@
 #include "Engine/EngineBaseTypes.h"
 #include "Engine/PackageMapClient.h"
 #include "NetworkReplayStreaming.h"
+#include "Net/Core/Connection/NetResult.h"
+#include "Net/ReplayResult.h"
 #include "ReplayTypes.h"
 #include "Containers/ArrayView.h"
 
 class APlayerController;
 class UNetConnection;
 
+class FReplayHelper;
+
+class FReplayResultHandler final : public UE::Net::FNetResultHandler
+{
+	friend class FReplayHelper;
+
+private:
+	void InitResultHandler(FReplayHelper* InReplayHelper);
+
+	virtual UE::Net::EHandleNetResult HandleNetResult(UE::Net::FNetResult&& InResult) override;
+
+private:
+	FReplayHelper* ReplayHelper = nullptr;
+};
+
 class FReplayHelper
 {
 	friend class UDemoNetDriver;
 	friend class UDemoNetConnection;
 	friend class UReplayNetConnection;
-	PRAGMA_DISABLE_DEPRECATION_WARNINGS
-	friend class FScopedPacketManager;
-	PRAGMA_ENABLE_DEPRECATION_WARNINGS
+	friend class FReplayResultHandler;
 
 public:
 	FReplayHelper();
@@ -87,10 +102,10 @@ private:
 	/** Returns either CheckpointSaveMaxMSPerFrame or the value of demo.CheckpointSaveMaxMSPerFrameOverride if it's >= 0. */
 	float GetCheckpointSaveMaxMSPerFrame() const;
 
-	DECLARE_MULTICAST_DELEGATE(FOnReplayRecordError);
+	DECLARE_DELEGATE_OneParam(FOnReplayRecordError, const UE::Net::TNetResult<EReplayResult>&);
 	FOnReplayRecordError OnReplayRecordError;
 
-	DECLARE_MULTICAST_DELEGATE_OneParam(FOnReplayPlaybackError, EDemoPlayFailure::Type);
+	DECLARE_DELEGATE_OneParam(FOnReplayPlaybackError, const UE::Net::TNetResult<EReplayResult>&);
 	FOnReplayPlaybackError OnReplayPlaybackError;
 
 	static float GetClampedDeltaSeconds(UWorld* World, const float DeltaSeconds);
@@ -459,6 +474,8 @@ private:
 
 	bool bPendingCheckpointRequest;
 
+	UE::Net::FNetResultManager ResultManager;
+
 	static FString GetLevelPackageName(const ULevel& InLevel);
 
 	void ResetLevelStatuses();
@@ -539,6 +556,10 @@ private:
 
 	void ProcessCheckpointActors(UNetConnection* Connection, TArrayView<FPendingCheckPointActor> PendingActors, int32& NextIndex, FRepActorsCheckpointParams& Params);
 
+	void NotifyReplayError(UE::Net::TNetResult<EReplayResult>&& Result);
+
+	bool bRecording;
+
 	static constexpr int32 MAX_DEMO_READ_WRITE_BUFFER = 1024 * 2;
 	static constexpr int32 MAX_DEMO_STRING_SERIALIZATION_SIZE = 16 * 1024 * 1024;
-};
\ No newline at end of file
+};
diff --git a/Engine/Source/Runtime/Engine/Public/ReplayTypes.h b/Engine/Source/Runtime/Engine/Public/ReplayTypes.h
--- a/Engine/Source/Runtime/Engine/Public/ReplayTypes.h
+++ b/Engine/Source/Runtime/Engine/Public/ReplayTypes.h
@@ -9,6 +9,7 @@
 #include "Misc/NetworkGuid.h"
 #include "Misc/NetworkVersion.h"
 #include "Net/Common/Packets/PacketTraits.h"
+#include "Net/ReplayResult.h"
 #include "IPAddress.h"
 #include "Serialization/BitReader.h"
 #include "Traits/IsCharEncodingCompatibleWith.h"
@@ -43,15 +44,6 @@ enum class EWriteDemoFrameFlags : uint32
 
 ENUM_CLASS_FLAGS(EWriteDemoFrameFlags);
 
-enum class EDemoRecordFailure : uint32
-{
-	Generic,
-	StreamingArchive,
-	StartStreaming,
-};
-
-ENGINE_API const TCHAR* LexToString(EDemoRecordFailure FailureType);
-
 struct FPlaybackPacket
 {
 	TArray<uint8>		Data;
diff --git a/Engine/Source/Runtime/Net/Core/Public/Net/Core/Connection/NetResult.h b/Engine/Source/Runtime/Net/Core/Public/Net/Core/Connection/NetResult.h
--- a/Engine/Source/Runtime/Net/Core/Public/Net/Core/Connection/NetResult.h
+++ b/Engine/Source/Runtime/Net/Core/Public/Net/Core/Connection/NetResult.h
@@ -379,7 +379,7 @@ public:
 
 		for (FConstIterator It(*this); It; ++It)
 		{
-			const TCHAR* ResultLex = LexToString(static_cast<ResultEnum>(It->Result));
+			const TCHAR* ResultLex = ToCStr(LexToString(static_cast<ResultEnum>(It->Result)));
 
 			if (ConversionType == ENetResultString::ResultEnumOnly)
 			{
@@ -480,7 +480,7 @@ private:
 	/** Default ErrorContext based on ResultEnum */
 	static const TCHAR* DefaultErrorContext(ResultEnum InResult)
 	{
-		return (InResult == ResultEnum::Success ? TEXT("") : LexToString(InResult));
+		return (InResult == ResultEnum::Success ? TEXT("") : ToCStr(LexToString(InResult)));
 	}
 };
 
diff --git a/Engine/Source/Runtime/NetworkReplayStreaming/HttpNetworkReplayStreaming/Private/HttpNetworkReplayStreaming.cpp b/Engine/Source/Runtime/NetworkReplayStreaming/HttpNetworkReplayStreaming/Private/HttpNetworkReplayStreaming.cpp
--- a/Engine/Source/Runtime/NetworkReplayStreaming/HttpNetworkReplayStreaming/Private/HttpNetworkReplayStreaming.cpp
+++ b/Engine/Source/Runtime/NetworkReplayStreaming/HttpNetworkReplayStreaming/Private/HttpNetworkReplayStreaming.cpp
@@ -216,7 +216,6 @@ FHttpNetworkReplayStreamer::FHttpNetworkReplayStreamer() :
 	StreamTimeRangeStart( 0 ),
 	StreamTimeRangeEnd( 0 ),
 	HighPriorityEndTime( 0 ),
-	StreamerLastError( ENetworkReplayError::None ),
 	DownloadCheckpointIndex( -1 ),
 	DeltaDownloadCheckpointIndex( -1 ),
 	LastGotoTimeInMS( -1 ),
@@ -671,7 +670,7 @@ void FHttpNetworkReplayStreamer::FlushStream()
 
 	if (!CompressRequest(HttpRequest, StreamArchive.Buffer))
 	{
-		SetLastError( ENetworkReplayError::ServiceUnavailable );
+		SetLastError(EHttpReplayResult::CompressionFailed);
 		return;
 	}
 
@@ -1033,7 +1032,7 @@ void FHttpNetworkReplayStreamer::FlushCheckpointInternal( uint32 TimeInMS )
 
 		if (!CompressRequest(HttpRequest, CheckpointArchive.Buffer))
 		{
-			SetLastError(ENetworkReplayError::ServiceUnavailable);
+			SetLastError(EHttpReplayResult::CompressionFailed);
 			return;
 		}
 
@@ -1382,11 +1381,11 @@ void FHttpNetworkReplayStreamer::ConditionallyRefreshViewer()
 	}
 };
 
-void FHttpNetworkReplayStreamer::SetLastError( const ENetworkReplayError::Type InLastError )
+void FHttpNetworkReplayStreamer::SetLastError(FHttpReplayResult&& Result)
 {
 	CancelStreamingRequests();
-
-	StreamerLastError = InLastError;
+	
+	SetExtendedError(MoveTemp(Result));
 }
 
 void FHttpNetworkReplayStreamer::CancelStreamingRequests()
@@ -1410,11 +1409,6 @@ void FHttpNetworkReplayStreamer::CancelStreamingRequests()
 	bStopStreamingCalled	= false;
 }
 
-ENetworkReplayError::Type FHttpNetworkReplayStreamer::GetLastError() const
-{
-	return StreamerLastError;
-}
-
 FArchive* FHttpNetworkReplayStreamer::GetHeaderArchive()
 {
 	return &HeaderArchive;
@@ -1445,7 +1439,7 @@ void FHttpNetworkReplayStreamer::UpdateTotalDemoTime( uint32 TimeInMS )
 
 bool FHttpNetworkReplayStreamer::IsDataAvailable() const
 {
-	if ( GetLastError() != ENetworkReplayError::None )
+	if (HasError())
 	{
 		return false;
 	}
@@ -1481,7 +1475,7 @@ void FHttpNetworkReplayStreamer::SetHighPriorityTimeRange( const uint32 StartTim
 
 bool FHttpNetworkReplayStreamer::IsDataAvailableForTimeRange( const uint32 StartTimeInMS, const uint32 EndTimeInMS )
 {
-	if ( GetLastError() != ENetworkReplayError::None )
+	if (HasError())
 	{
 		return false;
 	}
@@ -1751,7 +1745,7 @@ void FHttpNetworkReplayStreamer::HttpStartUploadingFinished( FHttpRequestPtr Htt
 		}
 
 		UE_LOG( LogHttpReplay, Error, TEXT( "FHttpNetworkReplayStreamer::HttpStartUploadingFinished. FAILED, %s" ), *BuildRequestErrorString( HttpRequest, HttpResponse ) );
-		SetLastError( ENetworkReplayError::ServiceUnavailable );
+		SetLastError(EHttpReplayResult::InvalidHttpResponse);
 	}
 }
 
@@ -1773,7 +1767,7 @@ void FHttpNetworkReplayStreamer::HttpStopUploadingFinished( FHttpRequestPtr Http
 		}
 
 		UE_LOG( LogHttpReplay, Error, TEXT( "FHttpNetworkReplayStreamer::HttpStopUploadingFinished. FAILED, %s" ), *BuildRequestErrorString( HttpRequest, HttpResponse ) );
-		SetLastError( ENetworkReplayError::ServiceUnavailable );
+		SetLastError(EHttpReplayResult::InvalidHttpResponse);
 	}
 
 	StreamArchive.SetIsLoading(false);
@@ -1811,7 +1805,7 @@ void FHttpNetworkReplayStreamer::HttpHeaderUploadFinished( FHttpRequestPtr HttpR
 		}
 
 		UE_LOG( LogHttpReplay, Error, TEXT( "FHttpNetworkReplayStreamer::HttpHeaderUploadFinished. FAILED, %s" ), *BuildRequestErrorString( HttpRequest, HttpResponse ) );
-		SetLastError( ENetworkReplayError::ServiceUnavailable );
+		SetLastError(EHttpReplayResult::InvalidHttpResponse);
 	}
 
 	StartStreamingDelegate.ExecuteIfBound(Result);
@@ -1836,7 +1830,7 @@ void FHttpNetworkReplayStreamer::HttpUploadStreamFinished( FHttpRequestPtr HttpR
 		}
 
 		UE_LOG( LogHttpReplay, Error, TEXT( "FHttpNetworkReplayStreamer::HttpUploadStreamFinished. FAILED, %s" ), *BuildRequestErrorString( HttpRequest, HttpResponse ) );
-		SetLastError( ENetworkReplayError::ServiceUnavailable );
+		SetLastError(EHttpReplayResult::InvalidHttpResponse);
 	}
 }
 
@@ -1863,7 +1857,7 @@ void FHttpNetworkReplayStreamer::HttpUploadCheckpointFinished( FHttpRequestPtr H
 		}
 
 		UE_LOG( LogHttpReplay, Warning, TEXT( "FHttpNetworkReplayStreamer::HttpUploadCheckpointFinished. FAILED, %s" ), *BuildRequestErrorString( HttpRequest, HttpResponse ) );
-		SetLastError( ENetworkReplayError::ServiceUnavailable );
+		SetLastError(EHttpReplayResult::InvalidHttpResponse);
 	}
 }
 
@@ -1924,7 +1918,7 @@ void FHttpNetworkReplayStreamer::HttpStartDownloadingFinished( FHttpRequestPtr H
 			// Reset delegate
 			StartStreamingDelegate = FStartStreamingCallback();
 
-			SetLastError( ENetworkReplayError::ServiceUnavailable );
+			SetLastError(EHttpReplayResult::InvalidPayload);
 		}
 	}
 	else
@@ -1936,7 +1930,7 @@ void FHttpNetworkReplayStreamer::HttpStartDownloadingFinished( FHttpRequestPtr H
 		// Reset delegate
 		StartStreamingDelegate = FStartStreamingCallback();
 
-		SetLastError( ENetworkReplayError::ServiceUnavailable );
+		SetLastError(EHttpReplayResult::InvalidHttpResponse);
 	}
 }
 
@@ -1965,7 +1959,7 @@ void FHttpNetworkReplayStreamer::HttpDownloadHeaderFinished( FHttpRequestPtr Htt
 		// Reset delegate
 		StartStreamingDelegate = FStartStreamingCallback();
 
-		SetLastError( ENetworkReplayError::ServiceUnavailable );
+		SetLastError(EHttpReplayResult::InvalidHttpResponse);
 	}
 
 	Delegate.ExecuteIfBound(Result);
@@ -2013,7 +2007,7 @@ void FHttpNetworkReplayStreamer::HttpDownloadFinished( FHttpRequestPtr HttpReque
 		{
 			UE_LOG( LogHttpReplay, Error, TEXT( "FHttpNetworkReplayStreamer::HttpDownloadFinished. StreamChunkIndex changed while request was in flight" ) );
 			StreamArchive.Buffer.Empty();
-			SetLastError( ENetworkReplayError::ServiceUnavailable );
+			SetLastError(EHttpReplayResult::DataUnavailable);
 			return;
 		}
 
@@ -2063,7 +2057,7 @@ void FHttpNetworkReplayStreamer::HttpDownloadFinished( FHttpRequestPtr HttpReque
 				
 				if (!DecompressResponse(HttpResponse, StreamArchive.Buffer))
 				{
-					SetLastError( ENetworkReplayError::ServiceUnavailable );
+					SetLastError(EHttpReplayResult::DecompressionFailed);
 					return;
 				}
 
@@ -2087,7 +2081,7 @@ void FHttpNetworkReplayStreamer::HttpDownloadFinished( FHttpRequestPtr HttpReque
 			{
 				UE_LOG( LogHttpReplay, Error, TEXT( "FHttpNetworkReplayStreamer::HttpDownloadFinished. FAILED (no content when not live). Live: %i, Progress: %i / %i, Start: %i, End: %i, DemoTime: %2.2f. %s" ), ( int )bStreamIsLive, StreamChunkIndex, NumTotalStreamChunks, ( int )StreamTimeRangeStart, ( int )StreamTimeRangeEnd, ( float )TotalDemoTimeInMS / 1000, *BuildRequestErrorString( HttpRequest, HttpResponse ) );
 				StreamArchive.Buffer.Empty();
-				SetLastError( ENetworkReplayError::ServiceUnavailable );
+				SetLastError(EHttpReplayResult::InvalidHttpResponse);
 			}
 		}		
 	}
@@ -2103,7 +2097,7 @@ void FHttpNetworkReplayStreamer::HttpDownloadFinished( FHttpRequestPtr HttpReque
 		{
 			UE_LOG( LogHttpReplay, Error, TEXT( "FHttpNetworkReplayStreamer::HttpDownloadFinished. FAILED, %s" ), *BuildRequestErrorString( HttpRequest, HttpResponse ) );
 			StreamArchive.Buffer.Empty();
-			SetLastError( ENetworkReplayError::ServiceUnavailable );
+			SetLastError(EHttpReplayResult::InvalidHttpResponse);
 		}
 	}
 }
@@ -2231,7 +2225,7 @@ void FHttpNetworkReplayStreamer::HttpRefreshViewerFinished( FHttpRequestPtr Http
 		if ( ++RefreshViewerFails > MaxViewerRefreshFails )
 		{
 			UE_LOG( LogHttpReplay, Error, TEXT( "FHttpNetworkReplayStreamer::HttpRefreshViewerFinished. FAILED, %s" ), *BuildRequestErrorString( HttpRequest, HttpResponse ) );
-			SetLastError( ENetworkReplayError::ServiceUnavailable );
+			SetLastError(EHttpReplayResult::InvalidHttpResponse);
 		}
 		else
 		{
@@ -2314,7 +2308,7 @@ void FHttpNetworkReplayStreamer::HttpEnumerateCheckpointsFinished( FHttpRequestP
 		if ( !CheckpointList.FromJson( JsonString ) )
 		{
 			UE_LOG( LogHttpReplay, Warning, TEXT( "FHttpNetworkReplayStreamer::HttpEnumerateCheckpointsFinished. FromJson FAILED" ) );
-			SetLastError( ENetworkReplayError::ServiceUnavailable );
+			SetLastError(EHttpReplayResult::FailedJsonParse);
 		}
 		else
 		{
@@ -2335,7 +2329,7 @@ void FHttpNetworkReplayStreamer::HttpEnumerateCheckpointsFinished( FHttpRequestP
 	else
 	{
 		UE_LOG( LogHttpReplay, Warning, TEXT( "FHttpNetworkReplayStreamer::HttpEnumerateCheckpointsFinished. FAILED, %s" ), *BuildRequestErrorString( HttpRequest, HttpResponse ) );
-		SetLastError( ENetworkReplayError::ServiceUnavailable );
+		SetLastError(EHttpReplayResult::InvalidHttpResponse);
 	}
 	
 	StartStreamingDelegate.ExecuteIfBound(Result);
@@ -2360,13 +2354,13 @@ void FHttpNetworkReplayStreamer::HttpEnumerateEventsFinished(FHttpRequestPtr Htt
 		else
 		{
 			UE_LOG(LogHttpReplay, Warning, TEXT("FHttpNetworkReplayStreamer::HttpEnumerateEventsFinished. FromJson FAILED"));
-			SetLastError(ENetworkReplayError::ServiceUnavailable);
+			SetLastError(EHttpReplayResult::FailedJsonParse);
 		}
 	}
 	else
 	{
 		UE_LOG( LogHttpReplay, Warning, TEXT( "FHttpNetworkReplayStreamer::HttpEnumerateEventsFinished. FAILED, %s" ), *BuildRequestErrorString( HttpRequest, HttpResponse ) );
-		SetLastError(ENetworkReplayError::ServiceUnavailable);
+		SetLastError(EHttpReplayResult::InvalidHttpResponse);
 	}
 	
 	Delegate.ExecuteIfBound(Result);
@@ -2737,3 +2731,26 @@ TStatId FHttpNetworkReplayStreamingFactory::GetStatId() const
 {
 	RETURN_QUICK_DECLARE_CYCLE_STAT( FHttpNetworkReplayStreamingFactory, STATGROUP_Tickables );
 }
+
+#define CASE_EHTTPREPLAYRESULT_TO_TEXT_RET(txt) case txt: ReturnVal = TEXT(#txt); break;
+
+const TCHAR* LexToString(EHttpReplayResult Enum)
+{
+	const TCHAR* ReturnVal = TEXT("::Invalid");
+
+	switch (Enum)
+	{
+		FOREACH_ENUM_EHTTPREPLAYRESULT(CASE_EHTTPREPLAYRESULT_TO_TEXT_RET)
+	}
+
+	while (*ReturnVal != ':')
+	{
+		ReturnVal++;
+	}
+
+	ReturnVal += 2;
+
+	return ReturnVal;
+}
+
+#undef CASE_EHTTPREPLAYRESULT_TO_TEXT_RET
\ No newline at end of file
diff --git a/Engine/Source/Runtime/NetworkReplayStreaming/HttpNetworkReplayStreaming/Public/HttpNetworkReplayStreaming.h b/Engine/Source/Runtime/NetworkReplayStreaming/HttpNetworkReplayStreaming/Public/HttpNetworkReplayStreaming.h
--- a/Engine/Source/Runtime/NetworkReplayStreaming/HttpNetworkReplayStreaming/Public/HttpNetworkReplayStreaming.h
+++ b/Engine/Source/Runtime/NetworkReplayStreaming/HttpNetworkReplayStreaming/Public/HttpNetworkReplayStreaming.h
@@ -6,6 +6,7 @@
 #include "NetworkReplayStreaming.h"
 #include "Interfaces/IHttpRequest.h"
 #include "Tickable.h"
+#include "HttpNetworkReplayStreaming.generated.h"
 
 class FHttpNetworkReplayStreamer;
 
@@ -197,11 +198,30 @@ public:
 
 PRAGMA_DISABLE_DEPRECATION_WARNINGS
 
+UENUM()
+enum class EHttpReplayResult : uint32
+{
+	Success,
+	FailedJsonParse,
+	DataUnavailable,
+	InvalidHttpResponse,
+	CompressionFailed,
+	DecompressionFailed,
+	InvalidPayload,
+	Unknown,
+};
+
+DECLARE_NETRESULT_ENUM(EHttpReplayResult);
+
+HTTPNETWORKREPLAYSTREAMING_API const TCHAR* LexToString(EHttpReplayResult Enum);
+
 /**
  * Http network replay streaming manager
  */
 class HTTPNETWORKREPLAYSTREAMING_API FHttpNetworkReplayStreamer : public INetworkReplayStreamer
 {
+	using FHttpReplayResult = UE::Net::TNetResult<EHttpReplayResult>;
+
 public:
 	FHttpNetworkReplayStreamer();
 
@@ -244,7 +264,6 @@ public:
 	virtual void		RenameReplay(const FString& ReplayName, const FString& NewName, const FRenameReplayCallback& Delegate) override;
 	virtual void		RenameReplay(const FString& ReplayName, const FString& NewName, const int32 UserIndex, const FRenameReplayCallback& Delegate) override;
 
-	virtual ENetworkReplayError::Type GetLastError() const override;
 	virtual FString		GetReplayID() const override { return SessionName; }
 	virtual EReplayStreamerState GetReplayStreamerState() const override { return StreamerState; }
 	virtual void		SetTimeBufferHintSeconds(const float InTimeBufferHintSeconds) override {}
@@ -274,7 +293,11 @@ public:
 	void ConditionallyDownloadNextChunk();
 	void RefreshViewer( const bool bFinal );
 	void ConditionallyRefreshViewer();
-	void SetLastError( const ENetworkReplayError::Type InLastError );
+	
+	UE_DEPRECATED(5.2, "No longer used")
+	void SetLastError(const ENetworkReplayError::Type InLastError) { SetLastError(EHttpReplayResult::Unknown); }
+	void SetLastError(FHttpReplayResult&& Result);
+
 	virtual void CancelStreamingRequests();
 	void FlushCheckpointInternal( uint32 TimeInMS );
 	virtual void AddEvent( const uint32 TimeInMS, const FString& Group, const FString& Meta, const TArray<uint8>& Data ) override;
@@ -352,8 +375,6 @@ public:
 	FString					ViewerName;
 	uint32					HighPriorityEndTime;
 
-	ENetworkReplayError::Type		StreamerLastError;
-
 	FStartStreamingCallback			StartStreamingDelegate;		// Delegate passed in to StartStreaming
 	FGotoCallback					GotoCheckpointDelegate;
 	int32							DownloadCheckpointIndex;
diff --git a/Engine/Source/Runtime/NetworkReplayStreaming/InMemoryNetworkReplayStreaming/Public/InMemoryNetworkReplayStreaming.h b/Engine/Source/Runtime/NetworkReplayStreaming/InMemoryNetworkReplayStreaming/Public/InMemoryNetworkReplayStreaming.h
--- a/Engine/Source/Runtime/NetworkReplayStreaming/InMemoryNetworkReplayStreaming/Public/InMemoryNetworkReplayStreaming.h
+++ b/Engine/Source/Runtime/NetworkReplayStreaming/InMemoryNetworkReplayStreaming/Public/InMemoryNetworkReplayStreaming.h
@@ -143,7 +143,6 @@ public:
 	virtual void DeleteFinishedStream( const FString& StreamName, const int32 UserIndex, const FDeleteFinishedStreamCallback& Delegate ) override;
 	virtual void EnumerateStreams( const FNetworkReplayVersion& ReplayVersion, const int32 UserIndex, const FString& MetaString, const TArray< FString >& ExtraParms, const FEnumerateStreamsCallback& Delegate ) override;
 	virtual void EnumerateRecentStreams( const FNetworkReplayVersion& ReplayVersion, const int32 UserIndex, const FEnumerateStreamsCallback& Delegate ) override;
-	virtual ENetworkReplayError::Type GetLastError() const override { return ENetworkReplayError::None; }
 	virtual void AddUserToReplay(const FString& UserString) override;
 	virtual void AddEvent(const uint32 TimeInMS, const FString& Group, const FString& Meta, const TArray<uint8>& Data) override;
 	virtual void AddOrUpdateEvent( const FString& Name, const uint32 TimeInMS, const FString& Group, const FString& Meta, const TArray<uint8>& Data ) override {}
diff --git a/Engine/Source/Runtime/NetworkReplayStreaming/LocalFileNetworkReplayStreaming/Private/LocalFileNetworkReplayStreaming.cpp b/Engine/Source/Runtime/NetworkReplayStreaming/LocalFileNetworkReplayStreaming/Private/LocalFileNetworkReplayStreaming.cpp
--- a/Engine/Source/Runtime/NetworkReplayStreaming/LocalFileNetworkReplayStreaming/Private/LocalFileNetworkReplayStreaming.cpp
+++ b/Engine/Source/Runtime/NetworkReplayStreaming/LocalFileNetworkReplayStreaming/Private/LocalFileNetworkReplayStreaming.cpp
@@ -148,7 +148,6 @@ FLocalFileNetworkReplayStreamer::FLocalFileNetworkReplayStreamer()
 	, HighPriorityEndTime(0)
 	, LastGotoTimeInMS(-1)
 	, StreamerState(EReplayStreamerState::Idle)
-	, StreamerLastError(ENetworkReplayError::None)
 	, DemoSavePath(GetDefaultDemoSavePath())
 	, bCacheFileReadsInMemory(false)
 {
@@ -164,7 +163,6 @@ FLocalFileNetworkReplayStreamer::FLocalFileNetworkReplayStreamer(const FString&
 	, HighPriorityEndTime(0)
 	, LastGotoTimeInMS(-1)
 	, StreamerState(EReplayStreamerState::Idle)
-	, StreamerLastError(ENetworkReplayError::None)
 	, DemoSavePath(InDemoSavePath.EndsWith(TEXT("/")) ? InDemoSavePath : InDemoSavePath + FString("/"))
 	, bCacheFileReadsInMemory(false)
 {
@@ -188,22 +186,12 @@ bool FLocalFileNetworkReplayStreamer::ReadReplayInfo(const FString& StreamName,
 	return false;
 }
 
-bool FLocalFileNetworkReplayStreamer::ReadReplayInfo(FArchive& Archive, FLocalFileReplayInfo& Info) const
-{
-	return ReadReplayInfo(Archive, Info, EReadReplayInfoFlags::None);
-}
-
 bool FLocalFileNetworkReplayStreamer::ReadReplayInfo(FArchive& Archive, FLocalFileReplayInfo& Info, EReadReplayInfoFlags Flags) const
 {
 	FLocalFileSerializationInfo DefaultSerializationInfo;
 	return ReadReplayInfo(Archive, Info, DefaultSerializationInfo, Flags);
 }
 
-bool FLocalFileNetworkReplayStreamer::ReadReplayInfo(FArchive& Archive, FLocalFileReplayInfo& Info, FLocalFileSerializationInfo& SerializationInfo) const
-{
-	return ReadReplayInfo(Archive, Info, SerializationInfo, EReadReplayInfoFlags::None);
-}
-
 bool FLocalFileNetworkReplayStreamer::ReadReplayInfo(FArchive& Archive, FLocalFileReplayInfo& Info, FLocalFileSerializationInfo& SerializationInfo, EReadReplayInfoFlags Flags) const
 {
 	// reset the info before reading
@@ -369,9 +357,7 @@ bool FLocalFileNetworkReplayStreamer::ReadReplayInfo(FArchive& Archive, FLocalFi
 	
 						if (Info.bCompressed)
 						{
-							PRAGMA_DISABLE_DEPRECATION_WARNINGS
-							DataChunk.MemorySizeInBytes = GetDecompressedSize(Archive);
-							PRAGMA_ENABLE_DEPRECATION_WARNINGS
+							DataChunk.MemorySizeInBytes = GetDecompressedSizeBackCompat(Archive);
 						}
 						else
 						{
@@ -853,15 +839,11 @@ void FLocalFileNetworkReplayStreamer::CancelStreamingRequests()
 	bStopStreamingCalled = false;
 }
 
-void FLocalFileNetworkReplayStreamer::SetLastError(const ENetworkReplayError::Type InLastError)
+void FLocalFileNetworkReplayStreamer::SetLastError(FLocalFileReplayResult&& Result)
 {
-	CancelStreamingRequests();
-	StreamerLastError = InLastError;
-}
+	SetExtendedError(MoveTemp(Result));
 
-ENetworkReplayError::Type FLocalFileNetworkReplayStreamer::GetLastError() const
-{
-	return StreamerLastError;
+	CancelStreamingRequests();
 }
 
 void FLocalFileNetworkReplayStreamer::StopStreaming()
@@ -876,14 +858,14 @@ void FLocalFileNetworkReplayStreamer::StopStreaming()
 
 	if (!IsStreaming())
 	{
-		UE_LOG( LogLocalFileReplay, Warning, TEXT( "FLocalFileNetworkReplayStreamer::StopStreaming. Not currently streaming." ) );
+		UE_LOG(LogLocalFileReplay, Log, TEXT("FLocalFileNetworkReplayStreamer::StopStreaming. Not currently streaming."));
 		check( bStopStreamingCalled == false );
 		return;
 	}
 
 	if (bStopStreamingCalled)
 	{
-		UE_LOG( LogLocalFileReplay, Warning, TEXT( "FLocalFileNetworkReplayStreamer::StopStreaming. Already called" ) );
+		UE_LOG(LogLocalFileReplay, Log, TEXT("FLocalFileNetworkReplayStreamer::StopStreaming. Already called"));
 		return;
 	}
 
@@ -953,7 +935,7 @@ void FLocalFileNetworkReplayStreamer::UpdateTotalDemoTime(uint32 TimeInMS)
 
 bool FLocalFileNetworkReplayStreamer::IsDataAvailable() const
 {
-	if (GetLastError() != ENetworkReplayError::None)
+	if (HasError())
 	{
 		return false;
 	}
@@ -1718,8 +1700,8 @@ void FLocalFileNetworkReplayStreamer::FlushStream(const uint32 TimeInMS)
 	// Save any newly streamed data to disk
 	UE_LOG(LogLocalFileReplay, Verbose, TEXT("FLocalFileNetworkReplayStreamer::FlushStream. StreamChunkIndex: %i, Size: %i"), StreamChunkIndex, StreamAr.Buffer.Num());
 
-	AddGenericRequestToQueue<FLocalFileReplayInfo>(EQueuedLocalFileRequestType::WritingStream, 
-		[this, StreamChunkStartMS, StreamChunkEndMS, StreamData=MoveTemp(StreamAr.Buffer), EncryptionKey=CurrentReplayInfo.EncryptionKey](FLocalFileReplayInfo& ReplayInfo) mutable
+	AddGenericRequestToQueue<FLocalFileReplayInfo>(EQueuedLocalFileRequestType::WritingStream,
+		[this, StreamChunkStartMS, StreamChunkEndMS, StreamData = MoveTemp(StreamAr.Buffer), EncryptionKey = CurrentReplayInfo.EncryptionKey](FLocalFileReplayInfo& ReplayInfo) mutable
 		{
 			SCOPE_CYCLE_COUNTER(STAT_LocalReplay_FlushStream);
 
@@ -1797,7 +1779,7 @@ void FLocalFileNetworkReplayStreamer::FlushStream(const uint32 TimeInMS)
 					LocalFileAr->Seek(SavedPos);
 					*LocalFileAr << ChunkSize;
 				}
-			
+
 				LocalFileAr = nullptr;
 			}
 
@@ -1812,7 +1794,7 @@ void FLocalFileNetworkReplayStreamer::FlushStream(const uint32 TimeInMS)
 			else
 			{
 				UE_LOG(LogLocalFileReplay, Error, TEXT("FLocalFileNetworkReplayStreamer::FlushStream failed with invalid replay info."));
-				SetLastError(ENetworkReplayError::ServiceUnavailable);
+				SetLastError(ELocalFileReplayResult::InvalidReplayInfo);
 			}
 		});
 
@@ -1965,7 +1947,7 @@ void FLocalFileNetworkReplayStreamer::FlushCheckpointInternal(const uint32 TimeI
 			else
 			{
 				UE_LOG(LogLocalFileReplay, Error, TEXT("FLocalFileNetworkReplayStreamer::FlushCheckpointInternal failed with invalid replay info."));
-				SetLastError(ENetworkReplayError::ServiceUnavailable);
+				SetLastError(ELocalFileReplayResult::InvalidReplayInfo);
 			}
 		});
 
@@ -3014,7 +2996,7 @@ void FLocalFileNetworkReplayStreamer::WriteHeader()
 			else
 			{
 				UE_LOG(LogLocalFileReplay, Error, TEXT("FLocalFileNetworkReplayStreamer::WriteHeader failed with invalid replay info."));
-				SetLastError(ENetworkReplayError::ServiceUnavailable);
+				SetLastError(ELocalFileReplayResult::InvalidReplayInfo);
 			}
 		});
 
@@ -3044,7 +3026,7 @@ void FLocalFileNetworkReplayStreamer::SetHighPriorityTimeRange(const uint32 Star
 
 bool FLocalFileNetworkReplayStreamer::IsDataAvailableForTimeRange(const uint32 StartTimeInMS, const uint32 EndTimeInMS)
 {
-	if (GetLastError() != ENetworkReplayError::None)
+	if (HasError())
 	{
 		return false;
 	}
@@ -3186,7 +3168,7 @@ void FLocalFileNetworkReplayStreamer::ConditionallyLoadNextChunk()
 							{
 								UE_LOG(LogLocalFileReplay, Error, TEXT("ConditionallyLoadNextChunk failed to decrypt data."));
 								RequestData.DataBuffer.Empty();
-								RequestData.bAsyncError = true;
+								RequestData.AsyncError = ELocalFileReplayResult::DecryptBuffer;
 								return;
 							}
 						}
@@ -3194,7 +3176,7 @@ void FLocalFileNetworkReplayStreamer::ConditionallyLoadNextChunk()
 						{
 							UE_LOG(LogLocalFileReplay, Error, TEXT("ConditionallyLoadNextChunk: Replay is marked encrypted but streamer does not support it."));
 							RequestData.DataBuffer.Empty();
-							RequestData.bAsyncError = true;
+							RequestData.AsyncError = ELocalFileReplayResult::EncryptionNotSupported;
 							return;
 						}
 					}
@@ -3214,7 +3196,7 @@ void FLocalFileNetworkReplayStreamer::ConditionallyLoadNextChunk()
 							{
 								UE_LOG(LogLocalFileReplay, Error, TEXT("ConditionallyLoadNextChunk failed to uncompresss data."));
 								RequestData.DataBuffer.Empty();
-								RequestData.bAsyncError = true;
+								RequestData.AsyncError = ELocalFileReplayResult::DecompressBuffer;
 								return;
 							}
 						}
@@ -3222,7 +3204,7 @@ void FLocalFileNetworkReplayStreamer::ConditionallyLoadNextChunk()
 						{
 							UE_LOG(LogLocalFileReplay, Error, TEXT("ConditionallyLoadNextChunk: Replay is marked compressed but streamer does not support it."));
 							RequestData.DataBuffer.Empty();
-							RequestData.bAsyncError = true;
+							RequestData.AsyncError = ELocalFileReplayResult::CompressionNotSupported;
 							return;
 						}
 					}
@@ -3236,9 +3218,9 @@ void FLocalFileNetworkReplayStreamer::ConditionallyLoadNextChunk()
 			LLM_SCOPE(ELLMTag::Replays);
 
 			// Hijacking this error code to indicate a failure in encryption/compression
-			if (RequestData.bAsyncError)
+			if (RequestData.AsyncError != ELocalFileReplayResult::Success)
 			{
-				SetLastError(ENetworkReplayError::ServiceUnavailable);
+				SetLastError(FLocalFileReplayResult(RequestData.AsyncError));
 				return;
 			}
 
@@ -3248,7 +3230,7 @@ void FLocalFileNetworkReplayStreamer::ConditionallyLoadNextChunk()
 				UE_LOG(LogLocalFileReplay, Warning, TEXT("FLocalFileNetworkReplayStreamer::ConditionallyLoadNextChunk failed with requested chunk index mismatch."));
 
 				StreamAr.Reset();
-				SetLastError(ENetworkReplayError::ServiceUnavailable);
+				SetLastError(ELocalFileReplayResult::StreamChunkIndexMismatch);
 				return;
 			}
 
@@ -3480,12 +3462,33 @@ void FLocalFileNetworkReplayStreamer::UpdateCurrentReplayInfo(FLocalFileReplayIn
 	}
 }
 
-PRAGMA_DISABLE_DEPRECATION_WARNINGS
-int32 FLocalFileNetworkReplayStreamer::GetDecompressedSize(FArchive& InCompressed) const
+int32 FLocalFileNetworkReplayStreamer::GetDecompressedSizeBackCompat(FArchive& InCompressed) const
 {
 	return 0;
 }
-PRAGMA_ENABLE_DEPRECATION_WARNINGS
+
+#define CASE_ELOCALFILEREPLAYRESULT_TO_TEXT_RET(txt) case txt: ReturnVal = TEXT(#txt); break;
+
+const TCHAR* LexToString(ELocalFileReplayResult Enum)
+{
+	const TCHAR* ReturnVal = TEXT("::Invalid");
+
+	switch (Enum)
+	{
+		FOREACH_ENUM_ELOCALFILEREPLAYRESULT(CASE_ELOCALFILEREPLAYRESULT_TO_TEXT_RET)
+	}
+
+	while (*ReturnVal != ':')
+	{
+		ReturnVal++;
+	}
+
+	ReturnVal += 2;
+
+	return ReturnVal;
+}
+
+#undef CASE_ELOCALFILEREPLAYRESULT_TO_TEXT_RET
 
 IMPLEMENT_MODULE(FLocalFileNetworkReplayStreamingFactory, LocalFileNetworkReplayStreaming)
 
diff --git a/Engine/Source/Runtime/NetworkReplayStreaming/LocalFileNetworkReplayStreaming/Public/LocalFileNetworkReplayStreaming.h b/Engine/Source/Runtime/NetworkReplayStreaming/LocalFileNetworkReplayStreaming/Public/LocalFileNetworkReplayStreaming.h
--- a/Engine/Source/Runtime/NetworkReplayStreaming/LocalFileNetworkReplayStreaming/Public/LocalFileNetworkReplayStreaming.h
+++ b/Engine/Source/Runtime/NetworkReplayStreaming/LocalFileNetworkReplayStreaming/Public/LocalFileNetworkReplayStreaming.h
@@ -12,6 +12,7 @@
 #include "Async/Async.h"
 #include "Templates/SharedPointer.h"
 #include "HAL/ThreadSafeBool.h"
+#include "LocalFileNetworkReplayStreaming.generated.h"
 
 class FNetworkReplayVersion;
 class FLocalFileNetworkReplayStreamer;
@@ -231,6 +232,23 @@ namespace EQueuedLocalFileRequestType
 	}
 };
 
+UENUM()
+enum class ELocalFileReplayResult : uint32
+{
+	Success,
+	InvalidReplayInfo,
+	StreamChunkIndexMismatch,
+	DecompressBuffer,
+	CompressionNotSupported,
+	DecryptBuffer,
+	EncryptionNotSupported,
+	Unknown,
+};
+
+DECLARE_NETRESULT_ENUM(ELocalFileReplayResult);
+
+LOCALFILENETWORKREPLAYSTREAMING_API const TCHAR* LexToString(ELocalFileReplayResult Enum);
+
 class FCachedFileRequest
 {
 public:
@@ -389,8 +407,12 @@ public:
 	DelegateResultType DelegateResult;
 	FLocalFileReplayInfo ReplayInfo;
 	TArray<uint8> DataBuffer;
+
+	UE_DEPRECATED(5.2, "No longer used")
 	bool bAsyncError = false;
-};
+
+	ELocalFileReplayResult AsyncError = ELocalFileReplayResult::Success;
+};	
 
 template <typename DelegateResultType>
 class TGenericCachedLocalFileRequest : public TGenericQueuedLocalFileRequest<TLocalFileRequestCommonData<DelegateResultType>>
@@ -422,6 +444,8 @@ protected:
 /** Local file streamer that supports playback/recording to a single file on disk */
 class LOCALFILENETWORKREPLAYSTREAMING_API FLocalFileNetworkReplayStreamer : public INetworkReplayStreamer, public TSharedFromThis<FLocalFileNetworkReplayStreamer>
 {
+	using FLocalFileReplayResult = UE::Net::TNetResult<ELocalFileReplayResult>;
+
 public:
 	FLocalFileNetworkReplayStreamer();
 	FLocalFileNetworkReplayStreamer(const FString& InDemoSavePath);
@@ -448,7 +472,6 @@ public:
 	virtual void DeleteFinishedStream( const FString& StreamName, const int32 UserIndex, const FDeleteFinishedStreamCallback& Delegate ) override;
 	virtual void EnumerateStreams( const FNetworkReplayVersion& InReplayVersion, const int32 UserIndex, const FString& MetaString, const TArray< FString >& ExtraParms, const FEnumerateStreamsCallback& Delegate ) override;
 	virtual void EnumerateRecentStreams( const FNetworkReplayVersion& ReplayVersion, const int32 UserIndex, const FEnumerateStreamsCallback& Delegate ) override;
-	virtual ENetworkReplayError::Type GetLastError() const override;
 	virtual void AddUserToReplay(const FString& UserString) override;
 	virtual void AddEvent(const uint32 TimeInMS, const FString& Group, const FString& Meta, const TArray<uint8>& Data) override;
 	virtual void AddOrUpdateEvent(const FString& Name, const uint32 TimeInMS, const FString& Group, const FString& Meta, const TArray<uint8>& Data) override;
@@ -478,9 +501,6 @@ public:
 
 	virtual bool SupportsCompression() const { return false; }
 
-	UE_DEPRECATED(4.25, "No longer used")
-	virtual int32 GetDecompressedSize(FArchive& InCompressed) const;
-
 	virtual bool DecompressBuffer(const TArray<uint8>& InCompressed, TArray<uint8>& OutBuffer) const { return false; }
 	virtual bool CompressBuffer(const TArray<uint8>& InBuffer, TArray<uint8>& OutCompressed) const { return false; }
 
@@ -589,7 +609,15 @@ protected:
 	bool IsFileRequestPendingOrInProgress(const EQueuedLocalFileRequestType::Type RequestType) const;
 	void CancelStreamingRequests();
 
-	void SetLastError(const ENetworkReplayError::Type InLastError);
+	PRAGMA_DISABLE_DEPRECATION_WARNINGS
+	UE_DEPRECATED(5.2, "No longer used")
+	void SetLastError(const ENetworkReplayError::Type InLastError) 
+	{ 
+		SetLastError(ELocalFileReplayResult::Unknown); 
+	}
+	PRAGMA_ENABLE_DEPRECATION_WARNINGS
+
+	void SetLastError(FLocalFileReplayResult&& Result);
 
 	void ConditionallyFlushStream();
 	void ConditionallyLoadNextChunk();
@@ -611,15 +639,7 @@ protected:
 	}
 
 	bool ReadReplayInfo(const FString& StreamName, FLocalFileReplayInfo& OutReplayInfo, EReadReplayInfoFlags Flags) const;
-
-	UE_DEPRECATED(4.25, "Now takes a set of read flags")
-	bool ReadReplayInfo(FArchive& Archive, FLocalFileReplayInfo& OutReplayInfo) const;
-
 	bool ReadReplayInfo(FArchive& Archive, FLocalFileReplayInfo& OutReplayInfo, EReadReplayInfoFlags Flags) const;
-
-	UE_DEPRECATED(4.25, "Now takes a set of read flags")
-	bool ReadReplayInfo(FArchive& Archive, FLocalFileReplayInfo& OutReplayInfo, struct FLocalFileSerializationInfo& SerializationInfo) const;
-
 	bool ReadReplayInfo(FArchive& Archive, FLocalFileReplayInfo& OutReplayInfo, struct FLocalFileSerializationInfo& SerializationInfo, EReadReplayInfoFlags Flags) const;
 
 	bool WriteReplayInfo(const FString& StreamName, const FLocalFileReplayInfo& ReplayInfo);
@@ -660,8 +680,6 @@ protected:
 	/** Overall state of the streamer */
 	EReplayStreamerState StreamerState;
 
-	ENetworkReplayError::Type StreamerLastError;
-
 	/** Remember the name of the current stream, if any. */
 	FString CurrentStreamName;
 
@@ -677,6 +695,8 @@ protected:
 
 	void UpdateCurrentReplayInfo(FLocalFileReplayInfo& ReplayInfo);
 
+	virtual int32 GetDecompressedSizeBackCompat(FArchive& InCompressed) const;
+
 public:
 	static const FString& GetDefaultDemoSavePath();
 	static FString GetDemoFullFilename(const FString& DemoPath, const FString& FileName);
diff --git a/Engine/Source/Runtime/NetworkReplayStreaming/NetworkReplayStreaming/NetworkReplayStreaming.Build.cs b/Engine/Source/Runtime/NetworkReplayStreaming/NetworkReplayStreaming/NetworkReplayStreaming.Build.cs
--- a/Engine/Source/Runtime/NetworkReplayStreaming/NetworkReplayStreaming/NetworkReplayStreaming.Build.cs
+++ b/Engine/Source/Runtime/NetworkReplayStreaming/NetworkReplayStreaming/NetworkReplayStreaming.Build.cs
@@ -14,6 +14,7 @@ namespace UnrealBuildTool.Rules
 					"Core",
                     "Json",
 					"Analytics",
+					"NetCore",
 				}
 			);
 		}
diff --git a/Engine/Source/Runtime/NetworkReplayStreaming/NetworkReplayStreaming/Private/NetworkReplayStreaming.cpp b/Engine/Source/Runtime/NetworkReplayStreaming/NetworkReplayStreaming/Private/NetworkReplayStreaming.cpp
--- a/Engine/Source/Runtime/NetworkReplayStreaming/NetworkReplayStreaming/Private/NetworkReplayStreaming.cpp
+++ b/Engine/Source/Runtime/NetworkReplayStreaming/NetworkReplayStreaming/Private/NetworkReplayStreaming.cpp
@@ -171,4 +171,33 @@ TArray<FAnalyticsEventAttribute> INetworkReplayStreamer::AppendCommonReplayAttri
 	INetworkReplayStreamer::OnReplayGetAnalyticsAttributes.Broadcast(this, CommonAttributes);
 
 	return CommonAttributes;
+}
+
+PRAGMA_DISABLE_DEPRECATION_WARNINGS
+ENetworkReplayError::Type INetworkReplayStreamer::GetLastError() const
+{
+	return HasError() ? ENetworkReplayError::ServiceUnavailable : ENetworkReplayError::None;
+}
+PRAGMA_ENABLE_DEPRECATION_WARNINGS
+
+bool INetworkReplayStreamer::HasError() const
+{
+	return ExtendedError.IsValid();
+}
+
+void INetworkReplayStreamer::SetExtendedError(UE::Net::FNetResult&& Result)
+{
+	AddToChainResultPtr(ExtendedError, MoveTemp(Result));
+}
+
+UE::Net::EHandleNetResult INetworkReplayStreamer::HandleLastError(UE::Net::FNetResultManager& ResultManager)
+{
+	UE::Net::EHandleNetResult ReturnVal = UE::Net::EHandleNetResult::NotHandled;
+
+	if (ExtendedError.IsValid())
+	{
+		ReturnVal = ResultManager.HandleNetResult(MoveTemp(*ExtendedError));
+	}
+	
+	return ReturnVal;
 }
\ No newline at end of file
diff --git a/Engine/Source/Runtime/NetworkReplayStreaming/NetworkReplayStreaming/Public/NetworkReplayStreaming.h b/Engine/Source/Runtime/NetworkReplayStreaming/NetworkReplayStreaming/Public/NetworkReplayStreaming.h
--- a/Engine/Source/Runtime/NetworkReplayStreaming/NetworkReplayStreaming/Public/NetworkReplayStreaming.h
+++ b/Engine/Source/Runtime/NetworkReplayStreaming/NetworkReplayStreaming/Public/NetworkReplayStreaming.h
@@ -16,7 +16,10 @@
 #include "Misc/NetworkVersion.h"
 #include "Modules/ModuleInterface.h"
 #include "Modules/ModuleManager.h"
+#include "Net/Core/Connection/NetResult.h"
+#include "Net/Core/Connection/NetResultManager.h"
 #include "Serialization/JsonSerializerMacros.h"
+#include "Templates/PimplPtr.h"
 #include "Templates/SharedPointer.h"
 #include "UObject/NameTypes.h"
 
@@ -98,7 +101,7 @@ struct FNetworkReplayStreamInfo
 
 namespace ENetworkReplayError
 {
-	enum Type
+	enum UE_DEPRECATED(5.2, "No longer used") Type
 	{
 		/** There are currently no issues */
 		None,
@@ -107,6 +110,7 @@ namespace ENetworkReplayError
 		ServiceUnavailable,
 	};
 
+	PRAGMA_DISABLE_DEPRECATION_WARNINGS
 	inline const TCHAR* ToString( const ENetworkReplayError::Type FailureType )
 	{
 		switch ( FailureType )
@@ -118,6 +122,7 @@ namespace ENetworkReplayError
 		}
 		return TEXT( "Unknown ENetworkReplayError error" );
 	}
+	PRAGMA_ENABLE_DEPRECATION_WARNINGS
 }
 
 /**
@@ -617,8 +622,17 @@ public:
 	 */
 	virtual void EnumerateRecentStreams(const FNetworkReplayVersion& ReplayVersion, const int32 UserIndex, const FEnumerateStreamsCallback& Delegate) = 0;
 
+	PRAGMA_DISABLE_DEPRECATION_WARNINGS
 	/** Returns the last error that occurred while streaming replays */
-	virtual ENetworkReplayError::Type GetLastError() const = 0;
+	UE_DEPRECATED(5.2, "Deprecated in favor of HasError and HandleLastError")
+	virtual ENetworkReplayError::Type GetLastError() const;
+	PRAGMA_ENABLE_DEPRECATION_WARNINGS
+
+	virtual bool HasError() const;
+
+	virtual UE::Net::EHandleNetResult HandleLastError(UE::Net::FNetResultManager& ResultManager);
+
+	virtual void SetExtendedError(UE::Net::FNetResult&& Result);
 
 	/**
 	 * Adds a join-in-progress user to the set of users associated with the currently recording replay (if any)
@@ -659,6 +673,9 @@ public:
 
 	/** Called from base streamer interface AppendCommonReplayAttributes to set common attributes */
 	static FOnReplayGetAnalyticsAttributes OnReplayGetAnalyticsAttributes;
+
+private:
+	TPimplPtr<UE::Net::FNetResult, EPimplPtrMode::DeepCopy> ExtendedError;
 };
 
 /** Replay streamer factory */
@@ -710,4 +727,4 @@ public:
 
 private:
 	TSet<FName> LoadedFactories;
-};
+};
\ No newline at end of file
diff --git a/Engine/Source/Runtime/NetworkReplayStreaming/NullNetworkReplayStreaming/Public/NullNetworkReplayStreaming.h b/Engine/Source/Runtime/NetworkReplayStreaming/NullNetworkReplayStreaming/Public/NullNetworkReplayStreaming.h
--- a/Engine/Source/Runtime/NetworkReplayStreaming/NullNetworkReplayStreaming/Public/NullNetworkReplayStreaming.h
+++ b/Engine/Source/Runtime/NetworkReplayStreaming/NullNetworkReplayStreaming/Public/NullNetworkReplayStreaming.h
@@ -64,7 +64,6 @@ public:
 	virtual void DeleteFinishedStream( const FString& StreamName, const int32 UserIndex, const FDeleteFinishedStreamCallback& Delegate ) override;
 	virtual void EnumerateStreams( const FNetworkReplayVersion& InReplayVersion, const int32 UserIndex, const FString& MetaString, const TArray< FString >& ExtraParms, const FEnumerateStreamsCallback& Delegate ) override;
 	virtual void EnumerateRecentStreams( const FNetworkReplayVersion& ReplayVersion, const int32 UserIndex, const FEnumerateStreamsCallback& Delegate ) override;
-	virtual ENetworkReplayError::Type GetLastError() const override { return ENetworkReplayError::None; }
 	virtual void AddUserToReplay(const FString& UserString) override;
 	virtual void AddEvent(const uint32 TimeInMS, const FString& Group, const FString& Meta, const TArray<uint8>& Data) override;
 	virtual void AddOrUpdateEvent( const FString& Name, const uint32 TimeInMS, const FString& Group, const FString& Meta, const TArray<uint8>& Data ) override {}
