Change GetNaniteOverride() to use local recursion guard instead of the material reentrant flag.
The material reentrant flag assumed game or render thread only.

diff --git a/Engine/Source/Runtime/Engine/Classes/Materials/Material.h b/Engine/Source/Runtime/Engine/Classes/Materials/Material.h
--- a/Engine/Source/Runtime/Engine/Classes/Materials/Material.h
+++ b/Engine/Source/Runtime/Engine/Classes/Materials/Material.h
@@ -1025,7 +1025,7 @@ public:
 	ENGINE_API virtual UPhysicalMaterial* GetPhysicalMaterial() const override;
 	ENGINE_API virtual UPhysicalMaterialMask* GetPhysicalMaterialMask() const override;
 	ENGINE_API virtual UPhysicalMaterial* GetPhysicalMaterialFromMap(int32 Index) const override;
-	ENGINE_API virtual UMaterialInterface* GetNaniteOverride() override;
+	ENGINE_API virtual UMaterialInterface* GetNaniteOverride(TMicRecursionGuard RecursionGuard = TMicRecursionGuard()) override;
 	ENGINE_API virtual void GetUsedTextures(TArray<UTexture*>& OutTextures, EMaterialQualityLevel::Type QualityLevel, bool bAllQualityLevels, ERHIFeatureLevel::Type FeatureLevel, bool bAllFeatureLevels) const override;
 	ENGINE_API virtual void GetUsedTexturesAndIndices(TArray<UTexture*>& OutTextures, TArray< TArray<int32> >& OutIndices, EMaterialQualityLevel::Type QualityLevel, ERHIFeatureLevel::Type FeatureLevel) const override;
 	ENGINE_API virtual void OverrideTexture(const UTexture* InTextureToOverride, UTexture* OverrideTexture, ERHIFeatureLevel::Type InFeatureLevel) override;
diff --git a/Engine/Source/Runtime/Engine/Classes/Materials/MaterialInstance.h b/Engine/Source/Runtime/Engine/Classes/Materials/MaterialInstance.h
--- a/Engine/Source/Runtime/Engine/Classes/Materials/MaterialInstance.h
+++ b/Engine/Source/Runtime/Engine/Classes/Materials/MaterialInstance.h
@@ -715,7 +715,7 @@ public:
 	virtual ENGINE_API UPhysicalMaterial* GetPhysicalMaterial() const override;
 	virtual ENGINE_API UPhysicalMaterialMask* GetPhysicalMaterialMask() const override;
 	virtual ENGINE_API UPhysicalMaterial* GetPhysicalMaterialFromMap(int32 Index) const override;
-	virtual ENGINE_API UMaterialInterface* GetNaniteOverride() override;
+	virtual ENGINE_API UMaterialInterface* GetNaniteOverride(TMicRecursionGuard RecursionGuard = TMicRecursionGuard()) override;
 	virtual ENGINE_API bool UpdateLightmassTextureTracking() override;
 	virtual ENGINE_API bool GetCastShadowAsMasked() const override;
 	virtual ENGINE_API float GetEmissiveBoost() const override;
diff --git a/Engine/Source/Runtime/Engine/Classes/Materials/MaterialInterface.h b/Engine/Source/Runtime/Engine/Classes/Materials/MaterialInterface.h
--- a/Engine/Source/Runtime/Engine/Classes/Materials/MaterialInterface.h
+++ b/Engine/Source/Runtime/Engine/Classes/Materials/MaterialInterface.h
@@ -503,7 +503,7 @@ public:
 	 * Get the associated nanite override material
 	 * @return - nanite override, will be this material if none was set
 	 */
-	virtual UMaterialInterface* GetNaniteOverride() PURE_VIRTUAL(UMaterialInterface::GetNaniteOverride, return nullptr;);
+	virtual UMaterialInterface* GetNaniteOverride(TMicRecursionGuard RecursionGuard = TMicRecursionGuard()) PURE_VIRTUAL(UMaterialInterface::GetNaniteOverride, return nullptr;);
 
 	/**
 	 * Precache PSOs which can be used for this material for the given vertex factory type and material paramaters
diff --git a/Engine/Source/Runtime/Engine/Private/Materials/Material.cpp b/Engine/Source/Runtime/Engine/Private/Materials/Material.cpp
--- a/Engine/Source/Runtime/Engine/Private/Materials/Material.cpp
+++ b/Engine/Source/Runtime/Engine/Private/Materials/Material.cpp
@@ -2084,7 +2084,7 @@ void UMaterial::SetRenderTracePhysicalMaterialOutputs(TArrayView<TObjectPtr<clas
 	RenderTracePhysicalMaterialOutputs = PhysicalMaterials;
 }
 
-UMaterialInterface* UMaterial::GetNaniteOverride()
+UMaterialInterface* UMaterial::GetNaniteOverride(TMicRecursionGuard RecursionGuard)
 {
 	return NaniteOverrideMaterial.GetOverrideMaterial();
 }
diff --git a/Engine/Source/Runtime/Engine/Private/Materials/MaterialInstance.cpp b/Engine/Source/Runtime/Engine/Private/Materials/MaterialInstance.cpp
--- a/Engine/Source/Runtime/Engine/Private/Materials/MaterialInstance.cpp
+++ b/Engine/Source/Runtime/Engine/Private/Materials/MaterialInstance.cpp
@@ -1745,22 +1745,16 @@ UPhysicalMaterial* UMaterialInstance::GetPhysicalMaterialFromMap(int32 Index) co
 	return PhysicalMaterialMap[Index];
 }
 
-UMaterialInterface* UMaterialInstance::GetNaniteOverride()
+UMaterialInterface* UMaterialInstance::GetNaniteOverride(TMicRecursionGuard RecursionGuard)
 {
-	if (GetReentrantFlag())
-	{
-		return nullptr;
-	}
-
-	FMICReentranceGuard	Guard(const_cast<UMaterialInstance*>(this));
-	
 	if (NaniteOverrideMaterial.bEnableOverride)
 	{
 		return NaniteOverrideMaterial.GetOverrideMaterial();
 	}
-	else if (Parent)
+	else if (Parent && !RecursionGuard.Contains(this))
 	{
-		return Parent->GetNaniteOverride();
+		RecursionGuard.Set(this);
+		return Parent->GetNaniteOverride(RecursionGuard);
 	}
 	else
 	{
