Edigrate: Added early exit to USoundWave::CachePlatformData() if we are loading an already cooked asset in an editor environment.

diff --git a/Engine/Source/Runtime/Engine/Private/AudioDerivedData.cpp b/Engine/Source/Runtime/Engine/Private/AudioDerivedData.cpp
--- a/Engine/Source/Runtime/Engine/Private/AudioDerivedData.cpp
+++ b/Engine/Source/Runtime/Engine/Private/AudioDerivedData.cpp
@@ -1810,6 +1810,12 @@ void USoundWave::CachePlatformData(bool bAsyncCache)
 {
 	check(SoundWaveDataPtr);
 
+	// don't interact with the DDC if we were loaded from cooked data in editor
+	if (bLoadedFromCookedData)
+	{
+		return;
+	}
+
 	FString DerivedDataKey;
 	FName AudioFormat = GetWaveFormatForRunningPlatform(*this);
 	const FPlatformAudioCookOverrides* CompressionOverrides = GetCookOverridesForRunningPlatform();
