Deprecated UWorldSettings::bEnableLargeWorlds.
Worlds are already large by default. Set UE_USE_UE4_WORLD_MAX to revert to the UE4 world boundaries.

diff --git a/Engine/Plugins/Runtime/GeoReferencing/Source/GeoReferencing/Private/RoundPlanetPawn.cpp b/Engine/Plugins/Runtime/GeoReferencing/Source/GeoReferencing/Private/RoundPlanetPawn.cpp
--- a/Engine/Plugins/Runtime/GeoReferencing/Source/GeoReferencing/Private/RoundPlanetPawn.cpp
+++ b/Engine/Plugins/Runtime/GeoReferencing/Source/GeoReferencing/Private/RoundPlanetPawn.cpp
@@ -176,10 +176,6 @@ void ARoundPlanetPawn::BeginPlay()
 {
 	Super::BeginPlay();
 
-	// Round planet is huge, if we enable world bounds check, the pawn behavior might be unpredictable when flying...
-	// LWC_TODO - should not be needed after WORLD_MAX updates. 
-	GetWorldSettings()->bEnableLargeWorlds = true;
-
 	if (!GeoReferencingSystem)
 	{
 		GeoReferencingSystem = AGeoReferencingSystem::GetGeoReferencingSystem(GetWorld());
diff --git a/Engine/Source/Runtime/Engine/Classes/GameFramework/WorldSettings.h b/Engine/Source/Runtime/Engine/Classes/GameFramework/WorldSettings.h
--- a/Engine/Source/Runtime/Engine/Classes/GameFramework/WorldSettings.h
+++ b/Engine/Source/Runtime/Engine/Classes/GameFramework/WorldSettings.h
@@ -481,13 +481,8 @@ class ENGINE_API AWorldSettings : public AInfo, public IInterface_AssetUserData
 
 	/** DEFAULT BASIC PHYSICS SETTINGS **/
 
-	/** If true, configures engine for large world testing. Disables CheckStillInWorld checks and octree visibility testing.
-	 *	See UE_USE_UE4_WORLD_MAX for a more correct alternative. This setting will be removed once UE_LARGE_WORLD_MAX is considered stable. */
+	/** If true, enables CheckStillInWorld checks */
 	UPROPERTY(EditAnywhere, BlueprintReadOnly, Category=World, AdvancedDisplay)
-	uint8 bEnableLargeWorlds:1;
-
-	/** If true, enables CheckStillInWorld checks. Note: Do not set this manually if experimenting with large worlds. @see bEnableLargeWorlds */
-	UPROPERTY(EditAnywhere, BlueprintReadOnly, Category=World, AdvancedDisplay, meta=(EditCondition = "!bEnableLargeWorlds"))
 	uint8 bEnableWorldBoundsChecks:1;
 
 protected:
@@ -636,11 +631,11 @@ public:
 	float WorldToMeters;
 
 	// any actor falling below this level gets destroyed
-	UPROPERTY(config, EditAnywhere, BlueprintReadOnly, Category=World, meta=(editcondition = "bEnableWorldBoundsChecks && !bEnableLargeWorlds"))
+	UPROPERTY(config, EditAnywhere, BlueprintReadOnly, Category=World, meta=(editcondition = "bEnableWorldBoundsChecks"))
 	float KillZ;
 
 	// The type of damage inflicted when a actor falls below KillZ
-	UPROPERTY(EditAnywhere, BlueprintReadOnly, Category=World, AdvancedDisplay, meta=(editcondition = "bEnableWorldBoundsChecks && !bEnableLargeWorlds"))
+	UPROPERTY(EditAnywhere, BlueprintReadOnly, Category=World, AdvancedDisplay, meta=(editcondition = "bEnableWorldBoundsChecks"))
 	TSubclassOf<UDamageType> KillZDamageType;
 
 	// current gravity actually being used
@@ -888,7 +883,7 @@ public:
 	bool IsAISystemEnabled() const { return bEnableAISystem; }
 
 	/** @return whether given world is restricting actors to +-HALF_WORLD_MAX bounds, destroying actors that move below KillZ */
-	bool AreWorldBoundsChecksEnabled() const { return bEnableWorldBoundsChecks && !bEnableLargeWorlds; }
+	bool AreWorldBoundsChecksEnabled() const { return bEnableWorldBoundsChecks; }
 	
 	/**
 	 * Called from GameStateBase, calls BeginPlay on all actors
@@ -940,8 +935,6 @@ private:
 	void SanitizeBookmarkClasses();
 	void UpdateBookmarkClass();
 
-	void UpdateEnableLargeWorldsCVars(bool bEnable) const;
-
 	/**
 	 * Maximum number of bookmarks allowed.
 	 * Changing this will change the allocation of the bookmarks array, and when shrinking
@@ -1048,6 +1041,9 @@ public:
 private: //DEPRECATED
 	UPROPERTY()
 	bool bEnableHierarchicalLODSystem_DEPRECATED;
+
+	UPROPERTY(meta=(DeprecatedProperty, DeprecationMessage="As of UE 5.1 all worlds are large. Set UE_USE_UE4_WORLD_MAX in EngineDefines.h to alter this."))
+	uint8 bEnableLargeWorlds_DEPRECATED:1;
 #endif
 };
 
diff --git a/Engine/Source/Runtime/Engine/Private/WorldSettings.cpp b/Engine/Source/Runtime/Engine/Private/WorldSettings.cpp
--- a/Engine/Source/Runtime/Engine/Private/WorldSettings.cpp
+++ b/Engine/Source/Runtime/Engine/Private/WorldSettings.cpp
@@ -73,7 +73,6 @@ AWorldSettings::AWorldSettings(const FObjectInitializer& ObjectInitializer)
 	};
 	static FConstructorStatics ConstructorStatics;
 
-	bEnableLargeWorlds = false;
 	bEnableWorldBoundsChecks = true;
 PRAGMA_DISABLE_DEPRECATION_WARNINGS
 	bEnableNavigationSystem = true;
@@ -84,6 +83,7 @@ PRAGMA_ENABLE_DEPRECATION_WARNINGS
 	bEnableWorldComposition = false;
 	bEnableWorldOriginRebasing = false;
 #if WITH_EDITORONLY_DATA
+	bEnableLargeWorlds_DEPRECATED = (UE_USE_UE4_WORLD_MAX != 0);
 	bEnableHierarchicalLODSystem_DEPRECATED = true;
 	NumHLODLevels = 0;
 	bGenerateSingleClusterForLevel = false;
@@ -213,11 +213,6 @@ void AWorldSettings::PostRegisterAllComponents()
 	{
 		AudioDevice->SetDefaultAudioSettings(World, DefaultReverbSettings, DefaultAmbientZoneSettings);
 	}
-	
-	if(bEnableLargeWorlds)
-	{
-		UpdateEnableLargeWorldsCVars(bEnableLargeWorlds);
-	}	
 }
 
 UWorldPartition* AWorldSettings::GetWorldPartition() const
@@ -298,12 +293,6 @@ void AWorldSettings::NotifyBeginPlay()
 
 		World->bBegunPlay = true;
 	}
-
-	if(bEnableLargeWorlds)
-	{
-		// Done post BeginPlay to give code a chance to set bEnableLargeWorlds.
-		UpdateEnableLargeWorldsCVars(bEnableLargeWorlds);
-	}	
 }
 
 void AWorldSettings::NotifyMatchStarted()
@@ -313,13 +302,6 @@ void AWorldSettings::NotifyMatchStarted()
 	World->bMatchStarted = true;
 }
 
-void AWorldSettings::UpdateEnableLargeWorldsCVars(bool bEnable) const
-{
-	// LWC_TODO: Large world support. This will be removed once UE_LARGE_WORLD_MAX is stable.
-	IConsoleManager::Get().FindConsoleVariable(TEXT("r.UseVisibilityOctree"))->Set(!bEnable);
-	IConsoleManager::Get().FindConsoleVariable(TEXT("r.Shadow.UseOctreeForCulling"))->Set(!bEnable); 
-}
-
 void AWorldSettings::GetLifetimeReplicatedProps( TArray< FLifetimeProperty > & OutLifetimeProps ) const
 {
 	Super::GetLifetimeReplicatedProps( OutLifetimeProps );
@@ -803,10 +785,6 @@ void AWorldSettings::InternalPostPropertyChanged(FName PropertyName)
 	{
 		UpdateBookmarkClass();
 	}
-	else if(PropertyName == GET_MEMBER_NAME_CHECKED(AWorldSettings, bEnableLargeWorlds))
-	{
-		UpdateEnableLargeWorldsCVars(bEnableLargeWorlds);
-	}
 
 	if (GetWorld() != nullptr && GetWorld()->PersistentLevel && GetWorld()->PersistentLevel->GetWorldSettings() == this)
 	{
