Removing ULocalPlayer::PostInitProperties() to avoid creating view states before the UWorld is actually initialized. This avoids the view states from being set to the wrong FeatureLevel.

diff --git a/Engine/Source/Runtime/Engine/Classes/Engine/LocalPlayer.h b/Engine/Source/Runtime/Engine/Classes/Engine/LocalPlayer.h
--- a/Engine/Source/Runtime/Engine/Classes/Engine/LocalPlayer.h
+++ b/Engine/Source/Runtime/Engine/Classes/Engine/LocalPlayer.h
@@ -227,7 +227,6 @@ private:
 
 public:
 	// UObject interface
-	virtual void PostInitProperties() override;
 	virtual void FinishDestroy() override;
 	static void AddReferencedObjects(UObject* InThis, FReferenceCollector& Collector);
 	// End of UObject interface
diff --git a/Engine/Source/Runtime/Engine/Private/LocalPlayer.cpp b/Engine/Source/Runtime/Engine/Private/LocalPlayer.cpp
--- a/Engine/Source/Runtime/Engine/Private/LocalPlayer.cpp
+++ b/Engine/Source/Runtime/Engine/Private/LocalPlayer.cpp
@@ -239,29 +239,6 @@ ULocalPlayer::ULocalPlayer(const FObjectInitializer& ObjectInitializer)
 	PendingLevelPlayerControllerClass = APlayerController::StaticClass();
 }
 
-void ULocalPlayer::PostInitProperties()
-{
-	Super::PostInitProperties();
-	if ( !IsTemplate() )
-	{
-		int32 NumViews = 1;
-		if (GEngine->StereoRenderingDevice.IsValid())
-		{
-			NumViews = GEngine->StereoRenderingDevice->GetDesiredNumberOfViews(true);
-			check(NumViews > 0);			
-		}
-
-		const UWorld* CurrentWorld = GetWorld();
-		const ERHIFeatureLevel::Type FeatureLevel = CurrentWorld ? CurrentWorld->FeatureLevel.GetValue() : GMaxRHIFeatureLevel;
-
-		ViewStates.SetNum(NumViews);
-		for (auto& State : ViewStates)
-		{
-			State.Allocate(FeatureLevel);
-		}		
-	}
-}
-
 void ULocalPlayer::PlayerAdded(UGameViewportClient* InViewportClient, int32 InControllerID)
 {
 	ViewportClient = InViewportClient;
