Haptic Soundwave support for stereo sounds.
-Stereo sounds can be used to provide two-channel haptics, provided the platform supports that.

diff --git a/Engine/Source/Runtime/ApplicationCore/Public/GenericPlatform/IInputInterface.h b/Engine/Source/Runtime/ApplicationCore/Public/GenericPlatform/IInputInterface.h
--- a/Engine/Source/Runtime/ApplicationCore/Public/GenericPlatform/IInputInterface.h
+++ b/Engine/Source/Runtime/ApplicationCore/Public/GenericPlatform/IInputInterface.h
@@ -46,6 +46,8 @@ struct FHapticFeedbackBuffer
 	bool bFinishedPlaying;
 	int SamplingRate;
 	float ScaleFactor;
+	bool bUseStereo;
+	uint32 CurrentSampleIndex[2];
 
 	FHapticFeedbackBuffer()
 		: CurrentPtr(0)
@@ -53,6 +55,8 @@ struct FHapticFeedbackBuffer
 		, SamplesSent(0)
 		, bFinishedPlaying(false)
 		, SamplingRate(0)
+		, bUseStereo(false)
+		, CurrentSampleIndex{}
 	{
 	}
 
diff --git a/Engine/Source/Runtime/Engine/Classes/Haptics/HapticFeedbackEffect_SoundWave.h b/Engine/Source/Runtime/Engine/Classes/Haptics/HapticFeedbackEffect_SoundWave.h
--- a/Engine/Source/Runtime/Engine/Classes/Haptics/HapticFeedbackEffect_SoundWave.h
+++ b/Engine/Source/Runtime/Engine/Classes/Haptics/HapticFeedbackEffect_SoundWave.h
@@ -18,6 +18,10 @@ class UHapticFeedbackEffect_SoundWave : public UHapticFeedbackEffect_Base
 	UPROPERTY(EditAnywhere, Category = "HapticFeedbackEffect_SoundWave")
 	TObjectPtr<USoundWave> SoundWave;
 
+	/** If true on a vr controller the left and right stereo channels would be applied to the left and right controller, respectively. */
+	UPROPERTY(EditAnywhere, Category = "HapticFeedbackEffect_SoundWave")
+	bool bUseStereo;
+
 	~UHapticFeedbackEffect_SoundWave();
 
 	void Initialize(FHapticFeedbackBuffer& HapticBuffer) override;
@@ -28,6 +32,8 @@ class UHapticFeedbackEffect_SoundWave : public UHapticFeedbackEffect_Base
 
 private:
 	void PrepareSoundWaveBuffer();
+	void PrepareSoundWaveMonoBuffer(uint8* PCMData, int32 RawPCMDataSize);
+	void PrepareSoundWaveStereoBuffer(uint8* PCMData, int32 RawPCMDataSize);
 	bool bPrepared;
 
 	TArray<uint8> RawData;
diff --git a/Engine/Source/Runtime/Engine/Private/GameFramework/HapticFeedbackEffect.cpp b/Engine/Source/Runtime/Engine/Private/GameFramework/HapticFeedbackEffect.cpp
--- a/Engine/Source/Runtime/Engine/Private/GameFramework/HapticFeedbackEffect.cpp
+++ b/Engine/Source/Runtime/Engine/Private/GameFramework/HapticFeedbackEffect.cpp
@@ -8,7 +8,7 @@
 #include "Sound/SoundWave.h"
 #include "Haptics/HapticFeedbackEffect_SoundWave.h"
 #include "Haptics/HapticFeedbackEffect_Buffer.h"
-
+DEFINE_LOG_CATEGORY_STATIC(LogHaptics, Display, All);
 
 bool FActiveHapticFeedbackEffect::Update(const float DeltaTime, FHapticFeedbackValues& Values)
 {
@@ -115,6 +115,8 @@ void UHapticFeedbackEffect_Buffer::Initialize(FHapticFeedbackBuffer& HapticBuffe
 	HapticBuffer.SamplesSent = 0;
 	HapticBuffer.bFinishedPlaying = false;
 	HapticBuffer.RawData = nullptr;
+	HapticBuffer.CurrentSampleIndex[0] = 0;
+	HapticBuffer.CurrentSampleIndex[1] = 0;
 }
 
 void UHapticFeedbackEffect_Buffer::GetValues(const float EvalTime, FHapticFeedbackValues& Values)
@@ -156,6 +158,9 @@ void UHapticFeedbackEffect_SoundWave::Initialize(FHapticFeedbackBuffer& HapticBu
 	HapticBuffer.SamplesSent = 0;
 	HapticBuffer.bFinishedPlaying = false;
 	HapticBuffer.SamplingRate = SoundWave->GetSampleRateForCurrentPlatform();
+	HapticBuffer.bUseStereo = bUseStereo;
+	HapticBuffer.CurrentSampleIndex[0] = 0;
+	HapticBuffer.CurrentSampleIndex[1] = 0;
 }
 
 void UHapticFeedbackEffect_SoundWave::GetValues(const float EvalTime, FHapticFeedbackValues& Values)
@@ -179,16 +184,28 @@ void UHapticFeedbackEffect_SoundWave::PrepareSoundWaveBuffer()
 		return;
 	}
 	AD->Precache(SoundWave, true, false, true);
-	SoundWave->InitAudioResource(AD->GetRuntimeFormat(SoundWave));
+	// Remove call to InitAudioResource because AD->Precache calls the InitAudioResource code and calling it twice causes issues with Bulk data loading. 
 	uint8* PCMData = SoundWave->RawPCMData;
 	int32 RawPCMDataSize = SoundWave->RawPCMDataSize;
-	check((PCMData != nullptr) || (RawPCMDataSize == 0));	
+	check((PCMData != nullptr) || (RawPCMDataSize == 0));
+	if (bUseStereo)
+	{
+		PrepareSoundWaveStereoBuffer(PCMData, RawPCMDataSize);
+	}
+	else
+	{
+		PrepareSoundWaveMonoBuffer(PCMData, RawPCMDataSize);
+	}
+	bPrepared = true;
+}
 
+void UHapticFeedbackEffect_SoundWave::PrepareSoundWaveMonoBuffer(uint8* PCMData, int32 RawPCMDataSize)
+{
 	// Some platforms may need to resample the PCM data.  Such resampling should be performed at the platform specific plugin level
 	int32 NumChannels = SoundWave->NumChannels;
 	if (NumChannels > 1)
 	{
-		UE_LOG(LogTemp, Warning, TEXT("%s used for vibration has more than 1 channel. Only the first channel will be used."), *SoundWave->GetPathName());
+		UE_LOG(LogHaptics, Warning, TEXT("%s used for mono vibration has more than 1 channel. Only the first channel will be used."), *SoundWave->GetPathName());
 		check(RawPCMDataSize % (sizeof(int16) * NumChannels) == 0);
 		int32 NumSamples = RawPCMDataSize / sizeof(int16) / NumChannels;
 		RawData.AddUninitialized(NumSamples * sizeof(int16));
@@ -203,5 +220,30 @@ void UHapticFeedbackEffect_SoundWave::PrepareSoundWaveBuffer()
 	{
 		RawData.Append(PCMData, RawPCMDataSize);
 	}
-	bPrepared = true;
+}
+
+void UHapticFeedbackEffect_SoundWave::PrepareSoundWaveStereoBuffer(uint8* PCMData, int32 RawPCMDataSize)
+{
+	// Some platforms may need to resample the PCM data.  Such resampling should be performed at the platform specific plugin level
+	int32 NumChannels = SoundWave->NumChannels;
+	if (NumChannels < 2)
+	{
+		constexpr static int StereoChannels = 2;
+		UE_LOG(LogHaptics, Warning, TEXT("%s used for stereo vibration has only 1 channel. The first channel will be copied into the second channel."), *SoundWave->GetPathName());
+		check(RawPCMDataSize % (sizeof(int16)) == 0);
+		int32 NumSamples = RawPCMDataSize / sizeof(int16);
+		RawData.AddUninitialized(NumSamples * StereoChannels * sizeof(int16));
+		int16* SourceData = reinterpret_cast<int16*>(PCMData);
+		int16* DestData = reinterpret_cast<int16*>(RawData.GetData());
+
+		for (int32 i = 0; i < NumSamples; i++)
+		{
+			DestData[i * StereoChannels] = SourceData[i];
+			DestData[i * StereoChannels + 1] = SourceData[i];
+		}
+	}
+	else
+	{
+		RawData.Append(PCMData, RawPCMDataSize);
+	}
 }
