ExternalStreamingObject now manage outer & levelstreaming creation of their cells.
On cook the level streaming of external cells is not saved and their outer is changed to ExternalStreamingObject
On Load the externalstreamingobject re-outers the cells to the spatial hash and creates their level streaming

diff --git a/Engine/Source/Runtime/Engine/Private/WorldPartition/WorldPartitionLevelStreamingDynamic.cpp b/Engine/Source/Runtime/Engine/Private/WorldPartition/WorldPartitionLevelStreamingDynamic.cpp
--- a/Engine/Source/Runtime/Engine/Private/WorldPartition/WorldPartitionLevelStreamingDynamic.cpp
+++ b/Engine/Source/Runtime/Engine/Private/WorldPartition/WorldPartitionLevelStreamingDynamic.cpp
@@ -45,6 +45,33 @@ UWorldPartitionLevelStreamingDynamic::UWorldPartitionLevelStreamingDynamic(const
 #endif
 }
 
+/**
+ * Initializes a UWorldPartitionLevelStreamingDynamic.
+ */
+void UWorldPartitionLevelStreamingDynamic::Initialize(const UWorldPartitionRuntimeLevelStreamingCell& InCell)
+{
+	StreamingCell = &InCell;
+	UWorld* World = GetWorld();
+	check(!ShouldBeLoaded());
+	check((World->IsGameWorld() && !ShouldBeVisible()) || (!World->IsGameWorld() && !GetShouldBeVisibleFlag()));
+	check(!WorldAsset.IsNull());
+
+	bShouldBeAlwaysLoaded = InCell.IsAlwaysLoaded();
+	StreamingPriority = 0;
+
+#if WITH_EDITOR
+	check(ChildPackages.Num() == 0);
+
+	UnsavedActorsContainer = InCell.UnsavedActorsContainer;
+
+	IWorldPartitionRuntimeCellOwner* CellOwner = InCell.GetCellOwner();
+	Initialize(CellOwner->GetOuterWorld(), InCell.GetPackages());
+#else
+	IWorldPartitionRuntimeCellOwner* CellOwner = InCell.GetCellOwner();
+	OuterWorldPartition = CellOwner->GetOuterWorld()->GetWorldPartition();
+#endif
+}
+
 #if WITH_EDITOR
 
 UWorldPartitionLevelStreamingDynamic* UWorldPartitionLevelStreamingDynamic::LoadInEditor(UWorld* World, FName LevelStreamingName, const TArray<FWorldPartitionRuntimeCellObjectMapping>& InPackages)
@@ -94,26 +121,6 @@ void UWorldPartitionLevelStreamingDynamic::Initialize(UWorld* OuterWorld, const
 	OuterWorldPartition = OuterWorld->GetWorldPartition();
 }
 
-/**
- * Initializes a UWorldPartitionLevelStreamingDynamic.
- */
-void UWorldPartitionLevelStreamingDynamic::Initialize(const UWorldPartitionRuntimeLevelStreamingCell& InCell)
-{
-	StreamingCell = &InCell;
-	UWorld* World = GetWorld();
-	check(!ShouldBeLoaded());
-	check((World->IsGameWorld() && !ShouldBeVisible()) || (!World->IsGameWorld() && !GetShouldBeVisibleFlag()));
-	check(ChildPackages.Num() == 0);
-	check(!WorldAsset.IsNull());
-
-	bShouldBeAlwaysLoaded = InCell.IsAlwaysLoaded();
-	StreamingPriority = 0;
-	UnsavedActorsContainer = InCell.UnsavedActorsContainer;
-
-	IWorldPartitionRuntimeCellOwner* CellOwner = InCell.GetCellOwner();
-	Initialize(CellOwner->GetOuterWorld(), InCell.GetPackages());
-}
-
 /**
  Custom destroy (delegate removal)
  */
diff --git a/Engine/Source/Runtime/Engine/Private/WorldPartition/WorldPartitionRuntimeLevelStreamingCell.cpp b/Engine/Source/Runtime/Engine/Private/WorldPartition/WorldPartitionRuntimeLevelStreamingCell.cpp
--- a/Engine/Source/Runtime/Engine/Private/WorldPartition/WorldPartitionRuntimeLevelStreamingCell.cpp
+++ b/Engine/Source/Runtime/Engine/Private/WorldPartition/WorldPartitionRuntimeLevelStreamingCell.cpp
@@ -8,12 +8,12 @@
 #include "WorldPartition/HLOD/HLODSubsystem.h"
 #include "WorldPartition/HLOD/HLODActor.h"
 #include "WorldPartition/WorldPartitionDebugHelper.h"
+#include "WorldPartition/WorldPartitionLevelStreamingPolicy.h"
 
 #include UE_INLINE_GENERATED_CPP_BY_NAME(WorldPartitionRuntimeLevelStreamingCell)
 
 #if WITH_EDITOR
 #include "WorldPartition/ActorDescContainer.h"
-#include "WorldPartition/WorldPartitionLevelStreamingPolicy.h"
 #include "WorldPartition/WorldPartitionLevelHelper.h"
 #include "Engine/LevelStreamingAlwaysLoaded.h"
 #include "Algo/ForEach.h"
@@ -49,6 +49,68 @@
 	return LevelStreaming;
 }
 
+bool UWorldPartitionRuntimeLevelStreamingCell::HasActors() const
+{
+#if WITH_EDITOR
+	return GetActorCount() > 0;
+#else
+	return true;
+#endif
+}
+
+void UWorldPartitionRuntimeLevelStreamingCell::CreateAndSetLevelStreaming(const FString& InPackageName)
+{
+	LevelStreaming = CreateLevelStreaming(InPackageName);
+}
+
+UWorldPartitionLevelStreamingDynamic* UWorldPartitionRuntimeLevelStreamingCell::CreateLevelStreaming(const FString& InPackageName) const
+{
+	if (HasActors())
+	{
+		const UWorldPartition* WorldPartition = GetCellOwner()->GetWorldPartition();
+		UWorld* OuterWorld = GetCellOwner()->GetOuterWorld();
+		UWorld* OwningWorld = GetCellOwner()->GetOwningWorld();
+
+		const FName LevelStreamingName = FName(*FString::Printf(TEXT("WorldPartitionLevelStreaming_%s"), *GetName()));
+
+		// When called by Commandlet (PopulateGeneratedPackageForCook), LevelStreaming's outer is set to Cell/WorldPartition's outer to prevent warnings when saving Cell Levels (Warning: Obj in another map). 
+		// At runtime, LevelStreaming's outer will be properly set to the main world (see UWorldPartitionRuntimeLevelStreamingCell::Activate).
+		UWorld* LevelStreamingOuterWorld = IsRunningCommandlet() ? OuterWorld : OwningWorld;
+		UWorldPartitionLevelStreamingDynamic* NewLevelStreaming = NewObject<UWorldPartitionLevelStreamingDynamic>(LevelStreamingOuterWorld, UWorldPartitionLevelStreamingDynamic::StaticClass(), LevelStreamingName, RF_NoFlags, NULL);
+
+		FName WorldName = OuterWorld->GetFName();
+#if WITH_EDITOR
+		// In PIE make sure that we are using the proper original world name so that actors resolve their outer property
+		if (OwningWorld->IsPlayInEditor() && !OuterWorld->OriginalWorldName.IsNone())
+		{
+			WorldName = OuterWorld->OriginalWorldName;
+		}
+
+		FString PackageName = !InPackageName.IsEmpty() ? InPackageName : UWorldPartitionLevelStreamingPolicy::GetCellPackagePath(GetFName(), OuterWorld);
+#else
+		check(!InPackageName.IsEmpty());
+		FString PackageName = InPackageName;
+#endif
+
+		TSoftObjectPtr<UWorld> WorldAsset(FSoftObjectPath(FString::Printf(TEXT("%s.%s"), *PackageName, *WorldName.ToString())));
+		NewLevelStreaming->SetWorldAsset(WorldAsset);
+		// Transfer WorldPartition's transform to Level
+		NewLevelStreaming->LevelTransform = WorldPartition->GetInstanceTransform();
+		NewLevelStreaming->bClientOnlyVisible = GetClientOnlyVisible();
+		NewLevelStreaming->Initialize(*this);
+
+		if (OwningWorld->IsPlayInEditor() && OwningWorld->GetPackage()->HasAnyPackageFlags(PKG_PlayInEditor) && OwningWorld->GetPackage()->GetPIEInstanceID() != INDEX_NONE)
+		{
+			// When renaming for PIE, make sure to keep World's name so that linker can properly remap with Package's instancing context
+			NewLevelStreaming->RenameForPIE(OwningWorld->GetPackage()->GetPIEInstanceID(), /*bKeepWorldAssetName*/true);
+		}
+
+		return NewLevelStreaming;
+	}
+
+	return nullptr;
+}
+
 EStreamingStatus UWorldPartitionRuntimeLevelStreamingCell::GetStreamingStatus() const
 {
 	if (LevelStreaming)
@@ -109,48 +171,6 @@
 PRAGMA_ENABLE_DEPRECATION_WARNINGS
 }
 
-UWorldPartitionLevelStreamingDynamic* UWorldPartitionRuntimeLevelStreamingCell::CreateLevelStreaming(const FString& InPackageName) const
-{
-	if (GetActorCount() > 0)
-	{
-		const UWorldPartition* WorldPartition = GetCellOwner()->GetWorldPartition();
-		UWorld* OuterWorld = GetCellOwner()->GetOuterWorld();
-		UWorld* OwningWorld = GetCellOwner()->GetOwningWorld();
-		
-		const FName LevelStreamingName = FName(*FString::Printf(TEXT("WorldPartitionLevelStreaming_%s"), *GetName()));
-		
-		// When called by Commandlet (PopulateGeneratedPackageForCook), LevelStreaming's outer is set to Cell/WorldPartition's outer to prevent warnings when saving Cell Levels (Warning: Obj in another map). 
-		// At runtime, LevelStreaming's outer will be properly set to the main world (see UWorldPartitionRuntimeLevelStreamingCell::Activate).
-		UWorld* LevelStreamingOuterWorld = IsRunningCommandlet() ? OuterWorld : OwningWorld;
-		UWorldPartitionLevelStreamingDynamic* NewLevelStreaming = NewObject<UWorldPartitionLevelStreamingDynamic>(LevelStreamingOuterWorld, UWorldPartitionLevelStreamingDynamic::StaticClass(), LevelStreamingName, RF_NoFlags, NULL);
-		FString PackageName = !InPackageName.IsEmpty() ? InPackageName : UWorldPartitionLevelStreamingPolicy::GetCellPackagePath(GetFName(), OuterWorld);
-				
-		FName WorldName = OuterWorld->GetFName();
-		// In PIE make sure that we are using the proper original world name so that actors resolve their outer property
-		if (OwningWorld->IsPlayInEditor() && !OuterWorld->OriginalWorldName.IsNone())
-		{
-			WorldName = OuterWorld->OriginalWorldName;
-		}
-
-		TSoftObjectPtr<UWorld> WorldAsset(FSoftObjectPath(FString::Printf(TEXT("%s.%s"), *PackageName, *WorldName.ToString())));
-		NewLevelStreaming->SetWorldAsset(WorldAsset);
-		// Transfer WorldPartition's transform to Level
-		NewLevelStreaming->LevelTransform = WorldPartition->GetInstanceTransform();
-		NewLevelStreaming->bClientOnlyVisible = GetClientOnlyVisible();
-		NewLevelStreaming->Initialize(*this);
-
-		if (OwningWorld->IsPlayInEditor() && OwningWorld->GetPackage()->HasAnyPackageFlags(PKG_PlayInEditor) && OwningWorld->GetPackage()->GetPIEInstanceID() != INDEX_NONE)
-		{
-			// When renaming for PIE, make sure to keep World's name so that linker can properly remap with Package's instancing context
-			NewLevelStreaming->RenameForPIE(OwningWorld->GetPackage()->GetPIEInstanceID(), /*bKeepWorldAssetName*/true);
-		}
-
-		return NewLevelStreaming;
-	}
-
-	return nullptr;
-}
-
 bool UWorldPartitionRuntimeLevelStreamingCell::PopulateGeneratorPackageForCook(TArray<UPackage*>& OutModifiedPackages)
 {
 	check(IsAlwaysLoaded());
@@ -227,6 +247,12 @@
 
 		// Remap Level's SoftObjectPaths
 		FWorldPartitionLevelHelper::RemapLevelSoftObjectPaths(NewLevel, WorldPartition);
+
+		if (GetCellOwner()->ShouldSerializeCellLevelStreaming())
+		{
+			LevelStreaming = nullptr;
+			OutModifiedPackage.Add(GetOuter()->GetPackage());
+		}
 	}
 	return true;
 }
diff --git a/Engine/Source/Runtime/Engine/Private/WorldPartition/WorldPartitionRuntimeSpatialHash.cpp b/Engine/Source/Runtime/Engine/Private/WorldPartition/WorldPartitionRuntimeSpatialHash.cpp
--- a/Engine/Source/Runtime/Engine/Private/WorldPartition/WorldPartitionRuntimeSpatialHash.cpp
+++ b/Engine/Source/Runtime/Engine/Private/WorldPartition/WorldPartitionRuntimeSpatialHash.cpp
@@ -7,6 +7,7 @@
 
 #include "WorldPartition/WorldPartitionStreamingGenerationContext.h"
 #include "WorldPartition/WorldPartitionRuntimeSpatialHashCell.h"
+#include "WorldPartition/WorldPartitionRuntimeLevelStreamingCell.h"
 #include "WorldPartition/WorldPartitionActorDesc.h"
 #include "WorldPartition/WorldPartitionActorDescView.h"
 #include "WorldPartition/WorldPartitionStreamingPolicy.h"
@@ -754,6 +755,55 @@ ASpatialHashRuntimeGridInfo::ASpatialHashRuntimeGridInfo(const FObjectInitialize
 #endif
 }
 
+void URuntimeSpatialHashExternalStreamingObject::OnStreamingObjectLoaded()
+{
+	// Cooked streaming object's Cells do not have LevelStreaming and are outered to the streaming object.
+	// Create their level streaming now and re-outered to the runtime hash.
+	bool bIsACookedObject = !CellToLevelStreamingPackage.IsEmpty();
+	if (bIsACookedObject)
+	{
+		for (FSpatialHashStreamingGrid& StreamingGrid : StreamingGrids)
+		{
+			for (FSpatialHashStreamingGridLevel& GridLevel : StreamingGrid.GridLevels)
+			{
+				for (FSpatialHashStreamingGridLayerCell& GridLayerCell : GridLevel.LayerCells)
+				{
+					for (UWorldPartitionRuntimeSpatialHashCell* Cell : GridLayerCell.GridCells)
+					{
+						UWorldPartitionRuntimeLevelStreamingCell* RuntimeCell = CastChecked<UWorldPartitionRuntimeLevelStreamingCell>(Cell);
+
+						FName LevelStreamingPackage = CellToLevelStreamingPackage.FindChecked(RuntimeCell->GetFName());
+						RuntimeCell->Rename(nullptr, GetOuterWorld()->GetWorldPartition()->RuntimeHash);
+						RuntimeCell->CreateAndSetLevelStreaming(*LevelStreamingPackage.ToString());
+					}
+				}
+			}
+		}
+
+	}
+	
+}
+
+#if WITH_EDITOR
+void URuntimeSpatialHashExternalStreamingObject::PopulateGeneratorPackageForCook()
+{
+	for (FSpatialHashStreamingGrid& StreamingGrid : StreamingGrids)
+	{
+		for (FSpatialHashStreamingGridLevel& GridLevel : StreamingGrid.GridLevels)
+		{
+			for (FSpatialHashStreamingGridLayerCell& GridLayerCell : GridLevel.LayerCells)
+			{
+				for (UWorldPartitionRuntimeSpatialHashCell* Cell : GridLayerCell.GridCells)
+				{
+					UWorldPartitionLevelStreamingDynamic* LevelStreamingDynamic = CastChecked<UWorldPartitionRuntimeLevelStreamingCell>(Cell)->GetLevelStreaming();
+					CellToLevelStreamingPackage.Add(Cell->GetFName(), LevelStreamingDynamic->PackageNameToLoad);
+				}
+			}
+		}
+	}
+}
+#endif
+
 UWorldPartitionRuntimeSpatialHash::UWorldPartitionRuntimeSpatialHash(const FObjectInitializer& ObjectInitializer)
 	: Super(ObjectInitializer)
 #if WITH_EDITORONLY_DATA
@@ -812,6 +862,7 @@ URuntimeHashExternalStreamingObjectBase* UWorldPartitionRuntimeSpatialHash::Stor
 				for (UWorldPartitionRuntimeSpatialHashCell* Cell : GridLayerCell.GridCells)
 				{
 					Cell->SetGridName(StreamingGrid.GridName);
+					Cell->Rename(nullptr, StreamingObject);
 				}
 			}
 		}
diff --git a/Engine/Source/Runtime/Engine/Public/WorldPartition/WorldPartitionRuntimeCellOwner.h b/Engine/Source/Runtime/Engine/Public/WorldPartition/WorldPartitionRuntimeCellOwner.h
--- a/Engine/Source/Runtime/Engine/Public/WorldPartition/WorldPartitionRuntimeCellOwner.h
+++ b/Engine/Source/Runtime/Engine/Public/WorldPartition/WorldPartitionRuntimeCellOwner.h
@@ -26,5 +26,7 @@ public:
 	virtual UWorld* GetOwningWorld() const = 0;
 	virtual UWorld* GetOuterWorld() const = 0;
 
+	virtual bool ShouldSerializeCellLevelStreaming() const { return true; }
+
 	UWorldPartition* GetWorldPartition() const { return GetOuterWorld()->GetWorldPartition(); }
 };
diff --git a/Engine/Source/Runtime/Engine/Public/WorldPartition/WorldPartitionRuntimeHash.h b/Engine/Source/Runtime/Engine/Public/WorldPartition/WorldPartitionRuntimeHash.h
--- a/Engine/Source/Runtime/Engine/Public/WorldPartition/WorldPartitionRuntimeHash.h
+++ b/Engine/Source/Runtime/Engine/Public/WorldPartition/WorldPartitionRuntimeHash.h
@@ -45,7 +45,18 @@ public:
 	//~ Begin IWorldPartitionRuntimeCellOwner Interface
 	virtual UWorld* GetOwningWorld() const override { return OwningWorld.Get(); }
 	virtual UWorld* GetOuterWorld() const override { return OuterWorld.Get(); }
+
+	// If the cell is saved within an external streaming object avoid do not serialize its level streaming. 
+	// Level streaming are outered to the world and would not be saved within the ExternalStreamingObject.
+	// Instead they will be created once the external streaming object is loaded. 
+	virtual bool ShouldSerializeCellLevelStreaming() const { return false; }
 	//~ End IWorldPartitionRuntimeCellOwner Interface
+	
+	virtual void OnStreamingObjectLoaded() {};
+
+#if WITH_EDITOR
+	virtual void PopulateGeneratorPackageForCook() {};
+#endif
 
 private:
 	UPROPERTY();
diff --git a/Engine/Source/Runtime/Engine/Public/WorldPartition/WorldPartitionRuntimeLevelStreamingCell.h b/Engine/Source/Runtime/Engine/Public/WorldPartition/WorldPartitionRuntimeLevelStreamingCell.h
--- a/Engine/Source/Runtime/Engine/Public/WorldPartition/WorldPartitionRuntimeLevelStreamingCell.h
+++ b/Engine/Source/Runtime/Engine/Public/WorldPartition/WorldPartitionRuntimeLevelStreamingCell.h
@@ -31,6 +31,11 @@ class UWorldPartitionRuntimeLevelStreamingCell : public UWorldPartitionRuntimeSp
 
 	virtual void SetStreamingPriority(int32 InStreamingPriority) const override;
 	class UWorldPartitionLevelStreamingDynamic* GetLevelStreaming() const;
+
+	bool HasActors() const;
+
+	void CreateAndSetLevelStreaming(const FString& InPackageName);
+	class UWorldPartitionLevelStreamingDynamic* CreateLevelStreaming(const FString& InPackageName = FString()) const;
 	
 
 #if WITH_EDITOR
@@ -57,10 +62,6 @@ private:
 
 	class UWorldPartitionLevelStreamingDynamic* GetOrCreateLevelStreaming() const;
 
-#if WITH_EDITOR
-	class UWorldPartitionLevelStreamingDynamic* CreateLevelStreaming(const FString& InPackageName = FString()) const;
-#endif
-
 #if WITH_EDITORONLY_DATA
 	UPROPERTY()
 	TArray<FWorldPartitionRuntimeCellObjectMapping> Packages;
diff --git a/Engine/Source/Runtime/Engine/Public/WorldPartition/WorldPartitionRuntimeSpatialHash.h b/Engine/Source/Runtime/Engine/Public/WorldPartition/WorldPartitionRuntimeSpatialHash.h
--- a/Engine/Source/Runtime/Engine/Public/WorldPartition/WorldPartitionRuntimeSpatialHash.h
+++ b/Engine/Source/Runtime/Engine/Public/WorldPartition/WorldPartitionRuntimeSpatialHash.h
@@ -224,8 +224,18 @@ class URuntimeSpatialHashExternalStreamingObject : public URuntimeHashExternalSt
 	GENERATED_BODY()
 
 public:
+	virtual void OnStreamingObjectLoaded() override;
+
+#if WITH_EDITOR
+	virtual void PopulateGeneratorPackageForCook() override;
+#endif
+
 	UPROPERTY();
 	TArray<FSpatialHashStreamingGrid> StreamingGrids;
+
+private:
+	UPROPERTY();
+	TMap<FName, FName> CellToLevelStreamingPackage;
 };
 
 UCLASS()
