Add cooked paths to ContentBundlePaths

diff --git a/Engine/Source/Runtime/Engine/Private/WorldPartition/ContentBundle/ContentBundle.cpp b/Engine/Source/Runtime/Engine/Private/WorldPartition/ContentBundle/ContentBundle.cpp
--- a/Engine/Source/Runtime/Engine/Private/WorldPartition/ContentBundle/ContentBundle.cpp
+++ b/Engine/Source/Runtime/Engine/Private/WorldPartition/ContentBundle/ContentBundle.cpp
@@ -2,6 +2,7 @@
 #include "WorldPartition/ContentBundle/ContentBundle.h"
 
 #include "WorldPartition/WorldPartition.h"
+#include "WorldPartition/ContentBundle/ContentBundlePaths.h"
 #include "WorldPartition/ContentBundle/ContentBundleDescriptor.h"
 #include "WorldPartition/WorldPartitionRuntimeHash.h"
 #include "WorldPartition/ContentBundle/ContentBundleLog.h"
@@ -27,11 +28,10 @@ void FContentBundle::DoInitialize()
 #if WITH_EDITOR
 	InitializeForPIE();
 #else
-	FString StreamingObjectPathName = FPaths::RemoveDuplicateSlashes(GetDescriptor()->GetPackageRoot() + TEXT("/_GENERATED_/") + TEXT("StreamingObject"));
-	ExternalStreamingObjectPackage = LoadPackage(nullptr, *StreamingObjectPathName, LOAD_None);
+	ExternalStreamingObjectPackage = LoadPackage(nullptr, *GetExternalStreamingObjectPackageName(), LOAD_None);
 	if (ExternalStreamingObjectPackage != nullptr)
 	{
-		if (UObject* Object = StaticFindObjectFast(URuntimeHashExternalStreamingObjectBase::StaticClass(), ExternalStreamingObjectPackage, TEXT("StreamingObject")))
+		if (UObject* Object = StaticFindObjectFast(URuntimeHashExternalStreamingObjectBase::StaticClass(), ExternalStreamingObjectPackage, *GetExternalStreamingObjectName()))
 		{
 			ExternalStreamingObject = CastChecked<URuntimeHashExternalStreamingObjectBase>(Object);
 		}
diff --git a/Engine/Source/Runtime/Engine/Private/WorldPartition/ContentBundle/ContentBundleBase.cpp b/Engine/Source/Runtime/Engine/Private/WorldPartition/ContentBundle/ContentBundleBase.cpp
--- a/Engine/Source/Runtime/Engine/Private/WorldPartition/ContentBundle/ContentBundleBase.cpp
+++ b/Engine/Source/Runtime/Engine/Private/WorldPartition/ContentBundle/ContentBundleBase.cpp
@@ -5,6 +5,7 @@
 #include "WorldPartition/ContentBundle/ContentBundleDescriptor.h"
 #include "WorldPartition/ContentBundle/ContentBundleClient.h"
 #include "WorldPartition/ContentBundle/ContentBundleLog.h"
+#include "WorldPartition/ContentBundle/ContentBundlePaths.h"
 #include "Engine/World.h"
 
 FContentBundleBase::FContentBundleBase(TSharedPtr<FContentBundleClient>& InClient, UWorld* InWorld)
@@ -75,6 +76,16 @@ const UContentBundleDescriptor* FContentBundleBase::GetDescriptor() const
 	return Descriptor;
 }
 
+FString FContentBundleBase::GetExternalStreamingObjectPackageName() const
+{
+	return SlugStringForValidName(ContentBundlePaths::GetCookedContentBundleLevelFolder(*this) + ContentBundlePaths::GetGeneratedFolder() + TEXT("StreamingObject"));
+}
+
+FString FContentBundleBase::GetExternalStreamingObjectName() const
+{
+	return SlugStringForValidName(GetDisplayName() + TEXT("_") + GetDescriptor()->GetGuid().ToString() + TEXT("_ExternalStreamingObject"));
+}
+
 UWorld* FContentBundleBase::GetInjectedWorld() const
 {
 	return InjectedWorld.Get();
diff --git a/Engine/Source/Runtime/Engine/Private/WorldPartition/ContentBundle/ContentBundleEditor.cpp b/Engine/Source/Runtime/Engine/Private/WorldPartition/ContentBundle/ContentBundleEditor.cpp
--- a/Engine/Source/Runtime/Engine/Private/WorldPartition/ContentBundle/ContentBundleEditor.cpp
+++ b/Engine/Source/Runtime/Engine/Private/WorldPartition/ContentBundle/ContentBundleEditor.cpp
@@ -309,8 +309,7 @@ void FContentBundleEditor::GenerateStreaming()
 		return true;
 	});
 
-	FString ExternalStreamingObjectName = ObjectTools::SanitizeInvalidChars(GetDisplayName() + TEXT("_ExternalStreamingObject"), INVALID_LONGPACKAGE_CHARACTERS) ;
-	ExternalStreamingObject = WorldPartition->RuntimeHash->StoreToExternalStreamingObject(GetInjectedWorld()->ContentBundleManager, *ExternalStreamingObjectName);
+	ExternalStreamingObject = WorldPartition->RuntimeHash->StoreToExternalStreamingObject(GetInjectedWorld()->ContentBundleManager, *GetExternalStreamingObjectName());
 
 	UE_LOG(LogContentBundle, Log, TEXT("[CB: %s] Generated streaming cells. %u cells were generated"), *GetDescriptor()->GetDisplayName(), ContentBundleCells.Num());
 
diff --git a/Engine/Source/Runtime/Engine/Private/WorldPartition/ContentBundle/ContentBundlePaths.cpp b/Engine/Source/Runtime/Engine/Private/WorldPartition/ContentBundle/ContentBundlePaths.cpp
--- a/Engine/Source/Runtime/Engine/Private/WorldPartition/ContentBundle/ContentBundlePaths.cpp
+++ b/Engine/Source/Runtime/Engine/Private/WorldPartition/ContentBundle/ContentBundlePaths.cpp
@@ -1,11 +1,59 @@
 // Copyright Epic Games, Inc. All Rights Reserved.
 #include "WorldPartition/ContentBundle/ContentBundlePaths.h"
 
+#include "Engine/Level.h"
+#include "Engine/World.h"
+#include "WorldPartition/ContentBundle/ContentBundleBase.h"
+#include "WorldPartition/ContentBundle/ContentBundleDescriptor.h"
+
 #if WITH_EDITOR
 #include "String/Find.h"
-#include "Engine/Level.h"
 #endif
 
+namespace ContentBundlePaths
+{
+	// ContentBundleCookedFolder is reduce to avoid reaching MAX_PATH
+	constexpr FStringView ContentBundleCookedFolder = TEXTVIEW("/CB/");
+}
+
+FString ContentBundlePaths::GetCookedContentBundleLevelFolder(const FContentBundleBase& ContentBundle)
+{
+	TStringBuilderWithBuffer<TCHAR, NAME_SIZE> GeneratedContentBundleLevelFolder;
+	GeneratedContentBundleLevelFolder += TEXT("/");
+	GeneratedContentBundleLevelFolder += ContentBundle.GetDescriptor()->GetPackageRoot();
+	GeneratedContentBundleLevelFolder += TEXT("/");
+	GeneratedContentBundleLevelFolder += ContentBundlePaths::ContentBundleCookedFolder;
+	GeneratedContentBundleLevelFolder += TEXT("/");
+	
+	// ContentBundleGuid is reduced to avoid reaching MAX_PATH
+	ContentBundle.GetDescriptor()->GetGuid().AppendString(GeneratedContentBundleLevelFolder, EGuidFormats::Short); 
+
+	GeneratedContentBundleLevelFolder += TEXT("/");
+	GeneratedContentBundleLevelFolder += GetRelativeLevelFolder(ContentBundle);
+	GeneratedContentBundleLevelFolder += TEXT("/");
+
+	FString Result = *GeneratedContentBundleLevelFolder;
+	FPaths::RemoveDuplicateSlashes(Result);
+	return Result;
+}
+
+
+FString ContentBundlePaths::GetRelativeLevelFolder(const FContentBundleBase& ContentBundle)
+{
+	const UWorld* InjectedWorld = ContentBundle.GetInjectedWorld();
+	FString WorldPackageName = InjectedWorld->GetPackage()->GetName();
+
+	FString PackageRoot, PackagePath, PackageName;
+	ensure(FPackageName::SplitLongPackageName(WorldPackageName, PackageRoot, PackagePath, PackageName));
+
+	TStringBuilderWithBuffer<TCHAR, NAME_SIZE> RelativeLevelFolder;
+	RelativeLevelFolder += PackagePath;
+	RelativeLevelFolder += PackageName;
+	RelativeLevelFolder += TEXT("/");
+
+	return *RelativeLevelFolder;
+}
+
 #if WITH_EDITOR
 
 FString ContentBundlePaths::MakeExternalActorPackagePath(const FString& ContentBundleExternalActorFolder, const FString& ActorName)
diff --git a/Engine/Source/Runtime/Engine/Public/WorldPartition/ContentBundle/ContentBundleBase.h b/Engine/Source/Runtime/Engine/Public/WorldPartition/ContentBundle/ContentBundleBase.h
--- a/Engine/Source/Runtime/Engine/Public/WorldPartition/ContentBundle/ContentBundleBase.h
+++ b/Engine/Source/Runtime/Engine/Public/WorldPartition/ContentBundle/ContentBundleBase.h
@@ -29,6 +29,8 @@ public:
 	const TWeakPtr<FContentBundleClient>& GetClient() const;
 	const UContentBundleDescriptor* GetDescriptor() const;
 	EContentBundleStatus GetStatus() const { return Status; }
+	FString GetExternalStreamingObjectPackageName() const;
+	FString GetExternalStreamingObjectName() const;
 
 protected:
 	void SetStatus(EContentBundleStatus NewStatus);
diff --git a/Engine/Source/Runtime/Engine/Public/WorldPartition/ContentBundle/ContentBundlePaths.h b/Engine/Source/Runtime/Engine/Public/WorldPartition/ContentBundle/ContentBundlePaths.h
--- a/Engine/Source/Runtime/Engine/Public/WorldPartition/ContentBundle/ContentBundlePaths.h
+++ b/Engine/Source/Runtime/Engine/Public/WorldPartition/ContentBundle/ContentBundlePaths.h
@@ -4,15 +4,15 @@
 #include "CoreMinimal.h"
 #include "Containers/StringView.h"
 
+class FContentBundleBase;
+
 namespace ContentBundlePaths
 {
 	constexpr FStringView ContentBundleFolder = TEXTVIEW("/ContentBundle/");
+	constexpr FStringView GeneratedFolder = TEXTVIEW("/_Generated_/");
 
 	// Returns "/ContentBundle/"
-	constexpr FStringView GetContentBundleFolder()
-	{
-		return ContentBundleFolder;
-	}
+	constexpr FStringView GetContentBundleFolder() { return ContentBundleFolder; }
 	
 	// Returns "ContentBundle"
 	constexpr FStringView GetContentBundleFolderName()
@@ -21,6 +21,23 @@ namespace ContentBundlePaths
 		return FStringView(&GetContentBundleFolder()[1], GetContentBundleFolder().Len() - 2);
 	}
 
+	// Returns "/ContentBundle/"
+	constexpr FStringView GetGeneratedFolder() { return GeneratedFolder; }
+	
+	// Returns "_Generated_"
+	constexpr FStringView GetGeneratedFolderName()
+	{
+		// Start at index 1 to remove the first "/", Lenght - 2 to remove the trailing "/"
+		return FStringView(&GetGeneratedFolder()[1], GetGeneratedFolder().Len() - 2);
+	}
+
+	// Return value uses a reduced amount of character to avoid reaching character limit when cooking 
+	// return format is /{MountPoint}/CB/{ContentBundle_ShortUID}/{LevelPath}/{GeneratedFolder}/
+	FString GetCookedContentBundleLevelFolder(const FContentBundleBase& ContentBundle);
+	
+	// return format is {LevelPath}
+	FString GetRelativeLevelFolder(const FContentBundleBase& ContentBundle);
+
 #if WITH_EDITOR
 	// return an ExternalActor path following the format : /{MountPoint}/{ExternalActorFolder}/{ContentBundleFolder}/{ContentBundleUID}/{LevelPath}/{ExternalActorPackagePath}
 	ENGINE_API FString MakeExternalActorPackagePath(const FString& ContentBundleExternalActorFolder, const FString& ActorName);
