Disabling the player controller used by the demo net driver as a streaming source for World Partition

diff --git a/Engine/Source/Runtime/Engine/Private/ReplayHelper.cpp b/Engine/Source/Runtime/Engine/Private/ReplayHelper.cpp
--- a/Engine/Source/Runtime/Engine/Private/ReplayHelper.cpp
+++ b/Engine/Source/Runtime/Engine/Private/ReplayHelper.cpp
@@ -348,6 +348,7 @@ APlayerController* FReplayHelper::CreateSpectatorController(UNetConnection* Conn
 
 	// Streaming volumes logic must not be affected by replay spectator camera
 	NewDemoController->bIsUsingStreamingVolumes = false;
+	NewDemoController->bEnableStreamingSource = false;
 
 	if (Connection->GetDriver())
 	{
