Fix compile on development build

diff --git a/Engine/Source/Runtime/Engine/Private/WorldPartition/ContentBundle/ContentBundle.cpp b/Engine/Source/Runtime/Engine/Private/WorldPartition/ContentBundle/ContentBundle.cpp
--- a/Engine/Source/Runtime/Engine/Private/WorldPartition/ContentBundle/ContentBundle.cpp
+++ b/Engine/Source/Runtime/Engine/Private/WorldPartition/ContentBundle/ContentBundle.cpp
@@ -37,7 +37,7 @@ void FContentBundle::DoInitialize()
 		}
 		else
 		{
-			UE_LOG(LogContentBundle, Error, TEXT("[CB: %s] No streaming object found in package %s."), *GetDescriptor()->GetDisplayName(), *StreamingObjectPathName);
+			UE_LOG(LogContentBundle, Error, TEXT("[CB: %s] No streaming object found in package %s."), *GetDescriptor()->GetDisplayName(), *GetExternalStreamingObjectPackageName());
 		}
 	}
 	else
