Fix missing cells in cooked build

diff --git a/Engine/Source/Runtime/Engine/Private/WorldPartition/ContentBundle/ContentBundleBase.cpp b/Engine/Source/Runtime/Engine/Private/WorldPartition/ContentBundle/ContentBundleBase.cpp
--- a/Engine/Source/Runtime/Engine/Private/WorldPartition/ContentBundle/ContentBundleBase.cpp
+++ b/Engine/Source/Runtime/Engine/Private/WorldPartition/ContentBundle/ContentBundleBase.cpp
@@ -78,7 +78,7 @@ const UContentBundleDescriptor* FContentBundleBase::GetDescriptor() const
 
 FString FContentBundleBase::GetExternalStreamingObjectPackageName() const
 {
-	return SlugStringForValidName(ContentBundlePaths::GetCookedContentBundleLevelFolder(*this) + ContentBundlePaths::GetGeneratedFolder() + TEXT("StreamingObject"));
+	return ContentBundlePaths::GetCookedContentBundleLevelFolder(*this) + ContentBundlePaths::GetGeneratedFolderName() + TEXT("/StreamingObject");
 }
 
 FString FContentBundleBase::GetExternalStreamingObjectName() const
diff --git a/Engine/Source/Runtime/Engine/Private/WorldPartition/WorldPartitionRuntimeLevelStreamingCell.cpp b/Engine/Source/Runtime/Engine/Private/WorldPartition/WorldPartitionRuntimeLevelStreamingCell.cpp
--- a/Engine/Source/Runtime/Engine/Private/WorldPartition/WorldPartitionRuntimeLevelStreamingCell.cpp
+++ b/Engine/Source/Runtime/Engine/Private/WorldPartition/WorldPartitionRuntimeLevelStreamingCell.cpp
@@ -248,7 +248,7 @@ bool UWorldPartitionRuntimeLevelStreamingCell::PopulateGeneratedPackageForCook(U
 		// Remap Level's SoftObjectPaths
 		FWorldPartitionLevelHelper::RemapLevelSoftObjectPaths(NewLevel, WorldPartition);
 
-		if (GetCellOwner()->ShouldSerializeCellLevelStreaming())
+		if (!GetCellOwner()->ShouldSerializeCellLevelStreaming())
 		{
 			LevelStreaming = nullptr;
 			OutModifiedPackage.Add(GetOuter()->GetPackage());
