Submitted missing file from 22275012

diff --git a/Engine/Source/Runtime/Engine/Private/WorldPartition/HLOD/HLODSubsystem.cpp b/Engine/Source/Runtime/Engine/Private/WorldPartition/HLOD/HLODSubsystem.cpp
--- a/Engine/Source/Runtime/Engine/Private/WorldPartition/HLOD/HLODSubsystem.cpp
+++ b/Engine/Source/Runtime/Engine/Private/WorldPartition/HLOD/HLODSubsystem.cpp
@@ -71,7 +71,7 @@ namespace FHLODSubsystem
 		} // If not find it through the cell
 		else if (UWorldPartitionLevelStreamingDynamic* LevelStreaming = Cast<UWorldPartitionLevelStreamingDynamic>(FLevelUtils::FindStreamingLevel(InWorldPartitionHLOD->GetLevel())))
 		{
-			return LevelStreaming->GetWorldPartitionRuntimeCell()->GetTypedOuter<UWorldPartition>();
+			return LevelStreaming->GetWorldPartitionRuntimeCell()->GetCellOwner()->GetOuterWorld()->GetWorldPartition();;
 		}
 
 		return nullptr;
@@ -236,7 +236,7 @@ void UHLODSubsystem::UnregisterHLODActor(AWorldPartitionHLOD* InWorldPartitionHL
 
 void UHLODSubsystem::OnCellShown(const UWorldPartitionRuntimeCell* InCell)
 {
-	const UWorldPartition* WorldPartition = InCell->GetTypedOuter<UWorldPartition>();
+	const UWorldPartition* WorldPartition = InCell->GetCellOwner()->GetOuterWorld()->GetWorldPartition();
 	check(WorldPartition && WorldPartition->IsStreamingEnabled());
 
 	FWorldPartitionHLODRuntimeData& WorldPartitionHLODRuntimeData = WorldPartitionsHLODRuntimeData.FindChecked(WorldPartition);
@@ -258,7 +258,7 @@ void UHLODSubsystem::OnCellShown(const UWorldPartitionRuntimeCell* InCell)
 
 void UHLODSubsystem::OnCellHidden(const UWorldPartitionRuntimeCell* InCell)
 {
-	const UWorldPartition* WorldPartition = InCell->GetTypedOuter<UWorldPartition>();
+	const UWorldPartition* WorldPartition = InCell->GetCellOwner()->GetOuterWorld()->GetWorldPartition();
 	check(WorldPartition && WorldPartition->IsStreamingEnabled());
 
 	FWorldPartitionHLODRuntimeData& WorldPartitionHLODRuntimeData = WorldPartitionsHLODRuntimeData.FindChecked(WorldPartition);
@@ -401,7 +401,7 @@ bool UHLODSubsystem::RequestUnloading(const UWorldPartitionRuntimeCell* InCell)
 		return true;
 	}
 
-	const UWorldPartition* WorldPartition = InCell->GetTypedOuter<UWorldPartition>();
+	const UWorldPartition* WorldPartition = InCell->GetCellOwner()->GetOuterWorld()->GetWorldPartition();;
 	FWorldPartitionHLODRuntimeData& WorldPartitionHLODRuntimeData = WorldPartitionsHLODRuntimeData.FindChecked(WorldPartition);
 
 	FCellData& CellData = WorldPartitionHLODRuntimeData.CellsData.FindChecked(InCell->GetFName());
