Add missing include

diff --git a/Engine/Source/Runtime/Engine/Private/WorldPartition/ContentBundle/ContentBundlePaths.cpp b/Engine/Source/Runtime/Engine/Private/WorldPartition/ContentBundle/ContentBundlePaths.cpp
--- a/Engine/Source/Runtime/Engine/Private/WorldPartition/ContentBundle/ContentBundlePaths.cpp
+++ b/Engine/Source/Runtime/Engine/Private/WorldPartition/ContentBundle/ContentBundlePaths.cpp
@@ -5,6 +5,7 @@
 #include "Engine/World.h"
 #include "WorldPartition/ContentBundle/ContentBundleBase.h"
 #include "WorldPartition/ContentBundle/ContentBundleDescriptor.h"
+#include "Misc/Paths.h"
 
 #if WITH_EDITOR
 #include "String/Find.h"
