Chaos - physicshandlecomponent wakes its grabbed object on release

diff --git a/Engine/Source/Runtime/Engine/Private/PhysicsEngine/PhysicsHandleComponent.cpp b/Engine/Source/Runtime/Engine/Private/PhysicsEngine/PhysicsHandleComponent.cpp
--- a/Engine/Source/Runtime/Engine/Private/PhysicsEngine/PhysicsHandleComponent.cpp
+++ b/Engine/Source/Runtime/Engine/Private/PhysicsEngine/PhysicsHandleComponent.cpp
@@ -203,6 +203,8 @@ void UPhysicsHandleComponent::ReleaseComponent()
 
 	if (GrabbedComponent)
 	{
+		GrabbedComponent->WakeRigidBody();
+
 		GrabbedComponent = NULL;
 		GrabbedBoneName = NAME_None;
 		GrabbedHandle = nullptr;
