Fix an issue where some cell would still reference level streaming in non fast cook builds

diff --git a/Engine/Source/Runtime/Engine/Private/WorldPartition/WorldPartitionRuntimeLevelStreamingCell.cpp b/Engine/Source/Runtime/Engine/Private/WorldPartition/WorldPartitionRuntimeLevelStreamingCell.cpp
--- a/Engine/Source/Runtime/Engine/Private/WorldPartition/WorldPartitionRuntimeLevelStreamingCell.cpp
+++ b/Engine/Source/Runtime/Engine/Private/WorldPartition/WorldPartitionRuntimeLevelStreamingCell.cpp
@@ -247,12 +247,6 @@ bool UWorldPartitionRuntimeLevelStreamingCell::PopulateGeneratedPackageForCook(U
 
 		// Remap Level's SoftObjectPaths
 		FWorldPartitionLevelHelper::RemapLevelSoftObjectPaths(NewLevel, WorldPartition);
-
-		if (!GetCellOwner()->ShouldSerializeCellLevelStreaming())
-		{
-			LevelStreaming = nullptr;
-			OutModifiedPackage.Add(GetOuter()->GetPackage());
-		}
 	}
 	return true;
 }
diff --git a/Engine/Source/Runtime/Engine/Private/WorldPartition/WorldPartitionRuntimeSpatialHash.cpp b/Engine/Source/Runtime/Engine/Private/WorldPartition/WorldPartitionRuntimeSpatialHash.cpp
--- a/Engine/Source/Runtime/Engine/Private/WorldPartition/WorldPartitionRuntimeSpatialHash.cpp
+++ b/Engine/Source/Runtime/Engine/Private/WorldPartition/WorldPartitionRuntimeSpatialHash.cpp
@@ -796,8 +796,15 @@ void URuntimeSpatialHashExternalStreamingObject::PopulateGeneratorPackageForCook
 			{
 				for (UWorldPartitionRuntimeSpatialHashCell* Cell : GridLayerCell.GridCells)
 				{
-					UWorldPartitionLevelStreamingDynamic* LevelStreamingDynamic = CastChecked<UWorldPartitionRuntimeLevelStreamingCell>(Cell)->GetLevelStreaming();
-					CellToLevelStreamingPackage.Add(Cell->GetFName(), LevelStreamingDynamic->PackageNameToLoad);
+					UWorldPartitionRuntimeLevelStreamingCell* RuntimeCell = CastChecked<UWorldPartitionRuntimeLevelStreamingCell>(Cell);
+					RuntimeCell->Rename(nullptr, this);
+					
+					UWorldPartitionLevelStreamingDynamic* LevelStreamingDynamic = RuntimeCell->GetLevelStreaming();
+					CellToLevelStreamingPackage.Add(RuntimeCell->GetFName(), LevelStreamingDynamic->PackageNameToLoad);
+
+					// Level streaming are outered to the world and would not be saved within the ExternalStreamingObject.
+					// Do not save them, instead they will be created once the external streaming object is loaded at runtime. 
+					LevelStreamingDynamic->SetFlags(RF_Transient);
 				}
 			}
 		}
diff --git a/Engine/Source/Runtime/Engine/Public/WorldPartition/WorldPartitionRuntimeCellOwner.h b/Engine/Source/Runtime/Engine/Public/WorldPartition/WorldPartitionRuntimeCellOwner.h
--- a/Engine/Source/Runtime/Engine/Public/WorldPartition/WorldPartitionRuntimeCellOwner.h
+++ b/Engine/Source/Runtime/Engine/Public/WorldPartition/WorldPartitionRuntimeCellOwner.h
@@ -26,7 +26,5 @@ public:
 	virtual UWorld* GetOwningWorld() const = 0;
 	virtual UWorld* GetOuterWorld() const = 0;
 
-	virtual bool ShouldSerializeCellLevelStreaming() const { return true; }
-
 	UWorldPartition* GetWorldPartition() const { return GetOuterWorld()->GetWorldPartition(); }
 };
diff --git a/Engine/Source/Runtime/Engine/Public/WorldPartition/WorldPartitionRuntimeHash.h b/Engine/Source/Runtime/Engine/Public/WorldPartition/WorldPartitionRuntimeHash.h
--- a/Engine/Source/Runtime/Engine/Public/WorldPartition/WorldPartitionRuntimeHash.h
+++ b/Engine/Source/Runtime/Engine/Public/WorldPartition/WorldPartitionRuntimeHash.h
@@ -45,11 +45,6 @@ public:
 	//~ Begin IWorldPartitionRuntimeCellOwner Interface
 	virtual UWorld* GetOwningWorld() const override { return OwningWorld.Get(); }
 	virtual UWorld* GetOuterWorld() const override { return OuterWorld.Get(); }
-
-	// If the cell is saved within an external streaming object avoid do not serialize its level streaming. 
-	// Level streaming are outered to the world and would not be saved within the ExternalStreamingObject.
-	// Instead they will be created once the external streaming object is loaded. 
-	virtual bool ShouldSerializeCellLevelStreaming() const { return false; }
 	//~ End IWorldPartitionRuntimeCellOwner Interface
 	
 	virtual void OnStreamingObjectLoaded() {};
