Fix missing include

diff --git a/Engine/Source/Runtime/Engine/Public/WorldPartition/Cook/WorldPartitionCookPackage.h b/Engine/Source/Runtime/Engine/Public/WorldPartition/Cook/WorldPartitionCookPackage.h
--- a/Engine/Source/Runtime/Engine/Public/WorldPartition/Cook/WorldPartitionCookPackage.h
+++ b/Engine/Source/Runtime/Engine/Public/WorldPartition/Cook/WorldPartitionCookPackage.h
@@ -7,6 +7,7 @@
 #include "Hash/CityHash.h"
 #include "Misc/Paths.h"
 #include "WorldPartition/WorldPartitionLog.h"
+#include "UObject/Package.h"
 
 struct FWorldPartitionCookPackage
 {
