Made it possible to configure a project to not instantiate AvoidanceManager at all.

diff --git a/Engine/Source/Runtime/Engine/Private/UnrealEngine.cpp b/Engine/Source/Runtime/Engine/Private/UnrealEngine.cpp
--- a/Engine/Source/Runtime/Engine/Private/UnrealEngine.cpp
+++ b/Engine/Source/Runtime/Engine/Private/UnrealEngine.cpp
@@ -3167,7 +3167,10 @@ void UEngine::InitializeObjectReferences()
 	LoadEngineClass<AWorldSettings>(WorldSettingsClassName, WorldSettingsClass);
 	LoadEngineClass<UNavigationSystemBase>(NavigationSystemClassName, NavigationSystemClass);
 	LoadEngineClass<UNavigationSystemConfig>(NavigationSystemConfigClassName, NavigationSystemConfigClass);
-	LoadEngineClass<UAvoidanceManager>(AvoidanceManagerClassName, AvoidanceManagerClass);
+	if (AvoidanceManagerClassName.IsValid())
+	{
+		LoadEngineClass<UAvoidanceManager>(AvoidanceManagerClassName, AvoidanceManagerClass);
+	}
 	LoadEngineClass<UGameUserSettings>(GameUserSettingsClassName, GameUserSettingsClass);
 	LoadEngineClass<ALevelScriptActor>(LevelScriptActorClassName, LevelScriptActorClass);
 
