CIS - Fix static anlayser warning C6011: Dereferencing NULL pointer 'InMeshComponent'.

diff --git a/Engine/Source/Runtime/Engine/Private/SkeletalRenderCPUSkin.cpp b/Engine/Source/Runtime/Engine/Private/SkeletalRenderCPUSkin.cpp
--- a/Engine/Source/Runtime/Engine/Private/SkeletalRenderCPUSkin.cpp
+++ b/Engine/Source/Runtime/Engine/Private/SkeletalRenderCPUSkin.cpp
@@ -605,7 +605,10 @@ FDynamicSkelMeshObjectDataCPUSkin::FDynamicSkelMeshObjectDataCPUSkin(
 	}
 
 #if !(UE_BUILD_SHIPPING || UE_BUILD_TEST)
-	MeshComponentSpaceTransforms = InMeshComponent->GetComponentSpaceTransforms();
+	if (InMeshComponent)
+	{
+		MeshComponentSpaceTransforms = InMeshComponent->GetComponentSpaceTransforms();
+	}
 #endif
 }
 
