Chaos: update async defaults

diff --git a/Engine/Source/Runtime/Engine/Private/PhysicsEngine/PhysicsSettings.cpp b/Engine/Source/Runtime/Engine/Private/PhysicsEngine/PhysicsSettings.cpp
--- a/Engine/Source/Runtime/Engine/Private/PhysicsEngine/PhysicsSettings.cpp
+++ b/Engine/Source/Runtime/Engine/Private/PhysicsEngine/PhysicsSettings.cpp
@@ -24,7 +24,7 @@ UPhysicsSettings::UPhysicsSettings(const FObjectInitializer& ObjectInitializer)
 	, MinPhysicsDeltaTime(UE_SMALL_NUMBER)
 	, MaxPhysicsDeltaTime(1.f / 30.f)
 	, bSubstepping(false)
-	, bTickPhysicsAsync(true)
+	, bTickPhysicsAsync(false)
 	, AsyncFixedTimeStepSize(1.f / 30.f)
 	, MaxSubstepDeltaTime(1.f / 60.f)
 	, MaxSubsteps(6)
