Use TGuardValue on GIsEditorLoadingPackage in FActorDeferredScriptManager::ProcessAsyncTasks.

diff --git a/Engine/Source/Runtime/Engine/Private/ActorDeferredScriptManager.cpp b/Engine/Source/Runtime/Engine/Private/ActorDeferredScriptManager.cpp
--- a/Engine/Source/Runtime/Engine/Private/ActorDeferredScriptManager.cpp
+++ b/Engine/Source/Runtime/Engine/Private/ActorDeferredScriptManager.cpp
@@ -70,8 +70,7 @@ void FActorDeferredScriptManager::ProcessAsyncTasks(bool bLimitExecutionTime)
 
 		// since this deferred run of construction script was supposed to be done during level load
 		// temporarily set the global flag to prevent dirtying the level package.
-		// @note: it would quite preferable if we would have a scoped context than touching a global variable... 
-		GIsEditorLoadingPackage = true;
+		TGuardValue<bool> IsEditorLoadingPackageGuard(GIsEditorLoadingPackage, true);
 		while (!PendingConstructionScriptActors.IsEmpty() && bHasTimeLeft)
 		{
 			TWeakObjectPtr<AActor> WeakActor = PendingConstructionScriptActors.Last();
@@ -84,7 +83,6 @@ void FActorDeferredScriptManager::ProcessAsyncTasks(bool bLimitExecutionTime)
 			}
 			bHasTimeLeft = bLimitExecutionTime ? ((FPlatformTime::Seconds() - TickStartTime) < MaxSecondsPerFrame) : true;
 		}
-		GIsEditorLoadingPackage = false;
 	}
 
 	UpdateCompilationNotification();
