Added project setting to disable per-object translucency shadows (off by default)

diff --git a/Engine/Source/Runtime/Core/Private/Misc/MapErrors.cpp b/Engine/Source/Runtime/Core/Private/Misc/MapErrors.cpp
--- a/Engine/Source/Runtime/Core/Private/Misc/MapErrors.cpp
+++ b/Engine/Source/Runtime/Core/Private/Misc/MapErrors.cpp
@@ -66,6 +66,7 @@ FLazyName FMapErrors::LODActorMissingStaticMesh(TEXT("LODActorMissingStaticMesh"
 FLazyName FMapErrors::LODActorMissingActor(TEXT("LODActorMissingActor"));
 FLazyName FMapErrors::LODActorNoActorFound(TEXT("LODActorNoActor"));
 FLazyName FMapErrors::InvalidVirtualTextureUsage(TEXT("InvalidVirtualTextureUsage"));
+FLazyName FMapErrors::PrimitiveComponentHasInvalidTranslucentShadowSetting(TEXT("PrimitiveComponentHasInvalidTranslucentShadowSetting"));
 
 namespace
 {
diff --git a/Engine/Source/Runtime/Core/Public/Misc/MapErrors.h b/Engine/Source/Runtime/Core/Public/Misc/MapErrors.h
--- a/Engine/Source/Runtime/Core/Public/Misc/MapErrors.h
+++ b/Engine/Source/Runtime/Core/Public/Misc/MapErrors.h
@@ -143,6 +143,9 @@ struct CORE_API FMapErrors
    
 	/** Component is a static type but has invalid lightmap settings!  Indirect lighting will be black.  Common causes are lightmap resolution of 0, LightmapCoordinateIndex out of bounds. */
 	static FLazyName StaticComponentHasInvalidLightmapSettings;
+   
+	/** Component is a using CastVolumetricTranslucentShadow but this feature is disabled for the project! Turn on r.Shadow.TranslucentPerObject.ProjectEnabled in a project ini if required. */
+	static FLazyName PrimitiveComponentHasInvalidTranslucentShadowSetting;
 
 	/** Navigation */
 
diff --git a/Engine/Source/Runtime/Engine/Classes/Engine/RendererSettings.h b/Engine/Source/Runtime/Engine/Classes/Engine/RendererSettings.h
--- a/Engine/Source/Runtime/Engine/Classes/Engine/RendererSettings.h
+++ b/Engine/Source/Runtime/Engine/Classes/Engine/RendererSettings.h
@@ -892,6 +892,15 @@ class ENGINE_API URendererSettings : public UDeveloperSettings
 		ConfigRestartRequired = true))
 		uint32 bSupportCloudShadowOnForwardLitTranslucent : 1;
 
+	/**
+	""Enable translucent volumetric self-shadow, requires vertex and pixel shader permutations for all tranlucent materials even if not used by any light."
+	*/
+	UPROPERTY(config, EditAnywhere, Category = ShaderPermutationReduction, meta = (
+		ConsoleVariable = "r.Shadow.TranslucentPerObject.ProjectEnabled", DisplayName = "Support Volumetric Translucent Self Shadowing",
+		ToolTip = "Enable translucent volumetric self-shadow, requires vertex and pixel shader permutations for all tranlucent materials even if not used by any light.",
+		ConfigRestartRequired = true))
+		uint32 bSupportTranslucentPerObjectShadow : 1;
+
 	/**
 	"Enable cloud shadow on SingleLayerWater surface. This is evaluated per vertex to reduce GPU cost. The cloud system requires extra samplers/textures to be bound to vertex shaders."
 	*/
diff --git a/Engine/Source/Runtime/Engine/Private/Components/PrimitiveComponent.cpp b/Engine/Source/Runtime/Engine/Private/Components/PrimitiveComponent.cpp
--- a/Engine/Source/Runtime/Engine/Private/Components/PrimitiveComponent.cpp
+++ b/Engine/Source/Runtime/Engine/Private/Components/PrimitiveComponent.cpp
@@ -1221,6 +1221,15 @@ void UPrimitiveComponent::CheckForErrors()
 			->AddToken(FTextToken::Create(LOCTEXT( "MapCheck_Message_InvalidLightmapSettings", "Component is a static type but has invalid lightmap settings!  Indirect lighting will be black.  Common causes are lightmap resolution of 0, LightmapCoordinateIndex out of bounds." )))
 			->AddToken(FMapErrorToken::Create(FMapErrors::StaticComponentHasInvalidLightmapSettings));
 	}
+
+	static const auto CVar = IConsoleManager::Get().FindConsoleVariable(TEXT("r.Shadow.TranslucentPerObject.ProjectEnabled"));
+	if (bCastVolumetricTranslucentShadow && CastShadow && bCastDynamicShadow && CVar && CVar->GetInt() == 0)
+	{
+		FMessageLog("MapCheck").Warning()
+			->AddToken(FUObjectToken::Create(Owner))
+			->AddToken(FTextToken::Create(LOCTEXT( "MapCheck_Message_NoTranslucentShadowSupport", "Component is a using CastVolumetricTranslucentShadow but this feature is disabled for the project! Turn on r.Shadow.TranslucentPerObject.ProjectEnabled in a project ini if required." )))
+			->AddToken(FMapErrorToken::Create(FMapErrors::PrimitiveComponentHasInvalidTranslucentShadowSetting));
+	}
 }
 
 void UPrimitiveComponent::GetActorDescProperties(FPropertyPairsMap& PropertyPairsMap) const
diff --git a/Engine/Source/Runtime/Engine/Private/Particles/ParticleComponents.cpp b/Engine/Source/Runtime/Engine/Private/Particles/ParticleComponents.cpp
--- a/Engine/Source/Runtime/Engine/Private/Particles/ParticleComponents.cpp
+++ b/Engine/Source/Runtime/Engine/Private/Particles/ParticleComponents.cpp
@@ -3822,6 +3822,15 @@ void UParticleSystemComponent::CheckForErrors()
 			}
 		}
 	}
+
+	static const auto CVar = IConsoleManager::Get().FindConsoleVariable(TEXT("r.Shadow.TranslucentPerObject.ProjectEnabled"));
+	if (bCastVolumetricTranslucentShadow && CastShadow && bCastDynamicShadow && CVar && CVar->GetInt() == 0)
+	{
+		FMessageLog("MapCheck").Warning()
+			->AddToken(FUObjectToken::Create(this))
+			->AddToken(FTextToken::Create(LOCTEXT("MapCheck_Message_NoTranslucentShadowSupport", "Component is a using CastVolumetricTranslucentShadow but this feature is disabled for the project! Turn on r.Shadow.TranslucentPerObject.ProjectEnabled in a project ini if required.")))
+			->AddToken(FMapErrorToken::Create(FMapErrors::PrimitiveComponentHasInvalidTranslucentShadowSetting));
+	}
 }
 #endif
 
diff --git a/Engine/Source/Runtime/RenderCore/Private/RenderUtils.cpp b/Engine/Source/Runtime/RenderCore/Private/RenderUtils.cpp
--- a/Engine/Source/Runtime/RenderCore/Private/RenderUtils.cpp
+++ b/Engine/Source/Runtime/RenderCore/Private/RenderUtils.cpp
@@ -27,6 +27,14 @@ FAutoConsoleVariableRef CVarAllowNanite(
 	ECVF_ReadOnly | ECVF_RenderThreadSafe
 );
 
+int32 GAllowTranslucencyShadowsInProject = 0;
+FAutoConsoleVariableRef CVarAllowTranslucencyShadowsInProject(
+	TEXT("r.Shadow.TranslucentPerObject.ProjectEnabled"),
+	GAllowTranslucencyShadowsInProject,
+	TEXT("Enable/Disable translucency shadows on a per-project basis. Turning off can significantly reduce the number of permutations if your project has many translucent materials.\n"),
+	ECVF_ReadOnly | ECVF_RenderThreadSafe
+);
+
 int32 GRayTracingEnableInGame = 1;
 FAutoConsoleVariableRef CVarRayTracingEnableInGame(
 	TEXT("r.RayTracing.EnableInGame"),
@@ -1774,4 +1782,9 @@ bool IsUsingBasePassVelocity(const FStaticShaderPlatform Platform)
 	{
 		return (PerPlatformCVar.Get(Platform) == 1);
 	}
+}
+
+RENDERCORE_API bool AllowTranslucencyPerObjectShadows(const FStaticShaderPlatform& Platform)
+{
+	return IsFeatureLevelSupported(Platform, ERHIFeatureLevel::SM5) && GAllowTranslucencyShadowsInProject != 0;
 }
\ No newline at end of file
diff --git a/Engine/Source/Runtime/RenderCore/Public/RenderUtils.h b/Engine/Source/Runtime/RenderCore/Public/RenderUtils.h
--- a/Engine/Source/Runtime/RenderCore/Public/RenderUtils.h
+++ b/Engine/Source/Runtime/RenderCore/Public/RenderUtils.h
@@ -819,3 +819,7 @@ RENDERCORE_API bool AreSkinCacheShadersEnabled(EShaderPlatform Platform);
  */
 RENDERCORE_API bool DoesRuntimeSupportOnePassPointLightShadows(EShaderPlatform Platform);
 
+/**
+ * Read-only switch to check if translucency per object shadows are enabled.
+ */
+RENDERCORE_API bool AllowTranslucencyPerObjectShadows(const FStaticShaderPlatform &Platform);
diff --git a/Engine/Source/Runtime/Renderer/Private/ShadowSetup.cpp b/Engine/Source/Runtime/Renderer/Private/ShadowSetup.cpp
--- a/Engine/Source/Runtime/Renderer/Private/ShadowSetup.cpp
+++ b/Engine/Source/Runtime/Renderer/Private/ShadowSetup.cpp
@@ -2948,8 +2948,8 @@ void FSceneRenderer::CreatePerObjectProjectedShadow(
 				}
 			}
 
-			if (bTranslucentRelevance
-				&& Scene->GetFeatureLevel() >= ERHIFeatureLevel::SM5
+			if (AllowTranslucencyPerObjectShadows(Scene->GetShaderPlatform())
+				&& bTranslucentRelevance
 				&& bCreateTranslucentObjectShadow 
 				&& (bTranslucentShadowIsVisibleThisFrame || bShadowIsPotentiallyVisibleNextFrame))
 			{
diff --git a/Engine/Source/Runtime/Renderer/Private/TranslucentLighting.cpp b/Engine/Source/Runtime/Renderer/Private/TranslucentLighting.cpp
--- a/Engine/Source/Runtime/Renderer/Private/TranslucentLighting.cpp
+++ b/Engine/Source/Runtime/Renderer/Private/TranslucentLighting.cpp
@@ -242,7 +242,7 @@ public:
 
 	static bool ShouldCompilePermutation(const FMeshMaterialShaderPermutationParameters& Parameters)
 	{
-		return IsTranslucentBlendMode(Parameters.MaterialParameters.BlendMode) && IsFeatureLevelSupported(Parameters.Platform, ERHIFeatureLevel::SM5);
+		return AllowTranslucencyPerObjectShadows(Parameters.Platform) && IsTranslucentBlendMode(Parameters.MaterialParameters.BlendMode);
 	}
 
 	FTranslucencyShadowDepthVS() {}
@@ -288,7 +288,7 @@ class FTranslucencyShadowDepthPS : public FMeshMaterialShader
 public:
 	static bool ShouldCompilePermutation(const FMeshMaterialShaderPermutationParameters& Parameters)
 	{
-		return IsTranslucentBlendMode(Parameters.MaterialParameters.BlendMode) && IsFeatureLevelSupported(Parameters.Platform, ERHIFeatureLevel::SM5);
+		return AllowTranslucencyPerObjectShadows(Parameters.Platform) && IsTranslucentBlendMode(Parameters.MaterialParameters.BlendMode);
 	}
 
 	FTranslucencyShadowDepthPS() = default;
