Bind Input Delegates to a component when it is added dynamically to an actor via the AActor::AddComponet method. This fixes an issue where components that are on the CDO would recieve input callbacks, but dynamically added components form a BP graph did not.

Make this a config project setting to let you disable it, just as a safety net if it causes any unforeseen problems.

#rnx

diff --git a/Engine/Source/Runtime/Engine/Classes/GameFramework/InputSettings.h b/Engine/Source/Runtime/Engine/Classes/GameFramework/InputSettings.h
--- a/Engine/Source/Runtime/Engine/Classes/GameFramework/InputSettings.h
+++ b/Engine/Source/Runtime/Engine/Classes/GameFramework/InputSettings.h
@@ -79,6 +79,13 @@ class ENGINE_API UInputSettings
 	 */
 	UPROPERTY(config, EditAnywhere, Category = "Input")
 	uint8 bShouldFlushPressedKeysOnViewportFocusLost:1;
+
+	/**
+	 * Should components that are dynamically added via the 'AddComponent' function at runtime have input delegates bound to them?
+	 * @see AActor::FinishAddComponent
+	 */
+	UPROPERTY(config, EditAnywhere, Category = "Input")
+	uint8 bEnableDynamicComponentInputBinding:1;
 	
 	/** Should the touch input interface be shown always, or only when the platform has a touch screen? */
 	UPROPERTY(config, EditAnywhere, Category="Mobile")
diff --git a/Engine/Source/Runtime/Engine/Private/ActorConstruction.cpp b/Engine/Source/Runtime/Engine/Private/ActorConstruction.cpp
--- a/Engine/Source/Runtime/Engine/Private/ActorConstruction.cpp
+++ b/Engine/Source/Runtime/Engine/Private/ActorConstruction.cpp
@@ -32,6 +32,8 @@
 #include "ProfilingDebugging/CsvProfiler.h"
 #include "Algo/Transform.h"
 #include "Misc/ScopeExit.h"
+#include "Engine/InputDelegateBinding.h"
+#include "GameFramework/InputSettings.h"
 
 #if WITH_EDITOR
 #include "Editor.h"
@@ -1261,6 +1263,11 @@ void AActor::FinishAddComponent(UActorComponent* NewActorComp, bool bManualAttac
 				World->UpdateCullDistanceVolumes(this, NewPrimitiveComponent);
 			}
 		}
+		
+		if (InputComponent && GetDefault<UInputSettings>()->bEnableDynamicComponentInputBinding)
+		{
+			UInputDelegateBinding::BindInputDelegates(NewActorComp->GetClass(), InputComponent, NewActorComp);
+		}
 	}
 }
 
diff --git a/Engine/Source/Runtime/Engine/Private/InputDelegateBinding.cpp b/Engine/Source/Runtime/Engine/Private/InputDelegateBinding.cpp
--- a/Engine/Source/Runtime/Engine/Private/InputDelegateBinding.cpp
+++ b/Engine/Source/Runtime/Engine/Private/InputDelegateBinding.cpp
@@ -69,7 +69,7 @@ void UInputDelegateBinding::BindInputDelegatesWithSubojects(AActor* InActor, UIn
 		for(UActorComponent* Comp : ComponentArray)
 		{
 			const UClass* CompClass = Comp ? Comp->GetClass() : nullptr;
-			if(CompClass && SupportsInputDelegate(CompClass) && Comp != InputComponent)
+			if(CompClass && Comp != InputComponent)
 			{
 				UInputDelegateBinding::BindInputDelegates(CompClass, InputComponent, Comp);	
 			}			
diff --git a/Engine/Source/Runtime/Engine/Private/UserInterface/InputSettings.cpp b/Engine/Source/Runtime/Engine/Private/UserInterface/InputSettings.cpp
--- a/Engine/Source/Runtime/Engine/Private/UserInterface/InputSettings.cpp
+++ b/Engine/Source/Runtime/Engine/Private/UserInterface/InputSettings.cpp
@@ -27,6 +27,7 @@ UInputSettings::UInputSettings(const FObjectInitializer& ObjectInitializer)
 	, bEnableMotionControls(true)
 	, bFilterInputByPlatformUser(false)
 	, bShouldFlushPressedKeysOnViewportFocusLost(true)
+	, bEnableDynamicComponentInputBinding(true)
 	, DefaultViewportMouseCaptureMode(EMouseCaptureMode::CapturePermanently_IncludingInitialMouseDown)
 	, DefaultViewportMouseLockMode(EMouseLockMode::LockOnCapture)
 	, DefaultPlayerInputClass(UPlayerInput::StaticClass())
