Fix potential crash introduced in 22274097: don't destroy the loader adapter when unregistering volume actors, just unload.

#rnx

diff --git a/Engine/Source/Runtime/Engine/Private/LocationVolume.cpp b/Engine/Source/Runtime/Engine/Private/LocationVolume.cpp
--- a/Engine/Source/Runtime/Engine/Private/LocationVolume.cpp
+++ b/Engine/Source/Runtime/Engine/Private/LocationVolume.cpp
@@ -60,8 +60,6 @@ void ALocationVolume::UnregisterAllComponents(bool bForReregister)
 	if (!bForReregister && WorldPartitionActorLoader)
 	{
 		WorldPartitionActorLoader->Unload();
-		delete WorldPartitionActorLoader;
-		WorldPartitionActorLoader = nullptr;
 	}
 }
 
