Fix an issue where generated cook package (worldpartition cells) could be used as the main map when loading.

diff --git a/Engine/Source/Runtime/Engine/Private/URL.cpp b/Engine/Source/Runtime/Engine/Private/URL.cpp
--- a/Engine/Source/Runtime/Engine/Private/URL.cpp
+++ b/Engine/Source/Runtime/Engine/Private/URL.cpp
@@ -483,7 +483,7 @@ FURL::FURL( FURL* Base, const TCHAR* TextURL, ETravelType Type )
 							FName TargetTestName(*URLStr);
 							for (const FAssetData& MapAsset : MapList)
 							{
-								if (MapAsset.AssetName == TargetTestName)
+								if (MapAsset.AssetName == TargetTestName && !MapAsset.HasAnyPackageFlags(PKG_CookGenerated))
 								{
 									Map = MapAsset.PackageName.ToString();
 									bFoundMap = true;
