Optimus: More compact  UI for the mesh deformer settings.
Settings now appear without the Settings outer block.
#rnx

diff --git a/Engine/Plugins/Animation/DeformerGraph/Source/OptimusCore/Public/OptimusDeformerInstance.h b/Engine/Plugins/Animation/DeformerGraph/Source/OptimusCore/Public/OptimusDeformerInstance.h
--- a/Engine/Plugins/Animation/DeformerGraph/Source/OptimusCore/Public/OptimusDeformerInstance.h
+++ b/Engine/Plugins/Animation/DeformerGraph/Source/OptimusCore/Public/OptimusDeformerInstance.h
@@ -81,11 +81,11 @@ struct FOptimusDeformerInstanceComponentBinding
 	GENERATED_BODY()
 
 	/** Binding name on deformer graph. */
-	UPROPERTY(VisibleAnywhere, Category="Binding", meta = (DisplayName = "Binding"))
+	UPROPERTY(VisibleAnywhere, Category="Deformer", meta = (DisplayName = "Binding"))
 	FName ProviderName;
 
 	/** Component name to bind. This should be sanitized before storage. */
-	UPROPERTY(EditAnywhere, Category = "Binding", meta = (DisplayName = "Component"))
+	UPROPERTY(EditAnywhere, Category="Deformer", meta = (DisplayName = "Component"))
  	FName ComponentName;
 
 	/** Get the component on an actor that matches the stored component name. */
@@ -110,7 +110,7 @@ class OPTIMUSCORE_API UOptimusDeformerInstanceSettings :
 	TWeakObjectPtr<UOptimusDeformer> Deformer;
 
 	/** Array of binding descriptions. This is fixed and used by GetComponentBindings() to resolve final bindings for a given context. */
-	UPROPERTY(EditAnywhere, Category = "Deformer|Settings", EditFixedSize, meta = (EditFixedOrder))
+	UPROPERTY(EditAnywhere, Category="Deformer", EditFixedSize, meta = (DisplayName = "Deformer Bindings", EditFixedOrder))
 	TArray<FOptimusDeformerInstanceComponentBinding> Bindings;
 
 public:
diff --git a/Engine/Source/Runtime/Engine/Classes/Animation/MeshDeformerInstance.h b/Engine/Source/Runtime/Engine/Classes/Animation/MeshDeformerInstance.h
--- a/Engine/Source/Runtime/Engine/Classes/Animation/MeshDeformerInstance.h
+++ b/Engine/Source/Runtime/Engine/Classes/Animation/MeshDeformerInstance.h
@@ -7,15 +7,21 @@
 
 class FSceneInterface;
 
-UCLASS(Abstract, EditInlineNew)
+
+/**
+ * Base class for mesh deformers instance settings.
+ * This contains the serialized user settings to apply to the UMeshDeformer.
+ */
+UCLASS(Abstract, DefaultToInstanced)
 class ENGINE_API UMeshDeformerInstanceSettings : public UObject
 {
 	GENERATED_BODY()
 };
 
+
 /** 
  * Base class for mesh deformers instances.
- * This contains the per instance state for a UMeshDeformer.
+ * This contains the transient per instance state for a UMeshDeformer.
  */
 UCLASS(Abstract)
 class ENGINE_API UMeshDeformerInstance : public UObject
@@ -32,13 +38,14 @@ public:
 	};
 
 	/** Called to allocate any persistent render resources */
-	virtual void AllocateResources() {}
+	virtual void AllocateResources() PURE_VIRTUAL(, );
 
 	/** Called when persistent render resources should be released */ 
-	virtual void ReleaseResources() {}
+	virtual void ReleaseResources() PURE_VIRTUAL(, );
 
 	/** Get if mesh deformer is active (compiled and valid). */
 	virtual bool IsActive() const PURE_VIRTUAL(, return false;);
-	/** Enqueue the mesh deformer workload on a scene. */
+	
+	/** Enqueue the mesh deformer workload on a scene. InOwnerName is used for debugging. */
 	virtual void EnqueueWork(FSceneInterface* InScene, EWorkLoad InWorkLoadType, FName InOwnerName) PURE_VIRTUAL(, );
 };
diff --git a/Engine/Source/Runtime/Engine/Classes/Components/SkinnedMeshComponent.h b/Engine/Source/Runtime/Engine/Classes/Components/SkinnedMeshComponent.h
--- a/Engine/Source/Runtime/Engine/Classes/Components/SkinnedMeshComponent.h
+++ b/Engine/Source/Runtime/Engine/Classes/Components/SkinnedMeshComponent.h
@@ -291,7 +291,7 @@ protected:
 	UMeshDeformer* GetActiveMeshDeformer() const;
 
 	/** Object containing instance settings for the bound MeshDeformer. */
-	UPROPERTY(VisibleAnywhere, BlueprintReadOnly, Instanced, DisplayName = "Settings", Category = "Deformer", meta = (EditInline, EditCondition = "MeshDeformerInstanceSettings!=nullptr", HideEditConditionToggle, EditConditionHides))
+	UPROPERTY(VisibleAnywhere, BlueprintReadOnly, Category = "Deformer", meta = (DisplayName = "Settings", EditCondition = "MeshDeformerInstanceSettings!=nullptr", EditConditionHides, ShowOnlyInnerProperties))
 	TObjectPtr<UMeshDeformerInstanceSettings> MeshDeformerInstanceSettings;
 
 	/** Object containing state for the bound MeshDeformer. */
