[WaveformTransformations] Use TotalNumSamples to when initing transformations UObjects
fix for #jira UE-165624

diff --git a/Engine/Source/Runtime/Engine/Classes/Sound/SoundWave.h b/Engine/Source/Runtime/Engine/Classes/Sound/SoundWave.h
--- a/Engine/Source/Runtime/Engine/Classes/Sound/SoundWave.h
+++ b/Engine/Source/Runtime/Engine/Classes/Sound/SoundWave.h
@@ -595,6 +595,7 @@ private:
 	// Called when we change any properties about the underlying audio asset
 #if WITH_EDITOR
 	void UpdateAsset();
+	void UpdateTransformations();
 #endif
 
 public:
diff --git a/Engine/Source/Runtime/Engine/Private/SoundWave.cpp b/Engine/Source/Runtime/Engine/Private/SoundWave.cpp
--- a/Engine/Source/Runtime/Engine/Private/SoundWave.cpp
+++ b/Engine/Source/Runtime/Engine/Private/SoundWave.cpp
@@ -2485,20 +2485,7 @@ void USoundWave::PostEditChangeProperty(FPropertyChangedEvent& PropertyChangedEv
 		//Keep transformations UObjects configurations up to date
 		if (Name == TransformationsFName)
 		{
-			FWaveTransformUObjectConfiguration ChainConfiguration = FWaveTransformUObjectConfiguration();
-
-			ChainConfiguration.NumChannels = NumChannels;
-			ChainConfiguration.SampleRate = GetSampleRateForCurrentPlatform();
-			ChainConfiguration.StartTime = 0.f;
-			ChainConfiguration.EndTime = Duration;
-
-			for (TObjectPtr<UWaveformTransformationBase>& Transformation : Transformations)
-			{
-				if (Transformation)
-				{
-					Transformation->UpdateConfiguration(ChainConfiguration);
-				}
-			}
+			UpdateTransformations();
 		}
 	
 		// Prevent constant re-compression of SoundWave while properties are being changed interactively
@@ -3906,6 +3893,32 @@ void USoundWave::UpdateAsset()
 		}
 	}
 }
+
+void USoundWave::UpdateTransformations()
+{
+	FWaveTransformUObjectConfiguration ChainConfiguration = FWaveTransformUObjectConfiguration();
+
+	if (TotalSamples == 0)
+	{
+		UE_LOG(LogAudio, Warning, TEXT("Found %d total samples when updating transformations on soundwave %s"), TotalSamples, *GetNameSafe(this));
+		UE_LOG(LogAudio, Warning, TEXT("Consider reimporting the asset, transformations might not work correctly"));
+		return;
+	}
+
+	ChainConfiguration.NumChannels = NumChannels;
+	ChainConfiguration.SampleRate = ImportedSampleRate > 0 ? ImportedSampleRate : GetSampleRateForCurrentPlatform();
+	ChainConfiguration.StartTime = 0.f;
+	ChainConfiguration.EndTime = TotalSamples / ChainConfiguration.SampleRate;
+
+	for (TObjectPtr<UWaveformTransformationBase>& Transformation : Transformations)
+	{
+		if (Transformation)
+		{
+			Transformation->UpdateConfiguration(ChainConfiguration);
+		}
+	}
+}
+
 #endif // #if WITH_EDITOR
 
 void USoundWave::CacheInheritedLoadingBehavior() const
