Don't show the SM6 toasts on non-Windows platforms to avoid confusion around platform requirements.

diff --git a/Engine/Source/Runtime/Engine/Private/RendererSettings.cpp b/Engine/Source/Runtime/Engine/Private/RendererSettings.cpp
--- a/Engine/Source/Runtime/Engine/Private/RendererSettings.cpp
+++ b/Engine/Source/Runtime/Engine/Private/RendererSettings.cpp
@@ -347,6 +347,8 @@ bool URendererSettings::CanEditChange(const FProperty* InProperty) const
 
 void URendererSettings::CheckForMissingShaderModels()
 {
+	// Don't show the SM6 toasts on non-Windows platforms to avoid confusion around platform requirements.
+#if PLATFORM_WINDOWS
 	if (GIsEditor && ShadowMapMethod == EShadowMapMethod::VirtualShadowMaps)
 	{
 		TArray<FString> D3D11TargetedShaderFormats;
@@ -394,6 +396,7 @@ void URendererSettings::CheckForMissingShaderModels()
 			ShaderModelNotificationPtr = FSlateNotificationManager::Get().AddNotification(Info);
 		}
 	}
+#endif // PLATFORM_WINDOWS
 }
 #endif // #if WITH_EDITOR
 
diff --git a/Engine/Source/Runtime/Engine/Private/StaticMesh.cpp b/Engine/Source/Runtime/Engine/Private/StaticMesh.cpp
--- a/Engine/Source/Runtime/Engine/Private/StaticMesh.cpp
+++ b/Engine/Source/Runtime/Engine/Private/StaticMesh.cpp
@@ -6221,6 +6221,8 @@ void UStaticMesh::ExecutePostLoadInternal(FStaticMeshPostLoadContext& Context)
 #if WITH_EDITOR
 void UStaticMesh::CheckForMissingShaderModels()
 {
+	// Don't show the SM6 toasts on non-Windows platforms to avoid confusion around platform requirements.
+#if PLATFORM_WINDOWS
 	static bool bWarnedAboutMissingShaderModel = false;
 	if (GIsEditor && NaniteSettings.bEnabled && !bWarnedAboutMissingShaderModel)
 	{
@@ -6271,6 +6273,7 @@ void UStaticMesh::CheckForMissingShaderModels()
 			ShaderModelNotificationPtr = FSlateNotificationManager::Get().AddNotification(Info);
 		}
 	}
+#endif // PLATFORM_WINDOWS
 }
 #endif
 
