When checking if a project has SM6, we also need to check the old config "TargetedRHIs"

diff --git a/Engine/Source/Runtime/Engine/Private/RendererSettings.cpp b/Engine/Source/Runtime/Engine/Private/RendererSettings.cpp
--- a/Engine/Source/Runtime/Engine/Private/RendererSettings.cpp
+++ b/Engine/Source/Runtime/Engine/Private/RendererSettings.cpp
@@ -357,6 +357,14 @@ void URendererSettings::CheckForMissingShaderModels()
 		TArray<FString> D3D12TargetedShaderFormats;
 		GConfig->GetArray(TEXT("/Script/WindowsTargetPlatform.WindowsTargetSettings"), TEXT("D3D12TargetedShaderFormats"), D3D12TargetedShaderFormats, GEngineIni);
 
+		TArray<FString> TargetedRHIs;
+		GConfig->GetArray(TEXT("/Script/WindowsTargetPlatform.WindowsTargetSettings"), TEXT("TargetedRHIs"), TargetedRHIs, GEngineIni);
+
+		if (TargetedRHIs.Contains(TEXT("PCD3D_SM6")))
+		{
+			D3D12TargetedShaderFormats.AddUnique(TEXT("PCD3D_SM6"));
+		}
+
 		const bool bProjectUsesD3D = (D3D11TargetedShaderFormats.Num() + D3D12TargetedShaderFormats.Num()) > 0;
 
 		if (bProjectUsesD3D && !D3D12TargetedShaderFormats.Contains(TEXT("PCD3D_SM6")))
diff --git a/Engine/Source/Runtime/Engine/Private/StaticMesh.cpp b/Engine/Source/Runtime/Engine/Private/StaticMesh.cpp
--- a/Engine/Source/Runtime/Engine/Private/StaticMesh.cpp
+++ b/Engine/Source/Runtime/Engine/Private/StaticMesh.cpp
@@ -6232,6 +6232,14 @@ void UStaticMesh::CheckForMissingShaderModels()
 		TArray<FString> D3D12TargetedShaderFormats;
 		GConfig->GetArray(TEXT("/Script/WindowsTargetPlatform.WindowsTargetSettings"), TEXT("D3D12TargetedShaderFormats"), D3D12TargetedShaderFormats, GEngineIni);
 
+		TArray<FString> TargetedRHIs;
+		GConfig->GetArray(TEXT("/Script/WindowsTargetPlatform.WindowsTargetSettings"), TEXT("TargetedRHIs"), TargetedRHIs, GEngineIni);
+
+		if (TargetedRHIs.Contains(TEXT("PCD3D_SM6")))
+		{
+			D3D12TargetedShaderFormats.AddUnique(TEXT("PCD3D_SM6"));
+		}
+
 		const bool bProjectUsesD3D = (D3D11TargetedShaderFormats.Num() + D3D12TargetedShaderFormats.Num()) > 0;
 
 		if (bProjectUsesD3D && !D3D12TargetedShaderFormats.Contains(TEXT("PCD3D_SM6")))
