Clear reference to WorldPartitionCells between PIE
Fix an issue where ExternalStreamingObject was retrieving Cells from a previous PIE run that were not garbage collected

diff --git a/Engine/Source/Runtime/Engine/Private/WorldPartition/ContentBundle/ContentBundleEditor.cpp b/Engine/Source/Runtime/Engine/Private/WorldPartition/ContentBundle/ContentBundleEditor.cpp
--- a/Engine/Source/Runtime/Engine/Private/WorldPartition/ContentBundle/ContentBundleEditor.cpp
+++ b/Engine/Source/Runtime/Engine/Private/WorldPartition/ContentBundle/ContentBundleEditor.cpp
@@ -323,6 +323,8 @@ void FContentBundleEditor::GenerateStreaming(TArray<FString>* OutPackageToGenera
 		{
 			UE_LOG(LogContentBundle, Error, TEXT("[CB: %s] Failed to store streaming object for %s. PIE duplication will not work."), *GetDescriptor()->GetDisplayName());
 		}
+		ExternalStreamingObject = nullptr;
+		ContentBundleCells.Empty();
 	}
 
 	WorldPartition->FlushStreaming();
diff --git a/Engine/Source/Runtime/Engine/Private/WorldPartition/ContentBundle/ContentBundleWorldSubsystem.cpp b/Engine/Source/Runtime/Engine/Private/WorldPartition/ContentBundle/ContentBundleWorldSubsystem.cpp
--- a/Engine/Source/Runtime/Engine/Private/WorldPartition/ContentBundle/ContentBundleWorldSubsystem.cpp
+++ b/Engine/Source/Runtime/Engine/Private/WorldPartition/ContentBundle/ContentBundleWorldSubsystem.cpp
@@ -15,6 +15,7 @@
 
 #if WITH_EDITOR
 #include "WorldPartition/ContentBundle/ContentBundleEditor.h"
+#include "Editor.h"
 #endif
 
 UContentBundleManager::UContentBundleManager()
@@ -28,6 +29,7 @@ void UContentBundleManager::Initialize()
 	if (!GetWorld()->IsGameWorld())
 	{
 		PIEDuplicateHelper = NewObject<UContentBundleDuplicateForPIEHelper>(this);
+		PIEDuplicateHelper->Initialize();
 	}
 #endif
 
@@ -47,7 +49,11 @@ void UContentBundleManager::Deinitialize()
 	ContentBundleContainers.Empty();
 
 #if WITH_EDITOR
-	PIEDuplicateHelper = nullptr;
+	if(PIEDuplicateHelper != nullptr)
+	{
+		PIEDuplicateHelper->Deinitialize();
+		PIEDuplicateHelper = nullptr;
+	}
 #endif
 }
 
@@ -133,6 +139,18 @@ void UContentBundleManager::OnWorldPartitionUninitialized(UWorldPartition* InWor
 
 #if WITH_EDITOR
 
+void UContentBundleDuplicateForPIEHelper::Initialize()
+{
+	FEditorDelegates::EndPIE.AddUObject(this, &UContentBundleDuplicateForPIEHelper::OnPIEEnded);
+}
+
+void UContentBundleDuplicateForPIEHelper::Deinitialize()
+{
+	FEditorDelegates::EndPIE.RemoveAll(this);
+
+	Clear();
+}
+
 bool UContentBundleDuplicateForPIEHelper::StoreContentBundleStreamingObect(const FContentBundleEditor& ContentBundleEditor, URuntimeHashExternalStreamingObjectBase* StreamingObject)
 {
 	FGuid ContentBundleUid = ContentBundleEditor.GetDescriptor()->GetGuid();
@@ -150,4 +168,9 @@ URuntimeHashExternalStreamingObjectBase* UContentBundleDuplicateForPIEHelper::Re
 	return ContentBundleStreamingObject != nullptr ? *ContentBundleStreamingObject : nullptr;
 }
 
+void UContentBundleDuplicateForPIEHelper::OnPIEEnded(const bool bIsSimulating)
+{
+	Clear();
+}
+
 #endif
diff --git a/Engine/Source/Runtime/Engine/Public/WorldPartition/ContentBundle/ContentBundleWorldSubsystem.h b/Engine/Source/Runtime/Engine/Public/WorldPartition/ContentBundle/ContentBundleWorldSubsystem.h
--- a/Engine/Source/Runtime/Engine/Public/WorldPartition/ContentBundle/ContentBundleWorldSubsystem.h
+++ b/Engine/Source/Runtime/Engine/Public/WorldPartition/ContentBundle/ContentBundleWorldSubsystem.h
@@ -64,6 +64,9 @@ class UContentBundleDuplicateForPIEHelper : public UObject
 
 public:
 #if WITH_EDITOR
+	void Initialize();
+	void Deinitialize();
+
 	bool StoreContentBundleStreamingObect(const FContentBundleEditor& ContentBundleEditor, URuntimeHashExternalStreamingObjectBase* StreamingObject);
 	URuntimeHashExternalStreamingObjectBase* RetrieveContentBundleStreamingObject(const FContentBundle& ContentBundle) const;
 
@@ -72,6 +75,8 @@ public:
 #endif
 
 private:
+	void OnPIEEnded(const bool bIsSimulating);
+
 #if WITH_EDITORONLY_DATA
 	UPROPERTY()
 	TMap<FGuid, TObjectPtr<URuntimeHashExternalStreamingObjectBase>> StreamingObjects;
