Fix for TTS plugin being silent (new virtual ISubmixBufferListener::IsRenderingAudio())

diff --git a/Engine/Plugins/Experimental/TextToSpeech/Source/TextToSpeech/Private/Flite/FliteTextToSpeechSubmixListener.cpp b/Engine/Plugins/Experimental/TextToSpeech/Source/TextToSpeech/Private/Flite/FliteTextToSpeechSubmixListener.cpp
--- a/Engine/Plugins/Experimental/TextToSpeech/Source/TextToSpeech/Private/Flite/FliteTextToSpeechSubmixListener.cpp
+++ b/Engine/Plugins/Experimental/TextToSpeech/Source/TextToSpeech/Private/Flite/FliteTextToSpeechSubmixListener.cpp
@@ -86,6 +86,11 @@ void FFliteTextToSpeechSubmixListener::OnNewSubmixBuffer(const USoundSubmix* Own
 	}
 }
 
+bool FFliteTextToSpeechSubmixListener::IsRenderingAudio() const
+{
+	return bAllowPlayback;
+}
+
 void FFliteTextToSpeechSubmixListener::StartPlayback_GameThread()
 {
 	ensure(IsInGameThread());
diff --git a/Engine/Plugins/Experimental/TextToSpeech/Source/TextToSpeech/Private/Flite/FliteTextToSpeechSubmixListener.h b/Engine/Plugins/Experimental/TextToSpeech/Source/TextToSpeech/Private/Flite/FliteTextToSpeechSubmixListener.h
--- a/Engine/Plugins/Experimental/TextToSpeech/Source/TextToSpeech/Private/Flite/FliteTextToSpeechSubmixListener.h
+++ b/Engine/Plugins/Experimental/TextToSpeech/Source/TextToSpeech/Private/Flite/FliteTextToSpeechSubmixListener.h
@@ -39,8 +39,10 @@ public:
 	float GetVolume() const;
 	/** Sets the volume the speech data should be played at */
 	void SetVolume(float InVolume);
+
 	// ISubmixBufferListener
 	virtual void OnNewSubmixBuffer(const USoundSubmix* OwningSubmix, float* AudioData, int32 NumSamples, int32 NumChannels, const int32 SampleRate, double AudioClock) override;
+	virtual bool IsRenderingAudio() const override;
 	// ~ISubmixBufferListener
 private:
 	class FSynthesizedSpeechBuffer
diff --git a/Engine/Source/Runtime/AudioMixer/Private/AudioMixerSubmix.cpp b/Engine/Source/Runtime/AudioMixer/Private/AudioMixerSubmix.cpp
--- a/Engine/Source/Runtime/AudioMixer/Private/AudioMixerSubmix.cpp
+++ b/Engine/Source/Runtime/AudioMixer/Private/AudioMixerSubmix.cpp
@@ -1078,6 +1078,15 @@ namespace Audio
 			return true;
 		}
 
+		// query the SubmixBufferListeners to see if they plan to render audio into this buffer
+		for(const ISubmixBufferListener* listener : BufferListeners)
+		{
+			if(listener->IsRenderingAudio())
+			{
+				return true;
+			}
+		}
+
 		// If this submix is not rendering any sources directly and silence has been detected, we need to check it's children submixes
 		if (MixerSourceVoices.Num() == 0 && SilenceTimeStartSeconds >= 0.0)
 		{
diff --git a/Engine/Source/Runtime/Engine/Public/ISubmixBufferListener.h b/Engine/Source/Runtime/Engine/Public/ISubmixBufferListener.h
--- a/Engine/Source/Runtime/Engine/Public/ISubmixBufferListener.h
+++ b/Engine/Source/Runtime/Engine/Public/ISubmixBufferListener.h
@@ -17,4 +17,17 @@ public:
 	@param AudioClock		Double audio clock value, from start of audio rendering.
 	*/
 	virtual void OnNewSubmixBuffer(const USoundSubmix* OwningSubmix, float* AudioData, int32 NumSamples, int32 NumChannels, const int32 SampleRate, double AudioClock) = 0;
+
+	/**
+	 * Called if the submix is evaluating disabling itself for the next buffer of audio in FMixerSubmix::IsRenderingAudio()
+	 * if this returns true, FMixerSubmix::IsRenderingAudio() will return true.
+	 * Otherwise the submix will evaluate playing sounds, children submixes, etc to see if it should auto-disable
+	 *
+	 * This is called every evaluation.
+	 * ISubmixListeners that intermittently render audio should only return true when they have work to do.
+	 */
+	virtual bool IsRenderingAudio() const
+	{
+		return false;
+	}
 };
\ No newline at end of file
