Iris
- Fixed issue with AActor::SetReplicates(false) would not stop replication for Iris
- Fixed issue with AActor::EndReplication not being forwarded if bReplicates is set to false prior to call even if the object was replicating

diff --git a/Engine/Source/Runtime/Engine/Private/Actor.cpp b/Engine/Source/Runtime/Engine/Private/Actor.cpp
--- a/Engine/Source/Runtime/Engine/Private/Actor.cpp
+++ b/Engine/Source/Runtime/Engine/Private/Actor.cpp
@@ -3911,6 +3911,12 @@ void AActor::SetReplicates(bool bInReplicates)
 					ForcePropertyCompare();
 				}
 			}
+#if UE_WITH_IRIS
+			else if (HasActorBegunPlay())
+			{
+				EndReplication(EEndPlayReason::RemovedFromWorld);
+			}
+#endif
 
 			MARK_PROPERTY_DIRTY_FROM_NAME(AActor, RemoteRole, this);
 		}
diff --git a/Engine/Source/Runtime/Engine/Private/ActorReplication.cpp b/Engine/Source/Runtime/Engine/Private/ActorReplication.cpp
--- a/Engine/Source/Runtime/Engine/Private/ActorReplication.cpp
+++ b/Engine/Source/Runtime/Engine/Private/ActorReplication.cpp
@@ -846,10 +846,7 @@ void AActor::BeginReplication()
 
 void AActor::EndReplication(EEndPlayReason::Type EndPlayReason)
 {
-	if (GetIsReplicated())
-	{
-		UE::Net::FReplicationSystemUtil::EndReplication(this, EndPlayReason);
-	}
+	UE::Net::FReplicationSystemUtil::EndReplication(this, EndPlayReason);
 }
 
 void AActor::UpdateOwningNetConnection() const
diff --git a/Engine/Source/Runtime/Engine/Private/Net/Iris/ReplicationSystem/ReplicationSystemUtil.cpp b/Engine/Source/Runtime/Engine/Private/Net/Iris/ReplicationSystem/ReplicationSystemUtil.cpp
--- a/Engine/Source/Runtime/Engine/Private/Net/Iris/ReplicationSystem/ReplicationSystemUtil.cpp
+++ b/Engine/Source/Runtime/Engine/Private/Net/Iris/ReplicationSystem/ReplicationSystemUtil.cpp
@@ -56,12 +56,12 @@ UActorReplicationBridge* FReplicationSystemUtil::GetActorReplicationBridge(const
 
 FNetHandle FReplicationSystemUtil::GetNetHandle(const AActor* Actor)
 {
-	// We really should cache the handle in the actor (or UObject if we can afford it)
-	if (!Actor || !Actor->GetIsReplicated())
+	if (!Actor)
 	{
 		return FNetHandle();
 	}
 
+	// We really should cache the handle in the actor (or UObject if we can afford it)
 	UActorReplicationBridge* Bridge = GetActorReplicationBridge(Actor);
 	return Bridge ? Bridge->GetReplicatedHandle(Actor) : FNetHandle();
 }
