Fix potential race condition while registering to delegates during loading

diff --git a/Engine/Source/Runtime/Engine/Classes/Components/SceneCaptureComponent.h b/Engine/Source/Runtime/Engine/Classes/Components/SceneCaptureComponent.h
--- a/Engine/Source/Runtime/Engine/Classes/Components/SceneCaptureComponent.h
+++ b/Engine/Source/Runtime/Engine/Classes/Components/SceneCaptureComponent.h
@@ -196,6 +196,7 @@ public:
 	virtual void Serialize(FArchive& Ar);
 
 	virtual void OnUnregister() override;
+	virtual void PostLoad() override;
 
 	/** To leverage a component's bOwnerNoSee/bOnlyOwnerSee properties, the capture view requires an "owner". Override this to set a "ViewActor" for the scene. */
 	virtual const AActor* GetViewOwner() const { return nullptr; }
@@ -211,7 +212,10 @@ protected:
 	/** Update the show flags from our show flags settings (ideally, you'd be able to set this more directly, but currently unable to make FEngineShowFlags a UStruct to use it as a FProperty...) */
 	void UpdateShowFlags();
 
+	void RegisterDelegates();
+	void UnregisterDelegates();
 	void ReleaseGarbageReferences();
+	bool bDelegatesRegistered = false;
 
 	virtual void UpdateSceneCaptureContents(FSceneInterface* Scene) {};
 
diff --git a/Engine/Source/Runtime/Engine/Private/Components/SceneCaptureComponent.cpp b/Engine/Source/Runtime/Engine/Private/Components/SceneCaptureComponent.cpp
--- a/Engine/Source/Runtime/Engine/Private/Components/SceneCaptureComponent.cpp
+++ b/Engine/Source/Runtime/Engine/Private/Components/SceneCaptureComponent.cpp
@@ -169,17 +169,48 @@ USceneCaptureComponent::USceneCaptureComponent(const FObjectInitializer& ObjectI
 	ShowFlags.SetHMDDistortion(0);
 	ShowFlags.SetOnScreenDebug(0);
 
-	if (!HasAnyFlags(RF_ClassDefaultObject))
+	if (!HasAnyFlags(RF_ClassDefaultObject | RF_ArchetypeObject))
 	{
-		FCoreUObjectDelegates::GetPreGarbageCollectDelegate().AddUObject(this, &USceneCaptureComponent::ReleaseGarbageReferences);
+		if (HasAnyFlags(RF_NeedPostLoad) || GetOuter()->HasAnyFlags(RF_NeedPostLoad))
+		{
+			// Delegate registration is not thread-safe, so we postpone it on PostLoad when coming from loading which could be on another thread
+		}
+		else
+		{
+			RegisterDelegates();
+		}
 	}
 }
 
+void USceneCaptureComponent::RegisterDelegates()
+{
+	ensureMsgf(IsInGameThread(), TEXT("Potential race condition in USceneCaptureComponent registering to a delegate from non game-thread"));
+	FCoreUObjectDelegates::GetPreGarbageCollectDelegate().AddUObject(this, &USceneCaptureComponent::ReleaseGarbageReferences);
+	bDelegatesRegistered = true;
+}
+
+void USceneCaptureComponent::UnregisterDelegates()
+{
+	if (!HasAnyFlags(RF_ClassDefaultObject | RF_ArchetypeObject))
+	{
+		check(bDelegatesRegistered);
+
+		FCoreUObjectDelegates::GetPreGarbageCollectDelegate().RemoveAll(this);
+	}
+}
+
+void USceneCaptureComponent::PostLoad()
+{
+	Super::PostLoad();
+
+	RegisterDelegates();
+}
+
 void USceneCaptureComponent::BeginDestroy()
 {
 	Super::BeginDestroy();
-
-	FCoreUObjectDelegates::GetPreGarbageCollectDelegate().RemoveAll(this);
+	
+	UnregisterDelegates();
 }
 
 // Because HiddenActors and ShowOnlyActors are serialized they are not easily refactored into weak pointers.
