- Add IsComponentPSOPrecachingEnabled() & IsResourcePSOPrecachingEnabled() helper functions (components enabled by default, resources not and are opt in)
- Use new IsComponentPSOPrecachingEnabled() check before precaching component level PSOs
- Add cvar to disable precaching of global compute shaders and disable by default (takes extra startup time and not proven yet if it's really needed or not)

#rnx

diff --git a/Engine/Plugins/Experimental/Water/Source/Runtime/Private/WaterBodyComponent.cpp b/Engine/Plugins/Experimental/Water/Source/Runtime/Private/WaterBodyComponent.cpp
--- a/Engine/Plugins/Experimental/Water/Source/Runtime/Private/WaterBodyComponent.cpp
+++ b/Engine/Plugins/Experimental/Water/Source/Runtime/Private/WaterBodyComponent.cpp
@@ -1417,19 +1417,22 @@ void UWaterBodyComponent::PostLoad()
 
 	DeprecateData();
 
-	FPSOPrecacheParams PrecachePSOParams;
-	SetupPrecachePSOParams(PrecachePSOParams);
-	if (WaterMaterial)
+	if (IsComponentPSOPrecachingEnabled())
 	{
-		WaterMaterial->PrecachePSOs(&FLocalVertexFactory::StaticType, PrecachePSOParams);
-	}
-	if (UnderwaterPostProcessMaterial)
-	{
-		UnderwaterPostProcessMaterial->PrecachePSOs(&FLocalVertexFactory::StaticType, PrecachePSOParams);
-	}
-	if (WaterInfoMaterial)
-	{
-		WaterInfoMaterial->PrecachePSOs(&FLocalVertexFactory::StaticType, PrecachePSOParams);
+		FPSOPrecacheParams PrecachePSOParams;
+		SetupPrecachePSOParams(PrecachePSOParams);
+		if (WaterMaterial)
+		{
+			WaterMaterial->PrecachePSOs(&FLocalVertexFactory::StaticType, PrecachePSOParams);
+		}
+		if (UnderwaterPostProcessMaterial)
+		{
+			UnderwaterPostProcessMaterial->PrecachePSOs(&FLocalVertexFactory::StaticType, PrecachePSOParams);
+		}
+		if (WaterInfoMaterial)
+		{
+			WaterInfoMaterial->PrecachePSOs(&FLocalVertexFactory::StaticType, PrecachePSOParams);
+		}
 	}
 
 #if WITH_EDITOR
diff --git a/Engine/Plugins/Experimental/Water/Source/Runtime/Private/WaterMeshComponent.cpp b/Engine/Plugins/Experimental/Water/Source/Runtime/Private/WaterMeshComponent.cpp
--- a/Engine/Plugins/Experimental/Water/Source/Runtime/Private/WaterMeshComponent.cpp
+++ b/Engine/Plugins/Experimental/Water/Source/Runtime/Private/WaterMeshComponent.cpp
@@ -83,19 +83,22 @@ void UWaterMeshComponent::PrecachePSOs()
 		return;
 	}
 
-	FPSOPrecacheParams PrecachePSOParams;
-	SetupPrecachePSOParams(PrecachePSOParams);
-
-	const FVertexFactoryType* WaterVertexFactoryType = GetWaterVertexFactoryType(/*bWithWaterSelectionSupport = */ false);
-	if (FarDistanceMaterial)
+	if (IsComponentPSOPrecachingEnabled())
 	{
-		FarDistanceMaterial->PrecachePSOs(WaterVertexFactoryType, PrecachePSOParams);
-	}
-	for (UMaterialInterface* MaterialInterface : UsedMaterials)
-	{
-		if (MaterialInterface)
+		FPSOPrecacheParams PrecachePSOParams;
+		SetupPrecachePSOParams(PrecachePSOParams);
+
+		const FVertexFactoryType* WaterVertexFactoryType = GetWaterVertexFactoryType(/*bWithWaterSelectionSupport = */ false);
+		if (FarDistanceMaterial)
 		{
-			MaterialInterface->PrecachePSOs(WaterVertexFactoryType, PrecachePSOParams);
+			FarDistanceMaterial->PrecachePSOs(WaterVertexFactoryType, PrecachePSOParams);
+		}
+		for (UMaterialInterface* MaterialInterface : UsedMaterials)
+		{
+			if (MaterialInterface)
+			{
+				MaterialInterface->PrecachePSOs(WaterVertexFactoryType, PrecachePSOParams);
+			}
 		}
 	}
 }
diff --git a/Engine/Plugins/FX/Niagara/Source/Niagara/Private/NiagaraSystem.cpp b/Engine/Plugins/FX/Niagara/Source/Niagara/Private/NiagaraSystem.cpp
--- a/Engine/Plugins/FX/Niagara/Source/Niagara/Private/NiagaraSystem.cpp
+++ b/Engine/Plugins/FX/Niagara/Source/Niagara/Private/NiagaraSystem.cpp
@@ -948,7 +948,7 @@ void UNiagaraSystem::PostLoad()
 
 void UNiagaraSystem::PrecachePSOs()
 {
-	if (!FApp::CanEverRender() || !PipelineStateCache::IsPSOPrecachingEnabled())
+	if (!IsComponentPSOPrecachingEnabled() && !IsResourcePSOPrecachingEnabled())
 	{
 		return;
 	}
diff --git a/Engine/Plugins/FX/Niagara/Source/NiagaraShader/Private/NiagaraShared.cpp b/Engine/Plugins/FX/Niagara/Source/NiagaraShader/Private/NiagaraShared.cpp
--- a/Engine/Plugins/FX/Niagara/Source/NiagaraShader/Private/NiagaraShared.cpp
+++ b/Engine/Plugins/FX/Niagara/Source/NiagaraShader/Private/NiagaraShared.cpp
@@ -520,6 +520,7 @@ void FNiagaraShaderScript::UpdateCachedData_PostCompile(bool bCalledFromSerializ
 			CachedData.bViewUniformBufferUsed |= NiagaraShader->bNeedsViewUniformBuffer;
 
 			// request precache the compute shader
+			if (IsResourcePSOPrecachingEnabled() || IsComponentPSOPrecachingEnabled())
 			{
 				check(NiagaraShader->GetFrequency() == SF_Compute)
 				FRHIShader* RHIShader = GameThreadShaderMap->GetResource()->GetShader(Shader->GetResourceIndex());
diff --git a/Engine/Source/Runtime/Engine/Private/Components/MaterialBillboardComponent.cpp b/Engine/Source/Runtime/Engine/Private/Components/MaterialBillboardComponent.cpp
--- a/Engine/Source/Runtime/Engine/Private/Components/MaterialBillboardComponent.cpp
+++ b/Engine/Source/Runtime/Engine/Private/Components/MaterialBillboardComponent.cpp
@@ -381,20 +381,23 @@ void UMaterialBillboardComponent::PostLoad()
 {
 	Super::PostLoad();
 
-	FPSOPrecacheParams PrecachePSOParams;
-	SetupPrecachePSOParams(PrecachePSOParams);
-	PrecachePSOParams.PrimitiveType = PT_TriangleStrip;
-	PrecachePSOParams.bDisableBackFaceCulling = true;
+	if (IsComponentPSOPrecachingEnabled())
+	{
+		FPSOPrecacheParams PrecachePSOParams;
+		SetupPrecachePSOParams(PrecachePSOParams);
+		PrecachePSOParams.PrimitiveType = PT_TriangleStrip;
+		PrecachePSOParams.bDisableBackFaceCulling = true;
 
-	const FVertexFactoryType* VFType = &FLocalVertexFactory::StaticType;
+		const FVertexFactoryType* VFType = &FLocalVertexFactory::StaticType;
 
-	TArray<UMaterialInterface*> UsedMaterials;
-	GetUsedMaterials(UsedMaterials, false);
-	for (UMaterialInterface* MaterialInterface : UsedMaterials)
-	{
-		if (MaterialInterface)
+		TArray<UMaterialInterface*> UsedMaterials;
+		GetUsedMaterials(UsedMaterials, false);
+		for (UMaterialInterface* MaterialInterface : UsedMaterials)
 		{
-			MaterialInterface->PrecachePSOs(VFType, PrecachePSOParams);
+			if (MaterialInterface)
+			{
+				MaterialInterface->PrecachePSOs(VFType, PrecachePSOParams);
+			}
 		}
 	}
 }
diff --git a/Engine/Source/Runtime/Engine/Private/Components/SplineMeshComponent.cpp b/Engine/Source/Runtime/Engine/Private/Components/SplineMeshComponent.cpp
--- a/Engine/Source/Runtime/Engine/Private/Components/SplineMeshComponent.cpp
+++ b/Engine/Source/Runtime/Engine/Private/Components/SplineMeshComponent.cpp
@@ -581,7 +581,7 @@ bool USplineMeshComponent::Modify(bool bAlwaysMarkDirty)
 
 void USplineMeshComponent::PrecachePSOs()
 {
-	if (GetStaticMesh() == nullptr || GetStaticMesh()->GetRenderData() == nullptr || !FApp::CanEverRender())
+	if (!IsComponentPSOPrecachingEnabled() || GetStaticMesh() == nullptr || GetStaticMesh()->GetRenderData() == nullptr)
 	{
 		return;
 	}
diff --git a/Engine/Source/Runtime/Engine/Private/Components/StaticMeshComponent.cpp b/Engine/Source/Runtime/Engine/Private/Components/StaticMeshComponent.cpp
--- a/Engine/Source/Runtime/Engine/Private/Components/StaticMeshComponent.cpp
+++ b/Engine/Source/Runtime/Engine/Private/Components/StaticMeshComponent.cpp
@@ -1526,7 +1526,7 @@ void UStaticMeshComponent::InitResources()
 
 void UStaticMeshComponent::PrecachePSOs()
 {
-	if (!FApp::CanEverRender() || !PipelineStateCache::IsPSOPrecachingEnabled() || StaticMesh == nullptr || StaticMesh->GetRenderData() == nullptr)
+	if (!IsComponentPSOPrecachingEnabled() || StaticMesh == nullptr || StaticMesh->GetRenderData() == nullptr)
 	{
 		return;
 	}
diff --git a/Engine/Source/Runtime/Engine/Private/Components/TextRenderComponent.cpp b/Engine/Source/Runtime/Engine/Private/Components/TextRenderComponent.cpp
--- a/Engine/Source/Runtime/Engine/Private/Components/TextRenderComponent.cpp
+++ b/Engine/Source/Runtime/Engine/Private/Components/TextRenderComponent.cpp
@@ -1360,7 +1360,7 @@ void UTextRenderComponent::PostLoad()
 		}
 	}
 
-	if (TextMaterial)
+	if (IsComponentPSOPrecachingEnabled() && TextMaterial)
 	{
 		FPSOPrecacheParams PrecachePSOParams;
 		SetupPrecachePSOParams(PrecachePSOParams);
diff --git a/Engine/Source/Runtime/Engine/Private/InstancedStaticMesh.cpp b/Engine/Source/Runtime/Engine/Private/InstancedStaticMesh.cpp
--- a/Engine/Source/Runtime/Engine/Private/InstancedStaticMesh.cpp
+++ b/Engine/Source/Runtime/Engine/Private/InstancedStaticMesh.cpp
@@ -4481,7 +4481,7 @@ void UInstancedStaticMeshComponent::PostLoad()
 
 void UInstancedStaticMeshComponent::PrecachePSOs()
 {	
-	if (!FApp::CanEverRender() || !PipelineStateCache::IsPSOPrecachingEnabled() || GetStaticMesh() == nullptr || GetStaticMesh()->GetRenderData() == nullptr)
+	if (!IsComponentPSOPrecachingEnabled() || GetStaticMesh() == nullptr || GetStaticMesh()->GetRenderData() == nullptr)
 	{
 		return;
 	}
diff --git a/Engine/Source/Runtime/Engine/Private/PSOPrecache.cpp b/Engine/Source/Runtime/Engine/Private/PSOPrecache.cpp
--- a/Engine/Source/Runtime/Engine/Private/PSOPrecache.cpp
+++ b/Engine/Source/Runtime/Engine/Private/PSOPrecache.cpp
@@ -8,8 +8,41 @@
 #include "SceneInterface.h"
 #include "VertexFactory.h"
 
+static TAutoConsoleVariable<int32> CVarPrecacheGlobalComputeShaders(
+	TEXT("r.PSOPrecache.GlobalComputeShaders"),
+	0,
+	TEXT("Precache all global compute shaders during startup (default 0)."),
+	ECVF_ReadOnly
+);
+
+int32 GPSOPrecacheComponents = 1;
+static FAutoConsoleVariableRef CVarPSOPrecacheComponents(
+	TEXT("r.PSOPrecache.Components"),
+	GPSOPrecacheComponents,
+	TEXT("Precache all possible used PSOs by components during Postload (default 1 if PSOPrecaching is enabled)."),
+	ECVF_ReadOnly
+);
+
+int32 GPSOPrecacheResources = 0;
+static FAutoConsoleVariableRef CVarPSOPrecacheResources(
+	TEXT("r.PSOPrecache.Resources"),
+	GPSOPrecacheResources,
+	TEXT("Precache all possible used PSOs by resources during Postload (default 0 if PSOPrecaching is enabled)."),
+	ECVF_ReadOnly
+);
+
 PSOCollectorCreateFunction FPSOCollectorCreateManager::JumpTable[(int32)EShadingPath::Num][FPSOCollectorCreateManager::MaxPSOCollectorCount] = {};
 
+bool IsComponentPSOPrecachingEnabled()
+{
+	return FApp::CanEverRender() && PipelineStateCache::IsPSOPrecachingEnabled() && GPSOPrecacheComponents && !GIsEditor;
+}
+
+bool IsResourcePSOPrecachingEnabled()
+{
+	return FApp::CanEverRender() && PipelineStateCache::IsPSOPrecachingEnabled() && GPSOPrecacheResources && !GIsEditor;
+}
+
 FGraphEventArray PrecachePSOs(const TArray<FPSOPrecacheData>& PSOInitializers)
 {
 	FGraphEventArray GraphEvents;
@@ -33,7 +66,7 @@ FGraphEventArray PrecachePSOs(const TArray<FPSOPrecacheData>& PSOInitializers)
 
 int32 GValidatePrecaching = 0;
 static FAutoConsoleVariableRef CVarValidatePSOPrecaching(
-	TEXT("r.PSOPrecacheValidation"),
+	TEXT("r.PSOPrecache.Validation"),
 	GValidatePrecaching,
 	TEXT("Check if runtime used PSOs are correctly precached and track information per pass type, vertex factory and cache hit state (default 0)."),
 	ECVF_ReadOnly
diff --git a/Engine/Source/Runtime/Engine/Private/Particles/ParticleComponents.cpp b/Engine/Source/Runtime/Engine/Private/Particles/ParticleComponents.cpp
--- a/Engine/Source/Runtime/Engine/Private/Particles/ParticleComponents.cpp
+++ b/Engine/Source/Runtime/Engine/Private/Particles/ParticleComponents.cpp
@@ -2617,7 +2617,7 @@ void UParticleSystem::PostLoad()
 
 void UParticleSystem::PrecachePSOs()
 {
-	if (!FApp::CanEverRender() || !PipelineStateCache::IsPSOPrecachingEnabled())
+	if (!IsComponentPSOPrecachingEnabled() && !IsResourcePSOPrecachingEnabled())
 	{
 		return;
 	}
diff --git a/Engine/Source/Runtime/Engine/Private/ShaderCompiler/ShaderCompiler.cpp b/Engine/Source/Runtime/Engine/Private/ShaderCompiler/ShaderCompiler.cpp
--- a/Engine/Source/Runtime/Engine/Private/ShaderCompiler/ShaderCompiler.cpp
+++ b/Engine/Source/Runtime/Engine/Private/ShaderCompiler/ShaderCompiler.cpp
@@ -7033,6 +7033,12 @@ void VerifyGlobalShaders(EShaderPlatform Platform, bool bLoadedFromCacheFile, co
 
 void PrecacheComputePipelineStatesForGlobalShaders(EShaderPlatform Platform, const ITargetPlatform* TargetPlatform)
 {
+	static IConsoleVariable* PrecacheGlobalComputeShadersCVar = IConsoleManager::Get().FindConsoleVariable(TEXT("r.PSOPrecache.GlobalComputeShaders"));
+	if (!PipelineStateCache::IsPSOPrecachingEnabled() || PrecacheGlobalComputeShadersCVar == nullptr || PrecacheGlobalComputeShadersCVar->GetInt() == 0)
+	{
+		return;
+	}
+
 	FPlatformTypeLayoutParameters LayoutParams;
 	LayoutParams.InitializeForPlatform(TargetPlatform);
 	EShaderPermutationFlags PermutationFlags = GetShaderPermutationFlags(LayoutParams);
diff --git a/Engine/Source/Runtime/Engine/Public/PSOPrecache.h b/Engine/Source/Runtime/Engine/Public/PSOPrecache.h
--- a/Engine/Source/Runtime/Engine/Public/PSOPrecache.h
+++ b/Engine/Source/Runtime/Engine/Public/PSOPrecache.h
@@ -130,11 +130,22 @@ struct FPSOPrecacheData
 class IPSOCollector
 {
 public:
+	virtual ~IPSOCollector() {}
 	
 	// Collect all PSO for given material, vertex factory & params
 	virtual void CollectPSOInitializers(const FSceneTexturesConfig& SceneTexturesConfig, const FMaterial& Material, const FVertexFactoryType* VertexFactoryType, const FPSOPrecacheParams& PreCacheParams, TArray<FPSOPrecacheData>& PSOInitializers) = 0;
 };
 
+/**
+ * Precaching PSOs for components?
+ */
+extern ENGINE_API bool IsComponentPSOPrecachingEnabled();
+
+/**
+ * Precaching PSOs for resources?
+ */
+extern ENGINE_API bool IsResourcePSOPrecachingEnabled();
+
 /**
  * Try and create PSOs for all the given initializers and return an optional array of graph events of async compiling PSOs
  */
@@ -203,7 +214,6 @@ namespace PSOCollectorStats
 {
 	extern ENGINE_API int32 IsPrecachingValidationEnabled();
 
-
 	using VertexFactoryCountTableType = Experimental::TRobinHoodHashMap<const FVertexFactoryType*, uint32>;
 	struct FShaderStateUsage
 	{
diff --git a/Engine/Source/Runtime/Landscape/Private/Landscape.cpp b/Engine/Source/Runtime/Landscape/Private/Landscape.cpp
--- a/Engine/Source/Runtime/Landscape/Private/Landscape.cpp
+++ b/Engine/Source/Runtime/Landscape/Private/Landscape.cpp
@@ -914,6 +914,7 @@ void ULandscapeComponent::PostLoad()
 {
 	Super::PostLoad();
 
+	if (IsComponentPSOPrecachingEnabled())
 	{
 		TArray<UMaterialInterface*> Materials;
 		bool bGetDebugMaterials = false;
