Allow virtual textures to reuse texture builds from other platforms since they never tile.

diff --git a/Engine/Source/Runtime/Engine/Private/TextureDerivedData.cpp b/Engine/Source/Runtime/Engine/Private/TextureDerivedData.cpp
--- a/Engine/Source/Runtime/Engine/Private/TextureDerivedData.cpp
+++ b/Engine/Source/Runtime/Engine/Private/TextureDerivedData.cpp
@@ -46,6 +46,7 @@
 #include "Misc/StringBuilder.h"
 #include "ProfilingDebugging/CookStats.h"
 #include "UObject/ArchiveCookContext.h"
+#include "TextureBuildUtilities.h"
 #include "VT/VirtualTextureDataBuilder.h"
 #include "VT/LightmapVirtualTexture.h"
 #include "TextureCompiler.h"
@@ -1127,6 +1128,25 @@ static void GetBuildSettingsPerFormat(
 			FTextureBuildSettings& OutSettings = OutSettingPerLayer.Add_GetRef(SourceBuildSettings);
 			OutSettings.TextureFormatName = PlatformFormatsPerLayer[LayerIndex];
 
+			if (OutSettings.bVirtualStreamable)
+			{
+				// Virtual textures always strip the child format prefix prior to actual encode since
+				// VTs never tile. We do this here so that we can end up with the same DDC keys and
+				// avoid re-encoding the same texture N times under different keys.
+				// This is the same code used at the end of FVirtualTextureDataBuilder::BuildSourcePixels
+				FName TextureFormatPrefix;
+				FName TextureFormatName = UE::TextureBuildUtilities::TextureFormatRemovePrefixFromName(OutSettings.TextureFormatName, TextureFormatPrefix);
+
+				if (TextureFormatPrefix.IsNone())
+				{
+					OutSettings.TextureFormatName = TextureFormatName;
+				}
+				else
+				{
+					OutSettings.TextureFormatName = *(TextureFormatPrefix.ToString() + TextureFormatName.ToString());
+				}
+			}
+
 			FTexturePlatformData::FTextureEncodeResultMetadata* OutResultMetadata = nullptr;
 			if (OutResultMetadataPerLayer)
 			{
diff --git a/Engine/Source/Runtime/Engine/Private/VT/VirtualTextureDataBuilder.cpp b/Engine/Source/Runtime/Engine/Private/VT/VirtualTextureDataBuilder.cpp
--- a/Engine/Source/Runtime/Engine/Private/VT/VirtualTextureDataBuilder.cpp
+++ b/Engine/Source/Runtime/Engine/Private/VT/VirtualTextureDataBuilder.cpp
@@ -1139,6 +1139,11 @@ void FVirtualTextureDataBuilder::BuildSourcePixels(const FTextureSourceData& Sou
 		const FTextureBuildSettings& BuildSettingsForLayer = SettingsPerLayer[LayerIndex];
 		FVirtualTextureSourceLayerData& LayerData = SourceLayers[LayerIndex];
 
+		// Note that this code is duplicated at GetBuildSettingsPerFormat in TextureDerivedData.cpp so that the DDC key for 
+		// the VT ends up the same for child platforms, preventing re-encoding of the same texture data over and over. 
+		// As a result, this code shouldn't be doing anything when used via the ddc texture build/cook path - the prefix 
+		// should already be removed.
+		// 
 		// Don't want platform specific swizzling for VT tile data, this tends to add extra padding for textures with odd dimensions
 		// (VT physical tiles generally not power-of-2 after adding border)
 		FName TextureFormatPrefix;
