Fixed how world partition server streaming deals with data layer cells when server streaming-out is disabled.
Changes to Runtime DataLayer state will always be applied on all server cells affected by data layers.

diff --git a/Engine/Source/Runtime/Engine/Private/WorldPartition/WorldPartitionStreamingPolicy.cpp b/Engine/Source/Runtime/Engine/Private/WorldPartition/WorldPartitionStreamingPolicy.cpp
--- a/Engine/Source/Runtime/Engine/Private/WorldPartition/WorldPartitionStreamingPolicy.cpp
+++ b/Engine/Source/Runtime/Engine/Private/WorldPartition/WorldPartitionStreamingPolicy.cpp
@@ -103,7 +103,7 @@ void UWorldPartitionStreamingPolicy::UpdateStreamingSources()
 	}
 
 	const FTransform WorldToLocal = WorldPartition->GetInstanceTransform().Inverse();
-	UWorld* World = WorldPartition->GetWorld();
+	UWorld* World = GetWorld();
 
 	const bool bIsServer = WorldPartition->IsServer();
 	const bool bIsServerStreamingEnabled = WorldPartition->IsServerStreamingEnabled();
@@ -148,7 +148,7 @@ void UWorldPartitionStreamingPolicy::UpdateStreamingSources()
 				}
 			}
 
-			UWorldPartitionSubsystem* WorldPartitionSubsystem = WorldPartition->GetWorld()->GetSubsystem<UWorldPartitionSubsystem>();
+			UWorldPartitionSubsystem* WorldPartitionSubsystem = GetWorld()->GetSubsystem<UWorldPartitionSubsystem>();
 			check(WorldPartitionSubsystem);
 
 			for (IWorldPartitionStreamingSourceProvider* StreamingSourceProvider : WorldPartitionSubsystem->GetStreamingSourceProviders())
@@ -168,7 +168,7 @@ void UWorldPartitionStreamingPolicy::UpdateStreamingSources()
 	}
 		
 	// Update streaming sources velocity
-	const float CurrentTime = WorldPartition->GetWorld()->GetTimeSeconds();
+	const float CurrentTime = GetWorld()->GetTimeSeconds();
 	TSet<FName, DefaultKeyFuncs<FName>, TInlineSetAllocator<8>> ValidStreamingSources;
 	for (FWorldPartitionStreamingSource& StreamingSource : StreamingSources)
 	{
@@ -208,7 +208,7 @@ UE_SUPPRESS(LogWorldPartition, Verbosity, \
 
 bool UWorldPartitionStreamingPolicy::IsInBlockTillLevelStreamingCompleted(bool bIsCausedByBadStreamingPerformance /* = false*/) const
 {
-	const UWorld* World = WorldPartition->GetWorld();
+	const UWorld* World = GetWorld();
 	const bool bIsInBlockTillLevelStreamingCompleted = World->GetIsInBlockTillLevelStreamingCompleted();
 	if (bIsCausedByBadStreamingPerformance)
 	{
@@ -223,7 +223,7 @@ void UWorldPartitionStreamingPolicy::UpdateStreamingState()
 {
 	TRACE_CPUPROFILER_EVENT_SCOPE(UWorldPartitionStreamingPolicy::UpdateStreamingState);
 
-	UWorld* World = WorldPartition->GetWorld();
+	UWorld* World = GetWorld();
 	check(World && World->IsGameWorld());
 
 	// Dermine if the World's BlockTillLevelStreamingCompleted was triggered by WorldPartitionStreamingPolicy
@@ -306,35 +306,51 @@ void UWorldPartitionStreamingPolicy::UpdateStreamingState()
 			}
 		}
 
-		if (!bIsServerStreamingEnabled)
+		const UDataLayerSubsystem* DataLayerSubsystem = GetWorld()->GetSubsystem<UDataLayerSubsystem>();
+		TSet<FName> EffectiveActiveDataLayerNames = DataLayerSubsystem->GetEffectiveActiveDataLayerNames();
+		TSet<FName> EffectiveLoadedDataLayerNames = DataLayerSubsystem->GetEffectiveLoadedDataLayerNames();
+
+		auto AddServerFrameCells = [this, &EffectiveLoadedDataLayerNames, &EffectiveActiveDataLayerNames](const UWorldPartitionRuntimeCell* Cell)
 		{
-			if (WorldPartition->RuntimeHash)
+			// Non Data Layer Cells + Active Data Layers
+			if (!Cell->HasDataLayers() || Cell->HasAnyDataLayer(EffectiveActiveDataLayerNames))
 			{
-				const UDataLayerSubsystem* DataLayerSubsystem = WorldPartition->GetWorld()->GetSubsystem<UDataLayerSubsystem>();
-				const TSet<FName> EffectiveActiveDataLayerNames = DataLayerSubsystem->GetEffectiveActiveDataLayerNames();
-				const TSet<FName> EffectiveLoadedDataLayerNames = DataLayerSubsystem->GetEffectiveLoadedDataLayerNames();
+				FrameActivateCells.Add(Cell);
+			}
 
-				WorldPartition->RuntimeHash->ForEachStreamingCells([this, &EffectiveActiveDataLayerNames, &EffectiveLoadedDataLayerNames](const UWorldPartitionRuntimeCell* Cell)
+			// Loaded Data Layers Cells only
+			if (EffectiveLoadedDataLayerNames.Num())
+			{
+				if (Cell->HasDataLayers() && Cell->HasAnyDataLayer(EffectiveLoadedDataLayerNames))
 				{
-					// Non Data Layer Cells + Active Data Layers
-					if (!Cell->HasDataLayers() || Cell->HasAnyDataLayer(EffectiveActiveDataLayerNames))
-					{
-						FrameActivateCells.Add(Cell);
-					}
-
-					// Loaded Data Layers Cells only
-					if (EffectiveLoadedDataLayerNames.Num())
-					{
-						if (Cell->HasDataLayers() && Cell->HasAnyDataLayer(EffectiveLoadedDataLayerNames))
-						{
-							FrameLoadCells.Add(Cell);
-						}
-					}
+					FrameLoadCells.Add(Cell);
+				}
+			}
+		};
 
+		if (!bIsServerStreamingEnabled)
+		{
+			if (WorldPartition->RuntimeHash)
+			{
+				WorldPartition->RuntimeHash->ForEachStreamingCells([this, &AddServerFrameCells](const UWorldPartitionRuntimeCell* Cell)
+				{
+					AddServerFrameCells(Cell);
 					return true;
 				});
 			}
 		}
+		else if (!bCanDeactivateOrUnloadCells)
+		{
+			// When server streaming-out is disabled, revisit existing loaded/activated cells and add them in the proper FrameLoadCells/FrameActivateCells
+			for (const UWorldPartitionRuntimeCell* Cell : ActivatedCells)
+			{
+				AddServerFrameCells(Cell);
+			}
+			for (const UWorldPartitionRuntimeCell* Cell : LoadedCells)
+			{
+				AddServerFrameCells(Cell);
+			}
+		}
 	}
 
 	auto ShouldWaitForClientVisibility = [bIsServer, &ClientVisibleLevelNames, &bUpdateEpoch](const UWorldPartitionRuntimeCell* Cell)
@@ -375,14 +391,14 @@ void UWorldPartitionStreamingPolicy::UpdateStreamingState()
 	};
 
 	// Process cells to load
-	auto ProcessCellsToLoad = [this, bIsServer, bCanDeactivateOrUnloadCells, &ShouldWaitForClientVisibility, &ShouldSkipDisabledHLODCell](TSet<const UWorldPartitionRuntimeCell*>& Cells)
+	auto ProcessCellsToLoad = [this, bIsServer, &ShouldWaitForClientVisibility, &ShouldSkipDisabledHLODCell](TSet<const UWorldPartitionRuntimeCell*>& Cells)
 	{
 		for (TSet<const UWorldPartitionRuntimeCell*>::TIterator It = Cells.CreateIterator(); It; ++It)
 		{
 			const UWorldPartitionRuntimeCell* Cell = *It;
 			// Only deactivated server cells need to call ShouldWaitForClientVisibility (those part of ActivatedCells)
 			const bool bIsServerCellToDeactivate = bIsServer && ActivatedCells.Contains(Cell);
-			const bool bShoudSkipServerCell = bIsServerCellToDeactivate && (!bCanDeactivateOrUnloadCells || ShouldWaitForClientVisibility(Cell));
+			const bool bShoudSkipServerCell = bIsServerCellToDeactivate && ShouldWaitForClientVisibility(Cell);
 
 			if (bShoudSkipServerCell || ShouldSkipCellForPerformance(Cell) || ShouldSkipDisabledHLODCell(Cell))
 			{
@@ -396,7 +412,7 @@ void UWorldPartitionStreamingPolicy::UpdateStreamingState()
 	};
 
 	// Process cells to unload
-	auto ProcessCellsToUnload = [this, bIsServer, bCanDeactivateOrUnloadCells, &ShouldWaitForClientVisibility](TSet<const UWorldPartitionRuntimeCell*>& Cells)
+	auto ProcessCellsToUnload = [this, bIsServer, &ShouldWaitForClientVisibility](TSet<const UWorldPartitionRuntimeCell*>& Cells)
 	{
 		// Only loop for server as ShouldWaitForClientVisibility only concerns server
 		if (bIsServer)
@@ -404,7 +420,6 @@ void UWorldPartitionStreamingPolicy::UpdateStreamingState()
 			for (TSet<const UWorldPartitionRuntimeCell*>::TIterator It = Cells.CreateIterator(); It; ++It)
 			{
 				const UWorldPartitionRuntimeCell* Cell = *It;
-				check(bCanDeactivateOrUnloadCells || Cell->HasDataLayers());
 				if (ShouldWaitForClientVisibility(Cell))
 				{
 					It.RemoveCurrent();
@@ -739,7 +754,7 @@ bool UWorldPartitionStreamingPolicy::CanAddLoadedLevelToWorld(ULevel* InLevel) c
 
 bool UWorldPartitionStreamingPolicy::IsStreamingCompleted(EWorldPartitionRuntimeCellState QueryState, const TArray<FWorldPartitionStreamingQuerySource>& QuerySources, bool bExactState) const
 {
-	const UDataLayerSubsystem* DataLayerSubsystem = WorldPartition->GetWorld()->GetSubsystem<UDataLayerSubsystem>();
+	const UDataLayerSubsystem* DataLayerSubsystem = GetWorld()->GetSubsystem<UDataLayerSubsystem>();
 	const bool bIsHLODEnabled = UHLODSubsystem::IsHLODEnabled();
 
 	bool bResult = true;
