[NO NON-PII COMMIT MESSAGE]

diff --git a/Engine/Source/Runtime/Engine/Classes/Components/SkinnedMeshComponent.h b/Engine/Source/Runtime/Engine/Classes/Components/SkinnedMeshComponent.h
--- a/Engine/Source/Runtime/Engine/Classes/Components/SkinnedMeshComponent.h
+++ b/Engine/Source/Runtime/Engine/Classes/Components/SkinnedMeshComponent.h
@@ -208,12 +208,14 @@ struct ENGINE_API FSkelMeshComponentLODInfo
 
 	void ReleaseOverrideVertexColorsAndBlock();
 	void BeginReleaseOverrideVertexColors();
+	void EndReleaseOverrideVertexColors();
 private:
 	void CleanUpOverrideVertexColors();
 
 public:
 	void ReleaseOverrideSkinWeightsAndBlock();
 	void BeginReleaseOverrideSkinWeights();
+	void EndReleaseOverrideSkinWeights();
 private:
 	void CleanUpOverrideSkinWeights();
 };
diff --git a/Engine/Source/Runtime/Engine/Private/Components/SkinnedMeshComponent.cpp b/Engine/Source/Runtime/Engine/Private/Components/SkinnedMeshComponent.cpp
--- a/Engine/Source/Runtime/Engine/Private/Components/SkinnedMeshComponent.cpp
+++ b/Engine/Source/Runtime/Engine/Private/Components/SkinnedMeshComponent.cpp
@@ -1152,6 +1152,25 @@ void USkinnedMeshComponent::InitLODInfos()
 	{
 		if (GetSkinnedAsset()->GetLODNum() != LODInfo.Num())
 		{
+			// Perform cleanup if LOD infos have been initialized before 
+			if (!LODInfo.IsEmpty())
+			{
+				// Batch release all resources of LOD infos we're about to destruct.
+				// This is relevant when overrides have been set but LODInfos are
+				// re-initialized, for example by changing the mesh at runtime.
+				for (int32 Idx = 0; Idx < LODInfo.Num(); ++Idx)
+				{
+					LODInfo[Idx].BeginReleaseOverrideSkinWeights();
+					LODInfo[Idx].BeginReleaseOverrideVertexColors();
+				}
+				FlushRenderingCommands();
+				for (int32 Idx = 0; Idx < LODInfo.Num(); ++Idx)
+				{
+					LODInfo[Idx].EndReleaseOverrideSkinWeights();
+					LODInfo[Idx].EndReleaseOverrideVertexColors();
+				}
+			}
+			
 			LODInfo.Empty(GetSkinnedAsset()->GetLODNum());
 			for (int32 Idx=0; Idx < GetSkinnedAsset()->GetLODNum(); Idx++)
 			{
@@ -3919,6 +3938,11 @@ void FSkelMeshComponentLODInfo::BeginReleaseOverrideVertexColors()
 	}
 }
 
+void FSkelMeshComponentLODInfo::EndReleaseOverrideVertexColors()
+{
+	CleanUpOverrideVertexColors();
+}
+
 void FSkelMeshComponentLODInfo::CleanUpOverrideVertexColors()
 {
 	if (OverrideVertexColors)
@@ -3950,6 +3974,11 @@ void FSkelMeshComponentLODInfo::BeginReleaseOverrideSkinWeights()
 	}
 }
 
+void FSkelMeshComponentLODInfo::EndReleaseOverrideSkinWeights()
+{
+	CleanUpOverrideSkinWeights();
+}
+
 void FSkelMeshComponentLODInfo::CleanUpOverrideSkinWeights()
 {
 	if (OverrideSkinWeights)
