Add more info which shader compile job succeeded/failed in SCW report.

#rnx

diff --git a/Engine/Source/Runtime/Engine/Private/ShaderCompiler/ShaderCompiler.cpp b/Engine/Source/Runtime/Engine/Private/ShaderCompiler/ShaderCompiler.cpp
--- a/Engine/Source/Runtime/Engine/Private/ShaderCompiler/ShaderCompiler.cpp
+++ b/Engine/Source/Runtime/Engine/Private/ShaderCompiler/ShaderCompiler.cpp
@@ -1733,8 +1733,16 @@ static FString GetSingleJobCompilationDump(const FShaderCompileJob* SingleJob)
 	return String;
 }
 
+static const TCHAR* GetCompileJobSuccessText(FShaderCompileJob* SingleJob)
+{
+	if (SingleJob)
+	{
+		return SingleJob->Output.bSucceeded ? TEXT("Succeeded") : TEXT("Failed");
+	}
+	return TEXT("");
+}
 
-static void DumpCompilationJobs(const TArray<FShaderCommonCompileJobPtr>& QueuedJobs, int32 NumProcessedJobs)
+static void LogQueuedCompileJobs(const TArray<FShaderCommonCompileJobPtr>& QueuedJobs, int32 NumProcessedJobs)
 {
 	if (NumProcessedJobs == -1)
 	{
@@ -1749,7 +1757,7 @@ static void DumpCompilationJobs(const TArray<FShaderCommonCompileJobPtr>& Queued
 	{
 		if (FShaderCompileJob* SingleJob = QueuedJobs[Index]->GetSingleShaderJob())
 		{
-			UE_LOG(LogShaderCompilers, Error, TEXT("Job %d [Single] %s"), Index, *GetSingleJobCompilationDump(SingleJob));
+			UE_LOG(LogShaderCompilers, Error, TEXT("Job %d [Single] %s: %s"), Index, GetCompileJobSuccessText(SingleJob), *GetSingleJobCompilationDump(SingleJob));
 		}
 		else
 		{
@@ -1757,7 +1765,8 @@ static void DumpCompilationJobs(const TArray<FShaderCommonCompileJobPtr>& Queued
 			UE_LOG(LogShaderCompilers, Error, TEXT("Job %d: Pipeline %s "), Index, PipelineJob->Key.ShaderPipeline->GetName());
 			for (int32 JobIndex = 0; JobIndex < PipelineJob->StageJobs.Num(); ++JobIndex)
 			{
-				UE_LOG(LogShaderCompilers, Error, TEXT("PipelineJob %d %s"), JobIndex, *GetSingleJobCompilationDump(PipelineJob->StageJobs[JobIndex]->GetSingleShaderJob()));
+				FShaderCompileJob* StageJob = PipelineJob->StageJobs[JobIndex]->GetSingleShaderJob();
+				UE_LOG(LogShaderCompilers, Error, TEXT("PipelineJob %d %s: %s"), JobIndex, GetCompileJobSuccessText(StageJob), *GetSingleJobCompilationDump(StageJob));
 			}
 		}
 	}
@@ -1832,7 +1841,7 @@ static void HandleWorkerCrash(const TArray<FShaderCommonCompileJobPtr>& QueuedJo
 	default:
 	case ESCWErrorCode::GeneralCrash:
 	{
-		DumpCompilationJobs(QueuedJobs, NumProcessedJobs);
+		LogQueuedCompileJobs(QueuedJobs, NumProcessedJobs);
 		SCWErrorCode::HandleGeneralCrash(ExceptionInfo.GetData(), Callstack.GetData());
 	}
 	break;
@@ -1861,7 +1870,7 @@ static void HandleWorkerCrash(const TArray<FShaderCommonCompileJobPtr>& QueuedJo
 		SCWErrorCode::HandleCantCompileForSpecificFormat(ExceptionInfo.GetData());
 		break;
 	case ESCWErrorCode::CrashInsidePlatformCompiler:
-		DumpCompilationJobs(QueuedJobs, NumProcessedJobs);
+		LogQueuedCompileJobs(QueuedJobs, NumProcessedJobs);
 		SCWErrorCode::HandleCrashInsidePlatformCompiler(ExceptionInfo.GetData());
 		break;
 	case ESCWErrorCode::Success:
@@ -2621,7 +2630,7 @@ bool FShaderCompileThreadRunnable::LaunchWorkersIfNeeded()
 					else
 					{
 						UE_LOG(LogShaderCompilers, Error, TEXT("ShaderCompileWorker terminated unexpectedly!  Falling back to directly compiling which will be very slow.  Thread %u."), WorkerIndex);
-						DumpCompilationJobs(CurrentWorkerInfo.QueuedJobs, -1);
+						LogQueuedCompileJobs(CurrentWorkerInfo.QueuedJobs, -1);
 
 						bAbandonWorkers = true;
 						break;
@@ -4339,7 +4348,7 @@ void FShaderCompilingManager::BlockOnShaderMapCompletion(const TArray<int32>& Sh
 		int32 NumActiveWorkers = 0;
 		do 
 		{
-			for (const auto& Thread : Threads)
+			for (const TUniquePtr<FShaderCompileThreadRunnableBase>& Thread : Threads)
 			{
 				NumActiveWorkers = Thread->CompilingLoop();
 			}
@@ -4453,7 +4462,7 @@ void FShaderCompilingManager::BlockOnAllShaderMapCompletion(TMap<int32, FShaderM
 		int32 NumActiveWorkers = 0;
 		do 
 		{
-			for (const auto& Thread : Threads)
+			for (const TUniquePtr<FShaderCompileThreadRunnableBase>& Thread : Threads)
 			{
 				NumActiveWorkers = Thread->CompilingLoop();
 			}
