- Fix incorrect GetBoneWeight function

#rnx

diff --git a/Engine/Source/Runtime/Engine/Private/Components/SkinnedMeshComponent.cpp b/Engine/Source/Runtime/Engine/Private/Components/SkinnedMeshComponent.cpp
--- a/Engine/Source/Runtime/Engine/Private/Components/SkinnedMeshComponent.cpp
+++ b/Engine/Source/Runtime/Engine/Private/Components/SkinnedMeshComponent.cpp
@@ -4669,7 +4669,7 @@ void GetTypedSkinnedTangentBasis(
 #endif
 	{
 		const int32 MeshBoneIndex = Section.BoneMap[SkinWeightVertexBuffer.GetBoneIndex(BufferVertIndex, InfluenceIndex)];
-		const float	Weight = (float)SkinWeightVertexBuffer.GetBoneWeight(BufferVertIndex, InfluenceIndex) / 255.0f;
+		const float	Weight = (float)SkinWeightVertexBuffer.GetBoneWeight(BufferVertIndex, InfluenceIndex) * UE::AnimationCore::InvMaxRawBoneWeightFloat;
 		const FMatrix44f& RefToLocal = RefToLocals[MeshBoneIndex];
 		OutTangentX += RefToLocal.TransformVector(VertexTangentX) * Weight;
 		OutTangentY += RefToLocal.TransformVector(VertexTangentY) * Weight;
@@ -4720,7 +4720,7 @@ FVector3f GetTypedSkinnedVertexPosition(
 			TransformBoneIndex = LeaderBoneMap[MeshBoneIndex];
 		}
 
-		const float	Weight = (float)SkinWeightVertexBuffer.GetBoneWeight(BufferVertIndex, InfluenceIndex) / 255.0f;
+		const float	Weight = (float)SkinWeightVertexBuffer.GetBoneWeight(BufferVertIndex, InfluenceIndex) * UE::AnimationCore::InvMaxRawBoneWeightFloat;
 		{
 			if (bCachedMatrices)
 			{
diff --git a/Engine/Source/Runtime/Engine/Private/SkeletalMesh.cpp b/Engine/Source/Runtime/Engine/Private/SkeletalMesh.cpp
--- a/Engine/Source/Runtime/Engine/Private/SkeletalMesh.cpp
+++ b/Engine/Source/Runtime/Engine/Private/SkeletalMesh.cpp
@@ -107,6 +107,7 @@
 #include "RayTracingInstance.h"
 #endif
 #include "Animation/SkinWeightProfileManager.h"
+#include "BoneWeights.h"
 
 #define LOCTEXT_NAMESPACE "SkeltalMesh"
 
@@ -6980,7 +6981,7 @@ FVector GetRefVertexLocationTyped(
 	for (int32 InfluenceIndex = 0; InfluenceIndex < MaxBoneInfluences; InfluenceIndex++)
 #endif
 	{
-		const float	Weight = (float)SkinWeightVertexBuffer.GetBoneWeight(BufferVertIndex, InfluenceIndex) / 255.0f;
+		const float	Weight = (float)SkinWeightVertexBuffer.GetBoneWeight(BufferVertIndex, InfluenceIndex) * UE::AnimationCore::InvMaxRawBoneWeightFloat;
 		{
 			const FMatrix BoneTransformMatrix = FMatrix::Identity;
 			SkinnedPos += BoneTransformMatrix.TransformPosition((FVector)PositionBuffer.VertexPosition(BufferVertIndex)) * Weight;
@@ -7021,7 +7022,7 @@ void GetRefTangentBasisTyped(const USkeletalMesh* Mesh, const FSkelMeshRenderSec
 	for (int32 InfluenceIndex = 0; InfluenceIndex < MaxBoneInfluences; InfluenceIndex++)
 #endif
 	{
-		const float	Weight = (float)SkinWeightVertexBuffer.GetBoneWeight(BufferVertIndex, InfluenceIndex) / 255.0f;
+		const float	Weight = (float)SkinWeightVertexBuffer.GetBoneWeight(BufferVertIndex, InfluenceIndex) * UE::AnimationCore::InvMaxRawBoneWeightFloat;
 		const FMatrix44f BoneTransformMatrix = FMatrix44f::Identity;
 		OutTangentX += BoneTransformMatrix.TransformVector(VertexTangentX) * Weight;
 		OutTangentY += BoneTransformMatrix.TransformVector(VertexTangentY) * Weight;
