Check for zero-length queries in ConvertTraceResults

diff --git a/Engine/Source/Runtime/Engine/Private/Collision/CollisionConversions.cpp b/Engine/Source/Runtime/Engine/Private/Collision/CollisionConversions.cpp
--- a/Engine/Source/Runtime/Engine/Private/Collision/CollisionConversions.cpp
+++ b/Engine/Source/Runtime/Engine/Private/Collision/CollisionConversions.cpp
@@ -408,6 +408,10 @@ EConvertQueryResult ConvertTraceResults(bool& OutHasValidBlockingHit, const UWor
 	EConvertQueryResult ConvertResult = EConvertQueryResult::Valid;
 	bool bHadBlockingHit = false;
 	const FVector Dir = (EndLoc - StartLoc).GetSafeNormal();
+	if (Dir.IsNearlyZero())
+	{
+		return EConvertQueryResult::Invalid;
+	}
 
 	for(int32 i=0; i<NumHits; i++)
 	{
@@ -447,6 +451,11 @@ template <typename Hit>
 EConvertQueryResult ConvertTraceResults(bool& OutHasValidBlockingHit, const UWorld* World, int32 NumHits, Hit* Hits, float CheckLength, const FCollisionFilterData& QueryFilter, FHitResult& OutHit, const FVector& StartLoc, const FVector& EndLoc, const FPhysicsGeometry& Geom, const FTransform& QueryTM, float MaxDistance, bool bReturnFaceIndex, bool bReturnPhysMat)
 {
 	const FVector Dir = (EndLoc - StartLoc).GetSafeNormal();
+	if (Dir.IsNearlyZero())
+	{
+		return EConvertQueryResult::Invalid;
+	}
+
 	if (TIsSame<Hit, FHitSweep>::Value)
 	{
 		if (!HadInitialOverlap(Hits[0]))
