Removed some redundant UObject scope cycle's in the ReplicateActor path

diff --git a/Engine/Source/Runtime/Engine/Private/DataChannel.cpp b/Engine/Source/Runtime/Engine/Private/DataChannel.cpp
--- a/Engine/Source/Runtime/Engine/Private/DataChannel.cpp
+++ b/Engine/Source/Runtime/Engine/Private/DataChannel.cpp
@@ -4693,7 +4693,6 @@ FObjectReplicator & UActorChannel::GetActorReplicationData()
 TSharedRef<FObjectReplicator>* UActorChannel::FindReplicator(UObject* Obj, bool* bOutFoundInvalid)
 {
 	CONDITIONAL_SCOPE_CYCLE_COUNTER(Stat_ActorChanFindOrCreateRep, CVarNetEnableDetailedScopeCounters.GetValueOnAnyThread() > 0);
-	SCOPE_CYCLE_UOBJECT(ActorChannelFindOrCreateRep, Obj);
 
 	// First, try to find it on the channel replication map
 	TSharedRef<FObjectReplicator>* ReplicatorRefPtr = ReplicationMap.Find( Obj );
@@ -4717,7 +4716,6 @@ TSharedRef<FObjectReplicator>* UActorChannel::FindReplicator(UObject* Obj, bool*
 TSharedRef<FObjectReplicator>& UActorChannel::CreateReplicator(UObject* Obj, bool bCheckDormantReplicators)
 {
 	CONDITIONAL_SCOPE_CYCLE_COUNTER(Stat_ActorChanFindOrCreateRep, CVarNetEnableDetailedScopeCounters.GetValueOnAnyThread() > 0);
-	SCOPE_CYCLE_UOBJECT(ActorChannelFindOrCreateRep, Obj);
 
 	// Try to find in the dormancy map if desired
 	TSharedPtr<FObjectReplicator> NewReplicator;
diff --git a/Engine/Source/Runtime/Engine/Private/RepLayout.cpp b/Engine/Source/Runtime/Engine/Private/RepLayout.cpp
--- a/Engine/Source/Runtime/Engine/Private/RepLayout.cpp
+++ b/Engine/Source/Runtime/Engine/Private/RepLayout.cpp
@@ -31,7 +31,6 @@
 #endif // UE_WITH_IRIS
 
 DECLARE_CYCLE_STAT(TEXT("RepLayout InitFromObjectClass"), STAT_RepLayout_InitFromObjectClass, STATGROUP_Game);
-DECLARE_CYCLE_STAT(TEXT("RepLayout BuildShadowOffsets"), STAT_RepLayout_BuildShadowOffsets, STATGROUP_Game);
 DECLARE_CYCLE_STAT(TEXT("RepLayout DeltaSerializeFastArray"), STAT_RepLayout_DeltaSerializeFastArray, STATGROUP_Game);
 
 // LogRepProperties is very spammy, and the logs are in a very hot code path,
@@ -5900,8 +5899,6 @@ static void BuildShadowOffsets(
 	TArray<FRepLayoutCmd>& Cmds,
 	int32& ShadowOffset)
 {
-	SCOPE_CYCLE_COUNTER(STAT_RepLayout_BuildShadowOffsets);
-
 	if (ShadowType == ERepBuildType::Class && !!GUsePackedShadowBuffers)
 	{
 		ShadowOffset = 0;
@@ -6019,7 +6016,6 @@ void FRepLayout::InitFromClass(
 	const ECreateRepLayoutFlags CreateFlags)
 {
 	SCOPE_CYCLE_COUNTER(STAT_RepLayout_InitFromObjectClass);
-	SCOPE_CYCLE_UOBJECT(ObjectClass, InObjectClass);
 
 	const bool bIsPushModelEnabled = IS_PUSH_MODEL_ENABLED();
 	const bool bIsObjectActor = InObjectClass->IsChildOf(AActor::StaticClass());
