Add support for PSO precaching for SkinnedAssets if IsResourcePSOPrecachingEnabled is set

#rnx

diff --git a/Engine/Source/Runtime/Engine/Classes/Engine/SkinnedAsset.h b/Engine/Source/Runtime/Engine/Classes/Engine/SkinnedAsset.h
--- a/Engine/Source/Runtime/Engine/Classes/Engine/SkinnedAsset.h
+++ b/Engine/Source/Runtime/Engine/Classes/Engine/SkinnedAsset.h
@@ -34,6 +34,11 @@ enum class ESkinnedAssetAsyncPropertyLockType
 };
 ENUM_CLASS_FLAGS(ESkinnedAssetAsyncPropertyLockType);
 
+struct FSkinnedAssetVertexFactoryTypesPerMaterialData
+{
+	int16 MaterialIndex;
+	TArray<const FVertexFactoryType*, TInlineAllocator<2>> VertexFactoryTypes;
+};
 
 UCLASS(hidecategories = Object, config = Engine, editinlinenew, abstract)
 class ENGINE_API USkinnedAsset : public UStreamableRenderAsset, public IInterface_AsyncCompilation
@@ -230,6 +235,8 @@ public:
 
 	virtual void PostLoad() override;
 	//~ End UObject Interface
+		
+	TArray<FSkinnedAssetVertexFactoryTypesPerMaterialData, TInlineAllocator<4>> GetVertexFactoryTypesPerMaterialIndex(int32 MinLODIndex, bool bCPUSkin, ERHIFeatureLevel::Type FeatureLevel);
 
 #if WITH_EDITOR
 	/** IInterface_AsyncCompilation begin*/
diff --git a/Engine/Source/Runtime/Engine/Private/Components/SkinnedMeshComponent.cpp b/Engine/Source/Runtime/Engine/Private/Components/SkinnedMeshComponent.cpp
--- a/Engine/Source/Runtime/Engine/Private/Components/SkinnedMeshComponent.cpp
+++ b/Engine/Source/Runtime/Engine/Private/Components/SkinnedMeshComponent.cpp
@@ -553,8 +553,7 @@ void USkinnedMeshComponent::PostLoad()
 
 void USkinnedMeshComponent::PrecachePSOs()
 {
-	if (!FApp::CanEverRender() ||
-		!PipelineStateCache::IsPSOPrecachingEnabled() ||
+	if (!IsComponentPSOPrecachingEnabled() ||
 		GetSkinnedAsset() == nullptr ||
 		GetSkinnedAsset()->GetResourceForRendering() == nullptr)
 	{
@@ -562,66 +561,17 @@ void USkinnedMeshComponent::PrecachePSOs()
 	}
 
 	ERHIFeatureLevel::Type FeatureLevel = GetWorld() ? GetWorld()->FeatureLevel.GetValue() : GMaxRHIFeatureLevel;
-	FSkeletalMeshRenderData* SkelMeshRenderData = GetSkinnedAsset()->GetResourceForRendering();
 	int32 MinLODIndex = ComputeMinLOD();
+	bool bCPUSkin = bRenderStatic || ShouldCPUSkin();
 
-	bool bAnySectionCastsShadows = false;
-
-	struct VFsPerMaterialData
-	{
-		int16 MaterialIndex;
-		TArray<const FVertexFactoryType*, TInlineAllocator<2>> VertexFactoryTypes;
-	};
-	TArray<VFsPerMaterialData, TInlineAllocator<4>> VFsPerMaterials;
-
-	for (int32 LODIndex = MinLODIndex; LODIndex < SkelMeshRenderData->LODRenderData.Num(); LODIndex++)
-	{
-		FSkeletalMeshLODRenderData& LODRenderData = SkelMeshRenderData->LODRenderData[LODIndex];
-		const FSkeletalMeshLODInfo& Info = *(GetSkinnedAsset()->GetLODInfo(LODIndex));
-		
-		// Check all render sections for the used material indices
-		for (int32 SectionIndex = 0; SectionIndex < LODRenderData.RenderSections.Num(); SectionIndex++)
-		{
-			FSkelMeshRenderSection& RenderSection = LODRenderData.RenderSections[SectionIndex];
-
-			bAnySectionCastsShadows |= RenderSection.bCastShadow;
-
-			// By default use the material index of the render section
-			uint16 MaterialIndex = RenderSection.MaterialIndex;
-
-			// Can be remapped by the LODInfo
-			if (SectionIndex < Info.LODMaterialMap.Num() && GetSkinnedAsset()->IsValidMaterialIndex(Info.LODMaterialMap[SectionIndex]))
-			{
-				MaterialIndex = Info.LODMaterialMap[SectionIndex];
-				MaterialIndex = FMath::Clamp(MaterialIndex, 0, GetSkinnedAsset()->GetNumMaterials());
-			}
-
-			VFsPerMaterialData* VFsPerMaterial = VFsPerMaterials.FindByPredicate([MaterialIndex](const VFsPerMaterialData& Other) { return Other.MaterialIndex == MaterialIndex; });
-			if (VFsPerMaterial == nullptr)
-			{
-				VFsPerMaterial = &VFsPerMaterials.AddDefaulted_GetRef();
-				VFsPerMaterial->MaterialIndex = MaterialIndex;
-			}
-
-			// Find all the possible used vertex factory types needed to render this render section
-			if (bRenderStatic || ShouldCPUSkin())
-			{
-				// Force static from GPU point of view
-				VFsPerMaterial->VertexFactoryTypes.AddUnique(&FLocalVertexFactory::StaticType);
-			}
-			else if (!SkelMeshRenderData->RequiresCPUSkinning(FeatureLevel, MinLODIndex))
-			{
-				// Add all the vertex factories which can be used for gpu skinning
-				FSkeletalMeshObjectGPUSkin::GetUsedVertexFactories(LODRenderData, RenderSection, FeatureLevel, VFsPerMaterial->VertexFactoryTypes);
-			}
-		}
-	}
+	TArray<FSkinnedAssetVertexFactoryTypesPerMaterialData, TInlineAllocator<4>> VFsPerMaterials = GetSkinnedAsset()->GetVertexFactoryTypesPerMaterialIndex(MinLODIndex, bCPUSkin, FeatureLevel);
+	bool bAnySectionCastsShadows = GetSkinnedAsset()->GetResourceForRendering()->AnyRenderSectionCastsShadows(MinLODIndex);
 
 	FPSOPrecacheParams PrecachePSOParams;
 	SetupPrecachePSOParams(PrecachePSOParams);
 	PrecachePSOParams.bCastShadow = PrecachePSOParams.bCastShadow && bAnySectionCastsShadows;
 
-	for (VFsPerMaterialData& VFsPerMaterial : VFsPerMaterials)
+	for (FSkinnedAssetVertexFactoryTypesPerMaterialData& VFsPerMaterial : VFsPerMaterials)
 	{
 		UMaterialInterface* MaterialInterface = GetMaterial(VFsPerMaterial.MaterialIndex);
 		if (MaterialInterface)
diff --git a/Engine/Source/Runtime/Engine/Private/SkeletalMeshRenderData.cpp b/Engine/Source/Runtime/Engine/Private/SkeletalMeshRenderData.cpp
--- a/Engine/Source/Runtime/Engine/Private/SkeletalMeshRenderData.cpp
+++ b/Engine/Source/Runtime/Engine/Private/SkeletalMeshRenderData.cpp
@@ -734,6 +734,22 @@ int32 FSkeletalMeshRenderData::GetMaxBonesPerSection(int32 MinLODIdx) const
 	return MaxBonesPerSection;
 }
 
+bool FSkeletalMeshRenderData::AnyRenderSectionCastsShadows(int32 MinLODIdx) const
+{
+	for (int32 LODIndex = MinLODIdx; LODIndex < LODRenderData.Num(); LODIndex++)
+	{
+		for (const FSkelMeshRenderSection& RenderSection : LODRenderData[LODIndex].RenderSections)
+		{
+			if (RenderSection.bCastShadow)
+			{
+				return true;
+			}
+		}
+	}
+
+	return false;
+}
+
 int32 FSkeletalMeshRenderData::GetMaxBonesPerSection() const
 {
 	return GetMaxBonesPerSection(0);
diff --git a/Engine/Source/Runtime/Engine/Private/SkinnedAsset.cpp b/Engine/Source/Runtime/Engine/Private/SkinnedAsset.cpp
--- a/Engine/Source/Runtime/Engine/Private/SkinnedAsset.cpp
+++ b/Engine/Source/Runtime/Engine/Private/SkinnedAsset.cpp
@@ -3,6 +3,8 @@
 #include "Engine/SkinnedAsset.h"
 #include "SkinnedAssetCompiler.h"
 #include "Animation/AnimStats.h"
+#include "SkeletalRenderGPUSkin.h"
+#include "PSOPrecache.h"
 
 #if INTEL_ISPC
 #include "SkinnedAsset.ispc.generated.h"
@@ -52,6 +54,84 @@ void USkinnedAsset::PostLoad()
 		ExecutePostLoadInternal(Context);
 		FinishPostLoadInternal(Context);
 	}
+
+	if (IsResourcePSOPrecachingEnabled() &&
+		GetResourceForRendering() != nullptr)
+	{
+		// Assume GPU skinning when precaching PSOs
+		bool bCPUSkin = false;
+		ERHIFeatureLevel::Type FeatureLevel = GetWorld() ? GetWorld()->FeatureLevel.GetValue() : GMaxRHIFeatureLevel;
+		int32 MinLODIndex = GetMinLodIdx();
+		
+		TArray<FSkinnedAssetVertexFactoryTypesPerMaterialData, TInlineAllocator<4>> VFsPerMaterials = GetVertexFactoryTypesPerMaterialIndex(MinLODIndex, bCPUSkin, FeatureLevel);
+		bool bAnySectionCastsShadows = GetResourceForRendering()->AnyRenderSectionCastsShadows(MinLODIndex);
+
+		// Precache using default material & PSO precache params but mark movable
+		const TArray<FSkeletalMaterial>& Materials = GetMaterials();
+		FPSOPrecacheParams PrecachePSOParams;
+		PrecachePSOParams.SetMobility(EComponentMobility::Movable);
+		PrecachePSOParams.bCastShadow = bAnySectionCastsShadows;
+
+		for (auto VFsPerMaterial : VFsPerMaterials)
+		{
+			UMaterialInterface* MaterialInterface = Materials[VFsPerMaterial.MaterialIndex].MaterialInterface;
+			if (MaterialInterface)
+			{
+				MaterialInterface->PrecachePSOs(VFsPerMaterial.VertexFactoryTypes, PrecachePSOParams);
+			}
+		}
+	}
+}
+
+TArray<FSkinnedAssetVertexFactoryTypesPerMaterialData, TInlineAllocator<4>> USkinnedAsset::GetVertexFactoryTypesPerMaterialIndex(
+	int32 MinLODIndex, bool bCPUSkin, ERHIFeatureLevel::Type FeatureLevel)
+{
+	TArray<FSkinnedAssetVertexFactoryTypesPerMaterialData, TInlineAllocator<4>> VFTypesPerMaterialIndex;
+
+	FSkeletalMeshRenderData* SkelMeshRenderData = GetResourceForRendering();
+	for (int32 LODIndex = MinLODIndex; LODIndex < SkelMeshRenderData->LODRenderData.Num(); LODIndex++)
+	{
+		FSkeletalMeshLODRenderData& LODRenderData = SkelMeshRenderData->LODRenderData[LODIndex];
+		const FSkeletalMeshLODInfo& Info = *(GetLODInfo(LODIndex));
+
+		// Check all render sections for the used material indices
+		for (int32 SectionIndex = 0; SectionIndex < LODRenderData.RenderSections.Num(); SectionIndex++)
+		{
+			FSkelMeshRenderSection& RenderSection = LODRenderData.RenderSections[SectionIndex];
+
+			// By default use the material index of the render section
+			uint16 MaterialIndex = RenderSection.MaterialIndex;
+
+			// Can be remapped by the LODInfo
+			if (SectionIndex < Info.LODMaterialMap.Num() && IsValidMaterialIndex(Info.LODMaterialMap[SectionIndex]))
+			{
+				MaterialIndex = Info.LODMaterialMap[SectionIndex];
+				MaterialIndex = FMath::Clamp(MaterialIndex, 0, GetNumMaterials());
+			}
+
+			FSkinnedAssetVertexFactoryTypesPerMaterialData* VFsPerMaterial = VFTypesPerMaterialIndex.FindByPredicate(
+				[MaterialIndex](const FSkinnedAssetVertexFactoryTypesPerMaterialData& Other) { return Other.MaterialIndex == MaterialIndex; });
+			if (VFsPerMaterial == nullptr)
+			{
+				VFsPerMaterial = &VFTypesPerMaterialIndex.AddDefaulted_GetRef();
+				VFsPerMaterial->MaterialIndex = MaterialIndex;
+			}
+
+			// Find all the possible used vertex factory types needed to render this render section
+			if (bCPUSkin)
+			{
+				// Force static from GPU point of view
+				VFsPerMaterial->VertexFactoryTypes.AddUnique(&FLocalVertexFactory::StaticType);
+			}
+			else if (!SkelMeshRenderData->RequiresCPUSkinning(FeatureLevel, LODIndex))
+			{
+				// Add all the vertex factories which can be used for gpu skinning
+				FSkeletalMeshObjectGPUSkin::GetUsedVertexFactories(LODRenderData, RenderSection, FeatureLevel, VFsPerMaterial->VertexFactoryTypes);
+			}
+		}
+	}
+
+	return VFTypesPerMaterialIndex;
 }
 
 #if WITH_EDITOR
diff --git a/Engine/Source/Runtime/Engine/Public/Rendering/SkeletalMeshRenderData.h b/Engine/Source/Runtime/Engine/Public/Rendering/SkeletalMeshRenderData.h
--- a/Engine/Source/Runtime/Engine/Public/Rendering/SkeletalMeshRenderData.h
+++ b/Engine/Source/Runtime/Engine/Public/Rendering/SkeletalMeshRenderData.h
@@ -107,6 +107,9 @@ public:
 		return GetFirstValidLODIdx(FMath::Max<int32>(PendingFirstLODIdx, MinLODIdx));
 	}
 
+	/** Check if any rendersection casts shadows */
+	ENGINE_API bool AnyRenderSectionCastsShadows(int32 MinLODIdx) const;
+
 	/** 
 	 * Return the pending first LOD that can be used for rendering starting at MinLODIdx.
 	 * This takes into account the streaming status from PendingFirstLODIdx, 
