Fixed world partition Map Check not validating level instances.

diff --git a/Engine/Plugins/Editor/DataValidation/Source/DataValidation/Private/WorldPartitionChangelistValidator.cpp b/Engine/Plugins/Editor/DataValidation/Source/DataValidation/Private/WorldPartitionChangelistValidator.cpp
--- a/Engine/Plugins/Editor/DataValidation/Source/DataValidation/Private/WorldPartitionChangelistValidator.cpp
+++ b/Engine/Plugins/Editor/DataValidation/Source/DataValidation/Private/WorldPartitionChangelistValidator.cpp
@@ -206,7 +206,8 @@ EDataValidationResult UWorldPartitionChangelistValidator::ValidateActorsAndDataL
 
 		// Invoke static WorldPartition Validation from the ActorDescContainer
 		const bool bIsStreamingDisabled = ULevel::GetIsStreamingDisabledFromPackage(MapPath.GetPackageName());
-		UWorldPartition::CheckForErrors(this, ActorDescContainer, !bIsStreamingDisabled);
+		const bool bIsChangelistValidation = true;
+		UWorldPartition::CheckForErrors(this, ActorDescContainer, !bIsStreamingDisabled, bIsChangelistValidation);
 	}
 
 	if (Errors->Num())
diff --git a/Engine/Source/Runtime/Engine/Private/WorldPartition/WorldPartitionStreamingGeneration.cpp b/Engine/Source/Runtime/Engine/Private/WorldPartition/WorldPartitionStreamingGeneration.cpp
--- a/Engine/Source/Runtime/Engine/Private/WorldPartition/WorldPartitionStreamingGeneration.cpp
+++ b/Engine/Source/Runtime/Engine/Private/WorldPartition/WorldPartitionStreamingGeneration.cpp
@@ -303,10 +303,7 @@ class FWorldPartitionStreamingGenerator
 		{
 			if (InActorDescView.IsContainerInstance())
 			{
-				if (!bIsValidating)
-				{
-					OutContainerInstances.Add(InActorDescView);
-				}
+				OutContainerInstances.Add(InActorDescView);
 			}
 			else
 			{
@@ -404,8 +401,13 @@ class FWorldPartitionStreamingGenerator
 
 			if (!ContainerInstanceView.GetContainerInstance(SubContainer, SubTransform, SubClusterMode))
 			{
-				//@todo_ow: make a specific error for missing container instance sublevel?
-				ErrorHandler->OnInvalidReference(ContainerInstanceView, FGuid());
+				// Don't generate an error when validating changelist because container instances
+				// won't be registered only because the world is not loaded.
+				if (const bool bGenerateError = !bIsChangelistValidation)
+				{
+					//@todo_ow: make a specific error for missing container instance sublevel?
+					ErrorHandler->OnInvalidReference(ContainerInstanceView, FGuid());
+				}
 				continue;
 			}
 
@@ -747,9 +749,9 @@ class FWorldPartitionStreamingGenerator
 	}
 
 public:
-	FWorldPartitionStreamingGenerator(FActorDescList* InModifiedActorsDescList, IStreamingGenerationErrorHandler* InErrorHandler, bool bInEnableStreaming, bool bInIsValidating = false, TArray<TSubclassOf<AActor>> InFilteredClasses = {})
+	FWorldPartitionStreamingGenerator(FActorDescList* InModifiedActorsDescList, IStreamingGenerationErrorHandler* InErrorHandler, bool bInEnableStreaming, bool bInIsChangelistValidation = false, TArray<TSubclassOf<AActor>> InFilteredClasses = {})
 		: bEnableStreaming(bInEnableStreaming)
-		, bIsValidating(bInIsValidating)
+		, bIsChangelistValidation(bInIsChangelistValidation)
 		, ModifiedActorsDescList(InModifiedActorsDescList)
 		, FilteredClasses(InFilteredClasses)
 		, ErrorHandler(InErrorHandler ? InErrorHandler : &NullErrorHandler)	
@@ -857,7 +859,7 @@ public:
 
 private:
 	bool bEnableStreaming;
-	bool bIsValidating;
+	bool bIsChangelistValidation;
 	FActorDescList* ModifiedActorsDescList;
 	TArray<TSubclassOf<AActor>> FilteredClasses;
 	IStreamingGenerationErrorHandler* ErrorHandler;
@@ -940,8 +942,8 @@ void UWorldPartition::FlushStreaming()
 void UWorldPartition::GenerateHLOD(ISourceControlHelper* SourceControlHelper, bool bCreateActorsOnly)
 {
 	FStreamingGenerationLogErrorHandler LogErrorHandler;
-	const bool bIsValidating = false;
-	FWorldPartitionStreamingGenerator StreamingGenerator(nullptr, &LogErrorHandler, IsStreamingEnabled(), bIsValidating, { AWorldPartitionHLOD::StaticClass() });
+	const bool bInIsChangelistValidation = false;
+	FWorldPartitionStreamingGenerator StreamingGenerator(nullptr, &LogErrorHandler, IsStreamingEnabled(), bInIsChangelistValidation, { AWorldPartitionHLOD::StaticClass() });
 	StreamingGenerator.PreparationPhase(ActorDescContainer);
 
 	TUniquePtr<FArchive> LogFileAr = FWorldPartitionStreamingGenerator::CreateDumpStateLogArchive(TEXT("HLOD"));
@@ -953,14 +955,14 @@ void UWorldPartition::GenerateHLOD(ISourceControlHelper* SourceControlHelper, bo
 
 void UWorldPartition::CheckForErrors(IStreamingGenerationErrorHandler* ErrorHandler) const
 {
-	CheckForErrors(ErrorHandler, ActorDescContainer, IsStreamingEnabled());
+	const bool bIsChangelistValidation = false;
+	CheckForErrors(ErrorHandler, ActorDescContainer, IsStreamingEnabled(), bIsChangelistValidation);
 }
 
-void UWorldPartition::CheckForErrors(IStreamingGenerationErrorHandler* ErrorHandler, const UActorDescContainer* ActorDescContainer, bool bEnableStreaming)
+void UWorldPartition::CheckForErrors(IStreamingGenerationErrorHandler* ErrorHandler, const UActorDescContainer* ActorDescContainer, bool bEnableStreaming, bool bIsChangelistValidation)
 {
 	FActorDescList ModifiedActorDescList;
-	const bool bIsValidating = true;
-	FWorldPartitionStreamingGenerator StreamingGenerator(ActorDescContainer->GetWorld() ? &ModifiedActorDescList : nullptr, ErrorHandler, bEnableStreaming, bIsValidating);
+	FWorldPartitionStreamingGenerator StreamingGenerator(ActorDescContainer->GetWorld() ? &ModifiedActorDescList : nullptr, ErrorHandler, bEnableStreaming, bIsChangelistValidation);
 	StreamingGenerator.PreparationPhase(ActorDescContainer);
 }
 #endif // WITH_EDITOR
diff --git a/Engine/Source/Runtime/Engine/Public/WorldPartition/WorldPartition.h b/Engine/Source/Runtime/Engine/Public/WorldPartition/WorldPartition.h
--- a/Engine/Source/Runtime/Engine/Public/WorldPartition/WorldPartition.h
+++ b/Engine/Source/Runtime/Engine/Public/WorldPartition/WorldPartition.h
@@ -188,7 +188,7 @@ public:
 	void DumpActorDescs(const FString& Path);
 
 	void CheckForErrors(IStreamingGenerationErrorHandler* ErrorHandler) const;
-	static void CheckForErrors(IStreamingGenerationErrorHandler* ErrorHandler, const UActorDescContainer* ActorDescContainer, bool bEnableStreaming);
+	static void CheckForErrors(IStreamingGenerationErrorHandler* ErrorHandler, const UActorDescContainer* ActorDescContainer, bool bEnableStreaming, bool bIsChangelistValidation);
 
 	void AppendAssetRegistryTags(TArray<FAssetRegistryTag>& OutTags) const;
 
