Add support for PSO precaching from StaticMesh resources if IsResourcePSOPrecachingEnabled is set

#rnx

diff --git a/Engine/Source/Runtime/Engine/Private/StaticMesh.cpp b/Engine/Source/Runtime/Engine/Private/StaticMesh.cpp
--- a/Engine/Source/Runtime/Engine/Private/StaticMesh.cpp
+++ b/Engine/Source/Runtime/Engine/Private/StaticMesh.cpp
@@ -5681,6 +5681,40 @@ void UStaticMesh::PostLoad()
 		ExecutePostLoadInternal(Context);
 		FinishPostLoadInternal(Context);
 	}
+
+	if (IsResourcePSOPrecachingEnabled() &&
+		GetRenderData() != nullptr)
+	{
+		ERHIFeatureLevel::Type FeatureLevel = GetWorld() ? GetWorld()->FeatureLevel.GetValue() : GMaxRHIFeatureLevel;
+		EShaderPlatform ShaderPlatform = GetFeatureLevelShaderPlatform(FeatureLevel);
+		bool bUseNanite = UseNanite(ShaderPlatform) && HasValidNaniteData();
+
+		bool bAnySectionCastsShadows = false;
+		TArray<int16, TInlineAllocator<2>> UsedMaterialIndices;
+		for (FStaticMeshLODResources& LODRenderData : GetRenderData()->LODResources)
+		{
+			for (FStaticMeshSection& RenderSection : LODRenderData.Sections)
+			{
+				UsedMaterialIndices.AddUnique(RenderSection.MaterialIndex);
+				bAnySectionCastsShadows |= RenderSection.bCastShadow;
+			}
+		}
+
+		// Use default precache PSO params but take shadow casting into account and mark movable to have better coverage
+		FPSOPrecacheParams PrecachePSOParams;
+		PrecachePSOParams.bCastShadow = bAnySectionCastsShadows;
+		PrecachePSOParams.SetMobility(EComponentMobility::Movable);
+
+		const FVertexFactoryType* VFType = bUseNanite ? &Nanite::FVertexFactory::StaticType : &FLocalVertexFactory::StaticType;
+		for (uint16 MaterialIndex : UsedMaterialIndices)
+		{
+			UMaterialInterface* MaterialInterface = GetMaterial(MaterialIndex);
+			if (MaterialInterface)
+			{
+				MaterialInterface->PrecachePSOs(VFType, PrecachePSOParams);
+			}
+		}
+	}
 }
 
 void UStaticMesh::BeginPostLoadInternal(FStaticMeshPostLoadContext& Context)
