- Exposed world partition cell bounds used by world partition streaming in IWorldPartitionCell::GetCellBounds.
- Added documentation for IWorldPartitionCell interface.

diff --git a/Engine/Source/Runtime/Engine/Private/WorldPartition/WorldPartitionRuntimeSpatialHash.cpp b/Engine/Source/Runtime/Engine/Private/WorldPartition/WorldPartitionRuntimeSpatialHash.cpp
--- a/Engine/Source/Runtime/Engine/Private/WorldPartition/WorldPartitionRuntimeSpatialHash.cpp
+++ b/Engine/Source/Runtime/Engine/Private/WorldPartition/WorldPartitionRuntimeSpatialHash.cpp
@@ -459,12 +459,6 @@ void FSpatialHashStreamingGrid::Draw3D(UWorld* World, const TArray<FWorldPartiti
 				GetFilteredCellsForDebugDraw(LayerCell, DataLayerSubsystem, FilteredCells);
 				check(FilteredCells.Num());
 
-				FBox2D CellWorldBounds;
-				Helper.Levels[GridLevel].GetCellBounds(FGridCellCoord2(Coords.X, Coords.Y), CellWorldBounds);
-				double CellSizeY = CellWorldBounds.GetSize().Y / FilteredCells.Num();
-				CellWorldBounds.Max.Y = CellWorldBounds.Min.Y + CellSizeY;
-				FTranslationMatrix CellOffsetMatrix(FVector(0.f, CellSizeY, 0.f));
-
 				for (const UWorldPartitionRuntimeCell* Cell : FilteredCells)
 				{
 					bool bIsAlreadyInSet = false;
@@ -475,17 +469,13 @@ void FSpatialHashStreamingGrid::Draw3D(UWorld* World, const TArray<FWorldPartiti
 					}
 
 					// Use cell MinMaxZ to compute effective cell bounds
-					double CellZExtent = 0.5 * (Cell->GetMinMaxZ().Y - Cell->GetMinMaxZ().X);
-					double CellZPosition = Cell->GetMinMaxZ().X + CellZExtent;
-					FVector BoundsExtent(CellWorldBounds.GetExtent(), CellZExtent);
-					FVector BoundsOrigin(CellWorldBounds.GetCenter(), CellZPosition);
-					FBox CellBox = FBox::BuildAABB(BoundsOrigin, BoundsExtent);
+					FBox CellBox = Cell->GetCellBounds();
 
 					// Draw Cell using its debug color
 					FColor CellColor = Cell->GetDebugColor(VisualizeMode).ToFColor(false).WithAlpha(16);
 					DrawDebugSolidBox(World, CellBox, CellColor, Transform, false, -1.f, 255);
 					FVector CellPos = Transform.TransformPosition(CellBox.GetCenter());
-					DrawDebugBox(World, CellPos, BoundsExtent, Transform.GetRotation(), CellColor.WithAlpha(255), false, -1.f, 255, 10.f);
+					DrawDebugBox(World, CellPos, CellBox.GetExtent(), Transform.GetRotation(), CellColor.WithAlpha(255), false, -1.f, 255, 10.f);
 
 					// Draw Cell's DataLayers colored boxes
 					if (DataLayerDebugColors.Num() && Cell->GetDataLayers().Num() > 0)
@@ -501,7 +491,6 @@ void FSpatialHashStreamingGrid::Draw3D(UWorld* World, const TArray<FWorldPartiti
 							DataLayerColoredBox = DataLayerColoredBox.TransformBy(DataLayerOffsetMatrix);
 						}
 					}
-					CellBox = CellBox.TransformBy(CellOffsetMatrix);
 				}
 			});
 		}
@@ -1648,9 +1637,11 @@ bool UWorldPartitionRuntimeSpatialHash::Draw2D(UCanvas* Canvas, const TArray<FWo
 		int64 CellSize = (1LL << MinGridLevel) * (int64)StreamingGrid->CellSize;
 		const FVector MinExtent(CellSize * GridViewMinimumSizeInCellCount);
 		FBox Region(ForceInit);
+		// Keep rendered region around source location if loading range is zero
+		float LoadingRange = FMath::Max<float>(StreamingGrid->GetLoadingRange(), 1.f);
 		for (const FWorldPartitionStreamingSource& Source : Sources)
 		{
-			Region += Source.CalcBounds(StreamingGrid->GetLoadingRange(), StreamingGrid->GridName, StreamingGrid->HLODLayer, /*bCalcIn2D*/ true);
+			Region += Source.CalcBounds(LoadingRange, StreamingGrid->GridName, StreamingGrid->HLODLayer, /*bCalcIn2D*/ true);
 		}
 		Region += FBox(Region.GetCenter() - MinExtent, Region.GetCenter() + MinExtent);
 		const FVector2D GridReferenceWorldPos = FVector2D(Region.GetCenter());
diff --git a/Engine/Source/Runtime/Engine/Private/WorldPartition/WorldPartitionRuntimeSpatialHashCell.cpp b/Engine/Source/Runtime/Engine/Private/WorldPartition/WorldPartitionRuntimeSpatialHashCell.cpp
--- a/Engine/Source/Runtime/Engine/Private/WorldPartition/WorldPartitionRuntimeSpatialHashCell.cpp
+++ b/Engine/Source/Runtime/Engine/Private/WorldPartition/WorldPartitionRuntimeSpatialHashCell.cpp
@@ -112,6 +112,15 @@ void UWorldPartitionRuntimeSpatialHashCell::AppendStreamingSourceInfo(const FWor
 	CachedSourceModulatedDistances.Add(ModulatedSquareDistance);
 }
 
+FBox UWorldPartitionRuntimeSpatialHashCell::GetCellBounds() const
+{
+	FBox Box = FBox::BuildAABB(Position, FVector(Extent));
+	// Use content bounds for the Z extent
+	Box.Min.Z = GetContentBounds().Min.Z;
+	Box.Max.Z = GetContentBounds().Max.Z;
+	return Box;
+}
+
 void UWorldPartitionRuntimeSpatialHashCell::MergeStreamingSourceInfo() const
 {
 	Super::MergeStreamingSourceInfo();
diff --git a/Engine/Source/Runtime/Engine/Public/WorldPartition/WorldPartitionRuntimeCell.h b/Engine/Source/Runtime/Engine/Public/WorldPartition/WorldPartitionRuntimeCell.h
--- a/Engine/Source/Runtime/Engine/Public/WorldPartition/WorldPartitionRuntimeCell.h
+++ b/Engine/Source/Runtime/Engine/Public/WorldPartition/WorldPartitionRuntimeCell.h
@@ -196,6 +196,7 @@ class ENGINE_API UWorldPartitionRuntimeCell : public UObject, public IWorldParti
 		return Algo::AnyOf(DataLayers, [&InDataLayers](const FName& DataLayer) { return InDataLayers.Contains(DataLayer); });
 	}
 	virtual const FBox& GetContentBounds() const override { return ContentBounds; }
+	virtual FBox GetCellBounds() const override { return FBox(ForceInit); }
 	// IWorldPartitionCell Interface end
 
 	const FVector2D GetMinMaxZ() const { return FVector2D(ContentBounds.Min.Z, ContentBounds.Max.Z); }
diff --git a/Engine/Source/Runtime/Engine/Public/WorldPartition/WorldPartitionRuntimeCellInterface.h b/Engine/Source/Runtime/Engine/Public/WorldPartition/WorldPartitionRuntimeCellInterface.h
--- a/Engine/Source/Runtime/Engine/Public/WorldPartition/WorldPartitionRuntimeCellInterface.h
+++ b/Engine/Source/Runtime/Engine/Public/WorldPartition/WorldPartitionRuntimeCellInterface.h
@@ -9,12 +9,21 @@ class UDataLayerAsset;
 class IWorldPartitionCell
 {
 public:
+	/** Returns the cell content associated data layer instances. */
 	virtual TArray<const UDataLayerInstance*> GetDataLayerInstances() const = 0;
+	/** Returns whether the cell data layers referenced the provided data layer asset or not. */
 	virtual bool ContainsDataLayer(const UDataLayerAsset* DataLayerAsset) const = 0;
+	/** Returns whether the cell data layers referenced the provided data layer instance or not. */
 	virtual bool ContainsDataLayer(const UDataLayerInstance* DataLayerInstance) const = 0;
+	/** Returns whether the cell content is associated with data layers or not. */
 	virtual bool HasDataLayers() const = 0;
+	/** Returns the cell content associated data layers. */
 	virtual const TArray<FName>& GetDataLayers() const = 0;
+	/** Returns whether the cell data layers referenced any of the provided data layer or not. */
 	virtual bool HasAnyDataLayer(const TSet<FName>& InDataLayers) const = 0;
+	/** Returns the cell bounds computed using cell content. */
 	virtual const FBox& GetContentBounds() const = 0;
+	/** Returns the cell bounds. */
+	virtual FBox GetCellBounds() const = 0;
 };
 
diff --git a/Engine/Source/Runtime/Engine/Public/WorldPartition/WorldPartitionRuntimeSpatialHashCell.h b/Engine/Source/Runtime/Engine/Public/WorldPartition/WorldPartitionRuntimeSpatialHashCell.h
--- a/Engine/Source/Runtime/Engine/Public/WorldPartition/WorldPartitionRuntimeSpatialHashCell.h
+++ b/Engine/Source/Runtime/Engine/Public/WorldPartition/WorldPartitionRuntimeSpatialHashCell.h
@@ -18,6 +18,9 @@ class UWorldPartitionRuntimeSpatialHashCell : public UWorldPartitionRuntimeCell
 	virtual void ResetStreamingSourceInfo() const override;
 	virtual void AppendStreamingSourceInfo(const FWorldPartitionStreamingSource& Source, const FSphericalSector& SourceShape) const override;
 	virtual void MergeStreamingSourceInfo() const override;
+	//~ Begin IWorldPartitionCell Interface
+	virtual FBox GetCellBounds() const override;
+	//~ End IWorldPartitionCell Interface
 
 	virtual int32 SortCompare(const UWorldPartitionRuntimeCell* InOther) const override;
 
