Fix a crash on creating reflection because RG11B10 is not really supported on TVOS.

diff --git a/Engine/Source/Runtime/Engine/Private/MapBuildData.cpp b/Engine/Source/Runtime/Engine/Private/MapBuildData.cpp
--- a/Engine/Source/Runtime/Engine/Private/MapBuildData.cpp
+++ b/Engine/Source/Runtime/Engine/Private/MapBuildData.cpp
@@ -365,6 +365,55 @@ void FReflectionCaptureMapBuildData::FinalizeLoad()
 {
 	AllocatedSize = FullHDRCapturedData.GetAllocatedSize() + EncodedHDRCapturedData.GetAllocatedSize();
 	INC_DWORD_STAT_BY(STAT_ReflectionCaptureBuildData, AllocatedSize);
+
+	bool bMobileEnableClusteredReflections = MobileForwardEnableClusteredReflections(GMaxRHIShaderPlatform) || IsMobileDeferredShadingEnabled(GMaxRHIShaderPlatform);
+	bool bEncodedDataRequired = (GIsEditor || (GMaxRHIFeatureLevel == ERHIFeatureLevel::ES3_1 && !bMobileEnableClusteredReflections));
+	// If the RG11B10 format is not really supported, decode it to RGBA16F 
+	if (GPixelFormats[PF_FloatR11G11B10].BlockBytes == 8 && bEncodedDataRequired && EncodedHDRCapturedData.Num() > 0)
+	{
+		const int32 NumMips = FMath::CeilLogTwo(CubemapSize) + 1;
+
+		int32 SourceMipBaseIndex = 0;
+		int32 DestMipBaseIndex = 0;
+
+		TArray<uint8> DecodedHDRData;
+
+		int32 DecodedDataSize = EncodedHDRCapturedData.Num() * sizeof(FFloat16Color) / sizeof(FFloat3Packed);
+
+		DecodedHDRData.Empty(DecodedDataSize);
+		DecodedHDRData.AddZeroed(DecodedDataSize);
+
+		for (int32 MipIndex = 0; MipIndex < NumMips; MipIndex++)
+		{
+			const int32 MipSize = 1 << (NumMips - MipIndex - 1);
+			const int32 SourceCubeFaceBytes = MipSize * MipSize * sizeof(FFloat3Packed);
+			const int32 DestCubeFaceBytes = MipSize * MipSize * sizeof(FFloat16Color);
+
+			// Decode rest of texels
+			for (int32 CubeFace = 0; CubeFace < CubeFace_MAX; CubeFace++)
+			{
+				const int32 FaceSourceIndex = SourceMipBaseIndex + CubeFace * SourceCubeFaceBytes;
+				const int32 FaceDestIndex = DestMipBaseIndex + CubeFace * DestCubeFaceBytes;
+				const FFloat3Packed* FaceSourceData = (const FFloat3Packed*)&EncodedHDRCapturedData[FaceSourceIndex];
+				FFloat16Color* FaceDestData = (FFloat16Color*)&DecodedHDRData[FaceDestIndex];
+
+				// Convert each texel from RG11B10 to linear space FP16 FColor
+				for (int32 y = 0; y < MipSize; y++)
+				{
+					for (int32 x = 0; x < MipSize; x++)
+					{
+						int32 TexelIndex = x + y * MipSize;
+						FaceDestData[TexelIndex] = FFloat16Color(FaceSourceData[TexelIndex].ToLinearColor());
+					}
+				}
+			}
+
+			SourceMipBaseIndex += SourceCubeFaceBytes * CubeFace_MAX;
+			DestMipBaseIndex += DestCubeFaceBytes * CubeFace_MAX;
+		}
+
+		EncodedHDRCapturedData = MoveTemp(DecodedHDRData);
+	}
 }
 
 void FReflectionCaptureMapBuildData::AddReferencedObjects(FReferenceCollector& Collector)
