Blueprint - Outputting more information in one of the checks when blueprints are compiled/loaded.

diff --git a/Engine/Source/Runtime/Engine/Private/BlueprintGeneratedClass.cpp b/Engine/Source/Runtime/Engine/Private/BlueprintGeneratedClass.cpp
--- a/Engine/Source/Runtime/Engine/Private/BlueprintGeneratedClass.cpp
+++ b/Engine/Source/Runtime/Engine/Private/BlueprintGeneratedClass.cpp
@@ -844,7 +844,7 @@ void UBlueprintGeneratedClass::SetupObjectInitializer(FObjectInitializer& Object
 void UBlueprintGeneratedClass::InitPropertiesFromCustomList(uint8* DataPtr, const uint8* DefaultDataPtr)
 {
 	FScopeLock SerializeAndPostLoadLock(&SerializeAndPostLoadCritical);
-	check(bCustomPropertyListForPostConstructionInitialized); // Something went wrong, probably a race condition
+	checkf(bCustomPropertyListForPostConstructionInitialized, TEXT("Custom Property List Not Initialized for %s"), *GetPathNameSafe(this)); // Something went wrong, probably a race condition
 
 	if (const FCustomPropertyListNode* CustomPropertyList = GetCustomPropertyListForPostConstruction())
 	{
