LevelInstance: BreakLevelInstance can't be undone (reset trans buffer).
- The reset is done before unloading the level instance so that we don't get slowed down by UEngine::FindAndPrintStaleReferencesToObject.

diff --git a/Engine/Source/Runtime/Engine/Private/LevelInstance/LevelInstanceSubsystem.cpp b/Engine/Source/Runtime/Engine/Private/LevelInstance/LevelInstanceSubsystem.cpp
--- a/Engine/Source/Runtime/Engine/Private/LevelInstance/LevelInstanceSubsystem.cpp
+++ b/Engine/Source/Runtime/Engine/Private/LevelInstance/LevelInstanceSubsystem.cpp
@@ -33,6 +33,7 @@
 #include "AssetRegistry/AssetData.h"
 #include "FileHelpers.h"
 #include "Editor.h"
+#include "Editor/Transactor.h"
 #include "EditorLevelUtils.h"
 #include "HAL/PlatformTime.h"
 #include "Engine/Selection.h"
@@ -1331,6 +1332,15 @@ void ULevelInstanceSubsystem::BreakLevelInstance_Impl(ILevelInstanceInterface* L
 			return;
 		}
 
+		// Clear undo buffer here because operation of Breaking a level instance is not undoable.
+		// - Do it before Unload and Destroy calls because those calls will try and unload the level.
+		//	 The unloading of the level might cause a call to UEngine::FindAndPrintStaleReferencesToObject
+		//	 which will slow down the unloading because the Trans buffer has some references to the moved actors.
+		if (GEditor->Trans)
+		{
+			GEditor->Trans->Reset(LOCTEXT("BreakLevelInstance", "Break Level Instance"));
+		}
+
 		if (LevelInstanceActor->IsA<APackedLevelActor>())
 		{
 			BlockUnloadLevelInstance(LevelInstance);
