Change PlayerController to register itself as a Streaming Source with the WorldPartition and remove iteration within WorldPartitionStreamingPolicy.

diff --git a/Engine/Source/Runtime/Engine/Classes/GameFramework/PlayerController.h b/Engine/Source/Runtime/Engine/Classes/GameFramework/PlayerController.h
--- a/Engine/Source/Runtime/Engine/Classes/GameFramework/PlayerController.h
+++ b/Engine/Source/Runtime/Engine/Classes/GameFramework/PlayerController.h
@@ -26,6 +26,7 @@
 #include "GenericPlatform/IInputInterface.h"
 #include "GameFramework/PlayerInput.h"
 #include "Physics/AsyncPhysicsData.h"
+#include "WorldPartition/WorldPartitionStreamingSource.h"
 #include "PlayerController.generated.h"
 
 class ACameraActor;
@@ -236,7 +237,7 @@ struct FAsyncPhysicsTimestamp
  * @see https://docs.unrealengine.com/latest/INT/Gameplay/Framework/Controller/PlayerController/
  */
 UCLASS(config=Game, BlueprintType, Blueprintable, meta=(ShortTooltip="A Player Controller is an actor responsible for controlling a Pawn used by the player."))
-class ENGINE_API APlayerController : public AController
+class ENGINE_API APlayerController : public AController, public IWorldPartitionStreamingSourceProvider
 {
 	GENERATED_BODY()
 
@@ -561,6 +562,10 @@ public:
 	UPROPERTY(EditAnywhere, BlueprintReadWrite, Category = WorldPartition, meta=(EditCondition="bEnableStreamingSource"))
 	EStreamingSourcePriority StreamingSourcePriority;
 
+	/** Color used for debugging. */
+	UPROPERTY(EditAnywhere, Category = WorldPartition, meta = (EditCondition = "bEnableStreamingSource"))
+	FColor StreamingSourceDebugColor;
+
 	/** Scale applied to force feedback values */
 	UPROPERTY(config)
 	float ForceFeedbackScale;
@@ -785,6 +790,8 @@ public:
 	UFUNCTION(BlueprintCallable, Category = WorldPartition)
 	virtual EStreamingSourcePriority GetStreamingSourcePriority() const { return StreamingSourcePriority; }
 
+	virtual bool GetStreamingSource(FWorldPartitionStreamingSource& OutStreamingSource) override;
+
 protected:
 	/** Pawn has been possessed, so changing state to NAME_Playing. Start it walking and begin playing with it. */
 	virtual void BeginPlayingState();
diff --git a/Engine/Source/Runtime/Engine/Private/PlayerController.cpp b/Engine/Source/Runtime/Engine/Private/PlayerController.cpp
--- a/Engine/Source/Runtime/Engine/Private/PlayerController.cpp
+++ b/Engine/Source/Runtime/Engine/Private/PlayerController.cpp
@@ -68,7 +68,8 @@
 #include "GameMapsSettings.h"
 #include "Particles/EmitterCameraLensEffectBase.h"
 #include "LevelUtils.h"
-#include "WorldPartition/WorldPartitionStreamingSource.h"
+#include "WorldPartition/WorldPartition.h"
+#include "WorldPartition/WorldPartitionSubsystem.h"
 #include "Physics/AsyncPhysicsInputComponent.h"
 #include "GenericPlatform/GenericPlatformInputDeviceMapper.h"
 #include "PBDRigidsSolver.h"
@@ -1022,6 +1023,13 @@ void APlayerController::PostInitializeComponents()
 
 	bPlayerIsWaiting = true;
 	StateName = NAME_Spectating; // Don't use ChangeState, because we want to defer spawning the SpectatorPawn until the Player is received
+
+	// Register with WorldParition as a Streaming Source
+	UWorldPartitionSubsystem* WorldPartitionSubsystem = GetWorld()->GetSubsystem<UWorldPartitionSubsystem>();
+	if (WorldPartitionSubsystem)
+	{
+		WorldPartitionSubsystem->RegisterStreamingSourceProvider(this);
+	}
 }
 
 /// @cond DOXYGEN_WARNINGS
@@ -1584,6 +1592,12 @@ void APlayerController::Destroyed()
 	PlayerInput = NULL;
 	CheatManager = NULL;
 
+	UWorldPartitionSubsystem* WorldPartitionSubsystem = GetWorld()->GetSubsystem<UWorldPartitionSubsystem>();
+	if (WorldPartitionSubsystem)
+	{
+		verify(WorldPartitionSubsystem->UnregisterStreamingSourceProvider(this));
+	}
+
 	Super::Destroyed();
 }
 
@@ -6096,5 +6110,23 @@ void APlayerController::AsyncPhysicsTickActor(float DeltaTime, float SimTime)
 	}
 }
 
+bool APlayerController::GetStreamingSource(FWorldPartitionStreamingSource& OutStreamingSource)
+{
+	const ENetMode NetMode = GetNetMode();
+	const bool bIsServer = (NetMode == NM_DedicatedServer || NetMode == NM_ListenServer);
+
+	if (IsStreamingSourceEnabled() && (IsLocalController() || bIsServer))
+	{
+		GetPlayerViewPoint(OutStreamingSource.Location, OutStreamingSource.Rotation);
+		OutStreamingSource.Name = GetFName();
+		OutStreamingSource.TargetState = StreamingSourceShouldActivate() ? EStreamingSourceTargetState::Activated : EStreamingSourceTargetState::Loaded;
+		OutStreamingSource.bBlockOnSlowLoading = StreamingSourceShouldBlockOnSlowStreaming();
+		OutStreamingSource.DebugColor = StreamingSourceDebugColor;
+		OutStreamingSource.Priority = GetStreamingSourcePriority();
+		return true;
+	}
+	return false;
+}
+
 #undef LOCTEXT_NAMESPACE
 
diff --git a/Engine/Source/Runtime/Engine/Private/WorldPartition/WorldPartitionStreamingPolicy.cpp b/Engine/Source/Runtime/Engine/Private/WorldPartition/WorldPartitionStreamingPolicy.cpp
--- a/Engine/Source/Runtime/Engine/Private/WorldPartition/WorldPartitionStreamingPolicy.cpp
+++ b/Engine/Source/Runtime/Engine/Private/WorldPartition/WorldPartitionStreamingPolicy.cpp
@@ -15,8 +15,6 @@
 #include "WorldPartition/HLOD/HLODActor.h"
 #include "WorldPartition/HLOD/HLODSubsystem.h"
 #include "WorldPartition/WorldPartitionDebugHelper.h"
-#include "GameFramework/PlayerController.h"
-#include "Engine/LocalPlayer.h"
 #include "Engine/World.h"
 #include "Engine/Engine.h"
 #include "Engine/Canvas.h"
@@ -123,31 +121,6 @@ void UWorldPartitionStreamingPolicy::UpdateStreamingSources()
 #endif
 		if (!bIsServer || bIsServerStreamingEnabled || AWorldPartitionReplay::IsRecordingEnabled(World))
 		{
-			for (FConstPlayerControllerIterator It = World->GetPlayerControllerIterator(); It; ++It)
-			{
-				if (APlayerController* PlayerController = It->Get())
-				{
-					if (bIsServer || PlayerController->IsLocalController())
-					{
-						if (PlayerController->IsStreamingSourceEnabled())
-						{
-							FVector ViewLocation;
-							FRotator ViewRotation;
-							PlayerController->GetPlayerViewPoint(ViewLocation, ViewRotation);
-
-							// Transform to Local
-							ViewLocation = WorldToLocal.TransformPosition(ViewLocation);
-							ViewRotation = WorldToLocal.TransformRotation(ViewRotation.Quaternion()).Rotator();
-							const EStreamingSourceTargetState TargetState = PlayerController->StreamingSourceShouldActivate() ? EStreamingSourceTargetState::Activated : EStreamingSourceTargetState::Loaded;
-							const bool bBlockOnSlowLoading = PlayerController->StreamingSourceShouldBlockOnSlowStreaming();
-							const EStreamingSourcePriority StreamingSourcePriority = PlayerController->GetStreamingSourcePriority();
-							const bool bRemote = !PlayerController->IsLocalController();
-							StreamingSources.Add(FWorldPartitionStreamingSource(PlayerController->GetFName(), ViewLocation, ViewRotation, TargetState, bBlockOnSlowLoading, StreamingSourcePriority, bRemote));
-						}
-					}
-				}
-			}
-
 			UWorldPartitionSubsystem* WorldPartitionSubsystem = GetWorld()->GetSubsystem<UWorldPartitionSubsystem>();
 			check(WorldPartitionSubsystem);
 
