Change UWorldPartitionSubsystem's StreamingSourceProviders to a TSet to avoid any duplicate registrations.

diff --git a/Engine/Source/Runtime/Engine/Private/WorldPartition/WorldPartitionSubsystem.cpp b/Engine/Source/Runtime/Engine/Private/WorldPartition/WorldPartitionSubsystem.cpp
--- a/Engine/Source/Runtime/Engine/Private/WorldPartition/WorldPartitionSubsystem.cpp
+++ b/Engine/Source/Runtime/Engine/Private/WorldPartition/WorldPartitionSubsystem.cpp
@@ -218,7 +218,9 @@ void UWorldPartitionSubsystem::OnWorldPartitionUninitialized(UWorldPartition* In
 
 void UWorldPartitionSubsystem::RegisterStreamingSourceProvider(IWorldPartitionStreamingSourceProvider* StreamingSource)
 {
-	StreamingSourceProviders.Add(StreamingSource);
+	bool bIsAlreadyInSet = false;
+	StreamingSourceProviders.Add(StreamingSource, &bIsAlreadyInSet);
+	UE_CLOG(bIsAlreadyInSet, LogWorldPartition, Warning, TEXT("Streaming source provider already registered."));
 }
 
 bool UWorldPartitionSubsystem::UnregisterStreamingSourceProvider(IWorldPartitionStreamingSourceProvider* StreamingSource)
diff --git a/Engine/Source/Runtime/Engine/Public/WorldPartition/WorldPartitionSubsystem.h b/Engine/Source/Runtime/Engine/Public/WorldPartition/WorldPartitionSubsystem.h
--- a/Engine/Source/Runtime/Engine/Public/WorldPartition/WorldPartitionSubsystem.h
+++ b/Engine/Source/Runtime/Engine/Public/WorldPartition/WorldPartitionSubsystem.h
@@ -43,7 +43,7 @@ public:
 
 	void DumpStreamingSources(FOutputDevice& OutputDevice) const;
 
-	TArray<IWorldPartitionStreamingSourceProvider*> GetStreamingSourceProviders() const { return StreamingSourceProviders; }
+	TSet<IWorldPartitionStreamingSourceProvider*> GetStreamingSourceProviders() const { return StreamingSourceProviders; }
 	void RegisterStreamingSourceProvider(IWorldPartitionStreamingSourceProvider* StreamingSource);
 	bool UnregisterStreamingSourceProvider(IWorldPartitionStreamingSourceProvider* StreamingSource);
 
@@ -65,7 +65,7 @@ private:
 
 	TArray<TObjectPtr<UWorldPartition>> RegisteredWorldPartitions;
 
-	TArray<IWorldPartitionStreamingSourceProvider*> StreamingSourceProviders;
+	TSet<IWorldPartitionStreamingSourceProvider*> StreamingSourceProviders;
 
 	FDelegateHandle	DrawHandle;
 
