Fixing Crash in USoundNodeBranch::ParseNodes from dereferencing nullptr

#rnx

diff --git a/Engine/Source/Runtime/Engine/Private/AudioDevice.cpp b/Engine/Source/Runtime/Engine/Private/AudioDevice.cpp
--- a/Engine/Source/Runtime/Engine/Private/AudioDevice.cpp
+++ b/Engine/Source/Runtime/Engine/Private/AudioDevice.cpp
@@ -6073,8 +6073,6 @@ void FAudioDevice::UnlinkActiveSoundFromComponent(const FActiveSound& InActiveSo
 				{
 					if (ActiveSound->GetInstanceID() == InActiveSound.GetInstanceID())
 					{
-						// Clears the transmitter ensuring additional parameter requests are ignored.
-						ActiveSound->ClearTransmitter();
 						ActiveSoundsInComponent->RemoveAtSwap(i, 1, false);
 						break;
 					}
