Lazy Loading - Making a few more basic asset types LazyOnDemand. (CurveTable, DataAsset, DataTable, DialogueWave)

diff --git a/Engine/Source/Runtime/Engine/Classes/Engine/CurveTable.h b/Engine/Source/Runtime/Engine/Classes/Engine/CurveTable.h
--- a/Engine/Source/Runtime/Engine/Classes/Engine/CurveTable.h
+++ b/Engine/Source/Runtime/Engine/Classes/Engine/CurveTable.h
@@ -35,7 +35,7 @@ enum class ECurveTableMode : uint8
 /**
  * Imported spreadsheet table as curves.
  */
-UCLASS(MinimalAPI)
+UCLASS(MinimalAPI, Meta = (LoadBehavior = "LazyOnDemand"))
 class UCurveTable
 	: public UObject
 	, public FCurveOwnerInterface
diff --git a/Engine/Source/Runtime/Engine/Classes/Engine/DataAsset.h b/Engine/Source/Runtime/Engine/Classes/Engine/DataAsset.h
--- a/Engine/Source/Runtime/Engine/Classes/Engine/DataAsset.h
+++ b/Engine/Source/Runtime/Engine/Classes/Engine/DataAsset.h
@@ -14,7 +14,7 @@
  * Assets can be made in the Content Browser using any native class that inherits from this.
  * If you want data inheritance or a complicated hierarchy, Data Only Blueprint Classes should be created instead.
  */
-UCLASS(abstract, MinimalAPI)
+UCLASS(abstract, MinimalAPI, Meta = (LoadBehavior = "LazyOnDemand"))
 class UDataAsset : public UObject
 {
 	GENERATED_UCLASS_BODY()
diff --git a/Engine/Source/Runtime/Engine/Classes/Engine/DataTable.h b/Engine/Source/Runtime/Engine/Classes/Engine/DataTable.h
--- a/Engine/Source/Runtime/Engine/Classes/Engine/DataTable.h
+++ b/Engine/Source/Runtime/Engine/Classes/Engine/DataTable.h
@@ -62,7 +62,7 @@ struct FTableRowBase
 /**
  * Imported spreadsheet table.
  */
-UCLASS(MinimalAPI, BlueprintType, AutoExpandCategories = "DataTable,ImportOptions")
+UCLASS(MinimalAPI, BlueprintType, AutoExpandCategories = "DataTable,ImportOptions", Meta = (LoadBehavior = "LazyOnDemand"))
 class UDataTable
 	: public UObject
 {
diff --git a/Engine/Source/Runtime/Engine/Classes/Sound/DialogueWave.h b/Engine/Source/Runtime/Engine/Classes/Sound/DialogueWave.h
--- a/Engine/Source/Runtime/Engine/Classes/Sound/DialogueWave.h
+++ b/Engine/Source/Runtime/Engine/Classes/Sound/DialogueWave.h
@@ -77,7 +77,7 @@ ENGINE_API bool operator!=(const FDialogueContextMapping& LHS, const FDialogueCo
 
 class UDialogueWaveFactory;
 
-UCLASS(hidecategories=Object, editinlinenew, MinimalAPI, BlueprintType)
+UCLASS(hidecategories=Object, editinlinenew, MinimalAPI, BlueprintType, Meta = (LoadBehavior = "LazyOnDemand"))
 class UDialogueWave : public UObject
 {
 	GENERATED_UCLASS_BODY()
