World Partition: temp fix for regression introduced in 22294197 that wasn't including alwayd loaded actors in the persistent level during cook for levels with streaming disabled.

#rnx

diff --git a/Engine/Source/Runtime/Engine/Private/WorldPartition/WorldPartitionCookPackageSplitter.cpp b/Engine/Source/Runtime/Engine/Private/WorldPartition/WorldPartitionCookPackageSplitter.cpp
--- a/Engine/Source/Runtime/Engine/Private/WorldPartition/WorldPartitionCookPackageSplitter.cpp
+++ b/Engine/Source/Runtime/Engine/Private/WorldPartition/WorldPartitionCookPackageSplitter.cpp
@@ -167,6 +167,11 @@ bool FWorldPartitionCookPackageSplitter::PopulateGeneratorPackage(UPackage* Owne
 		{
 			bIsSuccess &= CookPackageGenerator->PopulateGeneratorPackageForCook(CookContext, *CookPackages, ModifiedPackages);
 		}
+		else
+		{
+			TArray<FWorldPartitionCookPackage*> Temp;
+			bIsSuccess &= CookPackageGenerator->PopulateGeneratorPackageForCook(CookContext, Temp, ModifiedPackages);
+		}
 	}
 
 	UE_LOG(LogWorldPartition, Log, TEXT("[Cook][PopulateGeneratorPackage] Gathered %u modified packages"), ModifiedPackages.Num());
