Clamp the local pixel mouse position to the range of an int32 to prevent invalid casting bounds. When you swap from a controller to keyboard, the previous mouse cursor position will be -int32 max, which when translated from an absolutely value to a local value will give you an out of bounds int32 value which would have been incorrectly wrapped.

diff --git a/Engine/Source/Runtime/Engine/Private/Slate/SceneViewport.cpp b/Engine/Source/Runtime/Engine/Private/Slate/SceneViewport.cpp
--- a/Engine/Source/Runtime/Engine/Private/Slate/SceneViewport.cpp
+++ b/Engine/Source/Runtime/Engine/Private/Slate/SceneViewport.cpp
@@ -223,8 +223,8 @@ void FSceneViewport::UpdateCachedCursorPos( const FGeometry& InGeometry, const F
 	if (UserId == FSlateApplication::SlateAppPrimaryPlatformUser)
 	{
 		FVector2D LocalPixelMousePos = InGeometry.AbsoluteToLocal(InMouseEvent.GetScreenSpacePosition());
-		LocalPixelMousePos.X *= CachedGeometry.Scale;
-		LocalPixelMousePos.Y *= CachedGeometry.Scale;
+		LocalPixelMousePos.X = FMath::Clamp(LocalPixelMousePos.X * CachedGeometry.Scale, (double)TNumericLimits<int32>::Min(), (double)TNumericLimits<int32>::Max());
+		LocalPixelMousePos.Y = FMath::Clamp(LocalPixelMousePos.Y * CachedGeometry.Scale, (double)TNumericLimits<int32>::Min(), (double)TNumericLimits<int32>::Max());
 
 		CachedCursorPos = LocalPixelMousePos.IntPoint();
 	}
