Disable Hlod callback on ContentBundle cells.
Content Bundle do not support Hlods and Hlods assumes knowledges of all cells

diff --git a/Engine/Source/Runtime/Engine/Private/WorldPartition/WorldPartitionRuntimeLevelStreamingCell.cpp b/Engine/Source/Runtime/Engine/Private/WorldPartition/WorldPartitionRuntimeLevelStreamingCell.cpp
--- a/Engine/Source/Runtime/Engine/Private/WorldPartition/WorldPartitionRuntimeLevelStreamingCell.cpp
+++ b/Engine/Source/Runtime/Engine/Private/WorldPartition/WorldPartitionRuntimeLevelStreamingCell.cpp
@@ -317,7 +317,9 @@ UWorldPartitionLevelStreamingDynamic* UWorldPartitionRuntimeLevelStreamingCell::
 	}
 #endif
 
-	if (LevelStreaming)
+	// @todo_ow ContentBundles do not support Hlods and events below are forwarding
+	// Events to the HLodSubsystem which assumes knowledge of all cells (not true with plugins)
+	if (LevelStreaming  && !GetContentBundleID().IsValid())
 	{
 		LevelStreaming->OnLevelShown.AddUniqueDynamic(this, &UWorldPartitionRuntimeLevelStreamingCell::OnLevelShown);
 		LevelStreaming->OnLevelHidden.AddUniqueDynamic(this, &UWorldPartitionRuntimeLevelStreamingCell::OnLevelHidden);
