Bump server epoch when ContentBundle Inject their contents so their cells a refresh

diff --git a/Engine/Source/Runtime/Engine/Private/WorldPartition/ContentBundle/ContentBundle.cpp b/Engine/Source/Runtime/Engine/Private/WorldPartition/ContentBundle/ContentBundle.cpp
--- a/Engine/Source/Runtime/Engine/Private/WorldPartition/ContentBundle/ContentBundle.cpp
+++ b/Engine/Source/Runtime/Engine/Private/WorldPartition/ContentBundle/ContentBundle.cpp
@@ -15,10 +15,12 @@
 #include "Editor.h"
 #endif
 
+int32 FContentBundle::ContentBundlesEpoch = 0;
+
 FContentBundle::FContentBundle(TSharedPtr<FContentBundleClient>& InClient, UWorld* InWorld)
-	: FContentBundleBase(InClient, InWorld),
-	ExternalStreamingObjectPackage(nullptr),
-	ExternalStreamingObject(nullptr)
+	: FContentBundleBase(InClient, InWorld)
+	, ExternalStreamingObjectPackage(nullptr)
+	, ExternalStreamingObject(nullptr)
 {
 
 }
@@ -66,6 +68,8 @@ void FContentBundle::DoInjectContent()
 		{
 			UE_LOG(LogContentBundle, Log, TEXT("[CB: %s] Streaming Object Injected."), *GetDescriptor()->GetDisplayName());
 			SetStatus(EContentBundleStatus::ContentInjected);
+
+			ContentBundlesEpoch++;
 		}
 		else
 		{
@@ -84,7 +88,11 @@ void FContentBundle::DoRemoveContent()
 {
 	if (ExternalStreamingObject != nullptr)
 	{
-		if (!GetInjectedWorld()->GetWorldPartition()->RuntimeHash->RemoveExternalStreamingObject(ExternalStreamingObject))
+		if (GetInjectedWorld()->GetWorldPartition()->RuntimeHash->RemoveExternalStreamingObject(ExternalStreamingObject))
+		{
+			ContentBundlesEpoch++;
+		}
+		else
 		{
 			UE_LOG(LogContentBundle, Error, TEXT("[CB: %s] Error while removing streaming object."), *GetDescriptor()->GetDisplayName());
 		}
diff --git a/Engine/Source/Runtime/Engine/Private/WorldPartition/WorldPartitionStreamingPolicy.cpp b/Engine/Source/Runtime/Engine/Private/WorldPartition/WorldPartitionStreamingPolicy.cpp
--- a/Engine/Source/Runtime/Engine/Private/WorldPartition/WorldPartitionStreamingPolicy.cpp
+++ b/Engine/Source/Runtime/Engine/Private/WorldPartition/WorldPartitionStreamingPolicy.cpp
@@ -14,6 +14,7 @@
 #include "WorldPartition/DataLayer/WorldDataLayers.h"
 #include "WorldPartition/HLOD/HLODActor.h"
 #include "WorldPartition/HLOD/HLODSubsystem.h"
+#include "WorldPartition/ContentBundle/ContentBundle.h"
 #include "WorldPartition/WorldPartitionDebugHelper.h"
 #include "Engine/World.h"
 #include "Engine/Engine.h"
@@ -70,6 +71,7 @@ UWorldPartitionStreamingPolicy::UWorldPartitionStreamingPolicy(const FObjectInit
 	, bCriticalPerformanceRequestedBlockTillOnWorld(false)
 	, CriticalPerformanceBlockTillLevelStreamingCompletedEpoch(0)
 	, DataLayersStatesServerEpoch(INT_MIN)
+	, ContentBundleServerEpoch(INT_MIN)
 	, ServerStreamingEnabledEpoch(INT_MIN)
 	, StreamingPerformance(EWorldPartitionStreamingPerformance::Good)
 #if !UE_BUILD_SHIPPING
@@ -261,6 +263,7 @@ void UWorldPartitionStreamingPolicy::UpdateStreamingState()
 		// Server will activate all non data layer cells at first and then load/activate/unload data layer cells only when the data layer states change
 		if (!bIsServerStreamingEnabled && 
 			(ServerStreamingEnabledEpoch == NewServerStreamingEnabledEpoch) &&
+			(ContentBundleServerEpoch == FContentBundle::GetContentBundleEpoch()) &&
 			(DataLayersStatesServerEpoch == AWorldDataLayers::GetDataLayersStateEpoch()))
 		{
 			// Server as nothing to do early out
@@ -495,6 +498,7 @@ void UWorldPartitionStreamingPolicy::UpdateStreamingState()
 	if (bUpdateEpoch)
 	{
 		DataLayersStatesServerEpoch = AWorldDataLayers::GetDataLayersStateEpoch();
+		ContentBundleServerEpoch = FContentBundle::GetContentBundleEpoch();
 		ServerStreamingEnabledEpoch = NewServerStreamingEnabledEpoch;
 		UE_LOG(LogWorldPartition, Verbose, TEXT("Server epoch updated"));
 	}
diff --git a/Engine/Source/Runtime/Engine/Public/WorldPartition/ContentBundle/ContentBundle.h b/Engine/Source/Runtime/Engine/Public/WorldPartition/ContentBundle/ContentBundle.h
--- a/Engine/Source/Runtime/Engine/Public/WorldPartition/ContentBundle/ContentBundle.h
+++ b/Engine/Source/Runtime/Engine/Public/WorldPartition/ContentBundle/ContentBundle.h
@@ -23,6 +23,8 @@ public:
 	virtual bool IsValid() const override;
 	//~ End IContentBundle Interface
 
+	static int32 GetContentBundleEpoch() { return ContentBundlesEpoch; }
+
 protected:
 	//~ Begin IContentBundle Interface
 	virtual void DoInitialize() override;
@@ -36,6 +38,8 @@ private:
 	void InitializeForPIE();
 #endif
 
+	static int32 ContentBundlesEpoch;
+
 	UPackage* ExternalStreamingObjectPackage;
 	URuntimeHashExternalStreamingObjectBase* ExternalStreamingObject;
 };
\ No newline at end of file
diff --git a/Engine/Source/Runtime/Engine/Public/WorldPartition/WorldPartitionStreamingPolicy.h b/Engine/Source/Runtime/Engine/Public/WorldPartition/WorldPartitionStreamingPolicy.h
--- a/Engine/Source/Runtime/Engine/Public/WorldPartition/WorldPartitionStreamingPolicy.h
+++ b/Engine/Source/Runtime/Engine/Public/WorldPartition/WorldPartitionStreamingPolicy.h
@@ -97,6 +97,7 @@ protected:
 	mutable TArray<const UWorldPartitionRuntimeCell*> SortedAddToWorldCells;
 
 	int32 DataLayersStatesServerEpoch;
+	int32 ContentBundleServerEpoch;
 	int32 ServerStreamingEnabledEpoch;
 
 	EWorldPartitionStreamingPerformance StreamingPerformance;
