Change Strata legacy conversion to use proper shading model input.

diff --git a/Engine/Source/Runtime/Engine/Classes/Materials/MaterialExpressionStrata.h b/Engine/Source/Runtime/Engine/Classes/Materials/MaterialExpressionStrata.h
--- a/Engine/Source/Runtime/Engine/Classes/Materials/MaterialExpressionStrata.h
+++ b/Engine/Source/Runtime/Engine/Classes/Materials/MaterialExpressionStrata.h
@@ -163,7 +163,7 @@ class UMaterialExpressionStrataLegacyConversion : public UMaterialExpressionStra
 	 * Shading models
 	 */
 	UPROPERTY()
-	FExpressionInput ShadingModel;
+	FShadingModelMaterialInput ShadingModel;
 	
 	/** SubsurfaceProfile, for Screen Space Subsurface Scattering. The profile needs to be set up on both the Strata diffuse node, and the material node at the moment. */
 	UPROPERTY(EditAnywhere, BlueprintReadOnly, Category = Material, meta = (DisplayName = "Subsurface Profile"))
diff --git a/Engine/Source/Runtime/Engine/Private/Materials/Material.cpp b/Engine/Source/Runtime/Engine/Private/Materials/Material.cpp
--- a/Engine/Source/Runtime/Engine/Private/Materials/Material.cpp
+++ b/Engine/Source/Runtime/Engine/Private/Materials/Material.cpp
@@ -3281,10 +3281,9 @@ void UMaterial::ConvertMaterialToStrataMaterial()
 					//check(EditorOnly->ShadingModelFromMaterialExpression.IsConnected());
 					if (!EditorOnly->ShadingModelFromMaterialExpression.IsConnected())
 					{
-						// Now setup the mean free path with a hard coded transmittance of 0.75 when viewing the surface perpendicularly
-						UMaterialExpressionConstant* DefaultShadingModel = NewObject<UMaterialExpressionConstant>(this);
-						DefaultShadingModel->R = 1; // Default Lit shading model
-						ConvertNode->ShadingModel.Connect(0, DefaultShadingModel);
+						UMaterialExpressionShadingModel* ShadingModelNode = NewObject<UMaterialExpressionShadingModel>(this);
+						ShadingModelNode->ShadingModel = MSM_DefaultLit;
+						ConvertNode->ShadingModel.Connect(0, ShadingModelNode);
 					}
 					else
 					{
@@ -3306,11 +3305,9 @@ void UMaterial::ConvertMaterialToStrataMaterial()
 				else
 				{
 					check(!bHasShadingModelMixture);
-				
-					// Add constant for the shading model
-					UMaterialExpressionConstant* ShadingModelNode = NewObject<UMaterialExpressionConstant>(this);
-					ShadingModelNode->SetParameterName(FName(TEXT("ConstantShadingModel")));
-					ShadingModelNode->R = ShadingModel;
+
+					UMaterialExpressionShadingModel* ShadingModelNode = NewObject<UMaterialExpressionShadingModel>(this);
+					ShadingModelNode->ShadingModel = ShadingModel;
 					ConvertNode->ShadingModel.Connect(0, ShadingModelNode);
 
 					AddStrataShadingModelFromMaterialShadingModel(ConvertNode->ConvertedStrataMaterialInfo, ShadingModels);
@@ -3406,9 +3403,8 @@ void UMaterial::ConvertMaterialToStrataMaterial()
 			MoveConnectionTo(EditorOnly->Opacity, ConvertNode, 11);
 
 			// Add constant for the Unlit shading model
-			UMaterialExpressionConstant* ShadingModelNode = NewObject<UMaterialExpressionConstant>(this);
-			ShadingModelNode->SetParameterName(FName(TEXT("ConstantShadingModel")));
-			ShadingModelNode->R = ShadingModel;
+			UMaterialExpressionShadingModel* ShadingModelNode = NewObject<UMaterialExpressionShadingModel>(this);
+			ShadingModelNode->ShadingModel = ShadingModel;
 			ConvertNode->ShadingModel.Connect(0, ShadingModelNode);
 
 			AddStrataShadingModelFromMaterialShadingModel(ConvertNode->ConvertedStrataMaterialInfo, ShadingModels);
