World Partition: cluster spatially loaded actors referencing non-spatially loaded actors together until we identify and fix the bulk data detach cook crash.

#rnx

diff --git a/Engine/Source/Runtime/Engine/Private/WorldPartition/WorldPartitionStreamingGeneration.cpp b/Engine/Source/Runtime/Engine/Private/WorldPartition/WorldPartitionStreamingGeneration.cpp
--- a/Engine/Source/Runtime/Engine/Private/WorldPartition/WorldPartitionStreamingGeneration.cpp
+++ b/Engine/Source/Runtime/Engine/Private/WorldPartition/WorldPartitionStreamingGeneration.cpp
@@ -111,6 +111,15 @@ class FWorldPartitionStreamingGenerator
 				{
 					const FWorldPartitionActorDescView& ReferenceActorDescView = ActorSetContainer.ActorDescViewMap->FindByGuidChecked(ActorSet.Actors[0]);
 
+					// Validate assumptions
+					for (const FGuid& ActorGuid : ActorSet.Actors)
+					{
+						const FWorldPartitionActorDescView& ActorDescView = ActorSetContainer.ActorDescViewMap->FindByGuidChecked(ActorGuid);
+						check(ActorDescView.GetRuntimeGrid() == ReferenceActorDescView.GetRuntimeGrid());
+						check(ActorDescView.GetIsSpatiallyLoaded() == ReferenceActorDescView.GetIsSpatiallyLoaded());
+						check(ActorDescView.GetContentBundleGuid() == ReferenceActorDescView.GetContentBundleGuid());
+					}
+
 					FActorSetInstance& ActorSetInstance = ActorSetInstances.Emplace_GetRef();
 				
 					ActorSetInstance.ContainerInstance = &ActorSetContainer;
@@ -580,14 +589,6 @@ class FWorldPartitionStreamingGenerator
 				{
 					const bool bIsActorDescSpatiallyLoaded = RefererActorDescView.GetIsSpatiallyLoaded();
 					const bool bIsActorDescRefSpatiallyLoaded = ReferenceActorDescView.GetIsSpatiallyLoaded();
-
-					// The only case we support right now is spatially loaded actors referencing non-spatially loaded actors, when target is not in data layers.
-					// For this to work with data layers, we need to implement dependency logic support in the content cooker splitter.
-					if (bIsActorDescSpatiallyLoaded && !bIsActorDescRefSpatiallyLoaded && ReferenceActorDescView.GetDataLayers().IsEmpty())
-					{
-						return true;
-					}
-
 					return bIsActorDescSpatiallyLoaded == bIsActorDescRefSpatiallyLoaded;
 				};
 
