Fixed client showing both server & client world partition debug display when using listen server.

diff --git a/Engine/Source/Runtime/Engine/Private/WorldPartition/WorldPartitionSubsystem.cpp b/Engine/Source/Runtime/Engine/Private/WorldPartition/WorldPartitionSubsystem.cpp
--- a/Engine/Source/Runtime/Engine/Private/WorldPartition/WorldPartitionSubsystem.cpp
+++ b/Engine/Source/Runtime/Engine/Private/WorldPartition/WorldPartitionSubsystem.cpp
@@ -313,7 +313,7 @@ void UWorldPartitionSubsystem::Draw(UCanvas* Canvas, class APlayerController* PC
 	}
 
 	// Filter out views that don't match our world
-	if (!WorldPartition->IsServer() &&
+	if (!WorldPartition->GetWorld()->IsNetMode(NM_DedicatedServer) &&
 		Canvas->SceneView->ViewActor != nullptr &&
 		Canvas->SceneView->ViewActor->GetWorld() != GetWorld())
 	{
