Fix for sync groups not working correctly when using 'Transition Leader' group role across non-zero-length transitions

Fix by [REDACTED_PERSON] [at] The Coalition

diff --git a/Engine/Source/Runtime/Engine/Classes/Animation/AnimationAsset.h b/Engine/Source/Runtime/Engine/Classes/Animation/AnimationAsset.h
--- a/Engine/Source/Runtime/Engine/Classes/Animation/AnimationAsset.h
+++ b/Engine/Source/Runtime/Engine/Classes/Animation/AnimationAsset.h
@@ -813,7 +813,7 @@ public:
 		, bIsMarkerPositionValid(ValidMarkerNames.Num() > 0)
 		, bIsLeader(true)
 		, bOnlyOneAnimationInGroup(bInOnlyOneAnimationInGroup)
-		, bIsJoiningSyncGroup(false)
+		, bResyncToSyncGroup(false)
 	{
 	}
 
@@ -826,7 +826,7 @@ public:
 		, bIsMarkerPositionValid(false)
 		, bIsLeader(true)
 		, bOnlyOneAnimationInGroup(bInOnlyOneAnimationInGroup)
-		, bIsJoiningSyncGroup(false)
+		, bResyncToSyncGroup(false)
 	{
 	}
 
@@ -904,15 +904,15 @@ public:
 		return bOnlyOneAnimationInGroup;
 	}
 
-	void SetJoiningSyncGroup(bool bInIsJoiningSyncGroup)
+	void SetResyncToSyncGroup(bool bInResyncToSyncGroup)
 	{
-		bIsJoiningSyncGroup = bInIsJoiningSyncGroup;
+		bResyncToSyncGroup = bInResyncToSyncGroup;
 	}
 
-	// Are we (re)joining the sync group this tick (eg: when initializing or resuming from zero weight)?
-	bool IsJoiningSyncGroup() const
+	// Should we resync to the sync group this tick (eg: when initializing or resuming from zero weight)?
+	bool ShouldResyncToSyncGroup() const
 	{
-		return bIsJoiningSyncGroup;
+		return bResyncToSyncGroup;
 	}
 
 	//Root Motion accumulated from this tick context
@@ -940,8 +940,8 @@ private:
 
 	bool bOnlyOneAnimationInGroup;
 
-	// True if the asset player being ticked was not previously part of the sync group (eg: it was inactive and has now reactivated)
-	bool bIsJoiningSyncGroup;
+	// True if the asset player being ticked should (re)synchronize to the sync group's time (eg: it was inactive and has now reactivated)
+	bool bResyncToSyncGroup;
 };
 
 USTRUCT()
diff --git a/Engine/Source/Runtime/Engine/Private/Animation/AnimSequenceBase.cpp b/Engine/Source/Runtime/Engine/Private/Animation/AnimSequenceBase.cpp
--- a/Engine/Source/Runtime/Engine/Private/Animation/AnimSequenceBase.cpp
+++ b/Engine/Source/Runtime/Engine/Private/Animation/AnimSequenceBase.cpp
@@ -504,7 +504,7 @@ void UAnimSequenceBase::RemapTracksToNewSkeleton(USkeleton* NewSkeleton, bool bC
 void UAnimSequenceBase::TickAssetPlayer(FAnimTickRecord& Instance, struct FAnimNotifyQueue& NotifyQueue, FAnimAssetTickContext& Context) const
 {
 	float CurrentTime = *(Instance.TimeAccumulator);
-	if (Context.IsJoiningSyncGroup())
+	if (Context.ShouldResyncToSyncGroup())
 	{
 		// Synchronize the asset player time to the other sync group members when (re)joining the group
 		CurrentTime = Context.GetAnimationPositionRatio() * GetPlayLength();
diff --git a/Engine/Source/Runtime/Engine/Private/Animation/AnimSync.cpp b/Engine/Source/Runtime/Engine/Private/Animation/AnimSync.cpp
--- a/Engine/Source/Runtime/Engine/Private/Animation/AnimSync.cpp
+++ b/Engine/Source/Runtime/Engine/Private/Animation/AnimSync.cpp
@@ -77,6 +77,7 @@ void FAnimSync::TickAssetPlayerInstances(FAnimInstanceProxy& InProxy, float InDe
 	FSyncGroupMap& SyncGroupMap = SyncGroupMaps[GetSyncGroupWriteIndex()];
 	const FSyncGroupMap& PreviousSyncGroupMap = SyncGroupMaps[GetSyncGroupReadIndex()];
 	TArray<FAnimTickRecord>& UngroupedActivePlayers = UngroupedActivePlayerArrays[GetSyncGroupWriteIndex()];
+	const TArray<FAnimTickRecord>& PreviousUngroupedActivePlayers = UngroupedActivePlayerArrays[GetSyncGroupReadIndex()];
 
 	for (auto& SyncGroupPair : SyncGroupMap)
 	{
@@ -126,20 +127,16 @@ void FAnimSync::TickAssetPlayerInstances(FAnimInstanceProxy& InProxy, float InDe
 				SCOPE_CYCLE_COUNTER(STAT_TickAssetPlayerInstance);
 				FScopeCycleCounterUObject Scope(GroupLeader.SourceAsset);
 
-				// determine if we are (re)joining the sync group (ie: that we were not part of the sync group on the previous update)
-				bool bIsJoiningSyncGroup = true;
-				if (PreviousGroup)
+				auto TickRecordMatchesGroupLeader = [&GroupLeader](const FAnimTickRecord& TickRecord) -> bool
 				{
-					for (const FAnimTickRecord& PreviousPlayer : PreviousGroup->ActivePlayers)
-					{
-						if (PreviousPlayer.SourceAsset == GroupLeader.SourceAsset && PreviousPlayer.TimeAccumulator == GroupLeader.TimeAccumulator)
-						{
-							bIsJoiningSyncGroup = false;
-							break;
-						}
-					}
-				}
-				TickContext.SetJoiningSyncGroup(bIsJoiningSyncGroup);
+					return TickRecord.SourceAsset == GroupLeader.SourceAsset && TickRecord.TimeAccumulator == GroupLeader.TimeAccumulator;
+				};
+
+				// if the group leader was previously inactive then set it to resynchronize to the sync group's time
+				// (maintains sync during inertialization or zero-length blends where the previous sync leader is no longer active)
+				const bool bGroupLeaderWasActive = (PreviousGroup && PreviousGroup->ActivePlayers.ContainsByPredicate(TickRecordMatchesGroupLeader)) ||
+					(PreviousUngroupedActivePlayers.ContainsByPredicate(TickRecordMatchesGroupLeader));
+				TickContext.SetResyncToSyncGroup(!bGroupLeaderWasActive);
 
 				TickContext.MarkerTickContext.MarkersPassedThisTick.Reset();
 				TickContext.RootMotionMovementParams.Clear();
@@ -234,7 +231,7 @@ void FAnimSync::TickAssetPlayerInstances(FAnimInstanceProxy& InProxy, float InDe
 				// if we don't have a good leader, no reason to convert to follower
 				// tick as leader
 				TickContext.ConvertToFollower();
-				TickContext.SetJoiningSyncGroup(false);
+				TickContext.SetResyncToSyncGroup(false);
 
 				for (int32 TickIndex = GroupLeaderIndex + 1; TickIndex < SyncGroup.ActivePlayers.Num(); ++TickIndex)
 				{
diff --git a/Engine/Source/Runtime/Engine/Private/Animation/BlendSpace.cpp b/Engine/Source/Runtime/Engine/Private/Animation/BlendSpace.cpp
--- a/Engine/Source/Runtime/Engine/Private/Animation/BlendSpace.cpp
+++ b/Engine/Source/Runtime/Engine/Private/Animation/BlendSpace.cpp
@@ -418,7 +418,7 @@ void UBlendSpace::TickAssetPlayer(FAnimTickRecord& Instance, struct FAnimNotifyQ
 			}
 
 			float& NormalizedCurrentTime = *(Instance.TimeAccumulator);
-			if (Context.IsJoiningSyncGroup())
+			if (Context.ShouldResyncToSyncGroup())
 			{
 				// Synchronize the asset player time to the other sync group members when (re)joining the group
 				NormalizedCurrentTime = Context.GetAnimationPositionRatio();
