Fixes to prevent client from making visible a streaming level before the server or server making invisible/unloading while client makes visible (could happen when activating/deactivating a data layer and/or when using server streaming)
- Added option for client to wait for ACK from server before making a streaming level visible (doesn't apply for client only visible levels).
- The feature is controlled by UWorld::SupportsMakingVisibleTransactionRequests which internally uses LevelStreaming.DefaultAllowClientUseMakingVisibleTransactionRequests. If world is partitioned, it will check for wp.Runtime.UseMakingVisibleTransactionRequests (defaulted to false for now).
- If server acknowledges success, it will be part of NetConnection's ClientMakingVisibleLevelNames until is it visible and moved into ClientVisibleLevelNames.
- WorldPartition now also makes sure a streaming level is not part of ClientMakingVisibleLevelNames nor ClientVisibleLevelNames before unloading streaming levels.
- LevelStreaming.ShouldClientUseMakingInvisibleTransactionRequest was converted to use the same pattern (now LevelStreaming.DefaultAllowClientUseMakingInvisibleTransactionRequests and added UWorld::SupportsMakingInvisibleTransactionRequests().  If world is partitioned, it will check for wp.Runtime.UseMakingInvisibleTransactionRequests (defaulted to false for now).

diff --git a/Engine/Source/Runtime/Engine/Classes/Engine/LevelStreaming.h b/Engine/Source/Runtime/Engine/Classes/Engine/LevelStreaming.h
--- a/Engine/Source/Runtime/Engine/Classes/Engine/LevelStreaming.h
+++ b/Engine/Source/Runtime/Engine/Classes/Engine/LevelStreaming.h
@@ -10,6 +10,7 @@
 #include "Engine/LatentActionManager.h"
 #include "LatentActions.h"
 #include "ProfilingDebugging/ProfilingHelpers.h"
+#include "GameFramework/UpdateLevelVisibilityLevelInfo.h"
 
 #if WITH_EDITOR
 #include "Folder.h"
@@ -22,12 +23,20 @@ class ULevel;
 class ULevelStreaming;
 struct FNetLevelVisibilityTransactionId;
 
+enum class ENetLevelVisibilityRequest
+{
+	MakingVisible,
+	MakingInvisible
+};
+
 struct FNetLevelVisibilityState
 {
+	TOptional<bool> ClientAckedRequestCanMakeVisible;
+	TOptional<ENetLevelVisibilityRequest> PendingRequestType;
 	uint32 ServerRequestIndex = 0;
-	uint32 PendingClientRequestIndex = 0;
-	uint32 AckedClientRequestIndex = 0;
-	bool bPendingInvisibleRequest = false;
+	uint32 ClientPendingRequestIndex = 0;
+	uint32 ClientAckedRequestIndex = 0;
+	bool bHasClientPendingRequest = false;
 };
 
 ENGINE_API DECLARE_LOG_CATEGORY_EXTERN(LogLevelStreaming, Log, All);
@@ -293,6 +302,19 @@ public:
 
 private:
 
+	/** If true client will wait for acknowledgment from server before making streaming levels invisible */
+	bool ShouldClientUseMakingInvisibleTransactionRequest() const;
+	/** If true client will wait for acknowledgment from server before making streaming levels visible */
+	bool ShouldClientUseMakingVisibleTransactionRequest() const;
+	/** Returns whether the streaming level should wait for the server ack before changing its visibility. */
+	bool ShouldWaitForServerAckBeforeChangingVisibilityState(ENetLevelVisibilityRequest InRequestType, bool bInShouldBeVisible);
+	/** Ack a client instigated visibility/streaming transaction */
+	void AckNetVisibilityTransaction(FNetLevelVisibilityTransactionId AckedClientTransactionId, bool bClientAckCanMakeVisible);
+	/** Returns whether the streaming level can make visible (can call AddToWorld). */
+	bool CanMakeVisible();
+	/** Returns whether the streaming level can make invisible (can call RemoveFromWorld). */
+	bool CanMakeInvisible();
+
 	/** Determine what the streaming levels target state should be. Returns whether the streaming level should be in the consider list. */
 	bool DetermineTargetState();
 
@@ -310,20 +332,18 @@ private:
 
 public:
 
+	/** Called when streaming level is being considered by world for next UpdateStreamingState. */
+	void OnStreamingLevelAddedToConsidered();
+
 	/** Begin a client instigated NetVisibility request */
 	void BeginClientNetVisibilityRequest(bool bInShouldBeVisible);
 
-	/** Check if we are waiting for a making invisible streaming transaction */
-	bool IsWaitingForNetVisibilityTransactionAck() const;
+	/** Check if we are waiting for a making visible or invisible streaming transaction */
+	bool IsWaitingForNetVisibilityTransactionAck(ENetLevelVisibilityRequest InRequestType = ENetLevelVisibilityRequest::MakingInvisible) const;
 
 	/** Set the current state of the current visibility/streaming transaction */
 	void UpdateNetVisibilityTransactionState(bool bInShouldBeVisible, FNetLevelVisibilityTransactionId TransactionId);
 
-	/** Ack a client instigated visibility/streaming transaction */
-	void AckNetVisibilityTransaction(FNetLevelVisibilityTransactionId AckedClientTransactionId);
-
-public:
-
 	/** Returns the value of bShouldBeVisible. Use ShouldBeVisible to query whether a streaming level should be visible based on its own criteria. */
 	bool GetShouldBeVisibleFlag() const { return bShouldBeVisible; }
 
@@ -522,8 +542,11 @@ public:
 	UPROPERTY(BlueprintAssignable)
 	FLevelStreamingVisibilityStatus		OnLevelHidden;
 
-	/** If true client will wait for acknowledgment from server before making streaming levels invisible */
-	static bool ShouldClientUseMakingInvisibleTransactionRequest();
+	/** Whether client should be using making invisible transaction requests to the server (default value). */
+	static bool DefaultAllowClientUseMakingInvisibleTransactionRequests();
+
+	/** Whether client should be using making visible transaction requests to the server (default value). */
+	static bool DefaultAllowClientUseMakingVisibleTransactionRequests();
 
 	/** If true server will wait for client acknowledgment before making treating streaming levels as visible for the client */
 	static bool ShouldServerUseMakingVisibleTransactionRequest();
@@ -626,9 +649,21 @@ private:
 	/** State for server handshake of NetLevelVisibilty requests */
 	FNetLevelVisibilityState NetVisibilityState;
 
+	friend struct FAckNetVisibilityTransaction;
 	friend struct FStreamingLevelPrivateAccessor;
 };
 
+/** Internal struct used by APlayerController to call AckNetVisibilityTransaction on streaming level object */
+struct FAckNetVisibilityTransaction
+{
+private:
+	static void Call(ULevelStreaming* StreamingLevel, FNetLevelVisibilityTransactionId TransactionId, bool bClientAckCanMakeVisibleResponse)
+	{
+		StreamingLevel->AckNetVisibilityTransaction(TransactionId, bClientAckCanMakeVisibleResponse);
+	}
+	friend class APlayerController;
+};
+
 struct FStreamingLevelPrivateAccessor
 {
 private:
diff --git a/Engine/Source/Runtime/Engine/Classes/Engine/NetConnection.h b/Engine/Source/Runtime/Engine/Classes/Engine/NetConnection.h
--- a/Engine/Source/Runtime/Engine/Classes/Engine/NetConnection.h
+++ b/Engine/Source/Runtime/Engine/Classes/Engine/NetConnection.h
@@ -786,6 +786,8 @@ public:
 
 	/** Called by PlayerController to tell connection about client level visibility change */
 	ENGINE_API void UpdateLevelVisibility(const struct FUpdateLevelVisibilityLevelInfo& LevelVisibility);
+
+	ENGINE_API const TSet<FName>& GetClientMakingVisibleLevelNames() const { return ClientMakingVisibleLevelNames; }
 	
 #if DO_ENABLE_NET_TEST
 
@@ -814,6 +816,9 @@ private:
 
 #endif //#if DO_ENABLE_NET_TEST
 
+	/** On the server, the package names of streaming levels that the client has told us it is making visible */
+	TSet<FName> ClientMakingVisibleLevelNames;
+
 public:
 
 	/** 
diff --git a/Engine/Source/Runtime/Engine/Classes/Engine/World.h b/Engine/Source/Runtime/Engine/Classes/Engine/World.h
--- a/Engine/Source/Runtime/Engine/Classes/Engine/World.h
+++ b/Engine/Source/Runtime/Engine/Classes/Engine/World.h
@@ -39,6 +39,7 @@ class AGameModeBase;
 class AGameStateBase;
 class APhysicsVolume;
 class APlayerController;
+class AServerStreamingLevelsVisibility;
 class AWorldSettings;
 class UWorldPartition;
 class Error;
@@ -1116,8 +1117,20 @@ private:
 	UPROPERTY(Transient, DuplicateTransient)
 	FStreamingLevelsToConsider StreamingLevelsToConsider;
 
+	UPROPERTY(Transient, DuplicateTransient)
+	TObjectPtr<AServerStreamingLevelsVisibility> ServerStreamingLevelsVisibility;
+
 public:
 
+	/** Returns whether the world supports for a client to use "making visible" transaction requests to the server. */
+	bool SupportsMakingVisibleTransactionRequests() const;
+
+	/** Returns whether the world supports for a client to use "making invisible" transaction requests to the server. */
+	bool SupportsMakingInvisibleTransactionRequests() const;
+
+	/** Returns the object used to query server streaming level visibility. */
+	const AServerStreamingLevelsVisibility* GetServerStreamingLevelsVisibility() const;
+
 	/** Return a const version of the streaming levels array */
 	const TArray<ULevelStreaming*>& GetStreamingLevels() const { return StreamingLevels; }
 
@@ -1197,6 +1210,12 @@ private:
 	/** Returns wether AddToWorld should be skipped on a given level */
 	bool CanAddLoadedLevelToWorld(ULevel* Level) const;
 
+	/** Stores whether the game world supports for a client to use "making visible" transaction requests to the server. */
+	mutable TOptional<bool> bSupportsMakingVisibleTransactionRequests;
+
+	/** Stores whether the game world supports for a client to use "making invisible" transaction requests to the server. */
+	mutable TOptional<bool> bSupportsMakingInvisibleTransactionRequests;
+
 	/** Pointer to the current level in the queue to be made visible, NULL if none are pending. */
 	UPROPERTY(Transient)
 	TObjectPtr<class ULevel>								CurrentLevelPendingVisibility;
diff --git a/Engine/Source/Runtime/Engine/Classes/GameFramework/PlayerController.h b/Engine/Source/Runtime/Engine/Classes/GameFramework/PlayerController.h
--- a/Engine/Source/Runtime/Engine/Classes/GameFramework/PlayerController.h
+++ b/Engine/Source/Runtime/Engine/Classes/GameFramework/PlayerController.h
@@ -859,7 +859,7 @@ public:
 	 * @param TransactionId - TransactionId being acknowledged
 	 */
 	UFUNCTION(Reliable, Client)
-	void ClientAckUpdateLevelVisibility(FName PackageName, FNetLevelVisibilityTransactionId TransactionId);
+	void ClientAckUpdateLevelVisibility(FName PackageName, FNetLevelVisibilityTransactionId TransactionId, bool bClientAckCanMakeVisible);
 
 	/** Enable voice chat transmission */
 	void StartTalking();
diff --git a/Engine/Source/Runtime/Engine/Classes/GameFramework/UpdateLevelVisibilityLevelInfo.h b/Engine/Source/Runtime/Engine/Classes/GameFramework/UpdateLevelVisibilityLevelInfo.h
--- a/Engine/Source/Runtime/Engine/Classes/GameFramework/UpdateLevelVisibilityLevelInfo.h
+++ b/Engine/Source/Runtime/Engine/Classes/GameFramework/UpdateLevelVisibilityLevelInfo.h
@@ -58,15 +58,17 @@ struct ENGINE_API FUpdateLevelVisibilityLevelInfo
 		: PackageName(NAME_None)
 		, FileName(NAME_None)
 		, bIsVisible(false)
+		, bTryMakeVisible(false)
 		, bSkipCloseOnError(false)
 	{
 	}
 
 	/**
-	 * @param Level			Level to pull PackageName and FileName from.
-	 * @param bInIsVisible	Default value for bIsVisible.
+	 * @param Level				Level to pull PackageName and FileName from.
+	 * @param bInIsVisible		Default value for bIsVisible.
+	 * @param bInTryMakeVisible	Whether the level is trying to be made visible or not.
 	 */
-	FUpdateLevelVisibilityLevelInfo(const class ULevel* const Level, const bool bInIsVisible);
+	FUpdateLevelVisibilityLevelInfo(const class ULevel* const Level, const bool bInIsVisible, const bool bInTryMakeVisible = false);
 
 	/** The name of the package for the level whose status changed. */
 	UPROPERTY()
@@ -84,6 +86,10 @@ struct ENGINE_API FUpdateLevelVisibilityLevelInfo
 	UPROPERTY()
 	uint32 bIsVisible : 1;
 
+	/** Whether the level is in the state of making visible and waits for server to acknowledge. */
+	UPROPERTY()
+	uint32 bTryMakeVisible : 1;
+
 	/** Skip connection close if level can't be found (not net serialized) */
 	uint32 bSkipCloseOnError : 1;
 
diff --git a/Engine/Source/Runtime/Engine/Private/GameFramework/UpdateLevelVisibilityLevelInfo.cpp b/Engine/Source/Runtime/Engine/Private/GameFramework/UpdateLevelVisibilityLevelInfo.cpp
--- a/Engine/Source/Runtime/Engine/Private/GameFramework/UpdateLevelVisibilityLevelInfo.cpp
+++ b/Engine/Source/Runtime/Engine/Private/GameFramework/UpdateLevelVisibilityLevelInfo.cpp
@@ -19,10 +19,14 @@ namespace PlayerControllerCVars
 	);
 }
 
-FUpdateLevelVisibilityLevelInfo::FUpdateLevelVisibilityLevelInfo(const ULevel* const Level, const bool bInIsVisible)
+FUpdateLevelVisibilityLevelInfo::FUpdateLevelVisibilityLevelInfo(const ULevel* const Level, const bool bInIsVisible, const bool bInTryMakeVisible)
 	: bIsVisible(bInIsVisible)
+	, bTryMakeVisible(bInTryMakeVisible)
 	, bSkipCloseOnError(false)
 {
+	// For backward compatibility, bTryMakeVisible was added instead of converting bIsVisible to an enum.
+	// Make sure we don't receive the invalid state (bIsVisible == true) && (bTryMakeVisible == true)
+	check(!bTryMakeVisible || (bIsVisible != bTryMakeVisible));
 	const UPackage* const LevelPackage = Level->GetOutermost();
 	PackageName = LevelPackage->GetFName();
 
@@ -36,9 +40,11 @@ bool FUpdateLevelVisibilityLevelInfo::NetSerialize(FArchive& Ar, UPackageMap* Pa
 {
 	bool bArePackageAndFileTheSame = !!((PlayerControllerCVars::LevelVisibilityDontSerializeFileName) || (FileName == PackageName) || (FileName == NAME_None));
 	bool bLocalIsVisible = !!bIsVisible;
+	bool bLocalTryMakeVisible = !!bTryMakeVisible;
 
 	Ar.SerializeBits(&bArePackageAndFileTheSame, 1);
 	Ar.SerializeBits(&bLocalIsVisible, 1);
+	Ar.SerializeBits(&bLocalTryMakeVisible, 1);
 	Ar << PackageName;
 
 	if (!bArePackageAndFileTheSame)
@@ -53,6 +59,7 @@ bool FUpdateLevelVisibilityLevelInfo::NetSerialize(FArchive& Ar, UPackageMap* Pa
 	VisibilityRequestId.NetSerialize(Ar, PackageMap, bOutSuccess);
 
 	bIsVisible = bLocalIsVisible;
+	bTryMakeVisible = bLocalTryMakeVisible;
 
 	bOutSuccess = !Ar.IsError();
 	return true;
diff --git a/Engine/Source/Runtime/Engine/Private/LevelStreaming.cpp b/Engine/Source/Runtime/Engine/Private/LevelStreaming.cpp
--- a/Engine/Source/Runtime/Engine/Private/LevelStreaming.cpp
+++ b/Engine/Source/Runtime/Engine/Private/LevelStreaming.cpp
@@ -47,14 +47,24 @@ namespace LevelStreamingCVars
 	// starting to replicate data based on an older visibility/streamingstatus update which can lead to broken channels
 	// to mitigate this problem we assign a TransactionId to each request/update to make sure that we are acting on the correct data
 #if UE_WITH_IRIS
-	static bool bShouldClientUseMakingInvisibleTransactionRequest = true;
+	static bool bDefaultAllowClientUseMakingInvisibleTransactionRequests = true;
 #else
-	static bool bShouldClientUseMakingInvisibleTransactionRequest = false;
+	static bool bDefaultAllowClientUseMakingInvisibleTransactionRequests = false;
 #endif
-	FAutoConsoleVariableRef CVarShouldClientUseMakingInvisibleTransactionRequest(
-		TEXT("LevelStreaming.ShouldClientUseMakingInvisibleTransactionRequest"),
-		bShouldClientUseMakingInvisibleTransactionRequest,
-		TEXT("Whether client should wait for server to acknowledge visibility update before making streaming levels invisible.\n")
+	FAutoConsoleVariableRef CVarDefaultAllowClientUseMakingInvisibleTransactionRequests(
+		TEXT("LevelStreaming.DefaultAllowClientUseMakingInvisibleTransactionRequests"),
+		bDefaultAllowClientUseMakingInvisibleTransactionRequests,
+		TEXT("Flag combined with world support to use making invisible transaction requests to the server\n")
+		TEXT("that determines whether the client should wait for the server to acknowledge visibility update before making streaming levels invisible.\n")
+		TEXT("0: Disable, 1: Enable"),
+		ECVF_Default);
+
+	static bool bDefaultAllowClientUseMakingVisibleTransactionRequests = false;
+	FAutoConsoleVariableRef CVarDefaultAllowClientUseMakingVisibleTransactionRequests(
+		TEXT("LevelStreaming.DefaultAllowClientUseMakingVisibleTransactionRequests"),
+		bDefaultAllowClientUseMakingVisibleTransactionRequests,
+		TEXT("Flag combined with world support to use making visible transaction requests to the server\n")
+		TEXT("that determines whether the client should wait for the server to acknowledge visibility update before making streaming levels visible.\n")
 		TEXT("0: Disable, 1: Enable"),
 		ECVF_Default);
 
@@ -71,9 +81,36 @@ namespace LevelStreamingCVars
 		ECVF_Default);
 }
 
-bool ULevelStreaming::ShouldClientUseMakingInvisibleTransactionRequest()
+bool ULevelStreaming::DefaultAllowClientUseMakingInvisibleTransactionRequests()
+{
+	return LevelStreamingCVars::bDefaultAllowClientUseMakingInvisibleTransactionRequests;
+}
+
+bool ULevelStreaming::DefaultAllowClientUseMakingVisibleTransactionRequests()
+{
+	return LevelStreamingCVars::bDefaultAllowClientUseMakingVisibleTransactionRequests;
+}
+
+bool ULevelStreaming::ShouldClientUseMakingInvisibleTransactionRequest() const
+{
+	if (UWorld* World = GetWorld())
+	{
+		// Rely on the world to decide whether the client should wait for the server to acknowledge
+		// visibility before making streaming levels invisible on the client.
+		return World->SupportsMakingInvisibleTransactionRequests(); 
+	}
+	return false;
+}
+
+bool ULevelStreaming::ShouldClientUseMakingVisibleTransactionRequest() const
 {
-	return LevelStreamingCVars::bShouldClientUseMakingInvisibleTransactionRequest;
+	if (UWorld* World = GetWorld())
+	{
+		// Rely on the world to decide whether the client should wait for the server to acknowledge
+		// visibility before making streaming levels visible on the client.
+		return World->SupportsMakingVisibleTransactionRequests();
+	}
+	return false;
 }
 
 bool ULevelStreaming::ShouldServerUseMakingVisibleTransactionRequest()
@@ -576,93 +613,175 @@ bool ULevelStreaming::DetermineTargetState()
 	return bContinueToConsider;
 }
 
-static bool ShouldWaitForServerAckBeforeMakingInvisible(const UWorld* World, const ULevel* LoadedLevel, const bool bShouldBeVisible, FNetLevelVisibilityState& NetVisibilityState)
+bool ULevelStreaming::IsWaitingForNetVisibilityTransactionAck(ENetLevelVisibilityRequest InRequestType) const
 {
-	if (ULevelStreaming::ShouldClientUseMakingInvisibleTransactionRequest())
+	auto IsConcernedNetVisibilityTransactionAck = [this]()
+	{
+		UWorld* World = GetWorld();
+		return LoadedLevel && !LoadedLevel->bClientOnlyVisible && World && World->IsGameWorld() && World->IsNetMode(NM_Client) && (World->NetDriver && World->NetDriver->ServerConnection->GetConnectionState() == USOCK_Open);
+	};
+
+	if (NetVisibilityState.PendingRequestType.IsSet() && (NetVisibilityState.PendingRequestType == InRequestType) && IsConcernedNetVisibilityTransactionAck())
 	{
-		if (!LoadedLevel->bClientOnlyVisible && World->IsNetMode(NM_Client) && (World->NetDriver && World->NetDriver->ServerConnection->GetConnectionState() == USOCK_Open))
+		check(((NetVisibilityState.PendingRequestType == ENetLevelVisibilityRequest::MakingInvisible) && ShouldClientUseMakingInvisibleTransactionRequest()) ||
+			  ((NetVisibilityState.PendingRequestType == ENetLevelVisibilityRequest::MakingVisible) && ShouldClientUseMakingVisibleTransactionRequest()));
+
+		return (NetVisibilityState.ClientPendingRequestIndex != NetVisibilityState.ClientAckedRequestIndex) || NetVisibilityState.bHasClientPendingRequest;
+	}
+
+	return false;
+}
+
+bool ULevelStreaming::ShouldWaitForServerAckBeforeChangingVisibilityState(ENetLevelVisibilityRequest InRequestType, bool bInShouldBeVisible)
+{
+	if (IsWaitingForNetVisibilityTransactionAck(InRequestType))
+	{
+		const bool bTargetShouldBeVisible = (NetVisibilityState.PendingRequestType == ENetLevelVisibilityRequest::MakingVisible);
+		const bool bTargetStateMatches = (bInShouldBeVisible == bTargetShouldBeVisible);
+		if (bTargetStateMatches && NetVisibilityState.bHasClientPendingRequest)
 		{
-			if (NetVisibilityState.PendingClientRequestIndex != NetVisibilityState.AckedClientRequestIndex)
-			{
-				// Wait for server to acknowledge the visibility change so that we are sure that we have no incoming data in flight when removing the actors from the world
-				return true;
-			}
-			else if (!bShouldBeVisible && NetVisibilityState.bPendingInvisibleRequest && NetVisibilityState.PendingClientRequestIndex == NetVisibilityState.AckedClientRequestIndex)
+			// We have a pending request, IncrementTransactionIndex and send ServerUpdateLevelVisibility request to server
+			FNetLevelVisibilityTransactionId TransactionId;
+			TransactionId.SetIsClientInstigator(true);
+			TransactionId.SetTransactionIndex(NetVisibilityState.ClientPendingRequestIndex);
+			NetVisibilityState.ClientAckedRequestIndex = NetVisibilityState.ClientPendingRequestIndex;
+			NetVisibilityState.ClientPendingRequestIndex = TransactionId.IncrementTransactionIndex();
+			NetVisibilityState.ClientAckedRequestCanMakeVisible.Reset();
+			NetVisibilityState.bHasClientPendingRequest = false;
+
+			UWorld* World = GetWorld();
+			const bool bIsVisible = false;
+			const bool bTryMakeVisible = (InRequestType == ENetLevelVisibilityRequest::MakingVisible);
+			for (FConstPlayerControllerIterator Iterator = World->GetPlayerControllerIterator(); Iterator; ++Iterator)
 			{
-				// We have a pending request, IncrementTransactionIndex and send ServerUpdateLevelVisibility request to server
-				FNetLevelVisibilityTransactionId TransactionId;
-				TransactionId.SetIsClientInstigator(true);
-				TransactionId.SetTransactionIndex(NetVisibilityState.PendingClientRequestIndex);
-				NetVisibilityState.PendingClientRequestIndex = TransactionId.IncrementTransactionIndex();
-				NetVisibilityState.bPendingInvisibleRequest = false;
-
-				for (FConstPlayerControllerIterator Iterator = World->GetPlayerControllerIterator(); Iterator; ++Iterator)
+				if (APlayerController* PlayerController = Iterator->Get())
 				{
-					if (APlayerController* PlayerController = Iterator->Get())
-					{
-						FUpdateLevelVisibilityLevelInfo LevelVisibility(LoadedLevel, false);
-						LevelVisibility.PackageName = PlayerController->NetworkRemapPath(LevelVisibility.PackageName, false);
-						LevelVisibility.VisibilityRequestId = TransactionId;
-							
-						PlayerController->ServerUpdateLevelVisibility(LevelVisibility);
-					}
+					FUpdateLevelVisibilityLevelInfo LevelVisibility(LoadedLevel, bIsVisible, bTryMakeVisible);
+					LevelVisibility.PackageName = PlayerController->NetworkRemapPath(LevelVisibility.PackageName, false);
+					LevelVisibility.VisibilityRequestId = TransactionId;
+					PlayerController->ServerUpdateLevelVisibility(LevelVisibility);
 				}
-				return true;
 			}
+			return true;
+		}
+		else if (NetVisibilityState.ClientPendingRequestIndex != NetVisibilityState.ClientAckedRequestIndex)
+		{
+			// Wait for server to acknowledge the visibility change
+			return true;
 		}
-	}
-
-	// Invalidate request
-	NetVisibilityState.bPendingInvisibleRequest = false;
-	NetVisibilityState.AckedClientRequestIndex = NetVisibilityState.PendingClientRequestIndex;	
 
+		// Invalidate request
+		NetVisibilityState.bHasClientPendingRequest = false;
+		NetVisibilityState.PendingRequestType.Reset();
+	}
 	return false;
-}
+};
 
-bool ULevelStreaming::IsWaitingForNetVisibilityTransactionAck() const
+bool ULevelStreaming::CanMakeInvisible()
 {
 	if (ShouldClientUseMakingInvisibleTransactionRequest())
 	{
-		if (LoadedLevel && !LoadedLevel->bClientOnlyVisible)
+		if (ShouldWaitForServerAckBeforeChangingVisibilityState(ENetLevelVisibilityRequest::MakingInvisible, bShouldBeVisible))
 		{
-			const UWorld* World = GetWorld();
-			if (World->IsNetMode(NM_Client) && (World->NetDriver && World->NetDriver->ServerConnection->GetConnectionState() == USOCK_Open) && ((NetVisibilityState.PendingClientRequestIndex != NetVisibilityState.AckedClientRequestIndex) || NetVisibilityState.bPendingInvisibleRequest))
-			{
-				return true;
-			}	
+			return false;
 		}
 	}
+	return true;
+}
 
-	return false;
+bool ULevelStreaming::CanMakeVisible()
+{
+	if (ShouldClientUseMakingVisibleTransactionRequest())
+	{
+		if (ShouldWaitForServerAckBeforeChangingVisibilityState(ENetLevelVisibilityRequest::MakingVisible, true))
+		{
+			return false;
+		}
+		else if (NetVisibilityState.ClientAckedRequestCanMakeVisible.IsSet() && !NetVisibilityState.ClientAckedRequestCanMakeVisible.GetValue())
+		{
+			// Server response was negative
+			// Until client and server streaming level state matches, client starts another visibilily request to try to make level visible
+			check(!IsWaitingForNetVisibilityTransactionAck(ENetLevelVisibilityRequest::MakingVisible));
+			BeginClientNetVisibilityRequest(true);
+			return false;
+		}
+	}
+	return true;
 }
 
 void ULevelStreaming::UpdateNetVisibilityTransactionState(bool bInShouldBeVisible, FNetLevelVisibilityTransactionId TransactionId)
 {
-	if (GetCurrentState() != ULevelStreaming::ECurrentState::MakingInvisible)
+	UWorld* World = GetWorld();
+	if (World && World->IsGameWorld())
 	{
-		// If this is a client request to unload, we will conditionally send a notification to the server before we make the level invisible
 		const bool bIsClientTransaction = TransactionId.IsClientTransaction();
+		const ENetLevelVisibilityRequest Target = bInShouldBeVisible ? ENetLevelVisibilityRequest::MakingVisible : ENetLevelVisibilityRequest::MakingInvisible;
+		// Check if client is already waiting
+		if (bIsClientTransaction && IsWaitingForNetVisibilityTransactionAck(Target))
+		{
+			return;
+		}
 
-		NetVisibilityState.bPendingInvisibleRequest = bIsClientTransaction && (bInShouldBeVisible == false);
+		NetVisibilityState.bHasClientPendingRequest = false;
+		NetVisibilityState.PendingRequestType.Reset();
 		NetVisibilityState.ServerRequestIndex = bIsClientTransaction ? FNetLevelVisibilityTransactionId::InvalidTransactionIndex : TransactionId.GetTransactionIndex();
+
+		if (bIsClientTransaction)
+		{
+			if (!bInShouldBeVisible && ShouldClientUseMakingInvisibleTransactionRequest())
+			{
+				// If this is a client request to make invisible, we will conditionally send a notification to the server before we make the level invisible
+				if (!LoadedLevel || (World->GetCurrentLevelPendingInvisibility() != LoadedLevel))
+				{
+					NetVisibilityState.bHasClientPendingRequest = true;
+					NetVisibilityState.PendingRequestType = ENetLevelVisibilityRequest::MakingInvisible;
+				}
+			}
+			else if (bInShouldBeVisible && ShouldClientUseMakingVisibleTransactionRequest())
+			{
+				// If this is a client request to make visible, we will conditionally send a notification to the server before we make the level visible
+				if (!LoadedLevel || (World->GetCurrentLevelPendingVisibility() != LoadedLevel))
+				{
+					NetVisibilityState.bHasClientPendingRequest = true;
+					NetVisibilityState.PendingRequestType = ENetLevelVisibilityRequest::MakingVisible;
+				}
+			}
+		}
 	}
 }
 
+void ULevelStreaming::OnStreamingLevelAddedToConsidered()
+{
+	// When streaming level is added to the next considered streaming levels to be updated,
+	// make sure that it will trigger a visibility request (if necessary)
+	BeginClientNetVisibilityRequest(ShouldBeVisible());
+}
+
 void ULevelStreaming::BeginClientNetVisibilityRequest(bool bInShouldBeVisible)
 {
-	FNetLevelVisibilityTransactionId TransactionId;
-	TransactionId.SetIsClientInstigator(true);
+	UWorld* World = GetWorld();
+	if (World && World->IsGameWorld() && World->IsNetMode(NM_Client))
+	{
+		FNetLevelVisibilityTransactionId TransactionId;
+		TransactionId.SetIsClientInstigator(true);
 
-	UpdateNetVisibilityTransactionState(bInShouldBeVisible, TransactionId);
+		UpdateNetVisibilityTransactionState(bInShouldBeVisible, TransactionId);
+	}
 }
 
-void ULevelStreaming::AckNetVisibilityTransaction(FNetLevelVisibilityTransactionId AckedClientTransactionId)
+void ULevelStreaming::AckNetVisibilityTransaction(FNetLevelVisibilityTransactionId InAckedClientTransactionId, bool bInClientAckCanMakeVisible)
 {
-	if (ensure(NetVisibilityState.PendingClientRequestIndex != NetVisibilityState.AckedClientRequestIndex))
+	if (ensure(NetVisibilityState.ClientPendingRequestIndex != NetVisibilityState.ClientAckedRequestIndex))
 	{
-		// Invalidate the pending request as there should only ever be a single one in flight
-		NetVisibilityState.AckedClientRequestIndex = AckedClientTransactionId.GetTransactionIndex();
-		NetVisibilityState.bPendingInvisibleRequest = false;
+		NetVisibilityState.ClientAckedRequestIndex = InAckedClientTransactionId.GetTransactionIndex();
+
+		// If received an ack for MakingVisible, store the server response in ClientAckedRequestCanMakeVisible
+		if ((NetVisibilityState.PendingRequestType == ENetLevelVisibilityRequest::MakingVisible) &&
+			(NetVisibilityState.ClientPendingRequestIndex == NetVisibilityState.ClientAckedRequestIndex))
+		{
+			check(ShouldClientUseMakingVisibleTransactionRequest());
+			NetVisibilityState.ClientAckedRequestCanMakeVisible = bInClientAckCanMakeVisible;
+		}
 	}
 }
 
@@ -721,6 +840,14 @@ void ULevelStreaming::UpdateStreamingState(bool& bOutUpdateAgain, bool& bOutRede
 				{
 					TransactionId.SetTransactionIndex(NetVisibilityState.ServerRequestIndex);
 				}
+
+				// Calling CanMakeVisible will trigger a request for visibility if necessary
+				if (!CanMakeVisible())
+				{
+					check(LoadedLevel != World->GetCurrentLevelPendingVisibility());
+					break;
+				}
+
 				World->AddToWorld(LoadedLevel, LevelTransform, !bShouldBlockOnLoad, TransactionId);
 
 				if (LoadedLevel->bIsVisible)
@@ -757,8 +884,9 @@ void ULevelStreaming::UpdateStreamingState(bool& bOutUpdateAgain, bool& bOutRede
 			const bool bWasVisible = LoadedLevel->bIsVisible;
 
 			// We do not want to have any changes in flights when ending play, so before making invisible we wait for server acknowledgment
-			if (ShouldWaitForServerAckBeforeMakingInvisible(World, LoadedLevel, bShouldBeVisible, NetVisibilityState))
+			if (!CanMakeInvisible())
 			{
+				check(LoadedLevel != World->GetCurrentLevelPendingInvisibility());
 				break;
 			}
 
diff --git a/Engine/Source/Runtime/Engine/Private/LevelUtils.cpp b/Engine/Source/Runtime/Engine/Private/LevelUtils.cpp
--- a/Engine/Source/Runtime/Engine/Private/LevelUtils.cpp
+++ b/Engine/Source/Runtime/Engine/Private/LevelUtils.cpp
@@ -12,6 +12,7 @@
 #include "Misc/FeedbackContext.h"
 #include "GameFramework/WorldSettings.h"
 #include "Components/ModelComponent.h"
+#include "Streaming/ServerStreamingLevelsVisibility.h"
 
 #if WITH_EDITOR
 #include "ScopedTransaction.h"
@@ -95,6 +96,32 @@ bool FLevelUtils::IsValidStreamingLevel(UWorld* InWorld, const TCHAR* InPackageN
 	return false;
 }
 
+bool FLevelUtils::SupportsMakingVisibleTransactionRequests(UWorld* InWorld)
+{
+	return InWorld && InWorld->SupportsMakingVisibleTransactionRequests();
+}
+
+bool FLevelUtils::SupportsMakingInvisibleTransactionRequests(UWorld* InWorld)
+{
+	return InWorld && InWorld->SupportsMakingInvisibleTransactionRequests();
+}
+
+bool FLevelUtils::IsServerStreamingLevelVisible(UWorld* InWorld, const FName& InPackageName)
+{
+	// If there's no implementation for this world to query the server visible streaming levels, return true
+	if (!SupportsMakingVisibleTransactionRequests(InWorld))
+	{
+		return true;
+	}
+
+	if (const AServerStreamingLevelsVisibility* ServerStreamingLevelsVisibility = InWorld->GetServerStreamingLevelsVisibility())
+	{
+		return ServerStreamingLevelsVisibility && ServerStreamingLevelsVisibility->Contains(InPackageName);
+	}
+
+	return false;
+}
+
 /////////////////////////////////////////////////////////////////////////////////////////
 //
 //	Level locking/unlocking.
diff --git a/Engine/Source/Runtime/Engine/Private/NetConnection.cpp b/Engine/Source/Runtime/Engine/Private/NetConnection.cpp
--- a/Engine/Source/Runtime/Engine/Private/NetConnection.cpp
+++ b/Engine/Source/Runtime/Engine/Private/NetConnection.cpp
@@ -1350,7 +1350,7 @@ FNetLevelVisibilityTransactionId UNetConnection::UpdateLevelStreamStatusChangedT
 	if (bShouldBeVisible == false)
 	{
 		ClientVisibleLevelNames.Remove(PackageName);
-
+		ClientMakingVisibleLevelNames.Remove(PackageName);
 #if UE_WITH_IRIS
 		if (UReplicationSystem* ReplicationSystem = Driver->GetReplicationSystem())
 		{
@@ -1471,6 +1471,10 @@ void UNetConnection::UpdateLevelVisibilityInternal(const FUpdateLevelVisibilityL
 
 		if (bVerifiedTransaction && VisibilityResult.bLevelExists)
 		{
+			// Verify that server knows this package
+			UE_CLOG(!FLevelUtils::IsServerStreamingLevelVisible(GetWorld(), LevelVisibility.PackageName), LogPlayerController, Warning, TEXT("ServerUpdateLevelVisibility() Added '%s', but level is not visible on server."), *LevelVisibility.PackageName.ToString());
+			
+			ClientMakingVisibleLevelNames.Remove(LevelVisibility.PackageName);
 			ClientVisibleLevelNames.Add(LevelVisibility.PackageName);
 			UE_LOG(LogPlayerController, Verbose, TEXT("ServerUpdateLevelVisibility() Added '%s'"), *LevelVisibility.PackageName.ToString());
 
@@ -1558,8 +1562,23 @@ void UNetConnection::UpdateLevelVisibilityInternal(const FUpdateLevelVisibilityL
 			}
 		}
 	}
+	else if (LevelVisibility.bTryMakeVisible)
+	{
+		if (FLevelUtils::SupportsMakingVisibleTransactionRequests(GetWorld()))
+		{
+			const FNetLevelVisibilityTransactionId VisibilityRequestId = LevelVisibility.VisibilityRequestId;
+			check(VisibilityRequestId.IsValid() && VisibilityRequestId.IsClientTransaction());
+
+			if (FLevelUtils::IsServerStreamingLevelVisible(GetWorld(), LevelVisibility.PackageName))
+			{
+				ClientMakingVisibleLevelNames.Add(LevelVisibility.PackageName);
+				check(!ClientVisibleLevelNames.Contains(LevelVisibility.PackageName));
+			}
+		}
+	}
 	else
 	{
+		ClientMakingVisibleLevelNames.Remove(LevelVisibility.PackageName);
 		ClientVisibleLevelNames.Remove(LevelVisibility.PackageName);
 		UE_LOG(LogPlayerController, Verbose, TEXT("ServerUpdateLevelVisibility() Removed '%s'"), *LevelVisibility.PackageName.ToString());
 		if (ReplicationConnectionDriver)
@@ -4686,6 +4705,7 @@ void UNetConnection::ResetGameWorldState()
 	ResetDestructionInfos();
 	ClientPendingStreamingStatusRequest.Empty();
 	ClientVisibleLevelNames.Empty();
+	ClientMakingVisibleLevelNames.Empty();
 	KeepProcessingActorChannelBunchesMap.Empty();
 	DormantReplicatorMap.Empty();
 	CleanupDormantActorState();
diff --git a/Engine/Source/Runtime/Engine/Private/PlayerController.cpp b/Engine/Source/Runtime/Engine/Private/PlayerController.cpp
--- a/Engine/Source/Runtime/Engine/Private/PlayerController.cpp
+++ b/Engine/Source/Runtime/Engine/Private/PlayerController.cpp
@@ -361,7 +361,7 @@ void APlayerController::ClientUpdateMultipleLevelsStreamingStatus_Implementation
 	}
 }
 
-void APlayerController::ClientAckUpdateLevelVisibility_Implementation(FName PackageName, FNetLevelVisibilityTransactionId TransactionId)
+void APlayerController::ClientAckUpdateLevelVisibility_Implementation(FName PackageName, FNetLevelVisibilityTransactionId TransactionId, bool bClientAckCanMakeVisibleResponse)
 {
 	if (ensureAlwaysMsgf(TransactionId.IsClientTransaction(), TEXT("APlayerController::ClientAckUpdateLevelVisibility Expected TransactionId to be ClientTransaction")))
 	{
@@ -370,7 +370,7 @@ void APlayerController::ClientAckUpdateLevelVisibility_Implementation(FName Pack
 
 		if (ULevelStreaming* LevelStreamingObject = FLevelUtils::FindStreamingLevel(GetWorld(), PackageName))
 		{
-			LevelStreamingObject->AckNetVisibilityTransaction(TransactionId);
+			FAckNetVisibilityTransaction::Call(LevelStreamingObject, TransactionId, bClientAckCanMakeVisibleResponse);
 		}
 	}
 }
@@ -407,7 +407,9 @@ void APlayerController::ServerUpdateLevelVisibility_Implementation(const FUpdate
 		// If this is a client instigated request respond with the request id so that the client knows that we have received the visibility update
 		if (LevelVisibilityCopy.VisibilityRequestId.IsClientTransaction())
 		{
-			ClientAckUpdateLevelVisibility(LevelVisibility.PackageName, LevelVisibilityCopy.VisibilityRequestId);
+			// In case a making visible request was done but querying the server is not supported, reponse true to unblock the client
+			const bool bClientAckCanMakeVisibleResponse = !FLevelUtils::SupportsMakingVisibleTransactionRequests(GetWorld()) || (LevelVisibilityCopy.bTryMakeVisible && Connection->GetClientMakingVisibleLevelNames().Contains(LevelVisibilityCopy.PackageName));
+			ClientAckUpdateLevelVisibility(LevelVisibility.PackageName, LevelVisibilityCopy.VisibilityRequestId, bClientAckCanMakeVisibleResponse);
 		}
 	}
 }
diff --git a/Engine/Source/Runtime/Engine/Private/Streaming/ServerStreamingLevelsVisibility.cpp b/Engine/Source/Runtime/Engine/Private/Streaming/ServerStreamingLevelsVisibility.cpp
new file mode 100644
--- /dev/null
+++ b/Engine/Source/Runtime/Engine/Private/Streaming/ServerStreamingLevelsVisibility.cpp
@@ -0,0 +1,58 @@
+// Copyright Epic Games, Inc. All Rights Reserved.
+
+#include "Streaming/ServerStreamingLevelsVisibility.h"
+#include "Net/UnrealNetwork.h"
+#include "Engine/LevelStreaming.h"
+#include "Engine/World.h"
+
+AServerStreamingLevelsVisibility::AServerStreamingLevelsVisibility(const FObjectInitializer& ObjectInitializer)
+	: Super(ObjectInitializer)
+{
+	bReplicates = false;
+#if WITH_EDITORONLY_DATA
+	bListedInSceneOutliner = false;
+#endif
+}
+
+AServerStreamingLevelsVisibility* AServerStreamingLevelsVisibility::SpawnServerActor(UWorld* InWorld)
+{
+	const ENetMode NetMode = InWorld->GetNetMode();
+	if (InWorld->IsGameWorld() && (NetMode == NM_DedicatedServer || NetMode == NM_ListenServer))
+	{
+		FActorSpawnParameters SpawnParams;
+		SpawnParams.Name = AServerStreamingLevelsVisibility::StaticClass()->GetFName();
+		SpawnParams.OverrideLevel = InWorld->PersistentLevel;
+		if (AServerStreamingLevelsVisibility* ServerStreamingLevelVisibility = InWorld->SpawnActor<AServerStreamingLevelsVisibility>(AServerStreamingLevelsVisibility::StaticClass(), SpawnParams))
+		{
+			for (ULevelStreaming* StreamingLevel : InWorld->GetStreamingLevels())
+			{
+				if (StreamingLevel && StreamingLevel->IsLevelVisible())
+				{
+					ServerStreamingLevelVisibility->SetIsVisible(StreamingLevel->GetWorldAssetPackageFName(), true);
+				}
+			}
+			return ServerStreamingLevelVisibility;
+		}
+	}
+	return nullptr;
+}
+
+bool AServerStreamingLevelsVisibility::Contains(const FName& InLevelPackageName) const
+{
+	return ServerVisibleLevelNames.Contains(InLevelPackageName);
+}
+
+void AServerStreamingLevelsVisibility::SetIsVisible(const FName& InLevelPackageName, bool bInIsVisible)
+{
+	if (ensure(GetLocalRole() == ROLE_Authority))
+	{
+		if (bInIsVisible)
+		{
+			ServerVisibleLevelNames.Add(InLevelPackageName);
+		}
+		else
+		{
+			ServerVisibleLevelNames.Remove(InLevelPackageName);
+		}
+	}
+}
\ No newline at end of file
diff --git a/Engine/Source/Runtime/Engine/Private/World.cpp b/Engine/Source/Runtime/Engine/Private/World.cpp
--- a/Engine/Source/Runtime/Engine/Private/World.cpp
+++ b/Engine/Source/Runtime/Engine/Private/World.cpp
@@ -95,6 +95,7 @@
 #include "Materials/MaterialParameterCollectionInstance.h"
 #include "ProfilingDebugging/LoadTimeTracker.h"
 #include "ProfilingDebugging/CsvProfiler.h"
+#include "Streaming/ServerStreamingLevelsVisibility.h"
 
 #if WITH_EDITOR
 	#include "DerivedDataCacheInterface.h"
@@ -2769,6 +2770,40 @@ static void ResetLevelFlagsOnLevelAddedToWorld(ULevel* Level)
 	Level->ResetRouteActorInitializationState();
 }
 
+bool UWorld::SupportsMakingVisibleTransactionRequests() const
+{
+	if (GetWorld()->IsGameWorld())
+	{
+		// We don't support this flag to change dynamically (read ULevelStreaming CVar once)
+		if (!bSupportsMakingVisibleTransactionRequests.IsSet())
+		{
+			bSupportsMakingVisibleTransactionRequests = ULevelStreaming::DefaultAllowClientUseMakingVisibleTransactionRequests() || (IsPartitionedWorld() && UWorldPartition::UseMakingVisibleTransactionRequests());
+		}
+		return bSupportsMakingVisibleTransactionRequests.GetValue();
+	}
+	return false;
+}
+
+bool UWorld::SupportsMakingInvisibleTransactionRequests() const
+{
+	if (GetWorld()->IsGameWorld())
+	{
+		// We don't support this flag to change dynamically (read ULevelStreaming CVar once)
+		if (!bSupportsMakingInvisibleTransactionRequests.IsSet())
+		{
+			bSupportsMakingInvisibleTransactionRequests = ULevelStreaming::DefaultAllowClientUseMakingInvisibleTransactionRequests() || (IsPartitionedWorld() && UWorldPartition::UseMakingInvisibleTransactionRequests());
+		}
+		return bSupportsMakingInvisibleTransactionRequests.GetValue();
+	}
+	return false;
+}
+
+const AServerStreamingLevelsVisibility* UWorld::GetServerStreamingLevelsVisibility() const
+{
+	check(!IsValid(ServerStreamingLevelsVisibility) || SupportsMakingVisibleTransactionRequests());
+	return ServerStreamingLevelsVisibility;
+}
+
 void UWorld::AddToWorld( ULevel* Level, const FTransform& LevelTransform, bool bConsiderTimeLimit, FNetLevelVisibilityTransactionId TransactionId)
 {
 	TRACE_CPUPROFILER_EVENT_SCOPE(AddToWorld);
@@ -2792,7 +2827,8 @@ void UWorld::AddToWorld( ULevel* Level, const FTransform& LevelTransform, bool b
 
 	if (bExecuteNextStep)
 	{
-		if (!CanAddLoadedLevelToWorld(Level))
+		// Once CurrentLevelPendingVisibility is set, don't call CanAddLoadedLevelToWorld as we must finish processing CurrentLevelPendingVisibility
+		if ((CurrentLevelPendingVisibility == nullptr) && !CanAddLoadedLevelToWorld(Level))
 		{
 			bExecuteNextStep = false;
 		}
@@ -3019,12 +3055,12 @@ void UWorld::AddToWorld( ULevel* Level, const FTransform& LevelTransform, bool b
 		// Finished making level visible - allow other levels to be added to the world.
 		CurrentLevelPendingVisibility = nullptr;
 
-		// notify server that the client has finished making this level visible
 		if (!Level->bClientOnlyVisible)
 		{
 			FUpdateLevelVisibilityLevelInfo LevelVisibility(Level, true);
 			const FName UnmappedPackageName = LevelVisibility.PackageName;
-			
+
+			// notify server that the client has finished making this level visible
 			for (FLocalPlayerIterator It(GEngine, this); It; ++It)
 			{
 				if (APlayerController* LocalPlayerController = It->GetPlayerController(this))
@@ -3034,6 +3070,12 @@ void UWorld::AddToWorld( ULevel* Level, const FTransform& LevelTransform, bool b
 					LocalPlayerController->ServerUpdateLevelVisibility(LevelVisibility);
 				}
 			}
+			// update server's visible levels
+			if (IsValid(ServerStreamingLevelsVisibility))
+			{
+				check(IsNetMode(NM_ListenServer) || IsNetMode(NM_DedicatedServer));
+				ServerStreamingLevelsVisibility->SetIsVisible(UnmappedPackageName, true);
+			}
 		}
 
 		// Set visibility before adding the rendering resource and adding to streaming.
@@ -3219,12 +3261,12 @@ void UWorld::RemoveFromWorld( ULevel* Level, bool bAllowIncrementalRemoval, FNet
 		
 			Level->ClearLevelComponents();
 
-			// notify server that the client has removed this level
 			if (!Level->bClientOnlyVisible)
 			{
 				FUpdateLevelVisibilityLevelInfo LevelVisibility(Level, false);
 				const FName UnmappedPackageName = LevelVisibility.PackageName;
 
+				// notify server that the client has removed this level
 				for (FLocalPlayerIterator It(GEngine, this); It; ++It)
 				{
 					if (APlayerController* LocalPlayerController = It->GetPlayerController(this))
@@ -3234,6 +3276,12 @@ void UWorld::RemoveFromWorld( ULevel* Level, bool bAllowIncrementalRemoval, FNet
 						LocalPlayerController->ServerUpdateLevelVisibility(LevelVisibility);
 					}
 				}
+				// update server's visible levels
+				if (IsValid(ServerStreamingLevelsVisibility))
+				{
+					check(IsNetMode(NM_ListenServer) || IsNetMode(NM_DedicatedServer));
+					ServerStreamingLevelsVisibility->SetIsVisible(UnmappedPackageName, false);
+				}
 			}
 
 			// We expect level that use the bRequireFullVisibilityToRender flag to already be !visible at this point
@@ -3771,6 +3819,7 @@ void FStreamingLevelsToConsider::Add_Internal(ULevelStreaming* StreamingLevel, b
 				};
 
 				StreamingLevels.Insert(StreamingLevel, Algo::LowerBound(StreamingLevels, StreamingLevel, PrioritySort));
+				StreamingLevel->OnStreamingLevelAddedToConsidered();
 			}
 		}
 	}
@@ -4989,6 +5038,11 @@ void UWorld::BeginPlay()
 {
 	const TArray<UWorldSubsystem*>& WorldSubsystems = SubsystemCollection.GetSubsystemArray<UWorldSubsystem>(UWorldSubsystem::StaticClass());
 
+	if (SupportsMakingVisibleTransactionRequests() && (IsNetMode(NM_DedicatedServer) || IsNetMode(NM_ListenServer)))
+	{
+		ServerStreamingLevelsVisibility = AServerStreamingLevelsVisibility::SpawnServerActor(this);
+	}
+
 	for (UWorldSubsystem* WorldSubsystem : WorldSubsystems)
 	{
 		WorldSubsystem->OnWorldBeginPlay(*this);
diff --git a/Engine/Source/Runtime/Engine/Private/WorldPartition/WorldPartition.cpp b/Engine/Source/Runtime/Engine/Private/WorldPartition/WorldPartition.cpp
--- a/Engine/Source/Runtime/Engine/Private/WorldPartition/WorldPartition.cpp
+++ b/Engine/Source/Runtime/Engine/Private/WorldPartition/WorldPartition.cpp
@@ -97,8 +97,30 @@ static FAutoConsoleVariableRef CVarEnableServerStreamingOut(
 	GEnableServerStreamingOut,
 	TEXT("Turn on/off to allow or not the server to stream out levels (only relevant when server streaming is enabled)"));
 
+bool UWorldPartition::bUseMakingVisibleTransactionRequests = false;
+FAutoConsoleVariableRef UWorldPartition::CVarUseMakingVisibleTransactionRequests(
+	TEXT("wp.Runtime.UseMakingVisibleTransactionRequests"),
+	UWorldPartition::bUseMakingVisibleTransactionRequests,
+	TEXT("Whether the client should wait for the server to acknowledge visibility update before making partitioned world streaming levels visible."));
+
+bool UWorldPartition::bUseMakingInvisibleTransactionRequests = false;
+FAutoConsoleVariableRef UWorldPartition::CVarUseMakingInvisibleTransactionRequests(
+	TEXT("wp.Runtime.UseMakingInvisibleTransactionRequests"),
+	UWorldPartition::bUseMakingInvisibleTransactionRequests,
+	TEXT("Whether the client should wait for the server to acknowledge visibility update before making partitioned world streaming levels invisible."));
+
 #define LOCTEXT_NAMESPACE "WorldPartition"
 
+bool UWorldPartition::UseMakingVisibleTransactionRequests()
+{
+	return bUseMakingVisibleTransactionRequests;
+}
+
+bool UWorldPartition::UseMakingInvisibleTransactionRequests()
+{
+	return bUseMakingInvisibleTransactionRequests;
+}
+
 #if WITH_EDITOR
 TMap<FName, FString> GetDataLayersDumpString(const UWorldPartition* WorldPartition)
 {
diff --git a/Engine/Source/Runtime/Engine/Private/WorldPartition/WorldPartitionStreamingPolicy.cpp b/Engine/Source/Runtime/Engine/Private/WorldPartition/WorldPartitionStreamingPolicy.cpp
--- a/Engine/Source/Runtime/Engine/Private/WorldPartition/WorldPartitionStreamingPolicy.cpp
+++ b/Engine/Source/Runtime/Engine/Private/WorldPartition/WorldPartitionStreamingPolicy.cpp
@@ -279,6 +279,7 @@ void UWorldPartitionStreamingPolicy::UpdateStreamingState()
 			{
 				ClientVisibleLevelNames.Add(Connection->GetClientWorldPackageName());
 				ClientVisibleLevelNames.Append(Connection->ClientVisibleLevelNames);
+				ClientVisibleLevelNames.Append(Connection->GetClientMakingVisibleLevelNames());
 			}
 		}
 
diff --git a/Engine/Source/Runtime/Engine/Public/LevelUtils.h b/Engine/Source/Runtime/Engine/Public/LevelUtils.h
--- a/Engine/Source/Runtime/Engine/Public/LevelUtils.h
+++ b/Engine/Source/Runtime/Engine/Public/LevelUtils.h
@@ -53,6 +53,21 @@ public:
 	 */
 	static bool IsValidStreamingLevel(UWorld* InWorld, const TCHAR* InPackageName);
 
+	/**
+	 * Returns whether the given package is part of the world server visible streaming levels or not.
+	 *
+	 * @param		InWorld			World to look in for the streaming level
+	 * @param		InPackageName	Name of the package containing the ULevel to query
+	 * @return						True if the given package is referenced by one of
+	 *								the world server visible streaming levels, else False.
+	 */
+	static bool IsServerStreamingLevelVisible(UWorld* InWorld, const FName& InPackageName);
+
+	/** Returns whether the world supports for a client to use "making visible" transaction requests to the server. */
+	static bool SupportsMakingVisibleTransactionRequests(UWorld* InWorld);
+
+	/** Returns whether the world supports for a client to use "making invisible" transaction requests to the server. */
+	static bool SupportsMakingInvisibleTransactionRequests(UWorld* InWorld);
 
 	///////////////////////////////////////////////////////////////////////////
 	// Locking/unlocking levels for edit.
diff --git a/Engine/Source/Runtime/Engine/Public/Streaming/ServerStreamingLevelsVisibility.h b/Engine/Source/Runtime/Engine/Public/Streaming/ServerStreamingLevelsVisibility.h
new file mode 100644
--- /dev/null
+++ b/Engine/Source/Runtime/Engine/Public/Streaming/ServerStreamingLevelsVisibility.h
@@ -0,0 +1,25 @@
+// Copyright Epic Games, Inc. All Rights Reserved.
+
+#pragma once
+
+#include "CoreMinimal.h"
+#include "UObject/ObjectMacros.h"
+#include "GameFramework/Actor.h"
+#include "ServerStreamingLevelsVisibility.generated.h"
+
+/**
+ * Actor used to replicate server's visible level streaming
+ */
+UCLASS(notplaceable, transient)
+class ENGINE_API AServerStreamingLevelsVisibility : public AActor
+{
+	GENERATED_UCLASS_BODY()
+
+public:
+	static AServerStreamingLevelsVisibility* SpawnServerActor(UWorld* World);
+	bool Contains(const FName& InPackageName) const;
+	void SetIsVisible(const FName& InLevelPackageName, bool bInIsVisible);
+
+private:
+	TSet<FName> ServerVisibleLevelNames;
+};
\ No newline at end of file
diff --git a/Engine/Source/Runtime/Engine/Public/WorldPartition/WorldPartition.h b/Engine/Source/Runtime/Engine/Public/WorldPartition/WorldPartition.h
--- a/Engine/Source/Runtime/Engine/Public/WorldPartition/WorldPartition.h
+++ b/Engine/Source/Runtime/Engine/Public/WorldPartition/WorldPartition.h
@@ -38,6 +38,7 @@ class FHLODActorDesc;
 class UHLODLayer;
 class UCanvas;
 class ULevel;
+class FAutoConsoleVariableRef;
 
 struct IWorldPartitionStreamingSourceProvider;
 
@@ -274,6 +275,9 @@ private:
 #endif
 
 public:
+	static bool UseMakingVisibleTransactionRequests();
+	static bool UseMakingInvisibleTransactionRequests();
+
 	UActorDescContainer* GetActorDescContainer() const { return ActorDescContainer; }
 
 	UPROPERTY()
@@ -322,6 +326,11 @@ private:
 	FWorldPartitionReference WorldDataLayersActor;
 #endif
 
+	static bool bUseMakingVisibleTransactionRequests;
+	static bool bUseMakingInvisibleTransactionRequests;
+	static FAutoConsoleVariableRef CVarUseMakingVisibleTransactionRequests;
+	static FAutoConsoleVariableRef CVarUseMakingInvisibleTransactionRequests;
+
 	void OnWorldMatchStarting();
 	void OnLevelRemovedFromWorld(ULevel* InLevel, UWorld* InWorld);
 	void OnPostBugItGoCalled(const FVector& Loc, const FRotator& Rot);
