Avoid clearing the ChildActor during reinstancing so that new instances of the ChildActorComponent can map from the old value of ChildActor to a new value

diff --git a/Engine/Source/Runtime/Engine/Private/Components/ChildActorComponent.cpp b/Engine/Source/Runtime/Engine/Private/Components/ChildActorComponent.cpp
--- a/Engine/Source/Runtime/Engine/Private/Components/ChildActorComponent.cpp
+++ b/Engine/Source/Runtime/Engine/Private/Components/ChildActorComponent.cpp
@@ -688,7 +688,13 @@ void UChildActorComponent::OnChildActorDestroyed(AActor* DestroyedActor)
 			}
 		}
 	}
-	ChildActor = nullptr;
+	// While reinstancing we need the reference to remain valid so that we can
+	// overwrite it when references to old instances are updated to references
+	// to new instances:
+	if (!GIsReinstancing)
+	{
+		ChildActor = nullptr;
+	}
 }
 
 void UChildActorComponent::CreateChildActor(TFunction<void(AActor*)> CustomizerFunc)
