World Partition - HLOD: Compute HLOD bounds in Z by using the subactors bounds rather than relying on the world's min/max Z

diff --git a/Engine/Source/Runtime/Engine/Private/WorldPartition/RuntimeSpatialHash/RuntimeSpatialHashHLOD.cpp b/Engine/Source/Runtime/Engine/Private/WorldPartition/RuntimeSpatialHash/RuntimeSpatialHashHLOD.cpp
--- a/Engine/Source/Runtime/Engine/Private/WorldPartition/RuntimeSpatialHash/RuntimeSpatialHashHLOD.cpp
+++ b/Engine/Source/Runtime/Engine/Private/WorldPartition/RuntimeSpatialHash/RuntimeSpatialHashHLOD.cpp
@@ -158,7 +158,6 @@ static TArray<FGuid> GenerateHLODsForGrid(UWorldPartition* WorldPartition, const
 
 		FBox2D CellBounds2D;
 		PartitionedActors.GetCellBounds(CellCoord, CellBounds2D);
-		FBox CellBounds = FBox(FVector(CellBounds2D.Min, WorldBounds.Min.Z), FVector(CellBounds2D.Max, WorldBounds.Max.Z));
 
 		FGridCellCoord CellGlobalCoord;
 		verify(PartitionedActors.GetCellGlobalCoords(CellCoord, CellGlobalCoord));
@@ -182,6 +181,28 @@ static TArray<FGuid> GenerateHLODsForGrid(UWorldPartition* WorldPartition, const
 
 				UE_LOG(LogWorldPartitionRuntimeSpatialHashHLOD, Display, TEXT("[%d / %d] Processing cell %s..."), (int32)SlowTask.CompletedWork + 1, (int32)SlowTask.TotalAmountOfWork, *CellName);
 
+				// We know the cell's 2D bound but must figure out the bounds in Z
+				FBox CellBounds = FBox(FVector(CellBounds2D.Min,  std::numeric_limits<double>::max()), 
+									   FVector(CellBounds2D.Max, -std::numeric_limits<double>::max()));
+
+				TArray<IStreamingGenerationContext::FActorInstance> ActorInstances;
+				for (const IStreamingGenerationContext::FActorSetInstance* ActorSetInstance : GridCellDataChunk.GetActorSetInstances())
+				{
+					for (const FGuid& ActorGuid : ActorSetInstance->ActorSet->Actors)
+					{
+						const IStreamingGenerationContext::FActorInstance& ActorInstance = ActorInstances.Emplace_GetRef(ActorGuid, ActorSetInstance);
+
+						const double ActorMinZ = ActorInstance.GetActorDescView().GetBounds().Min.Z;
+						const double ActorMaxZ = ActorInstance.GetActorDescView().GetBounds().Max.Z;
+
+						CellBounds.Min.Z = FMath::Min(CellBounds.Min.Z, ActorMinZ);
+						CellBounds.Max.Z = FMath::Max(CellBounds.Max.Z, ActorMaxZ);
+					}
+				}
+
+				// Ensure the Z bounds are valid
+				check(CellBounds.Min.Z <= CellBounds.Max.Z);
+
 				FHLODCreationParams CreationParams;
 				CreationParams.WorldPartition = WorldPartition;
 				CreationParams.GridIndexX = CellCoord.X;
@@ -194,16 +215,6 @@ static TArray<FGuid> GenerateHLODsForGrid(UWorldPartition* WorldPartition, const
 				CreationParams.MinVisibleDistance = RuntimeGrid.LoadingRange;
 
 				IWorldPartitionHLODUtilities* WPHLODUtilities = FModuleManager::Get().LoadModuleChecked<IWorldPartitionHLODUtilitiesModule>("WorldPartitionHLODUtilities").GetUtilities();
-
-				TArray<IStreamingGenerationContext::FActorInstance> ActorInstances;
-				for (const IStreamingGenerationContext::FActorSetInstance* ActorSetInstance : GridCellDataChunk.GetActorSetInstances())
-				{
-					for (const FGuid& ActorGuid : ActorSetInstance->ActorSet->Actors)
-					{
-						ActorInstances.Emplace(ActorGuid, ActorSetInstance);
-					}
-				}
-
 				TArray<AWorldPartitionHLOD*> CellHLODActors = WPHLODUtilities->CreateHLODActors(Context, CreationParams, ActorInstances, GridCellDataChunk.GetDataLayers());
 				if (!CellHLODActors.IsEmpty())
 				{
