Chaos : fix character movement physics logic applying forces to static objects and spamming log warnings

diff --git a/Engine/Source/Runtime/Engine/Private/Components/CharacterMovementComponent.cpp b/Engine/Source/Runtime/Engine/Private/Components/CharacterMovementComponent.cpp
--- a/Engine/Source/Runtime/Engine/Private/Components/CharacterMovementComponent.cpp
+++ b/Engine/Source/Runtime/Engine/Private/Components/CharacterMovementComponent.cpp
@@ -7025,7 +7025,7 @@ void UCharacterMovementComponent::ApplyImpactPhysicsForces(const FHitResult& Imp
 			FVector ForcePoint = Impact.ImpactPoint;
 			float BodyMass = 1.0f; // set to 1 as this is used as a multiplier
 
-			bool bCanBePushed = true;
+			bool bCanBePushed = false;
 			FBodyInstance* BI = ImpactComponent->GetBodyInstance(Impact.BoneName);
 			if(BI != nullptr && BI->IsInstanceSimulatingPhysics())
 			{
@@ -7043,11 +7043,13 @@ void UCharacterMovementComponent::ApplyImpactPhysicsForces(const FHitResult& Imp
 						ForcePoint.Z = Center.Z + Extents.Z * PushForcePointZOffsetFactor;
 					}
 				}
+
+				bCanBePushed = true;
 			}
 			else if (CharacterMovementCVars::bGeometryCollectionImpulseWorkAround)
 			{
 				const FName ClassName = ImpactComponent->GetClass()->GetFName();
-				const FName GeometryCollectionClassName("UGeometryCollectionComponent");
+				const FName GeometryCollectionClassName("GeometryCollectionComponent");
 				if (ClassName == GeometryCollectionClassName && ImpactComponent->BodyInstance.bSimulatePhysics)
 				{
 					// in some case GetBodyInstance can return null while the BodyInstance still exists ( geometry collection component for example )
@@ -7056,14 +7058,10 @@ void UCharacterMovementComponent::ApplyImpactPhysicsForces(const FHitResult& Imp
 
 					// because of the above limititation we have to ignore bPushForceUsingZOffset
 
-					bCanBePushed = true; // not necessary as this is already set by default but it's better than leaving this block empty
+					bCanBePushed = true;
 				}
 			}
-			else
-			{
-				// no body instance, not a GC, not supported scenario
-				bCanBePushed = false;
-			}
+
 			if (bCanBePushed)
 			{
 				FVector Force = Impact.ImpactNormal * -1.0f;
