Updating granular mem tracking namespace/defines

diff --git a/Engine/Source/Runtime/Engine/Private/Net/NetworkGranularMemoryLogging.cpp b/Engine/Source/Runtime/Engine/Private/Net/NetworkGranularMemoryLogging.cpp
--- a/Engine/Source/Runtime/Engine/Private/Net/NetworkGranularMemoryLogging.cpp
+++ b/Engine/Source/Runtime/Engine/Private/Net/NetworkGranularMemoryLogging.cpp
@@ -2,7 +2,7 @@
 
 #include "Net/NetworkGranularMemoryLogging.h"
 
-#if !(UE_BUILD_SHIPPING || UE_BUILD_TEST)
+#if UE_WITH_NETWORK_GRANULAR_MEM_TRACKING
 
 #include "Serialization/ArchiveCountMem.h"
 #include "EngineLogs.h"
@@ -12,7 +12,7 @@
 #include "Stats/Stats.h"
 #include "Containers/Ticker.h"
 
-namespace GranularNetworkMemoryTrackingPrivate
+namespace UE::Net::Private::GranularMemoryTracking
 {
 	static TAutoConsoleVariable<int32> CVarUseGranularNetworkTracking(
 		TEXT("Net.UseGranularNetworkTracking"),
@@ -260,4 +260,4 @@ namespace GranularNetworkMemoryTrackingPrivate
 	}
 };
 
-#endif
\ No newline at end of file
+#endif // UE_WITH_NETWORK_GRANULAR_MEM_TRACKING
\ No newline at end of file
diff --git a/Engine/Source/Runtime/Engine/Public/Net/NetworkGranularMemoryLogging.h b/Engine/Source/Runtime/Engine/Public/Net/NetworkGranularMemoryLogging.h
--- a/Engine/Source/Runtime/Engine/Public/Net/NetworkGranularMemoryLogging.h
+++ b/Engine/Source/Runtime/Engine/Public/Net/NetworkGranularMemoryLogging.h
@@ -4,9 +4,13 @@
 
 #include "CoreMinimal.h"
 
-#if !(UE_BUILD_SHIPPING || UE_BUILD_TEST)
+#ifndef UE_WITH_NETWORK_GRANULAR_MEM_TRACKING
+#define UE_WITH_NETWORK_GRANULAR_MEM_TRACKING	!(UE_BUILD_SHIPPING || UE_BUILD_TEST)
+#endif // UE_WITH_NETWORK_GRANULAR_MEM_TRACKING
 
-namespace GranularNetworkMemoryTrackingPrivate
+#if UE_WITH_NETWORK_GRANULAR_MEM_TRACKING
+
+namespace UE::Net::Private::GranularMemoryTracking
 {
 	struct ENGINE_API FScopeMarker
 	{
@@ -41,7 +45,7 @@ namespace GranularNetworkMemoryTrackingPrivate
 	};
 }
 
-#define GRANULAR_NETWORK_MEMORY_TRACKING_INIT(Archive, ScopeName) GranularNetworkMemoryTrackingPrivate::FScopeMarker GranularNetworkMemoryScope(Archive, ScopeName);
+#define GRANULAR_NETWORK_MEMORY_TRACKING_INIT(Archive, ScopeName) UE::Net::Private::GranularMemoryTracking::FScopeMarker GranularNetworkMemoryScope(Archive, ScopeName);
 #define GRANULAR_NETWORK_MEMORY_TRACKING_TRACK(Id, Work) \
 	{ \
 		GranularNetworkMemoryScope.BeginWork(); \
@@ -50,10 +54,10 @@ namespace GranularNetworkMemoryTrackingPrivate
 	}
 #define GRANULAR_NETWORK_MEMORY_TRACKING_CUSTOM_WORK(Id, Value) GranularNetworkMemoryScope.LogCustomWork(Id, Value);
 
-#else
+#else // UE_WITH_NETWORK_GRANULAR_MEM_TRACKING
 
 #define GRANULAR_NETWORK_MEMORY_TRACKING_INIT(Archive, ScopeName) 
 #define GRANULAR_NETWORK_MEMORY_TRACKING_TRACK(Id, Work) { Work; }
 #define GRANULAR_NETWORK_MEMORY_TRACKING_CUSTOM_WORK(Id, Work) 
 
-#endif
+#endif // UE_WITH_NETWORK_GRANULAR_MEM_TRACKING
