World Partition: proper fix for 22362951, add a prepare phase to content splitters.

#rnx

diff --git a/Engine/Source/Runtime/Engine/Private/WorldPartition/WorldPartition.cpp b/Engine/Source/Runtime/Engine/Private/WorldPartition/WorldPartition.cpp
--- a/Engine/Source/Runtime/Engine/Private/WorldPartition/WorldPartition.cpp
+++ b/Engine/Source/Runtime/Engine/Private/WorldPartition/WorldPartition.cpp
@@ -1396,6 +1396,12 @@ bool UWorldPartition::GatherPackagesToCook(IWorldPartitionCookPackageContext& Co
 	return false;
 }
 
+bool UWorldPartition::PrepareGeneratorPackageForCook(IWorldPartitionCookPackageContext& CookContext, TArray<UPackage*>& OutModifiedPackages)
+{
+	check(RuntimeHash);
+	return RuntimeHash->PrepareGeneratorPackageForCook(OutModifiedPackages);
+}
+
 bool UWorldPartition::PopulateGeneratorPackageForCook(IWorldPartitionCookPackageContext& CookContext, const TArray<FWorldPartitionCookPackage*>& InPackagesToCook, TArray<UPackage*>& OutModifiedPackages)
 {
 	check(RuntimeHash);
diff --git a/Engine/Source/Runtime/Engine/Private/WorldPartition/WorldPartitionCookPackageSplitter.cpp b/Engine/Source/Runtime/Engine/Private/WorldPartition/WorldPartitionCookPackageSplitter.cpp
--- a/Engine/Source/Runtime/Engine/Private/WorldPartition/WorldPartitionCookPackageSplitter.cpp
+++ b/Engine/Source/Runtime/Engine/Private/WorldPartition/WorldPartitionCookPackageSplitter.cpp
@@ -163,15 +163,11 @@ bool FWorldPartitionCookPackageSplitter::PopulateGeneratorPackage(UPackage* Owne
 	TArray<UPackage*> ModifiedPackages;
 	for (IWorldPartitionCookPackageGenerator* CookPackageGenerator : CookContext.GetCookPackageGenerators())
 	{
+		bIsSuccess &= CookPackageGenerator->PrepareGeneratorPackageForCook(CookContext, ModifiedPackages);
 		if (const TArray<FWorldPartitionCookPackage*>* CookPackages = CookContext.GetCookPackages(CookPackageGenerator))
 		{
 			bIsSuccess &= CookPackageGenerator->PopulateGeneratorPackageForCook(CookContext, *CookPackages, ModifiedPackages);
 		}
-		else
-		{
-			TArray<FWorldPartitionCookPackage*> Temp;
-			bIsSuccess &= CookPackageGenerator->PopulateGeneratorPackageForCook(CookContext, Temp, ModifiedPackages);
-		}
 	}
 
 	UE_LOG(LogWorldPartition, Log, TEXT("[Cook][PopulateGeneratorPackage] Gathered %u modified packages"), ModifiedPackages.Num());
diff --git a/Engine/Source/Runtime/Engine/Private/WorldPartition/WorldPartitionRuntimeSpatialHash.cpp b/Engine/Source/Runtime/Engine/Private/WorldPartition/WorldPartitionRuntimeSpatialHash.cpp
--- a/Engine/Source/Runtime/Engine/Private/WorldPartition/WorldPartitionRuntimeSpatialHash.cpp
+++ b/Engine/Source/Runtime/Engine/Private/WorldPartition/WorldPartitionRuntimeSpatialHash.cpp
@@ -771,9 +771,7 @@ void URuntimeSpatialHashExternalStreamingObject::OnStreamingObjectLoaded()
 				}
 			}
 		}
-
 	}
-	
 }
 
 #if WITH_EDITOR
@@ -1255,12 +1253,11 @@ TArray<UWorldPartitionRuntimeCell*> UWorldPartitionRuntimeSpatialHash::GetAlways
 	return Result;
 }
 
-bool UWorldPartitionRuntimeSpatialHash::PopulateGeneratorPackageForCook(const TArray<FWorldPartitionCookPackage*>& InPackagesToCook, TArray<UPackage*>& OutModifiedPackages)
+bool UWorldPartitionRuntimeSpatialHash::PrepareGeneratorPackageForCook(TArray<UPackage*>& OutModifiedPackages)
 {
 	check(IsRunningCookCommandlet());
 	check(ExternalStreamingObjects.IsEmpty()); // Make sure we don't have external streaming objects grids so we won't gather their cells
 
-	OutModifiedPackages.Reset();
 	for (UWorldPartitionRuntimeCell* Cell : GetAlwaysLoadedCells())
 	{
 		check(Cell->IsAlwaysLoaded());
@@ -1269,6 +1266,13 @@ bool UWorldPartitionRuntimeSpatialHash::PopulateGeneratorPackageForCook(const TA
 			return false;
 		}
 	}
+	return true;
+}
+
+bool UWorldPartitionRuntimeSpatialHash::PopulateGeneratorPackageForCook(const TArray<FWorldPartitionCookPackage*>& InPackagesToCook, TArray<UPackage*>& OutModifiedPackages)
+{
+	check(IsRunningCookCommandlet());
+	check(ExternalStreamingObjects.IsEmpty()); // Make sure we don't have external streaming objects grids so we won't gather their cells
 
 	for (const FWorldPartitionCookPackage* CookPackage : InPackagesToCook)
 	{
diff --git a/Engine/Source/Runtime/Engine/Public/WorldPartition/Cook/WorldPartitionCookPackageGenerator.h b/Engine/Source/Runtime/Engine/Public/WorldPartition/Cook/WorldPartitionCookPackageGenerator.h
--- a/Engine/Source/Runtime/Engine/Public/WorldPartition/Cook/WorldPartitionCookPackageGenerator.h
+++ b/Engine/Source/Runtime/Engine/Public/WorldPartition/Cook/WorldPartitionCookPackageGenerator.h
@@ -4,10 +4,8 @@
 #include "CoreMinimal.h"
 
 #if WITH_EDITOR
-
 struct FWorldPartitionCookPackage;
 class IWorldPartitionCookPackageContext;
-
 #endif
 
 class IWorldPartitionCookPackageGenerator
@@ -17,6 +15,7 @@ public:
 	virtual ~IWorldPartitionCookPackageGenerator() = default;
 
 	virtual bool GatherPackagesToCook(IWorldPartitionCookPackageContext& CookContext) = 0;
+	virtual bool PrepareGeneratorPackageForCook(IWorldPartitionCookPackageContext& CookContext, TArray<UPackage*>& OutModifiedPackages) { return true; }
 	virtual bool PopulateGeneratorPackageForCook(IWorldPartitionCookPackageContext& CookContext, const TArray<FWorldPartitionCookPackage*>& InPackagesToCook, TArray<UPackage*>& OutModifiedPackages) = 0;
 	virtual bool PopulateGeneratedPackageForCook(IWorldPartitionCookPackageContext& CookContext, const FWorldPartitionCookPackage& InPackagesToCook, TArray<UPackage*>& OutModifiedPackages) = 0;
 #endif
diff --git a/Engine/Source/Runtime/Engine/Public/WorldPartition/WorldPartition.h b/Engine/Source/Runtime/Engine/Public/WorldPartition/WorldPartition.h
--- a/Engine/Source/Runtime/Engine/Public/WorldPartition/WorldPartition.h
+++ b/Engine/Source/Runtime/Engine/Public/WorldPartition/WorldPartition.h
@@ -168,6 +168,7 @@ public:
 
 	//~ Begin IWorldPartitionCookPackageGenerator Interface 
 	virtual bool GatherPackagesToCook(IWorldPartitionCookPackageContext& CookContext) override;
+	virtual bool PrepareGeneratorPackageForCook(IWorldPartitionCookPackageContext& CookContext, TArray<UPackage*>& OutModifiedPackages) override;
 	virtual bool PopulateGeneratorPackageForCook(IWorldPartitionCookPackageContext& CookContext, const TArray<FWorldPartitionCookPackage*>& InPackagesToCook, TArray<UPackage*>& OutModifiedPackages) override;
 	virtual bool PopulateGeneratedPackageForCook(IWorldPartitionCookPackageContext& CookContext, const FWorldPartitionCookPackage& InPackagesToCool, TArray<UPackage*>& OutModifiedPackages) override;
 	//~ End IWorldPartitionCookPackageGenerator Interface 
diff --git a/Engine/Source/Runtime/Engine/Public/WorldPartition/WorldPartitionRuntimeHash.h b/Engine/Source/Runtime/Engine/Public/WorldPartition/WorldPartitionRuntimeHash.h
--- a/Engine/Source/Runtime/Engine/Public/WorldPartition/WorldPartitionRuntimeHash.h
+++ b/Engine/Source/Runtime/Engine/Public/WorldPartition/WorldPartitionRuntimeHash.h
@@ -69,6 +69,7 @@ class ENGINE_API UWorldPartitionRuntimeHash : public UObject, public IWorldParti
 #if WITH_EDITOR
 	virtual void SetDefaultValues() {}
 	virtual bool SupportsHLODs() const { return false; }
+	virtual bool PrepareGeneratorPackageForCook(TArray<UPackage*>& OutModifiedPackages) { return false; }
 	virtual bool PopulateGeneratorPackageForCook(const TArray<FWorldPartitionCookPackage*>& PackagesToCook, TArray<UPackage*>& OutModifiedPackages) { return false; }
 	virtual bool PopulateGeneratedPackageForCook(const FWorldPartitionCookPackage& PackagesToCook, TArray<UPackage*>& OutModifiedPackages) { return false; }
 	virtual bool GenerateStreaming(class UWorldPartitionStreamingPolicy* StreamingPolicy, const IStreamingGenerationContext* StreamingGenerationContext, TArray<FString>* OutPackagesToGenerate) { return false; }
diff --git a/Engine/Source/Runtime/Engine/Public/WorldPartition/WorldPartitionRuntimeSpatialHash.h b/Engine/Source/Runtime/Engine/Public/WorldPartition/WorldPartitionRuntimeSpatialHash.h
--- a/Engine/Source/Runtime/Engine/Public/WorldPartition/WorldPartitionRuntimeSpatialHash.h
+++ b/Engine/Source/Runtime/Engine/Public/WorldPartition/WorldPartitionRuntimeSpatialHash.h
@@ -253,6 +253,7 @@ public:
 #if WITH_EDITOR
 	virtual void SetDefaultValues() override;
 	virtual bool SupportsHLODs() const override { return true; }
+	virtual bool PrepareGeneratorPackageForCook(TArray<UPackage*>& OutModifiedPackages) override;
 	virtual bool PopulateGeneratorPackageForCook(const TArray<FWorldPartitionCookPackage*>& PackagesToCook, TArray<UPackage*>& OutModifiedPackage) override;
 	virtual bool PopulateGeneratedPackageForCook(const FWorldPartitionCookPackage& PackageToCook, TArray<UPackage*>& OutModifiedPackages) override;
 	virtual void FlushStreaming() override;
