Temporary fix for editor loading issue where static meshes may not have been serialized prior to StaticMeshComponent calling ConditionalPostLoad() on them. The root cause of this is loose enforcement of object dependencies in the editor async loader; this needs to be made more robust, but for the moment this will catch the rare cases where objects are being loaded out of order.
Clarified error report when attempting to build static mesh render data before the UStaticMesh has been serialized.

diff --git a/Engine/Source/Runtime/Engine/Private/Components/StaticMeshComponent.cpp b/Engine/Source/Runtime/Engine/Private/Components/StaticMeshComponent.cpp
--- a/Engine/Source/Runtime/Engine/Private/Components/StaticMeshComponent.cpp
+++ b/Engine/Source/Runtime/Engine/Private/Components/StaticMeshComponent.cpp
@@ -2049,9 +2049,21 @@ void UStaticMeshComponent::PostLoad()
 	NotifyIfStaticMeshChanged();
 
 	// need to postload the StaticMesh because super initializes variables based on GetStaticLightingType() which we override and use from the StaticMesh
-	if (GetStaticMesh())
+	UStaticMesh* ReferencedStaticMesh = GetStaticMesh();
+	if (ReferencedStaticMesh)
 	{
-		GetStaticMesh()->ConditionalPostLoad();
+#if WITH_EDITOR
+		// Preload the referenced static mesh here if it hasn't yet been serialized, prior to the ConditionalPostLoad call.
+		// This is temporary fix, pending a more robust solution for async loading in editor.
+		if (ReferencedStaticMesh->HasAnyFlags(RF_NeedLoad))
+		{
+			if (FLinkerLoad* Linker = ReferencedStaticMesh->GetLinker())
+			{
+				Linker->Preload(ReferencedStaticMesh);
+			}
+		}
+#endif // WITH_EDITOR
+		ReferencedStaticMesh->ConditionalPostLoad();
 	}
 
 	Super::PostLoad();
diff --git a/Engine/Source/Runtime/Engine/Private/StaticMesh.cpp b/Engine/Source/Runtime/Engine/Private/StaticMesh.cpp
--- a/Engine/Source/Runtime/Engine/Private/StaticMesh.cpp
+++ b/Engine/Source/Runtime/Engine/Private/StaticMesh.cpp
@@ -2942,6 +2942,7 @@ void FStaticMeshRenderData::Cache(const ITargetPlatform* TargetPlatform, UStatic
 			Args.Add(TEXT("StaticMeshName"), FText::FromString( Owner->GetName() ) );
 			FStaticMeshStatusMessageContext StatusContext( FText::Format( NSLOCTEXT("Engine", "BuildingStaticMeshStatus", "Building static mesh {StaticMeshName}..."), Args ) );
 
+			checkf(!Owner->HasAnyFlags(RF_NeedLoad), TEXT("StaticMesh %s being PostLoaded before having been serialized - this suggests an async loading problem."), *GetPathNameSafe(Owner));
 			checkf(Owner->IsMeshDescriptionValid(0), TEXT("Bad MeshDescription on %s"), *GetPathNameSafe(Owner));
 
 			if (Owner->bDoFastBuild)
