Added support for skeletalmeshes, and improved static mesh instance support, in the moderation artifact generation commandlet

diff --git a/Engine/Source/Runtime/Engine/Classes/Components/MeshComponent.h b/Engine/Source/Runtime/Engine/Classes/Components/MeshComponent.h
--- a/Engine/Source/Runtime/Engine/Classes/Components/MeshComponent.h
+++ b/Engine/Source/Runtime/Engine/Classes/Components/MeshComponent.h
@@ -79,6 +79,11 @@ public:
 	 */
 	void EmptyOverrideMaterials();
 
+	/**
+	 * Returns true if there are any override materials set for this component
+	 */
+	bool HasOverrideMaterials();
+
 #if WITH_EDITOR
 	virtual void PostEditChangeChainProperty(FPropertyChangedChainEvent& PropertyChangedEvent) override;
 #endif
diff --git a/Engine/Source/Runtime/Engine/Private/Components/MeshComponent.cpp b/Engine/Source/Runtime/Engine/Private/Components/MeshComponent.cpp
--- a/Engine/Source/Runtime/Engine/Private/Components/MeshComponent.cpp
+++ b/Engine/Source/Runtime/Engine/Private/Components/MeshComponent.cpp
@@ -187,6 +187,11 @@ void UMeshComponent::EmptyOverrideMaterials()
 	}
 }
 
+bool UMeshComponent::HasOverrideMaterials()
+{
+	return OverrideMaterials.Num() > 0;
+}
+
 int32 UMeshComponent::GetNumMaterials() const
 {
 	return 0;
