Made Content Bundle client resilient to asymetric call of content injection

diff --git a/Engine/Source/Runtime/Engine/Private/WorldPartition/ContentBundle/ContentBundleBase.cpp b/Engine/Source/Runtime/Engine/Private/WorldPartition/ContentBundle/ContentBundleBase.cpp
--- a/Engine/Source/Runtime/Engine/Private/WorldPartition/ContentBundle/ContentBundleBase.cpp
+++ b/Engine/Source/Runtime/Engine/Private/WorldPartition/ContentBundle/ContentBundleBase.cpp
@@ -52,6 +52,11 @@ void FContentBundleBase::InjectContent()
 	DoInjectContent();
 
 	check(GetStatus() == EContentBundleStatus::ReadyToInject || GetStatus() == EContentBundleStatus::ContentInjected || GetStatus() == EContentBundleStatus::FailedToInject);
+
+	if (TSharedPtr<FContentBundleClient> PinnedClient = Client.Pin())
+	{
+		PinnedClient->OnContentInjectedInWorld(GetStatus(), GetInjectedWorld());
+	}
 }
 
 void FContentBundleBase::RemoveContent()
@@ -64,6 +69,11 @@ void FContentBundleBase::RemoveContent()
 	}
 
 	check(GetStatus() == EContentBundleStatus::Registered);
+
+	if (TSharedPtr<FContentBundleClient> PinnedClient = Client.Pin())
+	{
+		PinnedClient->OnContentRemovedFromWorld(GetStatus(), GetInjectedWorld());
+	}
 }
 
 const TWeakPtr<FContentBundleClient>& FContentBundleBase::GetClient() const
diff --git a/Engine/Source/Runtime/Engine/Private/WorldPartition/ContentBundle/ContentBundleClient.cpp b/Engine/Source/Runtime/Engine/Private/WorldPartition/ContentBundle/ContentBundleClient.cpp
--- a/Engine/Source/Runtime/Engine/Private/WorldPartition/ContentBundle/ContentBundleClient.cpp
+++ b/Engine/Source/Runtime/Engine/Private/WorldPartition/ContentBundle/ContentBundleClient.cpp
@@ -3,6 +3,7 @@
 
 #include "WorldPartition/ContentBundle/ContentBundleEngineSubsystem.h"
 #include "WorldPartition/ContentBundle/ContentBundleDescriptor.h"
+#include "WorldPartition/ContentBundle/ContentBundleLog.h"
 #include "Engine/Engine.h"
 
 TSharedPtr<FContentBundleClient> FContentBundleClient::CreateClient(const UContentBundleDescriptor* InContentBundleDescriptor, FString const& InDisplayName)
@@ -11,29 +12,104 @@ TSharedPtr<FContentBundleClient> FContentBundleClient::CreateClient(const UConte
 }
 
 FContentBundleClient::FContentBundleClient(const UContentBundleDescriptor* InContentBundleDescriptor, FString const& InDisplayName)
-	:ContentBundleDescriptor(InContentBundleDescriptor),
-#if WITH_EDITOR
-	DisplayName(InDisplayName),
-#endif
-	ContentInjectionState(EContentBundleClientState::Unregistered)
+	:ContentBundleDescriptor(InContentBundleDescriptor)
+	, DisplayName(InDisplayName)
+	, State(EContentBundleClientState::Unregistered)
 {
-
+	UE_LOG(LogContentBundle, Log, TEXT("[CB: %s] Client Created"), *GetDescriptor()->GetDisplayName());
+	SetState(EContentBundleClientState::Registered);
 }
 
 void FContentBundleClient::RequestContentInjection()
-{
-	GEngine->GetEngineSubsystem<UContentBundleEngineSubsystem>()->RequestContentInjection(*this);
-	ContentInjectionState = EContentBundleClientState::ContentInjectionRequested;
+{ 
+	if (State == EContentBundleClientState::Registered)
+	{
+		SetState(EContentBundleClientState::ContentInjectionRequested);
+		GEngine->GetEngineSubsystem<UContentBundleEngineSubsystem>()->RequestContentInjection(*this);
+	}
+	else
+	{
+		UE_LOG(LogContentBundle, Log, TEXT("[CB: %s] Client failed to request content injection. Its state is %s"), *GetDescriptor()->GetDisplayName(), *UEnum::GetDisplayValueAsText(State).ToString());
+	}
 }
 
 void FContentBundleClient::RequestRemoveContent()
 {
-	ContentInjectionState = EContentBundleClientState::ContentRemovalRequested;
-	GEngine->GetEngineSubsystem<UContentBundleEngineSubsystem>()->RequestContentRemoval(*this);
+	if (State == EContentBundleClientState::ContentInjectionRequested)
+	{
+		SetState(EContentBundleClientState::ContentRemovalRequested);
+		GEngine->GetEngineSubsystem<UContentBundleEngineSubsystem>()->RequestContentRemoval(*this);
+	}
+	else
+	{
+		UE_LOG(LogContentBundle, Log, TEXT("[CB: %s] Client failed to request content removal. Its state is %s"), *GetDescriptor()->GetDisplayName(), *UEnum::GetDisplayValueAsText(State).ToString());
+	}
 }
 
 void FContentBundleClient::RequestUnregister()
 {
-	ContentInjectionState = EContentBundleClientState::ContentRemovalRequested;
-	GEngine->GetEngineSubsystem<UContentBundleEngineSubsystem>()->UnregisterContentBundle(*this);
+	if (State != EContentBundleClientState::Unregistered)
+	{
+		GEngine->GetEngineSubsystem<UContentBundleEngineSubsystem>()->UnregisterContentBundle(*this);
+		SetState(EContentBundleClientState::Unregistered);
+	}
+	else
+	{
+		UE_LOG(LogContentBundle, Log, TEXT("[CB: %s] Client failed to request unregister. Its state is %s"), *GetDescriptor()->GetDisplayName(), *UEnum::GetDisplayValueAsText(State).ToString());
+	}
+}
+	
+
+void FContentBundleClient::OnContentInjectedInWorld(EContentBundleStatus InjectionStatus, UWorld* InjectedWorld)
+{
+	check(State == EContentBundleClientState::ContentInjectionRequested);
+
+	if (InjectionStatus == EContentBundleStatus::ContentInjected || InjectionStatus == EContentBundleStatus::ReadyToInject)
+	{
+		SetWorldContentState(InjectedWorld, EWorldContentState::ContentBundleInjected);
+	}
+	else if (InjectionStatus == EContentBundleStatus::FailedToInject)
+	{
+		SetWorldContentState(InjectedWorld, EWorldContentState::ContentBundleInjected);
+	}
+	else
+	{
+		// Injection status unhandled
+		check(0);
+	}
+}
+
+void FContentBundleClient::OnContentRemovedFromWorld(EContentBundleStatus RemovalStatus, UWorld* InjectedWorld)
+{
+	if (RemovalStatus == EContentBundleStatus::Registered)
+	{
+		SetWorldContentState(InjectedWorld, EWorldContentState::NoContent);
+	}
+	else
+	{
+		// Removal status unhandled
+		check(0);
+	}
+
+}
+
+void FContentBundleClient::SetWorldContentState(UWorld* World, EWorldContentState NewState)
+{
+	if (EWorldContentState* OldState = WorldContentStates.Find(World))
+	{
+		UE_LOG(LogContentBundle, Log, TEXT("[CB: %s] Client WorldState for world %s changing from %s to %s"), *GetDescriptor()->GetDisplayName(), *World->GetName(), *UEnum::GetDisplayValueAsText(*OldState).ToString(), *UEnum::GetDisplayValueAsText(NewState).ToString());
+		*OldState = NewState;
+		return;
+	}
+	
+	UE_LOG(LogContentBundle, Log, TEXT("[CB: %s] Client WorldState for world %s changing to %s"), *GetDescriptor()->GetDisplayName(), *World->GetName(), *UEnum::GetDisplayValueAsText(NewState).ToString());
+	WorldContentStates.Add(World, NewState);
+}
+
+void FContentBundleClient::SetState(EContentBundleClientState NewState)
+{
+	check(NewState != State);
+
+	UE_LOG(LogContentBundle, Log, TEXT("[CB: %s] Client State changing from %s to %s"), *GetDescriptor()->GetDisplayName(), *UEnum::GetDisplayValueAsText(State).ToString(), *UEnum::GetDisplayValueAsText(NewState).ToString());
+	State = NewState;
 }
\ No newline at end of file
diff --git a/Engine/Source/Runtime/Engine/Private/WorldPartition/ContentBundle/ContentBundleContainer.cpp b/Engine/Source/Runtime/Engine/Private/WorldPartition/ContentBundle/ContentBundleContainer.cpp
--- a/Engine/Source/Runtime/Engine/Private/WorldPartition/ContentBundle/ContentBundleContainer.cpp
+++ b/Engine/Source/Runtime/Engine/Private/WorldPartition/ContentBundle/ContentBundleContainer.cpp
@@ -310,7 +310,7 @@ void FContentBundleContainer::UnregisterContentBundleClientEvents()
 
 bool FContentBundleContainer::ShouldInjectClientContent(const TSharedPtr<FContentBundleClient>& ContentBundleClient) const
 {
-	return ContentBundleClient->GetState() == EContentBundleClientState::ContentInjected || ContentBundleClient->GetState() == EContentBundleClientState::ContentInjectionRequested;
+	return ContentBundleClient->GetState() == EContentBundleClientState::ContentInjectionRequested;
 }
 
 void FContentBundleContainer::InitializeContentBundlesForegisteredClients()
diff --git a/Engine/Source/Runtime/Engine/Public/WorldPartition/ContentBundle/ContentBundleClient.h b/Engine/Source/Runtime/Engine/Public/WorldPartition/ContentBundle/ContentBundleClient.h
--- a/Engine/Source/Runtime/Engine/Public/WorldPartition/ContentBundle/ContentBundleClient.h
+++ b/Engine/Source/Runtime/Engine/Public/WorldPartition/ContentBundle/ContentBundleClient.h
@@ -5,6 +5,8 @@
 #include "UObject/ObjectMacros.h"
 #include "UObject/WeakObjectPtr.h"
 
+#include "WorldPartition/ContentBundle/ContentBundleStatus.h"
+
 class UContentBundleDescriptor;
 
 UENUM()
@@ -13,17 +15,24 @@ enum class EContentBundleClientState
 	Unregistered,
 	Registered,
 	ContentInjectionRequested,
-	ContentInjected,
 	ContentRemovalRequested,
 	
 	// Failed state
 	RegistrationFailed,
-	ContentInjectionFailed,
+};
+
+UENUM()
+enum class EWorldContentState
+{
+	NoContent,
+	ContentBundleInjected,
+	ContentBundleInjectionFailed
 };
 
 class ENGINE_API FContentBundleClient
 {
 	friend class UContentBundleEngineSubsystem;
+	friend class FContentBundleBase;
 
 public:
 	static TSharedPtr<FContentBundleClient> CreateClient(const UContentBundleDescriptor* InContentBundleDescriptor, FString const& InDisplayName);
@@ -38,14 +47,22 @@ public:
 	
 	void RequestUnregister();
 
-	EContentBundleClientState GetState() const { return ContentInjectionState; }
+	EContentBundleClientState GetState() const { return State; }
 
 	FString const& GetDisplayName() const { return DisplayName; }
 
 private:
+	void SetState(EContentBundleClientState State);
+	void SetWorldContentState(UWorld* World, EWorldContentState State);
+
+	void OnContentInjectedInWorld(EContentBundleStatus InjectionStatus, UWorld* InjectedWorld);
+	void OnContentRemovedFromWorld(EContentBundleStatus RemovalStatus, UWorld* InjectedWorld);
+
 	TWeakObjectPtr<const UContentBundleDescriptor> ContentBundleDescriptor;
+	
+	TMap<TWeakObjectPtr<UWorld>, EWorldContentState> WorldContentStates;
 
 	FString DisplayName;
 
-	EContentBundleClientState ContentInjectionState;
+	EContentBundleClientState State;
 };
\ No newline at end of file
