Add forward declaration

diff --git a/Engine/Source/Runtime/Engine/Public/WorldPartition/ContentBundle/ContentBundleClient.h b/Engine/Source/Runtime/Engine/Public/WorldPartition/ContentBundle/ContentBundleClient.h
--- a/Engine/Source/Runtime/Engine/Public/WorldPartition/ContentBundle/ContentBundleClient.h
+++ b/Engine/Source/Runtime/Engine/Public/WorldPartition/ContentBundle/ContentBundleClient.h
@@ -8,6 +8,7 @@
 #include "WorldPartition/ContentBundle/ContentBundleStatus.h"
 
 class UContentBundleDescriptor;
+class UWorld;
 
 UENUM()
 enum class EContentBundleClientState
