Do prompt the hlod subsystem to see if a ContentBundle cell can be unloaded.
Content Bundle do not have Hlods and Hlods assumes knowledge of all cells.

diff --git a/Engine/Source/Runtime/Engine/Private/WorldPartition/WorldPartitionRuntimeLevelStreamingCell.cpp b/Engine/Source/Runtime/Engine/Private/WorldPartition/WorldPartitionRuntimeLevelStreamingCell.cpp
--- a/Engine/Source/Runtime/Engine/Private/WorldPartition/WorldPartitionRuntimeLevelStreamingCell.cpp
+++ b/Engine/Source/Runtime/Engine/Private/WorldPartition/WorldPartitionRuntimeLevelStreamingCell.cpp
@@ -371,7 +371,9 @@ ULevel* UWorldPartitionRuntimeLevelStreamingCell::GetLevel() const
 
 bool UWorldPartitionRuntimeLevelStreamingCell::CanUnload() const
 {
-	if (LevelStreaming)
+	// @todo_ow ContentBundles do not support Hlods and events below are forwarding
+	// Events to the HLodSubsystem which assumes knowledge of all cells (not true with plugins)
+	if (LevelStreaming && !GetContentBundleID().IsValid())
 	{
 		if (UHLODSubsystem* HLODSubsystem = LevelStreaming->GetWorld()->GetSubsystem<UHLODSubsystem>())
 		{
