Horde build failure - missing parameter on Streaming Source construction. Due to back out issues.

diff --git a/Engine/Source/Runtime/Engine/Private/WorldPartition/WorldPartitionStreamingPolicy.cpp b/Engine/Source/Runtime/Engine/Private/WorldPartition/WorldPartitionStreamingPolicy.cpp
--- a/Engine/Source/Runtime/Engine/Private/WorldPartition/WorldPartitionStreamingPolicy.cpp
+++ b/Engine/Source/Runtime/Engine/Private/WorldPartition/WorldPartitionStreamingPolicy.cpp
@@ -143,7 +143,8 @@ void UWorldPartitionStreamingPolicy::UpdateStreamingSources()
 							const EStreamingSourceTargetState TargetState = PlayerController->StreamingSourceShouldActivate() ? EStreamingSourceTargetState::Activated : EStreamingSourceTargetState::Loaded;
 							const bool bBlockOnSlowLoading = PlayerController->StreamingSourceShouldBlockOnSlowStreaming();
 							const EStreamingSourcePriority StreamingSourcePriority = PlayerController->GetStreamingSourcePriority();
-							StreamingSources.Add(FWorldPartitionStreamingSource(PlayerController->GetFName(), ViewLocation, ViewRotation, TargetState, bBlockOnSlowLoading, StreamingSourcePriority));
+							const bool bRemote = !PlayerController->IsLocalController();
+							StreamingSources.Add(FWorldPartitionStreamingSource(PlayerController->GetFName(), ViewLocation, ViewRotation, TargetState, bBlockOnSlowLoading, StreamingSourcePriority, bRemote));
 						}
 					}
 				}
