Offload PSO collection and Shader decompression for PSO initializer create to background task

#rnx

diff --git a/Engine/Source/Runtime/Engine/Private/Materials/MaterialShared.cpp b/Engine/Source/Runtime/Engine/Private/Materials/MaterialShared.cpp
--- a/Engine/Source/Runtime/Engine/Private/Materials/MaterialShared.cpp
+++ b/Engine/Source/Runtime/Engine/Private/Materials/MaterialShared.cpp
@@ -45,6 +45,7 @@
 #include "HAL/FileManager.h"
 #include "ProfilingDebugging/LoadTimeTracker.h"
 #include "UObject/CoreRedirects.h"
+#include "UObject/StrongObjectPtr.h"
 #include "RayTracingDefinitions.h"
 #include "Interfaces/ITargetPlatform.h"
 #include "Misc/ConfigCacheIni.h"
@@ -2640,10 +2641,93 @@ bool FMaterial::CacheShaders(const FMaterialShaderMapId& ShaderMapId, EShaderPla
 #endif
 }
 
+/**
+ * Helper task used to release the strong object reference to the material interface on the game thread
+ * The release has to happen on the gamethread and the material interface can't be GCd while the PSO
+ * collection is happening because it touches the material resources
+ */
+class FMaterialInterfaceReleaseTask
+{
+public:
+	explicit FMaterialInterfaceReleaseTask(TStrongObjectPtr<UMaterialInterface>* InMaterialInterface)
+		: MaterialInterface(InMaterialInterface)
+	{
+	}
+
+	void DoTask(ENamedThreads::Type CurrentThread, const FGraphEventRef& MyCompletionGraphEvent)
+	{
+		check(IsInGameThread());
+		delete MaterialInterface;
+	}
+
+public:
+
+	TStrongObjectPtr<UMaterialInterface>* MaterialInterface;
+
+	static ESubsequentsMode::Type	GetSubsequentsMode() { return ESubsequentsMode::FireAndForget; }
+	ENamedThreads::Type				GetDesiredThread() { return ENamedThreads::GameThread; }
+	FORCEINLINE TStatId				GetStatId() const { return TStatId(); }
+};
+
+/**
+ * Helper task used to offload the PSO collection from the GameThread. The shader decompression
+ * takes too long to run this on the GameThread and it isn't blocking anything crucial.
+ * The graph event used to create this task is extended with the PSO compilation tasks itself so the user can optionally
+ * wait or known when all PSOs are ready for rendering
+ */
+class FMaterialPSOPrecacheCollectionTask
+{
+public:
+	explicit FMaterialPSOPrecacheCollectionTask(
+		TStrongObjectPtr<UMaterialInterface>* InMaterialInterface,
+		FMaterialShaderMap* InMaterialShaderMap,
+		ERHIFeatureLevel::Type InFeatureLevel,
+		const FMaterial* InMaterial,
+		const TConstArrayView<const FVertexFactoryType*>& InVertexFactoryTypes,
+		const FPSOPrecacheParams& InPreCacheParams)
+		: MaterialInterface(InMaterialInterface)
+		, MaterialShaderMap(InMaterialShaderMap)
+		, FeatureLevel(InFeatureLevel)
+		, Material(InMaterial)
+		, VertexFactoryTypes(InVertexFactoryTypes)
+		, PreCacheParams(InPreCacheParams)
+	{
+	}
+
+	void DoTask(ENamedThreads::Type CurrentThread, const FGraphEventRef& MyCompletionGraphEvent)
+	{
+		FTaskTagScope ParallelGTScope(ETaskTag::EParallelGameThread);
+
+		FGraphEventArray PSOCompileGraphEvents = MaterialShaderMap->CollectPSOs(FeatureLevel, Material, VertexFactoryTypes, PreCacheParams);
+		
+		// Won't touch the material interface anymore - PSO compile jobs take refs to all RHI resources while creating the task
+		TGraphTask<FMaterialInterfaceReleaseTask>::CreateTask().ConstructAndDispatchWhenReady(MaterialInterface);
+
+		// Extend MyCompletionGraphEvent to wait for all the async compile events
+		for (FGraphEventRef& GraphEvent : PSOCompileGraphEvents)
+		{
+			MyCompletionGraphEvent->DontCompleteUntil(GraphEvent);
+		}
+	}
+
+public:
+
+	TStrongObjectPtr<UMaterialInterface>* MaterialInterface;
+	FMaterialShaderMap* MaterialShaderMap;
+	ERHIFeatureLevel::Type FeatureLevel;
+	const FMaterial* Material;
+	TArray<const FVertexFactoryType*, TInlineAllocator<4>> VertexFactoryTypes;
+	const FPSOPrecacheParams PreCacheParams;
+
+	static ESubsequentsMode::Type	GetSubsequentsMode() { return ESubsequentsMode::TrackSubsequents; }
+	ENamedThreads::Type				GetDesiredThread() { return ENamedThreads::AnyBackgroundThreadNormalTask; }
+	FORCEINLINE TStatId				GetStatId() const { return TStatId(); }
+};
+
 FGraphEventArray FMaterial::CollectPSOs(ERHIFeatureLevel::Type InFeatureLevel, const TConstArrayView<const FVertexFactoryType*>& VertexFactoryTypes, const FPSOPrecacheParams& PreCacheParams)
 {
 	TRACE_CPUPROFILER_EVENT_SCOPE(FMaterial::CollectPSOs);
-
+	
 	FGraphEventArray GraphEvents;
 
 	// Only care about inline shaders for now
@@ -2691,7 +2775,31 @@ FGraphEventArray FMaterial::CollectPSOs(ERHIFeatureLevel::Type InFeatureLevel, c
 			}
 		}
 
-		GraphEvents.Append(GameThreadShaderMap->CollectPSOs(InFeatureLevel, this, ActualMissingVFs, PreCacheParams));
+		if (ActualMissingVFs.Num() > 0)
+		{
+			// Offload to background job task graph if threading is enabled
+			// Don't use background thread in editor because shader maps and material resources could be destroyed while the task is running
+			// If it's a perf problem at some point then OutPSOCollectionGraphEvent has to be used at material level in the correct places to wait for
+			bool bUseBackgroundTask = FApp::ShouldUseThreadingForPerformance() && !GIsEditor;
+#if PSO_PRECACHING_VALIDATE
+			// when validation is enabled then we want to make sure that the precache stats are updated before the mesh draw commands can be build
+			bUseBackgroundTask = bUseBackgroundTask && !PSOCollectorStats::IsPrecachingValidationEnabled();
+#endif // PSO_PRECACHING_VALIDATE
+			if (bUseBackgroundTask)
+			{
+				// Make sure the material instance isn't garbage collected or destroyed yet (create TStrongObjectPtr which will be destroyed on the GT when the collection is done)
+				TStrongObjectPtr<UMaterialInterface>* MaterialInterface = new TStrongObjectPtr<UMaterialInterface>(GetMaterialInterface());
+
+				// Create and kick the collection task
+				FGraphEventRef GraphEvent = TGraphTask<FMaterialPSOPrecacheCollectionTask>::CreateTask().ConstructAndDispatchWhenReady(
+					MaterialInterface, GameThreadShaderMap, InFeatureLevel, this, ActualMissingVFs, PreCacheParams);
+				GraphEvents.Add(GraphEvent);
+			}
+			else
+			{				
+				GraphEvents = GameThreadShaderMap->CollectPSOs(InFeatureLevel, this, ActualMissingVFs, PreCacheParams);
+			}
+		}
 	}
 
 	return GraphEvents;
diff --git a/Engine/Source/Runtime/Renderer/Private/BasePassRendering.cpp b/Engine/Source/Runtime/Renderer/Private/BasePassRendering.cpp
--- a/Engine/Source/Runtime/Renderer/Private/BasePassRendering.cpp
+++ b/Engine/Source/Runtime/Renderer/Private/BasePassRendering.cpp
@@ -1989,7 +1989,7 @@ TArray<ELightMapPolicyType, TInlineAllocator<2>> FBasePassMeshProcessor::GetUnif
 	const bool bUseVolumetricLightmap = true;// Scene&& Scene->VolumetricLightmapSceneData.HasData();
 
 	static const auto CVarSupportLightMapPolicyMode = IConsoleManager::Get().FindTConsoleVariableDataInt(TEXT("r.PSOPrecache.LightMapPolicyMode"));
-	const bool bNoLightMapOnlyMode = (CVarSupportLightMapPolicyMode && CVarSupportLightMapPolicyMode->GetValueOnGameThread() == 1);
+	const bool bNoLightMapOnlyMode = (CVarSupportLightMapPolicyMode && CVarSupportLightMapPolicyMode->GetValueOnAnyThread() == 1);
 
 	TArray<ELightMapPolicyType, TInlineAllocator<2>> UniformLightMapPolicyTypes;
 
@@ -2090,7 +2090,7 @@ void FBasePassMeshProcessor::CollectPSOInitializersForSkyLight(
 	const bool bIsLitMaterial = ShadingModels.IsLit();
 	
 	static const auto CVarSupportLightMapPolicyMode = IConsoleManager::Get().FindTConsoleVariableDataInt(TEXT("r.PSOPrecache.LightMapPolicyMode"));
-	const bool bNoLightMapOnlyMode = (CVarSupportLightMapPolicyMode && CVarSupportLightMapPolicyMode->GetValueOnGameThread() == 1);
+	const bool bNoLightMapOnlyMode = (CVarSupportLightMapPolicyMode && CVarSupportLightMapPolicyMode->GetValueOnAnyThread() == 1);
 
 	if (!bNoLightMapOnlyMode && bIsLitMaterial && bIsTranslucent)
 	{
