- Handle cancel of an already acknowledged response from the server to make visible a streaming level.
- Changed client instigated NetVisibility request to be set when streaming level current state changes to MakingVisible or MakingInvisible (more robust).
- Small refactoring to invalidate streaming level NetLevelVisibilityState.

diff --git a/Engine/Source/Runtime/Engine/Classes/Engine/LevelStreaming.h b/Engine/Source/Runtime/Engine/Classes/Engine/LevelStreaming.h
--- a/Engine/Source/Runtime/Engine/Classes/Engine/LevelStreaming.h
+++ b/Engine/Source/Runtime/Engine/Classes/Engine/LevelStreaming.h
@@ -37,6 +37,12 @@ struct FNetLevelVisibilityState
 	uint32 ClientPendingRequestIndex = 0;
 	uint32 ClientAckedRequestIndex = 0;
 	bool bHasClientPendingRequest = false;
+
+	void InvalidateClientPendingRequest()
+	{
+		bHasClientPendingRequest = false;
+		PendingRequestType.Reset();
+	}
 };
 
 ENGINE_API DECLARE_LOG_CATEGORY_EXTERN(LogLevelStreaming, Log, All);
@@ -332,9 +338,6 @@ private:
 
 public:
 
-	/** Called when streaming level is being considered by world for next UpdateStreamingState. */
-	void OnStreamingLevelAddedToConsidered();
-
 	/** Begin a client instigated NetVisibility request */
 	void BeginClientNetVisibilityRequest(bool bInShouldBeVisible);
 
@@ -649,6 +652,12 @@ private:
 	/** State for server handshake of NetLevelVisibilty requests */
 	FNetLevelVisibilityState NetVisibilityState;
 
+	/** Whether streaming level is concerned by net visibility transactions */
+	bool IsConcernedByNetVisibilityTransactionAck() const;
+
+	/** Helper method that updates server level visibility for each player controller of streaming level world */
+	void ServerUpdateLevelVisibility(bool bIsVisible, bool bTryMakeVisible = false, FNetLevelVisibilityTransactionId TransactionId = FNetLevelVisibilityTransactionId());
+
 	friend struct FAckNetVisibilityTransaction;
 	friend struct FStreamingLevelPrivateAccessor;
 };
diff --git a/Engine/Source/Runtime/Engine/Private/LevelStreaming.cpp b/Engine/Source/Runtime/Engine/Private/LevelStreaming.cpp
--- a/Engine/Source/Runtime/Engine/Private/LevelStreaming.cpp
+++ b/Engine/Source/Runtime/Engine/Private/LevelStreaming.cpp
@@ -613,15 +613,15 @@ bool ULevelStreaming::DetermineTargetState()
 	return bContinueToConsider;
 }
 
-bool ULevelStreaming::IsWaitingForNetVisibilityTransactionAck(ENetLevelVisibilityRequest InRequestType) const
+bool ULevelStreaming::IsConcernedByNetVisibilityTransactionAck() const
 {
-	auto IsConcernedNetVisibilityTransactionAck = [this]()
-	{
-		UWorld* World = GetWorld();
-		return LoadedLevel && !LoadedLevel->bClientOnlyVisible && World && World->IsGameWorld() && World->IsNetMode(NM_Client) && (World->NetDriver && World->NetDriver->ServerConnection->GetConnectionState() == USOCK_Open);
-	};
+	UWorld* World = GetWorld();
+	return LoadedLevel && !LoadedLevel->bClientOnlyVisible && World && World->IsGameWorld() && World->IsNetMode(NM_Client) && (World->NetDriver && World->NetDriver->ServerConnection->GetConnectionState() == USOCK_Open);
+}
 
-	if (NetVisibilityState.PendingRequestType.IsSet() && (NetVisibilityState.PendingRequestType == InRequestType) && IsConcernedNetVisibilityTransactionAck())
+bool ULevelStreaming::IsWaitingForNetVisibilityTransactionAck(ENetLevelVisibilityRequest InRequestType) const
+{
+	if (NetVisibilityState.PendingRequestType.IsSet() && (NetVisibilityState.PendingRequestType == InRequestType) && IsConcernedByNetVisibilityTransactionAck())
 	{
 		check(((NetVisibilityState.PendingRequestType == ENetLevelVisibilityRequest::MakingInvisible) && ShouldClientUseMakingInvisibleTransactionRequest()) ||
 			  ((NetVisibilityState.PendingRequestType == ENetLevelVisibilityRequest::MakingVisible) && ShouldClientUseMakingVisibleTransactionRequest()));
@@ -632,6 +632,24 @@ bool ULevelStreaming::IsWaitingForNetVisibilityTransactionAck(ENetLevelVisibilit
 	return false;
 }
 
+void ULevelStreaming::ServerUpdateLevelVisibility(bool bIsVisible, bool bTryMakeVisible, FNetLevelVisibilityTransactionId TransactionId)
+{
+	if (IsConcernedByNetVisibilityTransactionAck())
+	{
+		UWorld* World = GetWorld();
+		for (FConstPlayerControllerIterator Iterator = World->GetPlayerControllerIterator(); Iterator; ++Iterator)
+		{
+			if (APlayerController* PlayerController = Iterator->Get())
+			{
+				FUpdateLevelVisibilityLevelInfo LevelVisibility(LoadedLevel, bIsVisible, bTryMakeVisible);
+				LevelVisibility.PackageName = PlayerController->NetworkRemapPath(LevelVisibility.PackageName, false);
+				LevelVisibility.VisibilityRequestId = TransactionId;
+				PlayerController->ServerUpdateLevelVisibility(LevelVisibility);
+			}
+		}
+	}
+}
+
 bool ULevelStreaming::ShouldWaitForServerAckBeforeChangingVisibilityState(ENetLevelVisibilityRequest InRequestType, bool bInShouldBeVisible)
 {
 	if (IsWaitingForNetVisibilityTransactionAck(InRequestType))
@@ -649,19 +667,9 @@ bool ULevelStreaming::ShouldWaitForServerAckBeforeChangingVisibilityState(ENetLe
 			NetVisibilityState.ClientAckedRequestCanMakeVisible.Reset();
 			NetVisibilityState.bHasClientPendingRequest = false;
 
-			UWorld* World = GetWorld();
 			const bool bIsVisible = false;
 			const bool bTryMakeVisible = (InRequestType == ENetLevelVisibilityRequest::MakingVisible);
-			for (FConstPlayerControllerIterator Iterator = World->GetPlayerControllerIterator(); Iterator; ++Iterator)
-			{
-				if (APlayerController* PlayerController = Iterator->Get())
-				{
-					FUpdateLevelVisibilityLevelInfo LevelVisibility(LoadedLevel, bIsVisible, bTryMakeVisible);
-					LevelVisibility.PackageName = PlayerController->NetworkRemapPath(LevelVisibility.PackageName, false);
-					LevelVisibility.VisibilityRequestId = TransactionId;
-					PlayerController->ServerUpdateLevelVisibility(LevelVisibility);
-				}
-			}
+			ServerUpdateLevelVisibility(bIsVisible, bTryMakeVisible, TransactionId);
 			return true;
 		}
 		else if (NetVisibilityState.ClientPendingRequestIndex != NetVisibilityState.ClientAckedRequestIndex)
@@ -671,8 +679,7 @@ bool ULevelStreaming::ShouldWaitForServerAckBeforeChangingVisibilityState(ENetLe
 		}
 
 		// Invalidate request
-		NetVisibilityState.bHasClientPendingRequest = false;
-		NetVisibilityState.PendingRequestType.Reset();
+		NetVisibilityState.InvalidateClientPendingRequest();
 	}
 	return false;
 };
@@ -722,8 +729,7 @@ void ULevelStreaming::UpdateNetVisibilityTransactionState(bool bInShouldBeVisibl
 			return;
 		}
 
-		NetVisibilityState.bHasClientPendingRequest = false;
-		NetVisibilityState.PendingRequestType.Reset();
+		NetVisibilityState.InvalidateClientPendingRequest();
 		NetVisibilityState.ServerRequestIndex = bIsClientTransaction ? FNetLevelVisibilityTransactionId::InvalidTransactionIndex : TransactionId.GetTransactionIndex();
 
 		if (bIsClientTransaction)
@@ -750,13 +756,6 @@ void ULevelStreaming::UpdateNetVisibilityTransactionState(bool bInShouldBeVisibl
 	}
 }
 
-void ULevelStreaming::OnStreamingLevelAddedToConsidered()
-{
-	// When streaming level is added to the next considered streaming levels to be updated,
-	// make sure that it will trigger a visibility request (if necessary)
-	BeginClientNetVisibilityRequest(ShouldBeVisible());
-}
-
 void ULevelStreaming::BeginClientNetVisibilityRequest(bool bInShouldBeVisible)
 {
 	UWorld* World = GetWorld();
@@ -831,6 +830,11 @@ void ULevelStreaming::UpdateStreamingState(bool& bOutUpdateAgain, bool& bOutRede
 				CurrentState = ECurrentState::LoadedNotVisible;
 				bOutUpdateAgain = true;
 				bOutRedetermineTarget = true;
+				// Make sure to update level visibility state (in case the server already acknowledged a Making Visible request)
+				if (ShouldClientUseMakingVisibleTransactionRequest())
+				{
+					ServerUpdateLevelVisibility(false);
+				}
 			}
 			else
 			{
@@ -953,6 +957,9 @@ void ULevelStreaming::UpdateStreamingState(bool& bOutUpdateAgain, bool& bOutRede
 		{
 		case ETargetState::LoadedVisible:
 			CurrentState = ECurrentState::MakingVisible;
+			// Make sure client pending visibility request (if any) matches current state
+			NetVisibilityState.InvalidateClientPendingRequest();
+			BeginClientNetVisibilityRequest(ShouldBeVisible());
 			bOutUpdateAgain = true;
 			break;
 
@@ -998,6 +1005,9 @@ void ULevelStreaming::UpdateStreamingState(bool& bOutUpdateAgain, bool& bOutRede
 		{
 		case ETargetState::LoadedNotVisible:
 			CurrentState = ECurrentState::MakingInvisible;
+			// Make sure client pending visibility request (if any) matches current state
+			NetVisibilityState.InvalidateClientPendingRequest();
+			BeginClientNetVisibilityRequest(ShouldBeVisible());
 			bOutUpdateAgain = true;
 			break;
 
diff --git a/Engine/Source/Runtime/Engine/Private/World.cpp b/Engine/Source/Runtime/Engine/Private/World.cpp
--- a/Engine/Source/Runtime/Engine/Private/World.cpp
+++ b/Engine/Source/Runtime/Engine/Private/World.cpp
@@ -3819,7 +3819,6 @@ void FStreamingLevelsToConsider::Add_Internal(ULevelStreaming* StreamingLevel, b
 				};
 
 				StreamingLevels.Insert(StreamingLevel, Algo::LowerBound(StreamingLevels, StreamingLevel, PrioritySort));
-				StreamingLevel->OnStreamingLevelAddedToConsidered();
 			}
 		}
 	}
