Engine: add UPluginBlueprintLibrary for querying info about plugins from Blueprint and Python

diff --git a/Engine/Source/Runtime/Engine/Private/PluginBlueprintLibrary.cpp b/Engine/Source/Runtime/Engine/Private/PluginBlueprintLibrary.cpp
new file mode 100644
--- /dev/null
+++ b/Engine/Source/Runtime/Engine/Private/PluginBlueprintLibrary.cpp
@@ -0,0 +1,193 @@
+// Copyright Epic Games, Inc. All Rights Reserved.
+
+#include "PluginBlueprintLibrary.h"
+
+#include "Containers/Array.h"
+#include "Containers/Set.h"
+#include "Containers/UnrealString.h"
+#include "Interfaces/IPluginManager.h"
+#include "Interfaces/IProjectManager.h"
+#include "Templates/SharedPointer.h"
+#include "UObject/SoftObjectPath.h"
+
+
+TArray<FString> UPluginBlueprintLibrary::GetAdditionalPluginSearchPaths()
+{
+	TArray<FString> PluginSearchPaths;
+
+	const TSet<FString>& SearchPathsSet =
+		IPluginManager::Get().GetAdditionalPluginSearchPaths();
+	for (const FString& SearchPath : SearchPathsSet)
+	{
+		PluginSearchPaths.Add(SearchPath);
+	}
+
+	return PluginSearchPaths;
+}
+
+const TArray<FString>&
+UPluginBlueprintLibrary::GetAdditionalProjectPluginSearchPaths()
+{
+	return IProjectManager::Get().GetAdditionalPluginDirectories();
+}
+
+TArray<FString> UPluginBlueprintLibrary::GetEnabledPluginNames()
+{
+	const TArray<TSharedRef<IPlugin>> EnabledPlugins =
+		IPluginManager::Get().GetEnabledPlugins();
+
+	TArray<FString> PluginNames;
+	PluginNames.Reserve(EnabledPlugins.Num());
+
+	for (const TSharedRef<IPlugin>& Plugin : EnabledPlugins)
+	{
+		PluginNames.Add(Plugin->GetName());
+	}
+
+	return PluginNames;
+}
+
+bool UPluginBlueprintLibrary::GetPluginNameForObjectPath(
+		const FSoftObjectPath& ObjectPath,
+		FString& OutPluginName)
+{
+	const TSharedPtr<IPlugin> Plugin =
+		IPluginManager::Get().FindPluginFromPath(
+			ObjectPath.GetAssetPathString());
+	if (!Plugin)
+	{
+		return false;
+	}
+
+	OutPluginName = Plugin->GetName();
+
+	return true;
+}
+
+bool UPluginBlueprintLibrary::GetPluginDescriptorFilePath(
+		const FString& PluginName,
+		FString& OutFilePath)
+{
+	const TSharedPtr<IPlugin> Plugin =
+		IPluginManager::Get().FindPlugin(PluginName);
+	if (!Plugin)
+	{
+		return false;
+	}
+
+	OutFilePath = Plugin->GetDescriptorFileName();
+
+	return true;
+}
+
+bool UPluginBlueprintLibrary::GetPluginBaseDir(
+		const FString& PluginName,
+		FString& OutBaseDir)
+{
+	const TSharedPtr<IPlugin> Plugin =
+		IPluginManager::Get().FindPlugin(PluginName);
+	if (!Plugin)
+	{
+		return false;
+	}
+
+	OutBaseDir = Plugin->GetBaseDir();
+
+	return true;
+}
+
+bool UPluginBlueprintLibrary::GetPluginContentDir(
+		const FString& PluginName,
+		FString& OutContentDir)
+{
+	const TSharedPtr<IPlugin> Plugin =
+		IPluginManager::Get().FindPlugin(PluginName);
+	if (!Plugin)
+	{
+		return false;
+	}
+
+	OutContentDir = Plugin->GetContentDir();
+
+	return true;
+}
+
+bool UPluginBlueprintLibrary::GetPluginMountedAssetPath(
+		const FString& PluginName,
+		FString& OutAssetPath)
+{
+	const TSharedPtr<IPlugin> Plugin =
+		IPluginManager::Get().FindPlugin(PluginName);
+	if (!Plugin)
+	{
+		return false;
+	}
+
+	OutAssetPath = Plugin->GetMountedAssetPath();
+
+	return true;
+}
+
+bool UPluginBlueprintLibrary::GetPluginVersion(
+		const FString& PluginName,
+		int32& OutVersion)
+{
+	const TSharedPtr<IPlugin> Plugin =
+		IPluginManager::Get().FindPlugin(PluginName);
+	if (!Plugin)
+	{
+		return false;
+	}
+
+	OutVersion = Plugin->GetDescriptor().Version;
+
+	return true;
+}
+
+bool UPluginBlueprintLibrary::GetPluginVersionName(
+		const FString& PluginName,
+		FString& OutVersionName)
+{
+	const TSharedPtr<IPlugin> Plugin =
+		IPluginManager::Get().FindPlugin(PluginName);
+	if (!Plugin)
+	{
+		return false;
+	}
+
+	OutVersionName = Plugin->GetDescriptor().VersionName;
+
+	return true;
+}
+
+bool UPluginBlueprintLibrary::GetPluginDescription(
+		const FString& PluginName,
+		FString& OutDescription)
+{
+	const TSharedPtr<IPlugin> Plugin =
+		IPluginManager::Get().FindPlugin(PluginName);
+	if (!Plugin)
+	{
+		return false;
+	}
+
+	OutDescription = Plugin->GetDescriptor().Description;
+
+	return true;
+}
+
+bool UPluginBlueprintLibrary::GetPluginEditorCustomVirtualPath(
+		const FString& PluginName,
+		FString& OutVirtualPath)
+{
+	const TSharedPtr<IPlugin> Plugin =
+		IPluginManager::Get().FindPlugin(PluginName);
+	if (!Plugin)
+	{
+		return false;
+	}
+
+	OutVirtualPath = Plugin->GetDescriptor().EditorCustomVirtualPath;
+
+	return true;
+}
diff --git a/Engine/Source/Runtime/Engine/Public/PluginBlueprintLibrary.h b/Engine/Source/Runtime/Engine/Public/PluginBlueprintLibrary.h
new file mode 100644
--- /dev/null
+++ b/Engine/Source/Runtime/Engine/Public/PluginBlueprintLibrary.h
@@ -0,0 +1,173 @@
+// Copyright Epic Games, Inc. All Rights Reserved.
+
+#pragma once
+
+#include "Containers/Array.h"
+#include "Containers/UnrealString.h"
+#include "Kismet/BlueprintFunctionLibrary.h"
+#include "UObject/ObjectMacros.h"
+#include "UObject/SoftObjectPath.h"
+
+#include "PluginBlueprintLibrary.generated.h"
+
+
+/**
+ * A function library of utilities for querying information about plugins.
+ */
+UCLASS(MinimalAPI)
+class UPluginBlueprintLibrary : public UBlueprintFunctionLibrary
+{
+	GENERATED_BODY()
+
+public:
+
+	/**
+	 * Get the list of extra directories that are recursively searched for
+	 * plugins (aside from the engine and project plugin directories).
+	 *
+	 * @return The additional filesystem plugin search paths.
+	 */
+	UFUNCTION(BlueprintCallable, Category="Engine Scripting | Plugin Utilities")
+	static TArray<FString> GetAdditionalPluginSearchPaths();
+
+	/**
+	 * Get the list of extra directories added by the project that are
+	 * recursively searched for plugins.
+	 *
+	 * @return The additional project filesystem plugin search paths.
+	 */
+	UFUNCTION(BlueprintCallable, Category="Engine Scripting | Plugin Utilities")
+	static const TArray<FString>& GetAdditionalProjectPluginSearchPaths();
+
+	/**
+	 * Get the names of all enabled plugins.
+	 *
+	 * @return The names of all enabled plugins.
+	 */
+	UFUNCTION(BlueprintCallable, Category="Engine Scripting | Plugin Utilities")
+	static TArray<FString> GetEnabledPluginNames();
+
+	/**
+	 * Get the name of the plugin containing an object.
+	 *
+	 * @param ObjectPath - Path to the object
+	 * @param OutPluginName - Name of the plugin containing the object, if found
+	 *
+	 * @return true if the object is contained within a plugin and the plugin
+	 *         name was stored in OutPluginName, or false otherwise
+	 */
+	UFUNCTION(BlueprintCallable, Category="Engine Scripting | Plugin Utilities")
+	static bool GetPluginNameForObjectPath(
+			const FSoftObjectPath& ObjectPath,
+			FString& OutPluginName);
+
+	/**
+	 * Get the filesystem path to a plugin's descriptor.
+	 *
+	 * @param PluginName - Name of the plugin
+	 * @param OutFilePath - Filesystem path to the plugin's descriptor, if found
+	 *
+	 * @return true if the named plugin was found and the plugin descriptor
+	 *         filesystem path was stored in OutFilePath, or false otherwise
+	 */
+	UFUNCTION(BlueprintCallable, Category="Engine Scripting | Plugin Utilities")
+	static bool GetPluginDescriptorFilePath(
+			const FString& PluginName,
+			FString& OutFilePath);
+
+	/**
+	 * Get the filesystem path to a plugin's base directory.
+	 *
+	 * @param PluginName - Name of the plugin
+	 * @param OutBaseDir - Filesystem path to the plugin's base directory, if found
+	 *
+	 * @return true if the named plugin was found and the plugin base directory
+	 *         filesystem path was stored in OutBaseDir, or false otherwise
+	 */
+	UFUNCTION(BlueprintCallable, Category="Engine Scripting | Plugin Utilities")
+	static bool GetPluginBaseDir(
+			const FString& PluginName,
+			FString& OutBaseDir);
+
+	/**
+	 * Get the filesystem path to a plugin's content directory.
+	 *
+	 * @param PluginName - Name of the plugin
+	 * @param OutContentDir - Filesystem path to the plugin's content directory, if found
+	 *
+	 * @return true if the named plugin was found and the plugin content
+	 *         directory filesystem path was stored in OutContentDir, or false otherwise
+	 */
+	UFUNCTION(BlueprintCallable, Category="Engine Scripting | Plugin Utilities")
+	static bool GetPluginContentDir(
+			const FString& PluginName,
+			FString& OutContentDir);
+
+	/**
+	 * Get the virtual root path for assets in a plugin.
+	 *
+	 * @param PluginName - Name of the plugin
+	 * @param OutAssetPath - Virtual root path for the plugin's assets, if found
+	 *
+	 * @return true if the named plugin was found and the plugin's virtual
+	 *         root path was stored in OutAssetPath, or false otherwise
+	 */
+	UFUNCTION(BlueprintCallable, Category="Engine Scripting | Plugin Utilities")
+	static bool GetPluginMountedAssetPath(
+			const FString& PluginName,
+			FString& OutAssetPath);
+
+	/**
+	 * Get the version number of a plugin.
+	 *
+	 * @param PluginName - Name of the plugin
+	 * @param OutVersion - Version number of the plugin, if found
+	 *
+	 * @return true if the named plugin was found and the plugin's version
+	 *         number was stored in OutVersion, or false otherwise
+	 */
+	UFUNCTION(BlueprintCallable, Category="Engine Scripting | Plugin Utilities")
+	static bool GetPluginVersion(const FString& PluginName, int32& OutVersion);
+
+	/**
+	 * Get the version name of a plugin.
+	 *
+	 * @param PluginName - Name of the plugin
+	 * @param OutVersionName - Version name of the plugin, if found
+	 *
+	 * @return true if the named plugin was found and the plugin's version
+	 *         name was stored in OutVersionName, or false otherwise
+	 */
+	UFUNCTION(BlueprintCallable, Category="Engine Scripting | Plugin Utilities")
+	static bool GetPluginVersionName(
+			const FString& PluginName,
+			FString& OutVersionName);
+
+	/**
+	 * Get the description of a plugin.
+	 *
+	 * @param PluginName - Name of the plugin
+	 * @param OutDescription - Description of the plugin, if found
+	 *
+	 * @return true if the named plugin was found and the plugin's description
+	 *         was stored in OutDescription, or false otherwise
+	 */
+	UFUNCTION(BlueprintCallable, Category="Engine Scripting | Plugin Utilities")
+	static bool GetPluginDescription(
+			const FString& PluginName,
+			FString& OutDescription);
+
+	/**
+	 * Get the editor custom virtual path of a plugin.
+	 *
+	 * @param PluginName - Name of the plugin
+	 * @param OutVirtualPath - Editor custom virtual path of the plugin, if found
+	 *
+	 * @return true if the named plugin was found and the plugin's editor
+	 *         custom virtual path was stored in OutVirtualPath, or false otherwise
+	 */
+	UFUNCTION(BlueprintCallable, Category="Engine Scripting | Plugin Utilities")
+	static bool GetPluginEditorCustomVirtualPath(
+			const FString& PluginName,
+			FString& OutVirtualPath);
+};
