Don't try to apply resolution changes on platforms with fixed resolution, since all it does it cause an unnecessary reallocation of the backbuffer, which can fail in memory-constrained situations.

#rnx

diff --git a/Engine/Source/Runtime/Engine/Private/GameUserSettings.cpp b/Engine/Source/Runtime/Engine/Private/GameUserSettings.cpp
--- a/Engine/Source/Runtime/Engine/Private/GameUserSettings.cpp
+++ b/Engine/Source/Runtime/Engine/Private/GameUserSettings.cpp
@@ -75,7 +75,7 @@ FIntPoint UGameUserSettings::GetDesktopResolution() const
 
 void UGameUserSettings::SetScreenResolution(FIntPoint Resolution)
 {
-	if (ResolutionSizeX != Resolution.X || ResolutionSizeY != Resolution.Y)
+	if (!FPlatformProperties::HasFixedResolution() && (ResolutionSizeX != Resolution.X || ResolutionSizeY != Resolution.Y))
 	{
 		ResolutionSizeX = Resolution.X;
 		ResolutionSizeY = Resolution.Y;
@@ -102,6 +102,11 @@ EWindowMode::Type UGameUserSettings::GetLastConfirmedFullscreenMode() const
 
 void UGameUserSettings::SetFullscreenMode(EWindowMode::Type InFullscreenMode)
 {
+	if (FPlatformProperties::HasFixedResolution())
+	{
+		return;
+	}
+
 	if (FullscreenMode != InFullscreenMode)
 	{
 		switch (InFullscreenMode)
@@ -478,11 +483,14 @@ void UGameUserSettings::ApplyNonResolutionSettings()
 
 void UGameUserSettings::ApplyResolutionSettings(bool bCheckForCommandLineOverrides)
 {
-#if UE_SERVER
-	return;
-#endif
+#if !UE_SERVER
 	QUICK_SCOPE_CYCLE_COUNTER(GameUserSettings_ApplyResolutionSettings);
 
+	if (FPlatformProperties::HasFixedResolution())
+	{
+		return;
+	}
+
 	ValidateSettings();
 
 	EWindowMode::Type NewFullscreenMode = GetFullscreenMode();
@@ -496,6 +504,7 @@ void UGameUserSettings::ApplyResolutionSettings(bool bCheckForCommandLineOverrid
 	}
 
 	IConsoleManager::Get().CallAllConsoleVariableSinks();
+#endif
 }
 
 void UGameUserSettings::ApplySettings(bool bCheckForCommandLineOverrides)
@@ -545,6 +554,11 @@ void UGameUserSettings::LoadSettings(bool bForceReload/*=false*/)
 
 void UGameUserSettings::RequestResolutionChange(int32 InResolutionX, int32 InResolutionY, EWindowMode::Type InWindowMode, bool bInDoOverrides /* = true */)
 {
+	if (FPlatformProperties::HasFixedResolution())
+	{
+		return;
+	}
+
 	if (bInDoOverrides)
 	{
 		UGameEngine::ConditionallyOverrideSettings(InResolutionX, InResolutionY, InWindowMode);
