Skeleton Compatibility: Fix some issues where filters were not displaying the right assets as the Skeleton wasn't loaded yet.

diff --git a/Engine/Source/Runtime/Engine/Private/Animation/Skeleton.cpp b/Engine/Source/Runtime/Engine/Private/Animation/Skeleton.cpp
--- a/Engine/Source/Runtime/Engine/Private/Animation/Skeleton.cpp
+++ b/Engine/Source/Runtime/Engine/Private/Animation/Skeleton.cpp
@@ -126,7 +126,7 @@ bool USkeleton::IsCompatible(const USkeleton* InSkeleton) const
 		return true;
 	}
 
-	for (const auto& CompatibleSkeleton : CompatibleSkeletons)
+	for (const TSoftObjectPtr<USkeleton>& CompatibleSkeleton : CompatibleSkeletons)
 	{
 		if (CompatibleSkeleton == InSkeleton->CachedSoftObjectPtr)
 		{
@@ -146,17 +146,17 @@ bool USkeleton::IsCompatibleSkeletonByAssetString(const FString& SkeletonAssetSt
 		return true;
 	}
 
-	// Now check against the list of compatible skeletons.
-	FString CompatibleName;
-	for (const auto& CompatibleSkeleton : CompatibleSkeletons)
+	// Get the asset registry.
+	const FAssetRegistryModule& AssetRegistryModule = FModuleManager::LoadModuleChecked<FAssetRegistryModule>("AssetRegistry");
+	const IAssetRegistry& AssetRegistry = AssetRegistryModule.Get();
+
+	// Now check against the list of compatible skeletons and see if we're dealing with the same asset.
+	for (const TSoftObjectPtr<USkeleton>& CompatibleSkeleton : CompatibleSkeletons)
 	{
-		if (CompatibleSkeleton.IsValid())
+		const FAssetData AssetData = AssetRegistry.GetAssetByObjectPath(CompatibleSkeleton.ToString());
+		if (AssetData.IsValid() && AssetData.GetExportTextName() == SkeletonAssetString)
 		{
-			CompatibleName = FObjectPropertyBase::GetExportPath(CompatibleSkeleton.Get());
-			if (CompatibleName == SkeletonAssetString)
-			{
-				return true;
-			}
+			return true;
 		}
 	}
 
@@ -600,10 +600,10 @@ void USkeleton::PostLoad()
 		OnSmartNamesChangedEvent.AddUObject(this, &USkeleton::HandleSmartNamesChangedEvent);
 	}
 	
-	// Cleanup CompatibleSkeletons for convenience, this does not ensure CompatibleSkeletons is free of invalid TSoftObjectPtrs at access-time.
+	// Cleanup CompatibleSkeletons for convenience. This basically removes any soft object pointers that has an invalid soft object name.
 	CompatibleSkeletons = CompatibleSkeletons.FilterByPredicate([](const TSoftObjectPtr<USkeleton>& Skeleton)
 	{
-		return Skeleton.IsValid();
+		return Skeleton.ToSoftObjectPath().IsValid();
 	});
 }
 
