Texture: remove cached LOD Bias, was often out of date
calculate LOD Bias as needed on demand
fix calculations of cinematic mip bias
clarify effective mip count should not go below numnonstreaming or numintail
RequiresBlock4Alignment to identify pixel formats that need 4-alignment (eg. not ASTC/ETC)
use RequiresBlock4Alignment to make NumNonStreaming only when necessary
so mobile platform streaming is not constrained when not needed

diff --git a/Engine/Source/Runtime/Core/Public/PixelFormat.h b/Engine/Source/Runtime/Core/Public/PixelFormat.h
--- a/Engine/Source/Runtime/Core/Public/PixelFormat.h
+++ b/Engine/Source/Runtime/Core/Public/PixelFormat.h
@@ -240,6 +240,29 @@ FORCEINLINE bool IsInteger(EPixelFormat PixelFormat)
 	return false;
 }
 
+FORCEINLINE bool IsDXTCBlockCompressedTextureFormat(EPixelFormat PixelFormat)
+{
+	switch (PixelFormat)
+	{
+	case PF_DXT1:
+	case PF_DXT3:
+	case PF_DXT5:
+	case PF_BC4:
+	case PF_BC5:
+	case PF_BC6H:
+	case PF_BC7:
+		return true;
+	default:
+		return false;
+	}
+}
+
+FORCEINLINE bool RequiresBlock4Alignment(EPixelFormat PixelFormat)
+{
+	// BCN DXTC formats require 4x4 alignment, but ASTC/ETC do not
+	return IsDXTCBlockCompressedTextureFormat(PixelFormat);
+}
+
 /** 
 * Information about a pixel format. The majority of this structure is valid after static init, however RHI does keep
 * some state in here that is initialized by that module and should not be used by general use programs that don't
diff --git a/Engine/Source/Runtime/Engine/Classes/Engine/StreamableRenderAsset.h b/Engine/Source/Runtime/Engine/Classes/Engine/StreamableRenderAsset.h
--- a/Engine/Source/Runtime/Engine/Classes/Engine/StreamableRenderAsset.h
+++ b/Engine/Source/Runtime/Engine/Classes/Engine/StreamableRenderAsset.h
@@ -162,12 +162,13 @@ public:
 	ENGINE_API void SetForceMipLevelsToBeResident(float Seconds, int32 CinematicLODGroupMask = 0);
 
 	/**
-	* Returns the cached combined LOD bias based on texture LOD group and LOD bias.
+	* Returns the combined LOD bias based on texture LOD group and LOD bias.
+	* Function name is legacy and incorrect, it is no longer cached
 	* @return	LOD bias
 	*/
-	ENGINE_API int32 GetCachedLODBias() const 
+	ENGINE_API virtual int32 GetCachedLODBias() const 
 	{ 
-		return CachedCombinedLODBias; 
+		return 0; 
 	}
 
 	/** Return the streaming state of the render resources. Mirrors the state and lowers cache misses. Cleared if there are no resources. */
@@ -263,10 +264,6 @@ protected:
 	UPROPERTY(transient, duplicatetransient, NonTransactional)
 	int32 StreamingIndex = INDEX_NONE;
 
-	/** Cached combined group and texture LOD bias to use.	*/
-	UPROPERTY(transient)
-	int32 CachedCombinedLODBias;
-
 public:
 	UPROPERTY(EditAnywhere, BlueprintReadWrite, Category = LevelOfDetail, AssetRegistrySearchable, AdvancedDisplay)
 	uint8 NeverStream : 1;
diff --git a/Engine/Source/Runtime/Engine/Classes/Engine/Texture.h b/Engine/Source/Runtime/Engine/Classes/Engine/Texture.h
--- a/Engine/Source/Runtime/Engine/Classes/Engine/Texture.h
+++ b/Engine/Source/Runtime/Engine/Classes/Engine/Texture.h
@@ -1453,7 +1453,25 @@ public:
 	virtual class FTextureResource* CreateResource() PURE_VIRTUAL(UTexture::CreateResource,return NULL;);
 
 	/** Cache the combined LOD bias based on texture LOD group and LOD bias. */
-	ENGINE_API void UpdateCachedLODBias();
+	UE_DEPRECATED(5.2, "UpdateCachedLODBias does nothing, remove call")
+	ENGINE_API void UpdateCachedLODBias()
+	{
+		// no longer cached, now does nothing
+	}
+	
+	ENGINE_API int32 GetCachedLODBias() const override
+	{
+		// this is the combined LOD Bias with cinematic bias
+		return CalculateLODBias(true);
+	}
+
+	/**
+	* Calculate the combined LOD bias based on texture LOD group and LOD bias.
+	*   CalculateLODBias(true) >= CalculateLODBias(false)
+	*   CalculateLODBias(true) == CalculateLODBias(false) + NumCinematicMipLevels , except when a clamp was applied
+	* @return	LOD bias
+	*/
+	ENGINE_API int32 CalculateLODBias(bool bWithCinematicMipBias) const;
 
 	/**
 	 * @return The material value type of this texture.
diff --git a/Engine/Source/Runtime/Engine/Classes/Engine/Texture2D.h b/Engine/Source/Runtime/Engine/Classes/Engine/Texture2D.h
--- a/Engine/Source/Runtime/Engine/Classes/Engine/Texture2D.h
+++ b/Engine/Source/Runtime/Engine/Classes/Engine/Texture2D.h
@@ -155,7 +155,10 @@ public:
 
 
 	/**
-	 * Calculates the maximum number of mips the engine allows to be loaded for this texture. 
+	 * Calculates the maximum number of mips that will be in this texture after cooking
+	 *   (eg. after the "drop mip" lod bias is applied).
+	 * This function is not correct and should not be used. 
+	 *
 	 * The cinematic mips will be considered as loadable, streaming enabled or not.
 	 * Note that in the cooking process, mips smaller than the min residency count
 	 * can be stripped out by the cooker.
@@ -166,7 +169,9 @@ public:
 	ENGINE_API int32 GetNumMipsAllowed(bool bIgnoreMinResidency) const;
 
 public:
-	/** Returns the minimum number of mips that must be resident in memory (cannot be streamed). */
+	/** Returns the minimum number of mips that must be resident in memory (cannot be streamed). 
+	This does not correctly account for NonStreaming mips and other constraints.
+	This function is not correct and should not be used. */
 	FORCEINLINE int32 GetMinTextureResidentMipCount() const
 	{
 		return FMath::Max(GMinTextureResidentMipCount, GetPlatformData() ? (int32)GetPlatformData()->GetNumMipsInTail() : 0);
diff --git a/Engine/Source/Runtime/Engine/Private/StreamableRenderAsset.cpp b/Engine/Source/Runtime/Engine/Private/StreamableRenderAsset.cpp
--- a/Engine/Source/Runtime/Engine/Private/StreamableRenderAsset.cpp
+++ b/Engine/Source/Runtime/Engine/Private/StreamableRenderAsset.cpp
@@ -317,15 +317,22 @@ void UStreamableRenderAsset::UnlinkStreaming()
 
 bool UStreamableRenderAsset::IsFullyStreamedIn()
 {
+	// consider a texture fully streamed when it hits this number of LODs :
+	//	MaxNumLODs has already been reduced by the "drop mip" LOD Bias
+	//	Note that just subtracting off NumCinematicMipLevels is not the right way to get the cinematic lod bias
+	//	it should be CalculateLODBias(false) , but we don't have that information here
+	int32 FullyStreamedNumLODs = CachedSRRState.MaxNumLODs - NumCinematicMipLevels;
+
 	// Note that if CachedSRRState is not valid, then this asset is not streamable and is then at max resolution.
 	if (!CachedSRRState.IsValid() 
 		|| !CachedSRRState.bSupportsStreaming 
-		|| (CachedSRRState.NumResidentLODs >= (CachedSRRState.MaxNumLODs - CachedCombinedLODBias)))
+		|| CachedSRRState.NumResidentLODs >= FullyStreamedNumLODs)
 	{
 		return true;
 	}
 
 	// IsFullyStreamedIn() might be used incorrectly if any logic waits on it to be true.
+	// there could be optional mips which are not available to be loaded, so waiting on IsFullyStreamedIn would never finish
 	ensureMsgf(CachedSRRState.NumResidentLODs != CachedSRRState.NumNonOptionalLODs, TEXT("IsFullyStreamedIn() is being called on (%s) which might not have optional LODs mounted."), *GetFName().ToString());
 
 	return false;
diff --git a/Engine/Source/Runtime/Engine/Private/Streaming/StreamingTexture.cpp b/Engine/Source/Runtime/Engine/Private/Streaming/StreamingTexture.cpp
--- a/Engine/Source/Runtime/Engine/Private/Streaming/StreamingTexture.cpp
+++ b/Engine/Source/Runtime/Engine/Private/Streaming/StreamingTexture.cpp
@@ -187,9 +187,21 @@ void FStreamingRenderAsset::UpdateDynamicData(const int32* NumStreamedMips, int3
 
 		int32 LODBias = 0;
 		if (!Settings.bUseAllMips)
-		{
-			const int32 ResourceLODBias = FMath::Max<int32>(0, RenderAsset->GetCachedLODBias() - ResourceState.AssetLODBias);
-			LODBias = FMath::Max<int32>(ResourceLODBias - NumCinematicMipLevels, 0);
+		{			
+			UTexture * Texture = Cast<UTexture>(RenderAsset);
+			if ( Texture )
+			{
+				LODBias = FMath::Max<int32>(0, Texture->CalculateLODBias(false) - ResourceState.AssetLODBias);
+			}
+			else
+			{
+				// GetCachedLODBias() - NumCinematicMipLevels
+				//  is not the right way to get non-cinematic LODBias
+				//	you should call CalculateLODBias(false)
+				//	the difference occurs if any of the clamps come into play
+				const int32 ResourceLODBias = FMath::Max<int32>(0, RenderAsset->GetCachedLODBias() - ResourceState.AssetLODBias);
+				LODBias = FMath::Max<int32>(ResourceLODBias - NumCinematicMipLevels, 0);
+			}
 
 			// Reduce the max allowed resolution according to LODBias if the texture group allows it.
 			if (IsMaxResolutionAffectedByGlobalBias() && !Settings.bUsePerTextureBias)
diff --git a/Engine/Source/Runtime/Engine/Private/Texture.cpp b/Engine/Source/Runtime/Engine/Private/Texture.cpp
--- a/Engine/Source/Runtime/Engine/Private/Texture.cpp
+++ b/Engine/Source/Runtime/Engine/Private/Texture.cpp
@@ -462,7 +462,7 @@ void UTexture::ValidateSettingsAfterImportOrEdit(bool * pRequiresNotifyMaterials
 					// VT nonpow2 will fail to build
 					// force it into a state that will succeed? or just let it fail?
 					// you can either pad to pow2 or set MaxTextureSize and turn off VT
-				PowerOfTwoMode = ETexturePowerOfTwoSetting::PadToPowerOfTwo;
+					PowerOfTwoMode = ETexturePowerOfTwoSetting::PadToPowerOfTwo;
 				}
 				else
 				{
@@ -951,9 +951,6 @@ void UTexture::PostLoad()
 
 	if( !IsTemplate() )
 	{
-		// Update cached LOD bias.
-		UpdateCachedLODBias();
-
 		// The texture will be cached by the cubemap it is contained within on consoles.
 		UTextureCube* CubeMap = Cast<UTextureCube>(GetOuter());
 		if (CubeMap == NULL)
@@ -1257,15 +1254,24 @@ bool UTexture::ForceUpdateTextureStreaming()
 {
 	if (!IStreamingManager::HasShutdown())
 	{
+/*
+		// I'm not sure what the scope of this "ForceUpdate" is supposed to be
+		// but if you are trying to account for config changes that can change LODBias in Editor
+		// then that means the NumMips in the cached FStreamableRenderResourceState may have changed
+		// and they all need to be reset
+
 #if WITH_EDITOR
 		for( TObjectIterator<UTexture2D> It; It; ++It )
 		{
 			UTexture* Texture = *It;
 
-			// Update cached LOD bias.
-			Texture->UpdateCachedLODBias();
+			fill the 
+			FStreamableTextureResource::State
+			by re-calling
+			Texture->GetResourcePostInitState();
 		}
 #endif // #if WITH_EDITOR
+*/
 
 		// Make sure we iterate over all textures by setting it to high value.
 		IStreamingManager::Get().SetNumIterationsForNextFrame( 100 );
@@ -1397,35 +1403,84 @@ bool UTexture::IsCandidateForTextureStreamingOnPlatformDuringCook(const ITargetP
 
 FStreamableRenderResourceState UTexture::GetResourcePostInitState(const FTexturePlatformData* PlatformData, bool bAllowStreaming, int32 MinRequestMipCount, int32 MaxMipCount, bool bSkipCanBeLoaded) const
 {
+	// Async caching of PlatformData must be done before calling this
+	//	if you call while async CachePlatformData is in progress, you get garbage out
+
+	// "FullLODBias" is == NumCincematicMipLevels + also maybe drop mip LODBias
+	//  the LODBias to drop mips is not added in cooked runs, because those mips have been already dropped
+	//  it is added in non-cooked runs because they are still present but we are trying to pretend they are not there
+	// "CinematicLODBias" is typically zero in cooked runs, = drop mip count in non-cooked runs
+
+	int32 FullLODBias = CalculateLODBias(true);
+	int32 CinematicLODBias = CalculateLODBias(false);
+	check( FullLODBias >= CinematicLODBias );
+
+	bool bTextureIsStreamable = IsPossibleToStream();
+
+	int32 FullMipCount = PlatformData->Mips.Num();
+	int32 NumOfNonOptionalMips = PlatformData->GetNumNonOptionalMips();
+	int32 NumOfNonStreamingMips = PlatformData->GetNumNonStreamingMips(bTextureIsStreamable);
+	int32 NumMipsInTail = PlatformData->GetNumMipsInTail();
+
+	// Optional mips must be streaming mips :
+	check( NumOfNonOptionalMips >= NumOfNonStreamingMips );
+	// Mips in tail must be nonstreaming :
+	check( NumOfNonStreamingMips >= NumMipsInTail );
+
 	// Create the resource with a mip count limit taking in consideration the asset LODBias.
 	// This ensures that the mip count stays constant when toggling asset streaming at runtime.
-	const int32 NumMips = [&]() -> int32 
+	
+	const int32 ExpectedAssetLODBias = FMath::Clamp<int32>(CinematicLODBias, 0, FullMipCount - 1);
+	// "ExpectedAssetLODBias" is the number of mips that would be dropped in cook
+	//		in a cooked run, it is zero
+
+	// GMaxTextureMipCount is for the current running RHI
+	//  it may be lower than the number of mips we cooked (eg. on mobile)
+	//	we must limit the number of mips to this count.
+	const int32 MaxRuntimeMipCount = FMath::Min<int32>(GMaxTextureMipCount, FStreamableRenderResourceState::MAX_LOD_COUNT);
+
+	int32 NumMips = FMath::Min<int32>( FullMipCount - ExpectedAssetLODBias, MaxRuntimeMipCount );
+	// "NumMips" is the number of mips after drop LOD Bias
+	//	 it should be the same in Editor and Runtime
+
+	if (MaxMipCount > 0 && NumMips > MaxMipCount)
 	{
-		const int32 ExpectedAssetLODBias = FMath::Clamp<int32>(GetCachedLODBias() - NumCinematicMipLevels, 0, PlatformData->Mips.Num() - 1);
-		const int32 MaxRuntimeMipCount = FMath::Min<int32>(GMaxTextureMipCount, FStreamableRenderResourceState::MAX_LOD_COUNT);
-		if (MaxMipCount > 0)
-		{
-			return FMath::Min3<int32>(PlatformData->Mips.Num() - ExpectedAssetLODBias, MaxMipCount, MaxRuntimeMipCount);
-		}
-		else
+		// MaxMipCount is almost always either 0 or == MaxRuntimeMipCount
+		// one exception is :
+		//	MobileReduceLoadedMips(NumMips);
+		// which can cause an additional reduction of NumMips
+		NumMips = MaxMipCount;
+	}
+	
+	// don't allow less than NumOfNonStreamingMips :
+	if ( NumMips < NumOfNonStreamingMips )
+	{
+		// if NumMips went under NumOfNonStreamingMips due to ExpectedAssetLODBias
+		//  then force it back up
+		// but if it went under due to MaxRuntimeMipCount
+		// then that's a problem
+
+		if ( NumOfNonStreamingMips > MaxRuntimeMipCount )
 		{
-			return FMath::Min<int32>(PlatformData->Mips.Num() - ExpectedAssetLODBias, MaxRuntimeMipCount);
+			// this should never happen on a PC platform, only on mobile
+			// in that case streaming in the "nonstreaming" may actually be okay
+			// in the old code, this was expected behavior, let the MaxRuntimeMipCount trump the NonStreaming constraint
+			// in new code (with RequiredBlock4Alignment) we do not expect to see this any more, so warn :
+			UE_LOG(LogTexture, Warning, TEXT("NumOfNonStreamingMips > MaxRuntimeMipCount. (%s)"), *GetName());
+			NumOfNonStreamingMips = MaxRuntimeMipCount;
 		}
-	}();
-	
-	bool bTextureIsStreamable = IsPossibleToStream();
 
-	const int32 NumOfNonOptionalMips = FMath::Min<int32>(NumMips, PlatformData->GetNumNonOptionalMips());
-	const int32 NumOfNonStreamingMips = FMath::Min<int32>(NumMips, PlatformData->GetNumNonStreamingMips(bTextureIsStreamable));
-	// Optional mips must be streaming mips :
-	check( NumOfNonOptionalMips >= NumOfNonStreamingMips );
+		NumMips = NumOfNonStreamingMips;
+	}
 
+	check( NumMips >= NumMipsInTail );
+	
 	if ( NumOfNonStreamingMips == NumMips )
 	{
 		bTextureIsStreamable = false;
 	}
 
-	const int32 AssetMipIdxForResourceFirstMip = FMath::Max<int32>(0, PlatformData->Mips.Num() - NumMips);
+	const int32 AssetMipIdxForResourceFirstMip = FullMipCount - NumMips;
 
 	bool bMakeStreamable = false;
 	int32 NumRequestedMips = 0;
@@ -1477,40 +1532,38 @@ FStreamableRenderResourceState UTexture::GetResourcePostInitState(const FTexture
 		// but only if the texture itself is streamable
 
 		// Adjust CachedLODBias so that it takes into account FStreamableRenderResourceState::AssetLODBias.
-		const int32 ResourceLODBias = FMath::Max<int32>(0, GetCachedLODBias() - AssetMipIdxForResourceFirstMip);
-		//  ResourceLODBias almost always = NumCinematicMipLevels
-		//check( ResourceLODBias == NumCinematicMipLevels ); // this will be true unless you hit the MaxRuntimeMipCount clamp in NumMips
-		// if ResourceLODBias == 0 , this always selects NumRequestedMips = NumMips
+		const int32 ResourceLODBias = FMath::Max<int32>(0, FullLODBias - AssetMipIdxForResourceFirstMip);
+		//  ResourceLODBias almost always == NumCinematicMipLevels, unless you hit the MaxRuntimeMipCount clamp in NumMips
 
-		// Ensure NumMipsInTail is within valid range to safeguard on the above expressions. 
-		const int32 NumMipsInTail = FMath::Clamp<int32>(PlatformData->GetNumMipsInTail(), 1, NumMips);
-
-		// Bias is not allowed to shrink the mip count below NumMipsInTail.
-		NumRequestedMips = FMath::Max<int32>(NumMips - ResourceLODBias, NumMipsInTail);
+		// Bias is not allowed to shrink the mip count below NumOfNonStreamingMips.
+		NumRequestedMips = FMath::Max<int32>(NumMips - ResourceLODBias, NumOfNonStreamingMips);
 
 		// If trying to load optional mips, check if the first resource mip is available.
 		if (NumRequestedMips > NumOfNonOptionalMips && !DoesMipDataExist(AssetMipIdxForResourceFirstMip))
 		{
 			NumRequestedMips = NumOfNonOptionalMips;
 		}
-
-		// Ensure we don't request a top mip in the NonStreamingMips
-		NumRequestedMips = FMath::Max(NumRequestedMips,NumOfNonStreamingMips);
 	}
 
-	// @todo Oodle : this looks like a bug; did it mean to be MinRequestMipCount <= NumMips
+	// @todo Oodle : this looks like a bug; did it mean to be MinRequestMipCount <= NumMips ?
 	// typically MinRequestMipCount == 0
-	// the only place it's not zero is from UTexture2D::CreateResource from existing resource mem, where it is == NumMips
-	// but in that case it is ignored here because it is == NumMips, not <
-	if (NumRequestedMips < MinRequestMipCount && MinRequestMipCount < NumMips)
+	// the only place it's not zero is from UTexture2D::CreateResource from existing resource mem, where MinRequestMipCount is == NumMips
+	// but in that case it is ignored here because this branches on < instead of <=
+	if ( NumRequestedMips < MinRequestMipCount && MinRequestMipCount < NumMips )
 	{
+		// as written with < instead of <= this branch is not used
+
 		NumRequestedMips = MinRequestMipCount;
 	}
 
+	check( NumOfNonStreamingMips <= NumMips );
+	check( NumRequestedMips <= NumMips );
+	check( NumRequestedMips >= NumOfNonStreamingMips );
+
 	FStreamableRenderResourceState PostInitState;
 	PostInitState.bSupportsStreaming = bMakeStreamable;
 	PostInitState.NumNonStreamingLODs = (uint8)NumOfNonStreamingMips;
-	PostInitState.NumNonOptionalLODs = (uint8)NumOfNonOptionalMips;
+	PostInitState.NumNonOptionalLODs = (uint8) FMath::Min<int32>(NumOfNonOptionalMips,NumMips);
 	PostInitState.MaxNumLODs = (uint8)NumMips;
 	PostInitState.AssetLODBias = (uint8)AssetMipIdxForResourceFirstMip;
 	PostInitState.NumResidentLODs = (uint8)NumRequestedMips;
@@ -3094,7 +3147,7 @@ FName GetDefaultTextureFormatName( const ITargetPlatform* TargetPlatform, const
 
 #if WITH_EDITORONLY_DATA
 		const UTextureLODSettings& LODSettings = TargetPlatform->GetTextureLODSettings();
- 		const uint32 LODBiasNoCinematics = FMath::Max<uint32>(LODSettings.CalculateLODBias(SizeX, SizeY, Texture->MaxTextureSize, Texture->LODGroup, Texture->LODBias, 0, Texture->MipGenSettings, bVirtualTextureStreaming), 0);
+ 		const uint32 LODBiasNoCinematics = FMath::Max<int32>(LODSettings.CalculateLODBias(SizeX, SizeY, Texture->MaxTextureSize, Texture->LODGroup, Texture->LODBias, 0, Texture->MipGenSettings, bVirtualTextureStreaming), 0);
 		SizeX = FMath::Max<int32>(SizeX >> LODBiasNoCinematics, 1);
 		SizeY = FMath::Max<int32>(SizeY >> LODBiasNoCinematics, 1);
 #endif
diff --git a/Engine/Source/Runtime/Engine/Private/Texture2D.cpp b/Engine/Source/Runtime/Engine/Private/Texture2D.cpp
--- a/Engine/Source/Runtime/Engine/Private/Texture2D.cpp
+++ b/Engine/Source/Runtime/Engine/Private/Texture2D.cpp
@@ -118,6 +118,8 @@ static FAutoConsoleVariableRef CVarUseGenericStreamingPath(
 
 static int32 MobileReduceLoadedMips(int32 NumTotalMips)
 {
+	// apply cvar options to reduce the number of mips created at runtime
+	// note they are still cooked  &shipped
 	int32 NumReduceMips = FMath::Max(0, CVarMobileReduceLoadedMips.GetValueOnAnyThread());
 	int32 MaxLoadedMips = FMath::Clamp(CVarMobileMaxLoadedMips.GetValueOnAnyThread(), 1, GMaxTextureMipCount);
 
@@ -327,6 +329,7 @@ int32 UTexture2D::GetSizeX() const
 #if WITH_EDITOR
 		if (IsDefaultTexture())
 		{
+			// any calculation that actually uses this is garbage
 			return GetDefaultTexture2D(this)->GetSizeX();
 		}
 #endif
@@ -342,6 +345,7 @@ int32 UTexture2D::GetSizeY() const
 #if WITH_EDITOR
 		if (IsDefaultTexture())
 		{
+			// any calculation that actually uses this is garbage
 			return GetDefaultTexture2D(this)->GetSizeY();
 		}
 #endif
@@ -781,6 +785,12 @@ int32 UTexture2D::CalcTextureMemorySize( int32 MipCount ) const
 
 int32 UTexture2D::GetNumMipsAllowed(bool bIgnoreMinResidency) const
 {
+	// this function is trying to get the number of mips that will be in the texture after cooking
+	//	(eg. after "drop mip" lod bias is applied)
+	// but it doesn't exactly replicate the behavior of Serialize
+	// it's also similar to Texture::GetResourcePostInitState but not the same
+	// yay
+
 	const int32 NumMips = GetNumMips();
 
 	// Compute the number of mips that will be available after cooking, as some mips get cooked out.
@@ -906,7 +916,12 @@ bool UTexture2D::HasAlphaChannel() const
 #if WITH_EDITOR
 bool UTexture2D::GetStreamableRenderResourceState(FTexturePlatformData* InPlatformData, FStreamableRenderResourceState& OutState) const
 {
+	// we want to make queries about InPlatformData here, not the current platform data
+	// so for the scope of this function, change PrivatePlatformData to InPlatformData
 	TGuardValue<FTexturePlatformData*> Guard(const_cast<UTexture2D*>(this)->PrivatePlatformData, InPlatformData);
+
+	check( GetPlatformData() == InPlatformData );
+
 	if (GetPlatformData())
 	{
 		if (IsCurrentlyVirtualTextured())
@@ -914,11 +929,15 @@ bool UTexture2D::GetStreamableRenderResourceState(FTexturePlatformData* InPlatfo
 			return false;
 		}
 
+		// we're relying on these calls like GetPixelFormat() to actually be queries on InPlatformData :
 		const EPixelFormat PixelFormat = GetPixelFormat();
 		const int32 NumMips = FMath::Min3<int32>(GetPlatformData()->Mips.Num(), GMaxTextureMipCount, FStreamableRenderResourceState::MAX_LOD_COUNT);
 		if (NumMips && GPixelFormats[PixelFormat].Supported &&
 			(NumMips > 1 || FMath::Max(GetSizeX(), GetSizeY()) <= (int32)GetMax2DTextureDimension()))
 		{
+			// this wants to make a query on InPlatformData, for some other target platform
+			// @todo Oodle : this is a bit broken, will use the current DeviceProfile in CalculateLODBias
+			//		it should be getting LODGroup settings from the TargetPlatform but is not
 			OutState = GetResourcePostInitState(GetPlatformData(), true, 0, NumMips, /*bSkipCanBeLoaded*/ true);
 			return true;
 		}
diff --git a/Engine/Source/Runtime/Engine/Private/TextureCompiler.cpp b/Engine/Source/Runtime/Engine/Private/TextureCompiler.cpp
--- a/Engine/Source/Runtime/Engine/Private/TextureCompiler.cpp
+++ b/Engine/Source/Runtime/Engine/Private/TextureCompiler.cpp
@@ -207,9 +207,6 @@ void FTextureCompilingManager::PostCompilation(UTexture* Texture)
 	Texture->FinishCachePlatformData();
 	Texture->UpdateResource();
 
-	// Needs to be called after the placeholder resource has been replaced by the real one
-	Texture->UpdateCachedLODBias();
-
 	// Generate an empty property changed event, to force the asset registry tag
 	// to be refreshed now that pixel format and alpha channels are available.
 	FPropertyChangedEvent EmptyPropertyChangedEvent(nullptr);
diff --git a/Engine/Source/Runtime/Engine/Private/TextureDerivedData.cpp b/Engine/Source/Runtime/Engine/Private/TextureDerivedData.cpp
--- a/Engine/Source/Runtime/Engine/Private/TextureDerivedData.cpp
+++ b/Engine/Source/Runtime/Engine/Private/TextureDerivedData.cpp
@@ -2132,18 +2132,16 @@ int32 FTexturePlatformData::GetNumNonStreamingMips(bool bIsStreamingPossible) co
 		NumNonStreamingMips = FMath::Max(NumNonStreamingMips, (int32)GetNumMipsInTail());
 		NumNonStreamingMips = FMath::Max(NumNonStreamingMips, UTexture2D::GetStaticMinTextureResidentMipCount());
 		NumNonStreamingMips = FMath::Min(NumNonStreamingMips, MipCount);
-		int32 BlockSizeX = GPixelFormats[PixelFormat].BlockSizeX;
-		int32 BlockSizeY = GPixelFormats[PixelFormat].BlockSizeY;
-		if (BlockSizeX > 1 || BlockSizeY > 1)
+
+		if ( RequiresBlock4Alignment(PixelFormat) )
 		{
-			// ensure the top non-streamed mip size is >= BlockSize (and a multiple of block size!)
-			// @todo Oodle : only do for BCN, not for ASTC
+			// ensure the top non-streamed mip (and all streamed mips) size is >= BlockSize (and a multiple of block size!)
 			
 			// note: this is not right for non pow 2; NeverStream should set !bIsStreamingPossible in that case
 			if ( FMath::IsPowerOfTwo(Mips[0].SizeX) && FMath::IsPowerOfTwo(Mips[0].SizeY) )
 			{
-				NumNonStreamingMips = FMath::Max<int32>(NumNonStreamingMips, MipCount - FPlatformMath::FloorLog2(Mips[0].SizeX / BlockSizeX));
-				NumNonStreamingMips = FMath::Max<int32>(NumNonStreamingMips, MipCount - FPlatformMath::FloorLog2(Mips[0].SizeY / BlockSizeY));
+				NumNonStreamingMips = FMath::Max<int32>(NumNonStreamingMips, MipCount - FPlatformMath::FloorLog2(Mips[0].SizeX / 4));
+				NumNonStreamingMips = FMath::Max<int32>(NumNonStreamingMips, MipCount - FPlatformMath::FloorLog2(Mips[0].SizeY / 4));
 			}
 			else
 			{
@@ -2683,6 +2681,7 @@ UE::DerivedData::FValueId FTexturePlatformData::MakeMipId(int32 MipIndex)
 void FTexturePlatformData::SerializeCooked(FArchive& Ar, UTexture* Owner, bool bStreamable)
 {
 	EPlatformDataSerializationFlags Flags = EPlatformDataSerializationFlags::Cooked;
+	// @@CB bStreamable may be set to true because of bCookedIsStreamable even if I am NOT actually a streamable texture on this platform
 	if (bStreamable)
 	{
 		Flags |= EPlatformDataSerializationFlags::Streamable;
@@ -2792,9 +2791,12 @@ void UTextureCube::GetMipData(int32 FirstMipToLoad, void** OutMipData)
 	}
 }
 
-void UTexture::UpdateCachedLODBias()
+int32 UTexture::CalculateLODBias(bool bWithCinematicMipBias) const
 {
-	CachedCombinedLODBias = UDeviceProfileManager::Get().GetActiveProfile()->GetTextureLODSettings()->CalculateLODBias(this);
+	// Async caching of PlatformData must be done before calling this
+	//	if you call while async CachePlatformData is in progress, you get garbage out
+
+	return UDeviceProfileManager::Get().GetActiveProfile()->GetTextureLODSettings()->CalculateLODBias(this,bWithCinematicMipBias);
 }
 
 #if WITH_EDITOR
@@ -2985,8 +2987,6 @@ void UTexture::CachePlatformData(bool bAsyncCache, bool bAllowAsyncBuild, bool b
 			// If there is no source art available, create an empty platform data container.
 			PlatformDataLink = new FTexturePlatformData();
 		}
-
-		UpdateCachedLODBias();
 	}
 }
 
@@ -3305,6 +3305,7 @@ void UTexture::FinishCachePlatformData()
 			if ( RunningPlatformData == NULL )
 			{
 				// begin cache never called
+				//  do a non-async cache :
 				CachePlatformData();
 			}
 			else
@@ -3315,7 +3316,8 @@ void UTexture::FinishCachePlatformData()
 		}
 	}
 
-	UpdateCachedLODBias();
+	// FinishCachePlatformData is not reliably called
+	//  this is not a good place to put code that finalizes caching
 }
 
 void UTexture::ForceRebuildPlatformData(uint8 InEncodeSpeedOverride /* =255 ETextureEncodeSpeedOverride::Disabled */)
@@ -3489,6 +3491,7 @@ void UTexture::SerializeCookedPlatformData(FArchive& Ar)
 				PlatformDataToSave->FinishCache();
 
 				// Update bCookedIsStreamable for later use in IsCandidateForTextureStreaming
+				// @@CB bCookedIsStreamable is set to true if the texture can be streamed on ANY platform
 				FStreamableRenderResourceState State;
 				if (GetStreamableRenderResourceState(PlatformDataToSave, State))
 				{
@@ -3505,6 +3508,9 @@ void UTexture::SerializeCookedPlatformData(FArchive& Ar)
 				}
 
 				// Pass streamable flag for inlining mips
+				// @@CB GetTextureIsStreamableOnPlatform checks bCookedIsStreamable which changes while we go through this loop
+				//		that makes behavior dependent on order of platforms
+				// 
 				bool bTextureIsStreamable = GetTextureIsStreamableOnPlatform(*this, *Ar.CookingTarget());
 				PlatformDataToSave->SerializeCooked(Ar, this, bTextureIsStreamable);
 
diff --git a/Engine/Source/Runtime/Engine/Private/TextureLODSettings.cpp b/Engine/Source/Runtime/Engine/Private/TextureLODSettings.cpp
--- a/Engine/Source/Runtime/Engine/Private/TextureLODSettings.cpp
+++ b/Engine/Source/Runtime/Engine/Private/TextureLODSettings.cpp
@@ -140,7 +140,10 @@ int32 UTextureLODSettings::CalculateLODBias(const UTexture* Texture, bool bIncCi
 	TextureMaxSize = Texture->MaxTextureSize;
 #endif // #if WITH_EDITORONLY_DATA
 
-	return CalculateLODBias(Texture->GetSurfaceWidth(), Texture->GetSurfaceHeight(), TextureMaxSize, Texture->LODGroup, Texture->LODBias, bIncCinematicMips ? Texture->NumCinematicMipLevels : 0, MipGenSetting, Texture->IsCurrentlyVirtualTextured());
+	float Width = Texture->GetSurfaceWidth();
+	float Height = Texture->GetSurfaceHeight();
+
+	return CalculateLODBias( FMath::RoundToInt(Width), FMath::RoundToInt(Height), TextureMaxSize, Texture->LODGroup, Texture->LODBias, bIncCinematicMips ? Texture->NumCinematicMipLevels : 0, MipGenSetting, Texture->IsCurrentlyVirtualTextured());
 }
 
 int32 UTextureLODSettings::CalculateLODBias(int32 Width, int32 Height, int32 MaxSize, int32 LODGroup, int32 LODBias, int32 NumCinematicMipLevels, TextureMipGenSettings InMipGenSetting, bool bVirtualTexture ) const
@@ -150,6 +153,11 @@ int32 UTextureLODSettings::CalculateLODBias(int32 Width, int32 Height, int32 Max
 	// Find LOD group.
 	const FTextureLODGroup& LODGroupInfo = TextureLODGroups[LODGroup];
 
+	// Note: MaxLODSize sort of acts like a max texture size limit
+	//	but it isn't really a good way to limit texture size
+	// because there are various ways in which it can be ignored
+	// eg. on textures set to NoMipmaps
+
 	// Test to see if we have no mip generation as in which case the LOD bias will be ignored
 	// VTs don't respect NoMipmaps, mips are required for VTs
 	const TextureMipGenSettings FinalMipGenSetting = (InMipGenSetting == TMGS_FromTextureGroup) ? (TextureMipGenSettings)LODGroupInfo.MipGenSettings : InMipGenSetting;
@@ -187,7 +195,13 @@ int32 UTextureLODSettings::CalculateLODBias(int32 Width, int32 Height, int32 Max
 	
 	if (LODGroup == TEXTUREGROUP_UI)
 	{
-		UsedLODBias += GUITextureLODBias;
+		UsedLODBias += GUITextureLODBias;  
+		// @todo Oodle : GUITextureLODBias is applied at both cook time & run time , which is screwy
+		//	this is not inside the if (!FPlatformProperties::RequiresCookedData()) ?
+		//	it should either act like cinematic bias (runtime streaming)
+		//	or like the "drop mip" lod bias (cook time)
+		//	but it's not quite like either
+		// -> this looks broken, probably never used
 	}
 
 	int32 MinLOD = FMath::CeilLogTwo(LODGroupInfo.MinLODSize);
