Fix regression: bad condition was causing failed actor load false positive logs.

diff --git a/Engine/Source/Runtime/Engine/Private/Level.cpp b/Engine/Source/Runtime/Engine/Private/Level.cpp
--- a/Engine/Source/Runtime/Engine/Private/Level.cpp
+++ b/Engine/Source/Runtime/Engine/Private/Level.cpp
@@ -1156,10 +1156,10 @@ void ULevel::PostLoad()
 					ActorPackage->SetFlags(RF_Transient);
 				}
 
+				bool bFoundActor = false;
 				ActorPackage = LoadPackage(ActorPackage, *ActorPackageName, bPackageForPIE ? LOAD_PackageForPIE : LOAD_None, nullptr, &InstancingContext);
-
-				int32 PreviousActorCount = Actors.Num();
-				ForEachObjectWithPackage(ActorPackage, [this, &ActorsSet, bPackageForPIE](UObject* PackageObject)
+				
+				ForEachObjectWithPackage(ActorPackage, [this, &ActorsSet, &bFoundActor, bPackageForPIE](UObject* PackageObject)
 				{
 					// There might be multiple actors per package in the case where an actor as a child actor component as we put child actor in the same package as their parent
 					if (PackageObject->IsA<AActor>() && !PackageObject->IsTemplate())
@@ -1170,6 +1170,7 @@ void ULevel::PostLoad()
 						{
 							Actors.Add(Actor);
 						}
+						bFoundActor = true;
 					}
 					// In PIE, we make sure to clear RF_Standalone flag on objects in external packages (UMetaData) 
 					// This guarantees that external packages of actors that are destroyed during the PIE session will
@@ -1181,11 +1182,7 @@ void ULevel::PostLoad()
 					return true;
 				}, false);
 
-				bool bAddedPackageActors = Actors.Num() > PreviousActorCount;
-				if (!bAddedPackageActors)
-				{
-					UE_LOG(LogLevel, Error, TEXT("Failed to load Actor for External Actor Package %s"), *ActorPackageName);
-				}
+				UE_CLOG(!bFoundActor, LogLevel, Error, TEXT("Failed to load Actor for External Actor Package %s"), *ActorPackageName);
 			}
 		}
 	}
