fix Oodle Texture build memory estimate being used on all blocked textures, not just BCN
caused the log "Unexpected non-BC TextureFormatName" when cooking for Android

diff --git a/Engine/Source/Runtime/Engine/Private/TextureDerivedDataTask.cpp b/Engine/Source/Runtime/Engine/Private/TextureDerivedDataTask.cpp
--- a/Engine/Source/Runtime/Engine/Private/TextureDerivedDataTask.cpp
+++ b/Engine/Source/Runtime/Engine/Private/TextureDerivedDataTask.cpp
@@ -1071,10 +1071,10 @@ static int64 GetBuildRequiredMemoryEstimate(UTexture* InTexture,
 
 		MemoryEstimate += OutputSizeBytes;
 
-		// check BlockSize to see if it's uncompressed or a BCN format :
-		if ( PFI.BlockSizeX > 1 )
+		// check to see if it's uncompressed or a BCN format :
+		if ( IsDXTCBlockCompressedTextureFormat(PixelFormat) )
 		{
-			// block-compressed format
+			// block-compressed format ; assume it's using Oodle Texture
 			
 			if ( bRDO )
 			{
@@ -1151,6 +1151,10 @@ static int64 GetBuildRequiredMemoryEstimate(UTexture* InTexture,
 				MemoryEstimate += FMath::Max(RDOPhase1MemUse,RDOPhase2MemUse);
 			}
 		}
+		else
+		{
+			// note: memory ues of non-Oodle encoders is not estimated
+		}
 		
 		MemoryEstimate += 64 * 1024; // overhead room
 		
