Fix an issue where a ContentBundle cook keep a reference on generated package (external streaming object)

diff --git a/Engine/Source/Runtime/Engine/Private/WorldPartition/ContentBundle/ContentBundleEditor.cpp b/Engine/Source/Runtime/Engine/Private/WorldPartition/ContentBundle/ContentBundleEditor.cpp
--- a/Engine/Source/Runtime/Engine/Private/WorldPartition/ContentBundle/ContentBundleEditor.cpp
+++ b/Engine/Source/Runtime/Engine/Private/WorldPartition/ContentBundle/ContentBundleEditor.cpp
@@ -289,7 +289,6 @@ void FContentBundleEditor::AddReferencedObjects(FReferenceCollector& Collector)
 {
 	FContentBundleBase::AddReferencedObjects(Collector);
 
-	Collector.AddReferencedObjects(ContentBundleCells);
 	Collector.AddReferencedObject(ExternalStreamingObject);
 	Collector.AddReferencedObject(UnsavedActorMonitor);
 }
@@ -305,15 +304,15 @@ void FContentBundleEditor::GenerateStreaming(TArray<FString>* OutPackageToGenera
 	UWorldPartition* WorldPartition = GetInjectedWorld()->GetWorldPartition();
 	WorldPartition->GenerateContainerStreaming(ActorDescContainer.Get(), OutPackageToGenerate);
 
-	WorldPartition->RuntimeHash->ForEachStreamingCells([this](const UWorldPartitionRuntimeCell* Cell)
+	ExternalStreamingObject = WorldPartition->RuntimeHash->StoreToExternalStreamingObject(GetTransientPackage(), *GetExternalStreamingObjectName());
+
+	uint32 CellCount = 0;
+	ExternalStreamingObject->ForEachStreamingCells([&CellCount](const UWorldPartitionRuntimeCell& Cell)
 	{
-		ContentBundleCells.Add(Cell);
+		CellCount++;
 		return true;
 	});
-
-	ExternalStreamingObject = WorldPartition->RuntimeHash->StoreToExternalStreamingObject(GetInjectedWorld()->GetWorldPartition(), *GetExternalStreamingObjectName());
-
-	UE_LOG(LogContentBundle, Log, TEXT("[CB: %s] Generated streaming cells. %u cells were generated"), *GetDescriptor()->GetDisplayName(), ContentBundleCells.Num());
+	UE_LOG(LogContentBundle, Log, TEXT("[CB: %s] Generated streaming cells. %u cells were generated"), *GetDescriptor()->GetDisplayName(), CellCount);
 
 	if (!IsRunningCookCommandlet())
 	{
@@ -323,8 +322,9 @@ void FContentBundleEditor::GenerateStreaming(TArray<FString>* OutPackageToGenera
 		{
 			UE_LOG(LogContentBundle, Error, TEXT("[CB: %s] Failed to store streaming object for %s. PIE duplication will not work."), *GetDescriptor()->GetDisplayName());
 		}
+
+		// Clear streaming object. It will be kept alive & duplicated by UContentBundleDuplicateForPIEHelper.
 		ExternalStreamingObject = nullptr;
-		ContentBundleCells.Empty();
 	}
 
 	WorldPartition->FlushStreaming();
@@ -344,19 +344,19 @@ bool FContentBundleEditor::GatherPackagesToCook(class IWorldPartitionCookPackage
 
 	if (HasCookedContent())
 	{
-		for (const UWorldPartitionRuntimeCell* RuntimeCell : ContentBundleCells)
+		ExternalStreamingObject->ForEachStreamingCells([this, &bIsSuccess, &CookContext](UWorldPartitionRuntimeCell& RuntimeCell)
 		{
-			if (const FWorldPartitionCookPackage* CookPackage = CookContext.AddLevelStreamingPackageToGenerate(this, ContentBundlePaths::GetCookedContentBundleLevelFolder(*this), RuntimeCell->GetPackageNameToCreate()))
+			if (const FWorldPartitionCookPackage* CookPackage = CookContext.AddLevelStreamingPackageToGenerate(this, ContentBundlePaths::GetCookedContentBundleLevelFolder(*this), RuntimeCell.GetPackageNameToCreate()))
 			{
-				CookPackageIdsToCell.Emplace(CookPackage->PackageId, RuntimeCell);
+				CookPackageIdsToCell.Emplace(CookPackage->PackageId, &RuntimeCell);
 			}
 			else
 			{
-				UE_LOG(LogContentBundle, Error, TEXT("[CB: %s][Cook] Failed to add cell package %s in cook context."), *ContentBundlePaths::GetCookedContentBundleLevelFolder(*this), *RuntimeCell->GetPackageNameToCreate());
+				UE_LOG(LogContentBundle, Error, TEXT("[CB: %s][Cook] Failed to add cell package %s in cook context."), *ContentBundlePaths::GetCookedContentBundleLevelFolder(*this), *RuntimeCell.GetPackageNameToCreate());
 				bIsSuccess = false;
 			}
-		}
-
+		});
+		
 		const FWorldPartitionCookPackage* CookPackage = CookContext.AddGenericPackageToGenerate(this, ContentBundlePaths::GetCookedContentBundleLevelFolder(*this), TEXT("StreamingObject"));
 		if (CookPackage == nullptr)
 		{
@@ -373,6 +373,7 @@ bool FContentBundleEditor::PopulateGeneratorPackageForCook(class IWorldPartition
 {
 	bool bIsSuccess = true;
 
+
 	if (HasCookedContent())
 	{
 		for (const FWorldPartitionCookPackage* CookPackage : PackagesToCook)
@@ -399,7 +400,7 @@ bool FContentBundleEditor::PopulateGeneratorPackageForCook(class IWorldPartition
 			}
 			else
 			{
-				if (!ExternalStreamingObject->Rename(nullptr, CookPackage->GetPackage()))
+				if (!ExternalStreamingObject->Rename(nullptr, CookPackage->GetPackage(), REN_DontCreateRedirectors))
 				{
 					UE_LOG(LogContentBundle, Error, TEXT("[CB: %s][Cook] Failed to set streaming object package %s."), *GetDescriptor()->GetDisplayName());
 					bIsSuccess = false;
@@ -408,6 +409,10 @@ bool FContentBundleEditor::PopulateGeneratorPackageForCook(class IWorldPartition
 		}
 
 		ExternalStreamingObject->PopulateGeneratorPackageForCook();
+
+		// Clear Streaming Object Reference past this point since it is not used.
+		// It's lifetime must be managed by the GeneratedPackage from the CookOntheFlyServer.
+		ExternalStreamingObject = nullptr;
 	}
 	
 
@@ -418,28 +423,25 @@ bool FContentBundleEditor::PopulateGeneratedPackageForCook(class IWorldPartition
 {
 	bool bIsSuccess = true;
 
-	if (HasCookedContent())
+	if (PackageToCook.Type == FWorldPartitionCookPackage::EType::Level)
 	{
-		if (PackageToCook.Type == FWorldPartitionCookPackage::EType::Level)
+		if (UWorldPartitionRuntimeCell** MatchingCell = const_cast<UWorldPartitionRuntimeCell**>(CookPackageIdsToCell.Find(PackageToCook.PackageId)))
 		{
-			if (UWorldPartitionRuntimeCell** MatchingCell = const_cast<UWorldPartitionRuntimeCell**>(CookPackageIdsToCell.Find(PackageToCook.PackageId)))
+			UWorldPartitionRuntimeCell* Cell = *MatchingCell;
+			if (ensure(Cell))
 			{
-				UWorldPartitionRuntimeCell* Cell = *MatchingCell;
-				if (ensure(Cell))
+				TArray<UPackage*> ModifiedPackages;
+				if (!Cell->PopulateGeneratedPackageForCook(PackageToCook.GetPackage(), OutModifiedPackages))
 				{
-					TArray<UPackage*> ModifiedPackages;
-					if (!Cell->PopulateGeneratedPackageForCook(PackageToCook.GetPackage(), OutModifiedPackages))
-					{
-						UE_LOG(LogContentBundle, Error, TEXT("[CB: %s][Cook] Failed to populate cell package %s."), *GetDescriptor()->GetDisplayName(), *PackageToCook.RelativePath);
-						bIsSuccess = false;
-					}
+					UE_LOG(LogContentBundle, Error, TEXT("[CB: %s][Cook] Failed to populate cell package %s."), *GetDescriptor()->GetDisplayName(), *PackageToCook.RelativePath);
+					bIsSuccess = false;
 				}
 			}
-			else
-			{
-				UE_LOG(LogContentBundle, Error, TEXT("[CB: %s][Cook] Could not find cell for package %s while populating generated pacakges."), *GetDescriptor()->GetDisplayName(), *PackageToCook.RelativePath);
-				bIsSuccess = false;
-			}
+		}
+		else
+		{
+			UE_LOG(LogContentBundle, Error, TEXT("[CB: %s][Cook] Could not find cell for package %s while populating generated pacakges."), *GetDescriptor()->GetDisplayName(), *PackageToCook.RelativePath);
+			bIsSuccess = false;
 		}
 	}
 
diff --git a/Engine/Source/Runtime/Engine/Private/WorldPartition/WorldPartitionRuntimeSpatialHash.cpp b/Engine/Source/Runtime/Engine/Private/WorldPartition/WorldPartitionRuntimeSpatialHash.cpp
--- a/Engine/Source/Runtime/Engine/Private/WorldPartition/WorldPartitionRuntimeSpatialHash.cpp
+++ b/Engine/Source/Runtime/Engine/Private/WorldPartition/WorldPartitionRuntimeSpatialHash.cpp
@@ -749,33 +749,23 @@ ASpatialHashRuntimeGridInfo::ASpatialHashRuntimeGridInfo(const FObjectInitialize
 
 void URuntimeSpatialHashExternalStreamingObject::OnStreamingObjectLoaded()
 {
-	// Cooked streaming object's Cells do not have LevelStreaming and are outered to the streaming object.
-	// Create their level streaming now and re-outered to the runtime hash.
 	bool bIsACookedObject = !CellToLevelStreamingPackage.IsEmpty();
 	if (bIsACookedObject)
 	{
-		for (FSpatialHashStreamingGrid& StreamingGrid : StreamingGrids)
+		// Cooked streaming object's Cells do not have LevelStreaming and are outered to the streaming object.
+		// Create their level streaming now and re-outered to the runtime hash.
+		ForEachStreamingCells([this](UWorldPartitionRuntimeCell& Cell)
 		{
-			for (FSpatialHashStreamingGridLevel& GridLevel : StreamingGrid.GridLevels)
-			{
-				for (FSpatialHashStreamingGridLayerCell& GridLayerCell : GridLevel.LayerCells)
-				{
-					for (UWorldPartitionRuntimeSpatialHashCell* Cell : GridLayerCell.GridCells)
-					{
-						UWorldPartitionRuntimeLevelStreamingCell* RuntimeCell = CastChecked<UWorldPartitionRuntimeLevelStreamingCell>(Cell);
+			UWorldPartitionRuntimeLevelStreamingCell* RuntimeCell = CastChecked<UWorldPartitionRuntimeLevelStreamingCell>(&Cell);
 
-						FName LevelStreamingPackage = CellToLevelStreamingPackage.FindChecked(RuntimeCell->GetFName());
-						RuntimeCell->Rename(nullptr, GetOuterWorld()->GetWorldPartition()->RuntimeHash);
-						RuntimeCell->CreateAndSetLevelStreaming(*LevelStreamingPackage.ToString());
-					}
-				}
-			}
-		}
+			FName LevelStreamingPackage = CellToLevelStreamingPackage.FindChecked(RuntimeCell->GetFName());
+			RuntimeCell->Rename(nullptr, GetOuterWorld()->GetWorldPartition()->RuntimeHash);
+			RuntimeCell->CreateAndSetLevelStreaming(*LevelStreamingPackage.ToString());
+		});
 	}
 }
 
-#if WITH_EDITOR
-void URuntimeSpatialHashExternalStreamingObject::PopulateGeneratorPackageForCook()
+void URuntimeSpatialHashExternalStreamingObject::ForEachStreamingCells(TFunctionRef<void(UWorldPartitionRuntimeCell&)> Func)
 {
 	for (FSpatialHashStreamingGrid& StreamingGrid : StreamingGrids)
 	{
@@ -785,19 +775,29 @@ void URuntimeSpatialHashExternalStreamingObject::PopulateGeneratorPackageForCook
 			{
 				for (UWorldPartitionRuntimeSpatialHashCell* Cell : GridLayerCell.GridCells)
 				{
-					UWorldPartitionRuntimeLevelStreamingCell* RuntimeCell = CastChecked<UWorldPartitionRuntimeLevelStreamingCell>(Cell);
-					RuntimeCell->Rename(nullptr, this);
-					
-					UWorldPartitionLevelStreamingDynamic* LevelStreamingDynamic = RuntimeCell->GetLevelStreaming();
-					CellToLevelStreamingPackage.Add(RuntimeCell->GetFName(), LevelStreamingDynamic->PackageNameToLoad);
-
-					// Level streaming are outered to the world and would not be saved within the ExternalStreamingObject.
-					// Do not save them, instead they will be created once the external streaming object is loaded at runtime. 
-					LevelStreamingDynamic->SetFlags(RF_Transient);
+					Func(*Cell);
 				}
 			}
 		}
 	}
+
+}
+
+#if WITH_EDITOR
+void URuntimeSpatialHashExternalStreamingObject::PopulateGeneratorPackageForCook()
+{
+	ForEachStreamingCells([this](UWorldPartitionRuntimeCell& Cell)
+	{
+		UWorldPartitionRuntimeLevelStreamingCell* RuntimeCell = CastChecked<UWorldPartitionRuntimeLevelStreamingCell>(&Cell);
+		RuntimeCell->Rename(nullptr, this);
+
+		UWorldPartitionLevelStreamingDynamic* LevelStreamingDynamic = RuntimeCell->GetLevelStreaming();
+		CellToLevelStreamingPackage.Add(RuntimeCell->GetFName(), LevelStreamingDynamic->PackageNameToLoad);
+
+		// Level streaming are outered to the world and would not be saved within the ExternalStreamingObject.
+		// Do not save them, instead they will be created once the external streaming object is loaded at runtime. 
+		LevelStreamingDynamic->SetFlags(RF_Transient);
+	});
 }
 #endif
 
diff --git a/Engine/Source/Runtime/Engine/Public/WorldPartition/ContentBundle/ContentBundleEditor.h b/Engine/Source/Runtime/Engine/Public/WorldPartition/ContentBundle/ContentBundleEditor.h
--- a/Engine/Source/Runtime/Engine/Public/WorldPartition/ContentBundle/ContentBundleEditor.h
+++ b/Engine/Source/Runtime/Engine/Public/WorldPartition/ContentBundle/ContentBundleEditor.h
@@ -92,7 +92,6 @@ private:
 
 	TArray<FWorldPartitionReference> ForceLoadedActors;
 
-	TSet<const UWorldPartitionRuntimeCell*> ContentBundleCells;
 	TMap<uint32, const UWorldPartitionRuntimeCell*> CookPackageIdsToCell;
 	URuntimeHashExternalStreamingObjectBase* ExternalStreamingObject;
 
diff --git a/Engine/Source/Runtime/Engine/Public/WorldPartition/WorldPartitionRuntimeHash.h b/Engine/Source/Runtime/Engine/Public/WorldPartition/WorldPartitionRuntimeHash.h
--- a/Engine/Source/Runtime/Engine/Public/WorldPartition/WorldPartitionRuntimeHash.h
+++ b/Engine/Source/Runtime/Engine/Public/WorldPartition/WorldPartitionRuntimeHash.h
@@ -46,6 +46,8 @@ public:
 	virtual UWorld* GetOwningWorld() const override { return OwningWorld.Get(); }
 	virtual UWorld* GetOuterWorld() const override { return OuterWorld.Get(); }
 	//~ End IWorldPartitionRuntimeCellOwner Interface
+
+	virtual void ForEachStreamingCells(TFunctionRef<void(UWorldPartitionRuntimeCell&)> Func) {};
 	
 	virtual void OnStreamingObjectLoaded() {};
 
diff --git a/Engine/Source/Runtime/Engine/Public/WorldPartition/WorldPartitionRuntimeSpatialHash.h b/Engine/Source/Runtime/Engine/Public/WorldPartition/WorldPartitionRuntimeSpatialHash.h
--- a/Engine/Source/Runtime/Engine/Public/WorldPartition/WorldPartitionRuntimeSpatialHash.h
+++ b/Engine/Source/Runtime/Engine/Public/WorldPartition/WorldPartitionRuntimeSpatialHash.h
@@ -224,11 +224,15 @@ class URuntimeSpatialHashExternalStreamingObject : public URuntimeHashExternalSt
 	GENERATED_BODY()
 
 public:
+	//~ Begin URuntimeHashExternalStreamingObjectBase Interface
 	virtual void OnStreamingObjectLoaded() override;
 
+	virtual void ForEachStreamingCells(TFunctionRef<void(UWorldPartitionRuntimeCell&)> Func) override;
+
 #if WITH_EDITOR
 	virtual void PopulateGeneratorPackageForCook() override;
 #endif
+	//~ End URuntimeHashExternalStreamingObjectBase Interface
 
 	UPROPERTY();
 	TArray<FSpatialHashStreamingGrid> StreamingGrids;
