World Partition: reduce the cells sublevel names size to avoid getting the infamous maximum filename length on Windows.

#rnx

diff --git a/Engine/Source/Runtime/Engine/Private/WorldPartition/WorldPartitionRuntimeSpatialHash.cpp b/Engine/Source/Runtime/Engine/Private/WorldPartition/WorldPartitionRuntimeSpatialHash.cpp
--- a/Engine/Source/Runtime/Engine/Private/WorldPartition/WorldPartitionRuntimeSpatialHash.cpp
+++ b/Engine/Source/Runtime/Engine/Private/WorldPartition/WorldPartitionRuntimeSpatialHash.cpp
@@ -1033,7 +1033,7 @@ FString UWorldPartitionRuntimeSpatialHash::GetCellNameString(FName InGridName, c
 
 	if (InContentBundleID.IsValid())
 	{
-		CellName += TEXT("_CB") + InContentBundleID.ToString(EGuidFormats::Short);
+		CellName += FString::Printf(TEXT("CB%X"), GetTypeHash(InContentBundleID));
 	}
 
 	return CellName;
