Simplify logic of delegate registration rules

diff --git a/Engine/Source/Runtime/Engine/Classes/Components/SceneCaptureComponent.h b/Engine/Source/Runtime/Engine/Classes/Components/SceneCaptureComponent.h
--- a/Engine/Source/Runtime/Engine/Classes/Components/SceneCaptureComponent.h
+++ b/Engine/Source/Runtime/Engine/Classes/Components/SceneCaptureComponent.h
@@ -215,7 +215,6 @@ protected:
 	void RegisterDelegates();
 	void UnregisterDelegates();
 	void ReleaseGarbageReferences();
-	bool bDelegatesRegistered = false;
 
 	virtual void UpdateSceneCaptureContents(FSceneInterface* Scene) {};
 
diff --git a/Engine/Source/Runtime/Engine/Private/Components/SceneCaptureComponent.cpp b/Engine/Source/Runtime/Engine/Private/Components/SceneCaptureComponent.cpp
--- a/Engine/Source/Runtime/Engine/Private/Components/SceneCaptureComponent.cpp
+++ b/Engine/Source/Runtime/Engine/Private/Components/SceneCaptureComponent.cpp
@@ -169,7 +169,7 @@ USceneCaptureComponent::USceneCaptureComponent(const FObjectInitializer& ObjectI
 	ShowFlags.SetHMDDistortion(0);
 	ShowFlags.SetOnScreenDebug(0);
 
-	if (!HasAnyFlags(RF_ClassDefaultObject | RF_ArchetypeObject))
+	if (!IsTemplate())
 	{
 		if (HasAnyFlags(RF_NeedPostLoad) || GetOuter()->HasAnyFlags(RF_NeedPostLoad))
 		{
@@ -179,38 +179,39 @@ USceneCaptureComponent::USceneCaptureComponent(const FObjectInitializer& ObjectI
 		{
 			RegisterDelegates();
 		}
-	}
+}
 }
 
 void USceneCaptureComponent::RegisterDelegates()
 {
 	ensureMsgf(IsInGameThread(), TEXT("Potential race condition in USceneCaptureComponent registering to a delegate from non game-thread"));
 	FCoreUObjectDelegates::GetPreGarbageCollectDelegate().AddUObject(this, &USceneCaptureComponent::ReleaseGarbageReferences);
-	bDelegatesRegistered = true;
 }
 
 void USceneCaptureComponent::UnregisterDelegates()
 {
-	if (!HasAnyFlags(RF_ClassDefaultObject | RF_ArchetypeObject))
-	{
-		check(bDelegatesRegistered);
-
-		FCoreUObjectDelegates::GetPreGarbageCollectDelegate().RemoveAll(this);
-	}
+	ensureMsgf(IsInGameThread(), TEXT("Potential race condition in USceneCaptureComponent unregistering to a delegate from non game-thread"));
+	FCoreUObjectDelegates::GetPreGarbageCollectDelegate().RemoveAll(this);
 }
 
 void USceneCaptureComponent::PostLoad()
 {
 	Super::PostLoad();
 
-	RegisterDelegates();
+	if (!IsTemplate())
+	{
+		RegisterDelegates();
+	}
 }
 
 void USceneCaptureComponent::BeginDestroy()
 {
 	Super::BeginDestroy();
-	
-	UnregisterDelegates();
+
+	if (!IsTemplate())
+	{
+		UnregisterDelegates();
+	}
 }
 
 // Because HiddenActors and ShowOnlyActors are serialized they are not easily refactored into weak pointers.
