Avoid performing work unless required during FinishDestroy of UWorld

diff --git a/Engine/Source/Runtime/Engine/Private/World.cpp b/Engine/Source/Runtime/Engine/Private/World.cpp
--- a/Engine/Source/Runtime/Engine/Private/World.cpp
+++ b/Engine/Source/Runtime/Engine/Private/World.cpp
@@ -1141,22 +1141,25 @@ void UWorld::FinishDestroy()
 	{
 		UPackage* WorldPackage = GetOutermost();
 
-		bool bContainsAnotherWorld = false;
-		TArray<UObject*> PotentialWorlds;
-		GetObjectsWithPackage(WorldPackage, PotentialWorlds, false);
-		for (UObject* PotentialWorld : PotentialWorlds)
+		if (WorldPackage->HasAnyPackageFlags(PKG_ContainsMap))
 		{
-			UWorld* World = Cast<UWorld>(PotentialWorld);
-			if (World && World != this)
+			bool bContainsAnotherWorld = false;
+			TArray<UObject*> PotentialWorlds;
+			GetObjectsWithPackage(WorldPackage, PotentialWorlds, false);
+			for (UObject* PotentialWorld : PotentialWorlds)
 			{
-				bContainsAnotherWorld = true;
-				break;
+				UWorld* World = Cast<UWorld>(PotentialWorld);
+				if (World && World != this)
+				{
+					bContainsAnotherWorld = true;
+					break;
+				}
 			}
-		}
 
-		if ( !bContainsAnotherWorld )
-		{
-			WorldPackage->ClearPackageFlags(PKG_ContainsMap);
+			if ( !bContainsAnotherWorld )
+			{
+				WorldPackage->ClearPackageFlags(PKG_ContainsMap);
+			}
 		}
 	}
 
