Track dormant actors per connection per level to make it quicker to flush on level visibility changes

diff --git a/Engine/Source/Runtime/Engine/Classes/Engine/NetworkObjectList.h b/Engine/Source/Runtime/Engine/Classes/Engine/NetworkObjectList.h
--- a/Engine/Source/Runtime/Engine/Classes/Engine/NetworkObjectList.h
+++ b/Engine/Source/Runtime/Engine/Classes/Engine/NetworkObjectList.h
@@ -189,10 +189,18 @@ public:
 
 	void CountBytes(FArchive& Ar) const;
 
+	/** Marks any actors in the given package/level active if they were fully dormant or dormant for the passed in connection */
+	void FlushDormantActors(UNetConnection* const Connection, const FName& PackageName);
+
 private:
+	bool MarkActiveInternal(const TSharedPtr<FNetworkObjectInfo>& ObjectInfo, UNetConnection* const Connection, UNetDriver* NetDriver);
+
 	FNetworkObjectSet AllNetworkObjects;
 	FNetworkObjectSet ActiveNetworkObjects;
 	FNetworkObjectSet ObjectsDormantOnAllConnections;
 
 	TMap<TWeakObjectPtr<UNetConnection>, int32 > NumDormantObjectsPerConnection;
+
+	TMap<FName, FNetworkObjectSet> FullyDormantObjectsByLevel;
+	TMap<TObjectKey<UNetConnection>, TMap<FName, FNetworkObjectSet>> DormantObjectsPerConnection;
 };
diff --git a/Engine/Source/Runtime/Engine/Private/NetConnection.cpp b/Engine/Source/Runtime/Engine/Private/NetConnection.cpp
--- a/Engine/Source/Runtime/Engine/Private/NetConnection.cpp
+++ b/Engine/Source/Runtime/Engine/Private/NetConnection.cpp
@@ -226,6 +226,11 @@ static SIZE_T CountBytes(FWrittenChannelsRecord& WrittenChannelsRecord);
 
 };
 
+namespace UE::Net::Private
+{
+	extern bool bTrackDormantObjectsByLevel;
+}
+
 PRAGMA_DISABLE_DEPRECATION_WARNINGS
 const int32 UNetConnection::DEFAULT_MAX_CHANNEL_SIZE = 32767;
 PRAGMA_ENABLE_DEPRECATION_WARNINGS
@@ -1516,14 +1521,19 @@ void UNetConnection::UpdateLevelVisibilityInternal(const FUpdateLevelVisibilityL
 			}
 			else 
 #endif // UE_WITH_IRIS
-			if (VisibilityResult.Package)
 			{
-				LevelWorld = (UWorld*)FindObjectWithOuter(VisibilityResult.Package, UWorld::StaticClass());
-				if (LevelWorld)
+				LevelWorld = VisibilityResult.Package ? (UWorld*)FindObjectWithOuter(VisibilityResult.Package, UWorld::StaticClass()) : nullptr;
+
+				FNetworkObjectList& NetworkObjectList = Driver->GetNetworkObjectList();
+
+				if (UE::Net::Private::bTrackDormantObjectsByLevel)
+				{
+					NetworkObjectList.FlushDormantActors(this, LevelVisibility.PackageName);
+				}
+				else
 				{
-					if (LevelWorld->PersistentLevel)
+					if (LevelWorld && LevelWorld->PersistentLevel)
 					{
-						FNetworkObjectList& NetworkObjectList = Driver->GetNetworkObjectList();
 						for (AActor* Actor : LevelWorld->PersistentLevel->Actors)
 						{
 							// Dormant Initial actors have no changes. Dormant Never and Awake will be sent normal, so we only need
diff --git a/Engine/Source/Runtime/Engine/Private/NetworkObjectList.cpp b/Engine/Source/Runtime/Engine/Private/NetworkObjectList.cpp
--- a/Engine/Source/Runtime/Engine/Private/NetworkObjectList.cpp
+++ b/Engine/Source/Runtime/Engine/Private/NetworkObjectList.cpp
@@ -7,6 +7,16 @@
 #include "Serialization/Archive.h"
 #include "Net/NetworkGranularMemoryLogging.h"
 
+namespace UE::Net::Private
+{
+	bool bTrackDormantObjectsByLevel = false;
+	static FAutoConsoleVariableRef CVarNetTrackDormantObjectsByLevel(
+		TEXT("net.TrackDormantObjectsByLevel"),
+		bTrackDormantObjectsByLevel,
+		TEXT("When true, network object list will maintain a set of dormant actors per connnection per level."),
+		ECVF_Default);
+}
+
 void FNetworkObjectList::AddInitialObjects(UWorld* const World, UNetDriver* NetDriver)
 {
 	if (World == nullptr || NetDriver == nullptr)
@@ -115,6 +125,8 @@ void FNetworkObjectList::Remove(AActor* const Actor)
 		return;
 	}
 
+	const FName PackageName = Actor->GetLevel()->GetPackage()->GetFName();
+
 	FNetworkObjectInfo* NetworkObjectInfo = NetworkObjectInfoPtr->Get();
 
 	// Lower the dormant object count for each connection this object is dormant on
@@ -125,12 +137,33 @@ void FNetworkObjectList::Remove(AActor* const Actor)
 		if (Connection == nullptr || Connection->GetConnectionState() == USOCK_Closed)
 		{
 			ConnectionIt.RemoveCurrent();
+			if (UE::Net::Private::bTrackDormantObjectsByLevel)
+			{
+				DormantObjectsPerConnection.Remove(Connection);
+			}
 			continue;
 		}
 
 		int32& NumDormantObjectsPerConnectionRef = NumDormantObjectsPerConnection.FindOrAdd(Connection);
-		check( NumDormantObjectsPerConnectionRef > 0 );
+		check(NumDormantObjectsPerConnectionRef > 0);
 		NumDormantObjectsPerConnectionRef--;
+
+		if (UE::Net::Private::bTrackDormantObjectsByLevel)
+		{
+			if (TMap<FName, FNetworkObjectSet>* DormantObjectsByLevel = DormantObjectsPerConnection.Find(Connection))
+			{
+				if (FNetworkObjectSet* DormantObjects = DormantObjectsByLevel->Find(PackageName))
+				{
+					const int32 NumRemoved = DormantObjects->Remove(Actor);
+					checkf(NumRemoved > 0, TEXT("Actor not found in Connection->Level->Dormant map: %s"), *GetNameSafe(Actor));
+
+					if (DormantObjects->Num() == 0)
+					{
+						DormantObjectsByLevel->Remove(PackageName);
+					}
+				}
+			}
+		}
 	}
 
 	// Remove this object from all lists
@@ -138,6 +171,19 @@ void FNetworkObjectList::Remove(AActor* const Actor)
 	ActiveNetworkObjects.Remove(Actor);
 	ObjectsDormantOnAllConnections.Remove(Actor);
 
+	if (UE::Net::Private::bTrackDormantObjectsByLevel)
+	{
+		if (FNetworkObjectSet* FullyDormant = FullyDormantObjectsByLevel.Find(PackageName))
+		{
+			FullyDormant->Remove(Actor);
+
+			if (FullyDormant->Num() == 0)
+			{
+				FullyDormantObjectsByLevel.Remove(PackageName);
+			}
+		}
+	}
+
 	check((ActiveNetworkObjects.Num() + ObjectsDormantOnAllConnections.Num()) == AllNetworkObjects.Num());
 }
 
@@ -150,6 +196,8 @@ void FNetworkObjectList::MarkDormant(AActor* const Actor, UNetConnection* const
 		return;		// Actor doesn't belong to this net driver name
 	}
 
+	const FName PackageName = Actor->GetLevel()->GetPackage()->GetFName();
+
 	FNetworkObjectInfo* NetworkObjectInfo = NetworkObjectInfoPtr->Get();
 
 	// Add the connection to the list of dormant connections (if it's not already on the list)
@@ -163,6 +211,18 @@ void FNetworkObjectList::MarkDormant(AActor* const Actor, UNetConnection* const
 		int32& NumDormantObjectsPerConnectionRef = NumDormantObjectsPerConnection.FindOrAdd(Connection);
 		NumDormantObjectsPerConnectionRef++;
 
+		if (UE::Net::Private::bTrackDormantObjectsByLevel)
+		{
+			// if not fully dormant
+			if (NetworkObjectInfo->DormantConnections.Num() != NumConnections)
+			{
+				TMap<FName, FNetworkObjectSet>& DormantObjectsByLevel = DormantObjectsPerConnection.FindOrAdd(Connection);
+				FNetworkObjectSet& DormantObjects = DormantObjectsByLevel.FindOrAdd(PackageName);
+
+				DormantObjects.Add(*NetworkObjectInfoPtr);
+			}
+		}
+
 		UE_LOG(LogNetDormancy, Log, TEXT("FNetworkObjectList::MarkDormant: Actor is now dormant. Actor: %s. NumDormant: %i, Connection: %s"), *Actor->GetName(), NumDormantObjectsPerConnectionRef, *Connection->GetName());
 	}
 
@@ -172,6 +232,10 @@ void FNetworkObjectList::MarkDormant(AActor* const Actor, UNetConnection* const
 		if ((*ConnectionIt).Get() == nullptr || (*ConnectionIt).Get()->GetConnectionState() == USOCK_Closed)
 		{
 			ConnectionIt.RemoveCurrent();
+			if (UE::Net::Private::bTrackDormantObjectsByLevel)
+			{
+				DormantObjectsPerConnection.Remove((*ConnectionIt).Get());
+			}
 		}
 	}
 
@@ -182,6 +246,19 @@ void FNetworkObjectList::MarkDormant(AActor* const Actor, UNetConnection* const
 	if (NetworkObjectInfo->DormantConnections.Num() == NumConnections)
 	{
 		ObjectsDormantOnAllConnections.Add(*NetworkObjectInfoPtr);
+
+		if (UE::Net::Private::bTrackDormantObjectsByLevel)
+		{
+			FNetworkObjectSet& FullyDormant = FullyDormantObjectsByLevel.FindOrAdd(PackageName);
+			FullyDormant.Add(*NetworkObjectInfoPtr);
+
+			// Remove from connection object lists
+			for (auto ConnectionIt = NetworkObjectInfo->DormantConnections.CreateIterator(); ConnectionIt; ++ConnectionIt)
+			{
+				DormantObjectsPerConnection.Remove((*ConnectionIt).Get());
+			}
+		}
+
 		ActiveNetworkObjects.Remove(Actor);
 
 		UE_LOG(LogNetDormancy, Log, TEXT("FNetworkObjectList::MarkDormant: Actor is now dormant on all connections. Actor: %s. Total: %i, Active: %i, Connection: %s"), *Actor->GetName(), AllNetworkObjects.Num(), ActiveNetworkObjects.Num(), *Connection->GetName());
@@ -190,24 +267,33 @@ void FNetworkObjectList::MarkDormant(AActor* const Actor, UNetConnection* const
 	check((ActiveNetworkObjects.Num() + ObjectsDormantOnAllConnections.Num()) == AllNetworkObjects.Num());
 }
 
-bool FNetworkObjectList::MarkActive(AActor* const Actor, UNetConnection* const Connection, UNetDriver* NetDriver)
+bool FNetworkObjectList::MarkActiveInternal(const TSharedPtr<FNetworkObjectInfo>& ObjectInfo, UNetConnection* const Connection, UNetDriver* NetDriver)
 {
-	TSharedPtr<FNetworkObjectInfo>* NetworkObjectInfoPtr = FindOrAdd(Actor, NetDriver);
+	FNetworkObjectInfo* NetworkObjectInfo = ObjectInfo.Get();
+	AActor* Actor = NetworkObjectInfo->Actor;
 
-	if (NetworkObjectInfoPtr == nullptr)
-	{
-		return false;		// Actor doesn't belong to this net driver name
-	}
-
-	FNetworkObjectInfo* NetworkObjectInfo = NetworkObjectInfoPtr->Get();
+	const FName PackageName = Actor->GetLevel()->GetPackage()->GetFName();
 
 	// Remove from the ObjectsDormantOnAllConnections if needed
 	if (ObjectsDormantOnAllConnections.Remove(Actor) > 0)
 	{
 		// Put this object back on the active list
-		ActiveNetworkObjects.Add(*NetworkObjectInfoPtr);
+		ActiveNetworkObjects.Add(ObjectInfo);
 
 		UE_LOG(LogNetDormancy, Log, TEXT("FNetworkObjectList::MarkDormant: Actor is no longer dormant on all connections. Actor: %s. Total: %i, Active: %i, Connection: %s"), *Actor->GetName(), AllNetworkObjects.Num(), ActiveNetworkObjects.Num(), *Connection->GetName());
+
+		if (UE::Net::Private::bTrackDormantObjectsByLevel)
+		{
+			if (FNetworkObjectSet* FullyDormant = FullyDormantObjectsByLevel.Find(PackageName))
+			{
+				FullyDormant->Remove(Actor);
+
+				if (FullyDormant->Num() == 0)
+				{
+					FullyDormantObjectsByLevel.Remove(PackageName);
+				}
+			}
+		}
 	}
 
 	check((ActiveNetworkObjects.Num() + ObjectsDormantOnAllConnections.Num()) == AllNetworkObjects.Num());
@@ -216,12 +302,29 @@ bool FNetworkObjectList::MarkActive(AActor* const Actor, UNetConnection* const C
 	if (NetworkObjectInfo->DormantConnections.Remove(Connection) > 0)
 	{
 		// Add the connection to the list of recently dormant connections
-		NetworkObjectInfo->RecentlyDormantConnections.Add( Connection );
+		NetworkObjectInfo->RecentlyDormantConnections.Add(Connection);
 
 		int32& NumDormantObjectsPerConnectionRef = NumDormantObjectsPerConnection.FindOrAdd(Connection);
 		check(NumDormantObjectsPerConnectionRef > 0);
 		NumDormantObjectsPerConnectionRef--;
 
+		if (UE::Net::Private::bTrackDormantObjectsByLevel)
+		{
+			if (TMap<FName, FNetworkObjectSet>* DormantObjectsByLevel = DormantObjectsPerConnection.Find(Connection))
+			{
+				if (FNetworkObjectSet* DormantObjects = DormantObjectsByLevel->Find(PackageName))
+				{
+					const int32 NumRemoved = DormantObjects->Remove(Actor);
+					checkf(NumRemoved > 0, TEXT("Actor not found in Connection->Level->Dormant map: %s"), *GetNameSafe(Actor));
+
+					if (DormantObjects->Num() == 0)
+					{
+						DormantObjectsByLevel->Remove(PackageName);
+					}
+				}
+			}
+		}
+
 		UE_LOG(LogNetDormancy, Log, TEXT("FNetworkObjectList::MarkActive: Actor is no longer dormant. Actor: %s. NumDormant: %i, Connection: %s"), *Actor->GetName(), NumDormantObjectsPerConnectionRef, *Connection->GetName());
 		return true;
 	}
@@ -229,6 +332,18 @@ bool FNetworkObjectList::MarkActive(AActor* const Actor, UNetConnection* const C
 	return false;
 }
 
+bool FNetworkObjectList::MarkActive(AActor* const Actor, UNetConnection* const Connection, UNetDriver* NetDriver)
+{
+	TSharedPtr<FNetworkObjectInfo>* NetworkObjectInfoPtr = FindOrAdd(Actor, NetDriver);
+
+	if (NetworkObjectInfoPtr == nullptr)
+	{
+		return false;		// Actor doesn't belong to this net driver name
+	}
+
+	return MarkActiveInternal(*NetworkObjectInfoPtr, Connection, NetDriver);
+}
+
 void FNetworkObjectList::MarkDirtyForReplay(AActor* const Actor)
 {
 	if (Actor)
@@ -267,12 +382,14 @@ void FNetworkObjectList::HandleConnectionAdded()
 	}
 
 	ObjectsDormantOnAllConnections.Empty();
+	FullyDormantObjectsByLevel.Empty();
 }
 
 void FNetworkObjectList::ResetDormancyState()
 {
 	// Reset all state related to dormancy, and move all objects back on to the active list
 	ObjectsDormantOnAllConnections.Empty();
+	FullyDormantObjectsByLevel.Empty();
 
 	ActiveNetworkObjects = AllNetworkObjects;
 
@@ -285,6 +402,7 @@ void FNetworkObjectList::ResetDormancyState()
 	}
 
 	NumDormantObjectsPerConnection.Empty();
+	DormantObjectsPerConnection.Empty();
 }
 
 int32 FNetworkObjectList::GetNumDormantActorsForConnection(UNetConnection* const Connection) const
@@ -315,6 +433,8 @@ void FNetworkObjectList::Reset()
 	ActiveNetworkObjects.Empty();
 	ObjectsDormantOnAllConnections.Empty();
 	NumDormantObjectsPerConnection.Empty();
+	FullyDormantObjectsByLevel.Empty();
+	DormantObjectsPerConnection.Empty();
 }
 
 void FNetworkObjectInfo::CountBytes(FArchive& Ar) const
@@ -333,6 +453,27 @@ void FNetworkObjectList::CountBytes(FArchive& Ar) const
 	GRANULAR_NETWORK_MEMORY_TRACKING_TRACK("ObjectsDormantOnAllConnections", ObjectsDormantOnAllConnections.CountBytes(Ar));
 	GRANULAR_NETWORK_MEMORY_TRACKING_TRACK("NumDormantObjectsPerConnection", NumDormantObjectsPerConnection.CountBytes(Ar));
 
+	GRANULAR_NETWORK_MEMORY_TRACKING_TRACK("FullyDormantObjectsByLevel",
+		FullyDormantObjectsByLevel.CountBytes(Ar);
+		for (const TPair<FName, FNetworkObjectSet>& LevelPair : FullyDormantObjectsByLevel)
+		{
+			LevelPair.Value.CountBytes(Ar);
+		}
+	);
+	
+	GRANULAR_NETWORK_MEMORY_TRACKING_TRACK("DormantObjectsPerConnection",
+		DormantObjectsPerConnection.CountBytes(Ar);
+		for (const TPair<TObjectKey<UNetConnection>, TMap<FName, FNetworkObjectSet>>& ConnectionPair : DormantObjectsPerConnection)
+		{
+			ConnectionPair.Value.CountBytes(Ar);
+
+			for (const TPair<FName, FNetworkObjectSet>& LevelPair : ConnectionPair.Value)
+			{
+				LevelPair.Value.CountBytes(Ar);
+			}
+		}
+	);
+
 	// ObjectsDormantOnAllConnections and ActiveNetworkObjects are both sub sets of AllNetworkObjects
 	// and only have pointers back to the data there.
 	// So, to avoid double (or triple) counting, only explicit count the elements from AllNetworkObjects.
@@ -347,4 +488,29 @@ void FNetworkObjectList::CountBytes(FArchive& Ar) const
 			}
 		}
 	);
+}
+
+void FNetworkObjectList::FlushDormantActors(UNetConnection* const Connection, const FName& PackageName)
+{
+	FNetworkObjectSet DormantActors;
+
+	// fully dormant actors
+	if (const FNetworkObjectSet* FullyDormantObjects = FullyDormantObjectsByLevel.Find(PackageName))
+	{
+		DormantActors.Append(*FullyDormantObjects);
+	}
+
+	// connection specific dormant actors
+	if (const TMap<FName, FNetworkObjectSet>* DormantObjectsByLevel = DormantObjectsPerConnection.Find(Connection))
+	{
+		if (const FNetworkObjectSet* DormantObjects = DormantObjectsByLevel->Find(PackageName))
+		{
+			DormantActors.Append(*DormantObjects);
+		}
+	}
+
+	for (const TSharedPtr<FNetworkObjectInfo>& ActorInfo : DormantActors)
+	{
+		MarkActiveInternal(ActorInfo, Connection, Connection->Driver);
+	}
 }
\ No newline at end of file
