Do not clear ExternalStreamingObject in PopulateGeneratorPackage.
It can be garbage collected before all cells are processed.
Will re-introduce warning UE-165762 which will be fixed subsequently

diff --git a/Engine/Source/Runtime/Engine/Private/WorldPartition/ContentBundle/ContentBundleEditor.cpp b/Engine/Source/Runtime/Engine/Private/WorldPartition/ContentBundle/ContentBundleEditor.cpp
--- a/Engine/Source/Runtime/Engine/Private/WorldPartition/ContentBundle/ContentBundleEditor.cpp
+++ b/Engine/Source/Runtime/Engine/Private/WorldPartition/ContentBundle/ContentBundleEditor.cpp
@@ -409,10 +409,6 @@ bool FContentBundleEditor::PopulateGeneratorPackageForCook(class IWorldPartition
 		}
 
 		ExternalStreamingObject->PopulateGeneratorPackageForCook();
-
-		// Clear Streaming Object Reference past this point since it is not used.
-		// It's lifetime must be managed by the GeneratedPackage from the CookOntheFlyServer.
-		ExternalStreamingObject = nullptr;
 	}
 	
 
