Fix for editor async loading, where objects were being serialized out of order, leading to ConditionalPostLoad() being called on uninitialized static meshes.
Reverted previous workaround, preloading static mesh from UStaticMeshComponent::PostLoad.

Unshelved from pending changelist '22411424', submitted on behalf of Francis Hurteau.

diff --git a/Engine/Source/Runtime/CoreUObject/Private/Serialization/AsyncLoading.cpp b/Engine/Source/Runtime/CoreUObject/Private/Serialization/AsyncLoading.cpp
--- a/Engine/Source/Runtime/CoreUObject/Private/Serialization/AsyncLoading.cpp
+++ b/Engine/Source/Runtime/CoreUObject/Private/Serialization/AsyncLoading.cpp
@@ -6903,7 +6903,7 @@ EAsyncPackageState::Type FAsyncPackage::PostLoadObjects()
 				--PostLoadIndex;
 				break;
 			}
-			else if (!bIsMultithreaded || (bAsyncPostLoadEnabled && CanPostLoadOnAsyncLoadingThread(Object)))
+			else if (/*!bIsMultithreaded ||*/ (bAsyncPostLoadEnabled && CanPostLoadOnAsyncLoadingThread(Object)))
 			{
 				SCOPED_ACCUM_LOADTIME(PostLoad, StaticGetNativeClassName(Object->GetClass()));
 
diff --git a/Engine/Source/Runtime/Engine/Private/Components/StaticMeshComponent.cpp b/Engine/Source/Runtime/Engine/Private/Components/StaticMeshComponent.cpp
--- a/Engine/Source/Runtime/Engine/Private/Components/StaticMeshComponent.cpp
+++ b/Engine/Source/Runtime/Engine/Private/Components/StaticMeshComponent.cpp
@@ -2049,21 +2049,9 @@ void UStaticMeshComponent::PostLoad()
 	NotifyIfStaticMeshChanged();
 
 	// need to postload the StaticMesh because super initializes variables based on GetStaticLightingType() which we override and use from the StaticMesh
-	UStaticMesh* ReferencedStaticMesh = GetStaticMesh();
-	if (ReferencedStaticMesh)
+	if (GetStaticMesh())
 	{
-#if WITH_EDITOR
-		// Preload the referenced static mesh here if it hasn't yet been serialized, prior to the ConditionalPostLoad call.
-		// This is temporary fix, pending a more robust solution for async loading in editor.
-		if (ReferencedStaticMesh->HasAnyFlags(RF_NeedLoad))
-		{
-			if (FLinkerLoad* Linker = ReferencedStaticMesh->GetLinker())
-			{
-				Linker->Preload(ReferencedStaticMesh);
-			}
-		}
-#endif // WITH_EDITOR
-		ReferencedStaticMesh->ConditionalPostLoad();
+		GetStaticMesh()->ConditionalPostLoad();
 	}
 
 	Super::PostLoad();
