World Partition Content Bundle: implement URuntimeHashExternalStreamingObjectBase::GetWorld.

#rnx

diff --git a/Engine/Source/Runtime/Engine/Public/WorldPartition/WorldPartitionRuntimeHash.h b/Engine/Source/Runtime/Engine/Public/WorldPartition/WorldPartitionRuntimeHash.h
--- a/Engine/Source/Runtime/Engine/Public/WorldPartition/WorldPartitionRuntimeHash.h
+++ b/Engine/Source/Runtime/Engine/Public/WorldPartition/WorldPartitionRuntimeHash.h
@@ -47,6 +47,8 @@ public:
 	virtual UWorld* GetOuterWorld() const override { return OuterWorld.Get(); }
 	//~ End IWorldPartitionRuntimeCellOwner Interface
 
+	virtual class UWorld* GetWorld() const override { return GetOwningWorld(); }
+
 	virtual void ForEachStreamingCells(TFunctionRef<void(UWorldPartitionRuntimeCell&)> Func) {};
 	
 	virtual void OnStreamingObjectLoaded() {};
