[Engine]
* Moved some small types out of larger files just to be able to include only those types. Types involved are ECommonInputType and EChannelMaskParameterColor

This change is a strategical submit for larger refactorings coming

diff --git a/Engine/Plugins/Experimental/CommonUI/Source/CommonInput/Public/CommonInputBaseTypes.h b/Engine/Plugins/Experimental/CommonUI/Source/CommonInput/Public/CommonInputBaseTypes.h
--- a/Engine/Plugins/Experimental/CommonUI/Source/CommonInput/Public/CommonInputBaseTypes.h
+++ b/Engine/Plugins/Experimental/CommonUI/Source/CommonInput/Public/CommonInputBaseTypes.h
@@ -3,6 +3,7 @@
 #pragma once
 
 #include "CoreMinimal.h"
+#include "CommonInputTypeEnum.h"
 #include "UObject/Object.h"
 #include "UObject/SoftObjectPath.h"
 #include "Templates/SubclassOf.h"
@@ -27,17 +28,6 @@ struct COMMONINPUT_API FCommonInputDefaults
 	static const FName GamepadGeneric;
 };
 
-UENUM(BlueprintType)
-enum class ECommonInputType : uint8
-{
-	MouseAndKeyboard,
-	Gamepad,
-	Touch,
-	Count
-};
-
-ENUM_RANGE_BY_COUNT(ECommonInputType, ECommonInputType::Count);
-
 USTRUCT(Blueprintable)
 struct COMMONINPUT_API FCommonInputKeyBrushConfiguration
 {
diff --git a/Engine/Plugins/Experimental/CommonUI/Source/CommonInput/Public/CommonInputTypeEnum.h b/Engine/Plugins/Experimental/CommonUI/Source/CommonInput/Public/CommonInputTypeEnum.h
new file mode 100644
--- /dev/null
+++ b/Engine/Plugins/Experimental/CommonUI/Source/CommonInput/Public/CommonInputTypeEnum.h
@@ -0,0 +1,17 @@
+// Copyright Epic Games, Inc. All Rights Reserved.
+
+#pragma once
+
+#include "Misc/EnumRange.h"
+#include "CommonInputTypeEnum.generated.h"
+
+UENUM(BlueprintType)
+enum class ECommonInputType : uint8
+{
+	MouseAndKeyboard,
+	Gamepad,
+	Touch,
+	Count
+};
+
+ENUM_RANGE_BY_COUNT(ECommonInputType, ECommonInputType::Count);
diff --git a/Engine/Source/Runtime/Engine/Classes/Materials/MaterialExpressionChannelMaskParameter.h b/Engine/Source/Runtime/Engine/Classes/Materials/MaterialExpressionChannelMaskParameter.h
--- a/Engine/Source/Runtime/Engine/Classes/Materials/MaterialExpressionChannelMaskParameter.h
+++ b/Engine/Source/Runtime/Engine/Classes/Materials/MaterialExpressionChannelMaskParameter.h
@@ -4,21 +4,10 @@
 
 #include "CoreMinimal.h"
 #include "UObject/ObjectMacros.h"
+#include "Materials/MaterialExpressionChannelMaskParameterColor.h"
 #include "Materials/MaterialExpressionVectorParameter.h"
 #include "MaterialExpressionChannelMaskParameter.generated.h"
 
-UENUM()
-namespace EChannelMaskParameterColor
-{
-	enum Type
-	{
-		Red,
-		Green,
-		Blue,
-		Alpha,
-	};
-}
-
 UCLASS(collapsecategories, hidecategories=(Object, MaterialExpressionVectorParameter), MinimalAPI)
 class UMaterialExpressionChannelMaskParameter : public UMaterialExpressionVectorParameter
 {
diff --git a/Engine/Source/Runtime/Engine/Classes/Materials/MaterialExpressionChannelMaskParameterColor.h b/Engine/Source/Runtime/Engine/Classes/Materials/MaterialExpressionChannelMaskParameterColor.h
new file mode 100644
--- /dev/null
+++ b/Engine/Source/Runtime/Engine/Classes/Materials/MaterialExpressionChannelMaskParameterColor.h
@@ -0,0 +1,17 @@
+// Copyright Epic Games, Inc. All Rights Reserved.
+
+#pragma once
+
+#include "MaterialExpressionChannelMaskParameterColor.generated.h"
+
+UENUM()
+namespace EChannelMaskParameterColor
+{
+	enum Type
+	{
+		Red,
+		Green,
+		Blue,
+		Alpha,
+	};
+}
diff --git a/Engine/Source/Runtime/Engine/Public/Physics/PhysicsInterfaceScene.h b/Engine/Source/Runtime/Engine/Public/Physics/PhysicsInterfaceScene.h
new file mode 100644
--- /dev/null
+++ b/Engine/Source/Runtime/Engine/Public/Physics/PhysicsInterfaceScene.h
@@ -0,0 +1 @@
+#include "Physics/Experimental/PhysScene_Chaos.h"
\ No newline at end of file
diff --git a/Engine/Source/Runtime/Engine/Public/StaticMeshComponentLODInfo.h b/Engine/Source/Runtime/Engine/Public/StaticMeshComponentLODInfo.h
new file mode 100644
--- /dev/null
+++ b/Engine/Source/Runtime/Engine/Public/StaticMeshComponentLODInfo.h
@@ -0,0 +1 @@
+// placeholder
\ No newline at end of file
