[Engine]
* Added missing copyright comments

diff --git a/Engine/Source/Runtime/Engine/Public/Physics/PhysicsInterfaceScene.h b/Engine/Source/Runtime/Engine/Public/Physics/PhysicsInterfaceScene.h
--- a/Engine/Source/Runtime/Engine/Public/Physics/PhysicsInterfaceScene.h
+++ b/Engine/Source/Runtime/Engine/Public/Physics/PhysicsInterfaceScene.h
@@ -1 +1,5 @@
+// Copyright Epic Games, Inc. All Rights Reserved.
+
+#pragma once
+
 #include "Physics/Experimental/PhysScene_Chaos.h"
\ No newline at end of file
diff --git a/Engine/Source/Runtime/Engine/Public/StaticMeshComponentLODInfo.h b/Engine/Source/Runtime/Engine/Public/StaticMeshComponentLODInfo.h
--- a/Engine/Source/Runtime/Engine/Public/StaticMeshComponentLODInfo.h
+++ b/Engine/Source/Runtime/Engine/Public/StaticMeshComponentLODInfo.h
@@ -1 +1,5 @@
+// Copyright Epic Games, Inc. All Rights Reserved.
+
+#pragma once
+
 // placeholder
\ No newline at end of file
