Improve console keyboard shortcuts: left/right/backspace/delete with ctrl key now behaves like common text editors.

diff --git a/Engine/Source/Runtime/Engine/Private/UserInterface/Console.cpp b/Engine/Source/Runtime/Engine/Private/UserInterface/Console.cpp
--- a/Engine/Source/Runtime/Engine/Private/UserInterface/Console.cpp
+++ b/Engine/Source/Runtime/Engine/Private/UserInterface/Console.cpp
@@ -910,6 +910,23 @@ bool UConsole::InputKey_InputLine(FInputDeviceId DeviceId, FKey Key, EInputEvent
 		}
 	};
 
+	auto FindWordBreak = [](const FString& Str, uint32 StartPos, ESearchDir::Type Direction)
+	{
+		// find the nearest '.' or ' '
+		int32 SpacePos = Str.Find(TEXT(" "), ESearchCase::CaseSensitive, Direction, StartPos);
+		int32 PeriodPos = Str.Find(TEXT("."), ESearchCase::CaseSensitive, Direction, StartPos);
+		if (Direction == ESearchDir::FromEnd)
+		{
+			return FMath::Max(SpacePos, PeriodPos);
+		}
+		else
+		{ 
+			int32 Result = SpacePos < 0 ? PeriodPos : (PeriodPos < 0 ? SpacePos : FMath::Min(SpacePos, PeriodPos));
+			Result = Result == INDEX_NONE ? Str.Len() : Result;
+			return Result;
+		}
+	};
+
 	// if user input is open
 	if (ConsoleState != NAME_None)
 	{
@@ -1062,11 +1079,23 @@ bool UConsole::InputKey_InputLine(FInputDeviceId DeviceId, FKey Key, EInputEvent
 		{
 			if (TypedStrPos > 0)
 			{
-				SetInputText(FString::Printf(TEXT("%s%s"), *TypedStr.Left(TypedStrPos - 1), *TypedStr.Right(TypedStr.Len() - TypedStrPos)));
-				SetCursorPos(TypedStrPos - 1);
+				int32 NewPos;
+				if (bCtrl)
+				{
+					NewPos = FMath::Max(0, FindWordBreak(TypedStr, TypedStrPos, ESearchDir::FromEnd));
+				}
+				else
+				{
+					NewPos = TypedStrPos - 1;
+				}
+
+				SetInputText(FString::Printf(TEXT("%s%s"), *TypedStr.Left(NewPos), *TypedStr.Right(TypedStr.Len() - TypedStrPos)));
+				SetCursorPos(NewPos);
+
 				// unlock auto-complete (@todo - track the lock position so we don't bother unlocking under bogus cases)
 				bAutoCompleteLocked = false;
 			}
+			bCaptureKeyInput = true;
 
 			return true;
 		}
@@ -1074,44 +1103,46 @@ bool UConsole::InputKey_InputLine(FInputDeviceId DeviceId, FKey Key, EInputEvent
 		{
 			if (TypedStrPos < TypedStr.Len())
 			{
-				SetInputText(FString::Printf(TEXT("%s%s"), *TypedStr.Left(TypedStrPos), *TypedStr.Right(TypedStr.Len() - TypedStrPos - 1)));
+				int32 RightStart;
+				if (bCtrl)
+				{
+					RightStart = FindWordBreak(TypedStr, TypedStrPos + 1, ESearchDir::FromStart);
+				}
+				else
+				{
+					RightStart = TypedStrPos + 1;
+				}
+
+				SetInputText(FString::Printf(TEXT("%s%s"), *TypedStr.Left(TypedStrPos), *TypedStr.Right(TypedStr.Len() - RightStart)));
 			}
 			return true;
 		}
 		else if (Key == EKeys::Left)
 		{
+			int32 NewPos;
 			if (bCtrl)
 			{
-				// find the nearest '.' or ' '
-				int32 NewPos = FMath::Max(TypedStr.Find(TEXT("."), ESearchCase::CaseSensitive, ESearchDir::FromEnd, TypedStrPos), TypedStr.Find(TEXT(" "), ESearchCase::CaseSensitive, ESearchDir::FromEnd, TypedStrPos));
-				SetCursorPos(FMath::Max(0, NewPos));
+				NewPos = FMath::Min(FindWordBreak(TypedStr, FMath::Max(0, TypedStrPos - 1), ESearchDir::FromEnd) + 1, TypedStr.Len());
 			}
 			else
 			{
-				SetCursorPos(FMath::Max(0, TypedStrPos - 1));
+				NewPos = FMath::Max(0, TypedStrPos - 1);
 			}
+			SetCursorPos(NewPos);
 			return true;
 		}
 		else if (Key == EKeys::Right)
 		{
+			int32 NewPos;
 			if (bCtrl)
 			{
-				// find the nearest '.' or ' '
-				int32 SpacePos = TypedStr.Find(TEXT(" "));
-				int32 PeriodPos = TypedStr.Find(TEXT("."));
-				// pick the closest valid index
-				int32 NewPos = SpacePos < 0 ? PeriodPos : (PeriodPos < 0 ? SpacePos : FMath::Min(SpacePos, PeriodPos));
-				// jump to end if nothing in between
-				if (NewPos == INDEX_NONE)
-				{
-					NewPos = TypedStr.Len();
-				}
-				SetCursorPos(FMath::Min(TypedStr.Len(), FMath::Max(TypedStrPos, NewPos)));
+				NewPos = FindWordBreak(TypedStr, FMath::Min(TypedStrPos + 1, TypedStr.Len()), ESearchDir::FromStart);
 			}
 			else
 			{
-				SetCursorPos(FMath::Min(TypedStr.Len(), TypedStrPos + 1));
+				NewPos = FMath::Min(TypedStr.Len(), TypedStrPos + 1);
 			}
+			SetCursorPos(NewPos);
 			return true;
 		}
 		else if (Key == EKeys::Home)
