Fixed bad enum usage in GetFontParameterDefaultValue.

diff --git a/Engine/Source/Runtime/Engine/Private/Materials/MaterialInterface.cpp b/Engine/Source/Runtime/Engine/Private/Materials/MaterialInterface.cpp
--- a/Engine/Source/Runtime/Engine/Private/Materials/MaterialInterface.cpp
+++ b/Engine/Source/Runtime/Engine/Private/Materials/MaterialInterface.cpp
@@ -907,7 +907,7 @@ bool UMaterialInterface::GetRuntimeVirtualTextureParameterDefaultValue(const FHa
 bool UMaterialInterface::GetFontParameterDefaultValue(const FHashedMaterialParameterInfo& ParameterInfo, class UFont*& OutFontValue, int32& OutFontPage) const
 {
 	FMaterialParameterMetadata Result;
-	if (GetParameterDefaultValue(EMaterialParameterType::RuntimeVirtualTexture, ParameterInfo, Result))
+	if (GetParameterDefaultValue(EMaterialParameterType::Font, ParameterInfo, Result))
 	{
 		OutFontValue = Result.Value.Font.Value;
 		OutFontPage = Result.Value.Font.Page;
