Make DrawMaterialToRenderTarget properly block until material is ready.

diff --git a/Engine/Source/Runtime/Engine/Private/Materials/MaterialInterface.cpp b/Engine/Source/Runtime/Engine/Private/Materials/MaterialInterface.cpp
--- a/Engine/Source/Runtime/Engine/Private/Materials/MaterialInterface.cpp
+++ b/Engine/Source/Runtime/Engine/Private/Materials/MaterialInterface.cpp
@@ -1423,12 +1423,32 @@ void UMaterialInterface::RemoveUserDataOfClass(TSubclassOf<UAssetUserData> InUse
 
 void UMaterialInterface::EnsureIsComplete()
 {
-	if (!IsComplete())
+	// TODO:
+	// The commented out code is roughly what we want to do to make the shaders ready. However this currently breaks for 
+	// material instances, where IsComplete isn't checking for readiness of the parent, and seems to not work even when
+	// CacheShaders is called on the parent. For now, to avoid issues with user blueprints failing on their first
+	// draw with on demand shader compilation, we do the more direct requests below, which seem to be more reliable.
+	
+	//if (!IsComplete())
+	//{
+	//	FScopedSlowTask SlowTask(0.0f, LOCTEXT("CacheShaders", "Caching Shaders..."));
+	//	SlowTask.MakeDialog();
+	//	CacheShaders(EMaterialShaderPrecompileMode::Synchronous);
+	//}
+
+#if WITH_EDITOR
+	uint32 FeatureLevelsToCompile = GetFeatureLevelsToCompileForRendering();
+	while (FeatureLevelsToCompile != 0)
 	{
-		FScopedSlowTask SlowTask(0.0f, LOCTEXT("CacheShaders", "Caching Shaders..."));
-		SlowTask.MakeDialog();
-		CacheShaders(EMaterialShaderPrecompileMode::Synchronous);
+		const ERHIFeatureLevel::Type FeatureLevel = (ERHIFeatureLevel::Type)FBitSet::GetAndClearNextBit(FeatureLevelsToCompile);
+		FMaterialResource* MaterialResource = GetMaterialResource(FeatureLevel);
+		if (MaterialResource && !MaterialResource->IsGameThreadShaderMapComplete())
+		{
+			MaterialResource->SubmitCompileJobs_GameThread(EShaderCompileJobPriority::ForceLocal);
+			MaterialResource->FinishCompilation(); // this is WITH_EDITOR only
+		}
 	}
+#endif
 }
 
 #if WITH_EDITORONLY_DATA
