Fixed UInstancedStaticMeshComponent's navigation relevancy issue that could result in removal from nav octree if UNavigationSystemV1::UpdateActorAndComponentsInNavOctree was called after the first SM instance got added

diff --git a/Engine/Source/Runtime/Engine/Private/InstancedStaticMesh.cpp b/Engine/Source/Runtime/Engine/Private/InstancedStaticMesh.cpp
--- a/Engine/Source/Runtime/Engine/Private/InstancedStaticMesh.cpp
+++ b/Engine/Source/Runtime/Engine/Private/InstancedStaticMesh.cpp
@@ -3126,9 +3126,13 @@ int32 UInstancedStaticMeshComponent::AddInstanceInternal(int32 InstanceIndex, FI
 
 	// If it's the first instance, register the component. 
 	// If there was no instance on component register, component registration was skipped because of UInstancedStaticMeshComponent::IsNavigationRelevant().
-	if (GetInstanceCount() == 1 && IsNavigationRelevant())
+	if (GetInstanceCount() == 1)
 	{
-		FNavigationSystem::RegisterComponent(*this);
+		if (const bool bNewNavigationRelevant = IsNavigationRelevant())
+		{
+			bNavigationRelevant = bNewNavigationRelevant;
+			FNavigationSystem::RegisterComponent(*this);
+		}
 	}
 	
 	PartialNavigationUpdate(InstanceIndex);
@@ -3185,9 +3189,13 @@ TArray<int32> UInstancedStaticMeshComponent::AddInstancesInternal(const TArray<F
 		{
 			// If it's the first instance, register the component. 
 			// If there was no instance on component register, component registration was skipped because of UInstancedStaticMeshComponent::IsNavigationRelevant().
-			if (GetInstanceCount() == 1 && IsNavigationRelevant())
+			if (GetInstanceCount() == 1)
 			{
-				FNavigationSystem::RegisterComponent(*this);
+				if (const bool bNewNavigationRelevant = IsNavigationRelevant())
+				{
+					bNavigationRelevant = bNewNavigationRelevant;
+					FNavigationSystem::RegisterComponent(*this);
+				}
 			}
 			
 			PartialNavigationUpdate(InstanceIndex);
@@ -3278,7 +3286,7 @@ bool UInstancedStaticMeshComponent::RemoveInstanceInternal(int32 InstanceIndex,
 	// remove instance
 	if (!InstanceAlreadyRemoved && PerInstanceSMData.IsValidIndex(InstanceIndex))
 	{
-		const bool bWasNavRelevant = IsNavigationRelevant();
+		const bool bWasNavRelevant = bNavigationRelevant;
 		
 		// Request navigation update
 		PartialNavigationUpdate(InstanceIndex);
@@ -3290,6 +3298,7 @@ bool UInstancedStaticMeshComponent::RemoveInstanceInternal(int32 InstanceIndex,
 		// (because of GetInstanceCount() > 0 in UInstancedStaticMeshComponent::IsNavigationRelevant())
 		if (bWasNavRelevant && GetInstanceCount() == 0)
 		{
+			bNavigationRelevant = false;
 			FNavigationSystem::UnregisterComponent(*this);
 		}
 	}
