Use permissions lists to filter allowed material expressions in HLSL material translator.

#rnx

diff --git a/Engine/Source/Developer/MaterialBaking/Private/ExportMaterialProxy.h b/Engine/Source/Developer/MaterialBaking/Private/ExportMaterialProxy.h
--- a/Engine/Source/Developer/MaterialBaking/Private/ExportMaterialProxy.h
+++ b/Engine/Source/Developer/MaterialBaking/Private/ExportMaterialProxy.h
@@ -19,7 +19,6 @@
 #include "SceneTypes.h"
 #include "Materials/Material.h"
 #include "Materials/MaterialExpressionCustomOutput.h"
-#include "MaterialRestrictiveMode.h"
 
 struct FExportMaterialCompiler : public FProxyMaterialCompiler
 {
@@ -355,8 +354,6 @@ public:
 	/** helper for CompilePropertyAndSetMaterialProperty() */
 	int32 CompilePropertyAndSetMaterialPropertyWithoutCast(EMaterialProperty Property, FMaterialCompiler* Compiler) const
 	{
-		FScopedHLSLMaterialTranslatorRestrictiveModeChange HLSLMaterialTranslatorRestrictiveModeChangeScope(false);
-
 		if (Property == MP_EmissiveColor)
 		{
 			const EBlendMode BlendMode = MaterialInterface->GetBlendMode();
diff --git a/Engine/Source/Editor/UnrealEd/Private/MaterialGraphSchema.cpp b/Engine/Source/Editor/UnrealEd/Private/MaterialGraphSchema.cpp
--- a/Engine/Source/Editor/UnrealEd/Private/MaterialGraphSchema.cpp
+++ b/Engine/Source/Editor/UnrealEd/Private/MaterialGraphSchema.cpp
@@ -1142,9 +1142,8 @@ void UMaterialGraphSchema::OnPinConnectionDoubleCicked(UEdGraphPin* PinA, UEdGra
 	// Create a new knot
 	UEdGraph* ParentGraph = PinA->GetOwningNode()->GetGraph();
 
+	if (UMaterialExpression* Expression = FMaterialEditorUtilities::CreateNewMaterialExpression(ParentGraph, UMaterialExpressionReroute::StaticClass(), KnotTopLeft, true, true))
 	{
-		UMaterialExpression* Expression = FMaterialEditorUtilities::CreateNewMaterialExpression(ParentGraph, UMaterialExpressionReroute::StaticClass(), KnotTopLeft, true, true);
-
 		// Move the connections across (only notifying the knot, as the other two didn't really change)
 		PinA->BreakLinkTo(PinB);
 		PinA->MakeLinkTo((PinA->Direction == EGPD_Output) ? CastChecked<UMaterialGraphNode_Knot>(Expression->GraphNode)->GetInputPin() : CastChecked<UMaterialGraphNode_Knot>(Expression->GraphNode)->GetOutputPin());
diff --git a/Engine/Source/Runtime/Engine/Classes/Materials/MaterialExpressionAdd.h b/Engine/Source/Runtime/Engine/Classes/Materials/MaterialExpressionAdd.h
--- a/Engine/Source/Runtime/Engine/Classes/Materials/MaterialExpressionAdd.h
+++ b/Engine/Source/Runtime/Engine/Classes/Materials/MaterialExpressionAdd.h
@@ -34,6 +34,7 @@ class UMaterialExpressionAdd : public UMaterialExpression
 	virtual int32 Compile(class FMaterialCompiler* Compiler, int32 OutputIndex) override;
 	virtual void GetCaption(TArray<FString>& OutCaptions) const override;
 	virtual FText GetKeywords() const override {return FText::FromString(TEXT("+"));}
+	virtual FText GetCreationName() const override { return FText::FromString(TEXT("Add")); }
 
 	virtual bool GenerateHLSLExpression(FMaterialHLSLGenerator& Generator, UE::HLSLTree::FScope& Scope, int32 OutputIndex, UE::HLSLTree::FExpression const*& OutExpression) const override;
 #endif // WITH_EDITOR
diff --git a/Engine/Source/Runtime/Engine/Classes/Materials/MaterialExpressionCustom.h b/Engine/Source/Runtime/Engine/Classes/Materials/MaterialExpressionCustom.h
--- a/Engine/Source/Runtime/Engine/Classes/Materials/MaterialExpressionCustom.h
+++ b/Engine/Source/Runtime/Engine/Classes/Materials/MaterialExpressionCustom.h
@@ -96,6 +96,7 @@ class UMaterialExpressionCustom : public UMaterialExpression
 #if WITH_EDITOR
 	virtual int32 Compile(class FMaterialCompiler* Compiler, int32 OutputIndex) override;
 	virtual void GetCaption(TArray<FString>& OutCaptions) const override;
+	virtual FText GetCreationName() const override { return FText::FromString(TEXT("Custom")); }
 	virtual const TArray<FExpressionInput*> GetInputs() override;
 	virtual FExpressionInput* GetInput(int32 InputIndex) override;
 	virtual FName GetInputName(int32 InputIndex) const override;
diff --git a/Engine/Source/Runtime/Engine/Private/Materials/HLSLMaterialTranslator.cpp b/Engine/Source/Runtime/Engine/Private/Materials/HLSLMaterialTranslator.cpp
--- a/Engine/Source/Runtime/Engine/Private/Materials/HLSLMaterialTranslator.cpp
+++ b/Engine/Source/Runtime/Engine/Private/Materials/HLSLMaterialTranslator.cpp
@@ -5,23 +5,27 @@
 #include "HLSLMaterialTranslator.h"
 #include "VT/VirtualTextureScalability.h"
 #include "Materials/MaterialExpressionCustom.h"
-#include "MaterialRestrictiveMode.h"
+#include "MaterialExpressionSettings.h"
 #include "Engine/RendererSettings.h"
 
 bool IsExpressionClassPermitted(const UClass* const Class)
 {
 	if (Class)
 	{
-		if (Class == UMaterialExpressionCustom::StaticClass())
+		// Always allow cooked content
+		if (Class->bCooked)
 		{
-			// If custom HLSL expressions are restricted, they can only be used in cooked content
-			return
-				Class->bCooked
-#if WITH_EDITORONLY_DATA
-				|| !FHLSLMaterialTranslator::InRestrictiveMode()
-#endif
-			;
+			return true;
 		}
+
+#if WITH_EDITOR
+		// Uncooked content must be explicitly allowed if in restrictive mode
+		const bool bIsClassPathAllowed = FMaterialExpressionSettings::Get()->IsClassPathAllowed(FTopLevelAssetPath(Class->GetPathName()));
+		if (!bIsClassPathAllowed)
+		{
+			return false;
+		}
+#endif
 	}
 	return true;
 }
@@ -2571,29 +2575,6 @@ FString FHLSLMaterialTranslator::GetParameterCodeRaw(int32 Index, const TCHAR* D
 	}
 }
 
-// Public interface for restrictive mode.
-bool GetHLSLMaterialTranslatorRestrictiveMode()
-{
-	return FHLSLMaterialTranslator::InRestrictiveMode();
-}
-
-void SetHLSLMaterialTranslatorRestrictiveMode(bool bRestrictiveMode)
-{
-	FHLSLMaterialTranslator::SetRestrictiveMode(bRestrictiveMode);
-}
-
-static bool GIsMaterialTranslatorInRestrictiveMode = false;
-
-void FHLSLMaterialTranslator::SetRestrictiveMode(bool bRestrictiveMode)
-{
-	GIsMaterialTranslatorInRestrictiveMode = bRestrictiveMode;
-}
-
-bool FHLSLMaterialTranslator::InRestrictiveMode()
-{
-	return GIsMaterialTranslatorInRestrictiveMode;
-}
-
 uint64 FHLSLMaterialTranslator::GetParameterHash(int32 Index)
 {
 	if (Index == INDEX_NONE)
@@ -11243,18 +11224,20 @@ int32 FHLSLMaterialTranslator::AccessMaterialAttribute(int32 CodeIndex, const FG
 		*AttributeName);
 }
 
-static bool IsMaterialPathExemptFromRestrictiveMode(FString& OutPathName)
+static bool IsMaterialPathExemptFromRestrictiveMode(FStringBuilderBase& PathName)
 {
+	const FStringView PathNameView = PathName;
+
 	// Anything currently being edited is NOT exempt from restrictive mode
-	const TCHAR* EngineTransientPath = TEXT("/Engine/Transient.");
-	if (OutPathName.StartsWith(EngineTransientPath))
+	const FStringView EngineTransientPath = TEXT("/Engine/Transient.");
+	if (PathNameView.StartsWith(EngineTransientPath))
 	{
-		OutPathName.RightChopInline(FCString::Strlen(EngineTransientPath));
+		PathName.RemoveAt(0, EngineTransientPath.Len());
 		return false;
 	}
 
 	// Anything else in /Engine/ and /Game/ folders is exempt
-	if (OutPathName.StartsWith(TEXT("/Engine/")) || OutPathName.StartsWith(TEXT("/Game/")))
+	if (PathNameView.StartsWith(TEXT("/Engine/")) || PathNameView.StartsWith(TEXT("/Game/")))
 	{
 		return true;
 	}
@@ -11262,30 +11245,47 @@ static bool IsMaterialPathExemptFromRestrictiveMode(FString& OutPathName)
 	return false;
 }
 
-static bool IsMaterialExpressionExemptFromRestrictiveMode(const UMaterialExpression* Expression, FString& OutPathName)
+static bool IsMaterialExpressionExemptFromRestrictiveMode(const UMaterialExpression* Expression, FStringBuilderBase& PathName)
 {
 	if (Expression)
 	{
 		if (Expression->Material)
 		{
-			OutPathName = Expression->Material->GetPathName();
-			return IsMaterialPathExemptFromRestrictiveMode(OutPathName);
+			Expression->Material->GetPathName(nullptr, PathName);
+			return IsMaterialPathExemptFromRestrictiveMode(PathName);
 		}
 		if (Expression->Function)
 		{
-			OutPathName = Expression->Function->GetPathName();
-			return IsMaterialPathExemptFromRestrictiveMode(OutPathName);
+			Expression->Function->GetPathName(nullptr, PathName);
+			return IsMaterialPathExemptFromRestrictiveMode(PathName);
 		}
 	}
 	return false;
 }
 
+bool FHLSLMaterialTranslator::ValidateMaterialExpressionPermission(const UMaterialExpression* Expression)
+{
+	FNameBuilder MaterialPathName;
+	if (!IsExpressionClassPermitted(Expression->GetClass()) && !IsMaterialExpressionExemptFromRestrictiveMode(Expression, MaterialPathName))
+	{
+		if (MaterialPathName.Len() > 0)
+		{
+			Errorf(TEXT("Material expression %s not permitted in %s"), *Expression->GetCreationName().ToString(), *MaterialPathName);
+		}
+		else
+		{
+			Errorf(TEXT("Material expression %s not permitted"), *Expression->GetCreationName().ToString());
+		}
+		return false;
+	}
+	return true;
+}
+
 int32 FHLSLMaterialTranslator::CustomExpression( class UMaterialExpressionCustom* Custom, int32 OutputIndex, TArray<int32>& CompiledInputs )
 {
-	FString MaterialPathName;
-	if (!IsExpressionClassPermitted(Custom->GetClass()) && !IsMaterialExpressionExemptFromRestrictiveMode(Custom, MaterialPathName))
+	if (!ValidateMaterialExpressionPermission(Custom))
 	{
-		return Errorf(TEXT("Custom HLSL expressions are not permitted in %s"), *MaterialPathName);
+		return INDEX_NONE;
 	}
 
 	const FMaterialCustomExpressionEntry* CustomEntry = nullptr;
diff --git a/Engine/Source/Runtime/Engine/Private/Materials/HLSLMaterialTranslator.h b/Engine/Source/Runtime/Engine/Private/Materials/HLSLMaterialTranslator.h
--- a/Engine/Source/Runtime/Engine/Private/Materials/HLSLMaterialTranslator.h
+++ b/Engine/Source/Runtime/Engine/Private/Materials/HLSLMaterialTranslator.h
@@ -521,11 +521,6 @@ public:
 	// Must always be valid
 	virtual FString GetParameterCodeDeriv(int32 Index, ECompiledPartialDerivativeVariation Variation);
 
-	// Puts the translator into restrictive mode.
-	static void SetRestrictiveMode(bool bRestrictiveMode);
-	// Returns whether the material translator is in restrictive mode. If true, custom HLSL expressions are only allowed for cooked content.
-	static bool InRestrictiveMode();
-
 protected:
 
 	uint64 GetParameterHash(int32 Index);
@@ -673,6 +668,8 @@ protected:
 	void StrataEvaluateSharedLocalBases(uint8& UsedSharedLocalBasesCount, FString* OutStrataPixelNormalInitializerValues, FShaderCompilerEnvironment* OutEnvironment);
 	FString StrataGetCastParameterCode(int32 Index, EMaterialValueType DestType);
 
+	bool ValidateMaterialExpressionPermission(const UMaterialExpression* Expression);
+
 	// FMaterialCompiler interface.
 
 	/** 
diff --git a/Engine/Source/Runtime/Engine/Private/Materials/MaterialExpressionSettings.cpp b/Engine/Source/Runtime/Engine/Private/Materials/MaterialExpressionSettings.cpp
new file mode 100644
--- /dev/null
+++ b/Engine/Source/Runtime/Engine/Private/Materials/MaterialExpressionSettings.cpp
@@ -0,0 +1,41 @@
+// Copyright Epic Games, Inc. All Rights Reserved.
+
+#include "MaterialExpressionSettings.h"
+
+#if WITH_EDITOR
+
+FMaterialExpressionSettings* FMaterialExpressionSettings::Get()
+{
+	static FMaterialExpressionSettings Inst;
+	return &Inst;
+}
+
+void FMaterialExpressionSettings::RegisterIsClassPathAllowedDelegate(const FName OwnerName, FOnIsClassPathAllowed Delegate)
+{
+	IsClassPathAllowedDelegates.Add(OwnerName, Delegate);
+}
+
+void FMaterialExpressionSettings::UnregisterIsClassPathAllowedDelegate(const FName OwnerName)
+{
+	IsClassPathAllowedDelegates.Remove(OwnerName);
+}
+
+bool FMaterialExpressionSettings::IsClassPathAllowed(const FTopLevelAssetPath& InClassPath) const
+{
+	for (const TPair<FName, FOnIsClassPathAllowed>& DelegatePair : IsClassPathAllowedDelegates)
+	{
+		if (!DelegatePair.Value.Execute(InClassPath))
+		{
+			return false;
+		}
+	}
+	return true;
+}
+
+bool FMaterialExpressionSettings::HasClassPathFiltering() const
+{
+	return IsClassPathAllowedDelegates.Num() > 0;
+}
+
+
+#endif // WITH_EDITOR
diff --git a/Engine/Source/Runtime/Engine/Public/MaterialExpressionSettings.h b/Engine/Source/Runtime/Engine/Public/MaterialExpressionSettings.h
new file mode 100644
--- /dev/null
+++ b/Engine/Source/Runtime/Engine/Public/MaterialExpressionSettings.h
@@ -0,0 +1,33 @@
+// Copyright Epic Games, Inc. All Rights Reserved.
+#pragma once
+
+#include "CoreMinimal.h"
+
+#if WITH_EDITOR
+
+/**
+ * Singleton for the material expression settings and permissions
+ */
+class ENGINE_API FMaterialExpressionSettings
+{
+public:
+	/** Gets singleton instance */
+	static FMaterialExpressionSettings* Get();
+
+	/** Delegate to filter class paths from permissions lists */
+	DECLARE_DELEGATE_RetVal_OneParam(bool, FOnIsClassPathAllowed, const FTopLevelAssetPath& /*InClassPath*/)
+
+	void RegisterIsClassPathAllowedDelegate(const FName OwnerName, FOnIsClassPathAllowed Delegate);
+	void UnregisterIsClassPathAllowedDelegate(const FName OwnerName);
+	bool IsClassPathAllowed(const FTopLevelAssetPath& InClassPath) const;
+	bool HasClassPathFiltering() const;
+
+private:
+	FMaterialExpressionSettings() = default;
+	~FMaterialExpressionSettings() = default;
+
+	/** Delegates called to determine whether a class type is allowed to be processed in the material translator */
+	TMap<FName, FOnIsClassPathAllowed> IsClassPathAllowedDelegates;
+};
+
+#endif // WITH_EDITOR
diff --git a/Engine/Source/Runtime/Engine/Public/MaterialRestrictiveMode.h b/Engine/Source/Runtime/Engine/Public/MaterialRestrictiveMode.h
deleted file mode 100644
