Change StreamingObject outer PopulateGeneratedPackageForCook instead of PopulateGeneratedPackageForCook so the package generated in PopulateGeneratedPackageForCook can be garbage collected.

#test fullcook, fastcook, testcook (for GC)

diff --git a/Engine/Source/Runtime/Engine/Private/WorldPartition/ContentBundle/ContentBundleEditor.cpp b/Engine/Source/Runtime/Engine/Private/WorldPartition/ContentBundle/ContentBundleEditor.cpp
--- a/Engine/Source/Runtime/Engine/Private/WorldPartition/ContentBundle/ContentBundleEditor.cpp
+++ b/Engine/Source/Runtime/Engine/Private/WorldPartition/ContentBundle/ContentBundleEditor.cpp
@@ -120,6 +120,9 @@ void FContentBundleEditor::DoRemoveContent()
 		ActorDescContainer = nullptr;
 	}
 
+	ExternalStreamingObject = nullptr;
+	CookPackageIdsToCell.Empty();
+
 	SetStatus(EContentBundleStatus::Registered);
 
 	BroadcastChanged();
@@ -398,14 +401,6 @@ bool FContentBundleEditor::PopulateGeneratorPackageForCook(class IWorldPartition
 					bIsSuccess = false;
 				}
 			}
-			else
-			{
-				if (!ExternalStreamingObject->Rename(nullptr, CookPackage->GetPackage(), REN_DontCreateRedirectors))
-				{
-					UE_LOG(LogContentBundle, Error, TEXT("[CB: %s][Cook] Failed to set streaming object package %s."), *GetDescriptor()->GetDisplayName());
-					bIsSuccess = false;
-				}
-			}
 		}
 
 		ExternalStreamingObject->PopulateGeneratorPackageForCook();
@@ -440,6 +435,14 @@ bool FContentBundleEditor::PopulateGeneratedPackageForCook(class IWorldPartition
 			bIsSuccess = false;
 		}
 	}
+	else
+	{
+		if (!ExternalStreamingObject->Rename(nullptr, PackageToCook.GetPackage(), REN_DontCreateRedirectors))
+		{
+			UE_LOG(LogContentBundle, Error, TEXT("[CB: %s][Cook] Failed to set streaming object package %s."), *GetDescriptor()->GetDisplayName());
+			bIsSuccess = false;
+		}
+	}
 
 	return bIsSuccess;
 }
