Create loose parameters uniform buffer to encapsulate GPUSkinPassThroughVertexFactory PreviousPositionBuffer.
Fixes ensure failure in MessPassProcessor.cpp in debug_editor build, due to PreviousPositionBuffer being a loose SRV parameter.

diff --git a/Engine/Shaders/Private/LocalVertexFactory.ush b/Engine/Shaders/Private/LocalVertexFactory.ush
--- a/Engine/Shaders/Private/LocalVertexFactory.ush
+++ b/Engine/Shaders/Private/LocalVertexFactory.ush
@@ -1141,11 +1141,6 @@ float3 VertexFactoryGetWorldNormal(FVertexFactoryInput Input, FVertexFactoryInte
 	return Intermediates.TangentToWorld[2];
 }
 
-#if SUPPORT_GPUSKIN_PASSTHROUGH
-// Information when using pass through to get the previous position, as no matrix is available/needed
-Buffer<float> GPUSkinCachePreviousPositionBuffer;
-#endif
-
 // @return previous translated world position
 float4 VertexFactoryGetPreviousWorldPosition(FVertexFactoryInput Input, FVertexFactoryIntermediates Intermediates)
 {
@@ -1155,11 +1150,11 @@ float4 VertexFactoryGetPreviousWorldPosition(FVertexFactoryInput Input, FVertexF
 #if SUPPORT_GPUSKIN_PASSTHROUGH
 	if (IsGPUSkinPassThrough())
 	{
-	uint Offset = Input.VertexId * 3;
-	float3 PreviousPos;
-	PreviousPos.x = GPUSkinCachePreviousPositionBuffer[Offset + 0];
-	PreviousPos.y = GPUSkinCachePreviousPositionBuffer[Offset + 1];
-	PreviousPos.z = GPUSkinCachePreviousPositionBuffer[Offset + 2];
+		uint Offset = Input.VertexId * 3;
+		float3 PreviousPos;
+		PreviousPos.x = LocalVFLooseParameters.GPUSkinPassThroughPreviousPositionBuffer[Offset + 0];
+		PreviousPos.y = LocalVFLooseParameters.GPUSkinPassThroughPreviousPositionBuffer[Offset + 1];
+		PreviousPos.z = LocalVFLooseParameters.GPUSkinPassThroughPreviousPositionBuffer[Offset + 2];
 		PrevLocalPosition = float4(PreviousPos, 1);	
 	}
 	else
diff --git a/Engine/Source/Runtime/Engine/Private/GPUSkinCache.cpp b/Engine/Source/Runtime/Engine/Private/GPUSkinCache.cpp
--- a/Engine/Source/Runtime/Engine/Private/GPUSkinCache.cpp
+++ b/Engine/Source/Runtime/Engine/Private/GPUSkinCache.cpp
@@ -1757,10 +1757,8 @@ void FGPUSkinCache::Release(FGPUSkinCacheEntry*& SkinCacheEntry)
 void FGPUSkinCache::GetShaderBindings(
 	const FGPUSkinCacheEntry* Entry, 
 	int32 Section,
-	bool bVerticesInMotion,
 	const FGPUSkinPassthroughVertexFactory* VertexFactory,
 	FShaderResourceParameter GPUSkinCachePositionBuffer,
-	FShaderResourceParameter GPUSkinCachePreviousPositionBuffer,
 	FMeshDrawSingleShaderBindings& ShaderBindings,
 	FVertexInputStreamArray& VertexStreams)
 {
@@ -1781,9 +1779,6 @@ void FGPUSkinCache::GetShaderBindings(
 	}
 
 	ShaderBindings.Add(GPUSkinCachePositionBuffer, DispatchData.GetPositionRWBuffer()->Buffer.SRV);
-
-	// If world is paused, use current frame bone matrices, so velocity is canceled and skeletal mesh isn't blurred from motion.
-	ShaderBindings.Add(GPUSkinCachePreviousPositionBuffer, bVerticesInMotion ? DispatchData.GetPreviousPositionRWBuffer()->Buffer.SRV : DispatchData.GetPositionRWBuffer()->Buffer.SRV);
 }
 
 void FGPUSkinCache::PrepareUpdateSkinning(FGPUSkinCacheEntry* Entry, int32 Section, uint32 RevisionNumber, TArray<FSkinCacheRWBuffer*>* OverlappedUAVs)
diff --git a/Engine/Source/Runtime/Engine/Private/LocalVertexFactory.cpp b/Engine/Source/Runtime/Engine/Private/LocalVertexFactory.cpp
--- a/Engine/Source/Runtime/Engine/Private/LocalVertexFactory.cpp
+++ b/Engine/Source/Runtime/Engine/Private/LocalVertexFactory.cpp
@@ -43,6 +43,7 @@ void FLocalVertexFactoryShaderParametersBase::Bind(const FShaderParameterMap& Pa
 }
 
 IMPLEMENT_GLOBAL_SHADER_PARAMETER_STRUCT(FLocalVertexFactoryUniformShaderParameters, "LocalVF");
+IMPLEMENT_GLOBAL_SHADER_PARAMETER_STRUCT(FLocalVertexFactoryLooseParameters, "LocalVFLooseParameters");
 
 TUniformBufferRef<FLocalVertexFactoryUniformShaderParameters> CreateLocalVFUniformBuffer(
 	const FLocalVertexFactory* LocalVertexFactory, 
@@ -162,7 +163,6 @@ void FLocalVertexFactoryShaderParameters::Bind(const FShaderParameterMap& Parame
 {
 	FLocalVertexFactoryShaderParametersBase::Bind(ParameterMap);
 	GPUSkinCachePositionBuffer.Bind(ParameterMap, TEXT("GPUSkinCachePositionBuffer"));
-	GPUSkinCachePreviousPositionBuffer.Bind(ParameterMap, TEXT("GPUSkinCachePreviousPositionBuffer"));
 	IsGPUSkinPassThrough.Bind(ParameterMap, TEXT("bIsGPUSkinPassThrough"));
 }
 
@@ -195,9 +195,6 @@ void FLocalVertexFactoryShaderParameters::GetElementShaderBindings(
 	}
 	else
 	{
-		// Bind null SRV as not used by non-passthrough vertex factory
-		ShaderBindings.Add(GPUSkinCachePreviousPositionBuffer, GNullVertexBuffer.VertexBufferSRV);
-
 		// Decode VertexFactoryUserData as VertexFactoryUniformBuffer
 		FRHIUniformBuffer* VertexFactoryUniformBuffer = static_cast<FRHIUniformBuffer*>(BatchElement.VertexFactoryUserData);
 
@@ -213,6 +210,8 @@ void FLocalVertexFactoryShaderParameters::GetElementShaderBindings(
 			ShaderBindings,
 			VertexStreams);
 	}
+
+	ShaderBindings.Add(Shader->GetUniformBufferParameter<FLocalVertexFactoryLooseParameters>(), LocalVertexFactory->LooseParametersUniformBuffer);
 }
 
 void FLocalVertexFactoryShaderParameters::GetElementShaderBindingsGPUSkinPassThrough(
@@ -236,40 +235,31 @@ void FLocalVertexFactoryShaderParameters::GetElementShaderBindingsGPUSkinPassThr
 
 	check(VertexFactory->GetType() == &FGPUSkinPassthroughVertexFactory::StaticType);
 	FGPUSkinPassthroughVertexFactory const* PassthroughVertexFactory = static_cast<FGPUSkinPassthroughVertexFactory const*>(VertexFactory);
-
-	// Bone data is updated whenever animation triggers a dynamic update, animation can skip frames hence the frequency is not necessary every frame.
-	// So check if bone data is updated this frame, if not then the previous frame data is stale and not suitable for motion blur.
-	bool bBoneDataUpdatedThisFrame = (View->Family->FrameNumber == PassthroughVertexFactory->GetUpdatedFrameNumber());
-	// If world is paused, use current frame bone matrices, so velocity is canceled and skeletal mesh isn't blurred from motion.
-	bool bVerticesInMotion = !View->Family->bWorldIsPaused && bBoneDataUpdatedThisFrame;
-
 	if (bUsesSkinCache)
 	{
-		GetElementShaderBindingsSkinCache(PassthroughVertexFactory, BatchUserData, bVerticesInMotion, ShaderBindings, VertexStreams);
+		GetElementShaderBindingsSkinCache(PassthroughVertexFactory, BatchUserData, ShaderBindings, VertexStreams);
 	}
 	else
 	{
-		GetElementShaderBindingsMeshDeformer(PassthroughVertexFactory, bVerticesInMotion, ShaderBindings, VertexStreams);
+		GetElementShaderBindingsMeshDeformer(PassthroughVertexFactory, ShaderBindings, VertexStreams);
 	}
 }
 
 void FLocalVertexFactoryShaderParameters::GetElementShaderBindingsSkinCache(
 	FGPUSkinPassthroughVertexFactory const* PassthroughVertexFactory,
 	FGPUSkinBatchElementUserData* BatchUserData,
-	bool bVerticesInMotion,
 	FMeshDrawSingleShaderBindings& ShaderBindings,
 	FVertexInputStreamArray& VertexStreams) const
 {
 	FGPUSkinCache::GetShaderBindings(
-		BatchUserData->Entry, BatchUserData->Section, bVerticesInMotion,
+		BatchUserData->Entry, BatchUserData->Section,
 		PassthroughVertexFactory,
-		GPUSkinCachePositionBuffer, GPUSkinCachePreviousPositionBuffer,
+		GPUSkinCachePositionBuffer,
 		ShaderBindings, VertexStreams);
 }
 
 void FLocalVertexFactoryShaderParameters::GetElementShaderBindingsMeshDeformer(
 	FGPUSkinPassthroughVertexFactory const* PassthroughVertexFactory,
-	bool bVerticesInMotion,
 	FMeshDrawSingleShaderBindings& ShaderBindings,
 	FVertexInputStreamArray& VertexStreams) const
 {
@@ -282,15 +272,6 @@ void FLocalVertexFactoryShaderParameters::GetElementShaderBindingsMeshDeformer(
 	{
 		VertexStreams.Add(FVertexInputStream(PassthroughVertexFactory->GetTangentStreamIndex(), 0, PassthroughVertexFactory->TangentRDG->GetRHI()));
 	}
-
-	if (bVerticesInMotion && PassthroughVertexFactory->PrevPositionRDG)
-	{
-		ShaderBindings.Add(GPUSkinCachePreviousPositionBuffer, PassthroughVertexFactory->PrevPositionRDG->GetOrCreateSRV(FRHIBufferSRVCreateInfo(PF_R32_FLOAT)));
-	}
-	else
-	{
-		ShaderBindings.Add(GPUSkinCachePreviousPositionBuffer, PassthroughVertexFactory->GetPositionsSRV());
-	}
 }
 
 /**
@@ -545,6 +526,10 @@ void FLocalVertexFactory::InitRHI()
 		UniformBuffer = CreateLocalVFUniformBuffer(this, Data.LODLightmapDataIndex, nullptr, DefaultBaseVertexIndex, DefaultPreSkinBaseVertexIndex);
 	}
 
+	FLocalVertexFactoryLooseParameters LooseParameters;
+	LooseParameters.GPUSkinPassThroughPreviousPositionBuffer = GNullVertexBuffer.VertexBufferSRV;
+	LooseParametersUniformBuffer = TUniformBufferRef<FLocalVertexFactoryLooseParameters>::CreateUniformBufferImmediate(LooseParameters, UniformBuffer_MultiFrame);
+
 	check(IsValidRef(GetDeclaration()));
 }
 
diff --git a/Engine/Source/Runtime/Engine/Private/SkeletalMesh.cpp b/Engine/Source/Runtime/Engine/Private/SkeletalMesh.cpp
--- a/Engine/Source/Runtime/Engine/Private/SkeletalMesh.cpp
+++ b/Engine/Source/Runtime/Engine/Private/SkeletalMesh.cpp
@@ -6231,6 +6231,43 @@ void FSkeletalMeshSceneProxy::CreateBaseMeshBatch(const FSceneView* View, const
 #endif
 	BatchElement.PrimitiveUniformBuffer = GetUniformBuffer();
 	BatchElement.NumPrimitives = LODData.RenderSections[SectionIndex].NumTriangles;
+
+	UpdateLooseParametersUniformBuffer(View, SectionIndex, Mesh, BatchElement);
+}
+
+void FSkeletalMeshSceneProxy::UpdateLooseParametersUniformBuffer(const FSceneView* View, const int32 SectionIndex, const FMeshBatch& Mesh, const FMeshBatchElement& BatchElement) const
+{
+	// Loose parameters uniform buffer is only needed for PassThroughVF
+	if (!Mesh.VertexFactory || Mesh.VertexFactory->GetType() != &FGPUSkinPassthroughVertexFactory::StaticType)
+	{
+		return;
+	}
+	
+	FLocalVertexFactoryLooseParameters Parameters;
+
+	FVertexFactory* NonConstVertexFactory = const_cast<FVertexFactory*>(Mesh.VertexFactory);
+	FGPUSkinPassthroughVertexFactory* PassthroughVertexFactory = static_cast<FGPUSkinPassthroughVertexFactory*>(NonConstVertexFactory);
+	// Bone data is updated whenever animation triggers a dynamic update, animation can skip frames hence the frequency is not necessary every frame.
+	// So check if bone data is updated this frame, if not then the previous frame data is stale and not suitable for motion blur.
+	bool bBoneDataUpdatedThisFrame = (View->Family->FrameNumber == PassthroughVertexFactory->GetUpdatedFrameNumber());
+	// If world is paused, use current frame bone matrices, so velocity is canceled and skeletal mesh isn't blurred from motion.
+	bool bVerticesInMotion = !View->Family->bWorldIsPaused && bBoneDataUpdatedThisFrame;
+
+	const bool bUsesSkinCache = BatchElement.VertexFactoryUserData != nullptr;
+	if (bUsesSkinCache)
+	{
+		Parameters.GPUSkinPassThroughPreviousPositionBuffer = bVerticesInMotion ?
+			FGPUSkinCache::GetPreviousPositionBuffer(MeshObject->SkinCacheEntry, SectionIndex)->SRV :
+			FGPUSkinCache::GetPositionBuffer(MeshObject->SkinCacheEntry, SectionIndex)->SRV;
+	}
+	else // Mesh deformer
+	{
+		Parameters.GPUSkinPassThroughPreviousPositionBuffer = (bVerticesInMotion && PassthroughVertexFactory->PrevPositionRDG) ?
+			PassthroughVertexFactory->PrevPositionRDG->GetOrCreateSRV(FRHIBufferSRVCreateInfo(PF_R32_FLOAT)) :
+			PassthroughVertexFactory->GetPositionsSRV();
+	}
+
+	PassthroughVertexFactory->LooseParametersUniformBuffer.UpdateUniformBufferImmediate(Parameters);
 }
 
 uint8 FSkeletalMeshSceneProxy::GetCurrentFirstLODIdx_Internal() const
diff --git a/Engine/Source/Runtime/Engine/Public/GPUSkinCache.h b/Engine/Source/Runtime/Engine/Public/GPUSkinCache.h
--- a/Engine/Source/Runtime/Engine/Public/GPUSkinCache.h
+++ b/Engine/Source/Runtime/Engine/Public/GPUSkinCache.h
@@ -152,10 +152,8 @@ public:
 	static void GetShaderBindings(
 		const FGPUSkinCacheEntry* Entry,
 		int32 Section,
-		bool bVerticesInMotion,
 		const FGPUSkinPassthroughVertexFactory* VertexFactory,
 		FShaderResourceParameter GPUSkinCachePositionBuffer,
-		FShaderResourceParameter GPUSkinCachePreviousPositionBuffer,
 		class FMeshDrawSingleShaderBindings& ShaderBindings,
 		FVertexInputStreamArray& VertexStreams);
 
diff --git a/Engine/Source/Runtime/Engine/Public/GPUSkinVertexFactory.h b/Engine/Source/Runtime/Engine/Public/GPUSkinVertexFactory.h
--- a/Engine/Source/Runtime/Engine/Public/GPUSkinVertexFactory.h
+++ b/Engine/Source/Runtime/Engine/Public/GPUSkinVertexFactory.h
@@ -560,6 +560,7 @@ public:
 
 protected:
 	friend class FLocalVertexFactoryShaderParameters;
+	friend class FSkeletalMeshSceneProxy;
 
 	// Reference holders for RDG buffers
 	TRefCountPtr<FRDGPooledBuffer> PositionRDG;
diff --git a/Engine/Source/Runtime/Engine/Public/LocalVertexFactory.h b/Engine/Source/Runtime/Engine/Public/LocalVertexFactory.h
--- a/Engine/Source/Runtime/Engine/Public/LocalVertexFactory.h
+++ b/Engine/Source/Runtime/Engine/Public/LocalVertexFactory.h
@@ -26,6 +26,10 @@ BEGIN_GLOBAL_SHADER_PARAMETER_STRUCT(FLocalVertexFactoryUniformShaderParameters,
 	SHADER_PARAMETER_SRV(Buffer<float4>, VertexFetch_ColorComponentsBuffer)
 END_GLOBAL_SHADER_PARAMETER_STRUCT()
 
+BEGIN_GLOBAL_SHADER_PARAMETER_STRUCT(FLocalVertexFactoryLooseParameters,)
+	SHADER_PARAMETER_SRV(Buffer<float>, GPUSkinPassThroughPreviousPositionBuffer)
+END_GLOBAL_SHADER_PARAMETER_STRUCT()
+
 extern ENGINE_API TUniformBufferRef<FLocalVertexFactoryUniformShaderParameters> CreateLocalVFUniformBuffer(
 	const class FLocalVertexFactory* VertexFactory, 
 	uint32 LODLightmapDataIndex, 
@@ -86,6 +90,7 @@ public:
 	virtual void ReleaseRHI() override
 	{
 		UniformBuffer.SafeRelease();
+		LooseParametersUniformBuffer.SafeRelease();
 		FVertexFactory::ReleaseRHI();
 	}
 
@@ -157,11 +162,13 @@ public:
 
 protected:
 	friend class FLocalVertexFactoryShaderParameters;
+	friend class FSkeletalMeshSceneProxy;
 
 	const FDataType& GetData() const { return Data; }
 
 	FDataType Data;
 	TUniformBufferRef<FLocalVertexFactoryUniformShaderParameters> UniformBuffer;
+	TUniformBufferRef<FLocalVertexFactoryLooseParameters> LooseParametersUniformBuffer;
 
 	int32 ColorStreamIndex;
 
@@ -249,17 +256,14 @@ private:
 	void GetElementShaderBindingsSkinCache(
 		class FGPUSkinPassthroughVertexFactory const* PassthroughVertexFactory,
 		struct FGPUSkinBatchElementUserData* BatchUserData,
-		bool bVerticesInMotion,
 		FMeshDrawSingleShaderBindings& ShaderBindings,
 		FVertexInputStreamArray& VertexStreams) const;
 
 	void GetElementShaderBindingsMeshDeformer(
 		class FGPUSkinPassthroughVertexFactory const* PassthroughVertexFactory,
-		bool bVerticesInMotion,
 		FMeshDrawSingleShaderBindings& ShaderBindings,
 		FVertexInputStreamArray& VertexStreams) const;
 
 	LAYOUT_FIELD(FShaderResourceParameter, GPUSkinCachePositionBuffer);
-	LAYOUT_FIELD(FShaderResourceParameter, GPUSkinCachePreviousPositionBuffer);
 	LAYOUT_FIELD(FShaderParameter, IsGPUSkinPassThrough);
 };
\ No newline at end of file
diff --git a/Engine/Source/Runtime/Engine/Public/SkeletalMeshTypes.h b/Engine/Source/Runtime/Engine/Public/SkeletalMeshTypes.h
--- a/Engine/Source/Runtime/Engine/Public/SkeletalMeshTypes.h
+++ b/Engine/Source/Runtime/Engine/Public/SkeletalMeshTypes.h
@@ -337,6 +337,7 @@ protected:
 	uint8 GetCurrentFirstLODIdx_Internal() const;
 private:
 	void CreateBaseMeshBatch(const FSceneView* View, const FSkeletalMeshLODRenderData& LODData, const int32 LODIndex, const int32 SectionIndex, const FSectionElementInfo& SectionElementInfo, FMeshBatch& Mesh, ESkinVertexFactoryMode VFMode = ESkinVertexFactoryMode::Default) const;
+	void UpdateLooseParametersUniformBuffer(const FSceneView* View, const int32 SectionIndex, const FMeshBatch& Mesh, const FMeshBatchElement& BatchElement) const;
 
 public:
 #if WITH_EDITORONLY_DATA
