Fix compile error in MaterialExpressionSettings.h for NonUnity build.

#rnx

diff --git a/Engine/Source/Runtime/Engine/Public/MaterialExpressionSettings.h b/Engine/Source/Runtime/Engine/Public/MaterialExpressionSettings.h
--- a/Engine/Source/Runtime/Engine/Public/MaterialExpressionSettings.h
+++ b/Engine/Source/Runtime/Engine/Public/MaterialExpressionSettings.h
@@ -2,6 +2,8 @@
 #pragma once
 
 #include "CoreMinimal.h"
+#include "Delegates/Delegate.h"
+#include "UObject/TopLevelAssetPath.h"
 
 #if WITH_EDITOR
 
@@ -15,7 +17,7 @@ public:
 	static FMaterialExpressionSettings* Get();
 
 	/** Delegate to filter class paths from permissions lists */
-	DECLARE_DELEGATE_RetVal_OneParam(bool, FOnIsClassPathAllowed, const FTopLevelAssetPath& /*InClassPath*/)
+	DECLARE_DELEGATE_RetVal_OneParam(bool, FOnIsClassPathAllowed, const FTopLevelAssetPath& /*InClassPath*/);
 
 	void RegisterIsClassPathAllowedDelegate(const FName OwnerName, FOnIsClassPathAllowed Delegate);
 	void UnregisterIsClassPathAllowedDelegate(const FName OwnerName);
