SkeletonCompatibility: A slightly more optimized IsCompatibleSkeletonByAssetString method that doesn't rely on the asset manager.

diff --git a/Engine/Source/Runtime/Engine/Private/Animation/Skeleton.cpp b/Engine/Source/Runtime/Engine/Private/Animation/Skeleton.cpp
--- a/Engine/Source/Runtime/Engine/Private/Animation/Skeleton.cpp
+++ b/Engine/Source/Runtime/Engine/Private/Animation/Skeleton.cpp
@@ -146,15 +146,11 @@ bool USkeleton::IsCompatibleSkeletonByAssetString(const FString& SkeletonAssetSt
 		return true;
 	}
 
-	// Get the asset registry.
-	const FAssetRegistryModule& AssetRegistryModule = FModuleManager::LoadModuleChecked<FAssetRegistryModule>("AssetRegistry");
-	const IAssetRegistry& AssetRegistry = AssetRegistryModule.Get();
-
 	// Now check against the list of compatible skeletons and see if we're dealing with the same asset.
+	const FSoftObjectPath InPath(SkeletonAssetString);
 	for (const TSoftObjectPtr<USkeleton>& CompatibleSkeleton : CompatibleSkeletons)
 	{
-		const FAssetData AssetData = AssetRegistry.GetAssetByObjectPath(CompatibleSkeleton.ToString());
-		if (AssetData.IsValid() && AssetData.GetExportTextName() == SkeletonAssetString)
+		if (CompatibleSkeleton.ToSoftObjectPath() == InPath)
 		{
 			return true;
 		}
