Fix crash when loading CitySample with -nullrhi

diff --git a/Engine/Source/Runtime/Engine/Private/InstancedStaticMesh.cpp b/Engine/Source/Runtime/Engine/Private/InstancedStaticMesh.cpp
--- a/Engine/Source/Runtime/Engine/Private/InstancedStaticMesh.cpp
+++ b/Engine/Source/Runtime/Engine/Private/InstancedStaticMesh.cpp
@@ -2276,9 +2276,12 @@ void UInstancedStaticMeshComponent::FlushInstanceUpdateCommands(bool bFlushInsta
 		InstanceUpdateCmdBuffer.Reset();
 	}
 
-	FStaticMeshInstanceData RenderInstanceData = FStaticMeshInstanceData(GVertexElementTypeSupport.IsSupported(VET_Half2));
-	BuildRenderData(RenderInstanceData, PerInstanceRenderData->HitProxies);
-	PerInstanceRenderData->UpdateFromPreallocatedData(RenderInstanceData);
+	if (PerInstanceRenderData)
+	{
+		FStaticMeshInstanceData RenderInstanceData = FStaticMeshInstanceData(GVertexElementTypeSupport.IsSupported(VET_Half2));
+		BuildRenderData(RenderInstanceData, PerInstanceRenderData->HitProxies);
+		PerInstanceRenderData->UpdateFromPreallocatedData(RenderInstanceData);
+	}
 }
 
 bool UInstancedStaticMeshComponent::IsHLODRelevant() const
