PR #9447: Correct color display of grayscale compression with virtual textures when previewing textures (Contributed by [REDACTED_PERSON])

diff --git a/Engine/Source/Runtime/Engine/Private/Texture2D.cpp b/Engine/Source/Runtime/Engine/Private/Texture2D.cpp
--- a/Engine/Source/Runtime/Engine/Private/Texture2D.cpp
+++ b/Engine/Source/Runtime/Engine/Private/Texture2D.cpp
@@ -1265,6 +1265,8 @@ FVirtualTexture2DResource::FVirtualTexture2DResource(const UTexture2D* InOwner,
 	FirstMipToUse = FMath::Min((int32)MaxMip, InFirstMipToUse);
 
 	bSRGB = InOwner->SRGB;
+	const EPixelFormat PixelFormat = VTData->LayerTypes[0];
+	bGreyScaleFormat = (PixelFormat == PF_G8) || (PixelFormat == PF_BC4);
 
 	// Initialize this resource FeatureLevel, so it gets re-created on FeatureLevel changes
 	SetFeatureLevel(GMaxRHIFeatureLevel);
