[materials]
- fix non-determinism in material shader map; correctly account for permutation ID in keys when storing preprocessed source and deduplicate.
- fix material stats source lookup not accounting for permutation IDs and deprecate old API
- fix material stats searching for an invalid key in the case where vertex factory is null
- change material logging to construct & output the same hash as DDC, the FSHA1 hashed keys don't match
- change LogMaterial channel to allow VeryVerbose messages
- add VeryVerbose logging of the unhashed (i.e. human-readable) DDC key for material shader maps
- add a debug utility driven via cvar to output the textual representation of shader maps for a given named material to disk (was key in tracking down the non-determinism mentioned above)

diff --git a/Engine/Shaders/Private/MaterialTemplate.ush b/Engine/Shaders/Private/MaterialTemplate.ush
--- a/Engine/Shaders/Private/MaterialTemplate.ush
+++ b/Engine/Shaders/Private/MaterialTemplate.ush
@@ -31,7 +31,7 @@
 
 // Update this GUID to force all material shaders to recompile
 // Merge conflicts on this line should be resolved by generating a new GUID
-// GUID = 5F290441-5611-4232-B26D-613F5E421262_A
+// GUID = 4BCB6F1E-5842-441E-934D-F999D009CCB5
 
 #if USES_SPEEDTREE
 	#include "/Engine/Private/SpeedTreeCommon.ush"
diff --git a/Engine/Source/Editor/MaterialEditor/Private/MaterialStats.cpp b/Engine/Source/Editor/MaterialEditor/Private/MaterialStats.cpp
--- a/Engine/Source/Editor/MaterialEditor/Private/MaterialStats.cpp
+++ b/Engine/Source/Editor/MaterialEditor/Private/MaterialStats.cpp
@@ -123,12 +123,7 @@ FText FShaderPlatformSettings::GetShaderCode(const EMaterialQualityLevel::Type Q
 		if (Entry != nullptr)
 		{
 			const TShaderRef<FShader>& Shader = *Entry;
-
-			static FName NAME_NullVF(TEXT("NullVF"));
-			FVertexFactoryType* VertexFactory = Shader.GetVertexFactoryType();
-			const FName VertexFactoryName = VertexFactory ? VertexFactory->GetFName() : NAME_NullVF;
-			const FName ShaderTypeName = Shader.GetType()->GetFName();
-			const FMemoryImageString* ShaderSource = MaterialShaderMap->GetShaderSource(VertexFactoryName, ShaderTypeName);
+			const FMemoryImageString* ShaderSource = MaterialShaderMap->GetShaderSource(Shader.GetVertexFactoryType(), Shader.GetType(), ShaderId.PermutationId);
 			if (ShaderSource != nullptr)
 			{
 				PlatformData[QualityType].bUpdateShaderCode = false;
diff --git a/Engine/Source/Runtime/Engine/Private/Materials/MaterialShader.cpp b/Engine/Source/Runtime/Engine/Private/Materials/MaterialShader.cpp
--- a/Engine/Source/Runtime/Engine/Private/Materials/MaterialShader.cpp
+++ b/Engine/Source/Runtime/Engine/Private/Materials/MaterialShader.cpp
@@ -21,6 +21,7 @@
 #include "ProfilingDebugging/CountersTrace.h"
 #include "ProfilingDebugging/LoadTimeTracker.h"
 #include "Misc/ScopeRWLock.h"
+#include "Misc/PathViews.h"
 #include "SceneTexturesConfig.h"
 #include "PSOPrecache.h"
 
@@ -32,6 +33,14 @@ static FAutoConsoleVariableRef CVarMaterialExcludeNonPipelinedShaders(
 	ECVF_ReadOnly
 );
 
+static TAutoConsoleVariable<FString> CVarMaterialShaderMapDump(
+	TEXT("r.Material.ShaderMapDump"),
+	"",
+	TEXT("Outputs a textual dump of all shader maps found for the given named material (specified by path).\n")
+	TEXT("Note that this will include any instances of said material created by a MaterialInstance.\n")
+	TEXT("Files (.txt extension) will be dumped to Saved\\MaterialShaderMaps named with the DDC key hash.\n"),
+	ECVF_ReadOnly);
+
 #if ENABLE_COOK_STATS
 namespace MaterialShaderCookStats
 {
@@ -208,20 +217,10 @@ static UE::DerivedData::FSharedString GetMaterialShaderMapName(const FStringView
 	return UE::DerivedData::FSharedString(WriteToString<256>(MaterialPath,
 		TEXTVIEW(" ["), LegacyShaderPlatformToShaderFormat(Platform),
 		TEXTVIEW(", "), FeatureLevelName,
-		TEXTVIEW(", "), GetMaterialQualityLevelFName(ShaderMapId.QualityLevel),
+		TEXTVIEW(", "), LexToString(ShaderMapId.QualityLevel),
 		TEXTVIEW("]")));
 }
 
-static FString FSHA1_HashString(const FString& Key)
-{
-	FSHAHash DDCKeyHash;
-	FSHA1 HashState;
-	HashState.UpdateWithString(*Key, Key.Len());
-	HashState.Final();
-	HashState.GetHash(&DDCKeyHash.Hash[0]);
-	return DDCKeyHash.ToString();
-}
-
 #endif // WITH_EDITOR
 
 /** Called for every material shader to update the appropriate stats. */
@@ -1547,7 +1546,8 @@ TSharedRef<FMaterialShaderMap::FAsyncLoadContext> FMaterialShaderMap::BeginLoadF
 
 				GShaderCompilerStats->AddDDCHit(1);
 
-				UE_LOG(LogMaterial, Verbose, TEXT("Loaded shaders for %s from DDC (key hash: %s)"), *AssetName, *FSHA1_HashString(DataKey));
+				FCacheKey DdcKey = GetMaterialShaderMapKey(DataKey);
+				UE_LOG(LogMaterial, Verbose, TEXT("Loaded shaders for %s from DDC (key hash: %s)"), *AssetName, *LexToString(DdcKey.Hash));
 			}
 			else
 			{
@@ -1582,7 +1582,8 @@ TSharedRef<FMaterialShaderMap::FAsyncLoadContext> FMaterialShaderMap::BeginLoadF
 			COOK_STAT(Timer.TrackCyclesOnly());
 
 			Result->DataKey = GetMaterialShaderMapKeyString(ShaderMapId, InPlatform);
-			OutDDCKeyDesc = FSHA1_HashString(Result->DataKey);
+			FCacheKey CacheKey = GetMaterialShaderMapKey(Result->DataKey);
+			OutDDCKeyDesc = LexToString(CacheKey.Hash);
 
 			if (UNLIKELY(ShouldDumpShaderDDCKeys()))
 			{
@@ -1656,12 +1657,28 @@ void FMaterialShaderMap::SaveToDerivedDataCache()
 	COOK_STAT(Timer.AddMiss(SaveData.Num()));
 
 	const FString DataKey = GetMaterialShaderMapKeyString(ShaderMapId, GetShaderPlatform());
-	UE_LOG(LogMaterial, Verbose, TEXT("Saved shaders for %s to DDC (key hash: %s)"), GetMaterialPath(), *FSHA1_HashString(DataKey));
 
 	using namespace UE::DerivedData;
 	FCachePutValueRequest Request;
 	Request.Name = GetMaterialShaderMapName(GetMaterialPath(), ShaderMapId, GetShaderPlatform());
 	Request.Key = GetMaterialShaderMapKey(DataKey);
+
+	UE_LOG(LogMaterial, Verbose, TEXT("Saved shaders for %s to DDC (key hash: %s)"), *Request.Name, *LexToString(Request.Key.Hash));
+	UE_LOG(LogMaterial, VeryVerbose, TEXT("Full DDC data key for %s: %s"), *Request.Name, *DataKey);
+
+	if (!CVarMaterialShaderMapDump->GetString().IsEmpty() && CVarMaterialShaderMapDump->GetString() == GetMaterialPath())
+	{
+		TStringBuilder<256> Path;
+		FPathViews::Append(Path, FPaths::ProjectSavedDir(), TEXT("MaterialShaderMaps"), TEXT(""));
+		Path << Request.Key.Hash << ".txt";
+		FArchive* DumpAr = IFileManager::Get().CreateFileWriter(*Path, FILEWRITE_Silent);
+		if (DumpAr)
+		{
+			FTCHARToUTF8 Converter(*ToString());
+			DumpAr->Serialize(const_cast<UTF8CHAR*>(reinterpret_cast<const UTF8CHAR*>(Converter.Get())), Converter.Length());
+		}
+	}
+
 	Request.Value = FValue::Compress(MakeSharedBufferFromArray(MoveTemp(SaveData)));
 	FRequestOwner AsyncOwner(EPriority::Normal);
 	FRequestBarrier AsyncBarrier(AsyncOwner);
@@ -2310,9 +2327,16 @@ void FMaterialShaderMap::Compile(
 	}
 }
 
-static FHashedName PreprocessedSourceKeyFromName(FName VertexFactoryName, FName VertexTypeName)
+static FHashedName GetPreprocessedSourceKey(const FVertexFactoryType* VertexFactoryType, const FShaderType* ShaderType, int32 PermutationId)
 {
-	return FHashedName(FString::Printf(TEXT("%s/%s"), *VertexFactoryName.ToString(), *VertexTypeName.ToString()));
+	if (VertexFactoryType)
+	{
+		return FHashedName(FString::Printf(TEXT("%s/%s/%d"), VertexFactoryType->GetName(), ShaderType->GetName(), PermutationId));
+	}
+	else
+	{
+		return FHashedName(FString::Printf(TEXT("%s/%d"), ShaderType->GetName(), PermutationId));
+	}
 }
 
 FShader* FMaterialShaderMap::ProcessCompilationResultsForSingleJob(FShaderCompileJob* SingleJob, const FShaderPipelineType* ShaderPipeline, const FSHAHash& MaterialShaderMapHash)
@@ -2363,11 +2387,14 @@ FShader* FMaterialShaderMap::ProcessCompilationResultsForSingleJob(FShaderCompil
 #if WITH_EDITOR
 	// add shader source
 	{
-		// Keep the preprocessed source list sorted by type name
-		const FHashedName Key = CurrentJob.Key.VFType ? PreprocessedSourceKeyFromName(CurrentJob.Key.VFType->GetFName(), CurrentJob.Key.ShaderType->GetFName()) : FHashedName(CurrentJob.Key.ShaderType->GetFName());
+		// Keep the preprocessed source list sorted by a name constructed from VF/ShaderType/PermutationId and deduplicate entries
+		const FHashedName Key = GetPreprocessedSourceKey(CurrentJob.Key.VFType, CurrentJob.Key.ShaderType, CurrentJob.Key.PermutationId);
 
 		const int32 Index = Algo::LowerBoundBy(GetMutableContent()->ShaderProcessedSource, Key, [](const FMaterialProcessedSource& Value) { return Value.Name; });
-		GetMutableContent()->ShaderProcessedSource.EmplaceAt(Index, Key, *CurrentJob.Output.OptionalFinalShaderSource);
+		if (Index >= GetMutableContent()->ShaderProcessedSource.Num() || GetMutableContent()->ShaderProcessedSource[Index].Name != Key)
+		{
+			GetMutableContent()->ShaderProcessedSource.EmplaceAt(Index, Key, *CurrentJob.Output.OptionalFinalShaderSource);
+		}
 	}
 #endif
 
@@ -2875,10 +2902,9 @@ void FMaterialShaderMap::LoadMissingShadersFromMemory(const FMaterial* Material)
 #endif // WITH_EDITOR
 
 #if WITH_EDITOR
-const FMemoryImageString* FMaterialShaderMap::GetShaderSource(const FName VertexFactoryName, const FName ShaderTypeName) const
+const FMemoryImageString* FMaterialShaderMap::GetShaderSource(const FVertexFactoryType* VertexFactoryType, const FShaderType* ShaderType, int32 PermutationId) const
 {
-	const FHashedName Key = PreprocessedSourceKeyFromName(VertexFactoryName, ShaderTypeName);
-
+	FHashedName Key = GetPreprocessedSourceKey(VertexFactoryType, ShaderType, PermutationId);
 	for (const FMaterialProcessedSource& Source : GetContent()->ShaderProcessedSource)
 	{
 		if (Source.Name == Key)
@@ -2889,6 +2915,11 @@ const FMemoryImageString* FMaterialShaderMap::GetShaderSource(const FName Vertex
 
 	return nullptr;
 }
+
+const FMemoryImageString* FMaterialShaderMap::GetShaderSource(const FName VertexFactoryName, const FName ShaderTypeName) const
+{
+	return GetShaderSource(FindVertexFactoryType(VertexFactoryName), FindShaderTypeByName(ShaderTypeName), /* PermutationId */ 0);
+}
 #endif // WITH_EDITOR
 
 void FMaterialShaderMap::GetShaderList(TMap<FShaderId, TShaderRef<FShader>>& OutShaders) const
diff --git a/Engine/Source/Runtime/Engine/Public/MaterialShared.h b/Engine/Source/Runtime/Engine/Public/MaterialShared.h
--- a/Engine/Source/Runtime/Engine/Public/MaterialShared.h
+++ b/Engine/Source/Runtime/Engine/Public/MaterialShared.h
@@ -102,7 +102,7 @@ template <class ElementType> class TLinkedList;
 // Adds various checks to track FMaterial lifetime, may add some overhead
 #define UE_CHECK_FMATERIAL_LIFETIME PLATFORM_WINDOWS
 
-DECLARE_LOG_CATEGORY_EXTERN(LogMaterial,Log,Verbose);
+DECLARE_LOG_CATEGORY_EXTERN(LogMaterial, Log, All);
 
 /** Creates a string that represents the given quality level. */
 extern void GetMaterialQualityLevelName(EMaterialQualityLevel::Type InMaterialQualityLevel, FString& OutName);
@@ -1267,7 +1267,9 @@ public:
 #endif
 
 #if WITH_EDITOR
-	ENGINE_API const FMemoryImageString *GetShaderSource(FName VertexFactoryName, FName ShaderTypeName) const;
+	ENGINE_API const FMemoryImageString* GetShaderSource(const FVertexFactoryType* VertexFactoryType, const FShaderType* ShaderType, int32 PermutationId) const;
+	UE_DEPRECATED(5.2, "Use the GetShaderSource overload accepting a FVertexFactoryType/FShaderType/permutation ID; this one doesn't correctly account for permutations.")
+	ENGINE_API const FMemoryImageString* GetShaderSource(FName VertexFactoryName, FName ShaderTypeName) const;
 #endif
 
 	/** Builds a list of the shaders in a shader map. */
