Fix crash in accessing null skin cache entry in FSkeletalMeshSceneProxy::UpdateLooseParametersUniformBuffer. It is a fallout from CL22429289

diff --git a/Engine/Source/Runtime/Engine/Private/SkeletalMesh.cpp b/Engine/Source/Runtime/Engine/Private/SkeletalMesh.cpp
--- a/Engine/Source/Runtime/Engine/Private/SkeletalMesh.cpp
+++ b/Engine/Source/Runtime/Engine/Private/SkeletalMesh.cpp
@@ -6225,17 +6225,18 @@ void FSkeletalMeshSceneProxy::CreateBaseMeshBatch(const FSceneView* View, const
 	BatchElement.MinVertexIndex = LODData.RenderSections[SectionIndex].GetVertexBufferIndex();
 	BatchElement.MaxVertexIndex = LODData.RenderSections[SectionIndex].GetVertexBufferIndex() + LODData.RenderSections[SectionIndex].GetNumVertices() - 1;
 #if RHI_RAYTRACING
-	BatchElement.VertexFactoryUserData = FGPUSkinCache::GetFactoryUserData(VFMode == ESkinVertexFactoryMode::RayTracing ? MeshObject->GetSkinCacheEntryForRayTracing() : MeshObject->SkinCacheEntry, SectionIndex);
+	FGPUSkinBatchElementUserData* VertexFactoryUserData = FGPUSkinCache::GetFactoryUserData(VFMode == ESkinVertexFactoryMode::RayTracing ? MeshObject->GetSkinCacheEntryForRayTracing() : MeshObject->SkinCacheEntry, SectionIndex);
 #else
-	BatchElement.VertexFactoryUserData = FGPUSkinCache::GetFactoryUserData(MeshObject->SkinCacheEntry, SectionIndex);
+	FGPUSkinBatchElementUserData* VertexFactoryUserData = FGPUSkinCache::GetFactoryUserData(MeshObject->SkinCacheEntry, SectionIndex);
 #endif
+	BatchElement.VertexFactoryUserData = VertexFactoryUserData;
 	BatchElement.PrimitiveUniformBuffer = GetUniformBuffer();
 	BatchElement.NumPrimitives = LODData.RenderSections[SectionIndex].NumTriangles;
 
-	UpdateLooseParametersUniformBuffer(View, SectionIndex, Mesh, BatchElement);
+	UpdateLooseParametersUniformBuffer(View, SectionIndex, Mesh, VertexFactoryUserData);
 }
 
-void FSkeletalMeshSceneProxy::UpdateLooseParametersUniformBuffer(const FSceneView* View, const int32 SectionIndex, const FMeshBatch& Mesh, const FMeshBatchElement& BatchElement) const
+void FSkeletalMeshSceneProxy::UpdateLooseParametersUniformBuffer(const FSceneView* View, const int32 SectionIndex, const FMeshBatch& Mesh, const FGPUSkinBatchElementUserData* BatchUserData) const
 {
 	// Loose parameters uniform buffer is only needed for PassThroughVF
 	if (!Mesh.VertexFactory || Mesh.VertexFactory->GetType() != &FGPUSkinPassthroughVertexFactory::StaticType)
@@ -6253,12 +6254,12 @@ void FSkeletalMeshSceneProxy::UpdateLooseParametersUniformBuffer(const FSceneVie
 	// If world is paused, use current frame bone matrices, so velocity is canceled and skeletal mesh isn't blurred from motion.
 	bool bVerticesInMotion = !View->Family->bWorldIsPaused && bBoneDataUpdatedThisFrame;
 
-	const bool bUsesSkinCache = BatchElement.VertexFactoryUserData != nullptr;
+	const bool bUsesSkinCache = BatchUserData != nullptr;
 	if (bUsesSkinCache)
 	{
 		Parameters.GPUSkinPassThroughPreviousPositionBuffer = bVerticesInMotion ?
-			FGPUSkinCache::GetPreviousPositionBuffer(MeshObject->SkinCacheEntry, SectionIndex)->SRV :
-			FGPUSkinCache::GetPositionBuffer(MeshObject->SkinCacheEntry, SectionIndex)->SRV;
+			FGPUSkinCache::GetPreviousPositionBuffer(BatchUserData->Entry, SectionIndex)->SRV :
+			FGPUSkinCache::GetPositionBuffer(BatchUserData->Entry, SectionIndex)->SRV;
 	}
 	else // Mesh deformer
 	{
diff --git a/Engine/Source/Runtime/Engine/Public/SkeletalMeshTypes.h b/Engine/Source/Runtime/Engine/Public/SkeletalMeshTypes.h
--- a/Engine/Source/Runtime/Engine/Public/SkeletalMeshTypes.h
+++ b/Engine/Source/Runtime/Engine/Public/SkeletalMeshTypes.h
@@ -337,7 +337,7 @@ protected:
 	uint8 GetCurrentFirstLODIdx_Internal() const;
 private:
 	void CreateBaseMeshBatch(const FSceneView* View, const FSkeletalMeshLODRenderData& LODData, const int32 LODIndex, const int32 SectionIndex, const FSectionElementInfo& SectionElementInfo, FMeshBatch& Mesh, ESkinVertexFactoryMode VFMode = ESkinVertexFactoryMode::Default) const;
-	void UpdateLooseParametersUniformBuffer(const FSceneView* View, const int32 SectionIndex, const FMeshBatch& Mesh, const FMeshBatchElement& BatchElement) const;
+	void UpdateLooseParametersUniformBuffer(const FSceneView* View, const int32 SectionIndex, const FMeshBatch& Mesh, const struct FGPUSkinBatchElementUserData* BatchUserData) const;
 
 public:
 #if WITH_EDITORONLY_DATA
