Fixed crash trying to re-instance post-process anim instances on components with no mesh

diff --git a/Engine/Source/Runtime/Engine/Private/Animation/AnimInstance.cpp b/Engine/Source/Runtime/Engine/Private/Animation/AnimInstance.cpp
--- a/Engine/Source/Runtime/Engine/Private/Animation/AnimInstance.cpp
+++ b/Engine/Source/Runtime/Engine/Private/Animation/AnimInstance.cpp
@@ -3887,26 +3887,27 @@ void UAnimInstance::HandleObjectsReinstanced(const TMap<UObject*, UObject*>& Old
 		
 		if(bThisObjectWasReinstanced)
 		{
-			TRACE_OBJECT_LIFETIME_BEGIN(this);
-			// Minimally reinit proxy (i.e. dont call per-node initialization) unless we are in an editor preview world (i.e. we are in the anim BP editor)
-			UWorld* World = GetWorld();
-			if(World && World->WorldType == EWorldType::EditorPreview)
+			USkeletalMeshComponent* MeshComponent = GetSkelMeshComponent();
+			if(MeshComponent && MeshComponent->GetSkeletalMeshAsset())
 			{
-				InitializeAnimation(false);
-			}
-			else
-			{
-				RecalcRequiredBones();
+				TRACE_OBJECT_LIFETIME_BEGIN(this);
+				// Minimally reinit proxy (i.e. dont call per-node initialization) unless we are in an editor preview world (i.e. we are in the anim BP editor)
+				UWorld* World = GetWorld();
+				if(World && World->WorldType == EWorldType::EditorPreview)
+				{
+					InitializeAnimation(false);
+				}
+				else
+				{
+					RecalcRequiredBones();
 
-				FAnimInstanceProxy& Proxy = GetProxyOnGameThread<FAnimInstanceProxy>();
-				Proxy.Initialize(this);
-				Proxy.InitializeCachedClassData();
-				Proxy.InitializeRootNode_WithRoot(Proxy.RootNode);
-			}
+					FAnimInstanceProxy& Proxy = GetProxyOnGameThread<FAnimInstanceProxy>();
+					Proxy.Initialize(this);
+					Proxy.InitializeCachedClassData();
+					Proxy.InitializeRootNode_WithRoot(Proxy.RootNode);
+				}
 
-			if(USkeletalMeshComponent* Mesh = GetSkelMeshComponent())
-			{
-				Mesh->ClearMotionVector();
+				MeshComponent->ClearMotionVector();
 			}
 		}
 	}
