Characters desync when standing on a moving base.
#rnx

diff --git a/Engine/Source/Runtime/Engine/Classes/GameFramework/Character.h b/Engine/Source/Runtime/Engine/Classes/GameFramework/Character.h
--- a/Engine/Source/Runtime/Engine/Classes/GameFramework/Character.h
+++ b/Engine/Source/Runtime/Engine/Classes/GameFramework/Character.h
@@ -148,6 +148,12 @@ namespace MovementBaseUtility
 
 	/** Get the transforms for the given MovementBase, optionally at the location of a bone. Returns false if MovementBase is nullptr, or if BoneName is not a valid bone. */
 	ENGINE_API bool GetMovementBaseTransform(const UPrimitiveComponent* MovementBase, const FName BoneName, FVector& OutLocation, FQuat& OutQuat);
+
+	/** Get the world space location from a local space location for a given MovementBase. Returns false if MovementBase is nullptr, or if BoneName is not a valid bone. Scaling is ignored. */
+	ENGINE_API bool GetLocalMovementBaseLocationInWorldSpace(const UPrimitiveComponent* MovementBase, const FName BoneName, const FVector& LocalLocation, FVector& OutLocationWorldSpace);
+
+	/** Get the local space location for a given MovementBase, optionally at the location of a bone. Returns false if MovementBase is nullptr, or if BoneName is not a valid bone. Scaling is ignored. */
+	ENGINE_API bool GetLocalMovementBaseLocation(const UPrimitiveComponent* MovementBase, const FName BoneName, const FVector& WorldSpaceLocation, FVector& OutLocalLocation);
 }
 
 /** Struct to hold information about the "base" object the character is standing on. */
diff --git a/Engine/Source/Runtime/Engine/Private/Character.cpp b/Engine/Source/Runtime/Engine/Private/Character.cpp
--- a/Engine/Source/Runtime/Engine/Private/Character.cpp
+++ b/Engine/Source/Runtime/Engine/Private/Character.cpp
@@ -690,6 +690,24 @@ namespace MovementBaseUtility
 		OutQuat = FQuat::Identity;
 		return false;
 	}
+
+	bool GetLocalMovementBaseLocationInWorldSpace(const UPrimitiveComponent* MovementBase, const FName BoneName, const FVector& LocalLocation, FVector& OutLocationWorldSpace)
+	{
+		FVector OutLocation;
+		FQuat OutQuat;
+		const bool bResult = GetMovementBaseTransform(MovementBase, BoneName, OutLocation, OutQuat);
+		OutLocationWorldSpace = FTransform(OutQuat, OutLocation).TransformPositionNoScale(LocalLocation);
+		return bResult;
+	}
+
+	bool GetLocalMovementBaseLocation(const UPrimitiveComponent* MovementBase, const FName BoneName, const FVector& WorldSpaceLocation, FVector& OutLocalLocation)
+	{
+		FVector OutLocation;
+		FQuat OutQuat;
+		const bool bResult = GetMovementBaseTransform(MovementBase, BoneName, OutLocation, OutQuat);
+		OutLocalLocation = FTransform(OutQuat, OutLocation).InverseTransformPositionNoScale(WorldSpaceLocation);
+		return bResult;
+	}
 }
 
 
@@ -1149,7 +1167,8 @@ void ACharacter::OnRep_ReplicatedBasedMovement()
 		const FVector OldLocation = GetActorLocation();
 		const FQuat OldRotation = GetActorQuat();
 		MovementBaseUtility::GetMovementBaseTransform(ReplicatedBasedMovement.MovementBase, ReplicatedBasedMovement.BoneName, CharacterMovement->OldBaseLocation, CharacterMovement->OldBaseQuat);
-		const FVector NewLocation = CharacterMovement->OldBaseLocation + ReplicatedBasedMovement.Location;
+		const FTransform BaseTransform(CharacterMovement->OldBaseQuat, CharacterMovement->OldBaseLocation);
+		const FVector NewLocation = BaseTransform.TransformPositionNoScale(ReplicatedBasedMovement.Location);
 		FRotator NewRotation;
 
 		if (ReplicatedBasedMovement.HasRelativeRotation())
diff --git a/Engine/Source/Runtime/Engine/Private/Components/CharacterMovementComponent.cpp b/Engine/Source/Runtime/Engine/Private/Components/CharacterMovementComponent.cpp
--- a/Engine/Source/Runtime/Engine/Private/Components/CharacterMovementComponent.cpp
+++ b/Engine/Source/Runtime/Engine/Private/Components/CharacterMovementComponent.cpp
@@ -2730,7 +2730,8 @@ void UCharacterMovementComponent::SaveBaseLocation()
 		if (MovementBaseUtility::UseRelativeLocation(MovementBase))
 		{
 			// Relative Location
-			const FVector RelativeLocation = UpdatedComponent->GetComponentLocation() - OldBaseLocation;
+			FVector RelativeLocation;
+			MovementBaseUtility::GetLocalMovementBaseLocation(MovementBase, CharacterOwner->GetBasedMovement().BoneName, UpdatedComponent->GetComponentLocation(), RelativeLocation);
 
 			// Rotation
 			if (bIgnoreBaseRotation)
@@ -9327,10 +9328,7 @@ void UCharacterMovementComponent::ServerMoveHandleClientError(float ClientTimeSt
 	FVector ClientLoc = RelativeClientLoc;
 	if (MovementBaseUtility::UseRelativeLocation(ClientMovementBase))
 	{
-		FVector BaseLocation;
-		FQuat BaseRotation;
-		MovementBaseUtility::GetMovementBaseTransform(ClientMovementBase, ClientBaseBoneName, BaseLocation, BaseRotation);
-		ClientLoc += BaseLocation;
+		MovementBaseUtility::GetLocalMovementBaseLocationInWorldSpace(ClientMovementBase, ClientBaseBoneName, RelativeClientLoc, ClientLoc);
 	}
 	else
 	{
@@ -9436,12 +9434,8 @@ void UCharacterMovementComponent::ServerMoveHandleClientError(float ClientTimeSt
 				// To improve position syncing, use old base for take-off
 				if (MovementBaseUtility::UseRelativeLocation(LastServerMovementBasePtr))
 				{
-					FVector BaseLocation;
-					FQuat BaseQuat;
-					MovementBaseUtility::GetMovementBaseTransform(LastServerMovementBasePtr, LastServerMovementBaseBoneName, BaseLocation, BaseQuat);
-
 					// Relative Location
-					RelativeLocation = UpdatedComponent->GetComponentLocation() - BaseLocation;
+					MovementBaseUtility::GetLocalMovementBaseLocation(LastServerMovementBasePtr, LastServerMovementBaseBoneName, UpdatedComponent->GetComponentLocation(), RelativeLocation);
 					bUseLastBase = true;
 				}
 			}
@@ -10295,10 +10289,7 @@ void UCharacterMovementComponent::ClientAdjustPosition_Implementation
 	//  Received Location is relative to dynamic base
 	if (bBaseRelativePosition)
 	{
-		FVector BaseLocation;
-		FQuat BaseRotation;
-		MovementBaseUtility::GetMovementBaseTransform(NewBase, NewBaseBoneName, BaseLocation, BaseRotation); // TODO: error handling if returns false		
-		WorldShiftedNewLocation = NewLocation + BaseLocation;
+		MovementBaseUtility::GetLocalMovementBaseLocationInWorldSpace(NewBase, NewBaseBoneName, NewLocation, WorldShiftedNewLocation); // TODO: error handling if returns false	
 	}
 	else
 	{
@@ -11870,9 +11861,9 @@ FVector FSavedMove_Character::GetRevertedLocation() const
 	const UPrimitiveComponent* MovementBase = StartBase.Get();
 	if (MovementBaseUtility::UseRelativeLocation(MovementBase))
 	{
-		FVector BaseLocation; FQuat BaseRotation;
-		MovementBaseUtility::GetMovementBaseTransform(MovementBase, StartBoneName, BaseLocation, BaseRotation);
-		return BaseLocation + StartRelativeLocation;
+		FVector WorldSpacePosition;
+		MovementBaseUtility::GetLocalMovementBaseLocationInWorldSpace(MovementBase, StartBoneName, StartRelativeLocation, WorldSpacePosition);
+		return WorldSpacePosition;
 	}
 
 	return StartLocation;
