Fixing issues with net.TrackDormantObjectsByLevel

diff --git a/Engine/Source/Runtime/Engine/Private/NetworkObjectList.cpp b/Engine/Source/Runtime/Engine/Private/NetworkObjectList.cpp
--- a/Engine/Source/Runtime/Engine/Private/NetworkObjectList.cpp
+++ b/Engine/Source/Runtime/Engine/Private/NetworkObjectList.cpp
@@ -136,11 +136,12 @@ void FNetworkObjectList::Remove(AActor* const Actor)
 
 		if (Connection == nullptr || Connection->GetConnectionState() == USOCK_Closed)
 		{
-			ConnectionIt.RemoveCurrent();
 			if (UE::Net::Private::bTrackDormantObjectsByLevel)
 			{
 				DormantObjectsPerConnection.Remove(Connection);
 			}
+
+			ConnectionIt.RemoveCurrent();
 			continue;
 		}
 
@@ -154,9 +155,8 @@ void FNetworkObjectList::Remove(AActor* const Actor)
 			{
 				if (FNetworkObjectSet* DormantObjects = DormantObjectsByLevel->Find(PackageName))
 				{
-					const int32 NumRemoved = DormantObjects->Remove(Actor);
-					checkf(NumRemoved > 0, TEXT("Actor not found in Connection->Level->Dormant map: %s"), *GetNameSafe(Actor));
-
+					DormantObjects->Remove(Actor);
+					
 					if (DormantObjects->Num() == 0)
 					{
 						DormantObjectsByLevel->Remove(PackageName);
@@ -213,12 +213,13 @@ void FNetworkObjectList::MarkDormant(AActor* const Actor, UNetConnection* const
 
 		if (UE::Net::Private::bTrackDormantObjectsByLevel)
 		{
+			// make sure the connection map exists
+			TMap<FName, FNetworkObjectSet>& DormantObjectsByLevel = DormantObjectsPerConnection.FindOrAdd(Connection);
+
 			// if not fully dormant
 			if (NetworkObjectInfo->DormantConnections.Num() != NumConnections)
 			{
-				TMap<FName, FNetworkObjectSet>& DormantObjectsByLevel = DormantObjectsPerConnection.FindOrAdd(Connection);
 				FNetworkObjectSet& DormantObjects = DormantObjectsByLevel.FindOrAdd(PackageName);
-
 				DormantObjects.Add(*NetworkObjectInfoPtr);
 			}
 		}
@@ -231,11 +232,12 @@ void FNetworkObjectList::MarkDormant(AActor* const Actor, UNetConnection* const
 	{
 		if ((*ConnectionIt).Get() == nullptr || (*ConnectionIt).Get()->GetConnectionState() == USOCK_Closed)
 		{
-			ConnectionIt.RemoveCurrent();
 			if (UE::Net::Private::bTrackDormantObjectsByLevel)
 			{
 				DormantObjectsPerConnection.Remove((*ConnectionIt).Get());
 			}
+
+			ConnectionIt.RemoveCurrent();
 		}
 	}
 
@@ -255,7 +257,19 @@ void FNetworkObjectList::MarkDormant(AActor* const Actor, UNetConnection* const
 			// Remove from connection object lists
 			for (auto ConnectionIt = NetworkObjectInfo->DormantConnections.CreateIterator(); ConnectionIt; ++ConnectionIt)
 			{
-				DormantObjectsPerConnection.Remove((*ConnectionIt).Get());
+				if (TMap<FName, FNetworkObjectSet>* DormantObjectsByLevel = DormantObjectsPerConnection.Find((*ConnectionIt).Get()))
+				{
+					if (FNetworkObjectSet* DormantObjects = DormantObjectsByLevel->Find(PackageName))
+					{
+						const int32 NumRemoved = DormantObjects->Remove(Actor);
+						checkf((NumRemoved > 0) || (Connection == (*ConnectionIt).Get()), TEXT("Actor not found in Connection->Level->Dormant map: %s"), *GetNameSafe(Actor));
+
+						if (DormantObjects->Num() == 0)
+						{
+							DormantObjectsByLevel->Remove(PackageName);
+						}
+					}
+				}
 			}
 		}
 
@@ -286,12 +300,21 @@ bool FNetworkObjectList::MarkActiveInternal(const TSharedPtr<FNetworkObjectInfo>
 		{
 			if (FNetworkObjectSet* FullyDormant = FullyDormantObjectsByLevel.Find(PackageName))
 			{
-				FullyDormant->Remove(Actor);
+				const int32 NumRemoved = FullyDormant->Remove(Actor);
+				checkf(NumRemoved > 0, TEXT("Actor not found in full Level->Dormant map: %s"), *GetNameSafe(Actor));
 
 				if (FullyDormant->Num() == 0)
 				{
 					FullyDormantObjectsByLevel.Remove(PackageName);
 				}
+
+				// add back into per connection maps
+				for (auto ConnectionIt = NetworkObjectInfo->DormantConnections.CreateIterator(); ConnectionIt; ++ConnectionIt)
+				{
+					TMap<FName, FNetworkObjectSet>& DormantObjectsByLevel = DormantObjectsPerConnection.FindOrAdd((*ConnectionIt).Get());
+					FNetworkObjectSet& DormantObjects = DormantObjectsByLevel.FindOrAdd(PackageName);
+					DormantObjects.Add(ObjectInfo);
+				}
 			}
 		}
 	}
