Fix level streaming already existing on the server in PIE leading to double activation

diff --git a/Engine/Source/Runtime/Engine/Private/WorldPartition/ContentBundle/ContentBundleEditor.cpp b/Engine/Source/Runtime/Engine/Private/WorldPartition/ContentBundle/ContentBundleEditor.cpp
--- a/Engine/Source/Runtime/Engine/Private/WorldPartition/ContentBundle/ContentBundleEditor.cpp
+++ b/Engine/Source/Runtime/Engine/Private/WorldPartition/ContentBundle/ContentBundleEditor.cpp
@@ -307,7 +307,7 @@ void FContentBundleEditor::GenerateStreaming(TArray<FString>* OutPackageToGenera
 	UWorldPartition* WorldPartition = GetInjectedWorld()->GetWorldPartition();
 	WorldPartition->GenerateContainerStreaming(ActorDescContainer.Get(), OutPackageToGenerate);
 
-	ExternalStreamingObject = WorldPartition->RuntimeHash->StoreToExternalStreamingObject(GetTransientPackage(), *GetExternalStreamingObjectName());
+	ExternalStreamingObject = WorldPartition->RuntimeHash->StoreToExternalStreamingObject(WorldPartition, *GetExternalStreamingObjectName());
 
 	uint32 CellCount = 0;
 	ExternalStreamingObject->ForEachStreamingCells([&CellCount](const UWorldPartitionRuntimeCell& Cell)
