Have AudioComponents initialize their underlying FQuartzTickableObject w/ the UClockHandle's world instead of its own.
(makes it compatible with audio components spawned ith/ Create Sound 2D/"persisit across leveltransition" = true)

diff --git a/Engine/Source/Runtime/Engine/Private/Components/AudioComponent.cpp b/Engine/Source/Runtime/Engine/Private/Components/AudioComponent.cpp
--- a/Engine/Source/Runtime/Engine/Private/Components/AudioComponent.cpp
+++ b/Engine/Source/Runtime/Engine/Private/Components/AudioComponent.cpp
@@ -466,7 +466,9 @@ void UAudioComponent::PlayQuantized(
 	//Initialize the tickable object portion of the Audio Component, if it hasn't been initialized already
 	if (!FQuartzTickableObject::IsInitialized())
 	{
-		Init(GetWorld());
+		// if the WorldContextObject is null, attempt to fallback on our world
+		// (call to Init() will ensure on a valid world)
+		Init(WorldContextObject ? WorldContextObject->GetWorld() : GetWorld());
 	}
 	check(FQuartzTickableObject::IsInitialized());
 
