[Quartz] Fix for regression with in StartOtherClock (copy constructor omitting the other clock's name)

diff --git a/Engine/Source/Runtime/Engine/Classes/Sound/QuartzQuantizationUtilities.cpp b/Engine/Source/Runtime/Engine/Classes/Sound/QuartzQuantizationUtilities.cpp
--- a/Engine/Source/Runtime/Engine/Classes/Sound/QuartzQuantizationUtilities.cpp
+++ b/Engine/Source/Runtime/Engine/Classes/Sound/QuartzQuantizationUtilities.cpp
@@ -382,6 +382,7 @@ namespace Audio
 		, int32 InSourceID
 	)
 		: ClockName(RHS.ClockName)
+		, OtherClockName(RHS.OtherClockName)
 		, QuantizedCommandPtr(RHS.QuantizedCommandPtr)
 		, QuantizationBoundary(RHS.QuantizationBoundary)
 		, GameThreadSubscribers(RHS.GameThreadSubscribers)
