World Partition: fix crash when trying to render queue on partitioned map with streaming disabled.

#rnx

diff --git a/Engine/Source/Runtime/Engine/Private/WorldPartition/WorldPartitionRuntimeLevelStreamingCell.cpp b/Engine/Source/Runtime/Engine/Private/WorldPartition/WorldPartitionRuntimeLevelStreamingCell.cpp
--- a/Engine/Source/Runtime/Engine/Private/WorldPartition/WorldPartitionRuntimeLevelStreamingCell.cpp
+++ b/Engine/Source/Runtime/Engine/Private/WorldPartition/WorldPartitionRuntimeLevelStreamingCell.cpp
@@ -425,12 +425,26 @@ void UWorldPartitionRuntimeLevelStreamingCell::Deactivate() const
 void UWorldPartitionRuntimeLevelStreamingCell::OnLevelShown()
 {
 	check(LevelStreaming);
-	LevelStreaming->GetWorld()->GetSubsystem<UHLODSubsystem>()->OnCellShown(this);
+
+	const UWorldPartition* WorldPartition = LevelStreaming->GetWorld()->GetWorldPartition();
+	check(WorldPartition);
+
+	if (WorldPartition->IsStreamingEnabled())
+	{
+		LevelStreaming->GetWorld()->GetSubsystem<UHLODSubsystem>()->OnCellShown(this);
+	}
 }
 
 void UWorldPartitionRuntimeLevelStreamingCell::OnLevelHidden()
 {
 	check(LevelStreaming);
-	LevelStreaming->GetWorld()->GetSubsystem<UHLODSubsystem>()->OnCellHidden(this);
+
+	const UWorldPartition* WorldPartition = LevelStreaming->GetWorld()->GetWorldPartition();
+	check(WorldPartition);
+
+	if (WorldPartition->IsStreamingEnabled())
+	{
+		LevelStreaming->GetWorld()->GetSubsystem<UHLODSubsystem>()->OnCellHidden(this);
+	}
 }
 
