Revert 22344826: "Check for zero-length queries in ConvertTraceResults". Remove the ensure - having zero length dir is not "valid" per se, but also doesn't break anything, so long as we don't ensure on it.

The original change fixed an ensure but caused another: https://horde.devtools.epicgames.com/job/6344f86a687cb259e5bde6d9

diff --git a/Engine/Source/Runtime/Engine/Private/Collision/CollisionConversions.cpp b/Engine/Source/Runtime/Engine/Private/Collision/CollisionConversions.cpp
--- a/Engine/Source/Runtime/Engine/Private/Collision/CollisionConversions.cpp
+++ b/Engine/Source/Runtime/Engine/Private/Collision/CollisionConversions.cpp
@@ -408,10 +408,6 @@ EConvertQueryResult ConvertTraceResults(bool& OutHasValidBlockingHit, const UWor
 	EConvertQueryResult ConvertResult = EConvertQueryResult::Valid;
 	bool bHadBlockingHit = false;
 	const FVector Dir = (EndLoc - StartLoc).GetSafeNormal();
-	if (Dir.IsNearlyZero())
-	{
-		return EConvertQueryResult::Invalid;
-	}
 
 	for(int32 i=0; i<NumHits; i++)
 	{
@@ -451,11 +447,6 @@ template <typename Hit>
 EConvertQueryResult ConvertTraceResults(bool& OutHasValidBlockingHit, const UWorld* World, int32 NumHits, Hit* Hits, float CheckLength, const FCollisionFilterData& QueryFilter, FHitResult& OutHit, const FVector& StartLoc, const FVector& EndLoc, const FPhysicsGeometry& Geom, const FTransform& QueryTM, float MaxDistance, bool bReturnFaceIndex, bool bReturnPhysMat)
 {
 	const FVector Dir = (EndLoc - StartLoc).GetSafeNormal();
-	if (Dir.IsNearlyZero())
-	{
-		return EConvertQueryResult::Invalid;
-	}
-
 	if (TIsSame<Hit, FHitSweep>::Value)
 	{
 		if (!HadInitialOverlap(Hits[0]))
diff --git a/Engine/Source/Runtime/Experimental/Chaos/Public/Chaos/ImplicitObjectScaled.h b/Engine/Source/Runtime/Experimental/Chaos/Public/Chaos/ImplicitObjectScaled.h
--- a/Engine/Source/Runtime/Experimental/Chaos/Public/Chaos/ImplicitObjectScaled.h
+++ b/Engine/Source/Runtime/Experimental/Chaos/Public/Chaos/ImplicitObjectScaled.h
@@ -834,7 +834,6 @@ public:
 
 	virtual int32 FindMostOpposingFace(const FVec3& Position, const FVec3& UnitDir, int32 HintFaceIndex, FReal SearchDist) const override
 	{
-		ensure(FMath::IsNearlyEqual(UnitDir.SizeSquared(), FReal(1), FReal(UE_KINDA_SMALL_NUMBER)));
 		return MObject->FindMostOpposingFaceScaled(Position, UnitDir, HintFaceIndex, SearchDist, MScale);
 	}
 
