Adding LogRepCompares very verbose logging for property comparisons

diff --git a/Engine/Source/Runtime/Engine/Private/RepLayout.cpp b/Engine/Source/Runtime/Engine/Private/RepLayout.cpp
--- a/Engine/Source/Runtime/Engine/Private/RepLayout.cpp
+++ b/Engine/Source/Runtime/Engine/Private/RepLayout.cpp
@@ -39,9 +39,11 @@ DECLARE_CYCLE_STAT(TEXT("RepLayout DeltaSerializeFastArray"), STAT_RepLayout_Del
 #if !(UE_BUILD_SHIPPING || UE_BUILD_TEST)
 DEFINE_LOG_CATEGORY_STATIC(LogRepProperties, Warning, All);
 DEFINE_LOG_CATEGORY_STATIC(LogRepPropertiesBackCompat, Warning, All);
+DEFINE_LOG_CATEGORY_STATIC(LogRepCompares, Warning, All);
 #else
 DEFINE_LOG_CATEGORY_STATIC(LogRepProperties, Warning, Warning);
 DEFINE_LOG_CATEGORY_STATIC(LogRepPropertiesBackCompat, Warning, Warning);
+DEFINE_LOG_CATEGORY_STATIC(LogRepCompares, Warning, Warning);
 #endif
 
 int32 GDoPropertyChecksum = 0;
@@ -754,6 +756,8 @@ static FORCEINLINE bool CompareNetSerializeStructWithObjectProperties(
 		// This will work with packed or unpacked shadow buffers, because net serialize structs
 		// aren't packed.
 
+		UE_LOG(LogRepCompares, VeryVerbose, TEXT("CompareNetSerializeStructWithObjectProperties: CmdIndex: %d CmdType: %s Property: %s"), CmdIndex, LexToString(Cmd.Type), *GetNameSafe(Cmd.Property));
+
 		if (ERepLayoutCmdType::DynamicArray == Cmd.Type)
 		{
 			FScriptArrayHelper AArray((FArrayProperty*)Cmd.Property, ((const uint8*)A + Cmd.Offset));
@@ -815,6 +819,7 @@ static FORCEINLINE bool PropertiesAreIdentical(
 
 static FORCEINLINE void StoreProperty(const FRepLayoutCmd& Cmd, void* A, const void* B)
 {
+	UE_LOG(LogRepCompares, VeryVerbose, TEXT("StoreProperty: %s"), *GetNameSafe(Cmd.Property));
 	Cmd.Property->CopySingleValue(A, B);
 }
 
@@ -1345,6 +1350,9 @@ static bool CompareRoleProperty(
 	const FRepLayoutCmd& RoleOrRemoteRoleCmd = SharedParams.Cmds[RoleOrRemoteRoleParent.CmdStart];
 	const uint16 Handle = RoleOrRemoteRoleCmd.RelativeHandle;
 	const TEnumAsByte<ENetRole> ActorRoleOrRemoteRole = *(const TEnumAsByte<ENetRole>*)(StackParams.Data + RoleOrRemoteRoleParent).Data;
+
+	UE_LOG(LogRepCompares, VeryVerbose, TEXT("CompareRoleProperty: bForceFail: %d, SavedRole: %s ActorRole: %s"), SharedParams.bForceFail, *UEnum::GetValueAsString<ENetRole>(SavedRoleOrRemoteRole.GetValue()), *UEnum::GetValueAsString<ENetRole>(ActorRoleOrRemoteRole.GetValue()));
+
 	if (SharedParams.bForceFail || SavedRoleOrRemoteRole != ActorRoleOrRemoteRole)
 	{
 		SavedRoleOrRemoteRole = ActorRoleOrRemoteRole;
@@ -1480,6 +1488,8 @@ static void CompareParentProperties(
 		// If we're forcibly comparing all properties, then don't bother checking dirty state.
 		if (UNLIKELY(SharedParams.bForceFail))
 		{
+			UE_LOG(LogRepCompares, VeryVerbose, TEXT("CompareParentProperties: Force failed"));
+
 			for (int32 ParentIndex = 0; ParentIndex < SharedParams.Parents.Num(); ++ParentIndex)
 			{
 				UE_RepLayout_Private::CompareParentPropertyHelper(ParentIndex, SharedParams, StackParams);
@@ -1490,6 +1500,8 @@ static void CompareParentProperties(
 		// If we're running validation, then we'll check everything regardless of push model state.
 		else if (SharedParams.bValidateProperties)
 		{
+			UE_LOG(LogRepCompares, VeryVerbose, TEXT("CompareParentProperties: Property validation"));
+
 			for (int32 ParentIndex = 0; ParentIndex < SharedParams.Parents.Num(); ++ParentIndex)
 			{
 				const bool bRecompareInitialProperties = SharedParams.bIsInitial && SharedParams.Parents[ParentIndex].Condition == COND_InitialOnly;
@@ -1531,6 +1543,9 @@ static void CompareParentProperties(
 				To get around that, if we detect that this is an Initial Replication, and we have Initial Only
 				properties, we will consider them dirty.
 			*/
+
+			UE_LOG(LogRepCompares, VeryVerbose, TEXT("CompareParentProperties: Initial only test"));
+
 			for (int32 ParentIndex = 0; ParentIndex < SharedParams.Parents.Num(); ++ParentIndex)
 			{
 				if (SharedParams.Parents[ParentIndex].Condition == COND_InitialOnly ||
@@ -1544,6 +1559,8 @@ static void CompareParentProperties(
 		// If we have full push model property support, then we only need to check properties that are actually dirty.
 		else if (EnumHasAnyFlags(SharedParams.Flags, ERepLayoutFlags::FullPushProperties) && !bRecentlyCollectedGarbage)
 		{
+			UE_LOG(LogRepCompares, VeryVerbose, TEXT("CompareParentProperties: Full push properties: Has Dirty: %d"), !!SharedParams.PushModelState->HasDirtyProperties());
+
 			for (TConstSetBitIterator<> It = SharedParams.PushModelState->GetDirtyProperties(); It; ++It)
 			{
 				UE_RepLayout_Private::CompareParentPropertyHelper(It.GetIndex(), SharedParams, StackParams);
@@ -1551,6 +1568,8 @@ static void CompareParentProperties(
 		}
 		else
 		{
+			UE_LOG(LogRepCompares, VeryVerbose, TEXT("CompareParentProperties: Default"));
+
 			for (int32 ParentIndex = 0; ParentIndex < SharedParams.Parents.Num(); ++ParentIndex)
 			{
 				if (UE_RepLayout_Private::IsPropertyDirty(ParentIndex, bRecentlyCollectedGarbage, SharedParams, StackParams))
@@ -1590,6 +1609,8 @@ static uint16 CompareProperties_r(
 		const FConstRepObjectDataBuffer Data = StackParams.Data + Cmd;
 		FRepShadowDataBuffer ShadowData = StackParams.ShadowData + Cmd;
 
+		UE_LOG(LogRepCompares, VeryVerbose, TEXT("CompareProperties_r: CmdIndex: %d CmdType: %s Property: %s"), CmdIndex, LexToString(Cmd.Type), *GetNameSafe(Cmd.Property));
+
 		if (Cmd.Type == ERepLayoutCmdType::DynamicArray)
 		{
 			FComparePropertiesStackParams NewStackParams{
@@ -1654,6 +1675,8 @@ static void CompareProperties_Array_r(
 		bool& bForceFail = const_cast<bool&>(SharedParams.bForceFail);
 		TGuardValue<bool> ForceFailGuard(bForceFail, bForceFail);
 
+		UE_LOG(LogRepCompares, VeryVerbose, TEXT("CompareProperties_Array_r: ArrayNum: %d"), ArrayNum);
+
 		for (int32 i = 0; i < ArrayNum; i++)
 		{
 			const int32 ArrayElementOffset = i * Cmd.ElementSize;
@@ -1717,6 +1740,8 @@ ERepLayoutResult FRepLayout::CompareProperties(
 
 	FRepChangedHistory& NewHistoryItem = RepChangelistState->ChangeHistory[HistoryIndex];
 
+	UE_LOG(LogRepCompares, VeryVerbose, TEXT("CompareProperties: Owner: %s CompareIndex: %d HistoryIndex: %d"), *GetFullNameSafe(Owner), RepChangelistState->CompareIndex, HistoryIndex);
+
 	TArray<uint16>& Changed = NewHistoryItem.Changed;
 	Changed.Empty(1);
 
@@ -8383,6 +8408,68 @@ const TCHAR* LexToString(ERepLayoutFlags Flag)
 	}
 }
 
+const TCHAR* LexToString(ERepLayoutCmdType CmdType)
+{
+	switch (CmdType)
+	{
+	case ERepLayoutCmdType::DynamicArray:
+		return TEXT("DynamicArray");
+	case ERepLayoutCmdType::Return:
+		return TEXT("Return");
+	case ERepLayoutCmdType::Property:
+		return TEXT("Property");
+	case ERepLayoutCmdType::PropertyBool:
+		return TEXT("PropertyBool");
+	case ERepLayoutCmdType::PropertyFloat:
+		return TEXT("PropertyFloat");
+	case ERepLayoutCmdType::PropertyInt:
+		return TEXT("PropertyInt");
+	case ERepLayoutCmdType::PropertyByte:
+		return TEXT("PropertyByte");
+	case ERepLayoutCmdType::PropertyName:
+		return TEXT("PropertyName");
+	case ERepLayoutCmdType::PropertyObject:
+		return TEXT("PropertyObject");
+	case ERepLayoutCmdType::PropertyUInt32:
+		return TEXT("PropertyUInt32");
+	case ERepLayoutCmdType::PropertyVector:
+		return TEXT("PropertyVector");
+	case ERepLayoutCmdType::PropertyRotator:
+		return TEXT("PropertyRotator");
+	case ERepLayoutCmdType::PropertyPlane:
+		return TEXT("PropertyPlane");
+	case ERepLayoutCmdType::PropertyVector100:
+		return TEXT("PropertyVector100");
+	case ERepLayoutCmdType::PropertyNetId:
+		return TEXT("PropertyNetId");
+	case ERepLayoutCmdType::RepMovement:
+		return TEXT("RepMovement");
+	case ERepLayoutCmdType::PropertyVectorNormal:
+		return TEXT("PropertyVectorNormal");
+	case ERepLayoutCmdType::PropertyVector10:
+		return TEXT("PropertyVector10");
+	case ERepLayoutCmdType::PropertyVectorQ:
+		return TEXT("PropertyVectorQ");
+	case ERepLayoutCmdType::PropertyString:
+		return TEXT("PropertyString");
+	case ERepLayoutCmdType::PropertyUInt64:
+		return TEXT("PropertyUInt64");
+	case ERepLayoutCmdType::PropertyNativeBool:
+		return TEXT("PropertyNativeBool");
+	case ERepLayoutCmdType::PropertySoftObject:
+		return TEXT("PropertySoftObject");
+	case ERepLayoutCmdType::PropertyWeakObject:
+		return TEXT("PropertyWeakObject");
+	case ERepLayoutCmdType::PropertyInterface:
+		return TEXT("PropertyInterface");
+	case ERepLayoutCmdType::NetSerializeStructWithObjectReferences:
+		return TEXT("NetSerializeStructWithObjectReferences");
+	default:
+		ensureMsgf(false, TEXT("Unhandled layout command type."));
+		return TEXT("Unknown");
+	}
+}
+
 #define REPDATATYPE_SPECIALIZATION(DstType, SrcType) \
 template bool FRepLayout::DiffStableProperties(TArray<FProperty*>*, TArray<UObject*>*, TRepDataBuffer<DstType>, TConstRepDataBuffer<SrcType>) const; \
 template bool FRepLayout::DiffProperties(TArray<FProperty*>*, TRepDataBuffer<DstType>, TConstRepDataBuffer<SrcType>, const EDiffPropertiesFlags) const;
diff --git a/Engine/Source/Runtime/Engine/Public/Net/RepLayout.h b/Engine/Source/Runtime/Engine/Public/Net/RepLayout.h
--- a/Engine/Source/Runtime/Engine/Public/Net/RepLayout.h
+++ b/Engine/Source/Runtime/Engine/Public/Net/RepLayout.h
@@ -733,6 +733,8 @@ enum class ERepLayoutCmdType : uint8
 	NetSerializeStructWithObjectReferences = 25,
 };
 
+const TCHAR* LexToString(ERepLayoutCmdType CmdType);
+
 /** Various flags that describe how a Top Level Property should be handled. */
 enum class ERepParentFlags : uint32
 {
