Fix race in FDistanceFieldVolumeData constructor most likely leading to duplicate IDs

#rnx

diff --git a/Engine/Source/Runtime/Engine/Private/DistanceFieldAtlas.cpp b/Engine/Source/Runtime/Engine/Private/DistanceFieldAtlas.cpp
--- a/Engine/Source/Runtime/Engine/Private/DistanceFieldAtlas.cpp
+++ b/Engine/Source/Runtime/Engine/Private/DistanceFieldAtlas.cpp
@@ -36,6 +36,8 @@
 #include "IMeshBuilderModule.h"
 #endif
 
+#include <atomic>
+
 #define LOCTEXT_NAMESPACE "DistanceField"
 
 #if ENABLE_COOK_STATS
@@ -236,15 +238,14 @@ void FDistanceFieldVolumeData::CacheDerivedData(const FString& InStaticMeshDeriv
 
 #endif
 
-uint64 NextDistanceFieldVolumeDataId = 1;
+std::atomic<uint64> NextDistanceFieldVolumeDataId { 1 };
 
 FDistanceFieldVolumeData::FDistanceFieldVolumeData() :
 	LocalSpaceMeshBounds(ForceInit),
 	bMostlyTwoSided(false),
 	bAsyncBuilding(false)
 {
-	Id = NextDistanceFieldVolumeDataId;
-	NextDistanceFieldVolumeDataId++;
+	Id = NextDistanceFieldVolumeDataId++;
 }
 
 void FDistanceFieldVolumeData::Serialize(FArchive& Ar, UObject* Owner)
