Fix Strata shading mode conversion when driven from shading expression.

diff --git a/Engine/Source/Runtime/Engine/Private/Materials/Material.cpp b/Engine/Source/Runtime/Engine/Private/Materials/Material.cpp
--- a/Engine/Source/Runtime/Engine/Private/Materials/Material.cpp
+++ b/Engine/Source/Runtime/Engine/Private/Materials/Material.cpp
@@ -3130,7 +3130,7 @@ void UMaterial::ConvertMaterialToStrataMaterial()
 		bRelinkCustomOutputNodes = true;
 
 		// * Remove support for material attribute
-		// * explicitely connect the Strata node to the root node
+		// * explicitly connect the Strata node to the root node
 		// * Forward inputs to the root node (Do not reconnect the Opacity as we handle the opacity by internally within the conversion node)
 		// * Always forward masked opacity because this is required when the blend mode is overriden to masked on a material instance.
 		bUseMaterialAttributes = false;
@@ -3151,10 +3151,9 @@ void UMaterial::ConvertMaterialToStrataMaterial()
 		{
 			check(!bHasShadingModelMixture);
 
-			// Add constant for the shading model
-			UMaterialExpressionConstant* ShadingModelNode = NewObject<UMaterialExpressionConstant>(this);
-			ShadingModelNode->SetParameterName(FName(TEXT("ConstantShadingModel")));
-			ShadingModelNode->R = ShadingModel;
+			// Add shading model node
+			UMaterialExpressionShadingModel* ShadingModelNode = NewObject<UMaterialExpressionShadingModel>(this);
+			ShadingModelNode->ShadingModel = ShadingModel;
 			ConvertNode->ShadingModel.Connect(0, ShadingModelNode);
 
 			// Store strata shading model of the converted material. 
@@ -3278,7 +3277,6 @@ void UMaterial::ConvertMaterialToStrataMaterial()
 				// rebuilding the final Shading model (see RebuildShadingModelField())
 				if (ShadingModel == MSM_FromMaterialExpression)
 				{
-					//check(EditorOnly->ShadingModelFromMaterialExpression.IsConnected());
 					if (!EditorOnly->ShadingModelFromMaterialExpression.IsConnected())
 					{
 						UMaterialExpressionShadingModel* ShadingModelNode = NewObject<UMaterialExpressionShadingModel>(this);
@@ -3287,8 +3285,9 @@ void UMaterial::ConvertMaterialToStrataMaterial()
 					}
 					else
 					{
-						// Reconnect the shading model expression
-						MoveConnectionTo(EditorOnly->ShadingModelFromMaterialExpression, ConvertNode, 19);
+						// Reconnect the shading model expression. 
+						// Note: assign the expression directly, as using ConvertNode->GetInput(19)->Connect(..) causes the expression to not be assigned
+						ConvertNode->ShadingModel.Expression = EditorOnly->ShadingModelFromMaterialExpression.Expression;
 					}
 
 					// Store strata shading model of the converted material. 
