Strata - fix for wrong world space normal.

diff --git a/Engine/Source/Runtime/Engine/Private/Materials/MaterialExpressions.cpp b/Engine/Source/Runtime/Engine/Private/Materials/MaterialExpressions.cpp
--- a/Engine/Source/Runtime/Engine/Private/Materials/MaterialExpressions.cpp
+++ b/Engine/Source/Runtime/Engine/Private/Materials/MaterialExpressions.cpp
@@ -21748,14 +21748,17 @@ static int32 CompileWithDefaultNormalWS(class FMaterialCompiler* Compiler, FExpr
 			{
 				*bDefaultIsUsed = true;
 			}
+			// Nothing is plug in from the linked input, so specify world space normal the BSDF node expects.
 			return Compiler->VertexNormal();
 		}
+		// Transform into world space normal if needed. BSDF nodes always expects world space normal as input.
 		return Compiler->TransformNormalFromRequestedBasisToWorld(NormalCodeChunk);
 	}
 	if (bDefaultIsUsed)
 	{
 		*bDefaultIsUsed = true;
 	}
+	// Nothing is plug in on the BSDF node, so specify world space normal the node expects.
 	return Compiler->VertexNormal();
 }
 static int32 CompileWithDefaultTangentWS(class FMaterialCompiler* Compiler, FExpressionInput& Input, bool* bDefaultIsUsed = nullptr)
@@ -21770,14 +21773,17 @@ static int32 CompileWithDefaultTangentWS(class FMaterialCompiler* Compiler, FExp
 			{
 				*bDefaultIsUsed = true;
 			}
+			// Nothing is plug in from the linked input, so specify world space tangent the BSDF node expects.
 			return Compiler->VertexTangent();
 		}
+		// Transform into world space tangent if needed. BSDF nodes always expects world space tangent as input.
 		return Compiler->TransformNormalFromRequestedBasisToWorld(TangentCodeChunk);
 	}
 	if (bDefaultIsUsed)
 	{
 		*bDefaultIsUsed = true;
 	}
+	// Nothing is plug in on the BSDF node, so specify world space tangent the node expects.
 	return Compiler->VertexTangent();
 }
 
diff --git a/Engine/Source/Runtime/Engine/Private/Materials/MaterialShared.cpp b/Engine/Source/Runtime/Engine/Private/Materials/MaterialShared.cpp
--- a/Engine/Source/Runtime/Engine/Private/Materials/MaterialShared.cpp
+++ b/Engine/Source/Runtime/Engine/Private/Materials/MaterialShared.cpp
@@ -1470,7 +1470,14 @@ void FMaterial::DiscardShaderMap()
 }
 
 EMaterialDomain FMaterialResource::GetMaterialDomain() const { return Material->MaterialDomain; }
-bool FMaterialResource::IsTangentSpaceNormal() const { return Material->bTangentSpaceNormal || (!Material->IsPropertyConnected(MP_Normal) && !Material->bUseMaterialAttributes); }
+bool FMaterialResource::IsTangentSpaceNormal() const 
+{ 
+	if (IsStrataMaterial())
+	{
+		return Material->bTangentSpaceNormal; // We do not need to check MP_Normal with strata as this cannot be specified on the root node anymore.
+	}
+	return Material->bTangentSpaceNormal || (!Material->IsPropertyConnected(MP_Normal) && !Material->bUseMaterialAttributes);
+}
 bool FMaterialResource::ShouldGenerateSphericalParticleNormals() const { return Material->bGenerateSphericalParticleNormals; }
 bool FMaterialResource::ShouldDisableDepthTest() const { return Material->bDisableDepthTest; }
 bool FMaterialResource::ShouldWriteOnlyAlpha() const { return Material->bWriteOnlyAlpha; }
