Fixing GetShadingModelString() for hair.

diff --git a/Engine/Source/Runtime/Engine/Private/Materials/MaterialShader.cpp b/Engine/Source/Runtime/Engine/Private/Materials/MaterialShader.cpp
--- a/Engine/Source/Runtime/Engine/Private/Materials/MaterialShader.cpp
+++ b/Engine/Source/Runtime/Engine/Private/Materials/MaterialShader.cpp
@@ -93,6 +93,7 @@ FString GetShadingModelString(EMaterialShadingModel ShadingModel)
 		case MSM_ClearCoat:			ShadingModelName = TEXT("MSM_ClearCoat"); break;
 		case MSM_SubsurfaceProfile:	ShadingModelName = TEXT("MSM_SubsurfaceProfile"); break;
 		case MSM_TwoSidedFoliage:	ShadingModelName = TEXT("MSM_TwoSidedFoliage"); break;
+		case MSM_Hair:				ShadingModelName = TEXT("MSM_Hair"); break;
 		case MSM_Cloth:				ShadingModelName = TEXT("MSM_Cloth"); break;
 		case MSM_Eye:				ShadingModelName = TEXT("MSM_Eye"); break;
 		case MSM_SingleLayerWater:	ShadingModelName = TEXT("MSM_SingleLayerWater"); break;
