Fix outdated assertion in  FShaderCompilingManager::ProcessAsyncResults().

#rnx

diff --git a/Engine/Source/Runtime/Engine/Private/ShaderCompiler/ShaderCompiler.cpp b/Engine/Source/Runtime/Engine/Private/ShaderCompiler/ShaderCompiler.cpp
--- a/Engine/Source/Runtime/Engine/Private/ShaderCompiler/ShaderCompiler.cpp
+++ b/Engine/Source/Runtime/Engine/Private/ShaderCompiler/ShaderCompiler.cpp
@@ -5276,106 +5276,98 @@ void FShaderCompilingManager::ProcessAsyncResults(bool bLimitExecutionTime, bool
 
 	COOK_STAT(FScopedDurationTimer Timer(ShaderCompilerCookStats::ProcessAsyncResultsTimeSec));
 	check(IsInGameThread());
-	if (bAllowAsynchronousShaderCompiling)
-	{
-		{
-			const double StartTime = FPlatformTime::Seconds();
 
-			// Some controllers need to be manually ticked if the engine loop is not initialized or blocked
-			// to do things like tick the HTTPModule.
-			// Otherwise the results from the controller will never be processed.
-			// We check for bBlockOnGlobalShaderCompletion because the BlockOnShaderMapCompletion methods already do this.
-			if (!bBlockOnGlobalShaderCompletion && BuildDistributionController)
-			{
-				BuildDistributionController->Tick(0.0f);
-			}
-			
-			// Block on global shaders before checking for shader maps to finalize
-			// So if we block on global shaders for a long time, we will get a chance to finalize all the non-global shader maps completed during that time.
-			if (bBlockOnGlobalShaderCompletion)
-			{
-				TArray<int32> ShaderMapId;
-				ShaderMapId.Add(GlobalShaderMapId);
+	const double StartTime = FPlatformTime::Seconds();
 
-				// Block until the global shader map jobs are complete
-				GShaderCompilingManager->BlockOnShaderMapCompletion(ShaderMapId, PendingFinalizeShaderMaps);
-			}
+	// Some controllers need to be manually ticked if the engine loop is not initialized or blocked
+	// to do things like tick the HTTPModule.
+	// Otherwise the results from the controller will never be processed.
+	// We check for bBlockOnGlobalShaderCompletion because the BlockOnShaderMapCompletion methods already do this.
+	if (!bBlockOnGlobalShaderCompletion && BuildDistributionController)
+	{
+		BuildDistributionController->Tick(0.0f);
+	}
 
-			int32 NumCompilingShaderMaps = 0;
+	// Block on global shaders before checking for shader maps to finalize
+	// So if we block on global shaders for a long time, we will get a chance to finalize all the non-global shader maps completed during that time.
+	if (bBlockOnGlobalShaderCompletion)
+	{
+		TArray<int32> ShaderMapId;
+		ShaderMapId.Add(GlobalShaderMapId);
 
-			{
-				// Lock CompileQueueSection so we can access the input and output queues
-				FScopeLock Lock(&CompileQueueSection);
+		// Block until the global shader map jobs are complete
+		GShaderCompilingManager->BlockOnShaderMapCompletion(ShaderMapId, PendingFinalizeShaderMaps);
+	}
 
-				if (!bBlockOnGlobalShaderCompletion)
-				{
-					bCompilingDuringGame = true;
-				}
+	int32 NumCompilingShaderMaps = 0;
 
-				// Get all material shader maps to finalize
-				//
-				for (TMap<int32, FPendingShaderMapCompileResultsPtr>::TIterator It(ShaderMapJobs); It; ++It)
-				{
-					FPendingShaderMapCompileResultsPtr& Results = It.Value();
-					if (Results->FinishedJobs.Num() > 0)
-					{
-						FShaderMapFinalizeResults& FinalizeResults = PendingFinalizeShaderMaps.FindOrAdd(It.Key());
-						FinalizeResults.FinishedJobs.Append(Results->FinishedJobs);
-						Results->FinishedJobs.Reset();
-						FinalizeResults.bAllJobsSucceeded = FinalizeResults.bAllJobsSucceeded && Results->bAllJobsSucceeded;
-					}
+	{
+		// Lock CompileQueueSection so we can access the input and output queues
+		FScopeLock Lock(&CompileQueueSection);
 
-					checkf(Results->FinishedJobs.Num() == 0, TEXT("Failed to remove finished jobs, %d remain"), Results->FinishedJobs.Num());
-					if (Results->NumPendingJobs.GetValue() == 0)
-					{
-						It.RemoveCurrent();
-					}
-				}
+		if (!bBlockOnGlobalShaderCompletion)
+		{
+			bCompilingDuringGame = true;
+		}
 
-				NumCompilingShaderMaps = ShaderMapJobs.Num();
+		// Get all material shader maps to finalize
+		//
+		for (TMap<int32, FPendingShaderMapCompileResultsPtr>::TIterator It(ShaderMapJobs); It; ++It)
+		{
+			FPendingShaderMapCompileResultsPtr& Results = It.Value();
+			if (Results->FinishedJobs.Num() > 0)
+			{
+				FShaderMapFinalizeResults& FinalizeResults = PendingFinalizeShaderMaps.FindOrAdd(It.Key());
+				FinalizeResults.FinishedJobs.Append(Results->FinishedJobs);
+				Results->FinishedJobs.Reset();
+				FinalizeResults.bAllJobsSucceeded = FinalizeResults.bAllJobsSucceeded && Results->bAllJobsSucceeded;
 			}
 
-			int32 NumPendingShaderMaps = PendingFinalizeShaderMaps.Num();
-
-			if (PendingFinalizeShaderMaps.Num() > 0)
+			checkf(Results->FinishedJobs.Num() == 0, TEXT("Failed to remove finished jobs, %d remain"), Results->FinishedJobs.Num());
+			if (Results->NumPendingJobs.GetValue() == 0)
 			{
-				bool bRetry = false;
-				do 
-				{
-					bRetry = HandlePotentialRetryOnError(PendingFinalizeShaderMaps);
-				} 
-				while (bRetry);
-
-				const float TimeBudget = bLimitExecutionTime ? ProcessGameThreadTargetTime : FLT_MAX;
-				ProcessCompiledShaderMaps(PendingFinalizeShaderMaps, TimeBudget);
-				check(bLimitExecutionTime || PendingFinalizeShaderMaps.Num() == 0);
+				It.RemoveCurrent();
 			}
+		}
 
+		NumCompilingShaderMaps = ShaderMapJobs.Num();
+	}
 
-			if (bBlockOnGlobalShaderCompletion && !bLimitExecutionTime)
-			{
-				check(PendingFinalizeShaderMaps.Num() == 0);
+	int32 NumPendingShaderMaps = PendingFinalizeShaderMaps.Num();
 
-				if (NumPendingShaderMaps - PendingFinalizeShaderMaps.Num() > 0)
-				{
-					UE_LOG(LogShaders, Warning, TEXT("Blocking ProcessAsyncResults for %.1fs, processed %u shader maps, %u being compiled"), 
-						(float)(FPlatformTime::Seconds() - StartTime),
-						NumPendingShaderMaps - PendingFinalizeShaderMaps.Num(), 
-						NumCompilingShaderMaps);
-				}
-			}
-			else if (NumPendingShaderMaps - PendingFinalizeShaderMaps.Num() > 0)
-			{
-				UE_LOG(LogShaders, Verbose, TEXT("Completed %u async shader maps, %u more pending, %u being compiled"),
-					NumPendingShaderMaps - PendingFinalizeShaderMaps.Num(), 
-					PendingFinalizeShaderMaps.Num(),
-					NumCompilingShaderMaps);
-			}
+	if (PendingFinalizeShaderMaps.Num() > 0)
+	{
+		bool bRetry = false;
+		do 
+		{
+			bRetry = HandlePotentialRetryOnError(PendingFinalizeShaderMaps);
+		} 
+		while (bRetry);
+
+		const float TimeBudget = bLimitExecutionTime ? ProcessGameThreadTargetTime : FLT_MAX;
+		ProcessCompiledShaderMaps(PendingFinalizeShaderMaps, TimeBudget);
+		check(bLimitExecutionTime || PendingFinalizeShaderMaps.Num() == 0);
+	}
+
+
+	if (bBlockOnGlobalShaderCompletion && !bLimitExecutionTime)
+	{
+		check(PendingFinalizeShaderMaps.Num() == 0);
+
+		if (NumPendingShaderMaps - PendingFinalizeShaderMaps.Num() > 0)
+		{
+			UE_LOG(LogShaders, Warning, TEXT("Blocking ProcessAsyncResults for %.1fs, processed %u shader maps, %u being compiled"), 
+				(float)(FPlatformTime::Seconds() - StartTime),
+				NumPendingShaderMaps - PendingFinalizeShaderMaps.Num(), 
+				NumCompilingShaderMaps);
 		}
 	}
-	else
+	else if (NumPendingShaderMaps - PendingFinalizeShaderMaps.Num() > 0)
 	{
-		check(AllJobs.GetNumPendingJobs() == 0);
+		UE_LOG(LogShaders, Verbose, TEXT("Completed %u async shader maps, %u more pending, %u being compiled"),
+			NumPendingShaderMaps - PendingFinalizeShaderMaps.Num(), 
+			PendingFinalizeShaderMaps.Num(),
+			NumCompilingShaderMaps);
 	}
 
 	UpdateNumRemainingAssets();
