Fix race in FCardRepresentationData constructor most likely leading to duplicate IDs

#rnx

diff --git a/Engine/Source/Runtime/Engine/Public/MeshCardRepresentation.h b/Engine/Source/Runtime/Engine/Public/MeshCardRepresentation.h
--- a/Engine/Source/Runtime/Engine/Public/MeshCardRepresentation.h
+++ b/Engine/Source/Runtime/Engine/Public/MeshCardRepresentation.h
@@ -19,6 +19,8 @@
 #include "DerivedMeshDataTaskUtils.h"
 #include "Async/AsyncWork.h"
 
+#include <atomic>
+
 template <class T> class TLockFreePointerListLIFO;
 class FSignedDistanceFieldBuildMaterialData;
 
@@ -250,9 +252,8 @@ public:
 	FCardRepresentationData() 
 	{
 		// 0 means invalid id.
-		static uint32 NextCardRepresentationId = 0;
-		++NextCardRepresentationId;
-		CardRepresentationDataId.Value = NextCardRepresentationId;
+		static std::atomic<uint32> NextCardRepresentationId { 1 };
+		CardRepresentationDataId.Value = NextCardRepresentationId++;
 	}
 
 	void GetResourceSizeEx(FResourceSizeEx& CumulativeResourceSize) const
