Increase default r.ShaderCompiler.DistributedMinBatchSize from 20 to 50

* This attempts to reduce the impact of the startup overhead when individual shader compilation is relatively fast

diff --git a/Engine/Source/Runtime/Engine/Private/ShaderCompiler/ShaderCompilerDistributed.cpp b/Engine/Source/Runtime/Engine/Private/ShaderCompiler/ShaderCompilerDistributed.cpp
--- a/Engine/Source/Runtime/Engine/Private/ShaderCompiler/ShaderCompilerDistributed.cpp
+++ b/Engine/Source/Runtime/Engine/Private/ShaderCompiler/ShaderCompilerDistributed.cpp
@@ -7,7 +7,7 @@
 namespace DistributedShaderCompilerVariables
 {
 	//TODO: Remove the XGE doublet
-	int32 MinBatchSize = 20;
+	int32 MinBatchSize = 50;
 	FAutoConsoleVariableRef CVarXGEShaderCompileMinBatchSize(
         TEXT("r.XGEShaderCompile.MinBatchSize"),
         MinBatchSize,
@@ -204,7 +204,7 @@ int32 FShaderCompileDistributedThreadRunnable_Interface::CompilingLoop()
 		// Just to provide typical numbers: the number of total jobs is usually in tens of thousands at most, oftentimes in low thousands. Thus JobsPerBatch when calculated as a log2 rarely reaches the value of 16,
 		// and that seems to be a sweet spot: lowering it does not result in faster completion, while increasing the number of jobs per batch slows it down.
 		const uint32 JobsPerBatch = FMath::Max(MinJobsPerBatch, FMath::FloorToInt(FMath::LogX(2.f, PendingJobs.Num() + NumDispatchedJobs)));
-		UE_LOG(LogShaderCompilers, Verbose, TEXT("Current jobs: %d, Batch size: %d, Num Already Dispatched: %d"), PendingJobs.Num(), JobsPerBatch, NumDispatchedJobs);
+		UE_LOG(LogShaderCompilers, Display, TEXT("Current jobs: %d, Batch size: %d, Num Already Dispatched: %d"), PendingJobs.Num(), JobsPerBatch, NumDispatchedJobs);
 
 
 		struct FJobBatch
