Optimization: when not using world partition server streaming, mark all world partition's streaming levels that won't be affected by data layers to skip making any visible/invisible transaction requests as server will always have them loaded/visible.
Small refactor to make readonly some world partition CVars for cooked builds (when necessary) and not to support changing them while running PIE.

diff --git a/Engine/Source/Runtime/Engine/Classes/Engine/LevelStreaming.h b/Engine/Source/Runtime/Engine/Classes/Engine/LevelStreaming.h
--- a/Engine/Source/Runtime/Engine/Classes/Engine/LevelStreaming.h
+++ b/Engine/Source/Runtime/Engine/Classes/Engine/LevelStreaming.h
@@ -227,6 +227,12 @@ protected:
 	UPROPERTY(Category=LevelStreaming, BlueprintSetter=SetShouldBeLoaded, BlueprintGetter=ShouldBeLoaded)
 	uint8 bShouldBeLoaded:1;
 
+	/** Whether the streaming level can safely skip making invisible transaction request from the client to the server */
+	uint8 bSkipClientUseMakingInvisibleTransactionRequest:1;
+
+	/** Whether the streaming level can safely skip making visible transaction request from the client to the server */
+	uint8 bSkipClientUseMakingVisibleTransactionRequest:1;
+
 	/** What the current streamed state of the streaming level is */
 	ECurrentState CurrentState;
 
diff --git a/Engine/Source/Runtime/Engine/Private/LevelStreaming.cpp b/Engine/Source/Runtime/Engine/Private/LevelStreaming.cpp
--- a/Engine/Source/Runtime/Engine/Private/LevelStreaming.cpp
+++ b/Engine/Source/Runtime/Engine/Private/LevelStreaming.cpp
@@ -93,22 +93,24 @@ bool ULevelStreaming::DefaultAllowClientUseMakingVisibleTransactionRequests()
 
 bool ULevelStreaming::ShouldClientUseMakingInvisibleTransactionRequest() const
 {
-	if (UWorld* World = GetWorld())
+	if (!bSkipClientUseMakingInvisibleTransactionRequest)
 	{
 		// Rely on the world to decide whether the client should wait for the server to acknowledge
 		// visibility before making streaming levels invisible on the client.
-		return World->SupportsMakingInvisibleTransactionRequests(); 
+		UWorld* World = GetWorld();
+		return World && World->SupportsMakingInvisibleTransactionRequests();
 	}
 	return false;
 }
 
 bool ULevelStreaming::ShouldClientUseMakingVisibleTransactionRequest() const
 {
-	if (UWorld* World = GetWorld())
+	if (!bSkipClientUseMakingVisibleTransactionRequest)
 	{
 		// Rely on the world to decide whether the client should wait for the server to acknowledge
 		// visibility before making streaming levels visible on the client.
-		return World->SupportsMakingVisibleTransactionRequests();
+		UWorld* World = GetWorld();
+		return World && World->SupportsMakingVisibleTransactionRequests();
 	}
 	return false;
 }
@@ -366,6 +368,8 @@ ULevelStreaming::ULevelStreaming(const FObjectInitializer& ObjectInitializer)
 	bDrawOnLevelStatusMap = true;
 	LevelLODIndex = INDEX_NONE;
 	CurrentState = ECurrentState::Removed;
+	bSkipClientUseMakingInvisibleTransactionRequest = false;
+	bSkipClientUseMakingVisibleTransactionRequest = false;
 }
 
 void ULevelStreaming::PostLoad()
diff --git a/Engine/Source/Runtime/Engine/Private/World.cpp b/Engine/Source/Runtime/Engine/Private/World.cpp
--- a/Engine/Source/Runtime/Engine/Private/World.cpp
+++ b/Engine/Source/Runtime/Engine/Private/World.cpp
@@ -2780,7 +2780,7 @@ bool UWorld::SupportsMakingVisibleTransactionRequests() const
 		// We don't support this flag to change dynamically (read ULevelStreaming CVar once)
 		if (!bSupportsMakingVisibleTransactionRequests.IsSet())
 		{
-			bSupportsMakingVisibleTransactionRequests = ULevelStreaming::DefaultAllowClientUseMakingVisibleTransactionRequests() || (IsPartitionedWorld() && UWorldPartition::UseMakingVisibleTransactionRequests());
+			bSupportsMakingVisibleTransactionRequests = ULevelStreaming::DefaultAllowClientUseMakingVisibleTransactionRequests() || (IsPartitionedWorld() && GetWorldPartition()->UseMakingVisibleTransactionRequests());
 		}
 		return bSupportsMakingVisibleTransactionRequests.GetValue();
 	}
@@ -2794,7 +2794,7 @@ bool UWorld::SupportsMakingInvisibleTransactionRequests() const
 		// We don't support this flag to change dynamically (read ULevelStreaming CVar once)
 		if (!bSupportsMakingInvisibleTransactionRequests.IsSet())
 		{
-			bSupportsMakingInvisibleTransactionRequests = ULevelStreaming::DefaultAllowClientUseMakingInvisibleTransactionRequests() || (IsPartitionedWorld() && UWorldPartition::UseMakingInvisibleTransactionRequests());
+			bSupportsMakingInvisibleTransactionRequests = ULevelStreaming::DefaultAllowClientUseMakingInvisibleTransactionRequests() || (IsPartitionedWorld() && GetWorldPartition()->UseMakingInvisibleTransactionRequests());
 		}
 		return bSupportsMakingInvisibleTransactionRequests.GetValue();
 	}
diff --git a/Engine/Source/Runtime/Engine/Private/WorldPartition/WorldPartition.cpp b/Engine/Source/Runtime/Engine/Private/WorldPartition/WorldPartition.cpp
--- a/Engine/Source/Runtime/Engine/Private/WorldPartition/WorldPartition.cpp
+++ b/Engine/Source/Runtime/Engine/Private/WorldPartition/WorldPartition.cpp
@@ -60,66 +60,73 @@
 #include "Modules/ModuleManager.h"
 #include "GameDelegates.h"
 #include "IDocumentation.h"
+#endif //WITH_EDITOR
+
+#define LOCTEXT_NAMESPACE "WorldPartition"
+
+namespace WorldPartition
+{
+#if WITH_EDITOR
+	static const EConsoleVariableFlags ECVF_Runtime_ReadOnly = ECVF_Default;
+#else
+	static const EConsoleVariableFlags ECVF_Runtime_ReadOnly = ECVF_ReadOnly;
+#endif
+}
 
-static int32 GLoadingRangeBugItGo = 12800;
-static FAutoConsoleVariableRef CVarLoadingRangeBugItGo(
+#if WITH_EDITOR
+int32 UWorldPartition::LoadingRangeBugItGo = 12800;
+FAutoConsoleVariableRef UWorldPartition::CVarLoadingRangeBugItGo(
 	TEXT("wp.Editor.LoadingRangeBugItGo"),
-	GLoadingRangeBugItGo,
+	UWorldPartition::LoadingRangeBugItGo,
 	TEXT("Loading range for BugItGo command."),
-	ECVF_Default
-);
+	ECVF_Default);
 
-static int32 GWorldExtentToEnableStreaming = 400000;
-static FAutoConsoleVariableRef CVarWorldExtentToEnableStreaming(
+int32 UWorldPartition::WorldExtentToEnableStreaming = 400000;
+FAutoConsoleVariableRef UWorldPartition::CVarWorldExtentToEnableStreaming(
 	TEXT("wp.Editor.WorldExtentToEnableStreaming"),
-	GWorldExtentToEnableStreaming,
+	UWorldPartition::WorldExtentToEnableStreaming,
 	TEXT("World extend to justify enabling streaming."),
-	ECVF_Default
-);
+	ECVF_Default);
 
-static bool GDebugDedicatedServerStreaming = false;
-static FAutoConsoleVariableRef CVarDebugDedicatedServerStreaming(
+bool UWorldPartition::DebugDedicatedServerStreaming = false;
+FAutoConsoleVariableRef UWorldPartition::CVarDebugDedicatedServerStreaming(
 	TEXT("wp.Runtime.DebugDedicatedServerStreaming"),
-	GDebugDedicatedServerStreaming,
-	TEXT("Turn on/off to debug of server streaming."));
-
-#endif //WITH_EDITOR
+	UWorldPartition::DebugDedicatedServerStreaming,
+	TEXT("Turn on/off to debug of server streaming."),
+	ECVF_Default);
+#endif
 
-static int32 GEnableServerStreaming = 0;
-static FAutoConsoleVariableRef CVarEnableServerStreaming(
+int32 UWorldPartition::EnableServerStreaming = 0;
+FAutoConsoleVariableRef UWorldPartition::CVarEnableServerStreaming(
 	TEXT("wp.Runtime.EnableServerStreaming"),
-	GEnableServerStreaming,
-	TEXT("Set to 1 to enable server streaming, set to 2 to only enable it in PIE."));
+	UWorldPartition::EnableServerStreaming,
+	TEXT("Set to 1 to enable server streaming, set to 2 to only enable it in PIE.\n")
+	TEXT("Changing the value while the game is running won't be considered."),
+	WorldPartition::ECVF_Runtime_ReadOnly);
 
-static bool GEnableServerStreamingOut = false;
-static FAutoConsoleVariableRef CVarEnableServerStreamingOut(
+bool UWorldPartition::bEnableServerStreamingOut = false;
+FAutoConsoleVariableRef UWorldPartition::CVarEnableServerStreamingOut(
 	TEXT("wp.Runtime.EnableServerStreamingOut"),
-	GEnableServerStreamingOut,
-	TEXT("Turn on/off to allow or not the server to stream out levels (only relevant when server streaming is enabled)"));
+	UWorldPartition::bEnableServerStreamingOut,
+	TEXT("Turn on/off to allow or not the server to stream out levels (only relevant when server streaming is enabled)\n")
+	TEXT("Changing the value while the game is running won't be considered."),
+	WorldPartition::ECVF_Runtime_ReadOnly);
 
 bool UWorldPartition::bUseMakingVisibleTransactionRequests = false;
 FAutoConsoleVariableRef UWorldPartition::CVarUseMakingVisibleTransactionRequests(
 	TEXT("wp.Runtime.UseMakingVisibleTransactionRequests"),
 	UWorldPartition::bUseMakingVisibleTransactionRequests,
-	TEXT("Whether the client should wait for the server to acknowledge visibility update before making partitioned world streaming levels visible."));
+	TEXT("Whether the client should wait for the server to acknowledge visibility update before making partitioned world streaming levels visible.\n")
+	TEXT("Changing the value while the game is running won't be considered."),
+	WorldPartition::ECVF_Runtime_ReadOnly);
 
 bool UWorldPartition::bUseMakingInvisibleTransactionRequests = false;
 FAutoConsoleVariableRef UWorldPartition::CVarUseMakingInvisibleTransactionRequests(
 	TEXT("wp.Runtime.UseMakingInvisibleTransactionRequests"),
 	UWorldPartition::bUseMakingInvisibleTransactionRequests,
-	TEXT("Whether the client should wait for the server to acknowledge visibility update before making partitioned world streaming levels invisible."));
-
-#define LOCTEXT_NAMESPACE "WorldPartition"
-
-bool UWorldPartition::UseMakingVisibleTransactionRequests()
-{
-	return bUseMakingVisibleTransactionRequests;
-}
-
-bool UWorldPartition::UseMakingInvisibleTransactionRequests()
-{
-	return bUseMakingInvisibleTransactionRequests;
-}
+	TEXT("Whether the client should wait for the server to acknowledge visibility update before making partitioned world streaming levels invisible.\n")
+	TEXT("Changing the value while the game is running won't be considered."),
+	WorldPartition::ECVF_Runtime_ReadOnly);
 
 #if WITH_EDITOR
 TMap<FName, FString> GetDataLayersDumpString(const UWorldPartition* WorldPartition)
@@ -742,7 +749,7 @@ void UWorldPartition::OnPostBugItGoCalled(const FVector& Loc, const FRotator& Ro
 		IWorldPartitionEditorModule& WorldPartitionEditorModule = FModuleManager::LoadModuleChecked<IWorldPartitionEditorModule>("WorldPartitionEditor");
 		if (!WorldPartitionEditorModule.GetDisableLoadingInEditor())
 		{
-			const FVector LoadExtent(GLoadingRangeBugItGo, GLoadingRangeBugItGo, HALF_WORLD_MAX);
+			const FVector LoadExtent(UWorldPartition::LoadingRangeBugItGo, UWorldPartition::LoadingRangeBugItGo, HALF_WORLD_MAX);
 			const FBox LoadCellsBox(Loc - LoadExtent, Loc + LoadExtent);
 
 			UWorldPartitionEditorLoaderAdapter* EditorLoaderAdapter = CreateEditorLoaderAdapter<FLoaderAdapterShape>(World, LoadCellsBox, TEXT("BugItGo"));
@@ -957,19 +964,56 @@ bool UWorldPartition::IsServer() const
 
 bool UWorldPartition::IsServerStreamingEnabled() const
 {
-	switch (GEnableServerStreaming)
+	// Resolve once (we don't allow changing the state at runtime)
+	if (!bCachedIsServerStreamingEnabled.IsSet())
 	{
-	case 1:	return true;
+		bool bIsEnabled = false;
+		switch (UWorldPartition::EnableServerStreaming)
+		{
+		case 1:	bIsEnabled = true;
 #if WITH_EDITOR
-	case 2: return bIsPIE;
+		case 2: bIsEnabled = bIsPIE;
 #endif
+		}
+		UWorld* OwningWorld = GetWorld();
+		bCachedIsServerStreamingEnabled = (OwningWorld && OwningWorld->IsGameWorld() && bIsEnabled);
 	}
-	return false;
+	
+	return bCachedIsServerStreamingEnabled.Get(false);
 }
 
 bool UWorldPartition::IsServerStreamingOutEnabled() const
 {
-	return IsServerStreamingEnabled() && GEnableServerStreamingOut;
+	// Resolve once (we don't allow changing the state at runtime)
+	if (!bCachedIsServerStreamingOutEnabled.IsSet())
+	{
+		UWorld* OwningWorld = GetWorld();
+		bCachedIsServerStreamingOutEnabled = OwningWorld && OwningWorld->IsGameWorld() && IsServerStreamingEnabled() && UWorldPartition::bEnableServerStreamingOut;
+	}
+
+	return bCachedIsServerStreamingOutEnabled.Get(false);
+}
+
+bool UWorldPartition::UseMakingVisibleTransactionRequests() const
+{
+	// Resolve once (we don't allow changing the state at runtime)
+	if (!bCachedUseMakingVisibleTransactionRequests.IsSet())
+	{
+		UWorld* OwningWorld = GetWorld();
+		bCachedUseMakingVisibleTransactionRequests = OwningWorld && OwningWorld->IsGameWorld() && UWorldPartition::bUseMakingVisibleTransactionRequests;
+	}
+	return bCachedUseMakingVisibleTransactionRequests.Get(false);
+}
+
+bool UWorldPartition::UseMakingInvisibleTransactionRequests() const
+{
+	// Resolve once (we don't allow changing the state at runtime)
+	if (!bCachedUseMakingInvisibleTransactionRequests.IsSet())
+	{
+		UWorld* OwningWorld = GetWorld();
+		bCachedUseMakingInvisibleTransactionRequests = OwningWorld && OwningWorld->IsGameWorld() && UWorldPartition::bUseMakingInvisibleTransactionRequests;
+	}
+	return bCachedUseMakingInvisibleTransactionRequests.Get(false);
 }
 
 bool UWorldPartition::IsSimulating()
@@ -1274,7 +1318,7 @@ void UWorldPartition::Tick(float DeltaSeconds)
 					AllActorsBounds += ActorDescIterator->GetBounds();
 
 					// Warn the user if the world becomes larger that 4km in any axis
-					if (AllActorsBounds.GetSize().GetMax() >= GWorldExtentToEnableStreaming)
+					if (AllActorsBounds.GetSize().GetMax() >= UWorldPartition::WorldExtentToEnableStreaming)
 					{
 						bEnablingStreamingJustified = true;
 						break;
@@ -1319,8 +1363,8 @@ bool UWorldPartition::CanDebugDraw() const
 	const ENetMode NetMode = GetWorld()->GetNetMode();
 	const bool bIsDedicatedServer = (NetMode == NM_DedicatedServer);
 	const bool bIsClient = (NetMode == NM_Client);
-	if ((bIsDedicatedServer && !GDebugDedicatedServerStreaming) ||
-		(bIsClient && GDebugDedicatedServerStreaming))
+	if ((bIsDedicatedServer && !UWorldPartition::DebugDedicatedServerStreaming) ||
+		(bIsClient && UWorldPartition::DebugDedicatedServerStreaming))
 	{
 		return false;
 	}
diff --git a/Engine/Source/Runtime/Engine/Private/WorldPartition/WorldPartitionLevelStreamingDynamic.cpp b/Engine/Source/Runtime/Engine/Private/WorldPartition/WorldPartitionLevelStreamingDynamic.cpp
--- a/Engine/Source/Runtime/Engine/Private/WorldPartition/WorldPartitionLevelStreamingDynamic.cpp
+++ b/Engine/Source/Runtime/Engine/Private/WorldPartition/WorldPartitionLevelStreamingDynamic.cpp
@@ -70,6 +70,8 @@ void UWorldPartitionLevelStreamingDynamic::Initialize(const UWorldPartitionRunti
 	IWorldPartitionRuntimeCellOwner* CellOwner = InCell.GetCellOwner();
 	OuterWorldPartition = CellOwner->GetOuterWorld()->GetWorldPartition();
 #endif
+
+	UpdateShouldSkipMakingVisibilityTransactionRequest();
 }
 
 #if WITH_EDITOR
@@ -616,5 +618,25 @@ UWorld* UWorldPartitionLevelStreamingDynamic::GetOuterWorld() const
 	return OuterWorldPartition->GetTypedOuter<UWorld>();
 }
 
+void UWorldPartitionLevelStreamingDynamic::UpdateShouldSkipMakingVisibilityTransactionRequest()
+{
+	// It is safe to skip client visibility transaction requests for cells without data layers when world partition server streaming is disabled
+	const UWorldPartitionRuntimeCell* Cell = GetWorldPartitionRuntimeCell();
+	if (ensure(Cell && OuterWorldPartition.IsValid()))
+	{
+		bSkipClientUseMakingVisibleTransactionRequest = bSkipClientUseMakingInvisibleTransactionRequest = !OuterWorldPartition->IsServerStreamingEnabled() && !Cell->HasDataLayers();
+	}
+}
+
+#if !WITH_EDITOR
+void UWorldPartitionLevelStreamingDynamic::PostLoad()
+{
+	Super::PostLoad();
+
+	// UWorldPartitionLevelStreamingDynamic::Initialize is not called at runtime (except for content bundles)
+	UpdateShouldSkipMakingVisibilityTransactionRequest();
+}
+#endif
+
 #undef LOCTEXT_NAMESPACE
 
diff --git a/Engine/Source/Runtime/Engine/Public/WorldPartition/WorldPartition.h b/Engine/Source/Runtime/Engine/Public/WorldPartition/WorldPartition.h
--- a/Engine/Source/Runtime/Engine/Public/WorldPartition/WorldPartition.h
+++ b/Engine/Source/Runtime/Engine/Public/WorldPartition/WorldPartition.h
@@ -229,6 +229,8 @@ public:
 	bool IsServer() const;
 	bool IsServerStreamingEnabled() const;
 	bool IsServerStreamingOutEnabled() const;
+	bool UseMakingVisibleTransactionRequests() const;
+	bool UseMakingInvisibleTransactionRequests() const;
 
 	bool IsMainWorldPartition() const;
 
@@ -276,9 +278,6 @@ private:
 #endif
 
 public:
-	static bool UseMakingVisibleTransactionRequests();
-	static bool UseMakingInvisibleTransactionRequests();
-
 	UActorDescContainer* GetActorDescContainer() const { return ActorDescContainer; }
 
 	UPROPERTY()
@@ -318,6 +317,11 @@ private:
 	EWorldPartitionInitState InitState;
 	TOptional<FTransform> InstanceTransform;
 
+	mutable TOptional<bool> bCachedUseMakingInvisibleTransactionRequests;
+	mutable TOptional<bool> bCachedUseMakingVisibleTransactionRequests;
+	mutable TOptional<bool> bCachedIsServerStreamingEnabled;
+	mutable TOptional<bool> bCachedIsServerStreamingOutEnabled;
+
 	UPROPERTY(Transient)
 	mutable TObjectPtr<UWorldPartitionStreamingPolicy> StreamingPolicy;
 
@@ -327,8 +331,21 @@ private:
 	FWorldPartitionReference WorldDataLayersActor;
 #endif
 
+#if WITH_EDITOR
+	static int32 LoadingRangeBugItGo;
+	static int32 WorldExtentToEnableStreaming;
+	static bool DebugDedicatedServerStreaming;
+	static FAutoConsoleVariableRef CVarLoadingRangeBugItGo;
+	static FAutoConsoleVariableRef CVarWorldExtentToEnableStreaming;
+	static FAutoConsoleVariableRef CVarDebugDedicatedServerStreaming;
+#endif
+
+	static int32 EnableServerStreaming;
+	static bool bEnableServerStreamingOut;
 	static bool bUseMakingVisibleTransactionRequests;
 	static bool bUseMakingInvisibleTransactionRequests;
+	static FAutoConsoleVariableRef CVarEnableServerStreaming;
+	static FAutoConsoleVariableRef CVarEnableServerStreamingOut;
 	static FAutoConsoleVariableRef CVarUseMakingVisibleTransactionRequests;
 	static FAutoConsoleVariableRef CVarUseMakingInvisibleTransactionRequests;
 
diff --git a/Engine/Source/Runtime/Engine/Public/WorldPartition/WorldPartitionLevelStreamingDynamic.h b/Engine/Source/Runtime/Engine/Public/WorldPartition/WorldPartitionLevelStreamingDynamic.h
--- a/Engine/Source/Runtime/Engine/Public/WorldPartition/WorldPartitionLevelStreamingDynamic.h
+++ b/Engine/Source/Runtime/Engine/Public/WorldPartition/WorldPartitionLevelStreamingDynamic.h
@@ -35,6 +35,9 @@ class ENGINE_API UWorldPartitionLevelStreamingDynamic : public ULevelStreamingDy
 
 	virtual bool ShouldBeAlwaysLoaded() const override { return bShouldBeAlwaysLoaded; }
 	virtual bool ShouldRequireFullVisibilityToRender() const override { return true; }
+#if !WITH_EDITOR
+	virtual void PostLoad();
+#endif
 
 	void Initialize(const UWorldPartitionRuntimeLevelStreamingCell& InCell);
 
@@ -72,6 +75,9 @@ private:
 	FWorldPartitionLevelHelper::FPackageReferencer PackageReferencer;
 #endif
 
+private:
+	void UpdateShouldSkipMakingVisibilityTransactionRequest();
+
 	UPROPERTY()
 	bool bShouldBeAlwaysLoaded;
 
