Prevent crashing in Content Bundle when starting with -game

diff --git a/Engine/Source/Runtime/Engine/Private/WorldPartition/ContentBundle/ContentBundle.cpp b/Engine/Source/Runtime/Engine/Private/WorldPartition/ContentBundle/ContentBundle.cpp
--- a/Engine/Source/Runtime/Engine/Private/WorldPartition/ContentBundle/ContentBundle.cpp
+++ b/Engine/Source/Runtime/Engine/Private/WorldPartition/ContentBundle/ContentBundle.cpp
@@ -118,11 +118,13 @@ bool FContentBundle::IsValid() const
 void FContentBundle::InitializeForPIE()
 {
 	UContentBundleManager* ContentBundleManager = GetInjectedWorld()->ContentBundleManager;
-	UContentBundleDuplicateForPIEHelper* PIEHelper = ContentBundleManager->GetPIEDuplicateHelper();
-	ExternalStreamingObject = PIEHelper->RetrieveContentBundleStreamingObject(*this);
-	if (ExternalStreamingObject == nullptr)
+	if (UContentBundleDuplicateForPIEHelper* PIEHelper = ContentBundleManager->GetPIEDuplicateHelper())
 	{
-		UE_LOG(LogContentBundle, Log, TEXT("[CB: %s] No streaming object found. There are %u existing streaming objects."), *GetDescriptor()->GetDisplayName(), PIEHelper->GetStreamingObjectCount());
+		ExternalStreamingObject = PIEHelper->RetrieveContentBundleStreamingObject(*this);
+		if (ExternalStreamingObject == nullptr)
+		{
+			UE_LOG(LogContentBundle, Log, TEXT("[CB: %s] No streaming object found. There are %u existing streaming objects."), *GetDescriptor()->GetDisplayName(), PIEHelper->GetStreamingObjectCount());
+		}
 	}
 }
 #endif
\ No newline at end of file
