World Partition - HLOD: Validate type of HLOD builder settings upon HLODLayer loading. Recreate if invalid
* This fixes an issue with older HLOD layer assets as UHLODBuilderInstancing now has it's own setting type.

diff --git a/Engine/Source/Runtime/Engine/Private/WorldPartition/HLOD/HLODLayer.cpp b/Engine/Source/Runtime/Engine/Private/WorldPartition/HLOD/HLODLayer.cpp
--- a/Engine/Source/Runtime/Engine/Private/WorldPartition/HLOD/HLODLayer.cpp
+++ b/Engine/Source/Runtime/Engine/Private/WorldPartition/HLOD/HLODLayer.cpp
@@ -37,6 +37,8 @@ UHLODLayer::UHLODLayer(const FObjectInitializer& ObjectInitializer)
 
 #if WITH_EDITOR
 
+
+
 UHLODLayer* UHLODLayer::GetHLODLayer(const AActor* InActor)
 {
 	if (UHLODLayer* HLODLayer = InActor->GetHLODLayer())
@@ -158,10 +160,13 @@ void UHLODLayer::PostLoad()
 {
 	Super::PostLoad();
 
-	if (HLODBuilderSettings == nullptr)
+	IWorldPartitionHLODUtilities* WPHLODUtilities = FModuleManager::Get().LoadModuleChecked<IWorldPartitionHLODUtilitiesModule>("WorldPartitionHLODUtilities").GetUtilities();
+	if (WPHLODUtilities)
 	{
-		IWorldPartitionHLODUtilities* WPHLODUtilities = FModuleManager::Get().LoadModuleChecked<IWorldPartitionHLODUtilitiesModule>("WorldPartitionHLODUtilities").GetUtilities();
-		if (WPHLODUtilities)
+		const UClass* BuilderClass = WPHLODUtilities->GetHLODBuilderClass(this);
+		const UClass* BuilderSettingsClass = BuilderClass ? BuilderClass->GetDefaultObject<UHLODBuilder>()->GetSettingsClass() : nullptr;
+
+		if (!HLODBuilderSettings || (BuilderSettingsClass && !HLODBuilderSettings->IsA(BuilderSettingsClass)))
 		{
 			HLODBuilderSettings = WPHLODUtilities->CreateHLODBuilderSettings(this);
 		}
