Added WorldPartitionSubsystem::IsAllStreamingCompleted() that returns true only if World Partition has finished loading/unloading/adding to world/removing from world any of its streaming level.
Also added WorldPartitionSubsystem::IsStreamingCompleted that receives an optional IWorldPartitionStreamingSourceProvider in which case it will only test streaming completion for this streaming source provider (used by UWorldPartitionStreamingSourceComponent::IsStreamingCompleted()).

diff --git a/Engine/Source/Runtime/Engine/Classes/Components/WorldPartitionStreamingSourceComponent.h b/Engine/Source/Runtime/Engine/Classes/Components/WorldPartitionStreamingSourceComponent.h
--- a/Engine/Source/Runtime/Engine/Classes/Components/WorldPartitionStreamingSourceComponent.h
+++ b/Engine/Source/Runtime/Engine/Classes/Components/WorldPartitionStreamingSourceComponent.h
@@ -33,7 +33,7 @@ class ENGINE_API UWorldPartitionStreamingSourceComponent : public UActorComponen
 	bool IsStreamingSourceEnabled() const { return bStreamingSourceEnabled; }
 
 	// IWorldPartitionStreamingSourceProvider interface
-	virtual bool GetStreamingSource(FWorldPartitionStreamingSource& OutStreamingSource) override;
+	virtual bool GetStreamingSource(FWorldPartitionStreamingSource& OutStreamingSource) const override;
 
 	/** Returns true if streaming is completed for this streaming source component. */
 	UFUNCTION(BlueprintCallable, Category = "Streaming")
diff --git a/Engine/Source/Runtime/Engine/Private/Components/WorldPartitionStreamingSourceComponent.cpp b/Engine/Source/Runtime/Engine/Private/Components/WorldPartitionStreamingSourceComponent.cpp
--- a/Engine/Source/Runtime/Engine/Private/Components/WorldPartitionStreamingSourceComponent.cpp
+++ b/Engine/Source/Runtime/Engine/Private/Components/WorldPartitionStreamingSourceComponent.cpp
@@ -64,7 +64,7 @@ void UWorldPartitionStreamingSourceComponent::OnUnregister()
 	verify(WorldPartitionSubsystem->UnregisterStreamingSourceProvider(this));
 }
 
-bool UWorldPartitionStreamingSourceComponent::GetStreamingSource(FWorldPartitionStreamingSource& OutStreamingSource)
+bool UWorldPartitionStreamingSourceComponent::GetStreamingSource(FWorldPartitionStreamingSource& OutStreamingSource) const
 {
 	if (bStreamingSourceEnabled)
 	{
@@ -86,36 +86,11 @@ bool UWorldPartitionStreamingSourceComponent::GetStreamingSource(FWorldPartition
 
 bool UWorldPartitionStreamingSourceComponent::IsStreamingCompleted() const
 {
-	UWorld* World = GetWorld();
-	if (!bStreamingSourceEnabled || !World->IsGameWorld())
-	{
-		return false;
-	}
-	
-	UWorldPartitionSubsystem* WorldPartitionSubsystem = World->GetSubsystem<UWorldPartitionSubsystem>();
-	UDataLayerSubsystem* DataLayerSubsystem = World->GetSubsystem<UDataLayerSubsystem>();
-	if (!WorldPartitionSubsystem || !DataLayerSubsystem)	
+	if (UWorldPartitionSubsystem* WorldPartitionSubsystem = UWorld::GetSubsystem<UWorldPartitionSubsystem>(GetWorld()))
 	{
-		return false;
+		return WorldPartitionSubsystem->IsStreamingCompleted(this);
 	}
-
-	// Build a query source
-	AActor* Actor = GetOwner();
-	TArray<FWorldPartitionStreamingQuerySource> QuerySources;
-	FWorldPartitionStreamingQuerySource& QuerySource = QuerySources.Emplace_GetRef();
-	QuerySource.bSpatialQuery = true;
-	QuerySource.Location = Actor->GetActorLocation();
-	QuerySource.Rotation = Actor->GetActorRotation();
-	QuerySource.TargetGrid = TargetGrid;
-	QuerySource.Shapes = Shapes;
-	QuerySource.bUseGridLoadingRange = true;
-	QuerySource.Radius = 0.f;
-	QuerySource.bDataLayersOnly = false;
-	QuerySource.DataLayers = (TargetState == EStreamingSourceTargetState::Loaded) ? DataLayerSubsystem->GetEffectiveLoadedDataLayerNames().Array() : DataLayerSubsystem->GetEffectiveActiveDataLayerNames().Array();
-
-	// Execute query
-	const EWorldPartitionRuntimeCellState QueryState = (TargetState == EStreamingSourceTargetState::Loaded) ? EWorldPartitionRuntimeCellState::Loaded : EWorldPartitionRuntimeCellState::Activated;
-	return WorldPartitionSubsystem->IsStreamingCompleted(QueryState, QuerySources, /*bExactState*/ true);
+	return true;
 }
 
 void UWorldPartitionStreamingSourceComponent::DrawVisualization(const FSceneView* View, FPrimitiveDrawInterface* PDI) const
diff --git a/Engine/Source/Runtime/Engine/Private/WorldPartition/WorldPartition.cpp b/Engine/Source/Runtime/Engine/Private/WorldPartition/WorldPartition.cpp
--- a/Engine/Source/Runtime/Engine/Private/WorldPartition/WorldPartition.cpp
+++ b/Engine/Source/Runtime/Engine/Private/WorldPartition/WorldPartition.cpp
@@ -1347,6 +1347,15 @@ bool UWorldPartition::CanAddLoadedLevelToWorld(class ULevel* InLevel) const
 	return true;
 }
 
+bool UWorldPartition::IsStreamingCompleted(const FWorldPartitionStreamingSource* InStreamingSource) const
+{
+	if (GetWorld()->IsGameWorld())
+	{
+		return StreamingPolicy->IsStreamingCompleted(InStreamingSource);
+	}
+	return true;
+}
+
 bool UWorldPartition::IsStreamingCompleted(EWorldPartitionRuntimeCellState QueryState, const TArray<FWorldPartitionStreamingQuerySource>& QuerySources, bool bExactState) const
 {
 	if (GetWorld()->IsGameWorld())
@@ -1354,7 +1363,7 @@ bool UWorldPartition::IsStreamingCompleted(EWorldPartitionRuntimeCellState Query
 		return StreamingPolicy->IsStreamingCompleted(QueryState, QuerySources, bExactState);
 	}
 
-	return false;
+	return true;
 }
 
 bool UWorldPartition::CanDebugDraw() const
diff --git a/Engine/Source/Runtime/Engine/Private/WorldPartition/WorldPartitionLevelStreamingPolicy.cpp b/Engine/Source/Runtime/Engine/Private/WorldPartition/WorldPartitionLevelStreamingPolicy.cpp
--- a/Engine/Source/Runtime/Engine/Private/WorldPartition/WorldPartitionLevelStreamingPolicy.cpp
+++ b/Engine/Source/Runtime/Engine/Private/WorldPartition/WorldPartitionLevelStreamingPolicy.cpp
@@ -51,6 +51,33 @@ void UWorldPartitionLevelStreamingPolicy::ForEachActiveRuntimeCell(TFunctionRef<
 	}
 }
 
+bool UWorldPartitionLevelStreamingPolicy::IsStreamingCompleted(const FWorldPartitionStreamingSource* InStreamingSource) const
+{
+	const UWorld* World = GetWorld();
+	check(World);
+	check(World->IsGameWorld());
+
+	if (!Super::IsStreamingCompleted(InStreamingSource))
+	{
+		return false;
+	}
+
+	if (!InStreamingSource)
+	{
+		// Also verify that there's no remaining activity (mainly for unloading) 
+		// waiting to be processed on level streaming of world partition runtime cells
+		for (ULevelStreaming* StreamingLevel : World->GetStreamingLevels())
+		{
+			ULevel* Level = StreamingLevel ? StreamingLevel->GetLoadedLevel() : nullptr;
+			if (Level && Level->IsWorldPartitionRuntimeCell() && StreamingLevel->IsStreamingStatePending())
+			{
+				return false;
+			}
+		}
+	}
+	return true;
+}
+
 #if WITH_EDITOR
 
 FString UWorldPartitionLevelStreamingPolicy::GetCellPackagePath(const FName& InCellName, const UWorld* InWorld)
diff --git a/Engine/Source/Runtime/Engine/Private/WorldPartition/WorldPartitionStreamingPolicy.cpp b/Engine/Source/Runtime/Engine/Private/WorldPartition/WorldPartitionStreamingPolicy.cpp
--- a/Engine/Source/Runtime/Engine/Private/WorldPartition/WorldPartitionStreamingPolicy.cpp
+++ b/Engine/Source/Runtime/Engine/Private/WorldPartition/WorldPartitionStreamingPolicy.cpp
@@ -226,7 +226,8 @@ void UWorldPartitionStreamingPolicy::UpdateStreamingState()
 	TRACE_CPUPROFILER_EVENT_SCOPE(UWorldPartitionStreamingPolicy::UpdateStreamingState);
 
 	UWorld* World = GetWorld();
-	check(World && World->IsGameWorld());
+	check(World);
+	check(World->IsGameWorld());
 
 	// Dermine if the World's BlockTillLevelStreamingCompleted was triggered by WorldPartitionStreamingPolicy
 	if (bCriticalPerformanceRequestedBlockTillOnWorld && IsInBlockTillLevelStreamingCompleted())
@@ -757,6 +758,67 @@ bool UWorldPartitionStreamingPolicy::CanAddLoadedLevelToWorld(ULevel* InLevel) c
 	return true;
 }
 
+bool UWorldPartitionStreamingPolicy::IsStreamingCompleted(const FWorldPartitionStreamingSource* InStreamingSource) const
+{
+	const UWorld* World = GetWorld();
+	check(World);
+	check(World->IsGameWorld());
+	const UDataLayerSubsystem* DataLayerSubsystem = World->GetSubsystem<UDataLayerSubsystem>();
+	const bool bTestProvidedStreamingSource = !!InStreamingSource;
+
+	// Always test non-spatial cells
+	{
+		// Test non-data layer and activated data layers
+		TArray<FWorldPartitionStreamingQuerySource> QuerySources;
+		FWorldPartitionStreamingQuerySource& QuerySource = QuerySources.Emplace_GetRef();
+		QuerySource.bSpatialQuery = false;
+		QuerySource.bDataLayersOnly = false;
+		QuerySource.DataLayers = DataLayerSubsystem->GetEffectiveActiveDataLayerNames().Array();
+		if (!IsStreamingCompleted(EWorldPartitionRuntimeCellState::Activated, QuerySources, true))
+		{
+			return false;
+		}
+
+		// Test only loaded data layers
+		if (!DataLayerSubsystem->GetEffectiveLoadedDataLayerNames().IsEmpty())
+		{
+			QuerySource.bDataLayersOnly = true;
+			QuerySource.DataLayers = DataLayerSubsystem->GetEffectiveLoadedDataLayerNames().Array();
+			if (!IsStreamingCompleted(EWorldPartitionRuntimeCellState::Loaded, QuerySources, true))
+			{
+				return false;
+			}
+		}
+	}
+
+	// Test spatially loaded cells using streaming sources (or provided streaming source)
+	TArrayView<const FWorldPartitionStreamingSource> QueriedStreamingSources = bTestProvidedStreamingSource ? MakeArrayView({ *InStreamingSource }) : StreamingSources;
+	for (const FWorldPartitionStreamingSource& StreamingSource : QueriedStreamingSources)
+	{
+		// Build a query source from a Streaming Source
+		TArray<FWorldPartitionStreamingQuerySource> QuerySources;
+		FWorldPartitionStreamingQuerySource& QuerySource = QuerySources.Emplace_GetRef();
+		QuerySource.bSpatialQuery = true;
+		QuerySource.Location = StreamingSource.Location;
+		QuerySource.Rotation = StreamingSource.Rotation;
+		QuerySource.TargetGrid = StreamingSource.TargetGrid;
+		QuerySource.Shapes = StreamingSource.Shapes;
+		QuerySource.bUseGridLoadingRange = true;
+		QuerySource.Radius = 0.f;
+		QuerySource.bDataLayersOnly = false;
+		QuerySource.DataLayers = (StreamingSource.TargetState == EStreamingSourceTargetState::Loaded) ? DataLayerSubsystem->GetEffectiveLoadedDataLayerNames().Array() : DataLayerSubsystem->GetEffectiveActiveDataLayerNames().Array();
+
+		// Execute query
+		const EWorldPartitionRuntimeCellState QueryState = (StreamingSource.TargetState == EStreamingSourceTargetState::Loaded) ? EWorldPartitionRuntimeCellState::Loaded : EWorldPartitionRuntimeCellState::Activated;
+		if (!IsStreamingCompleted(QueryState, QuerySources, true))
+		{
+			return false;
+		}
+	}
+
+	return true;
+}
+
 bool UWorldPartitionStreamingPolicy::IsStreamingCompleted(EWorldPartitionRuntimeCellState QueryState, const TArray<FWorldPartitionStreamingQuerySource>& QuerySources, bool bExactState) const
 {
 	const UDataLayerSubsystem* DataLayerSubsystem = GetWorld()->GetSubsystem<UDataLayerSubsystem>();
diff --git a/Engine/Source/Runtime/Engine/Private/WorldPartition/WorldPartitionSubsystem.cpp b/Engine/Source/Runtime/Engine/Private/WorldPartition/WorldPartitionSubsystem.cpp
--- a/Engine/Source/Runtime/Engine/Private/WorldPartition/WorldPartitionSubsystem.cpp
+++ b/Engine/Source/Runtime/Engine/Private/WorldPartition/WorldPartitionSubsystem.cpp
@@ -261,6 +261,36 @@ TStatId UWorldPartitionSubsystem::GetStatId() const
 	RETURN_QUICK_DECLARE_CYCLE_STAT(UWorldPartitionSubsystem, STATGROUP_Tickables);
 }
 
+bool UWorldPartitionSubsystem::IsAllStreamingCompleted()
+{
+	return const_cast<UWorldPartitionSubsystem*>(this)->IsStreamingCompleted();
+}
+
+bool UWorldPartitionSubsystem::IsStreamingCompleted(const IWorldPartitionStreamingSourceProvider* InStreamingSourceProvider) const
+{
+	// Convert specified/optional streaming source provider to a world partition 
+	// streaming source and pass it along to each registered world partition
+	FWorldPartitionStreamingSource StreamingSource;
+	FWorldPartitionStreamingSource* StreamingSourcePtr = nullptr;
+	if (InStreamingSourceProvider)
+	{
+		StreamingSourcePtr = &StreamingSource;
+		if (!InStreamingSourceProvider->GetStreamingSource(StreamingSource))
+		{
+			return true;
+		}
+	}
+
+	for (UWorldPartition* RegisteredWorldPartition : RegisteredWorldPartitions)
+	{
+		if (!RegisteredWorldPartition->IsStreamingCompleted(StreamingSourcePtr))
+		{
+			return false;
+		}
+	}
+	return true;
+}
+
 bool UWorldPartitionSubsystem::IsStreamingCompleted(EWorldPartitionRuntimeCellState QueryState, const TArray<FWorldPartitionStreamingQuerySource>& QuerySources, bool bExactState) const
 {
 	for (UWorldPartition* RegisteredWorldPartition : RegisteredWorldPartitions)
diff --git a/Engine/Source/Runtime/Engine/Public/WorldPartition/WorldPartition.h b/Engine/Source/Runtime/Engine/Public/WorldPartition/WorldPartition.h
--- a/Engine/Source/Runtime/Engine/Public/WorldPartition/WorldPartition.h
+++ b/Engine/Source/Runtime/Engine/Public/WorldPartition/WorldPartition.h
@@ -237,6 +237,7 @@ public:
 	void Tick(float DeltaSeconds);
 	void UpdateStreamingState();
 	bool CanAddLoadedLevelToWorld(ULevel* InLevel) const;
+	bool IsStreamingCompleted(const FWorldPartitionStreamingSource* InStreamingSource) const;
 	bool IsStreamingCompleted(EWorldPartitionRuntimeCellState QueryState, const TArray<FWorldPartitionStreamingQuerySource>& QuerySources, bool bExactState) const;
 
 	const TArray<FWorldPartitionStreamingSource>& GetStreamingSources() const;
diff --git a/Engine/Source/Runtime/Engine/Public/WorldPartition/WorldPartitionLevelStreamingPolicy.h b/Engine/Source/Runtime/Engine/Public/WorldPartition/WorldPartitionLevelStreamingPolicy.h
--- a/Engine/Source/Runtime/Engine/Public/WorldPartition/WorldPartitionLevelStreamingPolicy.h
+++ b/Engine/Source/Runtime/Engine/Public/WorldPartition/WorldPartitionLevelStreamingPolicy.h
@@ -24,6 +24,7 @@ class UWorldPartitionLevelStreamingPolicy : public UWorldPartitionStreamingPolic
 public:
 	virtual void DrawRuntimeCellsDetails(class UCanvas* Canvas, FVector2D& Offset) override;
 	virtual void DrawStreamingStatusLegend(UCanvas* Canvas, FVector2D& Offset) override;
+	virtual bool IsStreamingCompleted(const FWorldPartitionStreamingSource* InStreamingSource) const override;
 
 #if WITH_EDITOR
 	virtual TSubclassOf<class UWorldPartitionRuntimeCell> GetRuntimeCellClass() const override;
diff --git a/Engine/Source/Runtime/Engine/Public/WorldPartition/WorldPartitionStreamingPolicy.h b/Engine/Source/Runtime/Engine/Public/WorldPartition/WorldPartitionStreamingPolicy.h
--- a/Engine/Source/Runtime/Engine/Public/WorldPartition/WorldPartitionStreamingPolicy.h
+++ b/Engine/Source/Runtime/Engine/Public/WorldPartition/WorldPartitionStreamingPolicy.h
@@ -50,6 +50,7 @@ public:
 	virtual void DrawRuntimeCellsDetails(class UCanvas* Canvas, FVector2D& Offset) {}
 	virtual void DrawStreamingStatusLegend(class UCanvas* Canvas, FVector2D& Offset) {}
 
+	virtual bool IsStreamingCompleted(const FWorldPartitionStreamingSource* InStreamingSource) const;
 	virtual bool IsStreamingCompleted(EWorldPartitionRuntimeCellState QueryState, const TArray<FWorldPartitionStreamingQuerySource>& QuerySources, bool bExactState = true) const;
 
 #if WITH_EDITOR
diff --git a/Engine/Source/Runtime/Engine/Public/WorldPartition/WorldPartitionStreamingSource.h b/Engine/Source/Runtime/Engine/Public/WorldPartition/WorldPartitionStreamingSource.h
--- a/Engine/Source/Runtime/Engine/Public/WorldPartition/WorldPartitionStreamingSource.h
+++ b/Engine/Source/Runtime/Engine/Public/WorldPartition/WorldPartitionStreamingSource.h
@@ -433,5 +433,5 @@ struct ENGINE_API FWorldPartitionStreamingSource
  */
 struct ENGINE_API IWorldPartitionStreamingSourceProvider
 {
-	virtual bool GetStreamingSource(FWorldPartitionStreamingSource& StreamingSource) = 0;
+	virtual bool GetStreamingSource(FWorldPartitionStreamingSource& StreamingSource) const = 0;
 };
\ No newline at end of file
diff --git a/Engine/Source/Runtime/Engine/Public/WorldPartition/WorldPartitionSubsystem.h b/Engine/Source/Runtime/Engine/Public/WorldPartition/WorldPartitionSubsystem.h
--- a/Engine/Source/Runtime/Engine/Public/WorldPartition/WorldPartitionSubsystem.h
+++ b/Engine/Source/Runtime/Engine/Public/WorldPartition/WorldPartitionSubsystem.h
@@ -41,6 +41,16 @@ public:
 	UFUNCTION(BlueprintCallable, Category = Streaming)
 	bool IsStreamingCompleted(EWorldPartitionRuntimeCellState QueryState, const TArray<FWorldPartitionStreamingQuerySource>& QuerySources, bool bExactState) const;
 
+	/** Returns true if world partition is done streaming levels, adding them to the world or removing them from the world. */
+	UFUNCTION(BlueprintCallable, Category = Streaming)
+	bool IsAllStreamingCompleted();
+
+	/*
+	 * Returns true if world partition is done streaming levels, adding them to the world or removing them from the world. 
+	 * When provided, the test is reduced to streaming levels affected by the optional streaming source provider.
+	 */
+	bool IsStreamingCompleted(const IWorldPartitionStreamingSourceProvider* InStreamingSourceProvider = nullptr) const;
+
 	void DumpStreamingSources(FOutputDevice& OutputDevice) const;
 
 	TSet<IWorldPartitionStreamingSourceProvider*> GetStreamingSourceProviders() const { return StreamingSourceProviders; }
