Fix level being dirtied on PIE

#test fastcook, pie

diff --git a/Engine/Source/Runtime/Engine/Private/WorldPartition/WorldPartitionRuntimeSpatialHash.cpp b/Engine/Source/Runtime/Engine/Private/WorldPartition/WorldPartitionRuntimeSpatialHash.cpp
--- a/Engine/Source/Runtime/Engine/Private/WorldPartition/WorldPartitionRuntimeSpatialHash.cpp
+++ b/Engine/Source/Runtime/Engine/Private/WorldPartition/WorldPartitionRuntimeSpatialHash.cpp
@@ -859,7 +859,10 @@ URuntimeHashExternalStreamingObjectBase* UWorldPartitionRuntimeSpatialHash::Stor
 				for (UWorldPartitionRuntimeSpatialHashCell* Cell : GridLayerCell.GridCells)
 				{
 					Cell->SetGridName(StreamingGrid.GridName);
-					Cell->Rename(nullptr, StreamingObject);
+
+					// Do not dirty, otherwise it dirties the level previous outer (WorldPartition) which dirties the map. Occurs when entering PIE. 
+					// Do not reset loaders, the cell was just created it did not exists when loading initially.
+					Cell->Rename(nullptr, StreamingObject,  REN_DoNotDirty | REN_ForceNoResetLoaders);
 				}
 			}
 		}
