Fix manual injection not working after a plugin has been deactivated

diff --git a/Engine/Source/Runtime/Engine/Private/WorldPartition/ContentBundle/ContentBundleClient.cpp b/Engine/Source/Runtime/Engine/Private/WorldPartition/ContentBundle/ContentBundleClient.cpp
--- a/Engine/Source/Runtime/Engine/Private/WorldPartition/ContentBundle/ContentBundleClient.cpp
+++ b/Engine/Source/Runtime/Engine/Private/WorldPartition/ContentBundle/ContentBundleClient.cpp
@@ -58,7 +58,19 @@ void FContentBundleClient::RequestUnregister()
 		UE_LOG(LogContentBundle, Log, TEXT("[CB: %s] Client failed to request unregister. Its state is %s"), *GetDescriptor()->GetDisplayName(), *UEnum::GetDisplayValueAsText(State).ToString());
 	}
 }
-	
+
+bool FContentBundleClient::HasInjectedAnyContent() const
+{
+	for (auto& WorldState : WorldContentStates)
+	{
+		if (WorldState.Value != EWorldContentState::NoContent)
+		{
+			return true;
+		}
+	}
+
+	return false;
+}
 
 void FContentBundleClient::OnContentInjectedInWorld(EContentBundleStatus InjectionStatus, UWorld* InjectedWorld)
 {
@@ -84,6 +96,13 @@ void FContentBundleClient::OnContentRemovedFromWorld(EContentBundleStatus Remova
 	if (RemovalStatus == EContentBundleStatus::Registered)
 	{
 		SetWorldContentState(InjectedWorld, EWorldContentState::NoContent);
+
+		// If the client request content removal and removed everything set the client state as registered.
+		// If the client did not request content removal but removed everything keep the client state. We just changed levels. The next level will have its content bundle injected.
+		if (GetState() == EContentBundleClientState::ContentRemovalRequested && !HasInjectedAnyContent())
+		{
+			SetState(EContentBundleClientState::Registered);
+		}
 	}
 	else
 	{
diff --git a/Engine/Source/Runtime/Engine/Public/WorldPartition/ContentBundle/ContentBundleClient.h b/Engine/Source/Runtime/Engine/Public/WorldPartition/ContentBundle/ContentBundleClient.h
--- a/Engine/Source/Runtime/Engine/Public/WorldPartition/ContentBundle/ContentBundleClient.h
+++ b/Engine/Source/Runtime/Engine/Public/WorldPartition/ContentBundle/ContentBundleClient.h
@@ -50,6 +50,8 @@ public:
 
 	EContentBundleClientState GetState() const { return State; }
 
+	bool HasInjectedAnyContent() const;
+
 	FString const& GetDisplayName() const { return DisplayName; }
 
 private:
