Added cpu scope to ValidateReplicatedSubObjects function

diff --git a/Engine/Source/Runtime/Engine/Private/DataChannel.cpp b/Engine/Source/Runtime/Engine/Private/DataChannel.cpp
--- a/Engine/Source/Runtime/Engine/Private/DataChannel.cpp
+++ b/Engine/Source/Runtime/Engine/Private/DataChannel.cpp
@@ -3527,6 +3527,7 @@ bool UActorChannel::DoSubObjectReplication(FOutBunch& Bunch, FReplicationFlags&
 #if SUBOBJECT_TRANSITION_VALIDATION
 		if (UE::Net::GCVarCompareSubObjectsReplicated)
 		{
+			QUICK_SCOPE_CYCLE_COUNTER(STAT_Validate_ReplicatedSubObjects);
 			ValidateReplicatedSubObjects();
 
 			DataChannelInternal::bTrackReplicatedSubObjects = false;
