World Partition: fix crash when disabling streaming on a partitioned map and goint in PIE.

#rnx

diff --git a/Engine/Source/Runtime/Engine/Private/WorldPartition/WorldPartitionRuntimeLevelStreamingCell.cpp b/Engine/Source/Runtime/Engine/Private/WorldPartition/WorldPartitionRuntimeLevelStreamingCell.cpp
--- a/Engine/Source/Runtime/Engine/Private/WorldPartition/WorldPartitionRuntimeLevelStreamingCell.cpp
+++ b/Engine/Source/Runtime/Engine/Private/WorldPartition/WorldPartitionRuntimeLevelStreamingCell.cpp
@@ -375,9 +375,15 @@ bool UWorldPartitionRuntimeLevelStreamingCell::CanUnload() const
 	// Events to the HLodSubsystem which assumes knowledge of all cells (not true with plugins)
 	if (LevelStreaming && !GetContentBundleID().IsValid())
 	{
-		if (UHLODSubsystem* HLODSubsystem = LevelStreaming->GetWorld()->GetSubsystem<UHLODSubsystem>())
+		const UWorldPartition* WorldPartition = LevelStreaming->GetWorld()->GetWorldPartition();
+		check(WorldPartition);
+
+		if (WorldPartition->IsStreamingEnabled())
 		{
-			return HLODSubsystem->RequestUnloading(this);
+			if (UHLODSubsystem* HLODSubsystem = LevelStreaming->GetWorld()->GetSubsystem<UHLODSubsystem>())
+			{
+				return HLODSubsystem->RequestUnloading(this);
+			}
 		}
 	}
 
