Parse enable base submix toggle for preview MetaSounds

diff --git a/Engine/Source/Runtime/Engine/Private/SoundWave.cpp b/Engine/Source/Runtime/Engine/Private/SoundWave.cpp
--- a/Engine/Source/Runtime/Engine/Private/SoundWave.cpp
+++ b/Engine/Source/Runtime/Engine/Private/SoundWave.cpp
@@ -3037,8 +3037,12 @@ void USoundWave::Parse(FAudioDevice* AudioDevice, const UPTRINT NodeWaveInstance
 		if (EnableSubmixSendsOnPreview())
 		{
 			WaveInstance->bEnableSubmixSends = ParseParams.bEnableSubmixSends;
+			WaveInstance->bEnableBaseSubmix = ParseParams.bEnableBaseSubmix;
+		}
+		else
+		{
+			WaveInstance->bEnableBaseSubmix = true;
 		}
-		WaveInstance->bEnableBaseSubmix = true;
 		WaveInstance->bEnableBusSends = ParseParams.bEnableBusSends;
 	}
 
