Don't allow running wp.Runtime.SetDataLayerRuntimeStateD and wp.Runtime.ToggleDataLayerActivation on clients.

diff --git a/Engine/Source/Runtime/Engine/Private/WorldPartition/DataLayer/DataLayerSubsystem.cpp b/Engine/Source/Runtime/Engine/Private/WorldPartition/DataLayer/DataLayerSubsystem.cpp
--- a/Engine/Source/Runtime/Engine/Private/WorldPartition/DataLayer/DataLayerSubsystem.cpp
+++ b/Engine/Source/Runtime/Engine/Private/WorldPartition/DataLayer/DataLayerSubsystem.cpp
@@ -595,7 +595,7 @@ FAutoConsoleCommand UDataLayerSubsystem::ToggleDataLayerActivation(
 		for (const FWorldContext& Context : GEngine->GetWorldContexts())
 		{
 			UWorld* World = Context.World();
-			if (World && World->IsGameWorld())
+			if (World && World->IsGameWorld() && !World->IsNetMode(NM_Client))
 			{
 				if (UDataLayerSubsystem* DataLayerSubsystem = World->GetSubsystem<UDataLayerSubsystem>())
 				{
@@ -627,14 +627,14 @@ FAutoConsoleCommand UDataLayerSubsystem::SetDataLayerRuntimeStateCommand(
 		EDataLayerRuntimeState State;
 		if (!GetDataLayerRuntimeStateFromName(StatetStr, State))
 		{
-			UE_LOG(LogWorldPartition, Warning, TEXT("wp.Runtime.SetDataLayerRuntimeState : Invalid first argument, expencted one of these values : Unloaded, Loaded, Activated."));
+			UE_LOG(LogWorldPartition, Warning, TEXT("wp.Runtime.SetDataLayerRuntimeState : Invalid first argument, expected one of these values : Unloaded, Loaded, Activated."));
 			return;
 		}
 
 		for (const FWorldContext& Context : GEngine->GetWorldContexts())
 		{
 			UWorld* World = Context.World();
-			if (World && World->IsGameWorld())
+			if (World && World->IsGameWorld() && !World->IsNetMode(NM_Client))
 			{
 				if (UDataLayerSubsystem* DataLayerSubsystem = World->GetSubsystem<UDataLayerSubsystem>())
 				{
