Fix bug where non-sphere attenuation on MetaSounds does not virtualize. Also a cvar crash fix

diff --git a/Engine/Source/Runtime/Engine/Private/AudioDebug.cpp b/Engine/Source/Runtime/Engine/Private/AudioDebug.cpp
--- a/Engine/Source/Runtime/Engine/Private/AudioDebug.cpp
+++ b/Engine/Source/Runtime/Engine/Private/AudioDebug.cpp
@@ -730,7 +730,7 @@ FAutoConsoleCommandWithWorldAndArgs CVarAudioVisualizeAttenuation(
 	TEXT("0: Not Enabled, 1: Enabled"),
 	FConsoleCommandWithWorldAndArgsDelegate::CreateStatic([](const TArray<FString>& Args, UWorld* InWorld)
 	{
-		if (Args.Num() < 0)
+		if (Args.Num() <= 0)
 		{
 			return;
 		}
diff --git a/Engine/Source/Runtime/Engine/Private/SoundWave.cpp b/Engine/Source/Runtime/Engine/Private/SoundWave.cpp
--- a/Engine/Source/Runtime/Engine/Private/SoundWave.cpp
+++ b/Engine/Source/Runtime/Engine/Private/SoundWave.cpp
@@ -2775,7 +2775,7 @@ void USoundWave::Parse(FAudioDevice* AudioDevice, const UPTRINT NodeWaveInstance
 	}
 
 	// Looping sounds are never actually finished
-	if (bLooping || ParseParams.bLooping)
+	if (!IsOneShot() || ParseParams.bLooping)
 	{
 		WaveInstance->bIsFinished = false;
 
@@ -2924,7 +2924,7 @@ void USoundWave::Parse(FAudioDevice* AudioDevice, const UPTRINT NodeWaveInstance
 
 	WaveInstance->WaveData = this;
 	WaveInstance->NotifyBufferFinishedHooks = ParseParams.NotifyBufferFinishedHooks;
-	WaveInstance->LoopingMode = ((bLooping || ParseParams.bLooping) ? LOOP_Forever : LOOP_Never);
+	WaveInstance->LoopingMode = ((!IsOneShot() || ParseParams.bLooping) ? LOOP_Forever : LOOP_Never);
 	WaveInstance->bIsPaused = ParseParams.bIsPaused;
 
 	// If we're normalizing 3d stereo spatialized sounds, we need to scale by -6 dB
