Fix static analysis issues

diff --git a/Engine/Source/Runtime/Engine/Classes/GameFramework/Actor.h b/Engine/Source/Runtime/Engine/Classes/GameFramework/Actor.h
--- a/Engine/Source/Runtime/Engine/Classes/GameFramework/Actor.h
+++ b/Engine/Source/Runtime/Engine/Classes/GameFramework/Actor.h
@@ -2327,8 +2327,8 @@ public:
 	virtual bool ShouldExport() { return true; }
 
 	/** Called before editor paste, true allow import */
-	UE_DEPRECATED(5.1, "Use the override that take a StringView instead instead")
-	virtual bool ShouldImport(FString* ActorPropString, bool IsMovingLevel) { return ActorPropString ? ShouldImport(FStringView(), IsMovingLevel) : ShouldImport(FStringView(*ActorPropString), IsMovingLevel); }
+	UE_DEPRECATED(5.2, "Use the override that takes a StringView instead")
+	virtual bool ShouldImport(FString* ActorPropString, bool IsMovingLevel) { return ActorPropString ? ShouldImport(FStringView(*ActorPropString), IsMovingLevel) : ShouldImport(FStringView(), IsMovingLevel); }
 
 	/** Called before editor paste, true allow import */
 	virtual bool ShouldImport(FStringView ActorPropString, bool IsMovingLevel) { return true; }
diff --git a/Engine/Source/Runtime/Engine/Private/Elements/Framework/TypedElementCommonActions.cpp b/Engine/Source/Runtime/Engine/Private/Elements/Framework/TypedElementCommonActions.cpp
--- a/Engine/Source/Runtime/Engine/Private/Elements/Framework/TypedElementCommonActions.cpp
+++ b/Engine/Source/Runtime/Engine/Private/Elements/Framework/TypedElementCommonActions.cpp
@@ -245,7 +245,10 @@ TArray<FTypedElementHandle> UTypedElementCommonActions::PasteNormalizedElements(
 			{
 				if (Line.StartsWith(LookingForEndTag))
 				{
-					(*CurrentBlock) = MakeStringView(StartOfBlock, int32((Line.GetData() + Line.Len()) - StartOfBlock));
+					if (CurrentBlock)
+					{
+						(*CurrentBlock) = MakeStringView(StartOfBlock, int32((Line.GetData() + Line.Len()) - StartOfBlock));
+					}
 					LookingForEndTag.Reset();
 				}
 			}
