DrawCentripetalCatmullRomSpline - increasing tolerance to UE_KINDA_SMALL_NUMBER for checking small floats to avoid numerical instability

diff --git a/Engine/Source/Runtime/Engine/Private/DrawDebugHelpers.cpp b/Engine/Source/Runtime/Engine/Private/DrawDebugHelpers.cpp
--- a/Engine/Source/Runtime/Engine/Private/DrawDebugHelpers.cpp
+++ b/Engine/Source/Runtime/Engine/Private/DrawDebugHelpers.cpp
@@ -1094,11 +1094,11 @@ void DrawCentripetalCatmullRomSpline(const UWorld* InWorld, TConstArrayView<FVec
 			const float T2T0 = T2 - T0;
 			const float T3T1 = T3 - T1;
 
-			const bool bIsNearlyZeroT1T0 = FMath::IsNearlyZero(T1T0);
-			const bool bIsNearlyZeroT2T1 = FMath::IsNearlyZero(T2T1);
-			const bool bIsNearlyZeroT3T2 = FMath::IsNearlyZero(T3T2);
-			const bool bIsNearlyZeroT2T0 = FMath::IsNearlyZero(T2T0);
-			const bool bIsNearlyZeroT3T1 = FMath::IsNearlyZero(T3T1);
+			const bool bIsNearlyZeroT1T0 = FMath::IsNearlyZero(T1T0, UE_KINDA_SMALL_NUMBER);
+			const bool bIsNearlyZeroT2T1 = FMath::IsNearlyZero(T2T1, UE_KINDA_SMALL_NUMBER);
+			const bool bIsNearlyZeroT3T2 = FMath::IsNearlyZero(T3T2, UE_KINDA_SMALL_NUMBER);
+			const bool bIsNearlyZeroT2T0 = FMath::IsNearlyZero(T2T0, UE_KINDA_SMALL_NUMBER);
+			const bool bIsNearlyZeroT3T1 = FMath::IsNearlyZero(T3T1, UE_KINDA_SMALL_NUMBER);
 
 			const FColor Color1 = Colors[FMath::Min(i, NumColors - 1)];
 			const FColor Color2 = Colors[FMath::Min(i + 1, NumColors - 1)];
