World Partition Runtime Spatial Hash: make sure adding new grids through the UI gets default values for cell size and loading range.

#rnx

diff --git a/Engine/Source/Runtime/Engine/Private/WorldPartition/WorldPartitionRuntimeSpatialHash.cpp b/Engine/Source/Runtime/Engine/Private/WorldPartition/WorldPartitionRuntimeSpatialHash.cpp
--- a/Engine/Source/Runtime/Engine/Private/WorldPartition/WorldPartitionRuntimeSpatialHash.cpp
+++ b/Engine/Source/Runtime/Engine/Private/WorldPartition/WorldPartitionRuntimeSpatialHash.cpp
@@ -877,8 +877,6 @@ void UWorldPartitionRuntimeSpatialHash::SetDefaultValues()
 
 	FSpatialHashRuntimeGrid& MainGrid = Grids.AddDefaulted_GetRef();
 	MainGrid.GridName = TEXT("MainGrid");
-	MainGrid.CellSize = 12800;
-	MainGrid.LoadingRange = 25600;
 	MainGrid.DebugColor = FLinearColor::Gray;
 }
 
diff --git a/Engine/Source/Runtime/Engine/Public/WorldPartition/WorldPartitionRuntimeSpatialHash.h b/Engine/Source/Runtime/Engine/Public/WorldPartition/WorldPartitionRuntimeSpatialHash.h
--- a/Engine/Source/Runtime/Engine/Public/WorldPartition/WorldPartitionRuntimeSpatialHash.h
+++ b/Engine/Source/Runtime/Engine/Public/WorldPartition/WorldPartitionRuntimeSpatialHash.h
@@ -167,8 +167,8 @@ struct FSpatialHashRuntimeGrid
 
 	FSpatialHashRuntimeGrid()
 #if WITH_EDITORONLY_DATA
-		: CellSize(0)
-		, LoadingRange(1000)
+		: CellSize(12800)
+		, LoadingRange(25600)
 		, bBlockOnSlowStreaming(false)
 		, Priority(0)
 		, DebugColor(FLinearColor::MakeRandomColor())
