Game user setting CVar `r.VSync` check using game thread access triggers ensure condition.

diff --git a/Engine/Source/Runtime/Engine/Private/GameUserSettings.cpp b/Engine/Source/Runtime/Engine/Private/GameUserSettings.cpp
--- a/Engine/Source/Runtime/Engine/Private/GameUserSettings.cpp
+++ b/Engine/Source/Runtime/Engine/Private/GameUserSettings.cpp
@@ -180,7 +180,7 @@ bool UGameUserSettings::IsVSyncDirty() const
 	if (GEngine && GEngine->GameViewport && GEngine->GameViewport->ViewportFrame)
 	{
 		static const auto CVar = IConsoleManager::Get().FindTConsoleVariableDataInt(TEXT("r.VSync"));
-		bIsDirty = (bUseVSync != (CVar->GetValueOnGameThread() != 0));
+		bIsDirty = (bUseVSync != (CVar->GetValueOnAnyThread() != 0));
 	}
 	return bIsDirty;
 }
