[Engine]
* Moved AudioDeviceHandle to its own file

diff --git a/Engine/Source/Runtime/Engine/Public/AudioDeviceHandle.h b/Engine/Source/Runtime/Engine/Public/AudioDeviceHandle.h
new file mode 100644
--- /dev/null
+++ b/Engine/Source/Runtime/Engine/Public/AudioDeviceHandle.h
@@ -0,0 +1,133 @@
+// Copyright Epic Games, Inc. All Rights Reserved.
+#pragma once
+
+// Set this to one if you'd like to check who owns
+// handles to an audio device.
+#ifndef INSTRUMENT_AUDIODEVICE_HANDLES
+#define INSTRUMENT_AUDIODEVICE_HANDLES 0
+#endif
+
+namespace Audio
+{
+	class FMixerDevice;
+
+	/**
+	 * Typed identifier for Audio Device Id
+	 */
+	using FDeviceId = uint32;
+}
+
+class FAudioDevice;
+class UWorld;
+
+
+// Strong handle to an audio device. Guarantees that the audio device it references will stay alive while it is in scope.
+class ENGINE_API FAudioDeviceHandle
+{
+public:
+	FAudioDeviceHandle();
+	FAudioDeviceHandle(const FAudioDeviceHandle& Other);
+	FAudioDeviceHandle(FAudioDeviceHandle&& Other);
+
+	FAudioDeviceHandle& operator=(const FAudioDeviceHandle& Other);
+	FAudioDeviceHandle& operator=(FAudioDeviceHandle&& Other);
+
+
+	~FAudioDeviceHandle();
+
+	// gets a pointer to the audio device.
+	FAudioDevice* GetAudioDevice() const;
+
+	// Returns the device ID for the audio device referenced by this handle.
+	Audio::FDeviceId GetDeviceID() const;
+
+	// Returns world set when handle was created
+	TWeakObjectPtr<UWorld> GetWorld() const;
+
+	// Checks whether this points to a valid compressed chunk.
+	bool IsValid() const;
+
+	// Resets handle to invalid state
+	void Reset();
+
+private:
+	// This constructor should only be called by FAudioDeviceManager.
+	FAudioDeviceHandle(FAudioDevice* InDevice, Audio::FDeviceId InID, UWorld* InWorld);
+
+	// The world that this FAudioDeviceHandle was requested with.
+	// Null if this device handle wasn't generated by RequestAudioDevice.
+	TWeakObjectPtr<UWorld> World;
+
+	FAudioDevice* Device;
+	Audio::FDeviceId DeviceId;
+
+#if INSTRUMENT_AUDIODEVICE_HANDLES
+	uint32 StackWalkID;
+	void AddStackDumpToAudioDeviceContainer();
+#endif
+
+	friend class FAudioDeviceManager;
+
+public:
+	// These are convenience operators to use an FAudioDeviceHandle like an FAudioDevice* or an Audio::FDeviceId.
+	// For safety, we still require explicit casting to an FAudioDevice* to ensure ownership isn't lost due to programmer error.
+	const FAudioDevice& operator*() const
+	{
+		check(IsValid());
+		return *Device;
+	}
+
+	FAudioDevice& operator*()
+	{
+		check(IsValid());
+		return *Device;
+	}
+
+	const FAudioDevice* operator->() const
+	{
+		check(IsValid());
+		return Device;
+	}
+
+	FAudioDevice* operator->()
+	{
+		check(IsValid());
+		return Device;
+	}
+
+	FORCEINLINE explicit operator bool() const
+	{
+		return IsValid();
+	}
+
+	FORCEINLINE bool operator==(const FAudioDeviceHandle& Other) const
+	{
+		return DeviceId == Other.DeviceId;
+	}
+
+	FORCEINLINE bool operator==(const FAudioDevice*& Other) const
+	{
+		return Device == Other;
+	}
+
+	FORCEINLINE bool operator==(FAudioDevice*& Other) const
+	{
+		return Device == Other;
+	}
+
+	FORCEINLINE bool operator==(const Audio::FDeviceId& Other) const
+	{
+		return DeviceId == Other;
+	}
+
+	explicit operator FAudioDevice*() const
+	{
+		check(IsValid());
+		return Device;
+	}
+
+	explicit operator Audio::FDeviceId() const
+	{
+		return DeviceId;
+	}
+};
diff --git a/Engine/Source/Runtime/Engine/Public/AudioDeviceManager.h b/Engine/Source/Runtime/Engine/Public/AudioDeviceManager.h
--- a/Engine/Source/Runtime/Engine/Public/AudioDeviceManager.h
+++ b/Engine/Source/Runtime/Engine/Public/AudioDeviceManager.h
@@ -2,6 +2,7 @@
 #pragma once
 
 #include "AudioDefines.h"
+#include "AudioDeviceHandle.h"
 #include "AudioThread.h"
 #include "CoreMinimal.h"
 #include "UObject/WeakObjectPtr.h"
@@ -11,12 +12,6 @@
 class FAudioDebugger;
 #endif // ENABLE_AUDIO_DEBUG
 
-// Set this to one if you'd like to check who owns
-// handles to an audio device.
-#ifndef INSTRUMENT_AUDIODEVICE_HANDLES
-#define INSTRUMENT_AUDIODEVICE_HANDLES 0
-#endif
-
 class FReferenceCollector;
 class FSoundBuffer;
 class IAudioDeviceModule;
@@ -33,12 +28,6 @@ class FSimpleAudioInfoFactory;
 namespace Audio
 {
 	class FMixerDevice;
-
-	/**
-	 * Typed identifier for Audio Device Id
-	 */
-	using FDeviceId = uint32;
-
 	class FAudioDebugger;
 	class FAudioFormatSettings;
 }
@@ -78,117 +67,6 @@ struct FAudioDeviceParams
 	int32 NumBuffersOverride = INDEX_NONE;
 };
 
-// Strong handle to an audio device. Guarantees that the audio device it references will stay alive while it is in scope.
-class ENGINE_API FAudioDeviceHandle
-{
-public:
-	FAudioDeviceHandle();
-	FAudioDeviceHandle(const FAudioDeviceHandle& Other);
-	FAudioDeviceHandle(FAudioDeviceHandle&& Other);
-
-	FAudioDeviceHandle& operator=(const FAudioDeviceHandle& Other);
-	FAudioDeviceHandle& operator=(FAudioDeviceHandle&& Other);
-
-
-	~FAudioDeviceHandle();
-
-	// gets a pointer to the audio device.
-	FAudioDevice* GetAudioDevice() const;
-
-	// Returns the device ID for the audio device referenced by this handle.
-	Audio::FDeviceId GetDeviceID() const;
-
-	// Returns world set when handle was created
-	TWeakObjectPtr<UWorld> GetWorld() const;
-
-	// Checks whether this points to a valid compressed chunk.
-	bool IsValid() const;
-
-	// Resets handle to invalid state
-	void Reset();
-
-private:
-	// This constructor should only be called by FAudioDeviceManager.
-	FAudioDeviceHandle(FAudioDevice* InDevice, Audio::FDeviceId InID, UWorld* InWorld);
-
-	// The world that this FAudioDeviceHandle was requested with.
-	// Null if this device handle wasn't generated by RequestAudioDevice.
-	TWeakObjectPtr<UWorld> World;
-
-	FAudioDevice* Device;
-	Audio::FDeviceId DeviceId;
-
-#if INSTRUMENT_AUDIODEVICE_HANDLES
-	uint32 StackWalkID;
-	void AddStackDumpToAudioDeviceContainer();
-#endif
-
-	friend class FAudioDeviceManager;
-
-public:
-	// These are convenience operators to use an FAudioDeviceHandle like an FAudioDevice* or an Audio::FDeviceId.
-	// For safety, we still require explicit casting to an FAudioDevice* to ensure ownership isn't lost due to programmer error.
-	const FAudioDevice& operator*() const
-	{
-		check(IsValid());
-		return *Device;
-	}
-
-	FAudioDevice& operator*()
-	{
-		check(IsValid());
-		return *Device;
-	}
-
-	const FAudioDevice* operator->() const
-	{
-		check(IsValid());
-		return Device;
-	}
-
-	FAudioDevice* operator->()
-	{
-		check(IsValid());
-		return Device;
-	}
-
-	FORCEINLINE explicit operator bool() const
-	{
-		return IsValid();
-	}
-
-	FORCEINLINE bool operator==(const FAudioDeviceHandle& Other) const
-	{
-		return DeviceId == Other.DeviceId;
-	}
-
-	FORCEINLINE bool operator==(const FAudioDevice*& Other) const
-	{
-		return Device == Other;
-	}
-
-	FORCEINLINE bool operator==(FAudioDevice*& Other) const
-	{
-		return Device == Other;
-	}
-
-	FORCEINLINE bool operator==(const Audio::FDeviceId& Other) const
-	{
-		return DeviceId == Other;
-	}
-
-	explicit operator FAudioDevice*() const
-	{
-		check(IsValid());
-		return Device;
-	}
-
-	explicit operator Audio::FDeviceId() const
-	{
-		return DeviceId;
-	}
-};
-
 // List of delegates for the audio device manager.
 class ENGINE_API FAudioDeviceManagerDelegates
 {
