[Engine]
* AudioDeviceHandle - Fixed non unity error

diff --git a/Engine/Source/Runtime/Engine/Public/AudioDeviceHandle.h b/Engine/Source/Runtime/Engine/Public/AudioDeviceHandle.h
--- a/Engine/Source/Runtime/Engine/Public/AudioDeviceHandle.h
+++ b/Engine/Source/Runtime/Engine/Public/AudioDeviceHandle.h
@@ -1,6 +1,8 @@
 // Copyright Epic Games, Inc. All Rights Reserved.
 #pragma once
 
+#include "UObject/WeakObjectPtr.h"
+
 // Set this to one if you'd like to check who owns
 // handles to an audio device.
 #ifndef INSTRUMENT_AUDIODEVICE_HANDLES
