WorldPartition: Fix editor loader not using proper world to filter DataLayers loaded actors for sub worlds (Level Instances)

#rnx

diff --git a/Engine/Source/Runtime/Engine/Private/WorldPartition/WorldPartitionActorLoaderInterface.cpp b/Engine/Source/Runtime/Engine/Private/WorldPartition/WorldPartitionActorLoaderInterface.cpp
--- a/Engine/Source/Runtime/Engine/Private/WorldPartition/WorldPartitionActorLoaderInterface.cpp
+++ b/Engine/Source/Runtime/Engine/Private/WorldPartition/WorldPartitionActorLoaderInterface.cpp
@@ -124,7 +124,8 @@ bool IWorldPartitionActorLoaderInterface::ILoaderAdapter::PassActorDescFilter(co
 
 bool IWorldPartitionActorLoaderInterface::ILoaderAdapter::PassDataLayersFilter(const FWorldPartitionHandle& Actor) const
 {
-	if (UDataLayerSubsystem* DataLayerSubsystem = UWorld::GetSubsystem<UDataLayerSubsystem>(World))
+	UWorld* OwningWorld = World->PersistentLevel->GetWorld();
+	if (UDataLayerSubsystem* DataLayerSubsystem = UWorld::GetSubsystem<UDataLayerSubsystem>(OwningWorld))
 	{
 		FWorldPartitionActorViewProxy ActorDescProxy(*Actor);
 
