Add blueprint function to set Nanite Override Material on a MIC.

diff --git a/Engine/Source/Runtime/Engine/Classes/Materials/MaterialInstanceConstant.h b/Engine/Source/Runtime/Engine/Classes/Materials/MaterialInstanceConstant.h
--- a/Engine/Source/Runtime/Engine/Classes/Materials/MaterialInstanceConstant.h
+++ b/Engine/Source/Runtime/Engine/Classes/Materials/MaterialInstanceConstant.h
@@ -62,6 +62,10 @@ class UMaterialInstanceConstant : public UMaterialInstance
 	FLinearColor K2_GetVectorParameterValue(FName ParameterName);
 
 #if WITH_EDITOR
+	/** Set an override material which will be used when rendering with nanite. */
+	UFUNCTION(BlueprintCallable, Category = "Rendering|Material")
+	void SetNaniteOverrideMaterial(bool bInEnableOverride, UMaterialInterface* InOverrideMaterial);
+
 	/**
 	 * Set the parent of this material instance. This function may only be called in the Editor!
 	 *   WARNING: You MUST call PostEditChange afterwards to propagate changes to other materials in the chain!
diff --git a/Engine/Source/Runtime/Engine/Private/Materials/MaterialInstanceConstant.cpp b/Engine/Source/Runtime/Engine/Private/Materials/MaterialInstanceConstant.cpp
--- a/Engine/Source/Runtime/Engine/Private/Materials/MaterialInstanceConstant.cpp
+++ b/Engine/Source/Runtime/Engine/Private/Materials/MaterialInstanceConstant.cpp
@@ -250,5 +250,11 @@ void UMaterialInstanceConstant::UpdateCachedData()
 	}
 }
 
-#endif // #if WITH_EDITOR
+void UMaterialInstanceConstant::SetNaniteOverrideMaterial(bool bInEnableOverride, UMaterialInterface* InOverrideMaterial)
+{
+	NaniteOverrideMaterial.bEnableOverride = bInEnableOverride;
+	NaniteOverrideMaterial.OverrideMaterialRef = InOverrideMaterial;
+	NaniteOverrideMaterial.PostEditChange();
+}
 
+#endif // #if WITH_EDITOR
