Improved tooltip for ShadingRate Material attribute

There are licensees who mistakenly assumed the shading rate meant that the material would be supersampled.
The updated tooltip clearly defines that using non-1x1 rates will reduce fidelity, not improve it.

diff --git a/Engine/Source/Runtime/Engine/Classes/Materials/Material.h b/Engine/Source/Runtime/Engine/Classes/Materials/Material.h
--- a/Engine/Source/Runtime/Engine/Classes/Materials/Material.h
+++ b/Engine/Source/Runtime/Engine/Classes/Materials/Material.h
@@ -838,7 +838,11 @@ public:
 	UPROPERTY(EditAnywhere, Category=Material, AdvancedDisplay)
 	uint8 Wireframe : 1;
 
-	/** Select what shading rate to apply for platforms that have variable rate shading */
+	/**
+	 * Select what shading rate to apply, on platforms that support variable rate shading.
+	 * Non-1x1 rates will reduce the rasterization fidelity for the material; they will not super-sample the material.
+	 * This can save GPU performance on materials where reduced fidelity is acceptable.
+	 */
 	UPROPERTY(EditAnywhere, Category = Material, AdvancedDisplay)
 	TEnumAsByte<EMaterialShadingRate> ShadingRate;
 
