World partition debug display fixes : fixed SIE not working for all net modes, fixed debug 2D grids wasting empty space for other debug display, enabled text clipping.

diff --git a/Engine/Source/Runtime/Engine/Private/WorldPartition/WorldPartition.cpp b/Engine/Source/Runtime/Engine/Private/WorldPartition/WorldPartition.cpp
--- a/Engine/Source/Runtime/Engine/Private/WorldPartition/WorldPartition.cpp
+++ b/Engine/Source/Runtime/Engine/Private/WorldPartition/WorldPartition.cpp
@@ -94,6 +94,13 @@ FAutoConsoleVariableRef UWorldPartition::CVarDebugDedicatedServerStreaming(
 	UWorldPartition::DebugDedicatedServerStreaming,
 	TEXT("Turn on/off to debug of server streaming."),
 	ECVF_Default);
+
+int32 UWorldPartition::EnableSimulationStreamingSource = 1;
+FAutoConsoleVariableRef UWorldPartition::CVarEnableSimulationStreamingSource(
+	TEXT("wp.Runtime.EnableSimulationStreamingSource"),
+	UWorldPartition::EnableSimulationStreamingSource,
+	TEXT("Set to 0 to if you want to disable the simulation/ejected camera streaming source."),
+	ECVF_Default);
 #endif
 
 int32 UWorldPartition::EnableServerStreaming = 0;
@@ -1012,10 +1019,10 @@ bool UWorldPartition::UseMakingInvisibleTransactionRequests() const
 	return bCachedUseMakingInvisibleTransactionRequests.Get(false);
 }
 
-bool UWorldPartition::IsSimulating()
+bool UWorldPartition::IsSimulating(bool bIncludeTestEnableSimulationStreamingSource)
 {
 #if WITH_EDITOR
-	return GEditor && GEditor->bIsSimulatingInEditor && GCurrentLevelEditingViewportClient && GCurrentLevelEditingViewportClient->IsSimulateInEditorViewport();
+	return GEditor && GEditor->bIsSimulatingInEditor && GCurrentLevelEditingViewportClient && GCurrentLevelEditingViewportClient->IsSimulateInEditorViewport() && (!bIncludeTestEnableSimulationStreamingSource || UWorldPartition::EnableSimulationStreamingSource);
 #else
 	return false;
 #endif
@@ -1374,24 +1381,18 @@ bool UWorldPartition::CanDebugDraw() const
 		return false;
 	}
 
-	if (!bIsDedicatedServer && UWorldPartition::IsSimulating() && StreamingPolicy)
+	if (!bIsDedicatedServer && UWorldPartition::IsSimulating(false) && (!GEditor || (GetWorld() != GEditor->PlayWorld)))
 	{
-		return true;
+		return false;
 	}
 #endif
 	return GetWorld()->IsGameWorld();
 }
 
-FVector2D UWorldPartition::GetDrawRuntimeHash2DDesiredFootprint(const FVector2D& CanvasSize)
-{
-	check(CanDebugDraw());
-	return StreamingPolicy->GetDrawRuntimeHash2DDesiredFootprint(CanvasSize);
-}
-
-bool UWorldPartition::DrawRuntimeHash2D(class UCanvas* Canvas, const FVector2D& PartitionCanvasSize, const FVector2D& Offset)
+bool UWorldPartition::DrawRuntimeHash2D(class UCanvas* Canvas, const FVector2D& PartitionCanvasSize, const FVector2D& Offset, FVector2D& OutUsedCanvasSize)
 {
 	check(CanDebugDraw());
-	return StreamingPolicy->DrawRuntimeHash2D(Canvas, PartitionCanvasSize, Offset);
+	return StreamingPolicy->DrawRuntimeHash2D(Canvas, PartitionCanvasSize, Offset, OutUsedCanvasSize);
 }
 
 void UWorldPartition::DrawRuntimeHash3D()
diff --git a/Engine/Source/Runtime/Engine/Private/WorldPartition/WorldPartitionDebugHelper.cpp b/Engine/Source/Runtime/Engine/Private/WorldPartition/WorldPartitionDebugHelper.cpp
--- a/Engine/Source/Runtime/Engine/Private/WorldPartition/WorldPartitionDebugHelper.cpp
+++ b/Engine/Source/Runtime/Engine/Private/WorldPartition/WorldPartitionDebugHelper.cpp
@@ -147,12 +147,15 @@ bool FWorldPartitionDebugHelper::IsDebugCellNameShown(const FString& CellName)
 
 void FWorldPartitionDebugHelper::DrawText(UCanvas* Canvas, const FString& Text, const UFont* Font, const FColor& Color, FVector2D& Pos, float* MaxTextWidth)
 {
+	const float XScale = 1.f;
+	const float YScale = 1.f;
+	FFontRenderInfo RenderInfo;
+	RenderInfo.bClipText = true;
+	RenderInfo.bEnableShadow = true;
+	Canvas->SetDrawColor(Color);
+	Canvas->DrawText(Font, Text, Pos.X, Pos.Y, XScale, YScale, RenderInfo);
 	float TextWidth, TextHeight;
 	Canvas->StrLen(Font, Text, TextWidth, TextHeight);
-	Canvas->SetDrawColor(FColor(0,0,0));
-	Canvas->DrawText(Font, Text, Pos.X + 1, Pos.Y + 1);
-	Canvas->SetDrawColor(Color);
-	Canvas->DrawText(Font, Text, Pos.X, Pos.Y);
 	Pos.Y += TextHeight + 1;
 	if (MaxTextWidth)
 	{
diff --git a/Engine/Source/Runtime/Engine/Private/WorldPartition/WorldPartitionRuntimeSpatialHash.cpp b/Engine/Source/Runtime/Engine/Private/WorldPartition/WorldPartitionRuntimeSpatialHash.cpp
--- a/Engine/Source/Runtime/Engine/Private/WorldPartition/WorldPartitionRuntimeSpatialHash.cpp
+++ b/Engine/Source/Runtime/Engine/Private/WorldPartition/WorldPartitionRuntimeSpatialHash.cpp
@@ -1603,14 +1603,10 @@ EWorldPartitionStreamingPerformance UWorldPartitionRuntimeSpatialHash::GetStream
 	return EWorldPartitionStreamingPerformance::Good;
 }
 
-FVector2D UWorldPartitionRuntimeSpatialHash::GetDraw2DDesiredFootprint(const FVector2D& CanvasSize) const
-{
-	return FVector2D(CanvasSize.X * GetFilteredStreamingGrids().Num(), CanvasSize.Y);
-}
-
-bool UWorldPartitionRuntimeSpatialHash::Draw2D(UCanvas* Canvas, const TArray<FWorldPartitionStreamingSource>& Sources, const FVector2D& PartitionCanvasSize, const FVector2D& Offset) const
+bool UWorldPartitionRuntimeSpatialHash::Draw2D(UCanvas* Canvas, const TArray<FWorldPartitionStreamingSource>& Sources, const FVector2D& PartitionCanvasSize, const FVector2D& Offset, FVector2D& OutUsedCanvasSize) const
 {
 	TRACE_CPUPROFILER_EVENT_SCOPE(UWorldPartitionRuntimeSpatialHash::Draw2D);
+	OutUsedCanvasSize = FVector2D::ZeroVector;
 
 	TArray<const FSpatialHashStreamingGrid*> FilteredStreamingGrids = GetFilteredStreamingGrids();
 	if (FilteredStreamingGrids.Num() == 0 || Sources.Num() == 0)
@@ -1633,6 +1629,7 @@ bool UWorldPartitionRuntimeSpatialHash::Draw2D(UCanvas* Canvas, const TArray<FWo
 	// Sort streaming grids to render them sorted by loading range
 	FilteredStreamingGrids.Sort([](const FSpatialHashStreamingGrid& A, const FSpatialHashStreamingGrid& B) { return A.LoadingRange < B.LoadingRange; });
 
+	FBox2D GridsBounds(ForceInit);
 	int32 GridIndex = 0;
 	for (const FSpatialHashStreamingGrid* StreamingGrid : FilteredStreamingGrids)
 	{
@@ -1657,6 +1654,7 @@ bool UWorldPartitionRuntimeSpatialHash::Draw2D(UCanvas* Canvas, const TArray<FWo
 		auto WorldToScreen = [&](const FVector2D& WorldPos) { return (WorldToScreenScale * (WorldPos - GridReferenceWorldPos)) + GridScreenOffset; };
 
 		StreamingGrid->Draw2D(Canvas, World, Sources, Region, GridScreenBounds, WorldToScreen);
+		GridsBounds += GridScreenBounds;
 
 		// Draw WorldPartition name
 		FVector2D GridInfoPos = GridScreenOffset - GridScreenHalfExtent;
@@ -1689,6 +1687,11 @@ bool UWorldPartitionRuntimeSpatialHash::Draw2D(UCanvas* Canvas, const TArray<FWo
 		++GridIndex;
 	}
 
+	if (GridsBounds.bIsValid)
+	{
+		OutUsedCanvasSize = GridsBounds.GetSize();
+	}
+
 	return true;
 }
 
diff --git a/Engine/Source/Runtime/Engine/Private/WorldPartition/WorldPartitionStreamingPolicy.cpp b/Engine/Source/Runtime/Engine/Private/WorldPartition/WorldPartitionStreamingPolicy.cpp
--- a/Engine/Source/Runtime/Engine/Private/WorldPartition/WorldPartitionStreamingPolicy.cpp
+++ b/Engine/Source/Runtime/Engine/Private/WorldPartition/WorldPartitionStreamingPolicy.cpp
@@ -43,12 +43,6 @@ static FAutoConsoleVariableRef CVarUpdateStreamingSources(
 	GUpdateStreamingSources,
 	TEXT("Set to 0 to stop updating (freeze) world partition streaming sources."));
 
-static int32 GEnableSimulationStreamingSource = 1;
-static FAutoConsoleVariableRef CVarEnableSimulationStreamingSource(
-	TEXT("wp.Runtime.EnableSimulationStreamingSource"),
-	GEnableSimulationStreamingSource,
-	TEXT("Set to 0 to if you want to disable the simulation/ejected camera streaming source."));
-
 static int32 GMaxLoadingStreamingCells = 4;
 static FAutoConsoleVariableRef CMaxLoadingStreamingCells(
 	TEXT("wp.Runtime.MaxLoadingStreamingCells"),
@@ -113,7 +107,7 @@ void UWorldPartitionStreamingPolicy::UpdateStreamingSources()
 	if (!AWorldPartitionReplay::IsPlaybackEnabled(World) || !WorldPartition->Replay->GetReplayStreamingSources(StreamingSources))
 	{
 #if WITH_EDITOR
-		if (GEnableSimulationStreamingSource && UWorldPartition::IsSimulating())
+		if (UWorldPartition::IsSimulating())
 		{
 			// We are in the SIE
 			const FVector ViewLocation = WorldToLocal.TransformPosition(GCurrentLevelEditingViewportClient->GetViewLocation());
@@ -866,20 +860,11 @@ bool UWorldPartitionStreamingPolicy::IsStreamingCompleted(EWorldPartitionRuntime
 	return bResult;
 }
 
-FVector2D UWorldPartitionStreamingPolicy::GetDrawRuntimeHash2DDesiredFootprint(const FVector2D& CanvasSize)
-{
-	if (WorldPartition->RuntimeHash)
-	{
-		return WorldPartition->RuntimeHash->GetDraw2DDesiredFootprint(CanvasSize);
-	}
-	return FVector2D::ZeroVector;
-}
-
-bool UWorldPartitionStreamingPolicy::DrawRuntimeHash2D(class UCanvas* Canvas, const FVector2D& PartitionCanvasSize, const FVector2D& Offset)
+bool UWorldPartitionStreamingPolicy::DrawRuntimeHash2D(class UCanvas* Canvas, const FVector2D& PartitionCanvasSize, const FVector2D& Offset, FVector2D& OutUsedCanvasSize)
 {
 	if (StreamingSources.Num() > 0 && WorldPartition->RuntimeHash)
 	{
-		return WorldPartition->RuntimeHash->Draw2D(Canvas, StreamingSources, PartitionCanvasSize, Offset);
+		return WorldPartition->RuntimeHash->Draw2D(Canvas, StreamingSources, PartitionCanvasSize, Offset, OutUsedCanvasSize);
 	}
 	return false;
 }
diff --git a/Engine/Source/Runtime/Engine/Private/WorldPartition/WorldPartitionSubsystem.cpp b/Engine/Source/Runtime/Engine/Private/WorldPartition/WorldPartitionSubsystem.cpp
--- a/Engine/Source/Runtime/Engine/Private/WorldPartition/WorldPartitionSubsystem.cpp
+++ b/Engine/Source/Runtime/Engine/Private/WorldPartition/WorldPartitionSubsystem.cpp
@@ -343,9 +343,8 @@ void UWorldPartitionSubsystem::Draw(UCanvas* Canvas, class APlayerController* PC
 	}
 
 	// Filter out views that don't match our world
-	if (!WorldPartition->GetWorld()->IsNetMode(NM_DedicatedServer) &&
-		Canvas->SceneView->ViewActor != nullptr &&
-		Canvas->SceneView->ViewActor->GetWorld() != GetWorld())
+	if (!WorldPartition->GetWorld()->IsNetMode(NM_DedicatedServer) && !UWorldPartition::IsSimulating(false) &&
+		(Canvas->SceneView->ViewActor == nullptr || Canvas->SceneView->ViewActor->GetWorld() != GetWorld()))
 	{
 		return;
 	}
@@ -362,10 +361,11 @@ void UWorldPartitionSubsystem::Draw(UCanvas* Canvas, class APlayerController* PC
 		const FVector2D CanvasMaxScreenSize = FVector2D::Max(MaxScreenRatio*FVector2D(Canvas->ClipX, Canvas->ClipY) - CanvasBottomRightPadding - CurrentOffset, CanvasMinimumSize);
 
 		FVector2D PartitionCanvasSize = FVector2D(CanvasMaxScreenSize.X, CanvasMaxScreenSize.Y);
-		if (WorldPartition->DrawRuntimeHash2D(Canvas, PartitionCanvasSize, CurrentOffset))
+		FVector2D UsedCanvasSize = FVector2D::ZeroVector;
+		if (WorldPartition->DrawRuntimeHash2D(Canvas, PartitionCanvasSize, CurrentOffset, UsedCanvasSize))
 		{
 			CurrentOffset.X = CanvasBottomRightPadding.X;
-			CurrentOffset.Y += CanvasMaxScreenSize.Y;
+			CurrentOffset.Y += UsedCanvasSize.Y;
 		}
 	}
 	
diff --git a/Engine/Source/Runtime/Engine/Public/WorldPartition/WorldPartition.h b/Engine/Source/Runtime/Engine/Public/WorldPartition/WorldPartition.h
--- a/Engine/Source/Runtime/Engine/Public/WorldPartition/WorldPartition.h
+++ b/Engine/Source/Runtime/Engine/Public/WorldPartition/WorldPartition.h
@@ -217,7 +217,7 @@ public:
 #endif
 
 public:
-	static bool IsSimulating();
+	static bool IsSimulating(bool bIncludeTestEnableSimulationStreamingSource = true);
 
 	void Initialize(UWorld* World, const FTransform& InTransform);
 	bool IsInitialized() const;
@@ -244,8 +244,7 @@ public:
 
 	// Debugging
 	bool CanDebugDraw() const;
-	FVector2D GetDrawRuntimeHash2DDesiredFootprint(const FVector2D& CanvasSize);
-	bool DrawRuntimeHash2D(UCanvas* Canvas, const FVector2D& PartitionCanvasSize, const FVector2D& Offset);
+	bool DrawRuntimeHash2D(UCanvas* Canvas, const FVector2D& PartitionCanvasSize, const FVector2D& Offset, FVector2D& OutUsedCanvasSize);
 	void DrawRuntimeHash3D();
 	void DrawRuntimeCellsDetails(UCanvas* Canvas, FVector2D& Offset);
 	void DrawStreamingStatusLegend(UCanvas* Canvas, FVector2D& Offset);
@@ -334,9 +333,11 @@ private:
 
 #if WITH_EDITOR
 	static int32 LoadingRangeBugItGo;
+	static int32 EnableSimulationStreamingSource;
 	static int32 WorldExtentToEnableStreaming;
 	static bool DebugDedicatedServerStreaming;
 	static FAutoConsoleVariableRef CVarLoadingRangeBugItGo;
+	static FAutoConsoleVariableRef CVarEnableSimulationStreamingSource;
 	static FAutoConsoleVariableRef CVarWorldExtentToEnableStreaming;
 	static FAutoConsoleVariableRef CVarDebugDedicatedServerStreaming;
 #endif
diff --git a/Engine/Source/Runtime/Engine/Public/WorldPartition/WorldPartitionRuntimeHash.h b/Engine/Source/Runtime/Engine/Public/WorldPartition/WorldPartitionRuntimeHash.h
--- a/Engine/Source/Runtime/Engine/Public/WorldPartition/WorldPartitionRuntimeHash.h
+++ b/Engine/Source/Runtime/Engine/Public/WorldPartition/WorldPartitionRuntimeHash.h
@@ -133,14 +133,7 @@ public:
 	virtual bool InjectExternalStreamingObject(URuntimeHashExternalStreamingObjectBase* ExternalStreamingObject) { return false; }
 	virtual bool RemoveExternalStreamingObject(URuntimeHashExternalStreamingObjectBase* ExternalStreamingObject) { return false; }
 
-	/* Returns desired footprint that Draw2D should take relative to given Canvas size (the value can exceed the given size).
-	 * UWorldPartitionSubSystem will re-adapt the size relative to all others UWorldPartitionRuntimeHash and provide the correct size to Draw2D.
-	 *
-	 * Return Draw2D's desired footprint.
-	 */
-	virtual FVector2D GetDraw2DDesiredFootprint(const FVector2D& CanvasSize) const { return FVector2D::ZeroVector; }
-
-	virtual bool Draw2D(class UCanvas* Canvas, const TArray<FWorldPartitionStreamingSource>& Sources, const FVector2D& PartitionCanvasSize, const FVector2D& Offset) const { return false; }
+	virtual bool Draw2D(class UCanvas* Canvas, const TArray<FWorldPartitionStreamingSource>& Sources, const FVector2D& PartitionCanvasSize, const FVector2D& Offset, FVector2D& OutUsedCanvasSize) const { return false; }
 	virtual void Draw3D(const TArray<FWorldPartitionStreamingSource>& Sources) const {}
 	virtual bool ContainsRuntimeHash(const FString& Name) const { return false; }
 
diff --git a/Engine/Source/Runtime/Engine/Public/WorldPartition/WorldPartitionRuntimeSpatialHash.h b/Engine/Source/Runtime/Engine/Public/WorldPartition/WorldPartitionRuntimeSpatialHash.h
--- a/Engine/Source/Runtime/Engine/Public/WorldPartition/WorldPartitionRuntimeSpatialHash.h
+++ b/Engine/Source/Runtime/Engine/Public/WorldPartition/WorldPartitionRuntimeSpatialHash.h
@@ -329,8 +329,7 @@ private:
 	UPROPERTY(Transient)
 	TArray<TWeakObjectPtr<URuntimeSpatialHashExternalStreamingObject>> ExternalStreamingObjects;
 	
-	virtual FVector2D GetDraw2DDesiredFootprint(const FVector2D& CanvasSize) const override;
-	virtual bool Draw2D(class UCanvas* Canvas, const TArray<FWorldPartitionStreamingSource>& Sources, const FVector2D& PartitionCanvasSize, const FVector2D& Offset) const override;
+	virtual bool Draw2D(class UCanvas* Canvas, const TArray<FWorldPartitionStreamingSource>& Sources, const FVector2D& PartitionCanvasSize, const FVector2D& Offset, FVector2D& OutUsedCanvasSize) const override;
 	virtual void Draw3D(const TArray<FWorldPartitionStreamingSource>& Sources) const override;
 	virtual bool ContainsRuntimeHash(const FString& Name) const override;
 
diff --git a/Engine/Source/Runtime/Engine/Public/WorldPartition/WorldPartitionStreamingPolicy.h b/Engine/Source/Runtime/Engine/Public/WorldPartition/WorldPartitionStreamingPolicy.h
--- a/Engine/Source/Runtime/Engine/Public/WorldPartition/WorldPartitionStreamingPolicy.h
+++ b/Engine/Source/Runtime/Engine/Public/WorldPartition/WorldPartitionStreamingPolicy.h
@@ -44,8 +44,7 @@ class UWorldPartitionStreamingPolicy : public UObject
 public:
 	virtual void UpdateStreamingState();
 	virtual bool CanAddLoadedLevelToWorld(class ULevel* InLevel) const;
-	virtual FVector2D GetDrawRuntimeHash2DDesiredFootprint(const FVector2D& CanvasSize);
-	virtual bool DrawRuntimeHash2D(class UCanvas* Canvas, const FVector2D& PartitionCanvasSize, const FVector2D& Offset);
+	virtual bool DrawRuntimeHash2D(class UCanvas* Canvas, const FVector2D& PartitionCanvasSize, const FVector2D& Offset, FVector2D& OutUsedCanvasSize);
 	virtual void DrawRuntimeHash3D();
 	virtual void DrawRuntimeCellsDetails(class UCanvas* Canvas, FVector2D& Offset) {}
 	virtual void DrawStreamingStatusLegend(class UCanvas* Canvas, FVector2D& Offset) {}
