- Cascade fix for local space mesh particles and LWC

diff --git a/Engine/Shaders/Private/MeshParticleVertexFactory.ush b/Engine/Shaders/Private/MeshParticleVertexFactory.ush
--- a/Engine/Shaders/Private/MeshParticleVertexFactory.ush
+++ b/Engine/Shaders/Private/MeshParticleVertexFactory.ush
@@ -201,7 +201,8 @@ FLWCMatrix GetParticleTransform(FVertexFactoryInput Input)
 	FLWCMatrix FinalMat;
 	if ( MeshParticleVF.bUseLocalSpace )
 	{
-		FinalMat = LWCMultiply(Transform4x4, GetPrimitiveDataFromUniformBuffer().LocalToWorld);
+		float3 LWCTile = LWCGetTile(LWCGetOrigin(GetPrimitiveDataFromUniformBuffer().LocalToWorld));
+		FinalMat = MakeLWCMatrix(LWCTile, Transform4x4);
 	}
 	else
 	{
diff --git a/Engine/Source/Runtime/Engine/Private/Particles/ParticleSystemRender.cpp b/Engine/Source/Runtime/Engine/Private/Particles/ParticleSystemRender.cpp
--- a/Engine/Source/Runtime/Engine/Private/Particles/ParticleSystemRender.cpp
+++ b/Engine/Source/Runtime/Engine/Private/Particles/ParticleSystemRender.cpp
@@ -1763,11 +1763,10 @@ void FDynamicMeshEmitterData::CalculateParticleTransform(
 	FScaleMatrix kScaleMat(FVector(1.0f));
 	FQuat kLockedAxisQuat = FQuat::Identity;
 
-	const FVector LWCTileOffset = FVector(Source.LWCTile) * FLargeWorldRenderScalar::GetTileSize();
 	FVector ParticlePosition(ParticleLocation + ParticlePayloadCameraOffset);
-	kTransMat.M[3][0] = ParticlePosition.X - LWCTileOffset.X;
-	kTransMat.M[3][1] = ParticlePosition.Y - LWCTileOffset.Y;
-	kTransMat.M[3][2] = ParticlePosition.Z - LWCTileOffset.Z;
+	kTransMat.M[3][0] = ParticlePosition.X;
+	kTransMat.M[3][1] = ParticlePosition.Y;
+	kTransMat.M[3][2] = ParticlePosition.Z;
 
 	FVector3f ScaledSize = ParticleSize * Source.Scale;
 	kScaleMat.M[0][0] = ScaledSize.X;
@@ -2038,6 +2037,17 @@ void FDynamicMeshEmitterData::CalculateParticleTransform(
 
 	FTranslationMatrix OrbitMatrix(OrbitOffset);
 	OutTransformMat *= OrbitMatrix;
+
+	if (Source.bUseLocalSpace)
+	{
+		OutTransformMat *= LocalToWorld;
+	}
+
+	const FVector3f LWCTile = Source.bUseLocalSpace ? FLargeWorldRenderScalar::GetTileFor(ProxyLocalToWorld.GetOrigin()) : Source.LWCTile;
+	const FVector LWCTileOffset = FVector(LWCTile) * FLargeWorldRenderScalar::GetTileSize();
+	OutTransformMat.M[3][0] -= LWCTileOffset.X;
+	OutTransformMat.M[3][1] -= LWCTileOffset.Y;
+	OutTransformMat.M[3][2] -= LWCTileOffset.Z;
 }
 
 void FDynamicMeshEmitterData::GetInstanceData(void* InstanceData, void* DynamicParameterData, void* PrevTransformBuffer, const FParticleSystemSceneProxy* Proxy, const FSceneView* View, uint32 InstanceFactor) const
