Fix bug in shader bytecode de-deuplication which was causing us not to get shader job cache hits and resulted in more compilation that was necessary.

- Before we serialized the shader job output we need to zero out compile time and preprocess time stats.  These times will obviously vary and can't be used in the key to uniquely identify shader bytecode.
- Restore the preprocess time on the job after serialization.

todo: Separate the non-deterministic fields into a separate struct which will make it more obvious which members contribute to the key.

diff --git a/Engine/Source/Runtime/Engine/Private/ShaderCompiler/ShaderCompiler.cpp b/Engine/Source/Runtime/Engine/Private/ShaderCompiler/ShaderCompiler.cpp
--- a/Engine/Source/Runtime/Engine/Private/ShaderCompiler/ShaderCompiler.cpp
+++ b/Engine/Source/Runtime/Engine/Private/ShaderCompiler/ShaderCompiler.cpp
@@ -317,7 +317,7 @@ static void DumpShaderSourceToLog(const FString& DumpedSource)
 	TGuardValue<ELogTimes::Type> DisableLogTimes(GPrintLogTimes, ELogTimes::None);
 	GLog->Serialize(TEXT("\n======================= DUMPED SHADER BEGIN =======================\n"), ELogVerbosity::Log, NAME_None);
 	for (const FString& Line : DumpedSourceLines)
-	{
+		{
 		GLog->Serialize(*Line, ELogVerbosity::Log, NAME_None);
 	}
 	GLog->Serialize(TEXT("\n======================= DUMPED SHADER END =======================\n"), ELogVerbosity::Log, NAME_None);
@@ -8283,12 +8283,15 @@ FShaderCommonCompileJob::FInputHash FShaderCompileJob::GetInputHash()
 void FShaderCompileJob::SerializeOutput(FArchive& Ar)
 {
 	double ActualCompileTime = 0.0;
+	double ActualPreprocessTime = 0.0;
 	if (Ar.IsSaving())
 	{
 		// Cached jobs won't have accurate results anyway, so reduce the storage requirements by setting those fields to a known value.
 		// This significantly reduces the memory needed to store the outputs (by more than a half)
 		ActualCompileTime = Output.CompileTime;
+		ActualPreprocessTime = Output.PreprocessTime;
 		Output.CompileTime = 0.0;
+		Output.PreprocessTime = 0.0;
 	}
 
 	Ar << Output;
@@ -8304,6 +8307,7 @@ void FShaderCompileJob::SerializeOutput(FArchive& Ar)
 	{
 		// restore the compile time for this jobs. Jobs that will be deserialized from the cache will have a compile time of 0.0
 		Output.CompileTime = ActualCompileTime;
+		Output.PreprocessTime = ActualPreprocessTime;
 	}
 }
 
