World Partition - Minimap: When performing scene captures for the minimap, use a more reasonable value for the ZOffset / ZScaling of the ortho matrix

diff --git a/Engine/Source/Runtime/Engine/Private/WorldPartition/WorldPartitionMiniMapHelper.cpp b/Engine/Source/Runtime/Engine/Private/WorldPartition/WorldPartitionMiniMapHelper.cpp
--- a/Engine/Source/Runtime/Engine/Private/WorldPartition/WorldPartitionMiniMapHelper.cpp
+++ b/Engine/Source/Runtime/Engine/Private/WorldPartition/WorldPartitionMiniMapHelper.cpp
@@ -131,7 +131,7 @@ void FWorldPartitionMiniMapHelper::CalTopViewOfWorld(FMatrix& OutProjectionMatri
 	const float OrthoWidth = FMath::Max(1.f, ViewportWidth * OrthoZoom);
 	const float OrthoHeight = FMath::Max(1.f, ViewportHeight * OrthoZoom);
 
-	const double ZOffset = HALF_WORLD_MAX;
+	const double ZOffset = UE_FLOAT_HUGE_DISTANCE / 2.0;
 	OutProjectionMatrix = FReversedZOrthoMatrix(
 		OrthoWidth,
 		OrthoHeight,
