Exposing Cue Points to Blueprints

diff --git a/Engine/Source/Runtime/Engine/Classes/Sound/SoundWave.h b/Engine/Source/Runtime/Engine/Classes/Sound/SoundWave.h
--- a/Engine/Source/Runtime/Engine/Classes/Sound/SoundWave.h
+++ b/Engine/Source/Runtime/Engine/Classes/Sound/SoundWave.h
@@ -357,25 +357,25 @@ namespace Audio
 }
 
 // Struct defining a cue point in a sound wave asset
-USTRUCT()
+USTRUCT(BlueprintType)
 struct FSoundWaveCuePoint
 {
 	GENERATED_USTRUCT_BODY()
 
 	// Unique identifier for the wave cue point
-	UPROPERTY(Category = Info, VisibleAnywhere)
+	UPROPERTY(Category = Info, VisibleAnywhere, BlueprintReadOnly)
 	int32 CuePointID = 0;
 
 	// The label for the cue point
-	UPROPERTY(Category = Info, VisibleAnywhere)
+	UPROPERTY(Category = Info, VisibleAnywhere, BlueprintReadOnly)
 	FString Label;
 
 	// The frame position of the cue point
-	UPROPERTY(Category = Info, VisibleAnywhere)
+	UPROPERTY(Category = Info, VisibleAnywhere, BlueprintReadOnly)
 	int32 FramePosition = 0;
 
 	// The frame length of the cue point (non-zero if it's a region)
-	UPROPERTY(Category = Info, VisibleAnywhere)
+	UPROPERTY(Category = Info, VisibleAnywhere, BlueprintReadOnly)
 	int32 FrameLength = 0;
 };
 
@@ -704,7 +704,7 @@ public:
 	int32 NumChannels;
 
 	/** Cue point data */
-	UPROPERTY(Category = Info, VisibleAnywhere)
+	UPROPERTY(Category = Info, VisibleAnywhere, BlueprintReadOnly)
 	TArray<FSoundWaveCuePoint> CuePoints;
 
 #if WITH_EDITORONLY_DATA
