World Partition - HLOD: Made properties of UHLODLayer !editor only

diff --git a/Engine/Source/Runtime/Engine/Private/WorldPartition/HLOD/HLODLayer.cpp b/Engine/Source/Runtime/Engine/Private/WorldPartition/HLOD/HLODLayer.cpp
--- a/Engine/Source/Runtime/Engine/Private/WorldPartition/HLOD/HLODLayer.cpp
+++ b/Engine/Source/Runtime/Engine/Private/WorldPartition/HLOD/HLODLayer.cpp
@@ -27,11 +27,9 @@ DEFINE_LOG_CATEGORY_STATIC(LogHLODLayer, Log, All);
 
 UHLODLayer::UHLODLayer(const FObjectInitializer& ObjectInitializer)
 	: Super(ObjectInitializer)
-#if WITH_EDITORONLY_DATA
 	, bIsSpatiallyLoaded(true)
 	, CellSize(25600)
 	, LoadingRange(51200)
-#endif
 {
 }
 
@@ -193,10 +191,6 @@ void UHLODLayer::PostEditChangeProperty(FPropertyChangedEvent& PropertyChangedEv
 	}
 }
 
-#endif // WITH_EDITOR
-
-#if WITH_EDITORONLY_DATA
-
 FName UHLODLayer::GetRuntimeGridName(uint32 InLODLevel, int32 InCellSize, double InLoadingRange)
 {
 	return *FString::Format(TEXT("HLOD{0}_{1}m_{2}m"), { InLODLevel, int32(InCellSize * 0.01f), int32(InLoadingRange * 0.01f)});
@@ -218,5 +212,4 @@ const void UHLODLayer::SetParentLayer(const TSoftObjectPtr<UHLODLayer>& InParent
 	ParentLayer = InParentLayer;
 }
 
-#endif // WITH_EDITORONLY_DATA
-
+#endif // WITH_EDITOR
diff --git a/Engine/Source/Runtime/Engine/Public/WorldPartition/HLOD/HLODLayer.h b/Engine/Source/Runtime/Engine/Public/WorldPartition/HLOD/HLODLayer.h
--- a/Engine/Source/Runtime/Engine/Public/WorldPartition/HLOD/HLODLayer.h
+++ b/Engine/Source/Runtime/Engine/Public/WorldPartition/HLOD/HLODLayer.h
@@ -47,9 +47,7 @@ public:
 
 	/** Duplicate the provided HLOD layers setup */
 	static UHLODLayer* DuplicateHLODLayersSetup(UHLODLayer* HLODLayer, const FString& DestinationPath, const FString& Prefix);
-#endif
 
-#if WITH_EDITORONLY_DATA
 	EHLODLayerType GetLayerType() const { return LayerType; }
 	void SetLayerType(EHLODLayerType InLayerType) { LayerType = InLayerType; }
 	const TSubclassOf<UHLODBuilder> GetHLODBuilderClass() const { return HLODBuilderClass; }
@@ -65,9 +63,7 @@ public:
 	bool DoesRequireWarmup() const;
 
 	static FName GetRuntimeGridName(uint32 InLODLevel, int32 InCellSize, double InLoadingRange);
-#endif
 
-#if WITH_EDITOR
 private:
 	//~ Begin UObject Interface.
 	virtual void PostLoad() override;
@@ -75,7 +71,6 @@ private:
 	//~ End UObject Interface.
 #endif
 
-#if WITH_EDITORONLY_DATA
 private:
 	/** Type of HLOD generation to use */
 	UPROPERTY(EditAnywhere, Config, Category=HLOD)
@@ -92,9 +87,6 @@ private:
 	UPROPERTY(EditAnywhere, Config, Category=HLOD)
 	uint32 bIsSpatiallyLoaded : 1;
 
-	UPROPERTY()
-	uint32 bAlwaysLoaded_DEPRECATED : 1;
-
 	/** Cell size of the runtime grid created to encompass HLOD actors generated for this HLOD Layer */
 	UPROPERTY(EditAnywhere, Config, Category=HLOD, meta = (EditConditionHides, EditCondition = "bIsSpatiallyLoaded"))
 	int32 CellSize;
@@ -107,7 +99,10 @@ private:
 	UPROPERTY(EditAnywhere, Config, Category=HLOD, meta = (EditConditionHides, EditCondition = "bIsSpatiallyLoaded"))
 	TSoftObjectPtr<UHLODLayer> ParentLayer;
 
-public:
+private:
+	friend class FWorldPartitionHLODUtilities;
+
+#if WITH_EDITORONLY_DATA
 	UPROPERTY()
 	FMeshMergingSettings MeshMergeSettings_DEPRECATED;
 	UPROPERTY()
@@ -116,5 +111,7 @@ public:
 	FMeshApproximationSettings MeshApproximationSettings_DEPRECATED;
 	UPROPERTY()
 	TSoftObjectPtr<UMaterialInterface> HLODMaterial_DEPRECATED;
+	UPROPERTY()
+	uint32 bAlwaysLoaded_DEPRECATED : 1;
 #endif
 };
