Sparse volumetexture initial commit.

Sparse Volume Texture asset from imported OpenVDB (only a single float channel for now). Static or animated sequence.
It will be possible to generate SVT at runtime from GPU later.
Using FEditorBulkData for handling raw source without loading everything when not caching or cooking.
BulkData used at runtime for loading. No streaming yet.
Importer with dependency on OpenVDB is in a SparseVolumeTexture module only loaded when in editor

Sparse volume texture can be sampled from any materials (sample, sample parameter) and overridden on material instance and material instance dynamic.

Added support for uint in compiler (fetch from page table, see SparseVolumeTextureGetVoxelCoord)
Volume texture with u8 VirtualTextureLayerIndex!=255 (INDEX_NONE) are sparse texture. The layer index then represent what texture/attribute to sample.

diff --git a/Engine/Shaders/Private/MaterialTemplate.ush b/Engine/Shaders/Private/MaterialTemplate.ush
--- a/Engine/Shaders/Private/MaterialTemplate.ush
+++ b/Engine/Shaders/Private/MaterialTemplate.ush
@@ -802,6 +802,16 @@ float2 GetPixelPosition(FMaterialVertexParameters Parameters)
 	return GetViewportUV(Parameters) * View.ViewSizeAndInvSize.xy;
 }
 
+uint3 SparseVolumeTextureGetVoxelCoord(uint PackedPhysicalTileCoord, uint PhysicalTileSize, uint3 PageTableCoord, uint3 VolumeCoord)
+{
+	const int3 PhysicalTileCoord = int3(
+		PackedPhysicalTileCoord & 0x7FF,
+		(PackedPhysicalTileCoord >> 11) & 0x7FF,
+		(PackedPhysicalTileCoord >> 22) & 0x3FF);
+
+	int3 VoxelCoord = PhysicalTileCoord * PhysicalTileSize + int3(VolumeCoord - PageTableCoord * PhysicalTileSize);
+	return VoxelCoord;
+}
 
 #if POST_PROCESS_MATERIAL
 
diff --git a/Engine/Shaders/Private/SparseVolumeTexture/VisualizeSparseVolumeTexture.usf b/Engine/Shaders/Private/SparseVolumeTexture/VisualizeSparseVolumeTexture.usf
new file mode 100644
--- /dev/null
+++ b/Engine/Shaders/Private/SparseVolumeTexture/VisualizeSparseVolumeTexture.usf
@@ -0,0 +1,192 @@
+// Copyright Epic Games, Inc. All Rights Reserved.
+
+/*=============================================================================
+	VisualizeSparseVolumeTexture.usf: Used to visualize and debug a sparse texture asset.
+=============================================================================*/
+
+#include "/Engine/Private/Common.ush"
+
+
+#ifdef VERTEX_SHADER
+
+float DepthAsDeviceZ;
+
+void VisualizeSparseVolumeTextureVS(
+	in uint VertexId : SV_VertexID,
+	out float4 Position : SV_POSITION)
+{
+	float2 UV = -1.0f;
+	UV = VertexId == 1 ? float2(-1.0f, 3.0f) : UV;
+	UV = VertexId == 2 ? float2(3.0f, -1.0f) : UV;
+
+	Position = float4(UV, DepthAsDeviceZ, 1.0f);
+}
+
+#endif // VERTEX_SHADER
+
+
+#ifdef PIXEL_SHADER
+
+// Updated from http://jcgt.org/published/0007/03/04/
+bool Slabs(float3 p0, float3 p1, float3 rayOrigin, float3 invRaydir, out float outTMin, out float outTMax)
+{
+	float3 t0 = (p0 - rayOrigin) * invRaydir;
+	float3 t1 = (p1 - rayOrigin) * invRaydir;
+	float3 tmin = min(t0, t1), tmax = max(t0, t1);
+	float maxtmin = max(max(tmin.x, tmin.y), tmin.z);
+	float mintmax = min(min(tmax.x, tmax.y), tmax.z);
+	outTMin = maxtmin;
+	outTMax = mintmax;
+	return maxtmin <= mintmax;
+}
+
+float3 VolumeBoundMinWorld;
+float3 VolumeBoundMaxWorld;
+
+
+Texture3D<float4> SparseVolumeTexture;
+float3 SparseVolumeTextureResolution;
+
+Texture3D<uint> SparseVolumeTexturePageTable;
+float3 SparseVolumeTexturePageTableResolution;
+
+#define SPARSE_VOLUME_TILE_RES	16
+
+void VisualizeSparseVolumeTexturePS(
+	in float4 SVPos : SV_POSITION,
+	out float4 OutLuminanceAlpha : SV_Target0
+)
+{
+	ResolvedView = ResolveView();
+	const uint2 PixelCoord = uint2(SVPos.xy);
+	float2 UvBuffer = PixelCoord * View.BufferSizeAndInvSize.zw;	// Uv for depth buffer read (size can be larger than viewport)
+
+	float2 ScreenPosition = SvPositionToScreenPosition(SVPos).xy;
+	const float Depth = 10000.0f;
+	float4 TranslatedWorldPos = mul(float4(ScreenPosition * Depth, Depth, 1), View.ScreenToTranslatedWorld);
+
+	const float3 RayDir = normalize(TranslatedWorldPos.xyz - View.TranslatedWorldCameraOrigin);
+	const float3 RayOrig = LWCHackToFloat(ResolvedView.WorldCameraOrigin);
+
+	OutLuminanceAlpha = float4(0.0, 0.0, 0.0, 1.0);
+
+	float TMin = 0.0f;
+	float TMax = 0.0f;
+	float Transmittance = 1.0f;
+	if (Slabs(VolumeBoundMinWorld, VolumeBoundMaxWorld, RayOrig, 1.0 / RayDir, TMin, TMax) && TMax > 0.0f)
+	{
+		Transmittance = 0.5; // show the bounding box
+
+		// Amanatides 3D DDA marching implementation - Paper: http://citeseerx.ist.psu.edu/viewdoc/download?doi=10.1.1.42.3443&rep=rep1&type=pdf
+		// See https://www.shadertoy.com/view/3sKXDK
+
+		float3 StartPos = RayOrig + RayDir * max(0.0, TMin);
+
+		float3 VolumeResolution = SparseVolumeTexturePageTableResolution * SPARSE_VOLUME_TILE_RES;
+
+		const float3 StartUVs = (StartPos - VolumeBoundMinWorld) / (VolumeBoundMaxWorld - VolumeBoundMinWorld);
+		float3 P = StartUVs * VolumeResolution;
+
+		// Force round to volume
+		if (P.x < 0) P.x = 0;
+		if (P.y < 0) P.y = 0;
+		if (P.z < 0) P.z = 0;
+		if (P.x > (VolumeResolution.x)) P.x = (VolumeResolution.x);
+		if (P.y > (VolumeResolution.y)) P.y = (VolumeResolution.y);
+		if (P.z > (VolumeResolution.z)) P.z = (VolumeResolution.z);
+
+		// Amanatides 3D-DDA data preparation
+		float3 stepSign = sign(RayDir);
+		float3 tDelta = abs(1.0 / RayDir);
+		float3 tMax = float3(0.0, 0.0, 0.0);
+		float3 refPoint = floor(P);
+		tMax.x = stepSign.x > 0.0 ? refPoint.x + 1.0 - P.x : P.x - refPoint.x;
+		tMax.y = stepSign.y > 0.0 ? refPoint.y + 1.0 - P.y : P.y - refPoint.y;
+		tMax.z = stepSign.z > 0.0 ? refPoint.z + 1.0 - P.z : P.z - refPoint.z;
+		tMax.x *= tDelta.x;
+		tMax.y *= tDelta.y;
+		tMax.z *= tDelta.z;
+
+		LOOP
+		while (P.x >= 0 && P.y >= 0 && P.z >= 0 && P.x <= VolumeResolution.x && P.y <= VolumeResolution.y && P.z <= VolumeResolution.z)
+		{
+			float StepLength = 0.0;
+#if 0
+			// Slow reference
+			P += RayDir * 0.005;
+			StepLength = 0.005;
+#else
+			// Amanatides 3D-DDA 
+			if (tMax.x < tMax.y)
+			{
+				if (tMax.x < tMax.z)
+				{
+					P.x += stepSign.x;
+					tMax.x += tDelta.x;
+
+					StepLength += 1;
+				}
+				else
+				{
+					P.z += stepSign.z;
+					tMax.z += tDelta.z;
+
+					StepLength += 1;
+				}
+			}
+			else
+			{
+				if (tMax.y < tMax.z)
+				{
+					P.y += stepSign.y;
+					tMax.y += tDelta.y;
+
+					StepLength += 1;
+				}
+				else
+				{
+					P.z += stepSign.z;
+					tMax.z += tDelta.z;
+
+					StepLength += 1;
+				}
+			}
+#endif
+
+			float3 PageCoordF = floor(P * (1.0 / SPARSE_VOLUME_TILE_RES));
+			const uint EncodedPage = SparseVolumeTexturePageTable.Load(int4(int3(PageCoordF), 0));
+			if (EncodedPage)
+			{
+				const int3 TileCoord = int3(
+					 EncodedPage         & 0x7FF,
+					(EncodedPage >> 11 ) & 0x7FF,
+					(EncodedPage >> 22 ) & 0x3FF);
+
+				int3 VoxelCoord = TileCoord * SPARSE_VOLUME_TILE_RES + int3(P - PageCoordF * SPARSE_VOLUME_TILE_RES);
+
+				const float4 VoxelData = SparseVolumeTexture.Load(int4(VoxelCoord, 0));
+
+				const float Density = VoxelData.x;
+				if (Density > 0.0f)
+
+				{
+				#if 0
+					//OutLuminanceAlpha = float4(Rand3DPCG16(PageCoordF).x / 65535.0f, Rand3DPCG16(VoxelCoord).y / 65535.0f, 0.0, 0.0);
+					//OutLuminanceAlpha = float4(Rand3DPCG16(PageCoordF) / 65535.0f, 0.0);
+					OutLuminanceAlpha = float4(Rand3DPCG16(VoxelCoord) / 65535.0f, 0.0);
+					break;
+				#else
+					Transmittance *= exp(-StepLength * Density * 0.025);
+				#endif
+				}
+			}
+		}
+	}
+
+	OutLuminanceAlpha.a = Transmittance;
+}
+
+
+#endif // PIXEL_SHADER
+
+
diff --git a/Engine/Source/Editor/MaterialEditor/Private/MaterialEditingLibrary.cpp b/Engine/Source/Editor/MaterialEditor/Private/MaterialEditingLibrary.cpp
--- a/Engine/Source/Editor/MaterialEditor/Private/MaterialEditingLibrary.cpp
+++ b/Engine/Source/Editor/MaterialEditor/Private/MaterialEditingLibrary.cpp
@@ -1057,6 +1057,27 @@ bool UMaterialEditingLibrary::SetMaterialInstanceRuntimeVirtualTextureParameterV
 }
 
 
+USparseVolumeTexture* UMaterialEditingLibrary::GetMaterialInstanceSparseVolumeTextureParameterValue(UMaterialInstanceConstant* Instance, FName ParameterName, EMaterialParameterAssociation Association)
+{
+	USparseVolumeTexture* Result = nullptr;
+	if (Instance)
+	{
+		Instance->GetSparseVolumeTextureParameterValue(FHashedMaterialParameterInfo(ParameterName, Association, Association == EMaterialParameterAssociation::LayerParameter ? 0 : INDEX_NONE), Result);
+	}
+	return Result;
+}
+
+bool UMaterialEditingLibrary::SetMaterialInstanceSparseVolumeTextureParameterValue(UMaterialInstanceConstant* Instance, FName ParameterName, USparseVolumeTexture* Value, EMaterialParameterAssociation Association)
+{
+	bool bResult = false;
+	if (Instance)
+	{
+		Instance->SetSparseVolumeTextureParameterValueEditorOnly(FMaterialParameterInfo(ParameterName, Association, Association == EMaterialParameterAssociation::LayerParameter ? 0 : INDEX_NONE), Value);
+	}
+	return bResult;
+}
+
+
 FLinearColor UMaterialEditingLibrary::GetMaterialInstanceVectorParameterValue(UMaterialInstanceConstant* Instance, FName ParameterName, EMaterialParameterAssociation Association)
 {
 	FLinearColor Result = FLinearColor::Black;
diff --git a/Engine/Source/Editor/MaterialEditor/Private/MaterialEditor.cpp b/Engine/Source/Editor/MaterialEditor/Private/MaterialEditor.cpp
--- a/Engine/Source/Editor/MaterialEditor/Private/MaterialEditor.cpp
+++ b/Engine/Source/Editor/MaterialEditor/Private/MaterialEditor.cpp
@@ -71,6 +71,7 @@
 #include "Materials/MaterialExpressionReroute.h"
 #include "Materials/MaterialExpressionRuntimeVirtualTextureSample.h"
 #include "Materials/MaterialExpressionRuntimeVirtualTextureSampleParameter.h"
+#include "Materials/MaterialExpressionSparseVolumeTextureSample.h"
 #include "Materials/MaterialExpressionScalarParameter.h"
 #include "Materials/MaterialExpressionStaticComponentMaskParameter.h"
 #include "Materials/MaterialExpressionStaticSwitchParameter.h"
@@ -1514,6 +1515,11 @@ void FMaterialEditor::CreateInternalWidgets()
 		LayoutExpressionParameterDetails
 		);
 
+	MaterialDetailsView->RegisterInstancedCustomPropertyLayout(
+		UMaterialExpressionSparseVolumeTextureSampleParameter::StaticClass(),
+		LayoutExpressionParameterDetails
+		);
+
 	FOnGetDetailCustomizationInstance LayoutLayerExpressionParameterDetails = FOnGetDetailCustomizationInstance::CreateStatic(
 		&FMaterialExpressionLayersParameterDetails::MakeInstance, FOnCollectParameterGroups::CreateSP(this, &FMaterialEditor::GetAllMaterialExpressionGroups));
 
@@ -3886,6 +3892,7 @@ void FMaterialEditor::OnConvertObjects()
 				UMaterialExpressionVectorParameter* VectorParameterExpression = Cast<UMaterialExpressionVectorParameter>(CurrentSelectedExpression);
 				UMaterialExpressionTextureObjectParameter* TextureObjectParameterExpression = Cast<UMaterialExpressionTextureObjectParameter>(CurrentSelectedExpression);
 				UMaterialExpressionRuntimeVirtualTextureSample* RuntimeVirtualTextureSampleExpression = Cast<UMaterialExpressionRuntimeVirtualTextureSample>(CurrentSelectedExpression);
+				UMaterialExpressionSparseVolumeTextureSample* SparseVolumeTextureSampleExpression = Cast<UMaterialExpressionSparseVolumeTextureSample>(CurrentSelectedExpression);
 
 				// Setup the class to convert to
 				UClass* ClassToCreate = NULL;
@@ -3929,6 +3936,10 @@ void FMaterialEditor::OnConvertObjects()
 				{
 					ClassToCreate = UMaterialExpressionRuntimeVirtualTextureSampleParameter::StaticClass();
 				}
+				else if (SparseVolumeTextureSampleExpression)
+				{
+					ClassToCreate = UMaterialExpressionSparseVolumeTextureSampleParameter::StaticClass();
+				}
 				else if (ComponentMaskExpression)
 				{
 					ClassToCreate = UMaterialExpressionStaticComponentMaskParameter::StaticClass();
@@ -4007,6 +4018,13 @@ void FMaterialEditor::OnConvertObjects()
 							NewRuntimeVirtualTextureExpression->MipValueMode = RuntimeVirtualTextureSampleExpression->MipValueMode;
 							NewGraphNode->ReconstructNode();
 						}
+						else if (SparseVolumeTextureSampleExpression)
+						{
+							bNeedsRefresh = true;
+							UMaterialExpressionSparseVolumeTextureSampleParameter* NewSparseVolumeTextureExpression = CastChecked<UMaterialExpressionSparseVolumeTextureSampleParameter>(NewExpression);
+							NewSparseVolumeTextureExpression->SparseVolumeTexture = SparseVolumeTextureSampleExpression->SparseVolumeTexture;
+							NewGraphNode->ReconstructNode();
+						}
 						else if (ComponentMaskExpression)
 						{
 							bNeedsRefresh = true;
diff --git a/Engine/Source/Editor/MaterialEditor/Private/MaterialEditorInstanceDetailCustomization.cpp b/Engine/Source/Editor/MaterialEditor/Private/MaterialEditorInstanceDetailCustomization.cpp
--- a/Engine/Source/Editor/MaterialEditor/Private/MaterialEditorInstanceDetailCustomization.cpp
+++ b/Engine/Source/Editor/MaterialEditor/Private/MaterialEditorInstanceDetailCustomization.cpp
@@ -14,6 +14,7 @@
 #include "MaterialEditor/DEditorFontParameterValue.h"
 #include "MaterialEditor/DEditorMaterialLayersParameterValue.h"
 #include "MaterialEditor/DEditorRuntimeVirtualTextureParameterValue.h"
+#include "MaterialEditor/DEditorSparseVolumeTextureParameterValue.h"
 #include "MaterialEditor/DEditorScalarParameterValue.h"
 #include "MaterialEditor/DEditorStaticComponentMaskParameterValue.h"
 #include "MaterialEditor/DEditorStaticSwitchParameterValue.h"
@@ -391,6 +392,7 @@ void FMaterialInstanceParameterDetails::CreateSingleGroupWidget(FEditorParameter
 		UDEditorStaticSwitchParameterValue* SwitchParam = Cast<UDEditorStaticSwitchParameterValue>(Parameter);
 		UDEditorTextureParameterValue* TextureParam = Cast<UDEditorTextureParameterValue>(Parameter);
 		UDEditorRuntimeVirtualTextureParameterValue* RuntimeVirtualTextureParam = Cast<UDEditorRuntimeVirtualTextureParameterValue>(Parameter);
+		UDEditorSparseVolumeTextureParameterValue* SparseVolumeTextureParam = Cast<UDEditorSparseVolumeTextureParameterValue>(Parameter);
 		UDEditorVectorParameterValue* VectorParam = Cast<UDEditorVectorParameterValue>(Parameter);
 
 		if (Parameter->ParameterInfo.Association == EMaterialParameterAssociation::GlobalParameter)
diff --git a/Engine/Source/Editor/MaterialEditor/Private/MaterialEditorUtilities.cpp b/Engine/Source/Editor/MaterialEditor/Private/MaterialEditorUtilities.cpp
--- a/Engine/Source/Editor/MaterialEditor/Private/MaterialEditorUtilities.cpp
+++ b/Engine/Source/Editor/MaterialEditor/Private/MaterialEditorUtilities.cpp
@@ -19,6 +19,7 @@
 #include "Materials/MaterialExpressionTextureSample.h"
 #include "Materials/MaterialExpressionTextureSampleParameter.h"
 #include "Materials/MaterialExpressionRuntimeVirtualTextureSampleParameter.h"
+#include "Materials/MaterialExpressionSparseVolumeTextureSample.h"
 #include "Materials/MaterialExpressionFontSampleParameter.h"
 #include "Materials/MaterialExpressionMaterialFunctionCall.h"
 #include "Materials/MaterialExpressionScalarParameter.h"
@@ -461,6 +462,7 @@ void FMaterialEditorUtilities::GetVisibleMaterialParametersFromExpression(
 	UMaterialExpressionParameter* Param = Cast<UMaterialExpressionParameter>( MaterialExpressionKey.Expression );
 	UMaterialExpressionTextureSampleParameter* TexParam = Cast<UMaterialExpressionTextureSampleParameter>( MaterialExpressionKey.Expression );
 	UMaterialExpressionRuntimeVirtualTextureSampleParameter* RuntimeVirtualTexParam = Cast<UMaterialExpressionRuntimeVirtualTextureSampleParameter>(MaterialExpressionKey.Expression);
+	UMaterialExpressionSparseVolumeTextureSampleParameter* SparseVolumeTexParam = Cast<UMaterialExpressionSparseVolumeTextureSampleParameter>(MaterialExpressionKey.Expression);
 	UMaterialExpressionFontSampleParameter* FontParam = Cast<UMaterialExpressionFontSampleParameter>( MaterialExpressionKey.Expression );
 
 	if (Param)
@@ -475,12 +477,16 @@ void FMaterialEditorUtilities::GetVisibleMaterialParametersFromExpression(
 	{
 		ParameterInfo.Name = RuntimeVirtualTexParam->ParameterName;
 	}
+	else if (SparseVolumeTexParam)
+	{
+		ParameterInfo.Name = SparseVolumeTexParam->ParameterName;
+	}
 	else if (FontParam)
 	{
 		ParameterInfo.Name = FontParam->ParameterName;
 	}
 		
-	if (Param || TexParam || FontParam || RuntimeVirtualTexParam)
+	if (Param || TexParam || FontParam || RuntimeVirtualTexParam || SparseVolumeTexParam)
 	{
 		VisibleExpressions.AddUnique(ParameterInfo);
 	}
diff --git a/Engine/Source/Editor/MaterialEditor/Private/MaterialInstanceEditor.cpp b/Engine/Source/Editor/MaterialEditor/Private/MaterialInstanceEditor.cpp
--- a/Engine/Source/Editor/MaterialEditor/Private/MaterialInstanceEditor.cpp
+++ b/Engine/Source/Editor/MaterialEditor/Private/MaterialInstanceEditor.cpp
@@ -14,6 +14,7 @@
 #include "Styling/CoreStyle.h"
 #include "MaterialEditor/DEditorTextureParameterValue.h"
 #include "MaterialEditor/DEditorRuntimeVirtualTextureParameterValue.h"
+#include "MaterialEditor/DEditorSparseVolumeTextureParameterValue.h"
 #include "Materials/Material.h"
 #include "MaterialEditor/MaterialEditorInstanceConstant.h"
 #include "ThumbnailRendering/SceneThumbnailInfoWithPrimitive.h"
@@ -29,6 +30,7 @@
 #include "Materials/MaterialExpressionTextureBase.h"
 #include "Materials/MaterialExpressionTextureSampleParameter.h"
 #include "Materials/MaterialExpressionRuntimeVirtualTextureSampleParameter.h"
+#include "Materials/MaterialExpressionSparseVolumeTextureSample.h"
 
 #include "MaterialEditor.h"
 #include "MaterialEditorActions.h"
@@ -52,6 +54,7 @@
 #include "Widgets/Layout/SScrollBox.h"
 #include "DebugViewModeHelpers.h"
 #include "VT/RuntimeVirtualTexture.h"
+#include "SparseVolumeTexture/SparseVolumeTexture.h"
 #include "Widgets/Input/SButton.h"
 #include "Subsystems/AssetEditorSubsystem.h"
 
@@ -524,6 +527,7 @@ void FMaterialInstanceEditor::ReInitMaterialFunctionProxies()
 		TArray<FDoubleVectorParameterValue> DoubleVectorParameterValues = FunctionInstanceProxy->DoubleVectorParameterValues;
 		TArray<FTextureParameterValue> TextureParameterValues = FunctionInstanceProxy->TextureParameterValues;
 		TArray<FRuntimeVirtualTextureParameterValue> RuntimeVirtualTextureParameterValues = FunctionInstanceProxy->RuntimeVirtualTextureParameterValues;
+		TArray<FSparseVolumeTextureParameterValue> SparseVolumeTextureParameterValues = FunctionInstanceProxy->SparseVolumeTextureParameterValues;
 		TArray<FFontParameterValue> FontParameterValues = FunctionInstanceProxy->FontParameterValues;
 
 		const FStaticParameterSet& OldStaticParameters = FunctionInstanceProxy->GetStaticParameters();
@@ -599,6 +603,18 @@ void FMaterialInstanceEditor::ReInitMaterialFunctionProxies()
 			}
 		}
 
+		FunctionInstanceProxy->GetAllSparseVolumeTextureParameterInfo(OutParameterInfo, Guids);
+		FunctionInstanceProxy->SparseVolumeTextureParameterValues.Empty();
+		for (FSparseVolumeTextureParameterValue& SparseVolumeTextureParameter : SparseVolumeTextureParameterValues)
+		{
+			int32 Index = Guids.Find(SparseVolumeTextureParameter.ExpressionGUID);
+			if (Index != INDEX_NONE)
+			{
+				FunctionInstanceProxy->SparseVolumeTextureParameterValues.Add(SparseVolumeTextureParameter);
+				FunctionInstanceProxy->SparseVolumeTextureParameterValues.Last().ParameterInfo = OutParameterInfo[Index];
+			}
+		}
+
 		FunctionInstanceProxy->GetAllFontParameterInfo(OutParameterInfo, Guids);
 		FunctionInstanceProxy->FontParameterValues.Empty();
 		for (FFontParameterValue& FontParameter : FontParameterValues)
@@ -1374,6 +1390,22 @@ void FMaterialInstanceEditor::RefreshOnScreenMessages()
 							}
 						}
 					}
+
+					UDEditorSparseVolumeTextureParameterValue * SparseVolumeTextureParameterValue = Cast<UDEditorSparseVolumeTextureParameterValue>(Group.Parameters[ParameterIndex]);
+					if (SparseVolumeTextureParameterValue && SparseVolumeTextureParameterValue->ExpressionId.IsValid())
+					{
+						USparseVolumeTexture* SparseVolumeTexture = NULL;
+						MaterialEditorInstance->SourceInstance->GetSparseVolumeTextureParameterValue(SparseVolumeTextureParameterValue->ParameterInfo, SparseVolumeTexture);
+						if (SparseVolumeTexture)
+						{
+							UMaterialExpressionSparseVolumeTextureSampleParameter* Expression = BaseMaterial->FindExpressionByGUID<UMaterialExpressionSparseVolumeTextureSampleParameter>(SparseVolumeTextureParameterValue->ExpressionId);
+							if (!Expression)
+							{
+								const FText ExpressionNameText = FText::Format(LOCTEXT("MissingSVTExpression", "Warning: Sparse Volume Texture Expression {0} not found."), FText::FromName(SparseVolumeTextureParameterValue->ParameterInfo.Name));
+								OnScreenMessages.Emplace(FLinearColor(1, 1, 0), ExpressionNameText.ToString());
+							}
+						}
+					}
 				}
 			}
 		}
diff --git a/Engine/Source/Editor/MaterialEditor/Private/MaterialPropertyHelpers.cpp b/Engine/Source/Editor/MaterialEditor/Private/MaterialPropertyHelpers.cpp
--- a/Engine/Source/Editor/MaterialEditor/Private/MaterialPropertyHelpers.cpp
+++ b/Engine/Source/Editor/MaterialEditor/Private/MaterialPropertyHelpers.cpp
@@ -425,7 +425,8 @@ FReply FMaterialPropertyHelpers::OnClickedSaveNewFunctionInstance(class UMateria
 					ChildInstance->VectorParameterValues = EditedInstance->VectorParameterValues;
 					ChildInstance->DoubleVectorParameterValues = EditedInstance->DoubleVectorParameterValues;
 					ChildInstance->TextureParameterValues = EditedInstance->TextureParameterValues;
-					ChildInstance->RuntimeVirtualTextureParameterValues = EditedInstance->RuntimeVirtualTextureParameterValues; 
+					ChildInstance->RuntimeVirtualTextureParameterValues = EditedInstance->RuntimeVirtualTextureParameterValues;
+					ChildInstance->SparseVolumeTextureParameterValues = EditedInstance->SparseVolumeTextureParameterValues;
 					ChildInstance->FontParameterValues = EditedInstance->FontParameterValues;
 
 					const FStaticParameterSet& StaticParameters = EditedInstance->GetStaticParameters();
@@ -518,6 +519,7 @@ FReply FMaterialPropertyHelpers::OnClickedSaveNewLayerInstance(class UMaterialFu
 					ChildInstance->DoubleVectorParameterValues = EditedInstance->DoubleVectorParameterValues;
 					ChildInstance->TextureParameterValues = EditedInstance->TextureParameterValues;
 					ChildInstance->RuntimeVirtualTextureParameterValues = EditedInstance->RuntimeVirtualTextureParameterValues;
+					ChildInstance->SparseVolumeTextureParameterValues = EditedInstance->SparseVolumeTextureParameterValues;
 					ChildInstance->FontParameterValues = EditedInstance->FontParameterValues;
 
 					const FStaticParameterSet& StaticParameters = EditedInstance->GetStaticParameters();
diff --git a/Engine/Source/Editor/MaterialEditor/Private/SMaterialParametersOverviewWidget.cpp b/Engine/Source/Editor/MaterialEditor/Private/SMaterialParametersOverviewWidget.cpp
--- a/Engine/Source/Editor/MaterialEditor/Private/SMaterialParametersOverviewWidget.cpp
+++ b/Engine/Source/Editor/MaterialEditor/Private/SMaterialParametersOverviewWidget.cpp
@@ -40,6 +40,7 @@
 #include "Widgets/Views/SExpanderArrow.h"
 #include "Widgets/Layout/SWidgetSwitcher.h"
 #include "MaterialEditor/DEditorRuntimeVirtualTextureParameterValue.h"
+#include "MaterialEditor/DEditorSparseVolumeTextureParameterValue.h"
 #include "ThumbnailRendering/ThumbnailManager.h"
 #include "Styling/StyleColors.h"
 
@@ -812,6 +813,20 @@ void SMaterialParametersOverviewTree::ShowSubParameters()
 					}
 				}
 			}
+
+			UDEditorSparseVolumeTextureParameterValue* SVTParameter = Cast<UDEditorSparseVolumeTextureParameterValue>(Parameter);
+			if (!SVTParameter)
+			{
+				for (TSharedPtr<struct FSortedParamData> GroupChild : SortedParameters)
+				{
+					if (GroupChild->Group.GroupName == Property.ParameterGroup.GroupName
+						&& GroupChild->ParameterInfo.Association == ChildProperty->ParameterInfo.Association
+						&&  GroupChild->ParameterInfo.Index == ChildProperty->ParameterInfo.Index)
+					{
+						GroupChild->Children.Add(ChildProperty);
+					}
+				}
+			}
 		}
 	}
 }
diff --git a/Engine/Source/Editor/MaterialEditor/Public/MaterialEditingLibrary.h b/Engine/Source/Editor/MaterialEditor/Public/MaterialEditingLibrary.h
--- a/Engine/Source/Editor/MaterialEditor/Public/MaterialEditingLibrary.h
+++ b/Engine/Source/Editor/MaterialEditor/Public/MaterialEditingLibrary.h
@@ -286,6 +286,15 @@ public:
 	static bool SetMaterialInstanceRuntimeVirtualTextureParameterValue(UMaterialInstanceConstant* Instance, FName ParameterName, URuntimeVirtualTexture* Value, EMaterialParameterAssociation Association = EMaterialParameterAssociation::GlobalParameter);
 
 
+	/** Get the current texture parameter value from a Material Instance */
+	UFUNCTION(BlueprintPure, Category = "MaterialEditing")
+	static USparseVolumeTexture* GetMaterialInstanceSparseVolumeTextureParameterValue(UMaterialInstanceConstant* Instance, FName ParameterName, EMaterialParameterAssociation Association = EMaterialParameterAssociation::GlobalParameter);
+
+	/** Set the texture parameter value for a Material Instance */
+	UFUNCTION(BlueprintCallable, Category = "MaterialEditing")
+	static bool SetMaterialInstanceSparseVolumeTextureParameterValue(UMaterialInstanceConstant* Instance, FName ParameterName, USparseVolumeTexture* Value, EMaterialParameterAssociation Association = EMaterialParameterAssociation::GlobalParameter);
+
+
 	/** Get the current vector parameter value from a Material Instance */
 	UFUNCTION(BlueprintPure, Category = "MaterialEditing")
 	static FLinearColor GetMaterialInstanceVectorParameterValue(UMaterialInstanceConstant* Instance, FName ParameterName, EMaterialParameterAssociation Association = EMaterialParameterAssociation::GlobalParameter);
diff --git a/Engine/Source/Editor/SparseVolumeTexture/Private/AssetTypeActions_SparseVolumeTexture.cpp b/Engine/Source/Editor/SparseVolumeTexture/Private/AssetTypeActions_SparseVolumeTexture.cpp
new file mode 100644
--- /dev/null
+++ b/Engine/Source/Editor/SparseVolumeTexture/Private/AssetTypeActions_SparseVolumeTexture.cpp
@@ -0,0 +1,24 @@
+// Copyright Epic Games, Inc. All Rights Reserved.
+
+#include "AssetTypeActions_SparseVolumeTexture.h"
+
+#define LOCTEXT_NAMESPACE "AssetTypeActions_SparseVolumeTexture"
+
+void FAssetTypeActions_SparseVolumeTexture::GetActions(const TArray<UObject*>& InObjects, FToolMenuSection& Section)
+{
+	// Checkout AssetTypeActions_Texture.cpp for some exemples
+}
+
+
+void FAssetTypeActions_SparseVolumeTexture::GetResolvedSourceFilePaths(const TArray<UObject*>& TypeAssets, TArray<FString>& OutSourceFilePaths) const
+{
+	// Checkout AssetTypeActions_Texture.cpp for some exemples
+}
+
+
+void FAssetTypeActions_SparseVolumeTexture::OpenAssetEditor(const TArray<UObject*>& InObjects, TSharedPtr<IToolkitHost> EditWithinLevelEditor)
+{
+	// Checkout AssetTypeActions_Texture.cpp for some exemples
+}
+
+#undef LOCTEXT_NAMESPACE
diff --git a/Engine/Source/Editor/SparseVolumeTexture/Private/AssetTypeActions_SparseVolumeTexture.h b/Engine/Source/Editor/SparseVolumeTexture/Private/AssetTypeActions_SparseVolumeTexture.h
new file mode 100644
--- /dev/null
+++ b/Engine/Source/Editor/SparseVolumeTexture/Private/AssetTypeActions_SparseVolumeTexture.h
@@ -0,0 +1,28 @@
+// Copyright Epic Games, Inc. All Rights Reserved.
+
+#pragma once
+
+#include "CoreMinimal.h"
+#include "AssetTypeActions_Base.h"
+#include "SparseVolumeTexture/SparseVolumeTexture.h"
+
+class FAssetTypeActions_SparseVolumeTexture : public FAssetTypeActions_Base
+{
+public:
+
+	//~ Begin IAssetTypeActions Interface
+
+	virtual FText GetName() const override { return NSLOCTEXT("AssetTypeActions", "AssetTypeActions_SparseVolumeTexture", "Base"); }
+	virtual FColor GetTypeColor() const override { return FColor(255, 128, 64); }
+	virtual UClass* GetSupportedClass() const override { return USparseVolumeTexture::StaticClass(); }
+	virtual bool HasActions(const TArray<UObject*>& InObjects) const override { return true; }
+	virtual void GetActions(const TArray<UObject*>& InObjects, struct FToolMenuSection& Section) override;
+	virtual void OpenAssetEditor(const TArray<UObject*>& InObjects, TSharedPtr<class IToolkitHost> EditWithinLevelEditor = TSharedPtr<IToolkitHost>()) override;
+	virtual bool CanFilter() override { return false; }
+	virtual uint32 GetCategories() override { return EAssetTypeCategories::Textures; }
+	virtual bool IsImportedAsset() const override { return true; }
+	virtual void GetResolvedSourceFilePaths(const TArray<UObject*>& TypeAssets, TArray<FString>& OutSourceFilePaths) const override;
+
+	//~ End IAssetTypeActions Interface
+
+}; 
diff --git a/Engine/Source/Editor/SparseVolumeTexture/Private/SparseVolumeTextureFactory.cpp b/Engine/Source/Editor/SparseVolumeTexture/Private/SparseVolumeTextureFactory.cpp
new file mode 100644
--- /dev/null
+++ b/Engine/Source/Editor/SparseVolumeTexture/Private/SparseVolumeTextureFactory.cpp
@@ -0,0 +1,630 @@
+// Copyright Epic Games, Inc. All Rights Reserved.
+
+#include "SparseVolumeTextureFactory.h"
+
+#include "SparseVolumeTexture/SparseVolumeTexture.h"
+
+#if WITH_EDITOR
+
+#include "SparseVolumeTextureOpenVDB.h"
+
+#include "Serialization/EditorBulkDataWriter.h"
+#include "Misc/Paths.h"
+#include "Misc/ScopedSlowTask.h"
+
+#include "Editor.h"
+
+#define LOCTEXT_NAMESPACE "USparseVolumeTextureFactory"
+
+DEFINE_LOG_CATEGORY_STATIC(LogSparseVolumeTextureFactory, Log, All);
+
+USparseVolumeTextureFactory::USparseVolumeTextureFactory(const FObjectInitializer& ObjectInitializer)
+	: Super(ObjectInitializer)
+{
+	bCreateNew = true;
+	bEditAfterNew = true;
+	bEditorImport = true;
+	SupportedClass = USparseVolumeTexture::StaticClass();
+
+	Formats.Add(TEXT("vdb;OpenVDB Format"));
+}
+
+FText USparseVolumeTextureFactory::GetDisplayName() const
+{
+	return LOCTEXT("SparseVolumeTextureFactoryDescription", "Sparse Volume Texture");
+}
+
+bool USparseVolumeTextureFactory::ConfigureProperties()
+{
+	return true;
+}
+
+bool USparseVolumeTextureFactory::ShouldShowInNewMenu() const
+{
+	return false;
+}
+
+
+///////////////////////////////////////////////////////////////////////////////
+// Create asset
+
+
+bool USparseVolumeTextureFactory::CanCreateNew() const
+{
+	return false;	// To be able to import files and call 
+}
+
+UObject* USparseVolumeTextureFactory::FactoryCreateNew(UClass* InClass, UObject* InParent, FName InName, EObjectFlags Flags, UObject* Context, FFeedbackContext* Warn)
+{
+	USparseVolumeTexture* Object = NewObject<USparseVolumeTexture>(InParent, InClass, InName, Flags);
+
+	// SVT_TODO initialize similarly to UTexture2DFactoryNew
+
+	return Object;
+}
+
+
+///////////////////////////////////////////////////////////////////////////////
+// Import asset
+
+
+bool USparseVolumeTextureFactory::DoesSupportClass(UClass* Class)
+{
+	return Class == USparseVolumeTexture::StaticClass();
+}
+
+UClass* USparseVolumeTextureFactory::ResolveSupportedClass()
+{
+	return USparseVolumeTexture::StaticClass();
+}
+
+bool USparseVolumeTextureFactory::FactoryCanImport(const FString& Filename)
+{
+	const FString Extension = FPaths::GetExtension(Filename);
+	if (Extension == TEXT("vdb"))
+	{
+		return true;
+	}
+	return false;
+}
+
+void USparseVolumeTextureFactory::CleanUp()
+{
+	Super::CleanUp();
+}
+
+struct FOpenVDBData
+{
+	FVector VolumeActiveAABBMin;
+	FVector VolumeActiveAABBMax;
+	FVector VolumeActiveDim;
+	FVector VolumeVoxelSize;
+	bool bIsInWorldSpace;
+	bool bHasUniformVoxels;
+};
+
+#if OPENVDB_AVAILABLE
+
+static FOpenVDBData GetOpenVDBData(openvdb::GridBase::Ptr GridBase)
+{
+	FOpenVDBData OpenVDBData;
+	openvdb::CoordBBox VolumeActiveAABB = GridBase->evalActiveVoxelBoundingBox();
+	openvdb::Coord VolumeActiveDim = GridBase->evalActiveVoxelDim();
+	openvdb::Vec3d VolumeVoxelSize = GridBase->voxelSize();
+
+	OpenVDBData.VolumeActiveAABBMin = FVector(VolumeActiveAABB.min().x(), VolumeActiveAABB.min().y(), VolumeActiveAABB.min().z());
+	OpenVDBData.VolumeActiveAABBMax = FVector(VolumeActiveAABB.max().x(), VolumeActiveAABB.max().y(), VolumeActiveAABB.max().z());
+	OpenVDBData.VolumeActiveDim = FVector(VolumeActiveDim.x(), VolumeActiveDim.y(), VolumeActiveDim.z());
+	OpenVDBData.VolumeVoxelSize = FVector(VolumeVoxelSize.x(), VolumeVoxelSize.y(), VolumeVoxelSize.z());
+	OpenVDBData.bIsInWorldSpace = GridBase->isInWorldSpace();
+	OpenVDBData.bHasUniformVoxels = GridBase->hasUniformVoxels();
+
+	return OpenVDBData;
+}
+
+static bool IsOpenVDBDataValid(FOpenVDBData& OpenVDBData, const FString& Filename)
+{
+	if (OpenVDBData.VolumeActiveDim.X * OpenVDBData.VolumeActiveDim.Y * OpenVDBData.VolumeActiveDim.Z == 0)
+	{
+		// SVT_TODO we should gently handle that case
+		UE_LOG(LogSparseVolumeTextureFactory, Error, TEXT("Could not import empty OpenVDB asset due to volume size being 0: %s"), *Filename);
+		return false;
+	}
+
+	if (!OpenVDBData.bHasUniformVoxels)
+	{
+		UE_LOG(LogSparseVolumeTextureFactory, Error, TEXT("OpenVDB importer cannot handle non uniform voxels: %s"), *Filename);
+		return false;
+	}
+	return true;
+}
+
+
+static bool ConvertOpenVDBToSparseVolumeTexture(
+	openvdb::FloatGrid::Ptr DensityGrid,
+	const FString& SourceFilename,
+	FSparseVolumeRawSource& OutSparseVolumeRawSource,
+	bool bOverrideActiveMinMax = false,
+	FVector ActiveMin = FVector::ZeroVector,
+	FVector ActiveMax = FVector::ZeroVector)
+{
+	FOpenVDBData OVDBData = GetOpenVDBData(DensityGrid);
+	if (!IsOpenVDBDataValid(OVDBData, SourceFilename))
+	{
+		return false;
+	}
+
+	if (bOverrideActiveMinMax)
+	{
+		OVDBData.VolumeActiveAABBMin = ActiveMin;
+		OVDBData.VolumeActiveAABBMax = ActiveMax;
+		OVDBData.VolumeActiveDim = ActiveMax - ActiveMin;
+	}
+
+	FSparseVolumeAssetHeader& Header = OutSparseVolumeRawSource.Header;
+
+	Header.SourceVolumeResolution = FIntVector(OVDBData.VolumeActiveDim);
+	// SVT_TODO handle error too high volume resolution assuming dx11 limitations?
+
+	Header.PageTableVolumeResolution = FIntVector3(
+		FMath::DivideAndRoundUp(Header.SourceVolumeResolution.X, SPARSE_VOLUME_TILE_RES),
+		FMath::DivideAndRoundUp(Header.SourceVolumeResolution.Y, SPARSE_VOLUME_TILE_RES),
+		FMath::DivideAndRoundUp(Header.SourceVolumeResolution.Z, SPARSE_VOLUME_TILE_RES));
+	Header.TileDataVolumeResolution = FIntVector::ZeroValue;	// unknown for now
+
+	const uint32 VolumeTileDataBytes = SPARSE_VOLUME_TILE_RES * SPARSE_VOLUME_TILE_RES * SPARSE_VOLUME_TILE_RES;
+	uint32 NumAllocatedPages = 0;
+
+	// Allocate some memory for temp data (worst case)
+	TArray<uint8> LinearAllocatedTiles;
+	TArray<FIntVector3> LinearAllocatedPages;
+	LinearAllocatedTiles.SetNum(Header.PageTableVolumeResolution.X * Header.PageTableVolumeResolution.Y * Header.PageTableVolumeResolution.Z * VolumeTileDataBytes);
+	LinearAllocatedPages.SetNum(Header.PageTableVolumeResolution.X * Header.PageTableVolumeResolution.Y * Header.PageTableVolumeResolution.Z);
+
+	auto accessor = DensityGrid->getConstAccessor();
+	const float GridBackgroundValue = DensityGrid->background();
+
+	// Go over each potential page table form the source data and push allocate it if it has any data.
+	// Otherwise point to the default empty page.
+	bool bAnyEmptyPageExists = false;
+	for (int32_t PageZ = 0; PageZ < Header.PageTableVolumeResolution.Z; ++PageZ)
+	{
+		for (int32_t PageY = 0; PageY < Header.PageTableVolumeResolution.Y; ++PageY)
+		{
+			for (int32_t PageX = 0; PageX < Header.PageTableVolumeResolution.X; ++PageX)
+			{
+				bool bHasAnyData = false;
+
+				uint8* NewTileData = LinearAllocatedTiles.GetData() + NumAllocatedPages * VolumeTileDataBytes;
+
+				for (int32_t z = 0; z < SPARSE_VOLUME_TILE_RES; ++z)
+				{
+					for (int32_t y = 0; y < SPARSE_VOLUME_TILE_RES; ++y)
+					{
+						for (int32_t x = 0; x < SPARSE_VOLUME_TILE_RES; ++x)
+						{
+							FVector VoxelCoord(PageX * SPARSE_VOLUME_TILE_RES + x, PageY * SPARSE_VOLUME_TILE_RES + y, PageZ * SPARSE_VOLUME_TILE_RES + z);	// This assumes sampling outside the boundary retuirn a default value
+							FVector VoxelCoord2 = OVDBData.VolumeActiveAABBMin + VoxelCoord;
+							float VoxelValue = accessor.getValue(openvdb::Coord(VoxelCoord2.X, VoxelCoord2.Y, VoxelCoord2.Z));
+
+							// clamp for now until we have fp16
+							VoxelValue = FMath::Clamp(VoxelValue, 0.0f, 1.0f);
+
+							bHasAnyData |= VoxelValue > GridBackgroundValue;
+
+							NewTileData[z * SPARSE_VOLUME_TILE_RES * SPARSE_VOLUME_TILE_RES + y * SPARSE_VOLUME_TILE_RES + x] = uint8(VoxelValue * 255.0f);
+						}
+					}
+				}
+
+				if (bHasAnyData)
+				{
+					LinearAllocatedPages[NumAllocatedPages] = FIntVector3(PageX, PageY, PageZ);
+					NumAllocatedPages++;
+				}
+				bAnyEmptyPageExists |= !bHasAnyData;
+			}
+		}
+	}
+
+	// Compute Page and Tile VolumeResolution from allocated pages
+	const uint32 EffectivelyAllocatedPageEntries = NumAllocatedPages + (bAnyEmptyPageExists ? 1 : 0);
+	uint32 TileVolumeResolutionCube = 1;
+	while (TileVolumeResolutionCube * TileVolumeResolutionCube * TileVolumeResolutionCube < EffectivelyAllocatedPageEntries)
+	{
+		TileVolumeResolutionCube++;				// We use a simple loop to compute the minimum resolution of a cube to store all the tile data
+	}
+	Header.TileDataVolumeResolution = FIntVector3(TileVolumeResolutionCube, TileVolumeResolutionCube, TileVolumeResolutionCube);
+	while (Header.TileDataVolumeResolution.X * Header.TileDataVolumeResolution.Y * (Header.TileDataVolumeResolution.Z - 1) > int32(EffectivelyAllocatedPageEntries))
+	{
+		Header.TileDataVolumeResolution.Z--;	// We then trim an edge to get back space.
+	}
+	const FIntVector3 TileCoordResolution = Header.TileDataVolumeResolution;
+	Header.TileDataVolumeResolution = Header.TileDataVolumeResolution * SPARSE_VOLUME_TILE_RES;
+
+	// Initialise the SparseVolumeTexture page and tile.
+	// It would be good to stream in the OpenVDB file there but could not find a way to do that yet so we just store the sparse texture.
+	OutSparseVolumeRawSource.DensityPage.SetNumZeroed(Header.PageTableVolumeResolution.X * Header.PageTableVolumeResolution.Y * Header.PageTableVolumeResolution.Z);
+	OutSparseVolumeRawSource.DensityData.SetNumZeroed(Header.TileDataVolumeResolution.X * Header.TileDataVolumeResolution.Y * Header.TileDataVolumeResolution.Z);
+
+	FIntVector DestinationTileCoord = FIntVector::ZeroValue;
+	auto GoToNextTileCoord = [&]()
+	{
+		DestinationTileCoord.X++;
+		if (DestinationTileCoord.X >= TileCoordResolution.X)
+		{
+			DestinationTileCoord.X = 0;
+			DestinationTileCoord.Y++;
+		}
+		if (DestinationTileCoord.Y >= TileCoordResolution.Y)
+		{
+			DestinationTileCoord.Y = 0;
+			DestinationTileCoord.Z++;
+		}
+	};
+
+	// Add an empty tile is needed, reserve slot at coord 0
+	if (bAnyEmptyPageExists)
+	{
+		// DensityPage is all cleared to zero, simply skip a tile
+		GoToNextTileCoord();
+	}
+
+	// Generate page table and tile volume data by splatting the data
+	for (uint32 i = 0; i < NumAllocatedPages; ++i)
+	{
+		FIntVector3 PageCoordToSplat = LinearAllocatedPages[i];
+		uint8* TileDataToSplat = LinearAllocatedTiles.GetData() + i * VolumeTileDataBytes;
+
+		// A page encodes the physical tile coord as unsigned int of 11 11 10 bits
+		// This means a page coord cannot be larger than 2047 for x and y and 1023 for z
+		// which mean we cannot have more than 2048*2048*1024 = 4 Giga tiles of 16^3 tiles.
+		uint32 DestinationTileCoord32bit = DestinationTileCoord.X & 0x7FF | ((DestinationTileCoord.Y & 0x7FF) << 11) | ((DestinationTileCoord.Z & 0x3FF) << 22);
+
+		// Setup the page table entry
+		OutSparseVolumeRawSource.DensityPage.GetData()
+			[
+				PageCoordToSplat.Z * Header.PageTableVolumeResolution.X * Header.PageTableVolumeResolution.Y +
+				PageCoordToSplat.Y * Header.PageTableVolumeResolution.X +
+			PageCoordToSplat.X
+			] = DestinationTileCoord32bit;
+
+		// Now copy the tile data
+		for (int32_t z = 0; z < SPARSE_VOLUME_TILE_RES; ++z)
+		{
+			for (int32_t y = 0; y < SPARSE_VOLUME_TILE_RES; ++y)
+			{
+				for (int32_t x = 0; x < SPARSE_VOLUME_TILE_RES; ++x)
+				{
+					OutSparseVolumeRawSource.DensityData.GetData()
+						[
+							(DestinationTileCoord.Z * SPARSE_VOLUME_TILE_RES + z) * Header.TileDataVolumeResolution.X * Header.TileDataVolumeResolution.Y +
+							(DestinationTileCoord.Y * SPARSE_VOLUME_TILE_RES + y) * Header.TileDataVolumeResolution.X +
+						(DestinationTileCoord.X * SPARSE_VOLUME_TILE_RES + x)
+						] = TileDataToSplat[z * SPARSE_VOLUME_TILE_RES * SPARSE_VOLUME_TILE_RES + y * SPARSE_VOLUME_TILE_RES + x];
+				}
+			}
+		}
+
+		// Set the next tile to be written to
+		GoToNextTileCoord();
+	}
+
+	return true;
+}
+
+struct FOpenVDBFrameData
+{
+	FString Filename;
+	openvdb::io::File* File;
+
+	FOpenVDBData OpenVDBData;
+
+	openvdb::FloatGrid::Ptr DensityGrid;
+};
+
+#endif // OPENVDB_AVAILABLE
+
+
+UObject* USparseVolumeTextureFactory::FactoryCreateFile(UClass* InClass, UObject* InParent, FName InName, EObjectFlags Flags, const FString& Filename,
+	const TCHAR* Parms, FFeedbackContext* Warn, bool& bOutOperationCanceled)
+{
+
+#if OPENVDB_AVAILABLE
+
+	GEditor->GetEditorSubsystem<UImportSubsystem>()->BroadcastAssetPreImport(this, InClass, InParent, InName, Parms);
+	TArray<UObject*> ResultAssets;
+
+	auto IsFilePotentiallyPartOfASequence = [](const FString& Filename)
+	{
+		// The file is potnetially a sequence of the character before the `.vdb` is a number.
+		return FChar::IsDigit(Filename[Filename.Len() - 5]);
+	};
+
+	if (IsFilePotentiallyPartOfASequence(Filename))
+	{
+		// Import as an animated sparse volume texture asset.
+
+		const FString FilenameWithoutExt = Filename.LeftChop(4);
+		const int32 LastNonDigitIndex = FilenameWithoutExt.FindLastCharByPredicate([](TCHAR Letter) { return !FChar::IsDigit(Letter); }) + 1;
+		const int32 DigitCount = FilenameWithoutExt.Len() - LastNonDigitIndex;
+		FString FilenameWithoutSuffix = FilenameWithoutExt.LeftChop(FilenameWithoutExt.Len() - LastNonDigitIndex);
+		TCHAR LastDigit = FilenameWithoutExt[FilenameWithoutExt.Len() - 5];
+
+		bool IndexStartsAtOne = false;
+		auto GetOpenVDBFileNameForFrame = [&](int32 FrameIndex)
+		{
+			FString IndexString = FString::FromInt(FrameIndex + (IndexStartsAtOne ? 1 : 0));
+			// User must select a frame with index in [0-9] so that we can count leading 0s
+			check(DigitCount==1 || (DigitCount>1 && IndexString.Len() <= DigitCount));
+			const int32 MissingLeadingZeroCount = DigitCount - IndexString.Len();
+			const FString StringZero = FString::FromInt(0);
+			for (int32 i = 0; i < MissingLeadingZeroCount; ++i)
+			{
+				IndexString = StringZero + IndexString;
+			}
+			return FString(FilenameWithoutSuffix + IndexString) + TEXT(".vdb");
+		};
+
+		const FString VDBFileAt0 = GetOpenVDBFileNameForFrame(0);
+		const FString VDBFileAt1 = GetOpenVDBFileNameForFrame(1);
+		const bool VDBFileAt0Exists = FPaths::FileExists(VDBFileAt0);
+		const bool VDBFileAt1Exists = FPaths::FileExists(VDBFileAt1);
+		if (!VDBFileAt0Exists && !VDBFileAt1Exists)
+		{
+			UE_LOG(LogSparseVolumeTextureFactory, Error, TEXT("An OpenVDB animated sequence must start at index 0 or 1: %s or %s not found."), *VDBFileAt0, *VDBFileAt1);
+			return nullptr;
+		}
+		IndexStartsAtOne = !VDBFileAt0Exists;
+
+		FName NewName(InName .ToString()+ TEXT("VDBAnim"));
+		UAnimatedSparseVolumeTexture* AnimatedSVTexture = NewObject<UAnimatedSparseVolumeTexture>(InParent, UAnimatedSparseVolumeTexture::StaticClass(), NewName, Flags);
+
+		// Go over all the frame index and stop at the first missing one.
+		int32 FrameCount = 0;
+		while(FPaths::FileExists(GetOpenVDBFileNameForFrame(FrameCount)))
+		{
+			FrameCount++;
+		}
+		TArray<FOpenVDBFrameData> OpenVDBFramesData;
+		OpenVDBFramesData.SetNum(FrameCount);
+
+		FScopedSlowTask ImportTask(FrameCount, LOCTEXT("ImportingVDBAnim", "Importing OpenVDB animation"));
+		ImportTask.MakeDialog(true);
+
+		// Open all each frame openvdb files.
+		for(int32 FrameIndex = 0; FrameIndex < FrameCount; ++FrameIndex)
+		{
+			FOpenVDBFrameData& OpenVDBFrameData = OpenVDBFramesData[FrameIndex];
+			OpenVDBFrameData.Filename = GetOpenVDBFileNameForFrame(FrameIndex);
+
+			std::string FileNameStr(TCHAR_TO_ANSI(*OpenVDBFrameData.Filename));
+			OpenVDBFrameData.File = new openvdb::io::File(FileNameStr.c_str());
+			try
+			{
+				OpenVDBFrameData.File->open();
+			}
+			catch (openvdb::Exception e)
+			{
+				UE_LOG(LogSparseVolumeTextureFactory, Error, TEXT("OpenVDB file coult not be opened: %s"), *OpenVDBFrameData.Filename);
+				return nullptr;
+			}
+
+			openvdb::GridBase::Ptr BaseGrid = nullptr;
+			for (openvdb::io::File::NameIterator nameIter = OpenVDBFrameData.File->beginName();
+				nameIter != OpenVDBFrameData.File->endName(); ++nameIter)
+			{
+				if (nameIter.gridName() == "density")
+				{
+					openvdb::GridBase::Ptr QueryGrid = OpenVDBFrameData.File->readGrid(nameIter.gridName());
+					if (QueryGrid->isType<openvdb::FloatGrid>())
+					{
+						BaseGrid = QueryGrid;
+						break;
+					}
+					break;
+				}
+			}
+
+			// If we have not found any density map, let's pick up the first float map
+			for (openvdb::io::File::NameIterator nameIter = OpenVDBFrameData.File->beginName();
+				nameIter != OpenVDBFrameData.File->endName(); ++nameIter)
+			{
+				openvdb::GridBase::Ptr QueryGrid = OpenVDBFrameData.File->readGrid(nameIter.gridName());
+				if (QueryGrid->isType<openvdb::FloatGrid>())
+				{
+					BaseGrid = QueryGrid;
+					break;
+				}
+			}
+
+			OpenVDBFrameData.DensityGrid = openvdb::gridPtrCast<openvdb::FloatGrid>(BaseGrid);
+
+			if (OpenVDBFrameData.DensityGrid == nullptr)
+			{
+				UE_LOG(LogSparseVolumeTextureFactory, Error, TEXT("Could not find a Grid named density or grid of type float: %s"), *OpenVDBFrameData.Filename);
+				return nullptr;
+			}
+
+			OpenVDBFrameData.OpenVDBData = GetOpenVDBData(OpenVDBFrameData.DensityGrid);
+			if (!IsOpenVDBDataValid(OpenVDBFrameData.OpenVDBData, OpenVDBFrameData.Filename))
+			{
+				return nullptr;
+			}
+		}
+
+
+		// Evaluate the maximum extent for all frames.
+		FVector AnimationAABBMin;
+		FVector AnimationAABBMax;
+		for (int32 FrameIndex = 0; FrameIndex < FrameCount; ++FrameIndex)
+		{
+			FOpenVDBFrameData& OpenVDBFrameData = OpenVDBFramesData[FrameIndex];
+			AnimationAABBMin = FrameIndex == 0 ? OpenVDBFrameData.OpenVDBData.VolumeActiveAABBMin : FVector::Min(AnimationAABBMin, OpenVDBFrameData.OpenVDBData.VolumeActiveAABBMin);
+			AnimationAABBMax = FrameIndex == 0 ? OpenVDBFrameData.OpenVDBData.VolumeActiveAABBMax : FVector::Max(AnimationAABBMin, OpenVDBFrameData.OpenVDBData.VolumeActiveAABBMax);
+		}
+
+		// Allocate space for each frame
+		AnimatedSVTexture->FrameCount = FrameCount;
+		AnimatedSVTexture->AnimationFrames.SetNum(FrameCount);
+
+		UE_LOG(LogSparseVolumeTextureFactory, Display, TEXT("Serializing: %i frame"), FrameCount);
+
+		// Convert all openvdb files to sparse volume texture
+		for (int32 FrameIndex = 0; FrameIndex < FrameCount; ++FrameIndex)
+		{
+			FOpenVDBFrameData& OpenVDBFrameData = OpenVDBFramesData[FrameIndex];
+			FSparseVolumeRawSource SparseVolumeRawSource;
+			if (!ConvertOpenVDBToSparseVolumeTexture(OpenVDBFrameData.DensityGrid, OpenVDBFrameData.Filename, SparseVolumeRawSource,
+				true, AnimationAABBMin, AnimationAABBMax))
+			{
+				return nullptr;
+			}
+
+			// Serialise the raw source data fro this frame into the asset object.
+			{
+				UE_LOG(LogSparseVolumeTextureFactory, Display, TEXT("  - frame %i (active dimension %i x %i x %i)"), FrameIndex, 
+					OpenVDBFrameData.OpenVDBData.VolumeActiveDim.X, OpenVDBFrameData.OpenVDBData.VolumeActiveDim.Y, OpenVDBFrameData.OpenVDBData.VolumeActiveDim.Z);
+
+				UE_LOG(LogSparseVolumeTextureFactory, Display, TEXT("    -  SourceVolumeResolution %i x %i x %i"),
+					SparseVolumeRawSource.Header.SourceVolumeResolution.X, SparseVolumeRawSource.Header.SourceVolumeResolution.Y, SparseVolumeRawSource.Header.SourceVolumeResolution.Z);
+				UE_LOG(LogSparseVolumeTextureFactory, Display, TEXT("    -  PageTableVolumeResolution %i x %i x %i"),
+					SparseVolumeRawSource.Header.PageTableVolumeResolution.X, SparseVolumeRawSource.Header.PageTableVolumeResolution.Y, SparseVolumeRawSource.Header.PageTableVolumeResolution.Z);
+				UE_LOG(LogSparseVolumeTextureFactory, Display, TEXT("    -  TileDataVolumeResolution %i x %i x %i"),
+					SparseVolumeRawSource.Header.TileDataVolumeResolution.X, SparseVolumeRawSource.Header.TileDataVolumeResolution.Y, SparseVolumeRawSource.Header.TileDataVolumeResolution.Z);
+
+				UE::Serialization::FEditorBulkDataWriter RawDataArchiveWriter(AnimatedSVTexture->AnimationFrames[FrameIndex].RawData);
+				SparseVolumeRawSource.Serialize(RawDataArchiveWriter);
+
+				if (ImportTask.ShouldCancel())
+				{
+					return nullptr;
+				}
+				ImportTask.EnterProgressFrame(1.0f, LOCTEXT("ConvertingVDBStatic", "Converting static OpenVDB"));
+			}
+		}
+
+		// Close all the files.
+		for (int32 FrameIndex = 0; FrameIndex < FrameCount; ++FrameIndex)
+		{
+			FOpenVDBFrameData& OpenVDBFrameData = OpenVDBFramesData[FrameIndex];
+			OpenVDBFrameData.File->close();
+			delete OpenVDBFrameData.File;
+		}
+
+		ResultAssets.Add(AnimatedSVTexture);
+	}
+	else
+	{
+		// Import as a static sparse volume texture asset.
+
+		FName NewName(InName.ToString() + TEXT("VDB"));
+		UStaticSparseVolumeTexture* StaticSVTexture = NewObject<UStaticSparseVolumeTexture>(InParent, UStaticSparseVolumeTexture::StaticClass(), NewName, Flags);
+
+		std::string FileNameStr(TCHAR_TO_ANSI(*Filename));
+		openvdb::io::File File(FileNameStr.c_str());
+
+		try
+		{
+			File.open();
+		}
+		catch (openvdb::Exception e)
+		{
+			UE_LOG(LogSparseVolumeTextureFactory, Warning, TEXT("OpenVDB file coult not be opened: %s"), *Filename);
+			return nullptr;
+		}
+
+		FScopedSlowTask ImportTask(1.0f, LOCTEXT("ImportingVDBStatic", "Importing static OpenVDB"));
+		ImportTask.MakeDialog(true);
+
+		openvdb::GridBase::Ptr BaseGrid = nullptr;
+		for (openvdb::io::File::NameIterator nameIter = File.beginName();
+			nameIter != File.endName(); ++nameIter)
+		{
+			if (nameIter.gridName() == "density")
+			{
+				openvdb::GridBase::Ptr QueryGrid = File.readGrid(nameIter.gridName());
+				if (QueryGrid->isType<openvdb::FloatGrid>())
+				{
+					BaseGrid = QueryGrid;
+					break;
+				}
+				break;
+			}
+		}
+
+		// If we have not found any density map, let's pick up the first float map
+		for (openvdb::io::File::NameIterator nameIter = File.beginName();
+			nameIter != File.endName(); ++nameIter)
+		{
+			openvdb::GridBase::Ptr QueryGrid = File.readGrid(nameIter.gridName());
+			if (QueryGrid->isType<openvdb::FloatGrid>())
+			{
+				BaseGrid = QueryGrid;
+				break;
+			}
+		}
+
+		openvdb::FloatGrid::Ptr DensityGrid = openvdb::gridPtrCast<openvdb::FloatGrid>(BaseGrid);
+
+		// Only open float grid for now
+		if (DensityGrid == nullptr)
+		{
+			UE_LOG(LogSparseVolumeTextureFactory, Error, TEXT("Could not convert the first grid to float: %s"), *Filename);
+			return nullptr;
+		}
+
+		FSparseVolumeRawSource SparseVolumeRawSource;
+		if (!ConvertOpenVDBToSparseVolumeTexture(DensityGrid, Filename, SparseVolumeRawSource))
+		{
+			return nullptr;
+		}
+
+		if (ImportTask.ShouldCancel())
+		{
+			return nullptr;
+		}
+		ImportTask.EnterProgressFrame(1.0f, LOCTEXT("ConvertingVDBStatic", "Converting static OpenVDB"));
+
+		File.close();
+
+		// Serialise the raw source data into the asset object.
+		{
+			UE::Serialization::FEditorBulkDataWriter RawDataArchiveWriter(StaticSVTexture->StaticFrame.RawData);
+			SparseVolumeRawSource.Serialize(RawDataArchiveWriter);
+		}
+		ResultAssets.Add(StaticSVTexture);
+	}
+
+
+
+	// Now notify the system about the imported/updated/created assets
+	AdditionalImportedObjects.Reserve(ResultAssets.Num());
+	for (UObject* Object : ResultAssets)
+	{
+		if (Object)
+		{
+			GEditor->GetEditorSubsystem<UImportSubsystem>()->BroadcastAssetPostImport(this, Object);
+			Object->MarkPackageDirty();
+			Object->PostEditChange();
+			AdditionalImportedObjects.Add(Object);
+		}
+	}
+
+	return (ResultAssets.Num() > 0) ? ResultAssets[0] : nullptr;
+
+#else // OPENVDB_AVAILABLE
+
+	// SVT_TODO Make sure we can also import on more platforms such as Linux. See SparseVolumeTextureOpenVDB.h
+	UE_LOG(LogSparseVolumeTextureFactory, Error, TEXT("Cannot import OpenVDB asset any platform other than Windows."));
+	return nullptr;
+
+#endif // OPENVDB_AVAILABLE
+
+}
+
+#endif // WITH_EDITORONLY_DATA
+
+#undef LOCTEXT_NAMESPACE
+
+#include "Serialization/EditorBulkDataWriter.h"
diff --git a/Engine/Source/Editor/SparseVolumeTexture/Private/SparseVolumeTextureFactory.h b/Engine/Source/Editor/SparseVolumeTexture/Private/SparseVolumeTextureFactory.h
new file mode 100644
--- /dev/null
+++ b/Engine/Source/Editor/SparseVolumeTexture/Private/SparseVolumeTextureFactory.h
@@ -0,0 +1,51 @@
+// Copyright Epic Games, Inc. All Rights Reserved.
+
+#pragma once
+
+#if WITH_EDITOR
+
+#include "CoreMinimal.h"
+#include "Factories/Factory.h"
+
+#include "SparseVolumeTextureFactory.generated.h"
+
+// Responsible for creating and importing Sparse Volume Texture objects
+UCLASS(hidecategories = Object, MinimalAPI)
+class USparseVolumeTextureFactory : public UFactory
+{
+	GENERATED_UCLASS_BODY()
+
+public:
+
+	//~ Begin UFactory Interface
+
+	virtual bool ShouldShowInNewMenu() const override;
+	virtual FText GetDisplayName() const override;
+	virtual bool ConfigureProperties() override;
+
+	///////////////////////////////////////////////////////////////////////////////
+	// Create asset
+
+	virtual UObject* FactoryCreateNew(UClass* InClass, UObject* InParent, FName InName, EObjectFlags Flags, UObject* Context, FFeedbackContext* Warn) override;
+
+	virtual bool CanCreateNew() const override;
+
+	///////////////////////////////////////////////////////////////////////////////
+	// Import asset
+
+	virtual UObject* FactoryCreateFile(UClass* InClass, UObject* InParent, FName InName, EObjectFlags Flags, const FString& Filename,
+		const TCHAR* Parms, FFeedbackContext* Warn, bool& bOutOperationCanceled) override;
+
+	virtual void CleanUp() override;
+	virtual bool FactoryCanImport(const FString& Filename) override;
+	virtual bool DoesSupportClass(UClass* Class) override;
+	virtual UClass* ResolveSupportedClass() override;
+
+	//~ End UFactory Interface
+
+protected:
+
+private:
+};
+
+#endif // WITH_EDITOR
diff --git a/Engine/Source/Editor/SparseVolumeTexture/Private/SparseVolumeTextureModule.cpp b/Engine/Source/Editor/SparseVolumeTexture/Private/SparseVolumeTextureModule.cpp
new file mode 100644
--- /dev/null
+++ b/Engine/Source/Editor/SparseVolumeTexture/Private/SparseVolumeTextureModule.cpp
@@ -0,0 +1,27 @@
+// Copyright Epic Games, Inc. All Rights Reserved.
+
+#include "SparseVolumeTextureModule.h"
+
+#include "SparseVolumeTextureOpenVDB.h"
+
+#define LOCTEXT_NAMESPACE "SparseVolumeTextureModule"
+
+IMPLEMENT_MODULE(FSparseVolumeTextureModule, SparseVolumeTexture);
+
+void FSparseVolumeTextureModule::StartupModule()
+{
+	IAssetTools& AssetTools = FModuleManager::LoadModuleChecked<FAssetToolsModule>(TEXT("AssetTools")).Get();
+	AssetTypeActionsForSparseVolumeTexture = MakeShareable(new FAssetTypeActions_SparseVolumeTexture);
+	AssetTools.RegisterAssetTypeActions(AssetTypeActionsForSparseVolumeTexture.ToSharedRef());
+
+#if PLATFORM_WINDOWS
+	// Global registration of  the vdb types.
+	openvdb::initialize();
+#endif
+}
+
+void FSparseVolumeTextureModule::ShutdownModule()
+{
+}
+
+#undef LOCTEXT_NAMESPACE
diff --git a/Engine/Source/Editor/SparseVolumeTexture/Private/SparseVolumeTextureOpenVDB.h b/Engine/Source/Editor/SparseVolumeTexture/Private/SparseVolumeTextureOpenVDB.h
new file mode 100644
--- /dev/null
+++ b/Engine/Source/Editor/SparseVolumeTexture/Private/SparseVolumeTextureOpenVDB.h
@@ -0,0 +1,60 @@
+// Copyright Epic Games, Inc. All Rights Reserved.
+
+#pragma once
+
+#if PLATFORM_WINDOWS
+
+#define OPENVDB_AVAILABLE 1
+
+THIRD_PARTY_INCLUDES_START
+
+UE_PUSH_MACRO("check")
+#undef check
+
+#ifdef _MSC_VER
+#pragma warning (push)
+#pragma warning (disable : 4146) // unary minus operator applied to unsigned type, result still unsigned
+#pragma warning (disable : 6297) // Arithmetic overflow:  32-bit value is shifted, then cast to 64-bit value.  Results might not be an expected value.
+#endif
+
+#ifndef M_PI
+#define M_PI    3.14159265358979323846
+#define LOCAL_M_PI 1
+#endif
+
+#ifndef M_PI_2
+#define M_PI_2  1.57079632679489661923 // pi/2
+#define LOCAL_M_PI_2 1
+#endif
+
+#define BOOST_ALLOW_DEPRECATED_HEADERS
+
+#include <openvdb/openvdb.h>
+#include <openvdb/io/Stream.h>
+#include <openvdb/io/File.h>
+#include <openvdb/tools/Prune.h>
+#include <openvdb/tools/Dense.h>
+#include <openvdb/math/Vec4.h>
+#include <openvdb//math/Half.h>
+
+#if LOCAL_M_PI
+#undef M_PI
+#endif
+
+#if LOCAL_M_PI_2
+#undef M_PI_2
+#endif
+
+#ifdef _MSC_VER
+#pragma warning (pop)
+#endif
+
+UE_POP_MACRO("check")
+
+THIRD_PARTY_INCLUDES_END
+
+#else
+
+#define OPENVDB_AVAILABLE 0
+
+#endif
\ No newline at end of file
diff --git a/Engine/Source/Editor/SparseVolumeTexture/Public/SparseVolumeTextureModule.h b/Engine/Source/Editor/SparseVolumeTexture/Public/SparseVolumeTextureModule.h
new file mode 100644
--- /dev/null
+++ b/Engine/Source/Editor/SparseVolumeTexture/Public/SparseVolumeTextureModule.h
@@ -0,0 +1,39 @@
+// Copyright Epic Games, Inc. All Rights Reserved.
+
+#pragma once
+
+#include "CoreMinimal.h"
+#include "Modules/ModuleInterface.h"
+#include "IAssetTools.h"
+#include "AssetRegistry/AssetRegistryModule.h"
+
+#include "AssetTypeActions_SparseVolumeTexture.h"
+
+/**
+ * A sparse volume texture
+ */
+class FSparseVolumeTextureModule : public IModuleInterface
+{
+public:
+
+	static FSparseVolumeTextureModule& GetSparseVolumeTextureModule()
+	{
+		static const FName ModuleName = TEXT("SparseVolumeTexture");
+		auto& ModuleInterface = FModuleManager::LoadModuleChecked<FSparseVolumeTextureModule>(ModuleName);
+		return ModuleInterface;
+	}
+
+private:
+
+	/**
+	 * Called right after the module DLL has been loaded and the module object has been created
+	 */
+	virtual void StartupModule();
+
+	/**
+	 * Called before the module is unloaded, right before the module object is destroyed.
+	 */
+	virtual void ShutdownModule();
+
+	TSharedPtr<FAssetTypeActions_SparseVolumeTexture> AssetTypeActionsForSparseVolumeTexture;
+};
diff --git a/Engine/Source/Editor/SparseVolumeTexture/SparseVolumeTexture.Build.cs b/Engine/Source/Editor/SparseVolumeTexture/SparseVolumeTexture.Build.cs
new file mode 100644
--- /dev/null
+++ b/Engine/Source/Editor/SparseVolumeTexture/SparseVolumeTexture.Build.cs
@@ -0,0 +1,71 @@
+// Copyright Epic Games, Inc. All Rights Reserved.
+
+using UnrealBuildTool;
+
+public class SparseVolumeTexture : ModuleRules
+{
+	public SparseVolumeTexture(ReadOnlyTargetRules Target) : base(Target)
+	{
+		PrivateIncludePaths.Add(ModuleDirectory + "/Private");
+		PublicIncludePaths.Add(ModuleDirectory + "/Public");
+
+		PublicDependencyModuleNames.AddRange(
+			new string[]
+			{
+				"Core",
+				"Projects",
+				"Engine",
+				"Renderer",
+				"RenderCore",
+			}
+		);
+
+		PrivateIncludePathModuleNames.AddRange(
+			new string[] {
+				"AssetRegistry",
+                "AssetTools",
+				"EditorWidgets",
+			}
+		);
+
+		PrivateDependencyModuleNames.AddRange(
+			new string[] {
+				"Renderer",
+				"RenderCore",
+				"RHI",
+				"Core",
+				"CoreUObject",
+				"Engine",
+                "InputCore",
+                "Settings",
+			}
+		);
+
+		if (Target.bBuildEditor)
+		{
+			PrivateDependencyModuleNames.AddRange(
+				new string[] 
+				{
+					"MainFrame",
+					"EditorFramework",
+					"UnrealEd"
+				}
+			);
+		}
+
+		// Specific to OpenVDB support
+		if (Target.IsInPlatformGroup(UnrealPlatformGroup.Windows))
+		{
+			bUseRTTI = true;
+			bEnableExceptions = true;
+
+			AddEngineThirdPartyPrivateStaticDependencies(Target,
+				"IntelTBB",
+				"Blosc",
+				"zlib",
+				"Boost",
+				"OpenVDB"
+			);
+		}
+	}
+}
diff --git a/Engine/Source/Editor/UnrealEd/Classes/MaterialEditor/DEditorSparseVolumeTextureParameterValue.h b/Engine/Source/Editor/UnrealEd/Classes/MaterialEditor/DEditorSparseVolumeTextureParameterValue.h
new file mode 100644
--- /dev/null
+++ b/Engine/Source/Editor/UnrealEd/Classes/MaterialEditor/DEditorSparseVolumeTextureParameterValue.h
@@ -0,0 +1,37 @@
+// Copyright Epic Games, Inc. All Rights Reserved.
+
+#pragma once
+
+#include "CoreMinimal.h"
+#include "UObject/ObjectMacros.h"
+#include "MaterialEditor/DEditorParameterValue.h"
+
+#include "DEditorSparseVolumeTextureParameterValue.generated.h"
+
+UCLASS(hidecategories=Object, collapsecategories)
+class UNREALED_API UDEditorSparseVolumeTextureParameterValue : public UDEditorParameterValue
+{
+	GENERATED_UCLASS_BODY()
+
+	UPROPERTY(EditAnywhere, Category=DEditorTextureParameterValue)
+	TObjectPtr<class USparseVolumeTexture> ParameterValue;
+
+	virtual FName GetDefaultGroupName() const override { return TEXT("Texture Parameter Values"); }
+
+	virtual bool GetValue(FMaterialParameterMetadata& OutResult) const override
+	{
+		UDEditorParameterValue::GetValue(OutResult);
+		OutResult.Value = ParameterValue;
+		return true;
+	}
+
+	virtual bool SetValue(const FMaterialParameterValue& Value) override
+	{
+		if (Value.Type == EMaterialParameterType::SparseVolumeTexture)
+		{
+			ParameterValue = Value.SparseVolumeTexture;
+			return true;
+		}
+		return false;
+	}
+};
diff --git a/Engine/Source/Editor/UnrealEd/Private/EditorParameters.cpp b/Engine/Source/Editor/UnrealEd/Private/EditorParameters.cpp
--- a/Engine/Source/Editor/UnrealEd/Private/EditorParameters.cpp
+++ b/Engine/Source/Editor/UnrealEd/Private/EditorParameters.cpp
@@ -6,6 +6,7 @@
 #include "MaterialEditor/DEditorMaterialLayersParameterValue.h"
 #include "MaterialEditor/DEditorParameterValue.h"
 #include "MaterialEditor/DEditorRuntimeVirtualTextureParameterValue.h"
+#include "MaterialEditor/DEditorSparseVolumeTextureParameterValue.h"
 #include "MaterialEditor/DEditorScalarParameterValue.h"
 #include "MaterialEditor/DEditorStaticComponentMaskParameterValue.h"
 #include "MaterialEditor/DEditorStaticSwitchParameterValue.h"
@@ -71,6 +72,11 @@ UDEditorRuntimeVirtualTextureParameterValue::UDEditorRuntimeVirtualTextureParame
 {
 }
 
+UDEditorSparseVolumeTextureParameterValue::UDEditorSparseVolumeTextureParameterValue(const FObjectInitializer& ObjectInitializer)
+	: Super(ObjectInitializer)
+{
+}
+
 static UDEditorParameterValue* CreateParameter_Scalar(UObject* Owner, const FMaterialParameterMetadata& Meta)
 {
 	UDEditorScalarParameterValue* Parameter = NewObject<UDEditorScalarParameterValue>(Owner);
@@ -129,6 +135,16 @@ static UDEditorParameterValue* CreateParameter_RuntimeVirtualTexture(UObject* Ow
 	return Parameter;
 }
 
+static UDEditorParameterValue* CreateParameter_SparseVolumeTexture(UObject* Owner, const FMaterialParameterMetadata& Meta)
+{
+	UDEditorSparseVolumeTextureParameterValue* Parameter = NewObject<UDEditorSparseVolumeTextureParameterValue>(Owner);
+	if (Meta.Value.Type == EMaterialParameterType::SparseVolumeTexture)
+	{
+		Parameter->ParameterValue = Meta.Value.SparseVolumeTexture;
+	}
+	return Parameter;
+}
+
 static UDEditorParameterValue* CreateParameter_Font(UObject* Owner, const FMaterialParameterMetadata& Meta)
 {
 	UDEditorFontParameterValue* Parameter = NewObject<UDEditorFontParameterValue>(Owner);
@@ -176,6 +192,7 @@ UDEditorParameterValue* UDEditorParameterValue::Create(UObject* Owner,
 	case EMaterialParameterType::DoubleVector: Parameter = CreateParameter_DoubleVector(Owner, Meta); break;
 	case EMaterialParameterType::Texture: Parameter = CreateParameter_Texture(Owner, Meta); break;
 	case EMaterialParameterType::RuntimeVirtualTexture: Parameter = CreateParameter_RuntimeVirtualTexture(Owner, Meta); break;
+	case EMaterialParameterType::SparseVolumeTexture: Parameter = CreateParameter_SparseVolumeTexture(Owner, Meta); break;
 	case EMaterialParameterType::Font: Parameter = CreateParameter_Font(Owner, Meta); break;
 	case EMaterialParameterType::StaticSwitch: Parameter = CreateParameter_StaticSwitch(Owner, Meta); break;
 	case EMaterialParameterType::StaticComponentMask: Parameter = CreateParameter_StaticComponentMask(Owner, Meta); break;
diff --git a/Engine/Source/Editor/UnrealEd/Private/MaterialGraphNode.cpp b/Engine/Source/Editor/UnrealEd/Private/MaterialGraphNode.cpp
--- a/Engine/Source/Editor/UnrealEd/Private/MaterialGraphNode.cpp
+++ b/Engine/Source/Editor/UnrealEd/Private/MaterialGraphNode.cpp
@@ -25,6 +25,7 @@
 #include "Materials/MaterialExpressionFunctionOutput.h"
 #include "Materials/MaterialExpressionMaterialAttributeLayers.h"
 #include "Materials/MaterialExpressionRuntimeVirtualTextureSample.h"
+#include "Materials/MaterialExpressionSparseVolumeTextureSample.h"
 #include "Materials/MaterialExpressionScalarParameter.h"
 #include "Materials/MaterialExpressionStaticBool.h"
 #include "Materials/MaterialExpressionStaticBoolParameter.h"
@@ -433,6 +434,7 @@ void UMaterialGraphNode::GetNodeContextMenuActions(UToolMenu* Menu, UGraphNodeCo
 				|| MaterialExpression->IsA(UMaterialExpressionConstant4Vector::StaticClass())
 				|| (MaterialExpression->IsA(UMaterialExpressionTextureSample::StaticClass()) && !MaterialExpression->HasAParameterName())
 				|| (MaterialExpression->IsA(UMaterialExpressionRuntimeVirtualTextureSample::StaticClass()) && !MaterialExpression->HasAParameterName())
+				|| (MaterialExpression->IsA(UMaterialExpressionSparseVolumeTextureSampleParameter::StaticClass()) && !MaterialExpression->HasAParameterName())
 				|| MaterialExpression->IsA(UMaterialExpressionTextureObject::StaticClass())
 				|| MaterialExpression->IsA(UMaterialExpressionComponentMask::StaticClass()))
 			{
diff --git a/Engine/Source/Editor/UnrealEd/Private/PreviewMaterial.cpp b/Engine/Source/Editor/UnrealEd/Private/PreviewMaterial.cpp
--- a/Engine/Source/Editor/UnrealEd/Private/PreviewMaterial.cpp
+++ b/Engine/Source/Editor/UnrealEd/Private/PreviewMaterial.cpp
@@ -833,6 +833,7 @@ void UMaterialEditorInstanceConstant::ApplySourceFunctionChanges()
 		SourceFunction->DoubleVectorParameterValues = SourceInstance->DoubleVectorParameterValues;
 		SourceFunction->TextureParameterValues = SourceInstance->TextureParameterValues;
 		SourceFunction->RuntimeVirtualTextureParameterValues = SourceInstance->RuntimeVirtualTextureParameterValues;
+		SourceFunction->SparseVolumeTextureParameterValues = SourceInstance->SparseVolumeTextureParameterValues;
 		SourceFunction->FontParameterValues = SourceInstance->FontParameterValues;
 
 		const FStaticParameterSetEditorOnlyData& StaticParameters = SourceInstance->GetEditorOnlyStaticParameters();
diff --git a/Engine/Source/Runtime/Engine/Classes/Materials/Material.h b/Engine/Source/Runtime/Engine/Classes/Materials/Material.h
--- a/Engine/Source/Runtime/Engine/Classes/Materials/Material.h
+++ b/Engine/Source/Runtime/Engine/Classes/Materials/Material.h
@@ -1178,6 +1178,7 @@ public:
 	ENGINE_API bool SetScalarParameterValueEditorOnly(FName ParameterName, float InValue);
 	ENGINE_API bool SetTextureParameterValueEditorOnly(FName ParameterName, class UTexture* InValue);
 	ENGINE_API bool SetRuntimeVirtualTextureParameterValueEditorOnly(FName ParameterName, class URuntimeVirtualTexture* InValue);
+	ENGINE_API bool SetSparseVolumeTextureParameterValueEditorOnly(FName ParameterName, class USparseVolumeTexture* InValue);
 	ENGINE_API bool SetFontParameterValueEditorOnly(FName ParameterName, class UFont* InFontValue, int32 InFontPage);
 	ENGINE_API bool SetStaticComponentMaskParameterValueEditorOnly(FName ParameterName, bool R, bool G, bool B, bool A, FGuid OutExpressionGuid);
 	ENGINE_API bool SetStaticSwitchParameterValueEditorOnly(FName ParameterName, bool OutValue, FGuid OutExpressionGuid);
diff --git a/Engine/Source/Runtime/Engine/Classes/Materials/MaterialExpressionSparseVolumeTextureSample.h b/Engine/Source/Runtime/Engine/Classes/Materials/MaterialExpressionSparseVolumeTextureSample.h
new file mode 100644
--- /dev/null
+++ b/Engine/Source/Runtime/Engine/Classes/Materials/MaterialExpressionSparseVolumeTextureSample.h
@@ -0,0 +1,112 @@
+// Copyright Epic Games, Inc. All Rights Reserved.
+
+#pragma once
+
+#include "CoreMinimal.h"
+#include "Materials/MaterialExpression.h"
+#include "UObject/ObjectMacros.h"
+
+#include "MaterialExpressionSparseVolumeTextureSample.generated.h"
+
+class USparseVolumeTexture;
+
+/** Material expression for sampling from a runtime virtual texture. */
+UCLASS(collapsecategories, hidecategories=Object)
+class ENGINE_API UMaterialExpressionSparseVolumeTextureSample : public UMaterialExpression
+{
+	GENERATED_UCLASS_BODY()
+
+	/** 3D texture coordinate used to sample the sparse volume texture. */
+	UPROPERTY(meta = (RequiredInput = "false"))
+	FExpressionInput Coordinates;
+
+	/** The Sparse Virtual Texture to sample. */
+	UPROPERTY(EditAnywhere, Category = SparseVolumeTexture)
+	TObjectPtr<class USparseVolumeTexture> SparseVolumeTexture;
+
+protected:
+
+	//~ Begin UMaterialExpression Interface
+	virtual UObject* GetReferencedTexture() const override;
+	virtual bool CanReferenceTexture() const { return true; }
+
+#if WITH_EDITOR
+	virtual void PostLoad() override;
+	virtual int32 Compile(class FMaterialCompiler* Compiler, int32 OutputIndex) override;
+	virtual void GetCaption(TArray<FString>& OutCaptions) const override;
+	virtual uint32 GetOutputType(int32 OutputIndex) override;
+	virtual uint32 GetInputType(int32 InputIndex) override;
+public:
+	virtual void PostEditChangeProperty(FPropertyChangedEvent& PropertyChangedEvent) override;
+#endif
+	//~ End UMaterialExpression Interface
+};
+
+
+UCLASS(collapsecategories, hidecategories = Object, MinimalAPI)
+class UMaterialExpressionSparseVolumeTextureSampleParameter : public UMaterialExpressionSparseVolumeTextureSample
+{
+	GENERATED_UCLASS_BODY()
+
+	/** Name to be referenced when we want to find and set this parameter */
+	UPROPERTY(EditAnywhere, Category = MaterialParameter)
+	FName ParameterName;
+
+	/** GUID that should be unique within the material, this is used for parameter renaming. */
+	UPROPERTY()
+	FGuid ExpressionGUID;
+
+	/** The name of the parameter Group to display in MaterialInstance Editor. Default is None group */
+	UPROPERTY(EditAnywhere, Category = MaterialParameter)
+	FName Group;
+
+	/** Controls where the this parameter is displayed in a material instance parameter list. The lower the number the higher up in the parameter list. */
+	UPROPERTY(EditAnywhere, Category = MaterialParameter)
+	int32 SortPriority = 32;
+
+#if WITH_EDITOR
+	/** If this is the named parameter from this material expression, then set its value. */
+	bool SetParameterValue(FName InParameterName, USparseVolumeTexture* InValue, EMaterialExpressionSetParameterValueFlags Flags = EMaterialExpressionSetParameterValueFlags::None);
+#endif
+
+	//~ Begin UMaterialExpression Interface
+#if WITH_EDITOR
+	virtual bool CanRenameNode() const override { return true; }
+	virtual void SetEditableName(const FString& NewName) override { ParameterName = *NewName; }
+	virtual FString GetEditableName() const override { return ParameterName.ToString(); }
+	virtual bool HasAParameterName() const override { return true; }
+	virtual void SetParameterName(const FName& Name) override { ParameterName = Name; }
+	virtual FName GetParameterName() const override { return ParameterName; }
+	virtual void ValidateParameterName(const bool bAllowDuplicateName) override;
+	virtual void GetCaption(TArray<FString>& OutCaptions) const override;
+	virtual bool MatchesSearchQuery(const TCHAR* SearchQuery) override;
+	virtual bool GetParameterValue(FMaterialParameterMetadata& OutMeta) const override
+	{
+		OutMeta.Value = SparseVolumeTexture;
+		OutMeta.Description = Desc;
+		OutMeta.ExpressionGuid = ExpressionGUID;
+		OutMeta.Group = Group;
+		OutMeta.SortPriority = SortPriority;
+		OutMeta.AssetPath = GetAssetPathName();
+		return true;
+	}
+	virtual bool SetParameterValue(const FName& Name, const FMaterialParameterMetadata& Meta, EMaterialExpressionSetParameterValueFlags Flags) override
+	{
+		if (Meta.Value.Type == EMaterialParameterType::SparseVolumeTexture)
+		{
+			if (SetParameterValue(Name, Meta.Value.SparseVolumeTexture, Flags))
+			{
+				if (EnumHasAnyFlags(Flags, EMaterialExpressionSetParameterValueFlags::AssignGroupAndSortPriority))
+				{
+					Group = Meta.Group;
+					SortPriority = Meta.SortPriority;
+				}
+				return true;
+			}
+		}
+		return false;
+	}
+#endif
+	virtual FGuid& GetParameterExpressionId() override { return ExpressionGUID; }
+	//~ End UMaterialExpression Interface
+};
diff --git a/Engine/Source/Runtime/Engine/Classes/Materials/MaterialFunction.h b/Engine/Source/Runtime/Engine/Classes/Materials/MaterialFunction.h
--- a/Engine/Source/Runtime/Engine/Classes/Materials/MaterialFunction.h
+++ b/Engine/Source/Runtime/Engine/Classes/Materials/MaterialFunction.h
@@ -190,6 +190,7 @@ public:
 	ENGINE_API bool SetScalarParameterValueEditorOnly(FName ParameterName, float InValue);
 	ENGINE_API bool SetTextureParameterValueEditorOnly(FName ParameterName, class UTexture* InValue);
 	ENGINE_API bool SetRuntimeVirtualTextureParameterValueEditorOnly(FName ParameterName, class URuntimeVirtualTexture* InValue);
+	ENGINE_API bool SetSparseVolumeTextureParameterValueEditorOnly(FName ParameterName, class USparseVolumeTexture* InValue);
 	ENGINE_API bool SetFontParameterValueEditorOnly(FName ParameterName, class UFont* InFontValue, int32 InFontPage);
 	ENGINE_API bool SetStaticComponentMaskParameterValueEditorOnly(FName ParameterName, bool R, bool G, bool B, bool A, FGuid OutExpressionGuid);
 	ENGINE_API bool SetStaticSwitchParameterValueEditorOnly(FName ParameterName, bool OutValue, FGuid OutExpressionGuid);
diff --git a/Engine/Source/Runtime/Engine/Classes/Materials/MaterialFunctionInstance.h b/Engine/Source/Runtime/Engine/Classes/Materials/MaterialFunctionInstance.h
--- a/Engine/Source/Runtime/Engine/Classes/Materials/MaterialFunctionInstance.h
+++ b/Engine/Source/Runtime/Engine/Classes/Materials/MaterialFunctionInstance.h
@@ -69,6 +69,10 @@ class UMaterialFunctionInstance : public UMaterialFunctionInterface
 	UPROPERTY(EditAnywhere, Category = MaterialFunctionInstance)
 	TArray<struct FRuntimeVirtualTextureParameterValue> RuntimeVirtualTextureParameterValues;
 
+	/** Sparse volume texture parameters. */
+	UPROPERTY(EditAnywhere, Category = MaterialFunctionInstance)
+	TArray<struct FSparseVolumeTextureParameterValue> SparseVolumeTextureParameterValues;
+
 #if WITH_EDITOR
 	ENGINE_API void UpdateParameterSet();
 	ENGINE_API void OverrideMaterialInstanceParameterValues(class UMaterialInstance* Instance);
diff --git a/Engine/Source/Runtime/Engine/Classes/Materials/MaterialFunctionInterface.h b/Engine/Source/Runtime/Engine/Classes/Materials/MaterialFunctionInterface.h
--- a/Engine/Source/Runtime/Engine/Classes/Materials/MaterialFunctionInterface.h
+++ b/Engine/Source/Runtime/Engine/Classes/Materials/MaterialFunctionInterface.h
@@ -348,6 +348,7 @@ public:
 	ENGINE_API bool OverrideNamedVectorParameter(const FHashedMaterialParameterInfo& ParameterInfo, FLinearColor& OutValue);
 	ENGINE_API bool OverrideNamedTextureParameter(const FHashedMaterialParameterInfo& ParameterInfo, class UTexture*& OutValue);
 	ENGINE_API bool OverrideNamedRuntimeVirtualTextureParameter(const FHashedMaterialParameterInfo& ParameterInfo, class URuntimeVirtualTexture*& OutValue);
+	ENGINE_API bool OverrideNamedSparseVolumeTextureParameter(const FHashedMaterialParameterInfo& ParameterInfo, class USparseVolumeTexture*& OutValue);
 	ENGINE_API bool OverrideNamedFontParameter(const FHashedMaterialParameterInfo& ParameterInfo, class UFont*& OutFontValue, int32& OutFontPage);
 	ENGINE_API bool OverrideNamedStaticSwitchParameter(const FHashedMaterialParameterInfo& ParameterInfo, bool& OutValue, FGuid& OutExpressionGuid);
 	ENGINE_API bool OverrideNamedStaticComponentMaskParameter(const FHashedMaterialParameterInfo& ParameterInfo, bool& OutR, bool& OutG, bool& OutB, bool& OutA, FGuid& OutExpressionGuid);
diff --git a/Engine/Source/Runtime/Engine/Classes/Materials/MaterialInstance.h b/Engine/Source/Runtime/Engine/Classes/Materials/MaterialInstance.h
--- a/Engine/Source/Runtime/Engine/Classes/Materials/MaterialInstance.h
+++ b/Engine/Source/Runtime/Engine/Classes/Materials/MaterialInstance.h
@@ -327,6 +327,54 @@ struct FRuntimeVirtualTextureParameterValue
 	}
 };
 
+/** Editable sparse volume texture parameter. */
+USTRUCT(BlueprintType)
+struct FSparseVolumeTextureParameterValue
+{
+	GENERATED_USTRUCT_BODY()
+
+	UPROPERTY(EditAnywhere, BlueprintReadWrite, Category = RuntimeVirtualTextureParameterValue)
+	FMaterialParameterInfo ParameterInfo;
+
+	UPROPERTY(EditAnywhere, BlueprintReadWrite, Category = RuntimeVirtualTextureParameterValue)
+	TObjectPtr<class USparseVolumeTexture> ParameterValue;
+
+	UPROPERTY()
+	FGuid ExpressionGUID;
+
+	explicit FSparseVolumeTextureParameterValue(const FMaterialParameterInfo& InParameterInfo = FMaterialParameterInfo(), class USparseVolumeTexture* InValue = nullptr)
+		: ParameterInfo(InParameterInfo), ParameterValue(InValue)
+	{
+	}
+
+	bool IsOverride() const { return true; }
+
+	bool IsValid() const { return GetValue(*this) != nullptr; }
+
+	bool operator==(const FSparseVolumeTextureParameterValue& Other) const
+	{
+		return
+			ParameterInfo == Other.ParameterInfo &&
+			ParameterValue == Other.ParameterValue &&
+			ExpressionGUID == Other.ExpressionGUID;
+	}
+	bool operator!=(const FSparseVolumeTextureParameterValue& Other) const
+	{
+		return !((*this) == Other);
+	}
+
+	typedef const USparseVolumeTexture* ValueType;
+	static ValueType GetValue(const FSparseVolumeTextureParameterValue& Parameter) { return Parameter.ParameterValue; }
+
+	void GetValue(FMaterialParameterMetadata& OutResult) const
+	{
+		OutResult.Value = ParameterValue;
+#if WITH_EDITORONLY_DATA
+		OutResult.ExpressionGuid = ExpressionGUID;
+#endif
+	}
+};
+
 /** Editable font parameter. */
 USTRUCT(BlueprintType)
 struct FFontParameterValue
@@ -602,6 +650,10 @@ public:
 	UPROPERTY(EditAnywhere, BlueprintReadOnly, Category = MaterialInstance, meta = (EditFixedOrder))
 	TArray<struct FRuntimeVirtualTextureParameterValue> RuntimeVirtualTextureParameterValues;
 
+	/** Sparse Volume Texture parameters. */
+	UPROPERTY(EditAnywhere, BlueprintReadOnly, Category = MaterialInstance, meta = (EditFixedOrder))
+	TArray<struct FSparseVolumeTextureParameterValue> SparseVolumeTextureParameterValues;
+
 	/** Font parameters. */
 	UPROPERTY(EditAnywhere, BlueprintReadOnly, Category=MaterialInstance, meta = (EditFixedOrder))
 	TArray<struct FFontParameterValue> FontParameterValues;
@@ -968,6 +1020,7 @@ protected:
 #endif
 	void SetTextureParameterValueInternal(const FMaterialParameterInfo& ParameterInfo, class UTexture* Value);
 	void SetRuntimeVirtualTextureParameterValueInternal(const FMaterialParameterInfo& ParameterInfo, class URuntimeVirtualTexture* Value);
+	void SetSparseVolumeTextureParameterValueInternal(const FMaterialParameterInfo& ParameterInfo, class USparseVolumeTexture* Value);
 	void SetFontParameterValueInternal(const FMaterialParameterInfo& ParameterInfo, class UFont* FontValue, int32 FontPage);
 	void ClearParameterValuesInternal(EMaterialInstanceClearParameterFlag Flags = EMaterialInstanceClearParameterFlag::AllNonStatic);
 
@@ -1278,6 +1331,11 @@ bool UpdateParameterSet(TArray<FRuntimeVirtualTextureParameterValue>& Parameters
 	return MaterialInstance_Private::UpdateParameterSet_WithCachedData<FRuntimeVirtualTextureParameterValue, ExpressionType>(EMaterialParameterType::RuntimeVirtualTexture, Parameters, ParentMaterial);
 }
 template <typename ParameterType, typename ExpressionType>
+bool UpdateParameterSet(TArray<FSparseVolumeTextureParameterValue>& Parameters, UMaterial* ParentMaterial)
+{
+	return MaterialInstance_Private::UpdateParameterSet_WithCachedData<FSparseVolumeTextureParameterValue, ExpressionType>(EMaterialParameterType::SparseVolumeTexture, Parameters, ParentMaterial);
+}
+template <typename ParameterType, typename ExpressionType>
 bool UpdateParameterSet(TArray<FStaticSwitchParameter>& Parameters, UMaterial* ParentMaterial)
 {
 	return MaterialInstance_Private::UpdateParameterSet_WithCachedData<FStaticSwitchParameter, ExpressionType>(EMaterialParameterType::StaticSwitch, Parameters, ParentMaterial);
diff --git a/Engine/Source/Runtime/Engine/Classes/Materials/MaterialInstanceConstant.h b/Engine/Source/Runtime/Engine/Classes/Materials/MaterialInstanceConstant.h
--- a/Engine/Source/Runtime/Engine/Classes/Materials/MaterialInstanceConstant.h
+++ b/Engine/Source/Runtime/Engine/Classes/Materials/MaterialInstanceConstant.h
@@ -93,6 +93,7 @@ class UMaterialInstanceConstant : public UMaterialInstance
 	ENGINE_API void SetScalarParameterAtlasEditorOnly(const FMaterialParameterInfo& ParameterInfo, FScalarParameterAtlasInstanceData AtlasData);
 	ENGINE_API void SetTextureParameterValueEditorOnly(const FMaterialParameterInfo& ParameterInfo, class UTexture* Value);
 	ENGINE_API void SetRuntimeVirtualTextureParameterValueEditorOnly(const FMaterialParameterInfo& ParameterInfo, class URuntimeVirtualTexture* Value);
+	ENGINE_API void SetSparseVolumeTextureParameterValueEditorOnly(const FMaterialParameterInfo& ParameterInfo, class USparseVolumeTexture* Value);
 	ENGINE_API void SetFontParameterValueEditorOnly(const FMaterialParameterInfo& ParameterInfo, class UFont* FontValue, int32 FontPage);
 
 	/**
diff --git a/Engine/Source/Runtime/Engine/Classes/Materials/MaterialInstanceDynamic.h b/Engine/Source/Runtime/Engine/Classes/Materials/MaterialInstanceDynamic.h
--- a/Engine/Source/Runtime/Engine/Classes/Materials/MaterialInstanceDynamic.h
+++ b/Engine/Source/Runtime/Engine/Classes/Materials/MaterialInstanceDynamic.h
@@ -74,6 +74,10 @@ class ENGINE_API UMaterialInstanceDynamic : public UMaterialInstance
 	UFUNCTION(BlueprintCallable, Category = "Rendering|Material")
 	void SetRuntimeVirtualTextureParameterValueByInfo(const FMaterialParameterInfo& ParameterInfo, class URuntimeVirtualTexture* Value);
 
+	/** Set an MID texture parameter value */
+	UFUNCTION(BlueprintCallable, Category = "Rendering|Material")
+	void SetSparseVolumeTextureParameterValue(FName ParameterName, class USparseVolumeTexture* Value);
+
 	/** Get the current MID texture parameter value */
 	UFUNCTION(BlueprintCallable, meta=(DisplayName = "Get Texture Parameter Value", ScriptName = "GetTextureParameterValue"), Category="Rendering|Material")
 	class UTexture* K2_GetTextureParameterValue(FName ParameterName);
diff --git a/Engine/Source/Runtime/Engine/Classes/Materials/MaterialInterface.h b/Engine/Source/Runtime/Engine/Classes/Materials/MaterialInterface.h
--- a/Engine/Source/Runtime/Engine/Classes/Materials/MaterialInterface.h
+++ b/Engine/Source/Runtime/Engine/Classes/Materials/MaterialInterface.h
@@ -567,6 +567,7 @@ public:
 	ENGINE_API void GetAllDoubleVectorParameterInfo(TArray<FMaterialParameterInfo>& OutParameterInfo, TArray<FGuid>& OutParameterIds) const;
 	ENGINE_API void GetAllTextureParameterInfo(TArray<FMaterialParameterInfo>& OutParameterInfo, TArray<FGuid>& OutParameterIds) const;
 	ENGINE_API void GetAllRuntimeVirtualTextureParameterInfo(TArray<FMaterialParameterInfo>& OutParameterInfo, TArray<FGuid>& OutParameterIds) const;
+	ENGINE_API void GetAllSparseVolumeTextureParameterInfo(TArray<FMaterialParameterInfo>& OutParameterInfo, TArray<FGuid>& OutParameterIds) const;
 	ENGINE_API void GetAllFontParameterInfo(TArray<FMaterialParameterInfo>& OutParameterInfo, TArray<FGuid>& OutParameterIds) const;
 
 #if WITH_EDITORONLY_DATA
@@ -585,6 +586,7 @@ public:
 	ENGINE_API bool GetDoubleVectorParameterDefaultValue(const FHashedMaterialParameterInfo& ParameterInfo, FVector4d& OutValue) const;
 	ENGINE_API bool GetTextureParameterDefaultValue(const FHashedMaterialParameterInfo& ParameterInfo, class UTexture*& OutValue) const;
 	ENGINE_API bool GetRuntimeVirtualTextureParameterDefaultValue(const FHashedMaterialParameterInfo& ParameterInfo, class URuntimeVirtualTexture*& OutValue) const;
+	ENGINE_API bool GetSparseVolumeTextureParameterDefaultValue(const FHashedMaterialParameterInfo& ParameterInfo, class USparseVolumeTexture*& OutValue) const;
 	ENGINE_API bool GetFontParameterDefaultValue(const FHashedMaterialParameterInfo& ParameterInfo, class UFont*& OutFontValue, int32& OutFontPage) const;
 	
 #if WITH_EDITOR
@@ -801,6 +803,7 @@ public:
 #endif // WITH_EDITOR
 	ENGINE_API bool GetTextureParameterValue(const FHashedMaterialParameterInfo& ParameterInfo, class UTexture*& OutValue, bool bOveriddenOnly = false) const;
 	ENGINE_API bool GetRuntimeVirtualTextureParameterValue(const FHashedMaterialParameterInfo& ParameterInfo, class URuntimeVirtualTexture*& OutValue, bool bOveriddenOnly = false) const;
+	ENGINE_API bool GetSparseVolumeTextureParameterValue(const FHashedMaterialParameterInfo& ParameterInfo, class USparseVolumeTexture*& OutValue, bool bOveriddenOnly = false) const;
 #if WITH_EDITOR
 	ENGINE_API bool GetTextureParameterChannelNames(const FHashedMaterialParameterInfo& ParameterInfo, FParameterChannelNames& OutValue) const;
 #endif
diff --git a/Engine/Source/Runtime/Engine/Classes/SparseVolumeTexture/SparseVolumeTexture.h b/Engine/Source/Runtime/Engine/Classes/SparseVolumeTexture/SparseVolumeTexture.h
new file mode 100644
--- /dev/null
+++ b/Engine/Source/Runtime/Engine/Classes/SparseVolumeTexture/SparseVolumeTexture.h
@@ -0,0 +1,286 @@
+// Copyright Epic Games, Inc. All Rights Reserved.
+
+#pragma once
+
+#include "CoreMinimal.h"
+
+#include "Engine/VolumeTexture.h"
+
+#include "Serialization/EditorBulkData.h"
+#include "Containers/Array.h"
+#include "UObject/ObjectSaveContext.h"
+
+#include "SparseVolumeTexture.generated.h"
+
+
+#define SPARSE_VOLUME_TILE_RES	16
+
+struct ENGINE_API FSparseVolumeAssetHeader
+{
+	FIntVector3 PageTableVolumeResolution;
+	FIntVector3 TileDataVolumeResolution;
+	FIntVector3 SourceVolumeResolution;
+
+	// The current data format version for the header.
+	static const uint32 kVersion = 0;
+
+	// This version can be used to convert existing header to new version later.
+	uint32 Version;
+
+	void Serialize(FArchive& Ar);
+
+	FSparseVolumeAssetHeader()
+		: PageTableVolumeResolution(FIntVector3(0, 0, 0))
+		, TileDataVolumeResolution(FIntVector3(0, 0, 0))
+		, SourceVolumeResolution(FIntVector3(0, 0, 0))
+		, Version(kVersion)
+	{
+	}
+};
+
+// The structure represent the source asset in high quality. It is used to cook the runtime data
+struct ENGINE_API FSparseVolumeRawSource
+{
+	FSparseVolumeAssetHeader Header;
+
+	// SVT_TODO Store the source asset as the saved openvdb instead of page table/physical tile as it is now.
+
+	TArray<uint32>	DensityPage;
+	TArray<uint8>	DensityData;
+
+	// The current data format version for the raw source data.
+	static const uint32 kVersion = 0;
+
+	// This version can be used to convert existing source data to new version later.
+	uint32 Version;
+
+	void Serialize(FArchive& Ar);
+
+	FSparseVolumeRawSource()
+		: Header()
+		, Version(kVersion)
+	{
+	}
+};
+
+// The structure represent the runtime data cooked runtime data.
+struct ENGINE_API FSparseVolumeTextureRuntime
+{
+	FSparseVolumeAssetHeader	Header;
+	TArray<uint32>				DensityPage;
+	TArray<uint8>				DensityData;
+
+	void SetAsDefaultTexture();
+
+	// The current data format version for the raw source data.
+	static const uint32 kVersion = 0;
+
+	// This version can be used to convert existing source data to new version later.
+	uint32 Version;
+
+	void Serialize(FArchive& Ar);
+
+	FSparseVolumeTextureRuntime()
+		: Header()
+		, Version(kVersion)
+	{
+	}
+};
+
+
+class FSparseVolumeTextureSceneProxy : public FRenderResource
+{
+public:
+
+	FSparseVolumeTextureSceneProxy();
+	virtual ~FSparseVolumeTextureSceneProxy() override;
+
+	void InitialiseRuntimeData(FSparseVolumeTextureRuntime& SparseVolumeTextureRuntime);
+
+	const FSparseVolumeAssetHeader& GetHeader() const 
+	{
+		check(SparseVolumeTextureRuntime);
+		return SparseVolumeTextureRuntime->Header;
+	}
+
+	FTextureRHIRef GetTileDataTextureRHI() const
+	{
+		return  TileDataTextureRHI;
+	}
+	FTextureRHIRef GetPageTableTextureRHI() const
+	{
+		return  PageTableTextureRHI;
+	}
+
+	virtual void InitRHI() override;
+	virtual void ReleaseRHI() override;
+
+private:
+
+	FSparseVolumeTextureRuntime*		SparseVolumeTextureRuntime;
+
+	FTextureRHIRef						PageTableTextureRHI;
+	FTextureRHIRef						TileDataTextureRHI;
+};
+
+
+struct ENGINE_API FSparseVolumeTextureFrame
+{
+	FSparseVolumeTextureFrame();
+	virtual ~FSparseVolumeTextureFrame();
+
+	// The frame data that can be streamed in when in game.
+	FByteBulkData						RuntimeStreamedInData;
+
+	// The runtime frame data that is created.
+	FSparseVolumeTextureRuntime			SparseVolumeTextureRuntime;
+
+	// The render side proxy for the sparse volume texture asset.
+	FSparseVolumeTextureSceneProxy		SparseVolumeTextureSceneProxy;
+
+#if WITH_EDITORONLY_DATA
+	/** The raw data that can be loaded when we want to update cook the data with different settings or updated code without re importing. */
+	UE::Serialization::FEditorBulkData	RawData;
+#endif
+};
+
+
+enum ESparseVolumeTextureShaderUniform
+{
+	ESparseVolumeTexture_PhysicalUVToPageUV,
+	ESparseVolumeTexture_TileSize,
+	ESparseVolumeTexture_PageTableSize,
+	ESparseVolumeTexture_Count,
+};
+
+// SparseVolumeTexture base interface to communicate with material graph and shader bindings.
+UCLASS(ClassGroup = Rendering, BlueprintType)
+class ENGINE_API USparseVolumeTexture : public UObject
+{
+	GENERATED_UCLASS_BODY()
+
+public:
+
+	USparseVolumeTexture();
+	virtual ~USparseVolumeTexture() = default;
+
+	virtual int32 GetFrameCount() const { return 0; }
+	virtual const FSparseVolumeAssetHeader* GetSparseVolumeTextureHeader() const { return nullptr; }
+	virtual FSparseVolumeTextureSceneProxy* GetSparseVolumeTextureSceneProxy() { return nullptr; }
+	virtual const FSparseVolumeTextureSceneProxy* GetSparseVolumeTextureSceneProxy() const { return nullptr; }
+
+	/** Getter for the shader uniform parameters with index as ESparseVolumeTextureShaderUniform. */
+	virtual FVector4 GetUniformParameter(int32 Index) const { return FVector4(ForceInitToZero); }
+
+	/** Getter for the shader uniform parameter type with index as ESparseVolumeTextureShaderUniform. */
+	static UE::Shader::EValueType GetUniformParameterType(int32 Index);
+
+private:
+};
+
+
+UCLASS(ClassGroup = Rendering, BlueprintType)//, hidecategories = (Object))
+class ENGINE_API UStaticSparseVolumeTexture : public USparseVolumeTexture
+{
+	GENERATED_UCLASS_BODY()
+
+public:
+
+	UStaticSparseVolumeTexture();
+	virtual ~UStaticSparseVolumeTexture() = default;
+
+	//~ Begin UObject Interface.
+	virtual void GetResourceSizeEx(FResourceSizeEx& CumulativeResourceSize) override;
+	virtual void PostLoad() override;
+	virtual void BeginDestroy() override;
+	virtual void Serialize(FArchive& Ar) override;
+	virtual void PreSave(FObjectPreSaveContext ObjectSaveContext) override;
+#if WITH_EDITOR
+	virtual void PostEditChangeProperty(FPropertyChangedEvent& PropertyChangedEvent) override;
+#endif
+	//~ End UObject Interface.
+
+
+	//~ Begin USparseVolumeTexture Interface.
+	int32 GetFrameCount() const override { return 1; }
+	const FSparseVolumeAssetHeader* GetSparseVolumeTextureHeader() const override;
+	FSparseVolumeTextureSceneProxy* GetSparseVolumeTextureSceneProxy() override;
+	const FSparseVolumeTextureSceneProxy* GetSparseVolumeTextureSceneProxy() const override;
+
+	/** Getter for the shader uniform parameters. */
+	FVector4 GetUniformParameter(int32 Index) const;
+	/** Getter for the shader uniform parameter type. */
+	static UE::Shader::EValueType GetUniformParameterType(int32 Index);
+	//~ Begin USparseVolumeTexture Interface.
+
+private:
+
+#if WITH_EDITOR
+	friend class USparseVolumeTextureFactory; // Importer
+#endif
+
+	FSparseVolumeTextureFrame	StaticFrame;
+
+	void ConvertRawSourceDataToSparseVolumeTextureRuntime();
+	void GenerateOrLoadDDCRuntimeData();
+	void GenerateOrLoadDDCRuntimeDataAndCreateSceneProxy();
+};
+
+// UAnimatedSparseVolumeTexture inherit from USparseVolumeTexture to be viewed using the first frame by default.
+UCLASS(ClassGroup = Rendering, BlueprintType)//, hidecategories = (Object))
+class ENGINE_API UAnimatedSparseVolumeTexture : public USparseVolumeTexture
+{
+	GENERATED_UCLASS_BODY()
+
+public:
+
+	UAnimatedSparseVolumeTexture();
+	virtual ~UAnimatedSparseVolumeTexture() = default;
+
+	//~ Begin UObject Interface.
+	virtual void GetResourceSizeEx(FResourceSizeEx& CumulativeResourceSize) override;
+	virtual void PostLoad() override;
+	virtual void BeginDestroy() override;
+	virtual void Serialize(FArchive& Ar) override;
+	virtual void PreSave(FObjectPreSaveContext ObjectSaveContext) override;
+#if WITH_EDITOR
+	virtual void PostEditChangeProperty(FPropertyChangedEvent& PropertyChangedEvent) override;
+#endif
+	//~ End UObject Interface.
+
+
+	//~ Begin USparseVolumeTexture Interface.
+	int32 GetFrameCount() const override { return FrameCount; }
+	const FSparseVolumeAssetHeader* GetSparseVolumeTextureHeader() const override;
+	FSparseVolumeTextureSceneProxy* GetSparseVolumeTextureSceneProxy() override;
+	const FSparseVolumeTextureSceneProxy* GetSparseVolumeTextureSceneProxy() const override;
+
+	/** Getter for the shader uniform parameters. */
+	FVector4 GetUniformParameter(int32 Index) const;
+	/** Getter for the shader uniform parameter type. */
+	static UE::Shader::EValueType GetUniformParameterType(int32 Index);
+	//~ Begin USparseVolumeTexture Interface.
+
+	// Used for debugging a specific frame of an animated sequence.
+	const FSparseVolumeTextureSceneProxy* GetSparseVolumeTextureFrameSceneProxy(int32 FrameIndex) const;
+
+	// The asset frame count.
+	UPROPERTY(VisibleAnywhere, Category = "Animation", meta = (DisplayName = "Frame Count"))
+	int32 FrameCount = 0;
+
+private:
+
+#if WITH_EDITOR
+	friend class USparseVolumeTextureFactory; // Importer
+#endif
+	
+	const uint32 PreviewFrameIndex = 0;
+	const bool bLoadAllFramesToProxies = true;	// SVT_TODO remove that once streaming is working
+	int32 GetFrameCountToLoad() const ;
+
+	TArray<FSparseVolumeTextureFrame>	AnimationFrames;
+
+	void ConvertRawSourceDataToSparseVolumeTextureRuntime(int32 FrameIndex);
+	void GenerateOrLoadDDCRuntimeData(int32 FrameIndex);
+	void GenerateOrLoadDDCRuntimeDataAndCreateSceneProxy(int32 FrameIndex);
+};
diff --git a/Engine/Source/Runtime/Engine/Engine.Build.cs b/Engine/Source/Runtime/Engine/Engine.Build.cs
--- a/Engine/Source/Runtime/Engine/Engine.Build.cs
+++ b/Engine/Source/Runtime/Engine/Engine.Build.cs
@@ -148,6 +148,12 @@ public class Engine : ModuleRules
 			}
 		);
 
+		if (Target.bBuildWithEditorOnlyData && Target.bBuildEditor)
+		{
+			// The SparseVolumeTexture module containing the importer is only loaded and used in the editor.
+			DynamicallyLoadedModuleNames.Add("SparseVolumeTexture");
+		}
+
 		// Cross platform Audio Codecs:
 		AddEngineThirdPartyPrivateStaticDependencies(Target,
 			"UEOgg",
diff --git a/Engine/Source/Runtime/Engine/Private/Materials/HLSLMaterialTranslator.cpp b/Engine/Source/Runtime/Engine/Private/Materials/HLSLMaterialTranslator.cpp
--- a/Engine/Source/Runtime/Engine/Private/Materials/HLSLMaterialTranslator.cpp
+++ b/Engine/Source/Runtime/Engine/Private/Materials/HLSLMaterialTranslator.cpp
@@ -2188,7 +2188,10 @@ FString FHLSLMaterialTranslator::GetMaterialShaderCode()
 		case MCT_Float2: HLSLType = TEXT("float2"); break;
 		case MCT_Float3: HLSLType = TEXT("float3"); break;
 		case MCT_Float4: HLSLType = TEXT("float4"); break;
-		case MCT_ShadingModel: HLSLType = TEXT("uint"); break;
+		case MCT_UInt: case MCT_UInt1: case MCT_ShadingModel: HLSLType = TEXT("uint"); break;
+		case MCT_UInt2: HLSLType = TEXT("uint2"); break;
+		case MCT_UInt3: HLSLType = TEXT("uint3"); break;
+		case MCT_UInt4: HLSLType = TEXT("uint4"); break;
 		case MCT_Strata: HLSLType = TEXT("FStrataData"); break;
 		default: break;
 		}
@@ -2782,6 +2785,11 @@ const TCHAR* FHLSLMaterialTranslator::DescribeType(EMaterialValueType Type) cons
 	case MCT_TextureVirtual:		return TEXT("TextureVirtual");
 	case MCT_VTPageTableResult:		return TEXT("VTPageTableResult");
 	case MCT_ShadingModel:			return TEXT("ShadingModel");
+	case MCT_UInt:					return TEXT("uint");
+	case MCT_UInt1:					return TEXT("uint");
+	case MCT_UInt2:					return TEXT("uint2");
+	case MCT_UInt3:					return TEXT("uint3");
+	case MCT_UInt4:					return TEXT("uint4");
 	case MCT_Strata:				return TEXT("Strata");
 	case MCT_LWCScalar:				return TEXT("LWCScalar");
 	case MCT_LWCVector2:			return TEXT("LWCVector2");
@@ -2813,6 +2821,11 @@ const TCHAR* FHLSLMaterialTranslator::HLSLTypeString(EMaterialValueType Type) co
 	case MCT_TextureVirtual:		return TEXT("TextureVirtual");
 	case MCT_VTPageTableResult:		return TEXT("VTPageTableResult");
 	case MCT_ShadingModel:			return TEXT("uint");
+	case MCT_UInt:					return TEXT("uint");
+	case MCT_UInt1:					return TEXT("uint");
+	case MCT_UInt2:					return TEXT("uint2");
+	case MCT_UInt3:					return TEXT("uint3");
+	case MCT_UInt4:					return TEXT("uint4");
 	case MCT_Strata:				return TEXT("FStrataData");
 	case MCT_LWCScalar:				return TEXT("FLWCScalar");
 	case MCT_LWCVector2:			return TEXT("FLWCVector2");
@@ -2844,6 +2857,11 @@ const TCHAR* FHLSLMaterialTranslator::HLSLTypeStringDeriv(EMaterialValueType Typ
 	case MCT_TextureVirtual:		return TEXT("TextureVirtual");
 	case MCT_VTPageTableResult:		return TEXT("VTPageTableResult");
 	case MCT_ShadingModel:			return TEXT("uint");
+	case MCT_UInt:					return TEXT("uint");
+	case MCT_UInt1:					return TEXT("uint");
+	case MCT_UInt2:					return TEXT("uint2");
+	case MCT_UInt3:					return TEXT("uint3");
+	case MCT_UInt4:					return TEXT("uint4");
 	case MCT_Strata:				return TEXT("FStrataData");
 	case MCT_LWCScalar:				return (DerivativeStatus == EDerivativeStatus::Valid) ? TEXT("FLWCScalarDeriv") : TEXT("FLWCScalar");
 	case MCT_LWCVector2:			return (DerivativeStatus == EDerivativeStatus::Valid) ? TEXT("FLWCVector2Deriv") : TEXT("FLWCVector2");
@@ -2930,7 +2948,7 @@ int32 FHLSLMaterialTranslator::AddCodeChunkInner(uint64 Hash, const TCHAR* Forma
 		new(*CurrentScopeChunks) FShaderCodeChunk(Hash, FormattedCode, FormattedCode, TEXT(""), Type, DerivativeStatus, true);
 	}
 	// Can only create temporaries for certain types
-	else if ((Type & (MCT_Float | MCT_LWCType | MCT_VTPageTableResult)) || Type == MCT_ShadingModel || Type == MCT_MaterialAttributes || Type == MCT_Strata)
+	else if ((Type & (MCT_Float | MCT_LWCType | MCT_VTPageTableResult | MCT_UInt)) || Type == MCT_ShadingModel || Type == MCT_MaterialAttributes || Type == MCT_Strata || Type == MCT_UInt)
 	{
 		// Check for existing
 		for (int32 i = 0; i < CurrentScopeChunks->Num(); ++i)
@@ -3010,7 +3028,7 @@ int32 FHLSLMaterialTranslator::AddCodeChunkInnerDeriv(const TCHAR* FormattedCode
 		new(*CurrentScopeChunks) FShaderCodeChunk(Hash, FormattedCodeFinite, FormattedCodeAnalytic, TEXT(""), Type, DerivativeStatus, true);
 	}
 	// Can only create temporaries for certain types
-	else if ((Type & (MCT_Float | MCT_LWCType | MCT_VTPageTableResult)) || Type == MCT_ShadingModel || Type == MCT_MaterialAttributes || Type == MCT_Strata)
+	else if ((Type & (MCT_Float | MCT_LWCType | MCT_VTPageTableResult | MCT_UInt)) || Type == MCT_ShadingModel || Type == MCT_MaterialAttributes || Type == MCT_Strata)
 	{
 		// Check for existing
 		for (int32 i = 0; i < CurrentScopeChunks->Num(); ++i)
@@ -5916,7 +5934,7 @@ int32 FHLSLMaterialTranslator::TextureSample(
 		// Note, this does not really do anything (by design) other than adding it to the UniformExpressionSet
 		/*TextureName =*/ CoerceParameter(TextureIndex, TextureType);
 
-		FMaterialUniformExpression* UniformExpression = GetParameterUniformExpression(TextureIndex);
+		FMaterialUniformExpression* UniformExpression = GetParameterUniformExpression(TextureIndex); 
 		if (UniformExpression == nullptr)
 		{
 			return Errorf(TEXT("Unable to find VT uniform expression."));
@@ -6788,7 +6806,7 @@ int32 FHLSLMaterialTranslator::TextureParameter(FName ParameterName, UTexture* I
 	return AddUniformExpression(new FMaterialUniformExpressionTextureParameter(ParameterInfo, TextureReferenceIndex, SamplerType, SamplerSource, bVirtual),ShaderType,TEXT(""));
 }
 
-int32 FHLSLMaterialTranslator::VirtualTexture(URuntimeVirtualTexture* InTexture, int32 TextureLayerIndex, int32 PageTableLayerIndex, int32& TextureReferenceIndex, EMaterialSamplerType SamplerType)
+int32 FHLSLMaterialTranslator::VirtualTexture(URuntimeVirtualTexture* InTexture, int32 TextureLayerIndex, int32 PageTableLayerIndex, int32& TextureReferenceIndex, EMaterialSamplerType SamplerType) 
 {
 	if (!UseVirtualTexturing(FeatureLevel, TargetPlatform))
 	{
@@ -11955,4 +11973,142 @@ bool FHLSLMaterialTranslator::IsDevelopmentFeatureEnabled(const FName& FeatureNa
 	return true;
 }
 
+int32 FHLSLMaterialTranslator::SparseVolumeTexture(USparseVolumeTexture* Texture, int32 SubTextureId, int32& TextureReferenceIndex, EMaterialSamplerType SamplerType)
+{
+	TextureReferenceIndex = Material->GetReferencedTextures().Find(Texture);
+	checkf(TextureReferenceIndex != INDEX_NONE, TEXT("Material expression called Compiler->SparseVolumeTexture() without implementing UMaterialExpression::GetReferencedTexture properly"));
+
+	return AddUniformExpression(new FMaterialUniformExpressionTexture(TextureReferenceIndex, SubTextureId, SamplerType), MCT_VolumeTexture, TEXT(""));
+}
+
+int32 FHLSLMaterialTranslator::SparseVolumeTextureParameter(FName ParameterName, USparseVolumeTexture* InDefaultTexture, int32 SubTextureId, int32& TextureReferenceIndex, EMaterialSamplerType SamplerType)
+{
+	USparseVolumeTexture* DefaultTexture = InDefaultTexture;
+
+	// If we're compiling a function, give the function a chance to override the default parameter value
+	FMaterialParameterMetadata Meta;
+	if (GetParameterOverrideValueForCurrentFunction(EMaterialParameterType::SparseVolumeTexture, ParameterName, Meta))
+	{
+		DefaultTexture = Meta.Value.SparseVolumeTexture;
+	}
+
+	TextureReferenceIndex = Material->GetReferencedTextures().Find(DefaultTexture);
+	checkf(TextureReferenceIndex != INDEX_NONE, TEXT("Material expression called Compiler->SparseVolumeTextureParameter() without implementing UMaterialExpression::GetReferencedTexture properly"));
+
+	FMaterialParameterInfo ParameterInfo = GetParameterAssociationInfo();
+	ParameterInfo.Name = ParameterName;
+
+	return AddUniformExpression(new FMaterialUniformExpressionTextureParameter(ParameterInfo, TextureReferenceIndex, SubTextureId, SamplerType), MCT_VolumeTexture, TEXT(""));
+}
+
+int32 FHLSLMaterialTranslator::SparseVolumeTextureGetVoxelCoord(int32 PackedPhysicalTileCoord, int32 TileSize, int32 CoordPageTable, int32 CoordVolume)
+{
+	FString SampleCode = FString::Printf(TEXT("SparseVolumeTextureGetVoxelCoord(%s, %s, %s, %s)"),
+		*GetParameterCode(PackedPhysicalTileCoord), *GetParameterCode(TileSize), *GetParameterCode(CoordPageTable), *GetParameterCode(CoordVolume));
+	return AddCodeChunk(MCT_UInt3, *SampleCode);
+}
+
+int32 FHLSLMaterialTranslator::SparseVolumeTextureSample(int32 TextureIndex, int32 CoordinateIndex)
+{
+	if (TextureIndex == INDEX_NONE || CoordinateIndex == INDEX_NONE)
+	{
+		return INDEX_NONE;
+	}
+
+	EMaterialValueType TextureType = GetParameterType(TextureIndex);
+
+	int32 SamplingCodeIndex = INDEX_NONE;
+	FString TextureName;
+	if (TextureType == MCT_TextureCube)
+	{
+		return Errorf(TEXT("FHLSLMaterialTranslator::SparseVolumeTextureSample MCT_TextureCube ERROR."));
+	}
+	else if (TextureType == MCT_Texture2DArray)
+	{
+		return Errorf(TEXT("FHLSLMaterialTranslator::SparseVolumeTextureSample MCT_Texture2DArray ERROR."));
+	}
+	else if (TextureType == MCT_TextureCubeArray)
+	{
+		return Errorf(TEXT("FHLSLMaterialTranslator::SparseVolumeTextureSample MCT_TextureCubeArray ERROR."));
+	}
+	else if (TextureType == MCT_VolumeTexture)
+	{
+		TextureName = CoerceParameter(TextureIndex, MCT_VolumeTexture);		// TODO this hsould be some texture reserved for sub textures
+
+
+		FMaterialUniformExpression* UniformExpression = GetParameterUniformExpression(TextureIndex);
+		if (UniformExpression == nullptr)
+		{
+			return Errorf(TEXT("Unable to find expression for Sparse Volume Texture."));
+		}
+		FMaterialUniformExpressionTexture* TextureUniformExpression = UniformExpression->GetTextureUniformExpression();
+		if (TextureUniformExpression == nullptr || TextureUniformExpression->GetPageTableLayerIndex() != INDEX_NONE || TextureUniformExpression->GetTextureLayerIndex() == INDEX_NONE)
+		{
+			// Sparse volume texture do not use PageTableLayerIndex but only TextureLayerIndex.
+			return Errorf(TEXT("The provided uniform expression is cannot be identified as a Sparse Volume Texture."));
+		}
+
+		int32 VirtualTextureIndex = UniformTextureExpressions[(uint32)EMaterialTextureParameterType::Volume].Find(TextureUniformExpression);
+		check(UniformTextureExpressions[(uint32)EMaterialTextureParameterType::Volume].IsValidIndex(VirtualTextureIndex));
+
+		const EDerivativeStatus UvDerivativeStatus = GetDerivativeStatus(CoordinateIndex);
+
+		check(TextureUniformExpression->GetTextureLayerIndex() != INDEX_NONE);
+		if (TextureUniformExpression->GetTextureLayerIndex() == 0)
+		{
+			FString SampleCode = FString::Printf(TEXT("%s.Load(int4(int3(%s), 0)).x"), *TextureName, *GetParameterCode(CoordinateIndex));
+			AddEstimatedTextureSample();
+			if (IsAnalyticDerivEnabled() && IsDerivativeValid(UvDerivativeStatus))
+			{
+				SamplingCodeIndex = AddCodeChunk(MCT_UInt1, *SampleCode);
+			}
+			else
+			{
+				SamplingCodeIndex = AddCodeChunk(MCT_UInt1, *SampleCode);
+			}
+		}
+		else if (TextureUniformExpression->GetTextureLayerIndex() == 1)
+		{
+			FString SampleCode = FString::Printf(TEXT("%s.Load(int4(int3(%s), 0))"), *TextureName, *GetParameterCode(CoordinateIndex));
+			AddEstimatedTextureSample();
+			if (IsAnalyticDerivEnabled() && IsDerivativeValid(UvDerivativeStatus))
+			{
+				SamplingCodeIndex = AddCodeChunk(MCT_Float4, *SampleCode);
+			}
+			else
+			{
+				SamplingCodeIndex = AddCodeChunk(MCT_Float4, *SampleCode);
+			}
+		}
+		else
+		{
+			return Errorf(TEXT("Sparse volume texture cannot be interpreted."));
+		}
+
+	}
+	else if (TextureType == MCT_TextureExternal)
+	{
+		return Errorf(TEXT("FHLSLMaterialTranslator::SparseVolumeTextureSample MCT_TextureExternal ERROR."));
+	}
+	else // MCT_Texture2D
+	{
+		return Errorf(TEXT("FHLSLMaterialTranslator::SparseVolumeTextureSample MCT_Texture2D ERROR."));
+	}
+
+	return SamplingCodeIndex;
+}
+
+int32 FHLSLMaterialTranslator::SparseVolumeTextureUniform(int32 TextureIndex, int32 VectorIndex, UE::Shader::EValueType Type)
+{
+	return AddUniformExpression(new FMaterialUniformExpressionSparseVolumeTextureUniform(TextureIndex, VectorIndex), GetMaterialValueType(Type), TEXT(""));
+}
+
+int32 FHLSLMaterialTranslator::SparseVolumeTextureUniformParameter(FName ParameterName, int32 TextureIndex, int32 VectorIndex, UE::Shader::EValueType Type)
+{
+	FMaterialParameterInfo ParameterInfo = GetParameterAssociationInfo();
+	ParameterInfo.Name = ParameterName;
+	return AddUniformExpression(new FMaterialUniformExpressionSparseVolumeTextureUniform(ParameterInfo, TextureIndex, VectorIndex), GetMaterialValueType(Type), TEXT(""));
+}
+
+
 #endif // WITH_EDITORONLY_DATA
diff --git a/Engine/Source/Runtime/Engine/Private/Materials/HLSLMaterialTranslator.h b/Engine/Source/Runtime/Engine/Private/Materials/HLSLMaterialTranslator.h
--- a/Engine/Source/Runtime/Engine/Private/Materials/HLSLMaterialTranslator.h
+++ b/Engine/Source/Runtime/Engine/Private/Materials/HLSLMaterialTranslator.h
@@ -32,6 +32,7 @@
 #include "Field/FieldSystemTypes.h"
 #include "Containers/Map.h"
 #include "Shader/ShaderTypes.h"
+#include "SparseVolumeTexture/SparseVolumeTexture.h"
 
 #if WITH_EDITORONLY_DATA
 #include "Materials/MaterialExpressionSceneTexture.h"
@@ -847,6 +848,7 @@ protected:
 	) override;
 
 	virtual int32 TextureProperty(int32 TextureIndex, EMaterialExposedTextureProperty Property) override;
+
 	virtual int32 TextureDecalMipmapLevel(int32 TextureSizeInput) override;
 	virtual int32 TextureDecalDerivative(bool bDDY) override;
 
@@ -889,6 +891,13 @@ protected:
 	virtual int32 ExternalTextureCoordinateOffset(int32 TextureReferenceIndex, TOptional<FName> ParameterName) override;
 	virtual int32 ExternalTextureCoordinateOffset(const FGuid& ExternalTextureGuid) override;
 
+	virtual int32 SparseVolumeTexture(USparseVolumeTexture* Texture, int32 SubTextureId, int32& TextureReferenceIndex, EMaterialSamplerType SamplerType) override;
+	virtual int32 SparseVolumeTextureParameter(FName ParameterName, USparseVolumeTexture* InDefaultTexture, int32 SubTextureId, int32& TextureReferenceIndex, EMaterialSamplerType SamplerType) override;
+	virtual int32 SparseVolumeTextureSample(int32 TextureIndex, int32 CoordinateIndex) override;
+	virtual int32 SparseVolumeTextureUniform(int32 TextureIndex, int32 VectorIndex, UE::Shader::EValueType Type) override;
+	virtual int32 SparseVolumeTextureUniformParameter(FName ParameterName, int32 TextureIndex, int32 VectorIndex, UE::Shader::EValueType Type) override;
+	virtual int32 SparseVolumeTextureGetVoxelCoord(int32 PackedPhysicalTileCoord, int32 TileSize, int32 CoordPageTable, int32 CoordVolume) override;
+
 	virtual UObject* GetReferencedTexture(int32 Index);
 
 	virtual int32 StaticBool(bool bValue) override;
diff --git a/Engine/Source/Runtime/Engine/Private/Materials/Material.cpp b/Engine/Source/Runtime/Engine/Private/Materials/Material.cpp
--- a/Engine/Source/Runtime/Engine/Private/Materials/Material.cpp
+++ b/Engine/Source/Runtime/Engine/Private/Materials/Material.cpp
@@ -1461,6 +1461,13 @@ bool UMaterial::SetRuntimeVirtualTextureParameterValueEditorOnly(FName Parameter
 	return SetParameterValueEditorOnly(ParameterName, Meta);
 }
 
+bool UMaterial::SetSparseVolumeTextureParameterValueEditorOnly(FName ParameterName, class USparseVolumeTexture* InValue)
+{
+	FMaterialParameterMetadata Meta;
+	Meta.Value = InValue;
+	return SetParameterValueEditorOnly(ParameterName, Meta);
+}
+
 bool UMaterial::SetFontParameterValueEditorOnly(FName ParameterName, class UFont* InFontValue, int32 InFontPage)
 {
 	FMaterialParameterMetadata Meta;
diff --git a/Engine/Source/Runtime/Engine/Private/Materials/MaterialCachedData.cpp b/Engine/Source/Runtime/Engine/Private/Materials/MaterialCachedData.cpp
--- a/Engine/Source/Runtime/Engine/Private/Materials/MaterialCachedData.cpp
+++ b/Engine/Source/Runtime/Engine/Private/Materials/MaterialCachedData.cpp
@@ -36,6 +36,7 @@
 #include "HLSLTree/HLSLTreeEmit.h"
 #include "HLSLTree/HLSLTree.h"
 #include "VT/RuntimeVirtualTexture.h"
+#include "SparseVolumeTexture/SparseVolumeTexture.h"
 #include "Engine/Font.h"
 #include "LandscapeGrassType.h"
 #include "ProfilingDebugging/LoadTimeTracker.h"
@@ -200,6 +201,10 @@ void FMaterialCachedExpressionData::AddParameter(const FMaterialParameterInfo& P
 			RuntimeVirtualTextureValues.Insert(ParameterMeta.Value.RuntimeVirtualTexture, Index);
 			OutReferencedTexture = ParameterMeta.Value.RuntimeVirtualTexture;
 			break;
+		case EMaterialParameterType::SparseVolumeTexture:
+			SparseVolumeTextureValues.Insert(ParameterMeta.Value.SparseVolumeTexture, Index);
+			OutReferencedTexture = ParameterMeta.Value.SparseVolumeTexture;
+			break;
 		case EMaterialParameterType::StaticSwitch:
 			EditorOnlyData->StaticSwitchValues.Insert(ParameterMeta.Value.AsStaticSwitch(), Index);
 			break;
@@ -612,6 +617,9 @@ void FMaterialCachedExpressionData::GetParameterValueByIndex(EMaterialParameterT
 	case EMaterialParameterType::RuntimeVirtualTexture:
 		OutResult.Value = RuntimeVirtualTextureValues[ParameterIndex].LoadSynchronous();
 		break;
+	case EMaterialParameterType::SparseVolumeTexture:
+		OutResult.Value = SparseVolumeTextureValues[ParameterIndex].LoadSynchronous();
+		break;
 	case EMaterialParameterType::Font:
 		OutResult.Value = FMaterialParameterValue(FontValues[ParameterIndex].LoadSynchronous(), FontPageValues[ParameterIndex]);
 		break;
diff --git a/Engine/Source/Runtime/Engine/Private/Materials/MaterialExpressions.cpp b/Engine/Source/Runtime/Engine/Private/Materials/MaterialExpressions.cpp
--- a/Engine/Source/Runtime/Engine/Private/Materials/MaterialExpressions.cpp
+++ b/Engine/Source/Runtime/Engine/Private/Materials/MaterialExpressions.cpp
@@ -46,6 +46,7 @@
 #include "Engine/SubsurfaceProfile.h"
 #include "Styling/CoreStyle.h"
 #include "VT/RuntimeVirtualTexture.h"
+#include "SparseVolumeTexture/SparseVolumeTexture.h"
 #include "ProfilingDebugging/LoadTimeTracker.h"
 
 #include "Materials/MaterialExpressionAbs.h"
@@ -275,6 +276,7 @@
 #include "Materials/MaterialExpressionGetLocal.h"
 #include "Materials/MaterialExpressionBinaryOp.h"
 #include "Materials/MaterialExpressionGenericConstant.h"
+#include "Materials/MaterialExpressionSparseVolumeTextureSample.h"
 #include "EditorSupportDelegates.h"
 #include "MaterialCompiler.h"
 #if WITH_EDITOR
@@ -13554,6 +13556,17 @@ bool UMaterialFunctionInterface::OverrideNamedRuntimeVirtualTextureParameter(con
 	return false;
 }
 
+bool UMaterialFunctionInterface::OverrideNamedSparseVolumeTextureParameter(const FHashedMaterialParameterInfo& ParameterInfo, class USparseVolumeTexture*& OutValue)
+{
+	FMaterialParameterMetadata Meta;
+	if (GetParameterOverrideValue(EMaterialParameterType::SparseVolumeTexture, ParameterInfo.GetName(), Meta))
+	{
+		OutValue = Meta.Value.SparseVolumeTexture;
+		return true;
+	}
+	return false;
+}
+
 bool UMaterialFunctionInterface::OverrideNamedFontParameter(const FHashedMaterialParameterInfo& ParameterInfo, class UFont*& OutFontValue, int32& OutFontPage)
 {
 	FMaterialParameterMetadata Meta;
@@ -14440,6 +14453,13 @@ bool UMaterialFunction::SetRuntimeVirtualTextureParameterValueEditorOnly(FName P
 	return SetParameterValueEditorOnly(ParameterName, Meta);
 };
 
+bool UMaterialFunction::SetSparseVolumeTextureParameterValueEditorOnly(FName ParameterName, class USparseVolumeTexture* InValue)
+{
+	FMaterialParameterMetadata Meta;
+	Meta.Value = InValue;
+	return SetParameterValueEditorOnly(ParameterName, Meta);
+};
+
 bool UMaterialFunction::SetFontParameterValueEditorOnly(FName ParameterName, class UFont* InFontValue, int32 InFontPage)
 {
 	FMaterialParameterMetadata Meta;
@@ -14630,6 +14650,17 @@ void UMaterialFunctionInstance::UpdateParameterSet()
 						}
 					}
 				}
+				else if (const UMaterialExpressionSparseVolumeTextureSampleParameter* SparseVolumeTextureParameter = Cast<const UMaterialExpressionSparseVolumeTextureSampleParameter>(FunctionExpression))
+				{
+					for (FSparseVolumeTextureParameterValue& SparseVolumeTextureParameterValue : SparseVolumeTextureParameterValues)
+					{
+						if (SparseVolumeTextureParameterValue.ExpressionGUID == SparseVolumeTextureParameter->ExpressionGUID)
+						{
+							SparseVolumeTextureParameterValue.ParameterInfo.Name = SparseVolumeTextureParameter->ParameterName;
+							break;
+						}
+					}
+				}
 				else if (const UMaterialExpressionFontSampleParameter* FontParameter = Cast<const UMaterialExpressionFontSampleParameter>(FunctionExpression))
 				{
 					for (FFontParameterValue& FontParameterValue : FontParameterValues)
@@ -14676,6 +14707,7 @@ void UMaterialFunctionInstance::OverrideMaterialInstanceParameterValues(UMateria
 	Instance->DoubleVectorParameterValues = DoubleVectorParameterValues;
 	Instance->TextureParameterValues = TextureParameterValues;
 	Instance->RuntimeVirtualTextureParameterValues = RuntimeVirtualTextureParameterValues;
+	Instance->SparseVolumeTextureParameterValues = SparseVolumeTextureParameterValues;
 	Instance->FontParameterValues = FontParameterValues;
 
 	// Static parameters
@@ -14836,6 +14868,7 @@ bool UMaterialFunctionInstance::GetParameterOverrideValue(EMaterialParameterType
 	case EMaterialParameterType::DoubleVector: bResult = GameThread_GetParameterValue(DoubleVectorParameterValues, ParameterInfo, OutResult); break;
 	case EMaterialParameterType::Texture: bResult = GameThread_GetParameterValue(TextureParameterValues, ParameterInfo, OutResult); break;
 	case EMaterialParameterType::RuntimeVirtualTexture: bResult = GameThread_GetParameterValue(RuntimeVirtualTextureParameterValues, ParameterInfo, OutResult); break;
+	case EMaterialParameterType::SparseVolumeTexture: bResult = GameThread_GetParameterValue(SparseVolumeTextureParameterValues, ParameterInfo, OutResult); break;
 	case EMaterialParameterType::Font: bResult = GameThread_GetParameterValue(FontParameterValues, ParameterInfo, OutResult); break;
 	case EMaterialParameterType::StaticSwitch: bResult = GameThread_GetParameterValue(StaticSwitchParameterValues, ParameterInfo, OutResult); break;
 	case EMaterialParameterType::StaticComponentMask: bResult = GameThread_GetParameterValue(StaticComponentMaskParameterValues, ParameterInfo, OutResult); break;
@@ -25019,4 +25052,221 @@ void UMaterialExpressionSkyLightEnvMapSample::GetCaption(TArray<FString>& OutCap
 }
 #endif // WITH_EDITOR
 
+///////////////////////////////////////////////////////////////////////////////
+// UMaterialExpressionSparseVolumeTextureSample
+///////////////////////////////////////////////////////////////////////////////
+
+UMaterialExpressionSparseVolumeTextureSample::UMaterialExpressionSparseVolumeTextureSample(const FObjectInitializer& ObjectInitializer)
+	: Super(ObjectInitializer)
+{
+	// Structure to hold one-time initialization
+	struct FConstructorStatics
+	{
+		FText NAME_VirtualTexture;
+		FConstructorStatics()
+			: NAME_VirtualTexture(LOCTEXT("Texture", "Texture"))
+		{
+		}
+	};
+	static FConstructorStatics ConstructorStatics;
+
+#if WITH_EDITORONLY_DATA
+	MenuCategories.Add(ConstructorStatics.NAME_VirtualTexture);
+#endif
+
+#if WITH_EDITOR
+	// SVT_TODO update that when attribute setup has been clarified
+	Outputs.Reset();
+	Outputs.Add(FExpressionOutput(TEXT("Density"), 1, 1, 1, 1, 0));
+	bShowOutputNameOnPin = true;
+	bShowMaskColorsOnPin = false;
+#endif
+}
+
+UObject* UMaterialExpressionSparseVolumeTextureSample::GetReferencedTexture() const
+{
+	return SparseVolumeTexture;
+}
+
+#if WITH_EDITOR
+
+void UMaterialExpressionSparseVolumeTextureSample::PostLoad()
+{
+	Super::PostLoad();
+}
+
+void UMaterialExpressionSparseVolumeTextureSample::PostEditChangeProperty(FPropertyChangedEvent& PropertyChangedEvent)
+{
+	// Update what needs to be when the texture is changed
+	if (PropertyChangedEvent.Property && PropertyChangedEvent.Property->GetName() == TEXT("SparseVirtualTexture"))
+	{
+		if (SparseVolumeTexture != nullptr)
+		{
+			FEditorSupportDelegates::ForcePropertyWindowRebuild.Broadcast(this);
+		}
+	}
+
+	Super::PostEditChangeProperty(PropertyChangedEvent);
+}
+
+uint32 UMaterialExpressionSparseVolumeTextureSample::GetOutputType(int32 OutputIndex)
+{
+	switch (OutputIndex)
+	{
+	case 0:
+		return MCT_Float;	// Density defined in constructor for now
+		break;
+	}
+
+	check(false);
+	return MCT_Float1;
+}
+
+uint32 UMaterialExpressionSparseVolumeTextureSample::GetInputType(int32 InputIndex)
+{
+	switch (InputIndex)
+	{
+	case 0:
+		return MCT_Float3;
+		break;
+	}
+
+	check(false);
+	return MCT_Float1;
+}
+
+int32 UMaterialExpressionSparseVolumeTextureSample::Compile(class FMaterialCompiler* Compiler, int32 OutputIndex)
+{
+	if (SparseVolumeTexture != nullptr)
+	{
+		const bool bIsParameter = HasAParameterName() && GetParameterName().IsValid() && !GetParameterName().IsNone();
+
+		int32 PageTableTextureReferenceIndex = INDEX_NONE;
+		int32 PageTableSubTextureIndex = 0;
+		int32 PageTableTextureParamCodeChunkIndex = INDEX_NONE;
+		if (bIsParameter)
+		{
+			PageTableTextureParamCodeChunkIndex = Compiler->SparseVolumeTextureParameter(GetParameterName(), SparseVolumeTexture, PageTableSubTextureIndex, PageTableTextureReferenceIndex, SAMPLERTYPE_LinearColor);
+		}
+		else
+		{
+			PageTableTextureParamCodeChunkIndex = Compiler->SparseVolumeTexture(SparseVolumeTexture, PageTableSubTextureIndex, PageTableTextureReferenceIndex, SAMPLERTYPE_LinearColor);
+		}
+
+		int32 PhysicalTileTextureReferenceIndex = INDEX_NONE;
+		int32 PhysicalTileSubTextureIndex = 1;
+		int32 PhysicalTileTextureParamCodeChunkIndex = INDEX_NONE;
+		if (bIsParameter)
+		{
+			PhysicalTileTextureParamCodeChunkIndex = Compiler->SparseVolumeTextureParameter(GetParameterName(), SparseVolumeTexture, PhysicalTileSubTextureIndex, PhysicalTileTextureReferenceIndex, SAMPLERTYPE_LinearColor);
+		}
+		else
+		{
+			PhysicalTileTextureParamCodeChunkIndex = Compiler->SparseVolumeTexture(SparseVolumeTexture, PhysicalTileSubTextureIndex, PhysicalTileTextureReferenceIndex, SAMPLERTYPE_LinearColor);
+		}
+
+		// Now register the uniform we need for sampling the sparse volume texture
+		int32 UniformCodeChunkIndex[ESparseVolumeTexture_Count];
+		for (int32 UniformIndex = 0; UniformIndex < ESparseVolumeTexture_Count; ++UniformIndex)
+		{
+			const UE::Shader::EValueType Type = USparseVolumeTexture::GetUniformParameterType(UniformIndex);
+			if (bIsParameter)
+			{
+				UniformCodeChunkIndex[UniformIndex] = Compiler->SparseVolumeTextureUniformParameter(GetParameterName(), PageTableTextureReferenceIndex, UniformIndex, Type);
+			}
+			else
+			{
+				UniformCodeChunkIndex[UniformIndex] = Compiler->SparseVolumeTextureUniform(PageTableTextureReferenceIndex, UniformIndex, Type);
+			}
+		}
+
+		// SVT_TODO Send parameters and do all the following math in shared functions (one to sample the page table and one to sample the attribute from the page table sample output)
+
+		int32 UVW = CompileWithDefaultFloat3(Compiler, Coordinates, 0.0f, 0.0f, 0.0f);
+		int32 UVWPageTable = Compiler->Mul(UVW, UniformCodeChunkIndex[ESparseVolumeTexture_PhysicalUVToPageUV]);
+		int32 CoordPageTable = Compiler->Mul(UVWPageTable, UniformCodeChunkIndex[ESparseVolumeTexture_PageTableSize]);
+
+		int32 PackedPhysicalTileCoord = Compiler->SparseVolumeTextureSample(PageTableTextureParamCodeChunkIndex, CoordPageTable);
+
+		int32 CoordVolume = Compiler->Mul(CoordPageTable, UniformCodeChunkIndex[ESparseVolumeTexture_TileSize]);
+		int32 CoordPageTableFloor = Compiler->Floor(CoordPageTable);
+
+		int32 CoordPhysicalTileVoxel = Compiler->SparseVolumeTextureGetVoxelCoord(PackedPhysicalTileCoord, UniformCodeChunkIndex[ESparseVolumeTexture_TileSize], CoordPageTableFloor, CoordVolume);
+
+		return Compiler->SparseVolumeTextureSample(PhysicalTileTextureParamCodeChunkIndex, CoordPhysicalTileVoxel);
+	}
+
+	return INDEX_NONE;
+}
+
+void UMaterialExpressionSparseVolumeTextureSample::GetCaption(TArray<FString>& OutCaptions) const
+{
+	OutCaptions.Add(FString(TEXT("Sparse Volume Texture Sample")));
+}
+
+#endif // WITH_EDITOR
+
+///////////////////////////////////////////////////////////////////////////////
+// UMaterialExpressionSparseVolumeTextureSample
+///////////////////////////////////////////////////////////////////////////////
+
+UMaterialExpressionSparseVolumeTextureSampleParameter::UMaterialExpressionSparseVolumeTextureSampleParameter(const FObjectInitializer& ObjectInitializer)
+	: Super(ObjectInitializer)
+{
+	// Structure to hold one-time initialization
+	struct FConstructorStatics
+	{
+		FText NAME_Parameters;
+		FConstructorStatics()
+			: NAME_Parameters(LOCTEXT("Parameters", "Parameters"))
+		{
+		}
+	};
+	static FConstructorStatics ConstructorStatics;
+
+	bIsParameterExpression = true;
+
+#if WITH_EDITORONLY_DATA
+	MenuCategories.Add(ConstructorStatics.NAME_Parameters);
+#endif
+}
+
+#if WITH_EDITOR
+
+bool UMaterialExpressionSparseVolumeTextureSampleParameter::SetParameterValue(FName InParameterName, USparseVolumeTexture* InValue, EMaterialExpressionSetParameterValueFlags Flags)
+{
+	if (InParameterName == ParameterName)
+	{
+		SparseVolumeTexture = InValue;
+		if (EnumHasAnyFlags(Flags, EMaterialExpressionSetParameterValueFlags::SendPostEditChangeProperty))
+		{
+			SendPostEditChangeProperty(this, TEXT("SparseVolumeTexture"));
+		}
+		return true;
+	}
+
+	return false;
+}
+
+void UMaterialExpressionSparseVolumeTextureSampleParameter::ValidateParameterName(const bool bAllowDuplicateName)
+{
+	ValidateParameterNameInternal(this, Material, bAllowDuplicateName);
+}
+
+void UMaterialExpressionSparseVolumeTextureSampleParameter::GetCaption(TArray<FString>& OutCaptions) const
+{
+	OutCaptions.Add(FString(TEXT("Sparse Volume Texture Sample Param ")));
+	OutCaptions.Add(FString::Printf(TEXT("'%s'"), *ParameterName.ToString()));
+}
+
+bool UMaterialExpressionSparseVolumeTextureSampleParameter::MatchesSearchQuery(const TCHAR* SearchQuery)
+{
+	if (ParameterName.ToString().Contains(SearchQuery))
+	{
+		return true;
+	}
+	return Super::MatchesSearchQuery(SearchQuery);
+}
+#endif
+
 #undef LOCTEXT_NAMESPACE
diff --git a/Engine/Source/Runtime/Engine/Private/Materials/MaterialHLSLEmitter.cpp b/Engine/Source/Runtime/Engine/Private/Materials/MaterialHLSLEmitter.cpp
--- a/Engine/Source/Runtime/Engine/Private/Materials/MaterialHLSLEmitter.cpp
+++ b/Engine/Source/Runtime/Engine/Private/Materials/MaterialHLSLEmitter.cpp
@@ -42,6 +42,11 @@ static const TCHAR* HLSLTypeString(EMaterialValueType Type)
 	case MCT_TextureVirtual:		return TEXT("TextureVirtual");
 	case MCT_VTPageTableResult:		return TEXT("VTPageTableResult");
 	case MCT_ShadingModel:			return TEXT("uint");
+	case MCT_UInt:					return TEXT("uint");
+	case MCT_UInt1:					return TEXT("uint");
+	case MCT_UInt2:					return TEXT("uint2");
+	case MCT_UInt3:					return TEXT("uint3");
+	case MCT_UInt4:					return TEXT("uint4");
 	case MCT_Strata:				return TEXT("FStrataData");
 	default:						return TEXT("unknown");
 	};
@@ -135,7 +140,10 @@ static FString GenerateMaterialTemplateHLSL(EShaderPlatform ShaderPlatform,
 		case MCT_Float2: HLSLType = TEXT("float2"); break;
 		case MCT_Float3: HLSLType = TEXT("float3"); break;
 		case MCT_Float4: HLSLType = TEXT("float4"); break;
-		case MCT_ShadingModel: HLSLType = TEXT("uint"); break;
+		case MCT_ShadingModel: case MCT_UInt:  case MCT_UInt1:  HLSLType = TEXT("uint"); break;
+		case MCT_UInt2:  HLSLType = TEXT("uint2"); break;
+		case MCT_UInt3:  HLSLType = TEXT("uint3"); break;
+		case MCT_UInt4:  HLSLType = TEXT("uint4"); break;
 		case MCT_Strata: HLSLType = TEXT("FStrataData"); break;
 		default: break;
 		}
diff --git a/Engine/Source/Runtime/Engine/Private/Materials/MaterialInstance.cpp b/Engine/Source/Runtime/Engine/Private/Materials/MaterialInstance.cpp
--- a/Engine/Source/Runtime/Engine/Private/Materials/MaterialInstance.cpp
+++ b/Engine/Source/Runtime/Engine/Private/Materials/MaterialInstance.cpp
@@ -21,6 +21,7 @@
 #include "Materials/MaterialExpressionFontSampleParameter.h"
 #include "Materials/MaterialExpressionMaterialAttributeLayers.h"
 #include "Materials/MaterialExpressionRuntimeVirtualTextureSampleParameter.h"
+#include "Materials/MaterialExpressionSparseVolumeTextureSample.h"
 #include "Materials/MaterialExpressionStaticSwitchParameter.h"
 #include "Materials/MaterialExpressionStaticComponentMaskParameter.h"
 #include "Materials/MaterialFunctionInstance.h"
@@ -53,6 +54,7 @@
 #include "ShaderCompilerCore.h"
 #include "ShaderCompiler.h"
 #include "MaterialCachedData.h"
+#include "SparseVolumeTexture/SparseVolumeTexture.h"
 #include "ComponentRecreateRenderStateContext.h"
 
 #include UE_INLINE_GENERATED_CPP_BY_NAME(MaterialInstance)
@@ -277,6 +279,7 @@ bool FMaterialInstanceResource::GetParameterValue(EMaterialParameterType Type, c
 		case EMaterialParameterType::DoubleVector: bResult = RenderThread_GetParameterValue<FVector4d>(ParameterInfo, OutValue); break;
 		case EMaterialParameterType::Texture: bResult = RenderThread_GetParameterValue<const UTexture*>(ParameterInfo, OutValue); break;
 		case EMaterialParameterType::RuntimeVirtualTexture: bResult = RenderThread_GetParameterValue<const URuntimeVirtualTexture*>(ParameterInfo, OutValue); break;
+		case EMaterialParameterType::SparseVolumeTexture: bResult = RenderThread_GetParameterValue<const USparseVolumeTexture*>(ParameterInfo, OutValue); break;
 		default: ensure(false); break; // other parameter types are not expected on the render thread
 		}
 	}
@@ -368,12 +371,14 @@ void FMaterialInstanceResource::InitMIParameters(FMaterialInstanceParameterSet&
 	ParameterSet.DoubleVectorParameters.Sort(SortMaterialInstanceParametersPredicate<FVector4d>);
 	ParameterSet.TextureParameters.Sort(SortMaterialInstanceParametersPredicate<const UTexture*>);
 	ParameterSet.RuntimeVirtualTextureParameters.Sort(SortMaterialInstanceParametersPredicate<const URuntimeVirtualTexture*>);
+	ParameterSet.SparseVolumeTextureParameters.Sort(SortMaterialInstanceParametersPredicate<const USparseVolumeTexture*>);
 
 	Swap(ScalarParameterArray.Array, ParameterSet.ScalarParameters);
 	Swap(VectorParameterArray.Array, ParameterSet.VectorParameters);
 	Swap(DoubleVectorParameterArray.Array, ParameterSet.DoubleVectorParameters);
 	Swap(TextureParameterArray.Array, ParameterSet.TextureParameters);
 	Swap(RuntimeVirtualTextureParameterArray.Array, ParameterSet.RuntimeVirtualTextureParameters);
+	Swap(SparseVolumeTextureParameterArray.Array, ParameterSet.SparseVolumeTextureParameters);
 
 	// Build hash tables.
 	ScalarParameterArray.HashAddAllItems();
@@ -381,6 +386,7 @@ void FMaterialInstanceResource::InitMIParameters(FMaterialInstanceParameterSet&
 	DoubleVectorParameterArray.HashAddAllItems();
 	TextureParameterArray.HashAddAllItems();
 	RuntimeVirtualTextureParameterArray.HashAddAllItems();
+	SparseVolumeTextureParameterArray.HashAddAllItems();
 }
 
 /**
@@ -533,6 +539,7 @@ void UMaterialInstance::SwapLayerParameterIndices(int32 OriginalIndex, int32 New
 		SwapLayerParameterIndicesArray(DoubleVectorParameterValues, OriginalIndex, NewIndex);
 		SwapLayerParameterIndicesArray(TextureParameterValues, OriginalIndex, NewIndex);
 		SwapLayerParameterIndicesArray(RuntimeVirtualTextureParameterValues, OriginalIndex, NewIndex);
+		SwapLayerParameterIndicesArray(SparseVolumeTextureParameterValues, OriginalIndex, NewIndex);
 		SwapLayerParameterIndicesArray(FontParameterValues, OriginalIndex, NewIndex);
 		if (EditorOnly)
 		{
@@ -550,6 +557,7 @@ void UMaterialInstance::RemoveLayerParameterIndex(int32 Index)
 	RemoveLayerParameterIndicesArray(DoubleVectorParameterValues, Index);
 	RemoveLayerParameterIndicesArray(TextureParameterValues, Index);
 	RemoveLayerParameterIndicesArray(RuntimeVirtualTextureParameterValues, Index);
+	RemoveLayerParameterIndicesArray(SparseVolumeTextureParameterValues, Index);
 	RemoveLayerParameterIndicesArray(FontParameterValues, Index);
 	if (EditorOnly)
 	{
@@ -600,6 +608,9 @@ bool UMaterialInstance::UpdateParameters()
 			// Runtime Virtual Texture parameters
 			bDirty = UpdateParameterSet<FRuntimeVirtualTextureParameterValue, UMaterialExpressionRuntimeVirtualTextureSampleParameter>(RuntimeVirtualTextureParameterValues, ParentMaterial) || bDirty;
 
+			// Sparse Volume Texture parameters
+			bDirty = UpdateParameterSet<FSparseVolumeTextureParameterValue, UMaterialExpressionSparseVolumeTextureSampleParameter>(SparseVolumeTextureParameterValues, ParentMaterial) || bDirty;
+
 			// Font parameters
 			bDirty = UpdateParameterSet<FFontParameterValue, UMaterialExpressionFontSampleParameter>(FontParameterValues, ParentMaterial) || bDirty;
 
@@ -641,6 +652,7 @@ bool UMaterialInstance::UpdateParameters()
 					RemapLayerParameterIndicesArray(DoubleVectorParameterValues, RemapLayerIndices);
 					RemapLayerParameterIndicesArray(TextureParameterValues, RemapLayerIndices);
 					RemapLayerParameterIndicesArray(RuntimeVirtualTextureParameterValues, RemapLayerIndices);
+					RemapLayerParameterIndicesArray(SparseVolumeTextureParameterValues, RemapLayerIndices);
 					RemapLayerParameterIndicesArray(FontParameterValues, RemapLayerIndices);
 					RemapLayerParameterIndicesArray(EditorOnly->StaticParameters.StaticSwitchParameters, RemapLayerIndices);
 					RemapLayerParameterIndicesArray(EditorOnly->StaticParameters.StaticComponentMaskParameters, RemapLayerIndices);
@@ -753,6 +765,15 @@ void GameThread_InitMIParameters(const UMaterialInstance& Instance)
 		ParamRef.Info = Parameter.ParameterInfo;
 		ParamRef.Value = FRuntimeVirtualTextureParameterValue::GetValue(Parameter);
 	}
+
+	// SparseVolumeTexture parameters
+	ParameterSet.SparseVolumeTextureParameters.Reserve(Instance.SparseVolumeTextureParameterValues.Num());
+	for (const FSparseVolumeTextureParameterValue& Parameter : Instance.SparseVolumeTextureParameterValues)
+	{
+		auto& ParamRef = ParameterSet.SparseVolumeTextureParameters.AddDefaulted_GetRef();
+		ParamRef.Info = Parameter.ParameterInfo;
+		ParamRef.Value = FSparseVolumeTextureParameterValue::GetValue(Parameter);
+	}
 	
 	ENQUEUE_RENDER_COMMAND(InitMIParameters)(
 		[Resource, Parameters = MoveTemp(ParameterSet)](FRHICommandListImmediate& RHICmdList) mutable
@@ -918,6 +939,7 @@ bool UMaterialInstance::GetParameterOverrideValue(EMaterialParameterType Type, c
 	case EMaterialParameterType::DoubleVector: bResult = GameThread_GetParameterValue(DoubleVectorParameterValues, ParameterInfo, OutResult); break;
 	case EMaterialParameterType::Texture: bResult = GameThread_GetParameterValue(TextureParameterValues, ParameterInfo, OutResult); break;
 	case EMaterialParameterType::RuntimeVirtualTexture: bResult = GameThread_GetParameterValue(RuntimeVirtualTextureParameterValues, ParameterInfo, OutResult); break;
+	case EMaterialParameterType::SparseVolumeTexture: bResult = GameThread_GetParameterValue(SparseVolumeTextureParameterValues, ParameterInfo, OutResult); break;
 	case EMaterialParameterType::Font: bResult = GameThread_GetParameterValue(FontParameterValues, ParameterInfo, OutResult); break;
 #if WITH_EDITORONLY_DATA
 	case EMaterialParameterType::StaticSwitch:
@@ -1885,6 +1907,7 @@ void UMaterialInstance::GetAllParametersOfType(EMaterialParameterType Type, TMap
 		case EMaterialParameterType::DoubleVector: GameThread_ApplyParameterOverrides(Instance->DoubleVectorParameterValues, LayerIndexRemap, bSetOverride, OverridenParameters, OutParameters); break;
 		case EMaterialParameterType::Texture: GameThread_ApplyParameterOverrides(Instance->TextureParameterValues, LayerIndexRemap, bSetOverride, OverridenParameters, OutParameters); break;
 		case EMaterialParameterType::RuntimeVirtualTexture: GameThread_ApplyParameterOverrides(Instance->RuntimeVirtualTextureParameterValues, LayerIndexRemap, bSetOverride, OverridenParameters, OutParameters); break;
+		case EMaterialParameterType::SparseVolumeTexture: GameThread_ApplyParameterOverrides(Instance->SparseVolumeTextureParameterValues, LayerIndexRemap, bSetOverride, OverridenParameters, OutParameters); break;
 		case EMaterialParameterType::Font: GameThread_ApplyParameterOverrides(Instance->FontParameterValues, LayerIndexRemap, bSetOverride, OverridenParameters, OutParameters); break;
 #if WITH_EDITORONLY_DATA
 		case EMaterialParameterType::StaticSwitch:
@@ -2860,6 +2883,17 @@ void UMaterialInstance::PostLoad()
 		}
 	}
 
+	// do the same for sparse virtual textures
+	for (int32 ValueIndex = 0; ValueIndex < SparseVolumeTextureParameterValues.Num(); ValueIndex++)
+	{
+		// Make sure the texture is postloaded so the resource isn't null.
+		USparseVolumeTexture* Value = SparseVolumeTextureParameterValues[ValueIndex].ParameterValue;
+		if (Value)
+		{
+			Value->ConditionalPostLoad();
+		}
+	}
+
 	// do the same for font textures
 	for( int32 ValueIndex=0; ValueIndex < FontParameterValues.Num(); ValueIndex++ )
 	{
@@ -3197,6 +3231,7 @@ void UMaterialInstance::ReserveParameterValuesInternal(EMaterialParameterType Ty
 	case EMaterialParameterType::Texture: TextureParameterValues.Reserve(Capacity); break;
 	case EMaterialParameterType::Font: FontParameterValues.Reserve(Capacity); break;
 	case EMaterialParameterType::RuntimeVirtualTexture: RuntimeVirtualTextureParameterValues.Reserve(Capacity); break;
+	case EMaterialParameterType::SparseVolumeTexture: SparseVolumeTextureParameterValues.Reserve(Capacity); break;
 	default: checkNoEntry();
 	}
 }
@@ -3224,6 +3259,7 @@ void UMaterialInstance::AddParameterValueInternal(const FMaterialParameterInfo&
 	case EMaterialParameterType::Texture: TextureParameterValues.Emplace(ParameterInfo, Value.Texture); break;
 	case EMaterialParameterType::Font: FontParameterValues.Emplace(ParameterInfo, Value.Font.Value, Value.Font.Page); break;
 	case EMaterialParameterType::RuntimeVirtualTexture: RuntimeVirtualTextureParameterValues.Emplace(ParameterInfo, Value.RuntimeVirtualTexture); break;
+	case EMaterialParameterType::SparseVolumeTexture: SparseVolumeTextureParameterValues.Emplace(ParameterInfo, Value.SparseVolumeTexture); break;
 	default: checkNoEntry();
 	}
 }
@@ -3251,6 +3287,7 @@ void UMaterialInstance::SetParameterValueInternal(const FMaterialParameterInfo&
 	case EMaterialParameterType::Texture: SetTextureParameterValueInternal(ParameterInfo, Value.Texture); break;
 	case EMaterialParameterType::Font: SetFontParameterValueInternal(ParameterInfo, Value.Font.Value, Value.Font.Page); break;
 	case EMaterialParameterType::RuntimeVirtualTexture: SetRuntimeVirtualTextureParameterValueInternal(ParameterInfo, Value.RuntimeVirtualTexture); break;
+	case EMaterialParameterType::SparseVolumeTexture: SetSparseVolumeTextureParameterValueInternal(ParameterInfo, Value.SparseVolumeTexture); break;
 	default: checkNoEntry();
 	}
 }
@@ -3465,6 +3502,39 @@ void UMaterialInstance::SetRuntimeVirtualTextureParameterValueInternal(const FMa
 	}
 }
 
+void UMaterialInstance::SetSparseVolumeTextureParameterValueInternal(const FMaterialParameterInfo& ParameterInfo, USparseVolumeTexture* Value)
+{
+	LLM_SCOPE(ELLMTag::MaterialInstance);
+
+	FSparseVolumeTextureParameterValue* ParameterValue = GameThread_FindParameterByName(SparseVolumeTextureParameterValues, ParameterInfo);
+
+	bool bForceUpdate = false;
+	if (!ParameterValue)
+	{
+		// If there's no element for the named parameter in array yet, add one.
+		ParameterValue = new(SparseVolumeTextureParameterValues) FSparseVolumeTextureParameterValue;
+		ParameterValue->ParameterInfo = ParameterInfo;
+		ParameterValue->ExpressionGUID.Invalidate();
+		bForceUpdate = true;
+	}
+
+	// Don't enqueue an update if it isn't needed
+	if (bForceUpdate || ParameterValue->ParameterValue != Value)
+	{
+		// set as an ensure, because it is somehow possible to accidentally pass non-textures into here via blueprints...
+		if (Value && ensureMsgf(Value->IsA(USparseVolumeTexture::StaticClass()), TEXT("Expecting a USparseVolumeTexture! Value='%s' class='%s'"), *Value->GetName(), *Value->GetClass()->GetName()))
+		{
+			ParameterValue->ParameterValue = Value;
+			// Update the material instance data in the rendering thread.
+			GameThread_UpdateMIParameter(this, *ParameterValue);
+
+#if WITH_EDITOR
+			FObjectCacheEventSink::NotifyReferencedTextureChanged_Concurrent(this);
+#endif
+		}
+	}
+}
+
 void UMaterialInstance::SetFontParameterValueInternal(const FMaterialParameterInfo& ParameterInfo,class UFont* FontValue,int32 FontPage)
 {
 	LLM_SCOPE(ELLMTag::MaterialInstance);
@@ -3508,6 +3578,7 @@ void UMaterialInstance::ClearParameterValuesInternal(EMaterialInstanceClearParam
 	{
 		TextureParameterValues.Empty();
 		RuntimeVirtualTextureParameterValues.Empty();
+		SparseVolumeTextureParameterValues.Empty();
 		FontParameterValues.Empty();
 		bUpdateResource = true;
 	}
@@ -3893,6 +3964,7 @@ void UMaterialInstance::GetResourceSizeEx(FResourceSizeEx& CumulativeResourceSiz
 		CumulativeResourceSize.AddDedicatedSystemMemoryBytes(DoubleVectorParameterValues.Num() * sizeof(THashedMaterialParameterMap<FVector4d>::TNamedParameter));
 		CumulativeResourceSize.AddDedicatedSystemMemoryBytes(TextureParameterValues.Num() * sizeof(THashedMaterialParameterMap<const UTexture*>::TNamedParameter));
 		CumulativeResourceSize.AddDedicatedSystemMemoryBytes(RuntimeVirtualTextureParameterValues.Num() * sizeof(THashedMaterialParameterMap<const URuntimeVirtualTexture*>::TNamedParameter));
+		CumulativeResourceSize.AddDedicatedSystemMemoryBytes(SparseVolumeTextureParameterValues.Num() * sizeof(THashedMaterialParameterMap<const USparseVolumeTexture*>::TNamedParameter));
 		CumulativeResourceSize.AddDedicatedSystemMemoryBytes(FontParameterValues.Num() * sizeof(THashedMaterialParameterMap<const UTexture*>::TNamedParameter));
 
 		// Record space for hash tables as well..
@@ -3901,6 +3973,7 @@ void UMaterialInstance::GetResourceSizeEx(FResourceSizeEx& CumulativeResourceSiz
 		if (DoubleVectorParameterValues.Num()) CumulativeResourceSize.AddDedicatedSystemMemoryBytes(FDefaultSetAllocator::GetNumberOfHashBuckets(DoubleVectorParameterValues.Num()) * sizeof(uint16));
 		if (TextureParameterValues.Num()) CumulativeResourceSize.AddDedicatedSystemMemoryBytes(FDefaultSetAllocator::GetNumberOfHashBuckets(TextureParameterValues.Num()) * sizeof(uint16));
 		if (RuntimeVirtualTextureParameterValues.Num()) CumulativeResourceSize.AddDedicatedSystemMemoryBytes(FDefaultSetAllocator::GetNumberOfHashBuckets(RuntimeVirtualTextureParameterValues.Num()) * sizeof(uint16));
+		if (SparseVolumeTextureParameterValues.Num()) CumulativeResourceSize.AddDedicatedSystemMemoryBytes(FDefaultSetAllocator::GetNumberOfHashBuckets(SparseVolumeTextureParameterValues.Num()) * sizeof(uint16));
 		if (FontParameterValues.Num()) CumulativeResourceSize.AddDedicatedSystemMemoryBytes(FDefaultSetAllocator::GetNumberOfHashBuckets(FontParameterValues.Num()) * sizeof(uint16));
 	}
 }
@@ -4300,6 +4373,10 @@ bool UMaterialInstance::Equivalent(const UMaterialInstance* CompareTo) const
 	{
 		return false;
 	}
+	if (!CompareValueArraysByExpressionGUID(SparseVolumeTextureParameterValues, CompareTo->SparseVolumeTextureParameterValues))
+	{
+		return false;
+	}
 	if (!CompareValueArraysByExpressionGUID(FontParameterValues, CompareTo->FontParameterValues))
 	{
 		return false;
@@ -4532,6 +4609,7 @@ void UMaterialInstance::CopyMaterialUniformParametersInternal(UMaterialInterface
 				MergeParameterOverrides(DoubleVectorParameterValues, AsInstance->DoubleVectorParameterValues);
 				MergeParameterOverrides(TextureParameterValues, AsInstance->TextureParameterValues);
 				MergeParameterOverrides(RuntimeVirtualTextureParameterValues, AsInstance->RuntimeVirtualTextureParameterValues);
+				MergeParameterOverrides(SparseVolumeTextureParameterValues, AsInstance->SparseVolumeTextureParameterValues);
 				// No fonts?
 			}
 			else if (UMaterial* AsMaterial = Cast<UMaterial>(Interface))
@@ -4542,6 +4620,7 @@ void UMaterialInstance::CopyMaterialUniformParametersInternal(UMaterialInterface
 				checkSlow(DoubleVectorParameterValues.Num() == 0);
 				checkSlow(TextureParameterValues.Num() == 0);
 				checkSlow(RuntimeVirtualTextureParameterValues.Num() == 0);
+				checkSlow(SparseVolumeTextureParameterValues.Num() == 0);
 
 				const FMaterialResource* MaterialResource = nullptr;
 				if (UWorld* World = AsMaterial->GetWorld())
diff --git a/Engine/Source/Runtime/Engine/Private/Materials/MaterialInstanceConstant.cpp b/Engine/Source/Runtime/Engine/Private/Materials/MaterialInstanceConstant.cpp
--- a/Engine/Source/Runtime/Engine/Private/Materials/MaterialInstanceConstant.cpp
+++ b/Engine/Source/Runtime/Engine/Private/Materials/MaterialInstanceConstant.cpp
@@ -148,6 +148,12 @@ void UMaterialInstanceConstant::SetRuntimeVirtualTextureParameterValueEditorOnly
 	SetRuntimeVirtualTextureParameterValueInternal(ParameterInfo, Value);
 }
 
+void UMaterialInstanceConstant::SetSparseVolumeTextureParameterValueEditorOnly(const FMaterialParameterInfo& ParameterInfo, USparseVolumeTexture* Value)
+{
+	check(GIsEditor || IsRunningCommandlet());
+	SetSparseVolumeTextureParameterValueInternal(ParameterInfo, Value);
+}
+
 void UMaterialInstanceConstant::SetFontParameterValueEditorOnly(const FMaterialParameterInfo& ParameterInfo,class UFont* FontValue,int32 FontPage)
 {
 	check(GIsEditor || IsRunningCommandlet());
diff --git a/Engine/Source/Runtime/Engine/Private/Materials/MaterialInstanceDynamic.cpp b/Engine/Source/Runtime/Engine/Private/Materials/MaterialInstanceDynamic.cpp
--- a/Engine/Source/Runtime/Engine/Private/Materials/MaterialInstanceDynamic.cpp
+++ b/Engine/Source/Runtime/Engine/Private/Materials/MaterialInstanceDynamic.cpp
@@ -259,6 +259,12 @@ void UMaterialInstanceDynamic::SetRuntimeVirtualTextureParameterValueByInfo(cons
 	SetRuntimeVirtualTextureParameterValueInternal(ParameterInfo, Value);
 }
 
+void UMaterialInstanceDynamic::SetSparseVolumeTextureParameterValue(FName ParameterName, class USparseVolumeTexture* Value)
+{
+	FMaterialParameterInfo ParameterInfo(ParameterName);
+	SetSparseVolumeTextureParameterValueInternal(ParameterInfo, Value);
+}
+
 UTexture* UMaterialInstanceDynamic::K2_GetTextureParameterValue(FName ParameterName)
 {
 	UTexture* Result = NULL;
diff --git a/Engine/Source/Runtime/Engine/Private/Materials/MaterialInstanceSupport.h b/Engine/Source/Runtime/Engine/Private/Materials/MaterialInstanceSupport.h
--- a/Engine/Source/Runtime/Engine/Private/Materials/MaterialInstanceSupport.h
+++ b/Engine/Source/Runtime/Engine/Private/Materials/MaterialInstanceSupport.h
@@ -246,6 +246,7 @@ public:
 		ScalarParameterArray.Empty();
 		TextureParameterArray.Empty();
 		RuntimeVirtualTextureParameterArray.Empty();
+		SparseVolumeTextureParameterArray.Empty();
 	}
 
 	/**
@@ -315,6 +316,7 @@ private:
 	static bool IsValidParameterValue(const FVector4d&) { return true; }
 	static bool IsValidParameterValue(const UTexture* Value) { return Value != nullptr; }
 	static bool IsValidParameterValue(const URuntimeVirtualTexture* Value) { return Value != nullptr; }
+	static bool IsValidParameterValue(const USparseVolumeTexture* Value) { return Value != nullptr; }
 
 	virtual void StartCacheUniformExpressions() const override;
 	virtual void FinishCacheUniformExpressions() const override;
@@ -400,6 +402,8 @@ private:
 	THashedMaterialParameterMap<const UTexture*> TextureParameterArray;
 	/** Runtime Virtual Texture parameters for this material instance. */
 	THashedMaterialParameterMap<const URuntimeVirtualTexture*> RuntimeVirtualTextureParameterArray;
+	/** Sparse Volume Texture parameters for this material instance. */
+	THashedMaterialParameterMap<const USparseVolumeTexture*> SparseVolumeTextureParameterArray;
 	/** Remap layer indices for parent */
 	TArray<int32> ParentLayerIndexRemap;
 };
@@ -409,11 +413,13 @@ template <> FORCEINLINE THashedMaterialParameterMap<FLinearColor>& FMaterialInst
 template <> FORCEINLINE THashedMaterialParameterMap<FVector4d>& FMaterialInstanceResource::GetValueArray() { return DoubleVectorParameterArray; }
 template <> FORCEINLINE THashedMaterialParameterMap<const UTexture*>& FMaterialInstanceResource::GetValueArray() { return TextureParameterArray; }
 template <> FORCEINLINE THashedMaterialParameterMap<const URuntimeVirtualTexture*>& FMaterialInstanceResource::GetValueArray() { return RuntimeVirtualTextureParameterArray; }
+template <> FORCEINLINE THashedMaterialParameterMap<const USparseVolumeTexture*>& FMaterialInstanceResource::GetValueArray() { return SparseVolumeTextureParameterArray; }
 template <> FORCEINLINE const THashedMaterialParameterMap<float>& FMaterialInstanceResource::GetValueArray() const { return ScalarParameterArray; }
 template <> FORCEINLINE const THashedMaterialParameterMap<FLinearColor>& FMaterialInstanceResource::GetValueArray() const { return VectorParameterArray; }
 template <> FORCEINLINE const THashedMaterialParameterMap<FVector4d>& FMaterialInstanceResource::GetValueArray() const { return DoubleVectorParameterArray; }
 template <> FORCEINLINE const THashedMaterialParameterMap<const UTexture*>& FMaterialInstanceResource::GetValueArray() const { return TextureParameterArray; }
 template <> FORCEINLINE const THashedMaterialParameterMap<const URuntimeVirtualTexture*>& FMaterialInstanceResource::GetValueArray() const { return RuntimeVirtualTextureParameterArray; }
+template <> FORCEINLINE const THashedMaterialParameterMap<const USparseVolumeTexture*>& FMaterialInstanceResource::GetValueArray() const { return SparseVolumeTextureParameterArray; }
 
 struct FMaterialInstanceParameterSet
 {
@@ -422,6 +428,7 @@ struct FMaterialInstanceParameterSet
 	TArray<THashedMaterialParameterMap<FVector4d>::TNamedParameter>						DoubleVectorParameters;
 	TArray<THashedMaterialParameterMap<const UTexture*>::TNamedParameter>				TextureParameters;
 	TArray<THashedMaterialParameterMap<const URuntimeVirtualTexture*>::TNamedParameter>	RuntimeVirtualTextureParameters;
+	TArray<THashedMaterialParameterMap<const USparseVolumeTexture*>::TNamedParameter>	SparseVolumeTextureParameters;
 };
 	
 /** Finds a parameter by name from the game thread. */
diff --git a/Engine/Source/Runtime/Engine/Private/Materials/MaterialInterface.cpp b/Engine/Source/Runtime/Engine/Private/Materials/MaterialInterface.cpp
--- a/Engine/Source/Runtime/Engine/Private/Materials/MaterialInterface.cpp
+++ b/Engine/Source/Runtime/Engine/Private/Materials/MaterialInterface.cpp
@@ -826,6 +826,17 @@ bool UMaterialInterface::GetRuntimeVirtualTextureParameterValue(const FHashedMat
 	return false;
 }
 
+bool UMaterialInterface::GetSparseVolumeTextureParameterValue(const FHashedMaterialParameterInfo& ParameterInfo, USparseVolumeTexture*& OutValue, bool bOveriddenOnly) const
+{
+	FMaterialParameterMetadata Result;
+	if (GetParameterValue(EMaterialParameterType::SparseVolumeTexture, ParameterInfo, Result, MakeParameterValueFlags(bOveriddenOnly)))
+	{
+		OutValue = Result.Value.SparseVolumeTexture;
+		return true;
+	}
+	return false;
+}
+
 #if WITH_EDITOR
 bool UMaterialInterface::GetTextureParameterChannelNames(const FHashedMaterialParameterInfo& ParameterInfo, FParameterChannelNames& OutValue) const
 {
@@ -912,6 +923,17 @@ bool UMaterialInterface::GetRuntimeVirtualTextureParameterDefaultValue(const FHa
 	return false;
 }
 
+bool UMaterialInterface::GetSparseVolumeTextureParameterDefaultValue(const FHashedMaterialParameterInfo& ParameterInfo, class USparseVolumeTexture*& OutValue) const
+{
+	FMaterialParameterMetadata Result;
+	if (GetParameterDefaultValue(EMaterialParameterType::SparseVolumeTexture, ParameterInfo, Result))
+	{
+		OutValue = Result.Value.SparseVolumeTexture;
+		return true;
+	}
+	return false;
+}
+
 bool UMaterialInterface::GetFontParameterDefaultValue(const FHashedMaterialParameterInfo& ParameterInfo, class UFont*& OutFontValue, int32& OutFontPage) const
 {
 	FMaterialParameterMetadata Result;
@@ -993,6 +1015,11 @@ void UMaterialInterface::GetAllRuntimeVirtualTextureParameterInfo(TArray<FMateri
 	GetAllParameterInfoOfType(EMaterialParameterType::RuntimeVirtualTexture, OutParameterInfo, OutParameterIds);
 }
 
+void UMaterialInterface::GetAllSparseVolumeTextureParameterInfo(TArray<FMaterialParameterInfo>& OutParameterInfo, TArray<FGuid>& OutParameterIds) const
+{
+	GetAllParameterInfoOfType(EMaterialParameterType::SparseVolumeTexture, OutParameterInfo, OutParameterIds);
+}
+
 void UMaterialInterface::GetAllFontParameterInfo(TArray<FMaterialParameterInfo>& OutParameterInfo, TArray<FGuid>& OutParameterIds) const
 {
 	GetAllParameterInfoOfType(EMaterialParameterType::Font, OutParameterInfo, OutParameterIds);
diff --git a/Engine/Source/Runtime/Engine/Private/Materials/MaterialShared.cpp b/Engine/Source/Runtime/Engine/Private/Materials/MaterialShared.cpp
--- a/Engine/Source/Runtime/Engine/Private/Materials/MaterialShared.cpp
+++ b/Engine/Source/Runtime/Engine/Private/Materials/MaterialShared.cpp
@@ -3651,6 +3651,17 @@ bool FMaterialRenderProxy::GetTextureValue(const FHashedMaterialParameterInfo& P
 	return false;
 }
 
+bool FMaterialRenderProxy::GetTextureValue(const FHashedMaterialParameterInfo& ParameterInfo, const USparseVolumeTexture** OutValue, const FMaterialRenderContext& Context) const
+{
+	FMaterialParameterValue Value;
+	if (GetParameterValue(EMaterialParameterType::SparseVolumeTexture, ParameterInfo, Value, Context))
+	{
+		*OutValue = Value.SparseVolumeTexture;
+		return true;
+	}
+	return false;
+}
+
 static void OnVirtualTextureDestroyedCB(const FVirtualTextureProducerHandle& InHandle, void* Baton)
 {
 	FMaterialRenderProxy* MaterialProxy = static_cast<FMaterialRenderProxy*>(Baton);
@@ -5870,6 +5881,7 @@ UE::Shader::FValue FMaterialParameterValue::AsShaderValue() const
 	case EMaterialParameterType::Texture:
 	case EMaterialParameterType::Font:
 	case EMaterialParameterType::RuntimeVirtualTexture:
+	case EMaterialParameterType::SparseVolumeTexture:
 		// Non-numeric types, can't represent as shader values
 		return UE::Shader::FValue();
 	default:
@@ -5885,6 +5897,7 @@ UObject* FMaterialParameterValue::AsTextureObject() const
 	{
 	case EMaterialParameterType::Texture: Result = Texture; break;
 	case EMaterialParameterType::RuntimeVirtualTexture: Result = RuntimeVirtualTexture; break;
+	case EMaterialParameterType::SparseVolumeTexture: Result = SparseVolumeTexture; break;
 	case EMaterialParameterType::Font:
 		if (Font.Value && Font.Value->Textures.IsValidIndex(Font.Page))
 		{
@@ -5910,6 +5923,7 @@ UE::Shader::FType GetShaderValueType(EMaterialParameterType Type)
 	case EMaterialParameterType::Font:
 		return FMaterialTextureValue::GetTypeName();
 	case EMaterialParameterType::RuntimeVirtualTexture:
+	case EMaterialParameterType::SparseVolumeTexture:
 		return UE::Shader::EValueType::Void; // TODO
 	default:
 		checkNoEntry();
diff --git a/Engine/Source/Runtime/Engine/Private/Materials/MaterialUniformExpressions.cpp b/Engine/Source/Runtime/Engine/Private/Materials/MaterialUniformExpressions.cpp
--- a/Engine/Source/Runtime/Engine/Private/Materials/MaterialUniformExpressions.cpp
+++ b/Engine/Source/Runtime/Engine/Private/Materials/MaterialUniformExpressions.cpp
@@ -17,6 +17,12 @@
 #include "RenderCore.h"
 #include "VirtualTexturing.h"
 #include "VT/RuntimeVirtualTexture.h"
+#include "SparseVolumeTexture/SparseVolumeTexture.h"
+
+static bool IsASparseVolumeTexture(const FMaterialTextureParameterInfo& Parameter)
+{
+	return Parameter.VirtualTextureLayerIndex != 255; // See FMaterialUniformExpressionTexture::GetTextureParameterInfo and constructor
+}
 
 void WriteMaterialUniformAccess(UE::Shader::EValueComponentType ComponentType, uint32 NumComponents, uint32 UniformOffset, FStringBuilderBase& OutResult)
 {
@@ -133,6 +139,14 @@ void FMaterialRenderContext::GetTextureParameterValue(const FHashedMaterialParam
 	}
 }
 
+void FMaterialRenderContext::GetTextureParameterValue(const FHashedMaterialParameterInfo& ParameterInfo, int32 TextureIndex, const USparseVolumeTexture*& OutValue) const
+{
+	if (ParameterInfo.Name.IsNone() || !MaterialRenderProxy || !MaterialRenderProxy->GetTextureValue(ParameterInfo, &OutValue, *this))
+	{
+		OutValue = GetIndexedTexture<USparseVolumeTexture>(Material, TextureIndex);
+	}
+}
+
 FMaterialUniformExpressionType::FMaterialUniformExpressionType(const TCHAR* InName)
 	: Name(InName)
 {
@@ -454,11 +468,37 @@ FShaderParametersMetadata* FUniformExpressionSet::CreateBufferStruct()
 
 	for (int32 i = 0; i < UniformTextureParameters[(uint32)EMaterialTextureParameterType::Volume].Num(); ++i)
 	{
-		check((NextMemberOffset % SHADER_PARAMETER_POINTER_ALIGNMENT) == 0);
-		new(Members) FShaderParametersMetadata::FMember(*VolumeTextureNames[i],TEXT("Texture3D"),__LINE__,NextMemberOffset,UBMT_TEXTURE,EShaderPrecisionModifier::Float,1,1,0,NULL);
-		NextMemberOffset += SHADER_PARAMETER_POINTER_ALIGNMENT;
-		new(Members) FShaderParametersMetadata::FMember(*VolumeTextureSamplerNames[i],TEXT("SamplerState"),__LINE__,NextMemberOffset,UBMT_SAMPLER,EShaderPrecisionModifier::Float,1,1,0,NULL);
-		NextMemberOffset += SHADER_PARAMETER_POINTER_ALIGNMENT;
+		const FMaterialTextureParameterInfo& Parameter = UniformTextureParameters[(uint32)EMaterialTextureParameterType::Volume][i];
+		if (!IsASparseVolumeTexture(Parameter))
+		{
+			check((NextMemberOffset % SHADER_PARAMETER_POINTER_ALIGNMENT) == 0);
+			new(Members) FShaderParametersMetadata::FMember(*VolumeTextureNames[i], TEXT("Texture3D"), __LINE__, NextMemberOffset, UBMT_TEXTURE, EShaderPrecisionModifier::Float, 1, 1, 0, NULL);
+			NextMemberOffset += SHADER_PARAMETER_POINTER_ALIGNMENT;
+			new(Members) FShaderParametersMetadata::FMember(*VolumeTextureSamplerNames[i], TEXT("SamplerState"), __LINE__, NextMemberOffset, UBMT_SAMPLER, EShaderPrecisionModifier::Float, 1, 1, 0, NULL);
+			NextMemberOffset += SHADER_PARAMETER_POINTER_ALIGNMENT;
+		}
+		else
+		{
+			// Sparse volume texture bindings
+			if (Parameter.VirtualTextureLayerIndex == 0)
+			{
+				check((NextMemberOffset % SHADER_PARAMETER_POINTER_ALIGNMENT) == 0);
+				new(Members) FShaderParametersMetadata::FMember(*VolumeTextureNames[i], TEXT("Texture3D<uint>"), __LINE__, NextMemberOffset, UBMT_TEXTURE, EShaderPrecisionModifier::Float, 1, 1, 0, NULL);
+				NextMemberOffset += SHADER_PARAMETER_POINTER_ALIGNMENT;
+				new(Members) FShaderParametersMetadata::FMember(*VolumeTextureSamplerNames[i], TEXT("SamplerState"), __LINE__, NextMemberOffset, UBMT_SAMPLER, EShaderPrecisionModifier::Float, 1, 1, 0, NULL);
+				NextMemberOffset += SHADER_PARAMETER_POINTER_ALIGNMENT;
+			}
+			else if (Parameter.VirtualTextureLayerIndex == 1)
+			{
+				check((NextMemberOffset % SHADER_PARAMETER_POINTER_ALIGNMENT) == 0);
+				new(Members) FShaderParametersMetadata::FMember(*VolumeTextureNames[i], TEXT("Texture3D"), __LINE__, NextMemberOffset, UBMT_TEXTURE, EShaderPrecisionModifier::Float, 1, 1, 0, NULL);
+				NextMemberOffset += SHADER_PARAMETER_POINTER_ALIGNMENT;
+				new(Members) FShaderParametersMetadata::FMember(*VolumeTextureSamplerNames[i], TEXT("SamplerState"), __LINE__, NextMemberOffset, UBMT_SAMPLER, EShaderPrecisionModifier::Float, 1, 1, 0, NULL);
+				NextMemberOffset += SHADER_PARAMETER_POINTER_ALIGNMENT;
+			}
+		}
+
+
 	}
 
 	for (int32 i = 0; i < UniformExternalTextureParameters.Num(); ++i)
@@ -666,6 +706,21 @@ void FUniformExpressionSet::GetTextureValue(int32 Index, const FMaterialRenderCo
 	}
 }
 
+void FUniformExpressionSet::GetTextureValue(int32 Index, const FMaterialRenderContext& Context, const FMaterial& Material, const USparseVolumeTexture*& OutValue) const
+{
+	check(IsInParallelRenderingThread());
+	const int32 VolumeTexturesNum = GetNumTextures(EMaterialTextureParameterType::Volume);
+	if (ensure(Index < VolumeTexturesNum))
+	{
+		const FMaterialTextureParameterInfo& Parameter = GetTextureParameter(EMaterialTextureParameterType::Volume, Index);
+		Context.GetTextureParameterValue(Parameter.ParameterInfo, Parameter.TextureIndex, OutValue);
+	}
+	else
+	{
+		OutValue = nullptr;
+	}
+}
+
 void FUniformExpressionSet::FillUniformBuffer(const FMaterialRenderContext& MaterialRenderContext, const FUniformExpressionCache& UniformExpressionCache, const FRHIUniformBufferLayout* UniformBufferLayout, uint8* TempBuffer, int TempBufferSize) const
 {
 	FillUniformBuffer(MaterialRenderContext, UniformExpressionCache.AllocatedVTs, UniformBufferLayout, TempBuffer, TempBufferSize);
@@ -1124,40 +1179,84 @@ void FUniformExpressionSet::FillUniformBuffer(const FMaterialRenderContext& Mate
 		{
 			const FMaterialTextureParameterInfo& Parameter = GetTextureParameter(EMaterialTextureParameterType::Volume, ExpressionIndex);
 
-			const UTexture* Value = nullptr;
-			GetTextureValue(EMaterialTextureParameterType::Volume, ExpressionIndex, MaterialRenderContext,MaterialRenderContext.Material,Value);
-
 			void** ResourceTableTexturePtr = (void**)((uint8*)BufferCursor + 0 * SHADER_PARAMETER_POINTER_ALIGNMENT);
 			void** ResourceTableSamplerPtr = (void**)((uint8*)BufferCursor + 1 * SHADER_PARAMETER_POINTER_ALIGNMENT);
 			BufferCursor = ((uint8*)BufferCursor) + (SHADER_PARAMETER_POINTER_ALIGNMENT * 2);
 			check(BufferCursor <= TempBuffer + TempBufferSize);
 
-			if(Value && Value->GetResource() && Value->TextureReference.TextureReferenceRHI && (Value->GetMaterialType() & MCT_VolumeTexture) != 0u)
+			if (IsASparseVolumeTexture(Parameter))
 			{
-				*ResourceTableTexturePtr = Value->TextureReference.TextureReferenceRHI;
-				const FSamplerStateRHIRef* SamplerSource = &Value->GetResource()->SamplerStateRHI;
-
-				const ESamplerSourceMode SourceMode = Parameter.SamplerSource;
-				if (SourceMode == SSM_Wrap_WorldGroupSettings)
-				{
-					SamplerSource = &Wrap_WorldGroupSettings->SamplerStateRHI;
-				}
-				else if (SourceMode == SSM_Clamp_WorldGroupSettings)
+				const USparseVolumeTexture* SVTexture = nullptr;
+				GetTextureValue(ExpressionIndex, MaterialRenderContext, MaterialRenderContext.Material, SVTexture);
+				if (SVTexture != nullptr)
 				{
-					SamplerSource = &Clamp_WorldGroupSettings->SamplerStateRHI;
-				}
-				check(SourceMode != SSM_TerrainWeightmapGroupSettings); // not allowed for texture volumes
+					const int32 LayerIndex = Parameter.VirtualTextureLayerIndex;
+					// 0 is page table
+					// 1 is density
+					// SVT_TODO update that when attribute setup has been clarified
 
-				check(*SamplerSource);
-				*ResourceTableSamplerPtr = *SamplerSource;
+					const FSparseVolumeTextureSceneProxy* SVTextureProxy = SVTexture->GetSparseVolumeTextureSceneProxy();
+					if (SVTextureProxy)
+					{
+						if (LayerIndex == 0)
+						{
+							check(SVTextureProxy->GetPageTableTextureRHI());
+							*ResourceTableTexturePtr = SVTextureProxy->GetPageTableTextureRHI();
+						}
+						else if(LayerIndex == 1)
+						{
+
+							check(SVTextureProxy->GetTileDataTextureRHI());
+							*ResourceTableTexturePtr = SVTextureProxy->GetTileDataTextureRHI();
+						}
+						else
+						{
+							check(false);
+						}
+
+						*ResourceTableSamplerPtr = TStaticSamplerState<SF_Point, AM_Clamp, AM_Clamp, AM_Clamp, 0, 8>::GetRHI();
+					}
+					else
+					{
+						check(GBlackVolumeTexture->TextureRHI);
+						*ResourceTableTexturePtr = GBlackVolumeTexture->TextureRHI;
+						check(GBlackVolumeTexture->SamplerStateRHI);
+						*ResourceTableSamplerPtr = GBlackVolumeTexture->SamplerStateRHI;
+					}
+				}
 			}
 			else
 			{
-				check(GBlackVolumeTexture->TextureRHI);
-				*ResourceTableTexturePtr = GBlackVolumeTexture->TextureRHI;
-				check(GBlackVolumeTexture->SamplerStateRHI);
-				*ResourceTableSamplerPtr = GBlackVolumeTexture->SamplerStateRHI;
+				const UTexture* Value = nullptr;
+				GetTextureValue(EMaterialTextureParameterType::Volume, ExpressionIndex, MaterialRenderContext, MaterialRenderContext.Material, Value);
+
+				if (Value && Value->GetResource() && Value->TextureReference.TextureReferenceRHI && (Value->GetMaterialType() & MCT_VolumeTexture) != 0u)
+				{
+					*ResourceTableTexturePtr = Value->TextureReference.TextureReferenceRHI;
+					const FSamplerStateRHIRef* SamplerSource = &Value->GetResource()->SamplerStateRHI;
+
+					const ESamplerSourceMode SourceMode = Parameter.SamplerSource;
+					if (SourceMode == SSM_Wrap_WorldGroupSettings)
+					{
+						SamplerSource = &Wrap_WorldGroupSettings->SamplerStateRHI;
+					}
+					else if (SourceMode == SSM_Clamp_WorldGroupSettings)
+					{
+						SamplerSource = &Clamp_WorldGroupSettings->SamplerStateRHI;
+					}
+
+					check(*SamplerSource);
+					*ResourceTableSamplerPtr = *SamplerSource;
+				}
+				else
+				{
+					check(GBlackVolumeTexture->TextureRHI);
+					*ResourceTableTexturePtr = GBlackVolumeTexture->TextureRHI;
+					check(GBlackVolumeTexture->SamplerStateRHI);
+					*ResourceTableSamplerPtr = GBlackVolumeTexture->SamplerStateRHI;
+				}
 			}
+
 		}
 
 		// Cache external texture uniform expressions.
@@ -1384,6 +1483,7 @@ FMaterialUniformExpressionTexture::FMaterialUniformExpressionTexture(int32 InTex
 {
 }
 
+// This constructor is called for setting up VirtualTextures
 FMaterialUniformExpressionTexture::FMaterialUniformExpressionTexture(int32 InTextureIndex, int16 InTextureLayerIndex, int16 InPageTableLayerIndex, EMaterialSamplerType InSamplerType)
 	: TextureIndex(InTextureIndex)
 	, TextureLayerIndex(InTextureLayerIndex)
@@ -1396,11 +1496,24 @@ FMaterialUniformExpressionTexture::FMaterialUniformExpressionTexture(int32 InTex
 {
 }
 
+// This constructor is called for setting up SparseVolumeTextures
+FMaterialUniformExpressionTexture::FMaterialUniformExpressionTexture(int32 InTextureIndex, int16 InSubTextureIndex, EMaterialSamplerType InSamplerType)
+	: TextureIndex(InTextureIndex)
+	, TextureLayerIndex(InSubTextureIndex)
+	, PageTableLayerIndex(INDEX_NONE)		// Sparse Volume Textures are using PageTableLayerIndex as a sub texture index (for page table and attributes).
+#if WITH_EDITORONLY_DATA
+	, SamplerType(InSamplerType)
+#endif
+	, SamplerSource(SSM_Wrap_WorldGroupSettings)
+	, bVirtualTexture(false)
+{
+}
+
 void FMaterialUniformExpressionTexture::GetTextureParameterInfo(FMaterialTextureParameterInfo& OutParameter) const
 {
 	OutParameter.TextureIndex = TextureIndex;
 	OutParameter.SamplerSource = SamplerSource;
-	OutParameter.VirtualTextureLayerIndex = TextureLayerIndex;
+	OutParameter.VirtualTextureLayerIndex = TextureLayerIndex;	// VirtualTextureLayerIndex will be 255 if PageTableLayerIndex==INDEX_NONE.
 }
 
 bool FMaterialUniformExpressionTexture::IsIdentical(const FMaterialUniformExpression* OtherExpression) const
@@ -1562,6 +1675,45 @@ void FMaterialUniformExpressionRuntimeVirtualTextureUniform::WriteNumberOpcodes(
 	OutData.WriteOpcode(UE::Shader::EPreshaderOpcode::RuntimeVirtualTextureUniform).Write(WriteParameterInfo).Write((int32)TextureIndex).Write((int32)VectorIndex);
 }
 
+FMaterialUniformExpressionSparseVolumeTextureUniform::FMaterialUniformExpressionSparseVolumeTextureUniform()
+	: bParameter(false)
+	, TextureIndex(INDEX_NONE)
+	, VectorIndex(INDEX_NONE)
+{
+}
+
+FMaterialUniformExpressionSparseVolumeTextureUniform::FMaterialUniformExpressionSparseVolumeTextureUniform(int32 InTextureIndex, int32 InVectorIndex)
+	: bParameter(false)
+	, TextureIndex(InTextureIndex)
+	, VectorIndex(InVectorIndex)
+{
+}
+
+FMaterialUniformExpressionSparseVolumeTextureUniform::FMaterialUniformExpressionSparseVolumeTextureUniform(const FMaterialParameterInfo& InParameterInfo, int32 InTextureIndex, int32 InVectorIndex)
+	: bParameter(true)
+	, ParameterInfo(InParameterInfo)
+	, TextureIndex(InTextureIndex)
+	, VectorIndex(InVectorIndex)
+{
+}
+
+bool FMaterialUniformExpressionSparseVolumeTextureUniform::IsIdentical(const FMaterialUniformExpression* OtherExpression) const
+{
+	if (GetType() != OtherExpression->GetType())
+	{
+		return false;
+	}
+
+	const auto* Other = static_cast<const FMaterialUniformExpressionSparseVolumeTextureUniform*>(OtherExpression);
+	return ParameterInfo == Other->ParameterInfo && TextureIndex == Other->TextureIndex && VectorIndex == Other->VectorIndex;
+}
+
+void FMaterialUniformExpressionSparseVolumeTextureUniform::WriteNumberOpcodes(UE::Shader::FPreshaderData& OutData) const
+{
+	const FHashedMaterialParameterInfo WriteParameterInfo = bParameter ? ParameterInfo : FHashedMaterialParameterInfo();
+	OutData.WriteOpcode(UE::Shader::EPreshaderOpcode::SparseVolumeTextureUniform).Write(WriteParameterInfo).Write((int32)TextureIndex).Write((int32)VectorIndex);
+}
+
 /**
  * Deprecated FMaterialUniformExpressionRuntimeVirtualTextureParameter in favor of FMaterialUniformExpressionRuntimeVirtualTextureUniform
  * Keep around until we no longer need to support serialization of 4.23 data
@@ -1582,6 +1734,7 @@ IMPLEMENT_MATERIALUNIFORMEXPRESSION_TYPE(FMaterialUniformExpressionExternalTextu
 IMPLEMENT_MATERIALUNIFORMEXPRESSION_TYPE(FMaterialUniformExpressionExternalTextureCoordinateScaleRotation);
 IMPLEMENT_MATERIALUNIFORMEXPRESSION_TYPE(FMaterialUniformExpressionExternalTextureCoordinateOffset);
 IMPLEMENT_MATERIALUNIFORMEXPRESSION_TYPE(FMaterialUniformExpressionRuntimeVirtualTextureUniform);
+IMPLEMENT_MATERIALUNIFORMEXPRESSION_TYPE(FMaterialUniformExpressionSparseVolumeTextureUniform);
 IMPLEMENT_MATERIALUNIFORMEXPRESSION_TYPE(FMaterialUniformExpressionFlipBookTextureParameter);
 IMPLEMENT_MATERIALUNIFORMEXPRESSION_TYPE(FMaterialUniformExpressionSine);
 IMPLEMENT_MATERIALUNIFORMEXPRESSION_TYPE(FMaterialUniformExpressionSquareRoot);
diff --git a/Engine/Source/Runtime/Engine/Private/Materials/MaterialUniformExpressions.h b/Engine/Source/Runtime/Engine/Private/Materials/MaterialUniformExpressions.h
--- a/Engine/Source/Runtime/Engine/Private/Materials/MaterialUniformExpressions.h
+++ b/Engine/Source/Runtime/Engine/Private/Materials/MaterialUniformExpressions.h
@@ -83,6 +83,7 @@ public:
 	FMaterialUniformExpressionTexture();
 	FMaterialUniformExpressionTexture(int32 InTextureIndex, EMaterialSamplerType InSamplerType, ESamplerSourceMode InSamplerSource, bool InVirtualTexture);
 	FMaterialUniformExpressionTexture(int32 InTextureIndex, int16 InTextureLayerIndex, int16 InPageTableLayerIndex, EMaterialSamplerType InSamplerType);
+	FMaterialUniformExpressionTexture(int32 InTextureIndex, int16 InSubTextureIndex, EMaterialSamplerType InSamplerType);
 
 	//~ Begin FMaterialUniformExpression Interface.
 	virtual class FMaterialUniformExpressionTexture* GetTextureUniformExpression() { return this; }
@@ -357,6 +358,11 @@ public:
 		, ParameterInfo(InParameterInfo)
 	{}
 
+	FMaterialUniformExpressionTextureParameter(const FMaterialParameterInfo& InParameterInfo, int32 InTextureIndex, int16 InSubTextureIndex, EMaterialSamplerType InSamplerType)
+		: Super(InTextureIndex, InSubTextureIndex, InSamplerType)
+		, ParameterInfo(InParameterInfo)
+	{}
+
 	// FMaterialUniformExpression interface.
 	virtual class FMaterialUniformExpressionTextureParameter* GetTextureParameterUniformExpression() override { return this; }
 
@@ -1559,3 +1565,34 @@ protected:
 	/** Index of the uniform vector to fetch from the URuntimeVirtualTexture. */
 	int32 VectorIndex;
 };
+
+/**
+ * A uniform expression to retrieve one of the vector uniform parameters stored in a USparseVolumeTexture
+ */
+class FMaterialUniformExpressionSparseVolumeTextureUniform : public FMaterialUniformExpression
+{
+	DECLARE_MATERIALUNIFORMEXPRESSION_TYPE(FMaterialUniformExpressionSparseVolumeTextureUniform);
+
+public:
+	FMaterialUniformExpressionSparseVolumeTextureUniform();
+	/** Construct with the index of the texture reference and the vector index that we want to retrieve. */
+	FMaterialUniformExpressionSparseVolumeTextureUniform(int32 InTextureIndex, int32 InVectorIndex);
+	/** Construct with a URuntimeVirtualTexture parameter and the vector index that we want to retrieve. */
+	FMaterialUniformExpressionSparseVolumeTextureUniform(const FMaterialParameterInfo& InParameterInfo, int32 InTextureIndex, int32 InVectorIndex);
+
+	//~ Begin FMaterialUniformExpression Interface.
+	virtual bool IsConstant() const override { return false; }
+	virtual bool IsIdentical(const FMaterialUniformExpression* OtherExpression) const override;
+	virtual void WriteNumberOpcodes(UE::Shader::FPreshaderData& OutData) const override;
+	//~ End FMaterialUniformExpression Interface.
+
+protected:
+	/** Is this expression using a material instance parameter. */
+	bool bParameter;
+	/** Contains the parameter info used if bParameter is true. */
+	FHashedMaterialParameterInfo ParameterInfo;
+	/** Index of the associated URuntimeVirtualTexture in the material texture references used if bParameter is false. */
+	int32 TextureIndex;
+	/** Index of the uniform vector to fetch from the URuntimeVirtualTexture. */
+	int32 VectorIndex;
+};
diff --git a/Engine/Source/Runtime/Engine/Private/Shader/Preshader.cpp b/Engine/Source/Runtime/Engine/Private/Shader/Preshader.cpp
--- a/Engine/Source/Runtime/Engine/Private/Shader/Preshader.cpp
+++ b/Engine/Source/Runtime/Engine/Private/Shader/Preshader.cpp
@@ -7,6 +7,7 @@
 #include "VT/RuntimeVirtualTexture.h"
 #include "Hash/xxhash.h"
 #include "ExternalTexture.h"
+#include "SparseVolumeTexture/SparseVolumeTexture.h"
 
 IMPLEMENT_TYPE_LAYOUT(UE::Shader::FPreshaderData);
 IMPLEMENT_TYPE_LAYOUT(UE::Shader::FPreshaderStructType);
@@ -806,6 +807,27 @@ static void EvaluateRuntimeVirtualTextureUniform(const FMaterialRenderContext& C
 	}
 }
 
+static void EvaluateSparseTextureUniform(const FMaterialRenderContext& Context, FPreshaderStack& Stack, FPreshaderDataContext& RESTRICT Data)
+{
+	const FHashedMaterialParameterInfo ParameterInfo = ReadPreshaderValue<FHashedMaterialParameterInfo>(Data);
+	const int32 TextureIndex = ReadPreshaderValue<int32>(Data);
+	const int32 VectorIndex = ReadPreshaderValue<int32>(Data);
+
+	const USparseVolumeTexture* Texture = nullptr;
+	if (ParameterInfo.Name.IsNone() || !Context.MaterialRenderProxy || !Context.MaterialRenderProxy->GetTextureValue(ParameterInfo, &Texture, Context))
+	{
+		Texture = GetIndexedTexture<USparseVolumeTexture>(Context.Material, TextureIndex);
+	}
+	if (Texture != nullptr && VectorIndex != INDEX_NONE)
+	{
+		Stack.PushValue(FValue(Texture->GetUniformParameter(VectorIndex)));
+	}
+	else
+	{
+		Stack.PushValue(FValue(0.f, 0.f, 0.f, 0.f));
+	}
+}
+
 static void EvaluateJump(FPreshaderDataContext& RESTRICT Data)
 {
 	const int32 JumpOffset = ReadPreshaderValue<int32>(Data);
@@ -888,6 +910,7 @@ FPreshaderValue EvaluatePreshader(const FUniformExpressionSet* UniformExpression
 		case EPreshaderOpcode::ExternalTextureCoordinateScaleRotation: EvaluateExternalTextureCoordinateScaleRotation(Context, Stack, Data); break;
 		case EPreshaderOpcode::ExternalTextureCoordinateOffset: EvaluateExternalTextureCoordinateOffset(Context, Stack, Data); break;
 		case EPreshaderOpcode::RuntimeVirtualTextureUniform: EvaluateRuntimeVirtualTextureUniform(Context, Stack, Data); break;
+		case EPreshaderOpcode::SparseVolumeTextureUniform: EvaluateSparseTextureUniform(Context, Stack, Data); break;
 		case EPreshaderOpcode::Jump: EvaluateJump(Data); break;
 		case EPreshaderOpcode::JumpIfFalse: EvaluateJumpIfFalse(Stack, Data); break;
 		default:
diff --git a/Engine/Source/Runtime/Engine/Private/SparseVolumeTexture/SparseVolumeTexture.cpp b/Engine/Source/Runtime/Engine/Private/SparseVolumeTexture/SparseVolumeTexture.cpp
new file mode 100644
--- /dev/null
+++ b/Engine/Source/Runtime/Engine/Private/SparseVolumeTexture/SparseVolumeTexture.cpp
@@ -0,0 +1,682 @@
+// Copyright Epic Games, Inc. All Rights Reserved.
+
+/*=============================================================================
+	SparseVolumeTexture.cpp: SparseVolumeTexture implementation.
+=============================================================================*/
+
+#include "SparseVolumeTexture/SparseVolumeTexture.h"
+
+#include "Misc/SecureHash.h"
+#include "EngineUtils.h"
+
+#if WITH_EDITORONLY_DATA
+#include "DerivedDataCache/Public/DerivedDataCacheInterface.h"
+#endif
+
+#include "Serialization/LargeMemoryReader.h"
+#include "Serialization/LargeMemoryWriter.h"
+#include "Serialization/MemoryReader.h"
+#include "Serialization/MemoryWriter.h"
+#include "Serialization/BulkDataReader.h"
+#include "Serialization/BulkDataWriter.h"
+#include "Serialization/EditorBulkDataReader.h"
+
+#define LOCTEXT_NAMESPACE "USparseVolumeTexture"
+
+DEFINE_LOG_CATEGORY_STATIC(LogSparseVolumeTexture, Log, All);
+
+void FSparseVolumeAssetHeader::Serialize(FArchive& Ar)
+{
+
+	Ar << Version;
+
+	if (Version == 0)
+	{
+		Ar << PageTableVolumeResolution;
+		Ar << TileDataVolumeResolution;
+		Ar << SourceVolumeResolution;
+	}
+	else
+	{
+		// FSparseVolumeAssetHeader needs to account for new version
+		check(false);
+	}
+}
+
+void FSparseVolumeRawSource::Serialize(FArchive& Ar)
+{
+	Header.Serialize(Ar);
+
+	Ar << Version;
+
+	if (Version == 0)
+	{
+		Ar << DensityData;
+		Ar << DensityPage;
+	}
+	else
+	{
+		// FSparseVolumeRawSource needs to account for new version
+		check(false);
+	}
+}
+
+void FSparseVolumeTextureRuntime::Serialize(FArchive& Ar)
+{
+	Header.Serialize(Ar);
+
+	Ar << Version;
+
+	if (Version == 0)
+	{
+		Ar << DensityPage;
+		Ar << DensityData;
+	}
+	else
+	{
+		// FSparseVolumeRawSource needs to account for new version
+		check(false);
+	}
+}
+
+void FSparseVolumeTextureRuntime::SetAsDefaultTexture()
+{
+	const uint32 VolumeSize = 1;
+	DensityPage.SetNumZeroed(VolumeSize * VolumeSize * VolumeSize);
+	DensityData.SetNumZeroed(VolumeSize * VolumeSize * VolumeSize);
+}
+
+////////////////////////////////////////////////////////////////////////////////////////////////
+
+FSparseVolumeTextureFrame::FSparseVolumeTextureFrame()
+	: RuntimeStreamedInData()
+	, SparseVolumeTextureRuntime()
+	, SparseVolumeTextureSceneProxy()
+#if WITH_EDITORONLY_DATA
+	, RawData()
+#endif
+{
+}
+
+FSparseVolumeTextureFrame::~FSparseVolumeTextureFrame()
+{
+}
+
+////////////////////////////////////////////////////////////////////////////////////////////////
+
+USparseVolumeTexture::USparseVolumeTexture(const FObjectInitializer& ObjectInitializer)
+	: Super(ObjectInitializer)
+{
+}
+
+UE::Shader::EValueType USparseVolumeTexture::GetUniformParameterType(int32 Index)
+{
+	switch (Index)
+	{
+	case ESparseVolumeTexture_PhysicalUVToPageUV:	return UE::Shader::EValueType::Float3;
+	case ESparseVolumeTexture_TileSize:				return UE::Shader::EValueType::Float1;
+	case ESparseVolumeTexture_PageTableSize:		return UE::Shader::EValueType::Float3;
+	default:
+		break;
+	}
+	check(0);
+	return UE::Shader::EValueType::Float4;
+}
+
+////////////////////////////////////////////////////////////////////////////////////////////////
+
+UStaticSparseVolumeTexture::UStaticSparseVolumeTexture(const FObjectInitializer& ObjectInitializer)
+	: Super(ObjectInitializer)
+	, StaticFrame()
+{
+}
+
+void UStaticSparseVolumeTexture::GetResourceSizeEx(FResourceSizeEx& CumulativeResourceSize)
+{
+	Super::GetResourceSizeEx(CumulativeResourceSize);
+	// CumulativeResourceSize.AddDedicatedSystemMemoryBytes but not the RawData size
+}
+
+void UStaticSparseVolumeTexture::PostLoad()
+{
+	GenerateOrLoadDDCRuntimeDataAndCreateSceneProxy();
+
+	Super::PostLoad();
+}
+
+void UStaticSparseVolumeTexture::BeginDestroy()
+{
+	Super::BeginDestroy();
+
+	BeginReleaseResource(&StaticFrame.SparseVolumeTextureSceneProxy);
+}
+
+void UStaticSparseVolumeTexture::Serialize(FArchive& Ar)
+{
+	Super::Serialize(Ar);
+
+	FStripDataFlags StripFlags(Ar);
+
+	// The runtime bulk data for static sparse volume texture is always loaded, not streamed in.
+	StaticFrame.RuntimeStreamedInData.SetBulkDataFlags(BULKDATA_ForceInlinePayload);
+
+	if (StripFlags.IsEditorDataStripped() && Ar.IsLoadingFromCookedPackage())
+	{
+		// In this case we are loading in game with a cooked build so we only need to load the runtime data.
+		{
+			FBulkDataReader RuntimeStreamedInData(StaticFrame.RuntimeStreamedInData);
+			StaticFrame.SparseVolumeTextureRuntime.Serialize(RuntimeStreamedInData);
+		}
+
+		// The bulk data is no longer needed
+		StaticFrame.RuntimeStreamedInData.RemoveBulkData();
+
+		// Runtime data is now valid, create the render thread proxy
+		StaticFrame.SparseVolumeTextureSceneProxy.InitialiseRuntimeData(StaticFrame.SparseVolumeTextureRuntime);
+		BeginInitResource(&StaticFrame.SparseVolumeTextureSceneProxy);
+	}
+	else if (Ar.IsCooking())
+	{
+		// We are cooking the game, serialize the asset out.
+		ConvertRawSourceDataToSparseVolumeTextureRuntime();
+
+		{
+			FBulkDataWriter RuntimeStreamedInData(StaticFrame.RuntimeStreamedInData);
+			StaticFrame.SparseVolumeTextureRuntime.Serialize(RuntimeStreamedInData);
+		}
+	}
+	else if (!Ar.IsObjectReferenceCollector())
+	{
+#if WITH_EDITORONLY_DATA
+		// When in EDITOR:
+		//  - We only serialize raw data 
+		//  - The runtime data is fetched/put from/to DDC
+		//  - This EditorBulk data do not load the full and huge OpenVDB data. That is only done explicitly later.
+		StaticFrame.RawData.Serialize(Ar, this);
+#endif
+	}
+}
+
+void UStaticSparseVolumeTexture::PreSave(FObjectPreSaveContext ObjectSaveContext)
+{
+	Super::PreSave(ObjectSaveContext);
+}
+
+#if WITH_EDITOR
+void UStaticSparseVolumeTexture::PostEditChangeProperty(FPropertyChangedEvent& PropertyChangedEvent)
+{
+	Super::PostEditChangeProperty(PropertyChangedEvent);
+
+	GenerateOrLoadDDCRuntimeDataAndCreateSceneProxy();
+}
+#endif // WITH_EDITOR
+
+void UStaticSparseVolumeTexture::ConvertRawSourceDataToSparseVolumeTextureRuntime()
+{
+#if WITH_EDITOR
+	// Check if the virtualized bulk data payload is available now
+	if (StaticFrame.RawData.HasPayloadData())
+	{
+		// When cooking, the runtime data must be serialized out.
+
+		// First, read the source data in from the raw data stored as bulk data
+		UE::Serialization::FEditorBulkDataReader RawDataArchiveReader(StaticFrame.RawData);
+		FSparseVolumeRawSource SparseVolumeRawSource;
+		SparseVolumeRawSource.Serialize(RawDataArchiveReader);
+
+		// Then, cook the runtime data and serialize it out.
+		StaticFrame.SparseVolumeTextureRuntime.Header = SparseVolumeRawSource.Header;
+		StaticFrame.SparseVolumeTextureRuntime.DensityData = SparseVolumeRawSource.DensityData;
+		StaticFrame.SparseVolumeTextureRuntime.DensityPage = SparseVolumeRawSource.DensityPage;
+
+		// Now unload the raw data
+		StaticFrame.RawData.UnloadData();
+	}
+	else
+	{
+		UE_LOG(LogSparseVolumeTexture, Error, TEXT("SparseVolumeTexture - Raw source data is not available for %s. Using default data."), *GetName());
+		StaticFrame.SparseVolumeTextureRuntime.SetAsDefaultTexture();
+	}
+#endif // WITH_EDITOR
+}
+
+void UStaticSparseVolumeTexture::GenerateOrLoadDDCRuntimeDataAndCreateSceneProxy()
+{
+	// Release any previously allocated render thread proxy
+	BeginReleaseResource(&StaticFrame.SparseVolumeTextureSceneProxy);
+
+#if WITH_EDITORONLY_DATA
+	// We only fetch/put DDC when in editor. Otherwise, StaticFrame.SparseVolumeTextureRuntime is serialize in.
+	GenerateOrLoadDDCRuntimeData();
+#endif
+
+	// Runtime data is now valid, create the render thread proxy
+	StaticFrame.SparseVolumeTextureSceneProxy.InitialiseRuntimeData(StaticFrame.SparseVolumeTextureRuntime);
+	BeginInitResource(&StaticFrame.SparseVolumeTextureSceneProxy);
+}
+
+void UStaticSparseVolumeTexture::GenerateOrLoadDDCRuntimeData()
+{
+#if WITH_EDITORONLY_DATA
+	static const FString SparseVolumeTextureDDCVersion = TEXT("381AE2A9-A903-4C8F-8486-891E24D6EC70");	// Bump this if you want to ignore all cached data so far.
+	const FString DerivedDataKey = StaticFrame.RawData.GetIdentifier().ToString() + SparseVolumeTextureDDCVersion;
+
+	bool bSuccess = true;
+	TArray<uint8> DerivedData;
+	if (GetDerivedDataCacheRef().GetSynchronous(*DerivedDataKey, DerivedData, GetPathName()))
+	{
+		UE_LOG(LogSparseVolumeTexture, Display, TEXT("SparseVolumeTexture - Caching %s"), *GetName());
+
+		FMemoryReader Ar(DerivedData, /*bIsPersistent=*/ true);
+
+		int64 UncompressedSize = 0;
+		Ar << UncompressedSize;
+
+		uint8* DecompressionBuffer = reinterpret_cast<uint8*>(FMemory::Malloc(UncompressedSize));
+		Ar.SerializeCompressed(DecompressionBuffer, UncompressedSize, NAME_Zlib);
+
+		FLargeMemoryReader LargeMemReader(DecompressionBuffer, UncompressedSize, ELargeMemoryReaderFlags::Persistent | ELargeMemoryReaderFlags::TakeOwnership);
+
+		StaticFrame.SparseVolumeTextureRuntime.Serialize(LargeMemReader);
+	}
+	else
+	{
+		// Using a LargeMemoryWriter for serialization since the data can be bigger than 2 GB
+		FLargeMemoryWriter LargeMemWriter(0, /*bIsPersistent=*/ true);
+
+		ConvertRawSourceDataToSparseVolumeTextureRuntime();
+		StaticFrame.SparseVolumeTextureRuntime.Serialize(LargeMemWriter);
+
+		int64 UncompressedSize = LargeMemWriter.TotalSize();
+
+		// Since the DDC doesn't support data bigger than 2 GB, we only cache for such uncompressed size.
+		constexpr int64 SizeThreshold = 2147483648;	// 2GB
+		const bool bIsCacheable = UncompressedSize < SizeThreshold;
+		if (bIsCacheable)
+		{
+			FMemoryWriter CompressedArchive(DerivedData, true);
+
+			CompressedArchive << UncompressedSize; // needed for allocating decompression buffer
+			CompressedArchive.SerializeCompressed(LargeMemWriter.GetData(), UncompressedSize, NAME_Zlib);
+
+			GetDerivedDataCacheRef().Put(*DerivedDataKey, DerivedData, GetPathName());
+		}
+		else
+		{
+			UE_LOG(LogSparseVolumeTexture, Error, TEXT("SparseVolumeTexture - the asset is too large to fit in Derived Data Cache %s"), *GetName());
+		}
+	}
+#endif // WITH_EDITORONLY_DATA
+}
+
+const FSparseVolumeAssetHeader* UStaticSparseVolumeTexture::GetSparseVolumeTextureHeader() const
+{
+	return &StaticFrame.SparseVolumeTextureRuntime.Header;
+}
+FSparseVolumeTextureSceneProxy* UStaticSparseVolumeTexture::GetSparseVolumeTextureSceneProxy()
+{ 
+	return &StaticFrame.SparseVolumeTextureSceneProxy; 
+}
+const FSparseVolumeTextureSceneProxy* UStaticSparseVolumeTexture::GetSparseVolumeTextureSceneProxy() const
+{ 
+	return &StaticFrame.SparseVolumeTextureSceneProxy; 
+}
+
+FVector4 UStaticSparseVolumeTexture::GetUniformParameter(int32 Index) const
+{
+	const FSparseVolumeTextureSceneProxy* Proxy = GetSparseVolumeTextureSceneProxy();
+	if (Proxy)
+	{
+		const FSparseVolumeAssetHeader& Header = Proxy->GetHeader();
+		switch (Index)
+		{
+		case ESparseVolumeTexture_PhysicalUVToPageUV:
+		{
+			// 3d uv coordinates are specified in [0, 1]. Before addressing the page tables which might have padding,
+			// since source volume resolution might not align to tile resolution, we have to rescale the uv so that [0,1] maps to the source texture boundaries.
+			FVector3f PhysicalUVToPageUV = FVector3f(Header.SourceVolumeResolution)/ FVector3f(Header.PageTableVolumeResolution * SPARSE_VOLUME_TILE_RES);
+			return FVector4(PhysicalUVToPageUV.X, PhysicalUVToPageUV.Y, PhysicalUVToPageUV.Z, 0.0f);
+		}
+		case ESparseVolumeTexture_PageTableSize:
+		{
+			return FVector4(Header.PageTableVolumeResolution.X, Header.PageTableVolumeResolution.Y, Header.PageTableVolumeResolution.Z, 0.0f);
+		}
+		case ESparseVolumeTexture_TileSize:				return FVector4(float(SPARSE_VOLUME_TILE_RES), 0.0f, 0.0f, 0.0f);
+		default:
+			break;
+		}
+		check(0);
+		return FVector4(ForceInitToZero);
+	}
+	// 0 while waiting for the proxy
+	return FVector4(ForceInitToZero);
+}
+
+////////////////////////////////////////////////////////////////////////////////////////////////
+
+FSparseVolumeTextureSceneProxy::FSparseVolumeTextureSceneProxy()
+	: FRenderResource()
+	, SparseVolumeTextureRuntime(nullptr)
+	, PageTableTextureRHI(nullptr)
+	, TileDataTextureRHI(nullptr)
+{
+}
+
+FSparseVolumeTextureSceneProxy::~FSparseVolumeTextureSceneProxy()
+{
+}
+
+void FSparseVolumeTextureSceneProxy::InitialiseRuntimeData(FSparseVolumeTextureRuntime& SparseVolumeTextureRuntimeIn)
+{
+	SparseVolumeTextureRuntime = &SparseVolumeTextureRuntimeIn;
+}
+
+void FSparseVolumeTextureSceneProxy::InitRHI()
+{
+	// Page table
+	{
+		EPixelFormat PageEntryFormat = PF_R32_UINT;
+		FIntVector3 PageTableVolumeResolution = SparseVolumeTextureRuntime->Header.PageTableVolumeResolution;
+		const FRHITextureCreateDesc Desc =
+			FRHITextureCreateDesc::Create3D(TEXT("SparseVolumeTexture.PageTable.RHITexture"),
+				PageTableVolumeResolution.X, PageTableVolumeResolution.Y, PageTableVolumeResolution.Z, PageEntryFormat)
+			.SetFlags(ETextureCreateFlags::ShaderResource);
+
+		PageTableTextureRHI = RHICreateTexture(Desc);
+
+		const int32 FormatSize = GPixelFormats[PageEntryFormat].BlockBytes;
+		const FUpdateTextureRegion3D UpdateRegion(0, 0, 0, 0, 0, 0, PageTableVolumeResolution.X, PageTableVolumeResolution.Y, PageTableVolumeResolution.Z);
+		RHIUpdateTexture3D(PageTableTextureRHI, 0, UpdateRegion, UpdateRegion.Width * FormatSize, UpdateRegion.Width * UpdateRegion.Height * FormatSize, (const uint8*)SparseVolumeTextureRuntime->DensityPage.GetData());
+	}
+
+	// Tile data
+	{
+		EPixelFormat VoxelFormat = PF_R8;
+		FIntVector3 TileDataVolumeResolution = SparseVolumeTextureRuntime->Header.TileDataVolumeResolution;
+		const FRHITextureCreateDesc Desc =
+			FRHITextureCreateDesc::Create3D(TEXT("SparseVolumeTexture.TileData.RHITexture"),
+				TileDataVolumeResolution.X, TileDataVolumeResolution.Y, TileDataVolumeResolution.Z, VoxelFormat)
+			.SetFlags(ETextureCreateFlags::ShaderResource);
+
+		TileDataTextureRHI = RHICreateTexture(Desc);
+
+		const int32 FormatSize = GPixelFormats[VoxelFormat].BlockBytes;
+		const FUpdateTextureRegion3D UpdateRegion(0, 0, 0, 0, 0, 0, TileDataVolumeResolution.X, TileDataVolumeResolution.Y, TileDataVolumeResolution.Z);
+		RHIUpdateTexture3D(TileDataTextureRHI, 0, UpdateRegion, UpdateRegion.Width * FormatSize, UpdateRegion.Width * UpdateRegion.Height * FormatSize, (const uint8*)SparseVolumeTextureRuntime->DensityData.GetData());
+	}
+}
+
+void FSparseVolumeTextureSceneProxy::ReleaseRHI()
+{
+}
+
+////////////////////////////////////////////////////////////////////////////////////////////////
+
+UAnimatedSparseVolumeTexture::UAnimatedSparseVolumeTexture(const FObjectInitializer& ObjectInitializer)
+	: Super(ObjectInitializer)
+	, FrameCount(0)
+	, AnimationFrames()
+{
+}
+
+void UAnimatedSparseVolumeTexture::GetResourceSizeEx(FResourceSizeEx& CumulativeResourceSize)
+{
+	Super::GetResourceSizeEx(CumulativeResourceSize);
+	// CumulativeResourceSize.AddDedicatedSystemMemoryBytes but not the RawData size
+}
+
+void UAnimatedSparseVolumeTexture::PostLoad()
+{
+	const int32 FrameCountToLoad = GetFrameCountToLoad();
+	for (int32 FrameIndex = 0; FrameIndex < FrameCountToLoad; FrameIndex++)
+	{
+		GenerateOrLoadDDCRuntimeDataAndCreateSceneProxy(FrameIndex);
+	}
+
+	Super::PostLoad();
+}
+
+int32 UAnimatedSparseVolumeTexture::GetFrameCountToLoad() const 
+{
+	if (FrameCount > 0)
+	{
+		return bLoadAllFramesToProxies ? FrameCount : 1;
+	}
+	return 0;
+}
+
+void UAnimatedSparseVolumeTexture::BeginDestroy()
+{
+	Super::BeginDestroy();
+
+	const int32 FrameCountToLoad = GetFrameCountToLoad();
+	for(int32 FrameIndex = 0; FrameIndex < FrameCountToLoad; FrameIndex++)
+	{
+		FSparseVolumeTextureFrame& Frame = AnimationFrames[FrameIndex];
+		BeginReleaseResource(&Frame.SparseVolumeTextureSceneProxy);
+	}
+}
+
+void UAnimatedSparseVolumeTexture::Serialize(FArchive& Ar)
+{
+	Super::Serialize(Ar);
+
+	FStripDataFlags StripFlags(Ar);
+
+	if (StripFlags.IsEditorDataStripped() && Ar.IsLoadingFromCookedPackage())
+	{
+		check(false); // SVT_TODO Implement
+	}
+	else if (Ar.IsCooking())
+	{
+		check(false); // SVT_TODO Implement
+	}
+	else if (!Ar.IsObjectReferenceCollector())
+	{
+#if WITH_EDITORONLY_DATA
+		// When in EDITOR:
+		//  - We only serialize raw data 
+		//  - The runtime data is fetched/put from/to DDC
+		//  - This EditorBulk data do not load the full and huge OpenVDB data. That is only done explicitly later.
+		if (Ar.IsSaving())
+		{
+			check(AnimationFrames.Num() == FrameCount);
+		}
+		else if (Ar.IsLoading())
+		{
+			AnimationFrames.SetNum(FrameCount);
+		}
+
+		for (int32 i = 0; i < FrameCount; ++i)
+		{
+			AnimationFrames[i].RawData.Serialize(Ar, this);
+		}
+#endif
+	}
+}
+
+void UAnimatedSparseVolumeTexture::PreSave(FObjectPreSaveContext ObjectSaveContext)
+{
+	Super::PreSave(ObjectSaveContext);
+}
+
+#if WITH_EDITOR
+void UAnimatedSparseVolumeTexture::PostEditChangeProperty(FPropertyChangedEvent& PropertyChangedEvent)
+{
+	Super::PostEditChangeProperty(PropertyChangedEvent);
+}
+#endif // WITH_EDITOR
+
+void UAnimatedSparseVolumeTexture::ConvertRawSourceDataToSparseVolumeTextureRuntime(int32 FrameIndex)
+{
+#if WITH_EDITOR
+	FSparseVolumeTextureFrame& Frame = AnimationFrames[FrameIndex];
+
+	// Check if the virtualized bulk data payload is available now
+	if (Frame.RawData.HasPayloadData())
+	{
+		// When cooking, the runtime data must be serialized out.
+
+		// First, read the source data in from the raw data stored as bulk data
+		UE::Serialization::FEditorBulkDataReader RawDataArchiveReader(Frame.RawData);
+		FSparseVolumeRawSource SparseVolumeRawSource;
+		SparseVolumeRawSource.Serialize(RawDataArchiveReader);
+
+		// Then, cook the runtime data and serialize it out.
+		Frame.SparseVolumeTextureRuntime.Header = SparseVolumeRawSource.Header;
+		Frame.SparseVolumeTextureRuntime.DensityData = SparseVolumeRawSource.DensityData;
+		Frame.SparseVolumeTextureRuntime.DensityPage = SparseVolumeRawSource.DensityPage;
+
+		// Now unload the raw data
+		Frame.RawData.UnloadData();
+	}
+	else
+	{
+		UE_LOG(LogSparseVolumeTexture, Error, TEXT("AnimatedSparseVolumeTexture - Raw source data is not available for %s - Frame %i. Using default data."), *GetName(), FrameIndex);
+		Frame.SparseVolumeTextureRuntime.SetAsDefaultTexture();
+	}
+#endif // WITH_EDITOR
+}
+
+void UAnimatedSparseVolumeTexture::GenerateOrLoadDDCRuntimeDataAndCreateSceneProxy(int32 FrameIndex)
+{
+	FSparseVolumeTextureFrame& Frame = AnimationFrames[FrameIndex];
+
+	// Release any previously allocated render thread proxy
+	BeginReleaseResource(&Frame.SparseVolumeTextureSceneProxy);
+
+#if WITH_EDITORONLY_DATA
+	// We only fetch/put DDC when in editor. Otherwise, StaticFrame.SparseVolumeTextureRuntime is serialize in.
+	GenerateOrLoadDDCRuntimeData(FrameIndex);
+#endif
+
+	// Runtime data is now valid, create the render thread proxy
+	Frame.SparseVolumeTextureSceneProxy.InitialiseRuntimeData(Frame.SparseVolumeTextureRuntime);
+	BeginInitResource(&Frame.SparseVolumeTextureSceneProxy);
+}
+
+void UAnimatedSparseVolumeTexture::GenerateOrLoadDDCRuntimeData(int32 FrameIndex)
+{
+#if WITH_EDITORONLY_DATA
+	static const FString SparseVolumeTextureDDCVersion = TEXT("381AE2A9-A903-4C8F-8486-891E24D6EC71");	// Bump this if you want to ignore all cached data so far.
+
+	FSparseVolumeTextureFrame& Frame = AnimationFrames[FrameIndex];
+
+	const FString DerivedDataKey = Frame.RawData.GetIdentifier().ToString() + SparseVolumeTextureDDCVersion;
+
+	bool bSuccess = true;
+	TArray<uint8> DerivedData;
+	if (GetDerivedDataCacheRef().GetSynchronous(*DerivedDataKey, DerivedData, GetPathName()))
+	{
+		UE_LOG(LogSparseVolumeTexture, Display, TEXT("SparseVolumeTexture - Caching %s"), *GetName());
+
+		FMemoryReader Ar(DerivedData, /*bIsPersistent=*/ true);
+
+		int64 UncompressedSize = 0;
+		Ar << UncompressedSize;
+
+		uint8* DecompressionBuffer = reinterpret_cast<uint8*>(FMemory::Malloc(UncompressedSize));
+		Ar.SerializeCompressed(DecompressionBuffer, UncompressedSize, NAME_Zlib);
+
+		FLargeMemoryReader LargeMemReader(DecompressionBuffer, UncompressedSize, ELargeMemoryReaderFlags::Persistent | ELargeMemoryReaderFlags::TakeOwnership);
+
+		Frame.SparseVolumeTextureRuntime.Serialize(LargeMemReader);
+	}
+	else
+	{
+		// Using a LargeMemoryWriter for serialization since the data can be bigger than 2 GB
+		FLargeMemoryWriter LargeMemWriter(0, /*bIsPersistent=*/ true);
+
+		ConvertRawSourceDataToSparseVolumeTextureRuntime(FrameIndex);
+		Frame.SparseVolumeTextureRuntime.Serialize(LargeMemWriter);
+
+		int64 UncompressedSize = LargeMemWriter.TotalSize();
+
+		// Since the DDC doesn't support data bigger than 2 GB, we only cache for such uncompressed size.
+		constexpr int64 SizeThreshold = 2147483648;	// 2GB
+		const bool bIsCacheable = UncompressedSize < SizeThreshold;
+		if (bIsCacheable)
+		{
+			FMemoryWriter CompressedArchive(DerivedData, true);
+
+			CompressedArchive << UncompressedSize; // needed for allocating decompression buffer
+			CompressedArchive.SerializeCompressed(LargeMemWriter.GetData(), UncompressedSize, NAME_Zlib);
+
+			GetDerivedDataCacheRef().Put(*DerivedDataKey, DerivedData, GetPathName());
+		}
+		else
+		{
+			UE_LOG(LogSparseVolumeTexture, Error, TEXT("SparseVolumeTexture - the asset is too large to fit in Derived Data Cache %s - Frame %i"), *GetName(), FrameIndex);
+		}
+	}
+#endif // WITH_EDITORONLY_DATA
+}
+
+const FSparseVolumeAssetHeader* UAnimatedSparseVolumeTexture::GetSparseVolumeTextureHeader() const
+{
+	// When an AnimatedSparseVolumeTexture is used as SparseVolumeTexture, it can only be previewed using its first frame.
+	check(AnimationFrames.Num() >= 1);
+	const FSparseVolumeTextureFrame& Frame = AnimationFrames[PreviewFrameIndex];
+	return &Frame.SparseVolumeTextureRuntime.Header;
+}
+
+FSparseVolumeTextureSceneProxy* UAnimatedSparseVolumeTexture::GetSparseVolumeTextureSceneProxy()
+{
+	// When an AnimatedSparseVolumeTexture is used as SparseVolumeTexture, it can only be previewed using its first frame.
+	check(AnimationFrames.Num() >= 1);
+	FSparseVolumeTextureFrame& Frame = AnimationFrames[PreviewFrameIndex];
+	return &Frame.SparseVolumeTextureSceneProxy;
+}
+
+const FSparseVolumeTextureSceneProxy* UAnimatedSparseVolumeTexture::GetSparseVolumeTextureSceneProxy() const
+{
+	// When an AnimatedSparseVolumeTexture is used as SparseVolumeTexture, it can only be previewed using its first frame.
+	check(AnimationFrames.Num() >= 1);
+	const FSparseVolumeTextureFrame& Frame = AnimationFrames[PreviewFrameIndex];
+	return &Frame.SparseVolumeTextureSceneProxy;
+}
+
+FVector4 UAnimatedSparseVolumeTexture::GetUniformParameter(int32 Index) const
+{
+	// When an AnimatedSparseVolumeTexture is used as SparseVolumeTexture, it can only be previewed using its first frame.
+	const FSparseVolumeTextureSceneProxy* Proxy = GetSparseVolumeTextureSceneProxy();
+	if (Proxy)
+	{
+		const FSparseVolumeAssetHeader& Header = Proxy->GetHeader();
+		switch (Index)
+		{
+		case ESparseVolumeTexture_PhysicalUVToPageUV:
+		{
+			// 3d uv coordinates are specified in [0, 1]. Before addressing the page tables which might have padding,
+			// since source volume resolution might not align to tile resolution, we have to rescale the uv so that [0,1] maps to the source texture boundaries.
+			FVector3f PhysicalUVToPageUV = FVector3f(Header.SourceVolumeResolution) / FVector3f(Header.PageTableVolumeResolution * SPARSE_VOLUME_TILE_RES);
+			return FVector4(PhysicalUVToPageUV.X, PhysicalUVToPageUV.Y, PhysicalUVToPageUV.Z, 0.0f);
+		}
+		case ESparseVolumeTexture_PageTableSize:
+		{
+			return FVector4(Header.PageTableVolumeResolution.X, Header.PageTableVolumeResolution.Y, Header.PageTableVolumeResolution.Z, 0.0f);
+		}
+		case ESparseVolumeTexture_TileSize:				return FVector4(float(SPARSE_VOLUME_TILE_RES), 0.0f, 0.0f, 0.0f);
+		default:
+			break;
+		}
+		check(0);
+		return FVector4(ForceInitToZero);
+	}
+
+	// 0 while waiting for the proxy
+	return FVector4(ForceInitToZero);
+}
+
+
+const FSparseVolumeTextureSceneProxy* UAnimatedSparseVolumeTexture::GetSparseVolumeTextureFrameSceneProxy(int32 FrameIndex) const
+{
+	check(AnimationFrames.Num() >= 1);
+	FrameIndex = FrameIndex % GetFrameCountToLoad();
+	// SVT_TODO when streaming is enabled, this will likely change.
+	return &AnimationFrames[FrameIndex].SparseVolumeTextureSceneProxy;
+}
+
+////////////////////////////////////////////////////////////////////////////////////////////////
+
+#undef LOCTEXT_NAMESPACE
diff --git a/Engine/Source/Runtime/Engine/Private/UnrealEngine.cpp b/Engine/Source/Runtime/Engine/Private/UnrealEngine.cpp
--- a/Engine/Source/Runtime/Engine/Private/UnrealEngine.cpp
+++ b/Engine/Source/Runtime/Engine/Private/UnrealEngine.cpp
@@ -2099,6 +2099,10 @@ void UEngine::Init(IEngineLoop* InEngineLoop)
 		FModuleManager::Get().LoadModuleChecked("MovieScene");
 		FModuleManager::Get().LoadModuleChecked("MovieSceneTracks");
 		FModuleManager::Get().LoadModule("LevelSequence");
+#if WITH_EDITOR
+		// The SparseVolumeTexture module containing the importer is only loaded and used in the editor.
+		FModuleManager::Get().LoadModuleChecked("SparseVolumeTexture");
+#endif
 	}
 
 	// Record large world coordinate state
diff --git a/Engine/Source/Runtime/Engine/Public/MaterialCachedData.h b/Engine/Source/Runtime/Engine/Public/MaterialCachedData.h
--- a/Engine/Source/Runtime/Engine/Public/MaterialCachedData.h
+++ b/Engine/Source/Runtime/Engine/Public/MaterialCachedData.h
@@ -290,6 +290,9 @@ struct FMaterialCachedExpressionData
 	UPROPERTY()
 	TArray<TSoftObjectPtr<URuntimeVirtualTexture>> RuntimeVirtualTextureValues;
 
+	UPROPERTY()
+	TArray<TSoftObjectPtr<USparseVolumeTexture>> SparseVolumeTextureValues;
+
 	/** Array of all texture referenced by this material */
 	UPROPERTY()
 	TArray<TObjectPtr<UObject>> ReferencedTextures;
diff --git a/Engine/Source/Runtime/Engine/Public/MaterialCompiler.h b/Engine/Source/Runtime/Engine/Public/MaterialCompiler.h
--- a/Engine/Source/Runtime/Engine/Public/MaterialCompiler.h
+++ b/Engine/Source/Runtime/Engine/Public/MaterialCompiler.h
@@ -29,6 +29,7 @@ class UMaterialParameterCollection;
 class URuntimeVirtualTexture;
 class UTexture;
 struct FMaterialParameterInfo;
+class USparseVolumeTexture;
 
 enum EMaterialForceCastFlags
 {
@@ -371,6 +372,61 @@ public:
 	virtual int32 ExternalTextureCoordinateOffset(int32 TextureReferenceIndex, TOptional<FName> ParameterName) = 0;
 	virtual int32 ExternalTextureCoordinateOffset(const FGuid& ExternalTextureGuid) = 0;
 
+	// Could be called sub texture and used to support multiple texture samples from a single node?
+	// Making it clear for now and setting explicitly a USparseVolumeTexture as object.
+
+	/**
+	 * Register a sparse volume texture sub-texture to be sampled.
+	 * @param Texture The sparse volume texture to sample.
+	 * @param SubTextureId The sub-texture we want to fetch from (page table, density, etc).
+	 * @param TextureReferenceIndex Output the index of the texture in the referenced textures for that shaders.
+	 * @param SamplerType The sampler type
+	 * @return The code chunk index of the texture.
+	 */
+	virtual int32 SparseVolumeTexture(USparseVolumeTexture* Texture, int32 SubTextureId, int32& TextureReferenceIndex, EMaterialSamplerType SamplerType) = 0;
+	/**
+	 * Register a parameterized sparse volume texture sub-texture to be sampled.
+	 * @param ParameterName The sparse volume texture parameter name.
+	 * @param InDefaultTexture The default static sparse volume texture to sample.
+	 * @param SubTextureId The sub-texture we want to fetch from (page table, density, etc).
+	 * @param TextureReferenceIndex Output the index of the texture in the referenced textures for that shaders.
+	 * @param SamplerType The sampler type
+	 * @return The code chunk index of the texture parameter.
+	 */
+	virtual int32 SparseVolumeTextureParameter(FName ParameterName, USparseVolumeTexture* InDefaultTexture, int32 SubTextureId, int32& TextureReferenceIndex, EMaterialSamplerType SamplerType) = 0;
+	/**
+	 * Sample a sparse volume sub-texture.
+	 * @param Texture The code chunk index of the texture.
+	 * @param Coordinate The coordinate at which to sample the texture.
+	 * @return The code chunk index of result of the texture sample.
+	 */
+	virtual int32 SparseVolumeTextureSample(int32 Texture, int32 Coordinate) = 0;
+	/**
+	 * Register a uniform parameter required to be able to sample a sparse volume texture.
+	 * @param TextureIndex The TextureReferenceIndex.
+	 * @param VectorIndex The index of the vector in the list of vector to bind.
+	 * @param Type The type of the vector.
+	 * @return The code chunk index of the uniform.
+	 */
+	virtual int32 SparseVolumeTextureUniform(int32 TextureIndex, int32 VectorIndex, UE::Shader::EValueType Type) = 0;
+	/**
+	 * Register a uniform parameter required to be able to sample a parameterizes sparse volume texture.
+	 * @param ParameterName The sparse volume texture parameter name.
+	 * @param TextureIndex The TextureReferenceIndex.
+	 * @param VectorIndex The index of the vector in the list of vector to bind.
+	 * @param Type The type of the vector.
+	 * @return The code chunk index of the uniform.
+	 */
+	virtual int32 SparseVolumeTextureUniformParameter(FName ParameterName, int32 TextureIndex, int32 VectorIndex, UE::Shader::EValueType Type) = 0;
+	/**
+	 * @param PackedPhysicalTileCoord: the packed coordinate fetched form the page table.
+	 * @param TileSize: the size of a tile.
+	 * @param CoordPageTable: The tile lower voxel coordinate in the physical tile volume.
+	 * @param CoordVolume: the float UV according the entire volume.
+	 * @return The code chunk index of the voxel coord within the tile volume texture.
+	 */
+	virtual int32 SparseVolumeTextureGetVoxelCoord(int32 PackedPhysicalTileCoord, int32 TileSize, int32 CoordPageTable, int32 CoordVolume) = 0;
+
 	virtual UObject* GetReferencedTexture(int32 Index) { return nullptr; }
 
 	int32 Texture(UTexture* InTexture, EMaterialSamplerType SamplerType, ESamplerSourceMode SamplerSource=SSM_FromTextureAsset)
@@ -379,6 +435,12 @@ public:
 		return Texture(InTexture, TextureReferenceIndex, SamplerType, SamplerSource);
 	}
 
+	int32 SparseVolumeTexture(USparseVolumeTexture* Texture, int32 SubTextureId, EMaterialSamplerType SamplerType)
+	{
+		int32 TextureReferenceIndex = INDEX_NONE;
+		return SparseVolumeTexture(Texture, SubTextureId, TextureReferenceIndex, SamplerType);
+	}
+
 	int32 VirtualTexture(URuntimeVirtualTexture* InTexture, int32 TextureLayerIndex, int32 PageTableLayerIndex, EMaterialSamplerType SamplerType)
 	{
 		int32 TextureReferenceIndex = INDEX_NONE;
@@ -788,6 +850,7 @@ public:
 	virtual int32 Texture(UTexture* InTexture, int32& TextureReferenceIndex, EMaterialSamplerType SamplerType, ESamplerSourceMode SamplerSource = SSM_FromTextureAsset, ETextureMipValueMode MipValueMode = TMVM_None) override { return Compiler->Texture(InTexture, TextureReferenceIndex, SamplerType, SamplerSource, MipValueMode); }
 	virtual int32 TextureParameter(FName ParameterName, UTexture* DefaultValue, int32& TextureReferenceIndex, EMaterialSamplerType SamplerType, ESamplerSourceMode SamplerSource = SSM_FromTextureAsset) override { return Compiler->TextureParameter(ParameterName, DefaultValue, TextureReferenceIndex, SamplerType, SamplerSource); }
 
+
 	virtual int32 VirtualTexture(URuntimeVirtualTexture* InTexture, int32 TextureLayerIndex, int32 PageTableLayerIndex, int32& TextureReferenceIndex, EMaterialSamplerType SamplerType) override
 	{
 		return Compiler->VirtualTexture(InTexture, TextureLayerIndex, PageTableLayerIndex, TextureReferenceIndex, SamplerType);
@@ -809,6 +872,13 @@ public:
 	virtual int32 ExternalTextureCoordinateOffset(int32 TextureReferenceIndex, TOptional<FName> ParameterName) override { return Compiler->ExternalTextureCoordinateOffset(TextureReferenceIndex, ParameterName); }
 	virtual int32 ExternalTextureCoordinateOffset(const FGuid& ExternalTextureGuid) override { return Compiler->ExternalTextureCoordinateOffset(ExternalTextureGuid); }
 
+	virtual int32 SparseVolumeTexture(USparseVolumeTexture* Texture, int32 SubTextureId, int32& TextureReferenceIndex, EMaterialSamplerType SamplerType) override											{ return Compiler->SparseVolumeTexture(Texture, SubTextureId, TextureReferenceIndex, SamplerType); }
+	virtual int32 SparseVolumeTextureParameter(FName ParameterName, USparseVolumeTexture* InDefaultTexture, int32 SubTextureId, int32& TextureReferenceIndex, EMaterialSamplerType SamplerType) override	{ return Compiler->SparseVolumeTextureParameter(ParameterName, InDefaultTexture, SubTextureId, TextureReferenceIndex, SamplerType); }
+	virtual int32 SparseVolumeTextureSample(int32 Texture, int32 Coordinate) override																														{ return Compiler->SparseVolumeTextureSample(Texture, Coordinate); }
+	virtual int32 SparseVolumeTextureUniform(int32 TextureIndex, int32 VectorIndex, UE::Shader::EValueType Type) override																					{ return Compiler->SparseVolumeTextureUniform(TextureIndex, VectorIndex, Type); }
+	virtual int32 SparseVolumeTextureUniformParameter(FName ParameterName, int32 TextureIndex, int32 VectorIndex, UE::Shader::EValueType Type) override														{ return Compiler->SparseVolumeTextureUniformParameter(ParameterName, TextureIndex, VectorIndex, Type); }
+	virtual int32 SparseVolumeTextureGetVoxelCoord(int32 PackedPhysicalTileCoord, int32 TileSize, int32 CoordPageTable, int32 CoordVolume) override															{ return Compiler->SparseVolumeTextureGetVoxelCoord(PackedPhysicalTileCoord, TileSize, CoordPageTable, CoordVolume); }
+
 	virtual UObject* GetReferencedTexture(int32 Index) override { return Compiler->GetReferencedTexture(Index); }
 
 	virtual	int32 PixelDepth() override { return Compiler->PixelDepth(); }
diff --git a/Engine/Source/Runtime/Engine/Public/MaterialHLSLTree.h b/Engine/Source/Runtime/Engine/Public/MaterialHLSLTree.h
--- a/Engine/Source/Runtime/Engine/Public/MaterialHLSLTree.h
+++ b/Engine/Source/Runtime/Engine/Public/MaterialHLSLTree.h
@@ -36,6 +36,7 @@ inline void AppendHash(FHasher& Hasher, const FMaterialParameterValue& Value)
 	case EMaterialParameterType::Texture: AppendHash(Hasher, Value.Texture); break;
 	case EMaterialParameterType::Font: AppendHash(Hasher, Value.Font); break;
 	case EMaterialParameterType::RuntimeVirtualTexture: AppendHash(Hasher, Value.RuntimeVirtualTexture); break;
+	case EMaterialParameterType::SparseVolumeTexture: AppendHash(Hasher, Value.SparseVolumeTexture); break;
 	case EMaterialParameterType::StaticSwitch: AppendHash(Hasher, Value.Bool[0]); break;
 	case EMaterialParameterType::StaticComponentMask: AppendHash(Hasher, Value.Bool); break;
 	default: checkNoEntry(); break;
diff --git a/Engine/Source/Runtime/Engine/Public/MaterialShared.h b/Engine/Source/Runtime/Engine/Public/MaterialShared.h
--- a/Engine/Source/Runtime/Engine/Public/MaterialShared.h
+++ b/Engine/Source/Runtime/Engine/Public/MaterialShared.h
@@ -57,6 +57,7 @@ class UMaterialExpressionMaterialFunctionCall;
 class UMaterialInstance;
 class UMaterialInterface;
 class URuntimeVirtualTexture;
+class USparseVolumeTexture;
 class USubsurfaceProfile;
 class UTexture;
 class UTexture2D;
@@ -193,7 +194,14 @@ enum EMaterialValueType
 	MCT_VoidStatement         = 1 << 22,
 
 	/** Non-static bool, only used in new HLSL translator */
-	MCT_Bool = 1 << 23,
+	MCT_Bool				 = 1 << 23,
+
+	/** Unsigned int types */
+	MCT_UInt1				 = 1 << 24,
+	MCT_UInt2				 = 1 << 25,
+	MCT_UInt3				 = 1 << 26,
+	MCT_UInt4				 = 1 << 27,
+	MCT_UInt				 = MCT_UInt1 | MCT_UInt2 | MCT_UInt3 | MCT_UInt4,
 
 	MCT_Numeric = MCT_Float | MCT_LWCType | MCT_Bool,
 };
@@ -203,10 +211,15 @@ inline uint32 GetNumComponents(EMaterialValueType Type)
 {
 	switch (Type)
 	{
+	case MCT_UInt:
+	case MCT_UInt1:
 	case MCT_Float:
 	case MCT_Float1: return 1;
+	case MCT_UInt2:
 	case MCT_Float2: return 2;
+	case MCT_UInt3:
 	case MCT_Float3: return 3;
+	case MCT_UInt4:
 	case MCT_Float4: return 4;
 	case MCT_LWCScalar: return 1;
 	case MCT_LWCVector2: return 2;
@@ -226,10 +239,15 @@ inline bool IsFloatNumericType(EMaterialValueType InType)
 	return (InType & MCT_Float) || IsLWCType(InType);
 }
 
+inline bool IsUIntNumericType(EMaterialValueType InType)
+{
+	return (InType & MCT_UInt);
+}
+
 inline bool IsNumericType(EMaterialValueType InType)
 {
 	// 'ShadingModel' is considered an 'int' 
-	return IsFloatNumericType(InType) || InType == MCT_ShadingModel;
+	return IsFloatNumericType(InType) || InType == MCT_ShadingModel || IsUIntNumericType(InType);
 }
 
 inline EMaterialValueType MakeNonLWCType(EMaterialValueType Type)
@@ -348,6 +366,7 @@ struct ENGINE_API FMaterialRenderContext
 
 	void GetTextureParameterValue(const FHashedMaterialParameterInfo& ParameterInfo, int32 TextureIndex, const UTexture*& OutValue) const;
 	void GetTextureParameterValue(const FHashedMaterialParameterInfo& ParameterInfo, int32 TextureIndex, const URuntimeVirtualTexture*& OutValue) const;
+	void GetTextureParameterValue(const FHashedMaterialParameterInfo& ParameterInfo, int32 TextureIndex, const USparseVolumeTexture*& OutValue) const;
 	FGuid GetExternalTextureGuid(const FGuid& ExternalTextureGuid, const FName& ParameterName, int32 SourceTextureIndex) const;
 };
 
@@ -376,6 +395,8 @@ public:
 	void GetTextureValues(const FMaterialRenderContext& Context, const FUniformExpressionSet& UniformExpressionSet, UTexture const** OutValues) const;
 	/** Get the URuntimeVirtualTexture object if one was used to initialize this stack. */
 	void GetTextureValue(const FMaterialRenderContext& Context, const FUniformExpressionSet& UniformExpressionSet, const URuntimeVirtualTexture*& OutValue) const;
+	/** Get the USparseVolumeTexture object if one was used to initialize this stack. */
+	void GetTextureValue(const FMaterialRenderContext& Context, const FUniformExpressionSet& UniformExpressionSet, const USparseVolumeTexture * &OutValue) const;
 
 	void Serialize(FArchive& Ar);
 
@@ -606,6 +627,7 @@ public:
 	ENGINE_API void GetGameThreadTextureValue(EMaterialTextureParameterType Type, int32 Index, const UMaterialInterface* MaterialInterface, const FMaterial& Material, UTexture*& OutValue, bool bAllowOverride = true) const;
 	ENGINE_API void GetTextureValue(EMaterialTextureParameterType Type, int32 Index, const FMaterialRenderContext& Context, const FMaterial& Material, const UTexture*& OutValue) const;
 	ENGINE_API void GetTextureValue(int32 Index, const FMaterialRenderContext& Context, const FMaterial& Material, const URuntimeVirtualTexture*& OutValue) const;
+	ENGINE_API void GetTextureValue(int32 Index, const FMaterialRenderContext& Context, const FMaterial& Material, const USparseVolumeTexture*& OutValue) const;
 
 	int32 FindOrAddTextureParameter(EMaterialTextureParameterType Type, const FMaterialTextureParameterInfo& Info);
 	int32 FindOrAddExternalTextureParameter(const FMaterialExternalTextureParameterInfo& Info);
@@ -2629,6 +2651,7 @@ public:
 	bool GetScalarValue(const FHashedMaterialParameterInfo& ParameterInfo, float* OutValue, const FMaterialRenderContext& Context) const;
 	bool GetTextureValue(const FHashedMaterialParameterInfo& ParameterInfo, const UTexture** OutValue, const FMaterialRenderContext& Context) const;
 	bool GetTextureValue(const FHashedMaterialParameterInfo& ParameterInfo, const URuntimeVirtualTexture** OutValue, const FMaterialRenderContext& Context) const;
+	bool GetTextureValue(const FHashedMaterialParameterInfo& ParameterInfo, const USparseVolumeTexture** OutValue, const FMaterialRenderContext& Context) const;
 	virtual bool GetParameterValue(EMaterialParameterType Type, const FHashedMaterialParameterInfo& ParameterInfo, FMaterialParameterValue& OutValue, const FMaterialRenderContext& Context) const = 0;
 
 	bool IsDeleted() const
diff --git a/Engine/Source/Runtime/Engine/Public/MaterialTypes.h b/Engine/Source/Runtime/Engine/Public/MaterialTypes.h
--- a/Engine/Source/Runtime/Engine/Public/MaterialTypes.h
+++ b/Engine/Source/Runtime/Engine/Public/MaterialTypes.h
@@ -15,6 +15,7 @@ class UCurveLinearColor;
 class UCurveLinearColorAtlas;
 class UFont;
 class URuntimeVirtualTexture;
+class USparseVolumeTexture;
 
 UENUM()
 enum EMaterialParameterAssociation
@@ -187,6 +188,7 @@ enum class EMaterialParameterType : uint8
 	Texture,
 	Font,
 	RuntimeVirtualTexture,
+	SparseVolumeTexture,
 
 	NumRuntime, // Runtime parameter types must go above here, and editor-only ones below
 
@@ -311,6 +313,8 @@ struct FMaterialParameterValue
 	FMaterialParameterValue(const TObjectPtr<UTexture>& InValue) : Type(EMaterialParameterType::Texture) { Texture = InValue; }
 	FMaterialParameterValue(URuntimeVirtualTexture* InValue) : Type(EMaterialParameterType::RuntimeVirtualTexture) { RuntimeVirtualTexture = InValue; }
 	FMaterialParameterValue(const TObjectPtr<URuntimeVirtualTexture>& InValue) : Type(EMaterialParameterType::RuntimeVirtualTexture) { RuntimeVirtualTexture = InValue; }
+	FMaterialParameterValue(USparseVolumeTexture* InValue) : Type(EMaterialParameterType::SparseVolumeTexture) { SparseVolumeTexture = InValue; }
+	FMaterialParameterValue(const TObjectPtr<USparseVolumeTexture>& InValue) : Type(EMaterialParameterType::SparseVolumeTexture) { SparseVolumeTexture = InValue; }
 
 	// Gamethread parameters are typically non-const, but renderthread parameters are const
 	// Would be possible to store an additional const-flag member, and provide runtime checks to ensure constness is not violated...maybe worth doing in the future
@@ -318,6 +322,8 @@ struct FMaterialParameterValue
 	FMaterialParameterValue(const TObjectPtr<const UTexture>& InValue) : Type(EMaterialParameterType::Texture) { Texture = const_cast<UTexture*>(InValue.Get()); }
 	FMaterialParameterValue(const URuntimeVirtualTexture* InValue) : Type(EMaterialParameterType::RuntimeVirtualTexture) { RuntimeVirtualTexture = const_cast<URuntimeVirtualTexture*>(InValue); }
 	FMaterialParameterValue(const TObjectPtr<const URuntimeVirtualTexture>& InValue) : Type(EMaterialParameterType::RuntimeVirtualTexture) { RuntimeVirtualTexture = const_cast<URuntimeVirtualTexture*>(InValue.Get()); }
+	FMaterialParameterValue(const USparseVolumeTexture* InValue) : Type(EMaterialParameterType::SparseVolumeTexture) { SparseVolumeTexture = const_cast<USparseVolumeTexture*>(InValue); }
+	FMaterialParameterValue(const TObjectPtr<const USparseVolumeTexture>& InValue) : Type(EMaterialParameterType::SparseVolumeTexture) { SparseVolumeTexture = const_cast<USparseVolumeTexture*>(InValue.Get()); }
 
 	FMaterialParameterValue(UFont* InValue, int32 InPage) : Type(EMaterialParameterType::Font) { Font.Value = InValue; Font.Page = InPage; }
 	FMaterialParameterValue(bool InValue) : Type(EMaterialParameterType::StaticSwitch) { Bool[0] = InValue; }
@@ -341,6 +347,7 @@ struct FMaterialParameterValue
 		bool Bool[4];
 		UTexture* Texture;
 		URuntimeVirtualTexture* RuntimeVirtualTexture;
+		USparseVolumeTexture* SparseVolumeTexture;
 		struct
 		{
 			UFont* Value;
@@ -374,6 +381,7 @@ inline bool operator==(const FMaterialParameterValue& Lhs, const FMaterialParame
 	case EMaterialParameterType::Texture: return Lhs.Texture == Rhs.Texture;
 	case EMaterialParameterType::Font: return Lhs.Font.Value == Rhs.Font.Value && Lhs.Font.Page == Rhs.Font.Page;
 	case EMaterialParameterType::RuntimeVirtualTexture: return Lhs.RuntimeVirtualTexture == Rhs.RuntimeVirtualTexture;
+	case EMaterialParameterType::SparseVolumeTexture: return Lhs.SparseVolumeTexture == Rhs.SparseVolumeTexture;
 	case EMaterialParameterType::StaticSwitch: return Lhs.Bool[0] == Rhs.Bool[0];
 	case EMaterialParameterType::StaticComponentMask: return
 		Lhs.Bool[0] == Rhs.Bool[0] &&
diff --git a/Engine/Source/Runtime/Engine/Public/SceneInterface.h b/Engine/Source/Runtime/Engine/Public/SceneInterface.h
--- a/Engine/Source/Runtime/Engine/Public/SceneInterface.h
+++ b/Engine/Source/Runtime/Engine/Public/SceneInterface.h
@@ -38,6 +38,7 @@ class FInstanceCullingManager;
 class FSceneViewStateInterface;
 struct FHairStrandsInstance;
 struct FLightRenderParameters;
+class FSparseVolumeTextureViewerSceneProxy;
 
 enum EBasePassDrawListType
 {
@@ -306,6 +307,9 @@ public:
 	virtual FSkyAtmosphereRenderSceneInfo* GetSkyAtmosphereSceneInfo() = 0;
 	virtual const FSkyAtmosphereRenderSceneInfo* GetSkyAtmosphereSceneInfo() const = 0;
 
+	virtual void AddSparseVolumeTextureViewer(FSparseVolumeTextureViewerSceneProxy* SVTV) = 0;
+	virtual void RemoveSparseVolumeTextureViewer(FSparseVolumeTextureViewerSceneProxy* SVTV) = 0;
+
 	/**
 	 * Adds the unique volumetric cloud component to the scene
 	 *
diff --git a/Engine/Source/Runtime/Engine/Public/Shader/Preshader.h b/Engine/Source/Runtime/Engine/Public/Shader/Preshader.h
--- a/Engine/Source/Runtime/Engine/Public/Shader/Preshader.h
+++ b/Engine/Source/Runtime/Engine/Public/Shader/Preshader.h
@@ -61,6 +61,7 @@ enum class EPreshaderOpcode : uint8
 	ExternalTextureCoordinateScaleRotation,
 	ExternalTextureCoordinateOffset,
 	RuntimeVirtualTextureUniform,
+	SparseVolumeTextureUniform,
 	GetField,
 	SetField,
 	Neg,
diff --git a/Engine/Source/Runtime/Renderer/Private/DeferredShadingRenderer.cpp b/Engine/Source/Runtime/Renderer/Private/DeferredShadingRenderer.cpp
--- a/Engine/Source/Runtime/Renderer/Private/DeferredShadingRenderer.cpp
+++ b/Engine/Source/Runtime/Renderer/Private/DeferredShadingRenderer.cpp
@@ -9,6 +9,7 @@
 #include "SingleLayerWaterRendering.h"
 #include "SkyAtmosphereRendering.h"
 #include "VolumetricCloudRendering.h"
+#include "SparseVolumeTexture/SparseVolumeTextureViewerRendering.h"
 #include "VolumetricRenderTarget.h"
 #include "ScenePrivate.h"
 #include "SceneOcclusion.h"
@@ -3775,6 +3776,11 @@ void FDeferredShadingSceneRenderer::Render(FRDGBuilder& GraphBuilder)
 		RenderVolumetricCloud(GraphBuilder, SceneTextures, true, false, HalfResolutionDepthCheckerboardMinMaxTexture, false, InstanceCullingManager);
 	}
 
+	if (!bHasRayTracedOverlay)
+	{
+		AddSparseVolumeTextureViewerRenderPass(GraphBuilder, *this, SceneTextures);
+	}
+
 	// Resolve the scene color for post processing.
 	AddResolveSceneColorPass(GraphBuilder, Views, SceneTextures.Color);
 
diff --git a/Engine/Source/Runtime/Renderer/Private/RendererScene.cpp b/Engine/Source/Runtime/Renderer/Private/RendererScene.cpp
--- a/Engine/Source/Runtime/Renderer/Private/RendererScene.cpp
+++ b/Engine/Source/Runtime/Renderer/Private/RendererScene.cpp
@@ -6115,6 +6115,9 @@ public:
 	virtual FSkyAtmosphereRenderSceneInfo* GetSkyAtmosphereSceneInfo() override { return NULL; }
 	virtual const FSkyAtmosphereRenderSceneInfo* GetSkyAtmosphereSceneInfo() const override { return NULL; }
 
+	virtual void AddSparseVolumeTextureViewer(FSparseVolumeTextureViewerSceneProxy* SVTV) override {}
+	virtual void RemoveSparseVolumeTextureViewer(FSparseVolumeTextureViewerSceneProxy* SVTV) override {}
+
 	virtual void AddHairStrands(FHairStrandsInstance* Proxy) override {}
 	virtual void RemoveHairStrands(FHairStrandsInstance* Proxy) override {}
 	virtual void GetRectLightAtlasSlot(const FRectLightSceneProxy* Proxy, FLightRenderParameters* Out) override {}
diff --git a/Engine/Source/Runtime/Renderer/Private/ScenePrivate.h b/Engine/Source/Runtime/Renderer/Private/ScenePrivate.h
--- a/Engine/Source/Runtime/Renderer/Private/ScenePrivate.h
+++ b/Engine/Source/Runtime/Renderer/Private/ScenePrivate.h
@@ -93,6 +93,7 @@ class FLumenSceneData;
 class FVirtualShadowMapArrayCacheManager;
 struct FHairStrandsInstance;
 struct FPathTracingState;
+class FSparseVolumeTextureViewerSceneProxy;
 
 /** Holds information about a single primitive's occlusion. */
 class FPrimitiveOcclusionHistory
@@ -3005,6 +3006,8 @@ public:
 	/** Used to track the order that skylights were enabled in. */
 	TArray<FVolumetricCloudSceneProxy*> VolumetricCloudStack;
 
+	TArray<FSparseVolumeTextureViewerSceneProxy*> SparseVolumeTextureViewers;
+
 	/** Global Field Manager */
 	class FPhysicsFieldSceneProxy* PhysicsField = nullptr;
 
@@ -3169,6 +3172,9 @@ public:
 	virtual FSkyAtmosphereRenderSceneInfo* GetSkyAtmosphereSceneInfo() override { return SkyAtmosphere; }
 	virtual const FSkyAtmosphereRenderSceneInfo* GetSkyAtmosphereSceneInfo() const override { return SkyAtmosphere; }
 
+	virtual void AddSparseVolumeTextureViewer(FSparseVolumeTextureViewerSceneProxy* SVTV) override;
+	virtual void RemoveSparseVolumeTextureViewer(FSparseVolumeTextureViewerSceneProxy* SVTV) override;
+
 	virtual void SetPhysicsField(class FPhysicsFieldSceneProxy* PhysicsFieldSceneProxy) override;
 	virtual void ResetPhysicsField() override;
 	virtual void ShowPhysicsField() override;
diff --git a/Engine/Source/Runtime/Renderer/Private/SparseVolumeTexture/SparseVolumeTextureViewerComponent.cpp b/Engine/Source/Runtime/Renderer/Private/SparseVolumeTexture/SparseVolumeTextureViewerComponent.cpp
new file mode 100644
--- /dev/null
+++ b/Engine/Source/Runtime/Renderer/Private/SparseVolumeTexture/SparseVolumeTextureViewerComponent.cpp
@@ -0,0 +1,234 @@
+// Copyright Epic Games, Inc. All Rights Reserved.
+
+#include "SparseVolumeTexture/SparseVolumeTextureViewerComponent.h"
+
+#include "SparseVolumeTexture/SparseVolumeTextureViewerSceneProxy.h"
+#include "Components/ArrowComponent.h"
+#include "Components/BillboardComponent.h"
+#include "Engine/MapBuildDataRegistry.h"
+#include "Internationalization/Text.h"
+#include "Logging/MessageLog.h"
+#include "Logging/TokenizedMessage.h"
+#include "Misc/MapErrors.h"
+#include "Misc/UObjectToken.h"
+#include "UObject/UObjectIterator.h"
+#include "UObject/ConstructorHelpers.h"
+
+#if WITH_EDITOR
+#include "ObjectEditorUtils.h"
+#endif
+
+#define LOCTEXT_NAMESPACE "SparseVolumeTextureViewerComponent"
+
+/*=============================================================================
+	USparseVolumeTextureViewerComponent implementation.
+=============================================================================*/
+
+USparseVolumeTextureViewerComponent::USparseVolumeTextureViewerComponent(const FObjectInitializer& ObjectInitializer)
+	: Super(ObjectInitializer)
+	, SparseVolumeTexturePreview(nullptr)
+	, bAnimate(false)
+	, AnimationFrame(0.0f)
+	, SparseVolumeTextureViewerSceneProxy(nullptr)
+{
+	PrimaryComponentTick.bCanEverTick = true;
+	PrimaryComponentTick.TickGroup = TG_DuringPhysics;
+	bTickInEditor = true;
+}
+
+USparseVolumeTextureViewerComponent::~USparseVolumeTextureViewerComponent()
+{
+}
+
+#if WITH_EDITOR
+
+void USparseVolumeTextureViewerComponent::CheckForErrors()
+{
+}
+
+void USparseVolumeTextureViewerComponent::PostEditChangeProperty(FPropertyChangedEvent& PropertyChangedEvent)
+{
+	Super::PostEditChangeProperty(PropertyChangedEvent);
+
+	if (!bAnimate && SparseVolumeTexturePreview)
+	{
+		FrameIndex = int32(AnimationFrame * float(SparseVolumeTexturePreview->GetFrameCount()));
+	}
+	MarkRenderStateDirty();
+
+	SendRenderTransformCommand();
+}
+
+#endif // WITH_EDITOR
+
+void USparseVolumeTextureViewerComponent::PostInterpChange(FProperty* PropertyThatChanged)
+{
+	// This is called when property is modified by InterpPropertyTracks
+	Super::PostInterpChange(PropertyThatChanged);
+}
+
+void USparseVolumeTextureViewerComponent::Serialize(FArchive& Ar)
+{
+	Super::Serialize(Ar);
+}
+
+FBoxSphereBounds USparseVolumeTextureViewerComponent::CalcBounds(const FTransform& LocalToWorld) const
+{
+	const float OneMeterInCentimeter = 100.0f;
+	FBoxSphereBounds NormalizedBound;
+	NormalizedBound.Origin = FVector(0.0f, 0.0f, 0.0f);
+
+	if (SparseVolumeTexturePreview)
+	{
+		// We assume that the maximum bound will have a length of 1 meter. 
+		// Then the other dimensions are scaled relatively.
+		// All this account for the size of volume with page table (they add padding).
+		// TODO can we recover world size from OpenVDB meta data in meter?
+		const FSparseVolumeAssetHeader* SVTHeader = SparseVolumeTexturePreview->GetSparseVolumeTextureHeader();
+		check(SVTHeader);
+		const float MaxDim = (float)FMath::Max(FMath::Max(SVTHeader->PageTableVolumeResolution.X, SVTHeader->PageTableVolumeResolution.Y), SVTHeader->PageTableVolumeResolution.Z);
+
+		NormalizedBound.BoxExtent.X = 0.5f * OneMeterInCentimeter * float(SVTHeader->PageTableVolumeResolution.X) / MaxDim;
+		NormalizedBound.BoxExtent.Y = 0.5f * OneMeterInCentimeter * float(SVTHeader->PageTableVolumeResolution.Y) / MaxDim;
+		NormalizedBound.BoxExtent.Z = 0.5f * OneMeterInCentimeter * float(SVTHeader->PageTableVolumeResolution.Z) / MaxDim;
+	}
+	else
+	{
+		NormalizedBound.BoxExtent = FVector(OneMeterInCentimeter, OneMeterInCentimeter, OneMeterInCentimeter);
+	}
+
+	NormalizedBound.SphereRadius = NormalizedBound.BoxExtent.Size();
+	return NormalizedBound.TransformBy(LocalToWorld);
+}
+
+void USparseVolumeTextureViewerComponent::CreateRenderState_Concurrent(FRegisterComponentContext* Context)
+{
+	Super::CreateRenderState_Concurrent(Context);
+	// If one day we need to look up lightmass built data, lookup it up here using the guid from the correct MapBuildData.
+
+	bool bHidden = false;
+#if WITH_EDITORONLY_DATA
+	bHidden = GetOwner() ? GetOwner()->bHiddenEdLevel : false;
+#endif // WITH_EDITORONLY_DATA
+	if (!ShouldComponentAddToScene())
+	{
+		bHidden = true;
+	}
+
+	if (GetVisibleFlag() && !bHidden &&
+		ShouldComponentAddToScene() && ShouldRender() && IsRegistered() && (GetOuter() == NULL || !GetOuter()->HasAnyFlags(RF_ClassDefaultObject)))
+	{
+		// Create the scene proxy.
+		SparseVolumeTextureViewerSceneProxy = new FSparseVolumeTextureViewerSceneProxy(this, FrameIndex);
+		GetWorld()->Scene->AddSparseVolumeTextureViewer(SparseVolumeTextureViewerSceneProxy);
+		SendRenderTransformCommand();
+	}
+}
+
+void USparseVolumeTextureViewerComponent::DestroyRenderState_Concurrent()
+{
+	Super::DestroyRenderState_Concurrent();
+
+	if (SparseVolumeTextureViewerSceneProxy)
+	{
+		GetWorld()->Scene->RemoveSparseVolumeTextureViewer(SparseVolumeTextureViewerSceneProxy);
+
+		FSparseVolumeTextureViewerSceneProxy* SVTViewerSceneProxy = SparseVolumeTextureViewerSceneProxy;
+		ENQUEUE_RENDER_COMMAND(FDestroySparseVolumeTextureViewerSceneProxyCommand)(
+			[SVTViewerSceneProxy](FRHICommandList& RHICmdList)
+			{
+				delete SVTViewerSceneProxy;
+			});
+
+		SparseVolumeTextureViewerSceneProxy = nullptr;
+	}
+}
+
+void USparseVolumeTextureViewerComponent::SendRenderTransform_Concurrent()
+{
+	Super::SendRenderTransform_Concurrent();
+
+	SendRenderTransformCommand();
+}
+
+void USparseVolumeTextureViewerComponent::SendRenderTransformCommand()
+{
+	if (SparseVolumeTextureViewerSceneProxy)
+	{
+		FTransform ToWorldMat = GetComponentToWorld();
+		ToWorldMat.SetRotation(FQuat::Identity);								// Do not allow rotation
+		FVector Scale3D = ToWorldMat.GetScale3D();
+		float MaxScaling = FMath::Max(Scale3D.X, FMath::Max(Scale3D.Y, Scale3D.Z));
+		ToWorldMat.SetScale3D(FVector(MaxScaling, MaxScaling, MaxScaling));		// Keep max scaling
+		FBoxSphereBounds ComponentBounds = CalcBounds(ToWorldMat);
+		FSparseVolumeTextureViewerSceneProxy* SVTViewerSceneProxy = SparseVolumeTextureViewerSceneProxy;
+		ENQUEUE_RENDER_COMMAND(FUpdateSparseVolumeTextureViewerProxyTransformCommand)(
+			[SVTViewerSceneProxy, ComponentBounds](FRHICommandList& RHICmdList)
+			{
+				SVTViewerSceneProxy->VolumeWorldBounds = ComponentBounds;
+			});
+	}
+}
+
+void USparseVolumeTextureViewerComponent::TickComponent(float DeltaTime, enum ELevelTick TickType, FActorComponentTickFunction* ThisTickFunction)
+{
+	if (bAnimate)
+	{
+		FrameIndex++;
+	}
+	else if(SparseVolumeTexturePreview)
+	{
+		FrameIndex = int32(AnimationFrame * float(SparseVolumeTexturePreview->GetFrameCount()));
+	}
+	MarkRenderStateDirty();
+}
+
+/*=============================================================================
+	ASparseVolumeTextureViewer implementation.
+=============================================================================*/
+
+#if WITH_EDITOR
+#include "ObjectEditorUtils.h"
+#endif
+
+ASparseVolumeTextureViewer::ASparseVolumeTextureViewer(const FObjectInitializer& ObjectInitializer)
+	: Super(ObjectInitializer)
+{
+	SparseVolumeTextureViewerComponent = CreateDefaultSubobject<USparseVolumeTextureViewerComponent>(TEXT("SparseVolumeTextureViewerComponent"));
+	RootComponent = SparseVolumeTextureViewerComponent;
+
+#if WITH_EDITORONLY_DATA
+
+	if (!IsRunningCommandlet())
+	{
+		// Structure to hold one-time initialization
+		struct FConstructorStatics
+		{
+			ConstructorHelpers::FObjectFinderOptional<UTexture2D> SparseVolumeTextureViewerTextureObject;
+			FName ID_SparseVolumeTextureViewer;
+			FText NAME_SparseVolumeTextureViewer;
+			FConstructorStatics()
+				: SparseVolumeTextureViewerTextureObject(TEXT("/Engine/EditorResources/S_VolumetricCloud"))	// SVT_TODO set a specific icon
+				, ID_SparseVolumeTextureViewer(TEXT("Fog"))
+				, NAME_SparseVolumeTextureViewer(NSLOCTEXT("SpriteCategory", "Fog", "Fog"))
+			{
+			}
+		};
+		static FConstructorStatics ConstructorStatics;
+
+		if (GetSpriteComponent())
+		{
+			GetSpriteComponent()->Sprite = ConstructorStatics.SparseVolumeTextureViewerTextureObject.Get();
+			GetSpriteComponent()->SetRelativeScale3D(FVector(0.5f, 0.5f, 0.5f));
+			GetSpriteComponent()->SpriteInfo.Category = ConstructorStatics.ID_SparseVolumeTextureViewer;
+			GetSpriteComponent()->SpriteInfo.DisplayName = ConstructorStatics.NAME_SparseVolumeTextureViewer;
+			GetSpriteComponent()->SetupAttachment(SparseVolumeTextureViewerComponent);
+		}
+	}
+#endif // WITH_EDITORONLY_DATA
+
+	PrimaryActorTick.bCanEverTick = true;
+	SetHidden(false);
+}
+
+#undef LOCTEXT_NAMESPACE
diff --git a/Engine/Source/Runtime/Renderer/Private/SparseVolumeTexture/SparseVolumeTextureViewerComponent.h b/Engine/Source/Runtime/Renderer/Private/SparseVolumeTexture/SparseVolumeTextureViewerComponent.h
new file mode 100644
--- /dev/null
+++ b/Engine/Source/Runtime/Renderer/Private/SparseVolumeTexture/SparseVolumeTextureViewerComponent.h
@@ -0,0 +1,99 @@
+// Copyright Epic Games, Inc. All Rights Reserved.
+
+#pragma once
+
+#include "CoreMinimal.h"
+#include "Components/SceneComponent.h"
+#include "Components/PrimitiveComponent.h"
+#include "EngineDefines.h"
+#include "GameFramework/Info.h"
+
+#include "SparseVolumeTexture/SparseVolumeTexture.h"
+
+#include "SparseVolumeTextureViewerComponent.generated.h"
+
+
+class FSparseVolumeTextureViewerSceneProxy;
+
+
+/**
+ * A component used to inspect sparse volume textures.
+ */
+UCLASS(ClassGroup = Rendering, collapsecategories, hidecategories = (Object, Mobility, Activation, "Components|Activation"), editinlinenew, meta = (BlueprintSpawnableComponent), MinimalAPI)
+class USparseVolumeTextureViewerComponent : public UPrimitiveComponent
+{
+	GENERATED_UCLASS_BODY()
+
+	~USparseVolumeTextureViewerComponent();
+	
+	UPROPERTY(EditAnywhere, BlueprintReadWrite, Category= "Asset Preview")
+	TObjectPtr<class USparseVolumeTexture> SparseVolumeTexturePreview;
+
+	UPROPERTY(EditAnywhere, BlueprintReadOnly, Category = "Asset Preview")
+	uint32 bAnimate : 1;
+
+	UPROPERTY(EditAnywhere, BlueprintReadWrite, Category = "Asset Preview", meta = (UIMin = 0.0, UIMax = 1.0, ClampMin = 0.0, ClampMax = 1.0, EditCondition = "!bAnimate"))
+	float AnimationFrame;
+
+	//~ Begin UPrimitiveComponent Interface.
+	virtual bool SupportsStaticLighting() const override { return false; }
+	//~ End UPrimitiveComponent Interface.
+
+	//~ Begin USceneComponent Interface.
+	virtual FBoxSphereBounds CalcBounds(const FTransform& LocalToWorld) const  override;
+	//~ End USceneComponent Interface.
+
+
+	//~ Begin UActorComponent Interface.
+	virtual void CreateRenderState_Concurrent(FRegisterComponentContext* Context) override;
+	virtual void DestroyRenderState_Concurrent() override;
+	virtual void SendRenderTransform_Concurrent() override;
+	virtual void TickComponent(float DeltaTime, enum ELevelTick TickType, FActorComponentTickFunction* ThisTickFunction) override;
+	//~ End UActorComponent Interface.
+
+protected:
+
+public:
+
+	//~ Begin UObject Interface. 
+	virtual void PostInterpChange(FProperty* PropertyThatChanged) override;
+	virtual void Serialize(FArchive& Ar) override;
+#if WITH_EDITOR
+	virtual void PostEditChangeProperty(FPropertyChangedEvent& PropertyChangedEvent) override;
+#endif // WITH_EDITOR
+	//~ End UObject Interface
+
+	//~ Begin UActorComponent Interface.
+#if WITH_EDITOR
+	virtual void CheckForErrors() override;
+#endif // WITH_EDITOR
+	//~ End UActorComponent Interface.
+
+private:
+	
+	int32 FrameIndex = 0;
+
+	void SendRenderTransformCommand();
+
+	FSparseVolumeTextureViewerSceneProxy* SparseVolumeTextureViewerSceneProxy;
+};
+
+
+/**
+ * A placeable actor that represents a participating media material around a planet, e.g. clouds.
+ */
+UCLASS(showcategories = (Movement, Rendering, Transformation, DataLayers, "Input|MouseInput", "Input|TouchInput"), ClassGroup = Fog, hidecategories = (Info, Object, Input), MinimalAPI)
+class ASparseVolumeTextureViewer : public AInfo
+{
+	GENERATED_UCLASS_BODY()
+
+private:
+
+	UPROPERTY(BlueprintReadOnly, VisibleAnywhere, Category = Atmosphere, meta = (AllowPrivateAccess = "true"))
+	TObjectPtr<class USparseVolumeTextureViewerComponent> SparseVolumeTextureViewerComponent;
+
+#if WITH_EDITOR
+	virtual bool ActorTypeSupportsDataLayer() const override { return true; }
+#endif
+
+};
diff --git a/Engine/Source/Runtime/Renderer/Private/SparseVolumeTexture/SparseVolumeTextureViewerRendering.cpp b/Engine/Source/Runtime/Renderer/Private/SparseVolumeTexture/SparseVolumeTextureViewerRendering.cpp
new file mode 100644
--- /dev/null
+++ b/Engine/Source/Runtime/Renderer/Private/SparseVolumeTexture/SparseVolumeTextureViewerRendering.cpp
@@ -0,0 +1,183 @@
+// Copyright Epic Games, Inc. All Rights Reserved.
+
+#include "SparseVolumeTexture/SparseVolumeTextureViewerRendering.h"
+#include "SparseVolumeTexture/SparseVolumeTextureViewerSceneProxy.h"
+#include "SparseVolumeTexture/SparseVolumeTexture.h"
+
+#include "SceneView.h"
+#include "ScenePrivate.h"
+#include "SceneRendering.h"
+#include "RendererInterface.h"
+#include "UniformBuffer.h"
+#include "SceneTextureParameters.h"
+#include "ShaderCompiler.h"
+#include "RendererUtils.h"
+#include "EngineAnalytics.h"
+
+
+////////////////////////////////////////////////////////////////////////////////////////////////////
+
+class FVisualizeSparseVolumeTextureVS : public FGlobalShader
+{
+	DECLARE_GLOBAL_SHADER(FVisualizeSparseVolumeTextureVS);
+	SHADER_USE_PARAMETER_STRUCT(FVisualizeSparseVolumeTextureVS, FGlobalShader);
+
+	using FPermutationDomain = TShaderPermutationDomain<>;
+
+	BEGIN_SHADER_PARAMETER_STRUCT(FParameters, )
+		SHADER_PARAMETER(float, DepthAsDeviceZ)
+	END_SHADER_PARAMETER_STRUCT()
+
+	static FPermutationDomain RemapPermutation(FPermutationDomain PermutationVector)
+	{
+		return PermutationVector;
+	}
+
+	static bool ShouldCompilePermutation(const FGlobalShaderPermutationParameters& Parameters)
+	{
+		return !IsMobilePlatform(Parameters.Platform) && EnumHasAllFlags(Parameters.Flags, EShaderPermutationFlags::HasEditorOnlyData);
+	}
+
+	static void ModifyCompilationEnvironment(const FGlobalShaderPermutationParameters& Parameters, FShaderCompilerEnvironment& OutEnvironment)
+	{
+		FGlobalShader::ModifyCompilationEnvironment(Parameters, OutEnvironment);
+		OutEnvironment.SetDefine(TEXT("VERTEX_SHADER"), 1);
+	}
+};
+
+class FVisualizeSparseVolumeTexturePS : public FGlobalShader
+{
+	DECLARE_GLOBAL_SHADER(FVisualizeSparseVolumeTexturePS);
+	SHADER_USE_PARAMETER_STRUCT(FVisualizeSparseVolumeTexturePS, FGlobalShader);
+
+	using FPermutationDomain = TShaderPermutationDomain<>;
+
+	BEGIN_SHADER_PARAMETER_STRUCT(FParameters, )
+		SHADER_PARAMETER_STRUCT_REF(FViewUniformShaderParameters, ViewUniformBuffer)
+		RENDER_TARGET_BINDING_SLOTS()
+		SHADER_PARAMETER(FVector3f, VolumeBoundMinWorld)
+		SHADER_PARAMETER(FVector3f, VolumeBoundMaxWorld)
+		SHADER_PARAMETER(FVector3f, SparseVolumeTextureResolution)
+		SHADER_PARAMETER(FVector3f, SparseVolumeTexturePageTableResolution)
+		SHADER_PARAMETER_TEXTURE(Texture3D, SparseVolumeTexture)
+		SHADER_PARAMETER_TEXTURE(Texture3D<uint>, SparseVolumeTexturePageTable)
+	END_SHADER_PARAMETER_STRUCT()
+
+	static FPermutationDomain RemapPermutation(FPermutationDomain PermutationVector)
+	{
+		return PermutationVector;
+	}
+
+	static bool ShouldCompilePermutation(const FGlobalShaderPermutationParameters& Parameters)
+	{
+		return !IsMobilePlatform(Parameters.Platform) && EnumHasAllFlags(Parameters.Flags, EShaderPermutationFlags::HasEditorOnlyData);
+	}
+
+	static void ModifyCompilationEnvironment(const FGlobalShaderPermutationParameters& Parameters, FShaderCompilerEnvironment& OutEnvironment)
+	{
+		FGlobalShader::ModifyCompilationEnvironment(Parameters, OutEnvironment);
+		OutEnvironment.SetDefine(TEXT("PIXEL_SHADER"), 1);
+	}
+};
+
+IMPLEMENT_GLOBAL_SHADER(FVisualizeSparseVolumeTextureVS, "/Engine/Private/SparseVolumeTexture/VisualizeSparseVolumeTexture.usf", "VisualizeSparseVolumeTextureVS", SF_Vertex);
+IMPLEMENT_GLOBAL_SHADER(FVisualizeSparseVolumeTexturePS, "/Engine/Private/SparseVolumeTexture/VisualizeSparseVolumeTexture.usf", "VisualizeSparseVolumeTexturePS", SF_Pixel);
+
+
+////////////////////////////////////////////////////////////////////////////////////////////////////
+
+DECLARE_GPU_STAT(SparseVolumeTextureViewer);
+
+void AddSparseVolumeTextureViewerRenderPass(FRDGBuilder& GraphBuilder, FSceneRenderer& SceneRenderer, FSceneTextures& SceneTextures)
+{
+	FScene* Scene = SceneRenderer.Scene;
+
+	if (Scene->SparseVolumeTextureViewers.Num() == 0)
+	{
+		return;
+	}
+
+	RDG_EVENT_SCOPE(GraphBuilder, "SparseVolumeTextureViewer");
+	RDG_GPU_STAT_SCOPE(GraphBuilder, SparseVolumeTextureViewer);
+	RDG_CSV_STAT_EXCLUSIVE_SCOPE(GraphBuilder, SparseVolumeTextureViewer);
+
+	for (FViewInfo& View : SceneRenderer.Views)
+	{
+		const FIntRect& ViewportRect = View.ViewRect;
+
+		for (auto& SVTProxy : Scene->SparseVolumeTextureViewers)
+		{
+			const FBoxSphereBounds& SVTProxyBound = SVTProxy->VolumeWorldBounds;
+			
+			FGlobalShaderMap* GlobalShaderMap = GetGlobalShaderMap(GMaxRHIFeatureLevel);
+			FVisualizeSparseVolumeTextureVS::FPermutationDomain VsPermutationVector;
+			TShaderMapRef<FVisualizeSparseVolumeTextureVS> VertexShader(GlobalShaderMap, VsPermutationVector);
+			FVisualizeSparseVolumeTexturePS::FPermutationDomain PsPermutationVector;
+			TShaderMapRef<FVisualizeSparseVolumeTexturePS> PixelShader(GlobalShaderMap, PsPermutationVector);
+
+			FVisualizeSparseVolumeTexturePS::FParameters* PsPassParameters = GraphBuilder.AllocParameters<FVisualizeSparseVolumeTexturePS::FParameters>();
+			PsPassParameters->ViewUniformBuffer = View.ViewUniformBuffer;
+			PsPassParameters->RenderTargets[0] = FRenderTargetBinding(SceneTextures.Color.Target, ERenderTargetLoadAction::ELoad);
+			PsPassParameters->VolumeBoundMinWorld = FVector3f(SVTProxyBound.Origin - SVTProxyBound.BoxExtent);
+			PsPassParameters->VolumeBoundMaxWorld = FVector3f(SVTProxyBound.Origin + SVTProxyBound.BoxExtent);
+			PsPassParameters->SparseVolumeTextureResolution = FVector3f::OneVector;
+			PsPassParameters->SparseVolumeTexturePageTableResolution = FVector3f::OneVector;
+			PsPassParameters->SparseVolumeTexture = GBlackVolumeTexture->TextureRHI;
+			PsPassParameters->SparseVolumeTexturePageTable = GBlackUintVolumeTexture->TextureRHI;
+
+
+			if (SVTProxy->SparseVolumeTextureSceneProxy)
+			{
+				const FSparseVolumeAssetHeader& Header = SVTProxy->SparseVolumeTextureSceneProxy->GetHeader();
+
+				PsPassParameters->SparseVolumeTextureResolution = FVector3f(Header.SourceVolumeResolution);
+				PsPassParameters->SparseVolumeTexturePageTableResolution = FVector3f(Header.PageTableVolumeResolution);
+
+				FRHITexture* Texture = SVTProxy->SparseVolumeTextureSceneProxy->GetTileDataTextureRHI();
+				if (Texture)
+				{
+					PsPassParameters->SparseVolumeTexture = Texture;
+				}
+				FRHITexture* PageTableTexture = SVTProxy->SparseVolumeTextureSceneProxy->GetPageTableTextureRHI();
+				if (PageTableTexture)
+				{
+					PsPassParameters->SparseVolumeTexturePageTable = PageTableTexture;
+				}
+			}
+
+			ClearUnusedGraphResources(PixelShader, PsPassParameters);
+
+			float VsDepthAsDeviceZ = 0.1f;
+
+			GraphBuilder.AddPass(
+				{},
+				PsPassParameters,
+				ERDGPassFlags::Raster,
+				[PsPassParameters, VertexShader, PixelShader, ViewportRect, VsDepthAsDeviceZ](FRHICommandList& RHICmdListLambda)
+				{
+					RHICmdListLambda.SetViewport(ViewportRect.Min.X, ViewportRect.Min.Y, 0.0f, ViewportRect.Max.X, ViewportRect.Max.Y, 1.0f);
+
+					FGraphicsPipelineStateInitializer GraphicsPSOInit;
+					RHICmdListLambda.ApplyCachedRenderTargets(GraphicsPSOInit);
+
+					GraphicsPSOInit.BlendState = TStaticBlendState<CW_RGB, BO_Add, BF_One, BF_SourceAlpha, BO_Add, BF_Zero, BF_One>::GetRHI();
+					GraphicsPSOInit.DepthStencilState = TStaticDepthStencilState<false, CF_Always>::GetRHI();
+					GraphicsPSOInit.RasterizerState = TStaticRasterizerState<FM_Solid, CM_None>::GetRHI();
+					GraphicsPSOInit.BoundShaderState.VertexDeclarationRHI = GEmptyVertexDeclaration.VertexDeclarationRHI;
+					GraphicsPSOInit.BoundShaderState.VertexShaderRHI = VertexShader.GetVertexShader();
+					GraphicsPSOInit.BoundShaderState.PixelShaderRHI = PixelShader.GetPixelShader();
+					GraphicsPSOInit.PrimitiveType = PT_TriangleList;
+					SetGraphicsPipelineState(RHICmdListLambda, GraphicsPSOInit, 0);
+
+					SetShaderParameters(RHICmdListLambda, PixelShader, PixelShader.GetPixelShader(), *PsPassParameters);
+
+					FVisualizeSparseVolumeTextureVS::FParameters VsPassParameters;
+					VsPassParameters.DepthAsDeviceZ = VsDepthAsDeviceZ;
+					SetShaderParameters(RHICmdListLambda, VertexShader, VertexShader.GetVertexShader(), VsPassParameters);
+
+					RHICmdListLambda.DrawPrimitive(0, 1, 1);
+				});
+		}
+	}
+
+}
diff --git a/Engine/Source/Runtime/Renderer/Private/SparseVolumeTexture/SparseVolumeTextureViewerRendering.h b/Engine/Source/Runtime/Renderer/Private/SparseVolumeTexture/SparseVolumeTextureViewerRendering.h
new file mode 100644
--- /dev/null
+++ b/Engine/Source/Runtime/Renderer/Private/SparseVolumeTexture/SparseVolumeTextureViewerRendering.h
@@ -0,0 +1,16 @@
+// Copyright Epic Games, Inc.All Rights Reserved.
+
+#pragma once
+
+#include "CoreMinimal.h"
+
+class FScene;
+class FSceneRenderer;
+class FRDGBuilder;
+struct FMinimalSceneTextures;
+struct FSceneTextures;
+struct FScreenPassTexture;
+
+
+void AddSparseVolumeTextureViewerRenderPass(FRDGBuilder& GraphBuilder, FSceneRenderer& SceneRenderer, FSceneTextures& SceneTextures);
+
diff --git a/Engine/Source/Runtime/Renderer/Private/SparseVolumeTexture/SparseVolumeTextureViewerSceneProxy.cpp b/Engine/Source/Runtime/Renderer/Private/SparseVolumeTexture/SparseVolumeTextureViewerSceneProxy.cpp
new file mode 100644
--- /dev/null
+++ b/Engine/Source/Runtime/Renderer/Private/SparseVolumeTexture/SparseVolumeTextureViewerSceneProxy.cpp
@@ -0,0 +1,95 @@
+// Copyright Epic Games, Inc. All Rights Reserved.
+
+#include "SparseVolumeTexture/SparseVolumeTextureViewerSceneProxy.h"
+#include "SparseVolumeTexture/SparseVolumeTextureViewerComponent.h"
+#include "SparseVolumeTexture/SparseVolumeTexture.h"
+
+#include "SceneView.h"
+#include "Renderer/Private/ScenePrivate.h"
+#include "RenderGraphUtils.h"
+
+#define LOCTEXT_NAMESPACE "FSparseVolumeTextureViewerSceneProxy"
+
+void FScene::AddSparseVolumeTextureViewer(FSparseVolumeTextureViewerSceneProxy* SVTV)
+{
+	check(SVTV);
+	FScene* Scene = this;
+
+	ENQUEUE_RENDER_COMMAND(FAddSparseVolumeTextureViewerCommand)(
+		[Scene, SVTV](FRHICommandListImmediate& RHICmdList)
+		{
+			Scene->SparseVolumeTextureViewers.Emplace(SVTV);
+		});
+}
+
+void FScene::RemoveSparseVolumeTextureViewer(FSparseVolumeTextureViewerSceneProxy* SVTV)
+{
+	check(SVTV);
+	FScene* Scene = this;
+
+	ENQUEUE_RENDER_COMMAND(FRemoveSparseVolumeTextureCommand)(
+		[Scene, SVTV](FRHICommandListImmediate& RHICmdList)
+		{
+			Scene->SparseVolumeTextureViewers.Remove(SVTV);
+		});
+}
+
+////////////////////////////////////////////////////////////////////////////////////////////////
+
+FSparseVolumeTextureViewerSceneProxy::FSparseVolumeTextureViewerSceneProxy(const USparseVolumeTextureViewerComponent* InComponent, int32 FrameIndex, FName ResourceName)
+	: FPrimitiveSceneProxy((UPrimitiveComponent*)InComponent, ResourceName)
+	, SparseVolumeTextureSceneProxy(nullptr)
+{
+	if (InComponent->SparseVolumeTexturePreview)
+	{
+		check(InComponent->SparseVolumeTexturePreview);
+		if (InComponent->SparseVolumeTexturePreview->IsA(UAnimatedSparseVolumeTexture::StaticClass()))
+		{
+			const UAnimatedSparseVolumeTexture* AnimatedSparseVolumeTexture = CastChecked<const UAnimatedSparseVolumeTexture>(InComponent->SparseVolumeTexturePreview);
+			SparseVolumeTextureSceneProxy = AnimatedSparseVolumeTexture->GetSparseVolumeTextureFrameSceneProxy(FrameIndex);
+		}
+		else
+		{
+			// Regular texture from proxy
+			SparseVolumeTextureSceneProxy = InComponent->SparseVolumeTexturePreview->GetSparseVolumeTextureSceneProxy();
+		}
+	}
+}
+
+SIZE_T FSparseVolumeTextureViewerSceneProxy::GetTypeHash() const
+{
+	static size_t UniquePointer = 0;
+	return reinterpret_cast<size_t>(&UniquePointer);
+}
+
+void FSparseVolumeTextureViewerSceneProxy::GetDynamicMeshElements(const TArray<const FSceneView*>& Views, const FSceneViewFamily& ViewFamily, uint32 VisibilityMap, FMeshElementCollector& Collector) const
+{
+	check(IsInRenderingThread());
+
+	for (int32 ViewIndex = 0; ViewIndex < Views.Num(); ViewIndex++)
+	{
+		const FSceneView* View = Views[ViewIndex];
+
+		if (IsShown(View) && (VisibilityMap & (1 << ViewIndex)))
+		{
+			// Only render Bounds
+			FPrimitiveDrawInterface* PDI = Collector.GetPDI(ViewIndex);
+			FBoxSphereBounds LocalBound = GetBounds();
+			RenderBounds(PDI, ViewFamily.EngineShowFlags, LocalBound, IsSelected());
+		}
+	}
+}
+
+FPrimitiveViewRelevance FSparseVolumeTextureViewerSceneProxy::GetViewRelevance(const FSceneView * View) const
+{
+	FPrimitiveViewRelevance Result;
+	Result.bDrawRelevance = IsShown(View);
+	Result.bDynamicRelevance = true;
+	Result.bStaticRelevance = false;
+	Result.bRenderInMainPass = ShouldRenderInMainPass();
+	Result.bOpaque = false;
+	Result.bNormalTranslucency = true;
+	return Result;
+}
+
+#undef LOCTEXT_NAMESPACE
diff --git a/Engine/Source/Runtime/Renderer/Private/SparseVolumeTexture/SparseVolumeTextureViewerSceneProxy.h b/Engine/Source/Runtime/Renderer/Private/SparseVolumeTexture/SparseVolumeTextureViewerSceneProxy.h
new file mode 100644
--- /dev/null
+++ b/Engine/Source/Runtime/Renderer/Private/SparseVolumeTexture/SparseVolumeTextureViewerSceneProxy.h
@@ -0,0 +1,37 @@
+// Copyright Epic Games, Inc. All Rights Reserved.
+
+#pragma once
+
+#include "CoreMinimal.h"
+#include "PrimitiveSceneProxy.h"
+#include "RendererInterface.h"
+#include "SceneViewExtension.h"
+
+class USparseVolumeTextureViewerComponent;
+class FSparseVolumeTextureSceneProxy;
+
+// Proxy representing the rendering of the SparseVolumeTextureViewer component on the render thread.
+class FSparseVolumeTextureViewerSceneProxy : public FPrimitiveSceneProxy
+{
+public:
+	FSparseVolumeTextureViewerSceneProxy(const USparseVolumeTextureViewerComponent* InComponent, int32 FrameIndex, FName ResourceName = NAME_None);
+	virtual ~FSparseVolumeTextureViewerSceneProxy() = default;
+
+	const FSparseVolumeTextureSceneProxy* SparseVolumeTextureSceneProxy = nullptr;
+	FBoxSphereBounds VolumeWorldBounds;
+
+protected:
+
+	//~ Begin FPrimitiveSceneProxy Interface
+
+	virtual SIZE_T GetTypeHash() const override;
+	virtual void CreateRenderThreadResources() override {};
+	virtual void DestroyRenderThreadResources() override {};
+	virtual void GetDynamicMeshElements(const TArray<const FSceneView*>& Views, const FSceneViewFamily& ViewFamily, uint32 VisibilityMap, FMeshElementCollector& Collector) const override;
+	virtual FPrimitiveViewRelevance GetViewRelevance(const FSceneView* View) const override;
+	virtual uint32 GetMemoryFootprint() const override { return(sizeof(*this) + GetAllocatedSize()); }
+
+	//~ End FPrimitiveSceneProxy Interface
+
+private:
+};
