Fix for vulnerability caused by possible integer overflow in UVolumeTexture::UpdateSourceFromSourceTexture(). Fixed by adding an overflow-safe multiply function to compute the allocation size.

#rnx

diff --git a/Engine/Source/Runtime/Engine/Private/VolumeTexture.cpp b/Engine/Source/Runtime/Engine/Private/VolumeTexture.cpp
--- a/Engine/Source/Runtime/Engine/Private/VolumeTexture.cpp
+++ b/Engine/Source/Runtime/Engine/Private/VolumeTexture.cpp
@@ -33,6 +33,18 @@ bool GSupportsVolumeTextureStreaming = true;
 // Limit the possible depth of volume texture otherwise when the user converts 2D textures, they can crash the engine.
 const int32 MAX_VOLUME_TEXTURE_DEPTH = 512;
 
+// Returns 0 if the product of A and B would overflow
+static int32 CheckedNonNegativeProduct(int32 A, int32 B)
+{
+	int64 Product = (int64)A * (int64)B;
+	// Either factor negative or product too large to fit in int32?
+	if (A < 0 || B < 0 || Product > 0x7fffffff)
+	{
+		return 0;
+	}
+	return Product;
+}
+
 //*****************************************************************************
 //***************************** UVolumeTexture ********************************
 //*****************************************************************************
@@ -157,12 +169,19 @@ bool UVolumeTexture::UpdateSourceFromSourceTexture()
 		if (TileSizeZ > 0)
 		{
 			const int32 FormatDataSize = InitialSource.GetBytesPerPixel();
-			if (FormatDataSize > 0)
+
+			// Do an overflow checked product to avoid AllocationSize being smaller than the actual required memory which would cause
+			// the following loops to write past the allocation. We also force allocations to be <2GB with this by using int32.
+			const int32 AllocationSize = CheckedNonNegativeProduct(
+				CheckedNonNegativeProduct(Source2DTileSizeX, Source2DTileSizeY), 
+				CheckedNonNegativeProduct(TileSizeZ, FormatDataSize));
+
+			if (AllocationSize > 0)
 			{
 				TArray64<uint8> Ref2DData;
 				if (InitialSource.GetMipData(Ref2DData, 0))
 				{
-					uint8* NewData = (uint8*)FMemory::Malloc(Source2DTileSizeX * Source2DTileSizeY * TileSizeZ * FormatDataSize);
+					uint8* NewData = (uint8*)FMemory::Malloc(AllocationSize);
 					uint8* CurPos = NewData;
 					
 					for (int32 PosZ = 0; PosZ < TileSizeZ; ++PosZ)
