Fix invalid warning.

diff --git a/Engine/Source/Runtime/Engine/Private/Materials/Material.cpp b/Engine/Source/Runtime/Engine/Private/Materials/Material.cpp
--- a/Engine/Source/Runtime/Engine/Private/Materials/Material.cpp
+++ b/Engine/Source/Runtime/Engine/Private/Materials/Material.cpp
@@ -4593,11 +4593,6 @@ void UMaterial::RebuildShadingModelField()
 	const bool bStrataEnabled = Engine_IsStrataEnabled();
 	UMaterialEditorOnlyData* EditorOnly = GetEditorOnlyData();
 
-	auto StrataCheckShadingModelFromExpressionWarning = [&]()
-	{
-		UE_LOG(LogMaterial, Warning, TEXT("%s: Material seems to have been converted with ShadingModel `FromExpression`, but now is using `%s`.\nThis indicates the material has been resaved with Strata disabled when it already had a conversion node setup.\nThe material may look wrong: please delete all nodes plugged in FrontMaterial to rerun the conversion process."), *GetName(), *GetShadingModelString(ShadingModel));
-	};
-
 	if (bStrataEnabled && EditorOnly->FrontMaterial.IsConnected())
 	{
 		FStrataMaterialInfo StrataMaterialInfo;
@@ -4624,8 +4619,6 @@ void UMaterial::RebuildShadingModelField()
 			}
 			else if (StrataMaterialInfo.HasShadingModelFromExpression())
 			{
-				StrataCheckShadingModelFromExpressionWarning();
-
 				if (BlendMode == EBlendMode::BLEND_Opaque || BlendMode == EBlendMode::BLEND_Masked)
 				{
 					// Nothing
@@ -4671,9 +4664,6 @@ void UMaterial::RebuildShadingModelField()
 		if (StrataMaterialInfo.HasShadingModelFromExpression())
 		{
 			MaterialDomain = EMaterialDomain::MD_Surface;
-
-			StrataCheckShadingModelFromExpressionWarning();
-
 			{
 				TArray<UMaterialExpressionShadingModel*> ShadingModelExpressions;
 				GetAllExpressionsInMaterialAndFunctionsOfType(ShadingModelExpressions);
