LevelInstance: Fix LevelInstanceID not being unique because of parent level instance being embedded (World Partition only)

#rnx

diff --git a/Engine/Source/Runtime/Engine/Private/LevelInstance/LevelInstanceSubsystem.cpp b/Engine/Source/Runtime/Engine/Private/LevelInstance/LevelInstanceSubsystem.cpp
--- a/Engine/Source/Runtime/Engine/Private/LevelInstance/LevelInstanceSubsystem.cpp
+++ b/Engine/Source/Runtime/Engine/Private/LevelInstance/LevelInstanceSubsystem.cpp
@@ -122,13 +122,22 @@ ILevelInstanceInterface* ULevelInstanceSubsystem::GetLevelInstance(const FLevelI
 
 FLevelInstanceID::FLevelInstanceID(ULevelInstanceSubsystem* LevelInstanceSubsystem, ILevelInstanceInterface* LevelInstance)
 {
-	LevelInstanceSubsystem->ForEachLevelInstanceAncestorsAndSelf(CastChecked<AActor>(LevelInstance), [this](const ILevelInstanceInterface* AncestorOrSelf)
+	AActor* LevelInstanceActor = CastChecked<AActor>(LevelInstance);
+	LevelInstanceSubsystem->ForEachLevelInstanceAncestorsAndSelf(LevelInstanceActor, [this](const ILevelInstanceInterface* AncestorOrSelf)
 	{
 		Guids.Add(AncestorOrSelf->GetLevelInstanceGuid());
 		return true;
 	});
 	check(!Guids.IsEmpty());
-	Hash = CityHash64((const char*)Guids.GetData(), Guids.Num() * sizeof(FGuid));
+	
+	// Add Actor Name to hash because with World Partition Embedding top level of Level Instance hierarchy can get stripped leaving us with clashing ids.
+	// When embedding actors we make sure their names are unique (they get suffixed with their parent container id)
+	ActorName = LevelInstanceActor->GetFName();
+	
+	FString NameStr = ActorName.ToString();
+	const uint64 NameHash = CityHash64((const char*)*NameStr, NameStr.Len() * sizeof(TCHAR));
+	
+	Hash = CityHash64WithSeed((const char*)Guids.GetData(), Guids.Num() * sizeof(FGuid), NameHash);
 }
 
 FLevelInstanceID ULevelInstanceSubsystem::RegisterLevelInstance(ILevelInstanceInterface* LevelInstance)
diff --git a/Engine/Source/Runtime/Engine/Public/LevelInstance/LevelInstanceTypes.h b/Engine/Source/Runtime/Engine/Public/LevelInstance/LevelInstanceTypes.h
--- a/Engine/Source/Runtime/Engine/Public/LevelInstance/LevelInstanceTypes.h
+++ b/Engine/Source/Runtime/Engine/Public/LevelInstance/LevelInstanceTypes.h
@@ -26,7 +26,7 @@ struct FLevelInstanceID
 
 	inline bool operator==(const FLevelInstanceID& Other) const
 	{
-		return Hash == Other.Hash && Guids == Other.Guids;
+		return Hash == Other.Hash && Guids == Other.Guids && ActorName == Other.ActorName;
 	}
 
 	inline bool IsValid() const { return !Guids.IsEmpty(); }
@@ -36,6 +36,7 @@ struct FLevelInstanceID
 private:
 	uint64 Hash = 0;
 	TArray<FGuid> Guids;
+	FName ActorName;
 };
 
 UENUM()
