World Partition - HLOD: HLOD Instancing should copy copy primitive data of static mesh components, not just of ISMC

diff --git a/Engine/Source/Runtime/Engine/Private/WorldPartition/HLOD/HLODBuilder.cpp b/Engine/Source/Runtime/Engine/Private/WorldPartition/HLOD/HLODBuilder.cpp
--- a/Engine/Source/Runtime/Engine/Private/WorldPartition/HLOD/HLODBuilder.cpp
+++ b/Engine/Source/Runtime/Engine/Private/WorldPartition/HLOD/HLODBuilder.cpp
@@ -204,6 +204,7 @@ TArray<UActorComponent*> UHLODBuilder::BatchInstances(const TArray<UActorCompone
 		}
 		else
 		{
+			InstancingData.NumCustomDataFloats = FMath::Max(InstancingData.NumCustomDataFloats, SMC->GetCustomPrimitiveData().Data.Num());
 			InstancingData.NumInstances++;
 		}
 	}
@@ -240,8 +241,12 @@ TArray<UActorComponent*> UHLODBuilder::BatchInstances(const TArray<UActorCompone
 		}
 		else
 		{
+			// Add transform
 			InstancingData.InstancesTransforms.Add(SMC->GetComponentTransform());
-			InstancingData.InstancesCustomData.AddDefaulted(InstancingData.NumCustomDataFloats);
+
+			// Add custom data
+			InstancingData.InstancesCustomData.Append(SMC->GetCustomPrimitiveData().Data);
+			InstancingData.InstancesCustomData.AddDefaulted(InstancingData.NumCustomDataFloats - SMC->GetCustomPrimitiveData().Data.Num());
 		}
 	}
 
