VolumetricCloud material shaders are now compiled optionally for material belonging to the Volumetric Domain.
This is now opt in.
If a material with bUseWithVolumetricCloud=0 is used on volumetric cloud component, the screen will display a error message to the user.
Conversion: volumetric domain material have bUseWithVolumetricCloud automatically set to true when any ofthe cloud node is used.
With that, volumetric fog shader won't compile for clouds by default but it is still possible to do so for volumetric fog to cloud transitions.

diff --git a/Engine/Source/Runtime/Core/Public/UObject/UE5MainStreamObjectVersion.h b/Engine/Source/Runtime/Core/Public/UObject/UE5MainStreamObjectVersion.h
--- a/Engine/Source/Runtime/Core/Public/UObject/UE5MainStreamObjectVersion.h
+++ b/Engine/Source/Runtime/Core/Public/UObject/UE5MainStreamObjectVersion.h
@@ -286,6 +286,9 @@ struct CORE_API FUE5MainStreamObjectVersion
 		// Support 16-bit skin weights on SkeletalMesh
 		IncreasedSkinWeightPrecision,
 
+		// bIsUsedWithVolumetricCloud flag auto conversion
+		MaterialHasIsUsedWithVolumetricCloudFlag,
+
 		// -----<new versions can be added above this line>-------------------------------------------------
 		VersionPlusOne,
 		LatestVersion = VersionPlusOne - 1
diff --git a/Engine/Source/Runtime/Engine/Classes/Materials/Material.h b/Engine/Source/Runtime/Engine/Classes/Materials/Material.h
--- a/Engine/Source/Runtime/Engine/Classes/Materials/Material.h
+++ b/Engine/Source/Runtime/Engine/Classes/Materials/Material.h
@@ -758,6 +758,13 @@ public:
 	UPROPERTY(EditAnywhere, BlueprintReadOnly, Category = Usage)
 	uint32 bUsedWithNanite : 1;
 
+	/**
+	 * Indicates that the material and its instances with volumetric cloud. Without that flag, it can only be used on volumetric fog.
+	 * This will result in the shaders required to support Volumetric Cloud rendering being compiled which will increase shader compile time and memory usage.
+	 */
+	UPROPERTY(EditAnywhere, BlueprintReadOnly, Category = Usage)
+	uint32 bUsedWithVolumetricCloud : 1;
+
 	/** 
 	 * Indicates that the material and its instances can be used with Slate UI and UMG
 	 * This will result in the shaders required to support UI materials being compiled which will increase shader compile time and memory usage.
diff --git a/Engine/Source/Runtime/Engine/Classes/Materials/MaterialInterface.h b/Engine/Source/Runtime/Engine/Classes/Materials/MaterialInterface.h
--- a/Engine/Source/Runtime/Engine/Classes/Materials/MaterialInterface.h
+++ b/Engine/Source/Runtime/Engine/Classes/Materials/MaterialInterface.h
@@ -71,6 +71,7 @@ enum EMaterialUsage
 	MATUSAGE_LidarPointCloud,
 	MATUSAGE_VirtualHeightfieldMesh,
 	MATUSAGE_Nanite,
+	MATUSAGE_VolumetricCloud,
 
 	MATUSAGE_MAX,
 };
diff --git a/Engine/Source/Runtime/Engine/Private/Materials/Material.cpp b/Engine/Source/Runtime/Engine/Private/Materials/Material.cpp
--- a/Engine/Source/Runtime/Engine/Private/Materials/Material.cpp
+++ b/Engine/Source/Runtime/Engine/Private/Materials/Material.cpp
@@ -78,6 +78,9 @@
 #include "Materials/MaterialExpressionDivide.h"
 #include "Materials/MaterialExpressionComponentMask.h"
 #include "Materials/MaterialExpressionSingleLayerWaterMaterialOutput.h"
+#include "Materials/MaterialExpressionVolumetricAdvancedMaterialInput.h"
+#include "Materials/MaterialExpressionVolumetricAdvancedMaterialOutput.h"
+#include "Materials/MaterialExpressionCloudLayer.h"
 #include "Materials/MaterialExpressionStrata.h"
 #include "Materials/StrataMaterial.h"
 #include "Materials/MaterialExpressionThinTranslucentMaterialOutput.h"
@@ -969,6 +972,7 @@ UMaterial::UMaterial(const FObjectInitializer& ObjectInitializer)
 	BlendableOutputAlpha = false;
 	bIsBlendable = true;
 	bEnableStencilTest = false;
+	bUsedWithVolumetricCloud = false;
 
 	bUseEmissiveForDynamicAreaLighting = false;
 	RefractionDepthBias = 0.0f;
@@ -1344,6 +1348,7 @@ bool UMaterial::GetUsageByFlag(EMaterialUsage Usage) const
 		case MATUSAGE_LidarPointCloud: UsageValue = bUsedWithLidarPointCloud; break;
 		case MATUSAGE_VirtualHeightfieldMesh: UsageValue = bUsedWithVirtualHeightfieldMesh; break;
 		case MATUSAGE_Nanite: UsageValue = bUsedWithNanite; break;
+		case MATUSAGE_VolumetricCloud: UsageValue = bUsedWithVolumetricCloud; break;
 		default: UE_LOG(LogMaterial, Fatal,TEXT("Unknown material usage: %u"), (int32)Usage);
 	};
 	return UsageValue;
@@ -1587,6 +1592,10 @@ void UMaterial::SetUsageByFlag(EMaterialUsage Usage, bool NewValue)
 		{
 			bUsedWithNanite = NewValue; break;
 		}
+		case MATUSAGE_VolumetricCloud:
+		{
+			bUsedWithVolumetricCloud = NewValue; break;
+		}
 		default: UE_LOG(LogMaterial, Fatal, TEXT("Unknown material usage: %u"), (int32)Usage);
 	};
 #if WITH_EDITOR
@@ -1619,6 +1628,7 @@ FString UMaterial::GetUsageName(EMaterialUsage Usage) const
 		case MATUSAGE_LidarPointCloud: UsageName = TEXT("bUsedWithLidarPointCloud"); break;
 		case MATUSAGE_VirtualHeightfieldMesh: UsageName = TEXT("bUsedWithVirtualHeightfieldMesh"); break;
 		case MATUSAGE_Nanite: UsageName = TEXT("bUsedWithNanite"); break;
+		case MATUSAGE_VolumetricCloud: UsageName = TEXT("bUsedWithVolumetricCloud"); break;
 		default: UE_LOG(LogMaterial, Fatal,TEXT("Unknown material usage: %u"), (int32)Usage);
 	};
 	return UsageName;
@@ -3754,6 +3764,20 @@ void UMaterial::PostLoad()
 
 #endif // WITH_EDITOR
 
+#if WITH_EDITORONLY_DATA
+	if (MaterialDomain == MD_Volume && UE5MainVer < FUE5MainStreamObjectVersion::MaterialHasIsUsedWithVolumetricCloudFlag)
+	{
+		bUsedWithVolumetricCloud = HasAnyExpressionsInMaterialAndFunctionsOfType<UMaterialExpressionCloudSampleAttribute>()
+								|| HasAnyExpressionsInMaterialAndFunctionsOfType<UMaterialExpressionVolumetricAdvancedMaterialOutput>()
+								|| HasAnyExpressionsInMaterialAndFunctionsOfType<UMaterialExpressionVolumetricAdvancedMaterialInput>();
+		if (bUsedWithVolumetricCloud)
+		{
+			static FGuid BackwardsCompatibilityUsedWithVolumetricCloudConversionGuid(TEXT("6F6336BF-D377-4FA5-9887-A7457BE6DE92"));
+			ReleaseResourcesAndMutateDDCKey(BackwardsCompatibilityUsedWithVolumetricCloudConversionGuid);
+		}
+	}
+#endif
+
 	// Strata materials conversion needs to be done after expressions are cached, otherwise material function won't have 
 	// valid inputs in certain cases
 	ConvertMaterialToStrataMaterial();
@@ -4180,7 +4204,7 @@ bool UMaterial::CanEditChange(const FProperty* InProperty) const
 
 		if (FCString::Strncmp(*PropertyName, TEXT("bUsedWith"), 9) == 0)
 		{
-			return MaterialDomain == MD_DeferredDecal || MaterialDomain == MD_Surface;
+			return MaterialDomain == MD_DeferredDecal || MaterialDomain == MD_Surface || MaterialDomain == MD_Volume;
 		}
 		else if (PropertyName == GET_MEMBER_NAME_STRING_CHECKED(UMaterial, bUsesDistortion))
 		{
diff --git a/Engine/Source/Runtime/Engine/Private/Materials/MaterialShared.cpp b/Engine/Source/Runtime/Engine/Private/Materials/MaterialShared.cpp
--- a/Engine/Source/Runtime/Engine/Private/Materials/MaterialShared.cpp
+++ b/Engine/Source/Runtime/Engine/Private/Materials/MaterialShared.cpp
@@ -1613,6 +1613,11 @@ bool FMaterialResource::IsUsedWithNanite() const
 	return Material->bUsedWithNanite;
 }
 
+bool FMaterialResource::IsUsedWithVolumetricCloud() const
+{
+	return Material->bUsedWithVolumetricCloud;
+}
+
 bool FMaterialResource::IsTranslucencyAfterDOFEnabled() const 
 { 
 	return Material->TranslucencyPass == MTP_AfterDOF
@@ -6063,6 +6068,7 @@ FMaterialShaderParameters::FMaterialShaderParameters(const FMaterial* InMaterial
 	bIsTranslucencySurface = InMaterial->GetTranslucencyLightingMode() == ETranslucencyLightingMode::TLM_Surface || InMaterial->GetTranslucencyLightingMode() == ETranslucencyLightingMode::TLM_SurfacePerPixelLighting;
 	bShouldDisableDepthTest = InMaterial->ShouldDisableDepthTest();
 	bHasRenderTracePhysicalMaterialOutput = InMaterial->HasRenderTracePhysicalMaterialOutputs();
+	bIsUsedWithVolumetricCloud = InMaterial->IsUsedWithVolumetricCloud();
 }
 
 #undef LOCTEXT_NAMESPACE
diff --git a/Engine/Source/Runtime/Engine/Public/MaterialShared.h b/Engine/Source/Runtime/Engine/Public/MaterialShared.h
--- a/Engine/Source/Runtime/Engine/Public/MaterialShared.h
+++ b/Engine/Source/Runtime/Engine/Public/MaterialShared.h
@@ -1894,6 +1894,7 @@ public:
 	virtual bool IsUsedWithGeometryCollections() const { return false; }
 	virtual bool IsUsedWithAPEXCloth() const { return false; }
 	virtual bool IsUsedWithNanite() const { return false; }
+	virtual bool IsUsedWithVolumetricCloud() const { return false; }
 	virtual bool IsUsedWithUI() const { return false; }
 	virtual bool IsUsedWithGeometryCache() const { return false; }
 	virtual bool IsUsedWithWater() const { return false; }
@@ -2931,6 +2932,7 @@ public:
 	ENGINE_API virtual bool IsUsedWithLidarPointCloud() const override;
 	ENGINE_API virtual bool IsUsedWithVirtualHeightfieldMesh() const override;
 	ENGINE_API virtual bool IsUsedWithNanite() const override;
+	ENGINE_API virtual bool IsUsedWithVolumetricCloud() const override;
 	ENGINE_API virtual bool IsFullyRough() const override;
 	ENGINE_API virtual bool UseNormalCurvatureToRoughness() const override;
 	ENGINE_API virtual enum EMaterialFloatPrecisionMode GetMaterialFloatPrecisionMode() const override;
@@ -3585,6 +3587,7 @@ struct FMaterialShaderParameters
 			uint64 bIsTranslucencySurface : 1;
 			uint64 bShouldDisableDepthTest : 1;
 			uint64 bHasRenderTracePhysicalMaterialOutput : 1;
+			uint64 bIsUsedWithVolumetricCloud : 1;
 		};
 	};
 
diff --git a/Engine/Source/Runtime/Renderer/Private/VolumetricCloudRendering.cpp b/Engine/Source/Runtime/Renderer/Private/VolumetricCloudRendering.cpp
--- a/Engine/Source/Runtime/Renderer/Private/VolumetricCloudRendering.cpp
+++ b/Engine/Source/Runtime/Renderer/Private/VolumetricCloudRendering.cpp
@@ -507,7 +507,7 @@ FORCEINLINE bool IsVolumetricCloudMaterialSupported(const EShaderPlatform Platfo
 
 FORCEINLINE bool IsMaterialCompatibleWithVolumetricCloud(const FMaterialShaderParameters& Material, const EShaderPlatform Platform)
 {
-	return IsVolumetricCloudMaterialSupported(Platform) && Material.MaterialDomain == MD_Volume;
+	return IsVolumetricCloudMaterialSupported(Platform) && Material.bIsUsedWithVolumetricCloud && Material.MaterialDomain == MD_Volume;
 }
 
 
@@ -2077,6 +2077,35 @@ void FSceneRenderer::RenderVolumetricCloudsInternal(FRDGBuilder& GraphBuilder, F
 
 	FMaterialRenderProxy* CloudVolumeMaterialProxy = CloudRC.CloudVolumeMaterialProxy;
 
+#if !UE_BUILD_SHIPPING
+	{
+		const FMaterialRenderProxy* MaterialRenderProxy = nullptr;
+		const FMaterial* MaterialResource = &CloudVolumeMaterialProxy->GetMaterialWithFallback(Scene->GetFeatureLevel(), MaterialRenderProxy);
+		MaterialRenderProxy = MaterialRenderProxy ? MaterialRenderProxy : CloudVolumeMaterialProxy;
+
+		if (!MaterialResource->IsUsedWithVolumetricCloud())
+		{
+			// This is not a material designed to run with volumetric cloud.
+			OnGetOnScreenMessages.AddLambda([](FScreenMessageWriter& ScreenMessageWriter)->void
+				{
+					static const FText Message = NSLOCTEXT("Renderer", "MaterialNotVolumetricCloud", "A material assigned on a volumetric cloud component does not have the UsedWithVolumetricCloud flag set.");
+					ScreenMessageWriter.DrawLine(Message);
+				});
+			return;
+		}
+		if (MaterialResource->GetMaterialDomain() != MD_Volume)
+		{
+			// This is not a material designed to run with volumetric cloud.
+			OnGetOnScreenMessages.AddLambda([](FScreenMessageWriter& ScreenMessageWriter)->void
+				{
+					static const FText Message = NSLOCTEXT("Renderer", "MaterialNotVolumetricDomain", "A material assigned on a volumetric cloud component does not have the Volumetric domain set.");
+					ScreenMessageWriter.DrawLine(Message);
+				});
+			return;
+		}
+	}
+#endif
+
 	// Copy parameters to lambda
 	const bool bShouldViewRenderVolumetricRenderTarget = CloudRC.bShouldViewRenderVolumetricRenderTarget;
 	const bool bSkipAtmosphericLightShadowmap = CloudRC.bSkipAtmosphericLightShadowmap;
