Use common GUID invalidation workflow for some data conversion.

diff --git a/Engine/Source/Runtime/Engine/Private/Materials/Material.cpp b/Engine/Source/Runtime/Engine/Private/Materials/Material.cpp
--- a/Engine/Source/Runtime/Engine/Private/Materials/Material.cpp
+++ b/Engine/Source/Runtime/Engine/Private/Materials/Material.cpp
@@ -2757,12 +2757,8 @@ void UMaterial::Serialize(FArchive& Ar)
 
 			// Now force the material to recompile and we use a hash of the original StateId.
 			// This is to avoid having different StateId each time we load the material and to not forever recompile it,i.e. use a cached version.
-			uint32 HashBuffer[5];
-			FSHA1::HashBuffer(&StateId, sizeof(FGuid), reinterpret_cast<uint8*>(HashBuffer));
-			StateId.A = HashBuffer[0];
-			StateId.B = HashBuffer[1];
-			StateId.C = HashBuffer[2];
-			StateId.D = HashBuffer[3];
+			static FGuid VolumeExtinctionBecomesRGBConversionGuid(TEXT("2768E88D-9B58-4C53-9CB9-75696D1DF0CD"));
+			ReleaseResourcesAndMutateDDCKey(VolumeExtinctionBecomesRGBConversionGuid);
 		}
 	}
 #endif // WITH_EDITORONLY_DATA
@@ -3467,12 +3463,8 @@ void UMaterial::ConvertMaterialToStrataMaterial()
 	{
 		// Now force the material to recompile and we use a hash of the original StateId.
 		// This is to avoid having different StateId each time we load the material and to not forever recompile it, i.e. use a cached version.
-		uint32 HashBuffer[5];
-		FSHA1::HashBuffer(&StateId, sizeof(FGuid), reinterpret_cast<uint8*>(HashBuffer));
-		StateId.A = HashBuffer[0];
-		StateId.B = HashBuffer[1];
-		StateId.C = HashBuffer[2];
-		StateId.D = HashBuffer[3];
+		static FGuid LegacyToStrataConversionGuid(TEXT("0DAD35FE-21AE-4274-8B41-6C9D47285D8A"));
+		ReleaseResourcesAndMutateDDCKey(LegacyToStrataConversionGuid);
 	}
 
 	// For rebuild the shading mode since we have change it
