Fix Strata crash with shading model 'from expression' case.

diff --git a/Engine/Source/Runtime/Engine/Private/Materials/HLSLMaterialTranslator.cpp b/Engine/Source/Runtime/Engine/Private/Materials/HLSLMaterialTranslator.cpp
--- a/Engine/Source/Runtime/Engine/Private/Materials/HLSLMaterialTranslator.cpp
+++ b/Engine/Source/Runtime/Engine/Private/Materials/HLSLMaterialTranslator.cpp
@@ -9533,7 +9533,10 @@ int32 FHLSLMaterialTranslator::ShadingModel(EMaterialShadingModel InSelectedShad
 		InSelectedShadingModel = MSM_DefaultLit;
 	}
 	
-	ShadingModelsFromCompilation.AddShadingModel(InSelectedShadingModel);
+	if (InSelectedShadingModel < MSM_NUM)
+	{
+		ShadingModelsFromCompilation.AddShadingModel(InSelectedShadingModel);
+	}
 	return AddInlinedCodeChunk(MCT_ShadingModel, TEXT("%d"), InSelectedShadingModel);
 }
 
diff --git a/Engine/Source/Runtime/Engine/Private/Materials/Material.cpp b/Engine/Source/Runtime/Engine/Private/Materials/Material.cpp
--- a/Engine/Source/Runtime/Engine/Private/Materials/Material.cpp
+++ b/Engine/Source/Runtime/Engine/Private/Materials/Material.cpp
@@ -4686,7 +4686,11 @@ void UMaterial::RebuildShadingModelField()
 
 				for (UMaterialExpressionShadingModel* MatExpr : ShadingModelExpressions)
 				{
-					ShadingModels.AddShadingModel(MatExpr->ShadingModel);
+					// Ensure the Shading model is valid
+					if (MatExpr->ShadingModel < MSM_NUM)
+					{
+						ShadingModels.AddShadingModel(MatExpr->ShadingModel);
+					}
 				}
 
 				// If no expressions have been found, set a default
@@ -4799,7 +4803,11 @@ void UMaterial::RebuildShadingModelField()
 
 		for (UMaterialExpressionShadingModel* MatExpr : ShadingModelExpressions)
 		{
-			ShadingModels.AddShadingModel(MatExpr->ShadingModel);
+			// Ensure the Shading model is valid
+			if (MatExpr->ShadingModel < MSM_NUM)
+			{
+				ShadingModels.AddShadingModel(MatExpr->ShadingModel);
+			}
 		}
 
 		// If no expressions have been found, set a default
