Tweak tooltip

#rnx

diff --git a/Engine/Source/Runtime/Engine/Classes/GameFramework/SpringArmComponent.h b/Engine/Source/Runtime/Engine/Classes/GameFramework/SpringArmComponent.h
--- a/Engine/Source/Runtime/Engine/Classes/GameFramework/SpringArmComponent.h
+++ b/Engine/Source/Runtime/Engine/Classes/GameFramework/SpringArmComponent.h
@@ -12,7 +12,7 @@
  * This component tries to maintain its children at a fixed distance from the parent,
  * but will retract the children if there is a collision, and spring back when there is no collision.
  *
- * Example: Use as a 'camera boom' to keep the follow camera for a player from colliding into the world.
+ * Example: Use as a 'camera boom' or 'selfie stick' to keep the follow camera for a player from colliding into the world.
  */
 
 UCLASS(ClassGroup=Camera, meta=(BlueprintSpawnableComponent), hideCategories=(Mobility))
