Support enabling Extended Luminance Range for projects that have it disabled in the Renderer settings panel.
- We don't want new projects disabling this setting but still allow existing projects to enable it.

diff --git a/Engine/Source/Runtime/Engine/Classes/Engine/RendererSettings.h b/Engine/Source/Runtime/Engine/Classes/Engine/RendererSettings.h
--- a/Engine/Source/Runtime/Engine/Classes/Engine/RendererSettings.h
+++ b/Engine/Source/Runtime/Engine/Classes/Engine/RendererSettings.h
@@ -657,6 +657,13 @@ class ENGINE_API URendererSettings : public UDeveloperSettings
 		ToolTip = "Default Value for auto exposure bias."))
 	float DefaultFeatureAutoExposureBias;
 
+	UE_DEPRECATED(4.26, "Extend Default Luminance Range is deprecated, and is forced to ON at all times in new projects.")
+	UPROPERTY(config, EditAnywhere, Category = DefaultSettings, meta = (
+		ConsoleVariable = "r.DefaultFeature.AutoExposure.ExtendDefaultLuminanceRange", DisplayName = "Extend default luminance range in Auto Exposure settings",
+		ToolTip = "Whether the default values for AutoExposure should support an extended range of scene luminance. Also changes the exposure settings to be expressed in EV100. Having this setting disabled is deprecated and can only be done manually using r.DefaultFeature.AutoExposure.ExtendDefaultLuminanceRange.",
+		ConfigRestartRequired=true))
+	uint32 bExtendDefaultLuminanceRangeInAutoExposureSettings : 1;
+
 	UPROPERTY(config, EditAnywhere, Category = DefaultSettings, meta = (
 		ConsoleVariable = "r.DefaultFeature.MotionBlur", DisplayName = "Motion Blur",
 		ToolTip = "Whether the default for MotionBlur is enabled or not (postprocess volume/camera/game setting can still override and enable or disable it independently)"))
diff --git a/Engine/Source/Runtime/Engine/Private/RendererSettings.cpp b/Engine/Source/Runtime/Engine/Private/RendererSettings.cpp
--- a/Engine/Source/Runtime/Engine/Private/RendererSettings.cpp
+++ b/Engine/Source/Runtime/Engine/Private/RendererSettings.cpp
@@ -342,6 +342,17 @@ bool URendererSettings::CanEditChange(const FProperty* InProperty) const
 		return !bForwardShading;
 	}
 
+	PRAGMA_DISABLE_DEPRECATION_WARNINGS
+	// only allow changing ExtendDefaultLuminanceRange if it was disabled.
+	// we don't want new projects disabling this setting but still allow existing projects to enable it.
+	static bool bCanEditExtendDefaultLuminanceRange = !bExtendDefaultLuminanceRangeInAutoExposureSettings;
+
+	if ((InProperty->GetFName() == GET_MEMBER_NAME_CHECKED(URendererSettings, bExtendDefaultLuminanceRangeInAutoExposureSettings)))
+	{
+		return bCanEditExtendDefaultLuminanceRange;
+	}
+	PRAGMA_ENABLE_DEPRECATION_WARNINGS
+
 	return true;
 }
 
