Deprecated the CreateSubObjects array so it can be made private in future releases

diff --git a/Engine/Source/Runtime/Engine/Classes/Engine/ActorChannel.h b/Engine/Source/Runtime/Engine/Classes/Engine/ActorChannel.h
--- a/Engine/Source/Runtime/Engine/Classes/Engine/ActorChannel.h
+++ b/Engine/Source/Runtime/Engine/Classes/Engine/ActorChannel.h
@@ -126,9 +126,17 @@ public:
 
 	TSet<FNetworkGUID> PendingGuidResolves;	// These guids are waiting for their resolves, we need to queue up bunches until these are resolved
 
+	UE_DEPRECATED(5.1, "The CreateSubObjects array will be made private in future versions. Use GetCreatedSubObjects() instead")
 	UPROPERTY()
 	TArray< TObjectPtr<UObject> >					CreateSubObjects;		// Any sub-object we created on this channel
 
+	inline const TArray< TObjectPtr<UObject> >& GetCreatedSubObjects() const 
+	{ 
+		PRAGMA_DISABLE_DEPRECATION_WARNINGS
+		return CreateSubObjects;
+		PRAGMA_ENABLE_DEPRECATION_WARNINGS
+	}
+
 	TArray< FNetworkGUID >				QueuedMustBeMappedGuidsInLastBunch;	// Array of guids that will async load on client. This list is used for queued RPC's.
 	TArray< class FOutBunch * >			QueuedExportBunches;				// Bunches that need to be appended to the export list on the next SendBunch call. This list is used for queued RPC's.
 
@@ -419,6 +427,13 @@ private:
 	void TestLegacyReplicateSubObjects(UActorComponent* ReplicatedComponent, FOutBunch& Bunch, FReplicationFlags RepFlags);
 	void TestLegacyReplicateSubObjects(FOutBunch& Bunch, FReplicationFlags RepFlags);
 
+	inline TArray< TObjectPtr<UObject> >& GetCreatedSubObjects()
+	{
+		PRAGMA_DISABLE_DEPRECATION_WARNINGS
+		return CreateSubObjects;
+		PRAGMA_ENABLE_DEPRECATION_WARNINGS
+	}
+
 private:
 
 	// TODO: It would be nice to merge the tracking of these with PendingGuidResolves, to not duplicate memory,
diff --git a/Engine/Source/Runtime/Engine/Private/DataChannel.cpp b/Engine/Source/Runtime/Engine/Private/DataChannel.cpp
--- a/Engine/Source/Runtime/Engine/Private/DataChannel.cpp
+++ b/Engine/Source/Runtime/Engine/Private/DataChannel.cpp
@@ -1056,10 +1056,12 @@ UActorChannel::UActorChannel(const FObjectInitializer& ObjectInitializer)
 	QueuedCloseReason = EChannelCloseReason::Destroyed;
 }
 
+PRAGMA_DISABLE_DEPRECATION_WARNINGS
 UActorChannel::~UActorChannel()
 {
 
 }
+PRAGMA_ENABLE_DEPRECATION_WARNINGS
 
 void UActorChannel::AppendExportBunches( TArray<FOutBunch *>& OutExportBunches )
 {
@@ -2166,7 +2168,7 @@ void UActorChannel::ReleaseReferences(bool bKeepReplicators)
 {
 	Actor = nullptr;
 	CleanupReplicators(bKeepReplicators);
-	CreateSubObjects.Empty();
+	GetCreatedSubObjects().Empty();
 }
 
 void UActorChannel::BreakAndReleaseReferences()
@@ -2226,9 +2228,11 @@ void UActorChannel::CleanupReplicators(const bool bKeepReplicators)
 	// Cleanup or save replicators
 	for (auto CompIt = ReplicationMap.CreateIterator(); CompIt; ++CompIt)
 	{
+		TSharedRef< FObjectReplicator >& ObjectReplicatorRef = CompIt.Value();
+
 		// NOTE: FObjectReplicator::GetObject is just going to return a raw Object Pointer,
-		// so it won't actually check to see whether or not the Object was marked PendingKill.
-		if (bKeepReplicators && CompIt.Value()->GetObject() != nullptr)
+		// so it won't actually check to see whether or not the Object was marked PendingKill/Garbage.
+		if (bKeepReplicators && ObjectReplicatorRef->GetObject() != nullptr)
 		{
 			LLM_SCOPE_BYTAG(NetConnection);
 
@@ -2242,12 +2246,13 @@ void UActorChannel::CleanupReplicators(const bool bKeepReplicators)
 			//		KeepProcessingActorChannelBunchesMap will get in here, then when the channel closes a second time, we'll hit this assert
 			//		It should be okay to just set the most recent replicator
 			//check( Connection->DormantReplicatorMap.Find( CompIt.Value()->GetObject() ) == NULL );
-			Connection->AddDormantReplicator(CompIt.Value()->GetObject(), CompIt.Value());
-			CompIt.Value()->StopReplicating(this);		// Stop replicating on this channel
+			Connection->AddDormantReplicator(ObjectReplicatorRef->GetObject(), ObjectReplicatorRef);
+			ObjectReplicatorRef->StopReplicating(this);		// Stop replicating on this channel
 		}
 		else
 		{
-			CompIt.Value()->CleanUp();
+
+			ObjectReplicatorRef->CleanUp();
 		}
 	}
 
@@ -2267,6 +2272,7 @@ void UActorChannel::MoveMappedObjectToUnmapped(const UObject* Object)
 void UActorChannel::DestroyActorAndComponents()
 {
 	// Destroy any sub-objects we created
+	PRAGMA_DISABLE_DEPRECATION_WARNINGS
 	for (UObject* SubObject : CreateSubObjects)
 	{
 		if (SubObject != nullptr)
@@ -2287,6 +2293,7 @@ void UActorChannel::DestroyActorAndComponents()
 	}
 
 	CreateSubObjects.Empty();
+	PRAGMA_ENABLE_DEPRECATION_WARNINGS
 
 	// Destroy the actor
 	if (Actor != nullptr)
@@ -3880,7 +3887,7 @@ void UActorChannel::Serialize(FArchive& Ar)
 		);
 
 		GRANULAR_NETWORK_MEMORY_TRACKING_TRACK("PendingGuidResolves", PendingGuidResolves.CountBytes(Ar));
-		GRANULAR_NETWORK_MEMORY_TRACKING_TRACK("CreateSubObjects", CreateSubObjects.CountBytes(Ar));
+		GRANULAR_NETWORK_MEMORY_TRACKING_TRACK("CreateSubObjects", GetCreatedSubObjects().CountBytes(Ar));
 		GRANULAR_NETWORK_MEMORY_TRACKING_TRACK("QueuedMustBeMappedGuidsInLastBunch", QueuedMustBeMappedGuidsInLastBunch.CountBytes(Ar));
 
 		GRANULAR_NETWORK_MEMORY_TRACKING_TRACK("QueuedExportBunches",
@@ -4264,7 +4271,7 @@ UObject* UActorChannel::ReadContentBlockHeader(FInBunch& Bunch, bool& bObjectDel
 			MoveMappedObjectToUnmapped(SubObj);
 
 			// Stop tracking this sub-object
-			CreateSubObjects.Remove(SubObj);
+			GetCreatedSubObjects().Remove(SubObj);
 
 			if (Connection != nullptr && Connection->Driver != nullptr)
 			{
@@ -4383,7 +4390,7 @@ UObject* UActorChannel::ReadContentBlockHeader(FInBunch& Bunch, bool& bObjectDel
 		Connection->Driver->GuidCache->RegisterNetGUID_Client( NetGUID, SubObj );
 
 		// Track which sub-object guids we are creating
-		CreateSubObjects.Add( SubObj );
+		GetCreatedSubObjects().Add( SubObj );
 
 		// Add this sub-object to the ImportedNetGuids list so we can possibly map this object if needed
 		if (ensureMsgf(NetGUID.IsValid(), TEXT("Channel tried to add an invalid GUID to the import list: %s"), *Describe()))
