Fix for "Moving Nav Modifier Volume does not cause the navmesh to update."

Previous fix [REDACTED_HANDLE] hadn't accounted properly for all the ways PostRegisterAllComponents() can get called.

diff --git a/Engine/Source/Runtime/Engine/Classes/GameFramework/Actor.h b/Engine/Source/Runtime/Engine/Classes/GameFramework/Actor.h
--- a/Engine/Source/Runtime/Engine/Classes/GameFramework/Actor.h
+++ b/Engine/Source/Runtime/Engine/Classes/GameFramework/Actor.h
@@ -459,7 +459,7 @@ private:
 	/** True if this actor is currently running user construction script (used to defer component registration) */
 	uint8 bRunningUserConstructionScript:1;
 
-	/** Set when RegisterAllComponents() is called and unset when UnregisterAllComponents() is called */
+	/** Set true just before PostRegisterAllComponents() is called and false just before PostUnregisterAllComponents() is called */
 	uint8 bHasRegisteredAllComponents:1;
 
 	/**
@@ -1835,9 +1835,15 @@ public:
 		return bActorIsBeingDestroyed;
 	}
 
-	/** Returns true if RegisterAllComponents() has been called without a subsequent call to UnregisterAllComponents() and false otherwise */
+	/** Returns bHasRegisteredAllComponents which indicates whether this actor has registered all their components without 
+	 *	unregisatering all of them them.
+	 *	bHasRegisteredAllComponents is set true just before PostRegisterAllComponents() is called and false just before PostUnregisterAllComponents() is called.
+	 */ 
 	bool HasActorRegisteredAllComponents() const { return bHasRegisteredAllComponents; }
 
+	/** Sets bHasRegisteredAllComponents true. bHasRegisteredAllComponents must be set true just prior to calling PostRegisterAllComponents(). */
+	void SetHasActorRegisteredAllComponents() { bHasRegisteredAllComponents = true; }
+
 	/** Event when this actor takes ANY damage */
 	UFUNCTION(BlueprintImplementableEvent, BlueprintAuthorityOnly, meta=(DisplayName = "AnyDamage"), Category="Game|Damage")
 	void ReceiveAnyDamage(float Damage, const class UDamageType* DamageType, class AController* InstigatedBy, AActor* DamageCauser);
@@ -2844,7 +2850,10 @@ public:
 	/** Called before all the components in the Components array are registered, called both in editor and during gameplay */
 	virtual void PreRegisterAllComponents();
 
-	/** Called after all the components in the Components array are registered, called both in editor and during gameplay */
+	/** 
+	 * Called after all the components in the Components array are registered, called both in editor and during gameplay.
+	 * bHasRegisteredAllComponents must be set true prior to calling this function.
+	 */
 	virtual void PostRegisterAllComponents();
 
 	/** Returns true if Actor has deferred the RegisterAllComponents() call at spawn time (e.g. pending Blueprint SCS execution to set up a scene root component). */
diff --git a/Engine/Source/Runtime/Engine/Private/Actor.cpp b/Engine/Source/Runtime/Engine/Private/Actor.cpp
--- a/Engine/Source/Runtime/Engine/Private/Actor.cpp
+++ b/Engine/Source/Runtime/Engine/Private/Actor.cpp
@@ -3476,6 +3476,8 @@ void AActor::PreRegisterAllComponents()
 
 void AActor::PostRegisterAllComponents() 
 {
+	ensureMsgf(bHasRegisteredAllComponents == true, TEXT("bHasRegisteredAllComponents must be set to true prior to calling PostRegisterAllComponents()"));
+
 	FNavigationSystem::OnActorRegistered(*this);
 }
 
@@ -5186,6 +5188,8 @@ void AActor::UnregisterAllComponents(const bool bForReregister)
 
 void AActor::PostUnregisterAllComponents()
 {
+	ensureMsgf(bHasRegisteredAllComponents == false, TEXT("bHasRegisteredAllComponents must be set to false prior to calling PostUnregisterAllComponents()"));
+
 	FNavigationSystem::OnActorUnregistered(*this);
 }
 
@@ -5201,7 +5205,6 @@ void AActor::RegisterAllComponents()
 
 	// Clear this flag as it's no longer deferred
 	bHasDeferredComponentRegistration = false;
-	bHasRegisteredAllComponents = true;
 }
 
 // Walks through components hierarchy and returns closest to root parent component that is unregistered
@@ -5311,6 +5314,7 @@ bool AActor::IncrementalRegisterComponents(int32 NumComponentsToRegister, FRegis
 #if PERF_TRACK_DETAILED_ASYNC_STATS
 		QUICK_SCOPE_CYCLE_COUNTER(STAT_AActor_IncrementalRegisterComponents_PostRegisterAllComponents);
 #endif
+		bHasRegisteredAllComponents = true;
 		// Finally, call PostRegisterAllComponents
 		PostRegisterAllComponents();
 		return true;
diff --git a/Engine/Source/Runtime/Engine/Private/ActorEditor.cpp b/Engine/Source/Runtime/Engine/Private/ActorEditor.cpp
--- a/Engine/Source/Runtime/Engine/Private/ActorEditor.cpp
+++ b/Engine/Source/Runtime/Engine/Private/ActorEditor.cpp
@@ -188,6 +188,7 @@ void AActor::PostEditChangeProperty(FPropertyChangedEvent& PropertyChangedEvent)
 			}
 			else
 			{
+				bHasRegisteredAllComponents = true;
 				PostRegisterAllComponents();
 			}
 		}
diff --git a/Engine/Source/Runtime/Engine/Private/World.cpp b/Engine/Source/Runtime/Engine/Private/World.cpp
--- a/Engine/Source/Runtime/Engine/Private/World.cpp
+++ b/Engine/Source/Runtime/Engine/Private/World.cpp
@@ -1539,6 +1539,7 @@ void UWorld::RepairChaosActors()
 					PersistentLevel->OwningWorld = this;
 				}
 
+				PersistentLevel->Actors[i]->SetHasActorRegisteredAllComponents();
 				PersistentLevel->Actors[i]->PostRegisterAllComponents();
 
 				if (bClearOwningWorld)
