Make HLOD subsystem more resilient to unexpected conditions

diff --git a/Engine/Source/Runtime/Engine/Private/WorldPartition/HLOD/HLODSubsystem.cpp b/Engine/Source/Runtime/Engine/Private/WorldPartition/HLOD/HLODSubsystem.cpp
--- a/Engine/Source/Runtime/Engine/Private/WorldPartition/HLOD/HLODSubsystem.cpp
+++ b/Engine/Source/Runtime/Engine/Private/WorldPartition/HLOD/HLODSubsystem.cpp
@@ -195,24 +195,26 @@ void UHLODSubsystem::RegisterHLODActor(AWorldPartitionHLOD* InWorldPartitionHLOD
 	const UWorldPartition* WorldPartition = FHLODSubsystem::GetWorldPartition(InWorldPartitionHLOD);
 	check(WorldPartition && WorldPartition->IsStreamingEnabled());
 
-	FWorldPartitionHLODRuntimeData& WorldPartitionHLODRuntimeData = WorldPartitionsHLODRuntimeData.FindChecked(WorldPartition);
-
-	const FName CellName = InWorldPartitionHLOD->GetSourceCellName();
-	FCellData* CellData = WorldPartitionHLODRuntimeData.CellsData.Find(CellName);
+	FWorldPartitionHLODRuntimeData* WorldPartitionHLODRuntimeData = WorldPartitionsHLODRuntimeData.Find(WorldPartition);
+	if (WorldPartitionHLODRuntimeData)
+	{
+		const FName CellName = InWorldPartitionHLOD->GetSourceCellName();
+		FCellData* CellData = WorldPartitionHLODRuntimeData->CellsData.Find(CellName);
 
 #if WITH_EDITOR
-	UE_LOG(LogHLODSubsystem, Verbose, TEXT("Registering HLOD %s (%s) for cell %s"), *InWorldPartitionHLOD->GetActorLabel(), *InWorldPartitionHLOD->GetActorGuid().ToString(), *CellName.ToString());
+		UE_LOG(LogHLODSubsystem, Verbose, TEXT("Registering HLOD %s (%s) for cell %s"), *InWorldPartitionHLOD->GetActorLabel(), *InWorldPartitionHLOD->GetActorGuid().ToString(), *CellName.ToString());
 #endif
 
-	if (CellData)
-	{
-		CellData->LoadedHLODs.Add(InWorldPartitionHLOD);
-		InWorldPartitionHLOD->SetVisibility(UHLODSubsystem::WorldPartitionHLODEnabled && !CellData->bIsCellVisible);
-	}
-	else
-	{
-		UE_LOG(LogHLODSubsystem, Verbose, TEXT("Found HLOD referencing nonexistent cell '%s'"), *CellName.ToString());
-		InWorldPartitionHLOD->SetVisibility(false);
+		if (CellData)
+		{
+			CellData->LoadedHLODs.Add(InWorldPartitionHLOD);
+			InWorldPartitionHLOD->SetVisibility(UHLODSubsystem::WorldPartitionHLODEnabled && !CellData->bIsCellVisible);
+		}
+		else
+		{
+			UE_LOG(LogHLODSubsystem, Verbose, TEXT("Found HLOD referencing nonexistent cell '%s'"), *CellName.ToString());
+			InWorldPartitionHLOD->SetVisibility(false);
+		}
 	}
 }
 
@@ -223,19 +225,20 @@ void UHLODSubsystem::UnregisterHLODActor(AWorldPartitionHLOD* InWorldPartitionHL
 	const UWorldPartition* WorldPartition = FHLODSubsystem::GetWorldPartition(InWorldPartitionHLOD);
 	check(WorldPartition && WorldPartition->IsStreamingEnabled());
 
-	FWorldPartitionHLODRuntimeData& WorldPartitionHLODRuntimeData = WorldPartitionsHLODRuntimeData.FindChecked(WorldPartition);
-
-	const FName CellName = InWorldPartitionHLOD->GetSourceCellName();
-	FCellData* CellData = WorldPartitionHLODRuntimeData.CellsData.Find(CellName);
-
+	FWorldPartitionHLODRuntimeData* WorldPartitionHLODRuntimeData = WorldPartitionsHLODRuntimeData.Find(WorldPartition);
+	if (WorldPartitionHLODRuntimeData)
+	{
+		const FName CellName = InWorldPartitionHLOD->GetSourceCellName();
+		FCellData* CellData = WorldPartitionHLODRuntimeData->CellsData.Find(CellName);
+		if (CellData)
+		{
 #if WITH_EDITOR
-	UE_LOG(LogHLODSubsystem, Verbose, TEXT("Unregistering HLOD %s (%s) for cell %s"), *InWorldPartitionHLOD->GetActorLabel(), *InWorldPartitionHLOD->GetActorGuid().ToString(), *CellName.ToString());
+			UE_LOG(LogHLODSubsystem, Verbose, TEXT("Unregistering HLOD %s (%s) for cell %s"), *InWorldPartitionHLOD->GetActorLabel(), *InWorldPartitionHLOD->GetActorGuid().ToString(), *CellName.ToString());
 #endif
 
-	if (CellData)
-	{
-		int32 NumRemoved = CellData->LoadedHLODs.Remove(InWorldPartitionHLOD);
-		check(NumRemoved == 1);
+			int32 NumRemoved = CellData->LoadedHLODs.Remove(InWorldPartitionHLOD);
+			check(NumRemoved == 1);
+		}
 	}
 }
 
@@ -244,20 +247,26 @@ void UHLODSubsystem::OnCellShown(const UWorldPartitionRuntimeCell* InCell)
 	const UWorldPartition* WorldPartition = InCell->GetCellOwner()->GetOuterWorld()->GetWorldPartition();
 	check(WorldPartition && WorldPartition->IsStreamingEnabled());
 
-	FWorldPartitionHLODRuntimeData& WorldPartitionHLODRuntimeData = WorldPartitionsHLODRuntimeData.FindChecked(WorldPartition);
-	FCellData& CellData = WorldPartitionHLODRuntimeData.CellsData.FindChecked(InCell->GetFName());
-	CellData.bIsCellVisible = true;
+	FWorldPartitionHLODRuntimeData* WorldPartitionHLODRuntimeData = WorldPartitionsHLODRuntimeData.Find(WorldPartition);
+	if (WorldPartitionHLODRuntimeData)
+	{
+		FCellData* CellData = WorldPartitionHLODRuntimeData->CellsData.Find(InCell->GetFName());
+		if (CellData)
+		{
+			CellData->bIsCellVisible = true;
 
 #if WITH_EDITOR
-	UE_LOG(LogHLODSubsystem, Verbose, TEXT("Cell shown - %s - hiding %d HLOD actors"), *InCell->GetName(), CellData.LoadedHLODs.Num());
+			UE_LOG(LogHLODSubsystem, Verbose, TEXT("Cell shown - %s - hiding %d HLOD actors"), *InCell->GetName(), CellData->LoadedHLODs.Num());
 #endif
 
-	for (AWorldPartitionHLOD* HLODActor : CellData.LoadedHLODs)
-	{
+			for (AWorldPartitionHLOD* HLODActor : CellData->LoadedHLODs)
+			{
 #if WITH_EDITOR
-		UE_LOG(LogHLODSubsystem, Verbose, TEXT("\t\t%s - %s"), *HLODActor->GetActorLabel(), *HLODActor->GetActorGuid().ToString());
+				UE_LOG(LogHLODSubsystem, Verbose, TEXT("\t\t%s - %s"), *HLODActor->GetActorLabel(), *HLODActor->GetActorGuid().ToString());
 #endif
-		HLODActor->SetVisibility(false);
+				HLODActor->SetVisibility(false);
+			}
+		}
 	}
 }
 
@@ -266,20 +275,26 @@ void UHLODSubsystem::OnCellHidden(const UWorldPartitionRuntimeCell* InCell)
 	const UWorldPartition* WorldPartition = InCell->GetCellOwner()->GetOuterWorld()->GetWorldPartition();
 	check(WorldPartition && WorldPartition->IsStreamingEnabled());
 
-	FWorldPartitionHLODRuntimeData& WorldPartitionHLODRuntimeData = WorldPartitionsHLODRuntimeData.FindChecked(WorldPartition);
-	FCellData& CellData = WorldPartitionHLODRuntimeData.CellsData.FindChecked(InCell->GetFName());
-	CellData.bIsCellVisible = false;
+	FWorldPartitionHLODRuntimeData* WorldPartitionHLODRuntimeData = WorldPartitionsHLODRuntimeData.Find(WorldPartition);
+	if (WorldPartitionHLODRuntimeData)
+	{
+		FCellData* CellData = WorldPartitionHLODRuntimeData->CellsData.Find(InCell->GetFName());
+		if (CellData)
+		{
+			CellData->bIsCellVisible = false;
 
 #if WITH_EDITOR
-	UE_LOG(LogHLODSubsystem, Verbose, TEXT("Cell hidden - %s - showing %d HLOD actors"), *InCell->GetName(), CellData.LoadedHLODs.Num());
+			UE_LOG(LogHLODSubsystem, Verbose, TEXT("Cell hidden - %s - showing %d HLOD actors"), *InCell->GetName(), CellData->LoadedHLODs.Num());
 #endif
 
-	for (AWorldPartitionHLOD* HLODActor : CellData.LoadedHLODs)
-	{
+			for (AWorldPartitionHLOD* HLODActor : CellData->LoadedHLODs)
+			{
 #if WITH_EDITOR
-		UE_LOG(LogHLODSubsystem, Verbose, TEXT("\t\t%s - %s"), *HLODActor->GetActorLabel(), *HLODActor->GetActorGuid().ToString());
+				UE_LOG(LogHLODSubsystem, Verbose, TEXT("\t\t%s - %s"), *HLODActor->GetActorLabel(), *HLODActor->GetActorGuid().ToString());
 #endif
-		HLODActor->SetVisibility(UHLODSubsystem::WorldPartitionHLODEnabled);
+				HLODActor->SetVisibility(UHLODSubsystem::WorldPartitionHLODEnabled);
+			}
+		}
 	}
 }
 
@@ -407,19 +422,27 @@ bool UHLODSubsystem::RequestUnloading(const UWorldPartitionRuntimeCell* InCell)
 	}
 
 	const UWorldPartition* WorldPartition = InCell->GetCellOwner()->GetOuterWorld()->GetWorldPartition();;
-	FWorldPartitionHLODRuntimeData& WorldPartitionHLODRuntimeData = WorldPartitionsHLODRuntimeData.FindChecked(WorldPartition);
+	FWorldPartitionHLODRuntimeData* WorldPartitionHLODRuntimeData = WorldPartitionsHLODRuntimeData.Find(WorldPartition);
+	if (!WorldPartitionHLODRuntimeData)
+	{
+		return true;
+	}
 
-	FCellData& CellData = WorldPartitionHLODRuntimeData.CellsData.FindChecked(InCell->GetFName());
+	FCellData* CellData = WorldPartitionHLODRuntimeData->CellsData.Find(InCell->GetFName());
+	if (!CellData)
+	{
+		return true;
+	}
 
 	// If cell wasn't even visible yet or doesn't have HLOD actors, skip warmup.
-	bPerformWarmpup = !CellData.LoadedHLODs.IsEmpty() && CellData.bIsCellVisible;
+	bPerformWarmpup = !CellData->LoadedHLODs.IsEmpty() && CellData->bIsCellVisible;
 	if (!bPerformWarmpup)
 	{
 		return true;
 	}
 
 	// Verify that at least one HLOD actor associated with this cell needs warmup.
-	bPerformWarmpup = Algo::AnyOf(CellData.LoadedHLODs, [](const AWorldPartitionHLOD* HLODActor) { return HLODActor->DoesRequireWarmup(); });
+	bPerformWarmpup = Algo::AnyOf(CellData->LoadedHLODs, [](const AWorldPartitionHLOD* HLODActor) { return HLODActor->DoesRequireWarmup(); });
 	if (!bPerformWarmpup)
 	{
 		return true;
@@ -431,20 +454,20 @@ bool UHLODSubsystem::RequestUnloading(const UWorldPartitionRuntimeCell* InCell)
 	uint32 CurrentFrameNumber = GetWorld()->Scene->GetFrameNumber();
 
 	// Trigger warmup on the first request to unload, or if a warmup request expired
-	if (CellData.WarmupEndFrame == INDEX_NONE || CurrentFrameNumber > (CellData.WarmupEndFrame + WarmupExpiredFrames))
+	if (CellData->WarmupEndFrame == INDEX_NONE || CurrentFrameNumber > (CellData->WarmupEndFrame + WarmupExpiredFrames))
 	{
 		// Warmup will be triggered in the next BeginRenderView() call, at which point the frame number will have been incremented.
-		CellData.WarmupStartFrame = CurrentFrameNumber + 1; 
-		CellData.WarmupEndFrame = CellData.WarmupStartFrame + CVarHLODWarmupNumFrames.GetValueOnGameThread();
-		WorldPartitionHLODRuntimeData.CellsToWarmup.Add(&CellData);
+		CellData->WarmupStartFrame = CurrentFrameNumber + 1; 
+		CellData->WarmupEndFrame = CellData->WarmupStartFrame + CVarHLODWarmupNumFrames.GetValueOnGameThread();
+		WorldPartitionHLODRuntimeData->CellsToWarmup.Add(CellData);
 	}
 
 	// Test if warmup is completed
-	bool bCanUnload = CurrentFrameNumber >= CellData.WarmupEndFrame;
+	bool bCanUnload = CurrentFrameNumber >= CellData->WarmupEndFrame;
 	if (bCanUnload)
 	{
-		CellData.WarmupStartFrame = INDEX_NONE;
-		CellData.WarmupEndFrame = INDEX_NONE;
+		CellData->WarmupStartFrame = INDEX_NONE;
+		CellData->WarmupEndFrame = INDEX_NONE;
 	}
 
 	return bCanUnload;
