fix bCookedIsStreamable being set and queried in same loop
caused inconsistent behavior in multi-platform cooks
replace with DidSerializeStreamingMipsForPlatform
take it out of IsCandidate and instead read it directly in PrimitiveComponent
set it from if you actually serialize streaming mips, remove GetStreamableRenderResourceState which was not correct due to not using targetplatform LODBias

#b richard.malo,dan.thompson

diff --git a/Engine/Source/Runtime/Engine/Classes/Engine/Texture.h b/Engine/Source/Runtime/Engine/Classes/Engine/Texture.h
--- a/Engine/Source/Runtime/Engine/Classes/Engine/Texture.h
+++ b/Engine/Source/Runtime/Engine/Classes/Engine/Texture.h
@@ -1385,13 +1385,15 @@ protected:
 	UPROPERTY(EditAnywhere, AdvancedDisplay, Instanced, Category = Texture)
 	TArray<TObjectPtr<UAssetUserData>> AssetUserData;
 
-private:
-
+public:
 #if WITH_EDITOR
-	/** Used to mark texture streamable state when cooking. */
-	TOptional<bool> bCookedIsStreamable;
+	/** Used to record texture streamable state when cooking.
+	This a per-platform bool to record whether the PlatformData being cooked made streamable mips.
+	Key is TargetPlatform->PlatformName */
+	TMap<FString,bool> DidSerializeStreamingMipsForPlatform;
 #endif
 
+private:
 	/** The texture's resource, can be NULL */
 	class FTextureResource*	PrivateResource;
 	/** Value updated and returned by the render-thread to allow
@@ -1811,7 +1813,7 @@ public:
 
 #if WITH_EDITOR
 	/** Called by ULevel::MarkNoStreamableTexturesPrimitiveComponents when cooking level.
-	Checks bCookedIsStreamable.  Return false for VT. */
+	* Return false for VT. */
 	bool IsCandidateForTextureStreamingOnPlatformDuringCook(const ITargetPlatform* InTargetPlatform) const;
 
 	/** Get the largest allowed dimension of non-VT texture
@@ -1842,7 +1844,6 @@ protected:
 	/** Notify any loaded material instances that the texture has changed. */
 	ENGINE_API void NotifyMaterials(const ENotifyMaterialsEffectOnShaders EffectOnShaders = ENotifyMaterialsEffectOnShaders::Default);
 
-	virtual bool GetStreamableRenderResourceState(FTexturePlatformData* InPlatformData, FStreamableRenderResourceState& OutState) const { return false; }
 #endif //WITH_EDITOR
 
 	void BeginFinalReleaseResource();
diff --git a/Engine/Source/Runtime/Engine/Classes/Engine/Texture2D.h b/Engine/Source/Runtime/Engine/Classes/Engine/Texture2D.h
--- a/Engine/Source/Runtime/Engine/Classes/Engine/Texture2D.h
+++ b/Engine/Source/Runtime/Engine/Classes/Engine/Texture2D.h
@@ -317,9 +317,4 @@ public:
 	 */
 	virtual bool IsVirtualTexturedWithSinglePhysicalSpace() const { return false;  }
 
-protected:
-
-#if WITH_EDITOR
-	virtual bool GetStreamableRenderResourceState(FTexturePlatformData* InPlatformData, FStreamableRenderResourceState& OutState) const override;
-#endif
 };
diff --git a/Engine/Source/Runtime/Engine/Classes/Engine/Texture2DArray.cpp b/Engine/Source/Runtime/Engine/Classes/Engine/Texture2DArray.cpp
--- a/Engine/Source/Runtime/Engine/Classes/Engine/Texture2DArray.cpp
+++ b/Engine/Source/Runtime/Engine/Classes/Engine/Texture2DArray.cpp
@@ -144,23 +144,6 @@ FTexturePlatformData* UTexture2DArray::GetPlatformData()
 	return const_cast<FTexturePlatformData*>(ConstThis->GetPlatformData());
 }
 
-#if WITH_EDITOR
-bool UTexture2DArray::GetStreamableRenderResourceState(FTexturePlatformData* InPlatformData, FStreamableRenderResourceState& OutState) const
-{
-	TGuardValue<FTexturePlatformData*> Guard(const_cast<UTexture2DArray*>(this)->PrivatePlatformData, InPlatformData);
-	if (GetPlatformData())
-	{
-		const FPixelFormatInfo& FormatInfo = GPixelFormats[GetPixelFormat()];
-		if (GetNumMips() > 0 && FormatInfo.Supported)
-		{
-			OutState = GetResourcePostInitState(GetPlatformData(), GSupportsTexture2DArrayStreaming, 0, 0, /*bSkipCanBeLoaded*/ true);
-			return true;
-		}
-	}
-	return false;
-}
-#endif
-
 FTextureResource* UTexture2DArray::CreateResource()
 {
 #if WITH_EDITOR
diff --git a/Engine/Source/Runtime/Engine/Classes/Engine/Texture2DArray.h b/Engine/Source/Runtime/Engine/Classes/Engine/Texture2DArray.h
--- a/Engine/Source/Runtime/Engine/Classes/Engine/Texture2DArray.h
+++ b/Engine/Source/Runtime/Engine/Classes/Engine/Texture2DArray.h
@@ -137,9 +137,4 @@ public:
 	virtual bool StreamIn(int32 NewMipCount, bool bHighPrio) final override;
 	//~ End UStreamableRenderAsset Interface
 
-protected:
-
-#if WITH_EDITOR
-	virtual bool GetStreamableRenderResourceState(FTexturePlatformData* InPlatformData, FStreamableRenderResourceState& OutState) const override;
-#endif
 };
\ No newline at end of file
diff --git a/Engine/Source/Runtime/Engine/Classes/Engine/VolumeTexture.h b/Engine/Source/Runtime/Engine/Classes/Engine/VolumeTexture.h
--- a/Engine/Source/Runtime/Engine/Classes/Engine/VolumeTexture.h
+++ b/Engine/Source/Runtime/Engine/Classes/Engine/VolumeTexture.h
@@ -143,7 +143,6 @@ protected:
 
 #if WITH_EDITOR
 	void UpdateMipGenSettings();
-	virtual bool GetStreamableRenderResourceState(FTexturePlatformData* InPlatformData, FStreamableRenderResourceState& OutState) const override;
 #endif
 };
 
diff --git a/Engine/Source/Runtime/Engine/Private/Components/PrimitiveComponent.cpp b/Engine/Source/Runtime/Engine/Private/Components/PrimitiveComponent.cpp
--- a/Engine/Source/Runtime/Engine/Private/Components/PrimitiveComponent.cpp
+++ b/Engine/Source/Runtime/Engine/Private/Components/PrimitiveComponent.cpp
@@ -43,6 +43,7 @@
 #if WITH_EDITOR
 #include "Engine/LODActor.h"
 #include "Rendering/StaticLightingSystemInterface.h"
+#include "Interfaces/ITargetPlatform.h"
 #endif // WITH_EDITOR
 
 #if DO_CHECK
@@ -1427,8 +1428,18 @@ void UPrimitiveComponent::PreSave(FObjectPreSaveContext ObjectSaveContext)
 		{
 			if (Texture && Texture->IsCandidateForTextureStreamingOnPlatformDuringCook(ObjectSaveContext.GetTargetPlatform()))
 			{
-				bHasStreamableTextures = true;
-				break;
+				// furthermore refine based on what actually happened with that texture in *this* cook session
+				//	if that texture was serialized out before us
+				// can be false if texture IsCandidate but did not actually make streaming textures
+				// if texture has not serialized yet and set DidSerializeStreamingMipsForPlatform, we must assume it would be true
+				
+				const FString PlatformName = ObjectSaveContext.GetTargetPlatform()->PlatformName();
+				const bool * pBoolDidStream = Texture->DidSerializeStreamingMipsForPlatform.Find(PlatformName);
+				if ( pBoolDidStream == nullptr || *pBoolDidStream == true )
+				{
+					bHasStreamableTextures = true;
+					break;
+				}
 			}
 		}
 
diff --git a/Engine/Source/Runtime/Engine/Private/Texture.cpp b/Engine/Source/Runtime/Engine/Private/Texture.cpp
--- a/Engine/Source/Runtime/Engine/Private/Texture.cpp
+++ b/Engine/Source/Runtime/Engine/Private/Texture.cpp
@@ -1375,9 +1375,7 @@ bool UTexture::IsPossibleToStream() const
 }
 
 #if WITH_EDITOR
-// Based on target platform, returns wether texture is candidate to be streamed.
-// This method is used to decide if PrimitiveComponent's bHasNoStreamableTextures flag can be set to true.
-// See ULevel::MarkNoStreamableTexturesPrimitiveComponents for details.
+// Based on target platform, returns whether texture is candidate to be streamed.
 bool UTexture::IsCandidateForTextureStreamingOnPlatformDuringCook(const ITargetPlatform* InTargetPlatform) const
 {
 	const bool bIsVirtualTextureStreaming = InTargetPlatform->SupportsFeature(ETargetPlatformFeatures::VirtualTextureStreaming) ? VirtualTextureStreaming : false;
@@ -1386,15 +1384,6 @@ bool UTexture::IsCandidateForTextureStreamingOnPlatformDuringCook(const ITargetP
 	if (bIsCandidateForTextureStreaming &&
 		IsPossibleToStream())
 	{
-		// If bCookedIsStreamable flag was previously computed, use it.
-		// This is computed _after_ the derived data is cached. It's only Reset() and used in the SerializeCookedPlatformData
-		// path.. but not the BeginCacheForCookedPlatformData path. Need to know what happens if you do
-		// one, then the other, as the flag will be set. Also note that this is an OR when multiple platforms
-		// are cooked... so this can be true for a non-streaming platform if it gets cooked with a streaming platform!
-		if (bCookedIsStreamable.IsSet())
-		{
-			return *bCookedIsStreamable;
-		}
 		return true;
 	}
 	return false;
diff --git a/Engine/Source/Runtime/Engine/Private/Texture2D.cpp b/Engine/Source/Runtime/Engine/Private/Texture2D.cpp
--- a/Engine/Source/Runtime/Engine/Private/Texture2D.cpp
+++ b/Engine/Source/Runtime/Engine/Private/Texture2D.cpp
@@ -913,39 +913,6 @@ bool UTexture2D::HasAlphaChannel() const
 	return false;
 }
 
-#if WITH_EDITOR
-bool UTexture2D::GetStreamableRenderResourceState(FTexturePlatformData* InPlatformData, FStreamableRenderResourceState& OutState) const
-{
-	// we want to make queries about InPlatformData here, not the current platform data
-	// so for the scope of this function, change PrivatePlatformData to InPlatformData
-	TGuardValue<FTexturePlatformData*> Guard(const_cast<UTexture2D*>(this)->PrivatePlatformData, InPlatformData);
-
-	check( GetPlatformData() == InPlatformData );
-
-	if (GetPlatformData())
-	{
-		if (IsCurrentlyVirtualTextured())
-		{
-			return false;
-		}
-
-		// we're relying on these calls like GetPixelFormat() to actually be queries on InPlatformData :
-		const EPixelFormat PixelFormat = GetPixelFormat();
-		const int32 NumMips = FMath::Min3<int32>(GetPlatformData()->Mips.Num(), GMaxTextureMipCount, FStreamableRenderResourceState::MAX_LOD_COUNT);
-		if (NumMips && GPixelFormats[PixelFormat].Supported &&
-			(NumMips > 1 || FMath::Max(GetSizeX(), GetSizeY()) <= (int32)GetMax2DTextureDimension()))
-		{
-			// this wants to make a query on InPlatformData, for some other target platform
-			// @todo Oodle : this is a bit broken, will use the current DeviceProfile in CalculateLODBias
-			//		it should be getting LODGroup settings from the TargetPlatform but is not
-			OutState = GetResourcePostInitState(GetPlatformData(), true, 0, NumMips, /*bSkipCanBeLoaded*/ true);
-			return true;
-		}
-	}
-	return false;
-}
-#endif
-
 FTextureResource* UTexture2D::CreateResource()
 {
 	TRACE_CPUPROFILER_EVENT_SCOPE(UTexture2D::CreateResource)
diff --git a/Engine/Source/Runtime/Engine/Private/TextureDerivedData.cpp b/Engine/Source/Runtime/Engine/Private/TextureDerivedData.cpp
--- a/Engine/Source/Runtime/Engine/Private/TextureDerivedData.cpp
+++ b/Engine/Source/Runtime/Engine/Private/TextureDerivedData.cpp
@@ -2512,6 +2512,8 @@ static void SerializePlatformData(
 			if (bStreamable)
 #endif
 			{
+				check( Texture->IsPossibleToStream() );
+
 				int32 NumNonStreamingMips = PlatformData->GetNumNonStreamingMips(true);
 				// NumMips has been reduced by FirstMipToSerialize (LODBias)
 				NumNonStreamingMips = FMath::Min(NumNonStreamingMips,NumMips);
@@ -2531,6 +2533,16 @@ static void SerializePlatformData(
 				// Optional mips must not overlap the non-streaming mips : (MinMipToInline ensures this)
 				check( OptionalMips + NumNonStreamingMips <= NumMips );
 #endif
+			
+#if WITH_EDITOR
+				// if we make any streaming mips, record it on the owner :
+				if ( NumNonStreamingMips < NumMips )
+				{
+					// was added previously and set to false, so FindChecked here :
+					const FString PlatformName = Ar.CookingTarget()->PlatformName();
+					Texture->DidSerializeStreamingMipsForPlatform.FindChecked( PlatformName ) = true;
+				}
+#endif
 			}
 
 			for (int32 MipIndex = 0; MipIndex < NumMips && MipIndex < OptionalMips; ++MipIndex) //-V654
@@ -2681,7 +2693,6 @@ UE::DerivedData::FValueId FTexturePlatformData::MakeMipId(int32 MipIndex)
 void FTexturePlatformData::SerializeCooked(FArchive& Ar, UTexture* Owner, bool bStreamable)
 {
 	EPlatformDataSerializationFlags Flags = EPlatformDataSerializationFlags::Cooked;
-	// @@CB bStreamable may be set to true because of bCookedIsStreamable even if I am NOT actually a streamable texture on this platform
 	if (bStreamable)
 	{
 		Flags |= EPlatformDataSerializationFlags::Streamable;
@@ -3410,7 +3421,6 @@ void UTexture::SerializeCookedPlatformData(FArchive& Ar)
 #if WITH_EDITOR
 	if (Ar.IsCooking() && Ar.IsPersistent())
 	{
-		bCookedIsStreamable.Reset();
 		if (Ar.CookingTarget()->AllowAudioVisualData())
 		{
 			TArray<FTexturePlatformData*> PlatformDataToSerialize;
@@ -3486,18 +3496,15 @@ void UTexture::SerializeCookedPlatformData(FArchive& Ar)
 				}
 			}
 
+			// set DidSerializeStreamingMipsForPlatform to false, then it will change to true if any SerializeCooked makes streaming mips
+			const FString PlatformName = Ar.CookingTarget()->PlatformName();
+			DidSerializeStreamingMipsForPlatform.Add( PlatformName, false );
+
+			// this iteration is over NumLayers :
 			for (FTexturePlatformData* PlatformDataToSave : PlatformDataToSerialize)
 			{
 				PlatformDataToSave->FinishCache();
 
-				// Update bCookedIsStreamable for later use in IsCandidateForTextureStreaming
-				// @@CB bCookedIsStreamable is set to true if the texture can be streamed on ANY platform
-				FStreamableRenderResourceState State;
-				if (GetStreamableRenderResourceState(PlatformDataToSave, State))
-				{
-					bCookedIsStreamable = !bCookedIsStreamable.IsSet() ? State.bSupportsStreaming : (*bCookedIsStreamable || State.bSupportsStreaming);
-				}
-
 				FName PixelFormatName = PixelFormatEnum->GetNameByValue(PlatformDataToSave->PixelFormat);
 				Ar << PixelFormatName;
 
@@ -3508,9 +3515,6 @@ void UTexture::SerializeCookedPlatformData(FArchive& Ar)
 				}
 
 				// Pass streamable flag for inlining mips
-				// @@CB GetTextureIsStreamableOnPlatform checks bCookedIsStreamable which changes while we go through this loop
-				//		that makes behavior dependent on order of platforms
-				// 
 				bool bTextureIsStreamable = GetTextureIsStreamableOnPlatform(*this, *Ar.CookingTarget());
 				PlatformDataToSave->SerializeCooked(Ar, this, bTextureIsStreamable);
 
diff --git a/Engine/Source/Runtime/Engine/Private/VolumeTexture.cpp b/Engine/Source/Runtime/Engine/Private/VolumeTexture.cpp
--- a/Engine/Source/Runtime/Engine/Private/VolumeTexture.cpp
+++ b/Engine/Source/Runtime/Engine/Private/VolumeTexture.cpp
@@ -501,24 +501,6 @@ EPixelFormat UVolumeTexture::GetPixelFormat() const
 	return PF_Unknown;
 }
 
-#if WITH_EDITOR
-bool UVolumeTexture::GetStreamableRenderResourceState(FTexturePlatformData* InPlatformData, FStreamableRenderResourceState& OutState) const
-{
-	TGuardValue<FTexturePlatformData*> Guard(const_cast<UVolumeTexture*>(this)->PrivatePlatformData, InPlatformData);
-	if (GetPlatformData())
-	{
-		const FPixelFormatInfo& FormatInfo = GPixelFormats[GetPixelFormat()];
-		const bool bFormatIsSupported = FormatInfo.Supported;
-		if (GetNumMips() > 0 && GSupportsTexture3D && bFormatIsSupported)
-		{
-			OutState = GetResourcePostInitState(GetPlatformData(), GSupportsVolumeTextureStreaming, 0, 0, /*bSkipCanBeLoaded*/ true);
-			return true;
-		}
-	}
-	return false;
-}
-#endif
-
 FTextureResource* UVolumeTexture::CreateResource()
 {
 #if WITH_EDITOR
