Fixes for static analysis issues in ue5 main.

diff --git a/Engine/Source/Editor/SparseVolumeTexture/Private/SparseVolumeTextureFactory.cpp b/Engine/Source/Editor/SparseVolumeTexture/Private/SparseVolumeTextureFactory.cpp
--- a/Engine/Source/Editor/SparseVolumeTexture/Private/SparseVolumeTextureFactory.cpp
+++ b/Engine/Source/Editor/SparseVolumeTexture/Private/SparseVolumeTextureFactory.cpp
@@ -279,7 +279,7 @@ static bool ConvertOpenVDBToSparseVolumeTexture(
 		// A page encodes the physical tile coord as unsigned int of 11 11 10 bits
 		// This means a page coord cannot be larger than 2047 for x and y and 1023 for z
 		// which mean we cannot have more than 2048*2048*1024 = 4 Giga tiles of 16^3 tiles.
-		uint32 DestinationTileCoord32bit = DestinationTileCoord.X & 0x7FF | ((DestinationTileCoord.Y & 0x7FF) << 11) | ((DestinationTileCoord.Z & 0x3FF) << 22);
+		uint32 DestinationTileCoord32bit = (DestinationTileCoord.X & 0x7FF) | ((DestinationTileCoord.Y & 0x7FF) << 11) | ((DestinationTileCoord.Z & 0x3FF) << 22);
 
 		// Setup the page table entry
 		OutSparseVolumeRawSource.DensityPage.GetData()
diff --git a/Engine/Source/Editor/SparseVolumeTexture/Private/SparseVolumeTextureOpenVDB.h b/Engine/Source/Editor/SparseVolumeTexture/Private/SparseVolumeTextureOpenVDB.h
--- a/Engine/Source/Editor/SparseVolumeTexture/Private/SparseVolumeTextureOpenVDB.h
+++ b/Engine/Source/Editor/SparseVolumeTexture/Private/SparseVolumeTextureOpenVDB.h
@@ -35,7 +35,7 @@ UE_PUSH_MACRO("check")
 #include <openvdb/tools/Prune.h>
 #include <openvdb/tools/Dense.h>
 #include <openvdb/math/Vec4.h>
-#include <openvdb//math/Half.h>
+#include <openvdb/math/Half.h>
 
 #if LOCAL_M_PI
 #undef M_PI
diff --git a/Engine/Source/Runtime/Engine/Private/Materials/HLSLMaterialTranslator.cpp b/Engine/Source/Runtime/Engine/Private/Materials/HLSLMaterialTranslator.cpp
--- a/Engine/Source/Runtime/Engine/Private/Materials/HLSLMaterialTranslator.cpp
+++ b/Engine/Source/Runtime/Engine/Private/Materials/HLSLMaterialTranslator.cpp
@@ -12056,34 +12056,20 @@ int32 FHLSLMaterialTranslator::SparseVolumeTextureSample(int32 TextureIndex, int
 		int32 VirtualTextureIndex = UniformTextureExpressions[(uint32)EMaterialTextureParameterType::Volume].Find(TextureUniformExpression);
 		check(UniformTextureExpressions[(uint32)EMaterialTextureParameterType::Volume].IsValidIndex(VirtualTextureIndex));
 
-		const EDerivativeStatus UvDerivativeStatus = GetDerivativeStatus(CoordinateIndex);
-
 		check(TextureUniformExpression->GetTextureLayerIndex() != INDEX_NONE);
 		if (TextureUniformExpression->GetTextureLayerIndex() == 0)
 		{
 			FString SampleCode = FString::Printf(TEXT("%s.Load(int4(int3(%s), 0)).x"), *TextureName, *GetParameterCode(CoordinateIndex));
 			AddEstimatedTextureSample();
-			if (IsAnalyticDerivEnabled() && IsDerivativeValid(UvDerivativeStatus))
-			{
-				SamplingCodeIndex = AddCodeChunk(MCT_UInt1, *SampleCode);
-			}
-			else
-			{
-				SamplingCodeIndex = AddCodeChunk(MCT_UInt1, *SampleCode);
-			}
+			// We ignore derivative for now, see IsAnalyticDerivEnabled() and IsDerivativeValid(UvDerivativeStatus)
+			SamplingCodeIndex = AddCodeChunk(MCT_UInt1, *SampleCode);
 		}
 		else if (TextureUniformExpression->GetTextureLayerIndex() == 1)
 		{
 			FString SampleCode = FString::Printf(TEXT("%s.Load(int4(int3(%s), 0))"), *TextureName, *GetParameterCode(CoordinateIndex));
 			AddEstimatedTextureSample();
-			if (IsAnalyticDerivEnabled() && IsDerivativeValid(UvDerivativeStatus))
-			{
-				SamplingCodeIndex = AddCodeChunk(MCT_Float4, *SampleCode);
-			}
-			else
-			{
-				SamplingCodeIndex = AddCodeChunk(MCT_Float4, *SampleCode);
-			}
+			// We ignore derivative for now, see IsAnalyticDerivEnabled() and IsDerivativeValid(UvDerivativeStatus)
+			SamplingCodeIndex = AddCodeChunk(MCT_Float4, *SampleCode);
 		}
 		else
 		{
