Change FShaderJobCache eviction scheme to FIFO

* This improves trimming performance when memory budget is reached

diff --git a/Engine/Source/Runtime/Engine/Private/ShaderCompiler/ShaderCompiler.cpp b/Engine/Source/Runtime/Engine/Private/ShaderCompiler/ShaderCompiler.cpp
--- a/Engine/Source/Runtime/Engine/Private/ShaderCompiler/ShaderCompiler.cpp
+++ b/Engine/Source/Runtime/Engine/Private/ShaderCompiler/ShaderCompiler.cpp
@@ -8376,6 +8376,51 @@ namespace
 }
 #endif
 
+struct FShaderJobCacheStoredOutput
+{
+private:
+	/** How many times this output is referenced by the cached jobs */
+	int32 NumReferences = 0;
+
+public:
+
+	/** How many times this output has been returned as a cached result, no matter the input hash */
+	int32 NumHits = 0;
+
+	/** Canned output */
+	TArray<uint8> JobOutput;
+
+	/** Similar to FRefCountBase AddRef, but not atomic */
+	int32 AddRef()
+	{
+		++NumReferences;
+
+		return NumReferences;
+	}
+
+	/** Similar to FRefCountBase Release, but not atomic */
+	int32 Release()
+	{
+		checkf(NumReferences >= 0, TEXT("Attempting to release shader job cache output that was already released"));
+
+		--NumReferences;
+
+		const int32 RemainingNumReferences = NumReferences;
+
+		if (RemainingNumReferences == 0)
+		{
+			delete this;
+		}
+
+		return RemainingNumReferences;
+	}
+
+	uint64 GetAllocatedSize() const
+	{
+		return static_cast<uint64>(JobOutput.GetAllocatedSize() + sizeof(*this));
+	}
+};
+
 FShaderJobCache::FJobCachedOutput* FShaderJobCache::Find(const FJobInputHash& Hash, const bool bCheckDDC)
 {
 	++TotalSearchAttempts;
@@ -8446,12 +8491,8 @@ FShaderJobCache::FJobCachedOutput* FShaderJobCache::Find(const FJobInputHash& Ha
 				RequestOwner.Wait();
 				if (!Results.IsNull())
 				{
-					const uint64 OutputsOriginalSize = Outputs.GetAllocatedSize();
-
 					// Create a new entry to store in the FShaderJobCache
 					FStoredOutput* NewStoredOutput = new FStoredOutput();
-					NewStoredOutput->NumHits = 0;
-					NewStoredOutput->NumReferences = 1;
 					NewStoredOutput->JobOutput.Reserve(Results.GetSize());
 					NewStoredOutput->JobOutput.SetNum(Results.GetSize());
 
@@ -8462,10 +8503,13 @@ FShaderJobCache::FJobCachedOutput* FShaderJobCache::Find(const FJobInputHash& Ha
 
 					// Generate an output hash and cache the result in the FShaderJobCache
 					FJobOutputHash NewOutputHash = FBlake3::HashBuffer(NewStoredOutput->JobOutput.GetData(), NewStoredOutput->JobOutput.Num());
+
+					NewStoredOutput->AddRef();
 					Outputs.Add(NewOutputHash, NewStoredOutput);
 					InputHashToOutput.Add(Hash, NewOutputHash);
+					EvictionQueue.PushLast(Hash);
 
-					CurrentlyAllocatedMemory += NewStoredOutput->JobOutput.GetAllocatedSize() + sizeof(FStoredOutput);
+					CurrentlyAllocatedMemory += NewStoredOutput->GetAllocatedSize();
 
 					// return the results.
 					return &NewStoredOutput->JobOutput;
@@ -8490,6 +8534,14 @@ FShaderJobCache::FShaderJobCache()
 	CurrentlyAllocatedMemory = sizeof(*this) + InputHashToOutput.GetAllocatedSize() + Outputs.GetAllocatedSize();
 }
 
+FShaderJobCache::~FShaderJobCache()
+{
+	for (TMap<FJobOutputHash, FStoredOutput*>::TIterator Iter(Outputs); Iter; ++Iter)
+	{
+		delete Iter.Value();
+	}
+}
+
 void FShaderJobCache::Add(const FJobInputHash& Hash, const FJobCachedOutput& Contents, int32 InitialHitCount, const bool bAddToDDC)
 {
 	TRACE_CPUPROFILER_EVENT_SCOPE(FShaderJobCache::Add);
@@ -8509,7 +8561,9 @@ void FShaderJobCache::Add(const FJobInputHash& Hash, const FJobCachedOutput& Con
 
 	// add the record
 	const uint64 InputHashToOutputOriginalSize = InputHashToOutput.GetAllocatedSize();
+
 	InputHashToOutput.Add(Hash, OutputHash);
+	EvictionQueue.PushLast(Hash);
 
 	const bool bCachePerShaderDDC = IsShaderJobCacheDDCEnabled() && bAddToDDC;
 
@@ -8517,7 +8571,7 @@ void FShaderJobCache::Add(const FJobInputHash& Hash, const FJobCachedOutput& Con
 	if (CannedOutput && (bCachePerShaderDDC == false))
 	{
 		// update the output hit count
-		(*CannedOutput)->NumReferences++;
+		(*CannedOutput)->AddRef();
 	}
 	else
 	{
@@ -8525,11 +8579,12 @@ void FShaderJobCache::Add(const FJobInputHash& Hash, const FJobCachedOutput& Con
 
 		FStoredOutput* NewStoredOutput = new FStoredOutput();
 		NewStoredOutput->NumHits = InitialHitCount;
-		NewStoredOutput->NumReferences = 1;
 		NewStoredOutput->JobOutput = Contents;
+
+		NewStoredOutput->AddRef();
 		Outputs.Add(OutputHash, NewStoredOutput);
 
-		CurrentlyAllocatedMemory += NewStoredOutput->JobOutput.GetAllocatedSize() + sizeof(FStoredOutput);
+		CurrentlyAllocatedMemory += NewStoredOutput->GetAllocatedSize();
 
 #if WITH_EDITOR
 		if (bCachePerShaderDDC)
@@ -8562,45 +8617,10 @@ void FShaderJobCache::Add(const FJobInputHash& Hash, const FJobCachedOutput& Con
 			{
 				TRACE_CPUPROFILER_EVENT_SCOPE(FShaderJobCache::Trim);
 			
-				while (CurrentlyAllocatedMemory > MemoryBudgetBytes)
+				FJobInputHash KeyToRemove = {};
+				while (CurrentlyAllocatedMemory > MemoryBudgetBytes && EvictionQueue.TryPopFirst(KeyToRemove))
 				{
-					// heuristics: delete the entry that has the smallest hits. Don't account for references as if something is referenced often but not hit, it's of no value for us.
-					// (consider other heuristics: hits * memory, time it took to produce the output, last hit time)
-					int32 MinHits = INT_MAX;
-
-					// find the (new) min
-					for (TMap<FJobOutputHash, FStoredOutput*>::TConstIterator Iter(Outputs); Iter; ++Iter)
-					{
-						MinHits = FMath::Min(Iter.Value()->NumHits, MinHits);
-					}
-
-					// remove all matching this minimum until there's enough memory
-					for (TMap<FJobOutputHash, FStoredOutput*>::TIterator Iter(Outputs); Iter; ++Iter)
-					{
-						if (Iter.Value()->NumHits == MinHits)
-						{
-							CurrentlyAllocatedMemory -= static_cast<uint64>(Iter.Value()->JobOutput.GetAllocatedSize() + sizeof(FStoredOutput));
-							delete Iter.Value();
-
-							FJobOutputHash RemovedOutputHash = Iter.Key();
-							Iter.RemoveCurrent();
-
-							// remove all mappings
-							for (TMap<FJobInputHash, FJobOutputHash>::TIterator IterInputs(InputHashToOutput); IterInputs; ++IterInputs)
-							{
-								if (IterInputs.Value() == RemovedOutputHash)
-								{
-									IterInputs.RemoveCurrent();
-								}
-							}
-
-							// don't remove too much
-							if (CurrentlyAllocatedMemory < MemoryBudgetBytes)
-							{
-								break;
-							}
-						}
-					}
+					RemoveByInputHash(KeyToRemove);
 				}
 			}
 		}
@@ -8613,6 +8633,34 @@ void FShaderJobCache::Add(const FJobInputHash& Hash, const FJobCachedOutput& Con
 	CurrentlyAllocatedMemory -= InputHashToOutputOriginalSize;
 }
 
+void FShaderJobCache::RemoveByInputHash(const FJobInputHash& InputHash)
+{
+	const FJobOutputHash* FoundOutputHash = InputHashToOutput.Find(InputHash);
+
+	if (FoundOutputHash)
+	{
+		const FJobOutputHash& OutputHash = *FoundOutputHash;
+		FStoredOutput** FoundStoredOutput = Outputs.Find(OutputHash);
+
+		if (FoundStoredOutput)
+		{
+			FStoredOutput* StoredOutput = *FoundStoredOutput;
+			checkf(StoredOutput, TEXT("Invalid entry found in FShaderJobCache Output hash table. All values are expected to be valid pointers."));
+
+			const uint64 OutputSize = StoredOutput->GetAllocatedSize();
+
+			// Decrement reference count and remove cached object if it's no longer referenced by any input hashes
+			if (StoredOutput->Release() == 0)
+			{
+				Outputs.Remove(OutputHash);
+				CurrentlyAllocatedMemory -= OutputSize;
+			}
+		}
+
+		InputHashToOutput.Remove(InputHash);
+	}
+}
+
 /** Calculates memory used by the cache*/
 uint64 FShaderJobCache::GetAllocatedMemory()
 {
diff --git a/Engine/Source/Runtime/Engine/Public/ShaderCompiler.h b/Engine/Source/Runtime/Engine/Public/ShaderCompiler.h
--- a/Engine/Source/Runtime/Engine/Public/ShaderCompiler.h
+++ b/Engine/Source/Runtime/Engine/Public/ShaderCompiler.h
@@ -25,6 +25,7 @@
 #include "AssetCompilingManager.h"
 #include "Containers/HashTable.h"
 #include "Containers/List.h"
+#include "Containers/Deque.h"
 #include "Hash/Blake3.h"
 #include "SceneTypes.h"
 
@@ -300,10 +301,14 @@ inline void FShaderCommonCompileJob::Destroy() const
 	}
 }
 
+struct FShaderJobCacheStoredOutput;
+
 class FShaderJobCache
 {
 public:
+
 	FShaderJobCache();
+	~FShaderJobCache();
 
 	using FJobInputHash = FShaderCommonCompileJob::FInputHash;
 	using FJobCachedOutput = TArray<uint8>;
@@ -324,25 +329,21 @@ public:
 	uint64 GetCurrentMemoryBudget() const;
 
 private:
+
 	using FJobOutputHash = FBlake3Hash;
-	struct FStoredOutput
-	{
-		/** How many times this output is referenced by the cached jobs */
-		int32 NumReferences;
-
-		/** How many times this output has been returned as a cached result, no matter the input hash */
-		int32 NumHits;
-		
-		/** Canned output */
-		TArray<uint8> JobOutput;
-	};
+	using FStoredOutput = FShaderJobCacheStoredOutput;
+
+	void RemoveByInputHash(const FJobInputHash& InputHash);
 
-	/* a lot of outputs can be duplicated, so they are deduplicated before storing */
+	/* A lot of outputs can be duplicated, so they are deduplicated before storing */
 	TMap<FJobOutputHash, FStoredOutput*> Outputs;
 
 	/** Map of input hashes to output hashes */
 	TMap<FJobInputHash, FJobOutputHash> InputHashToOutput;
 
+	/** Queue to evict oldest elements when memory budget is exceeded */
+	TDeque<FJobInputHash> EvictionQueue;
+
 	/** Statistics - total number of times we tried to Find() some input hash */
 	uint64 TotalSearchAttempts = 0;
 
