Render stationary/movable RectLight approximately as SpotLight on mobile

#b Dmitriy.Dyomin

diff --git a/Engine/Shaders/Private/LightData.ush b/Engine/Shaders/Private/LightData.ush
--- a/Engine/Shaders/Private/LightData.ush
+++ b/Engine/Shaders/Private/LightData.ush
@@ -304,6 +304,8 @@ FDeferredLightData ConvertToDeferredLight_Mobile(const FLocalLightData In)
 	Out.Tangent			= 0;
 	Out.SoftSourceRadius= 0;
 	Out.SpecularScale	= 1.0f;
+	//Render RectLight approximately as SpotLight
+	Out.bSpotLight		= Out.bSpotLight || Out.bRectLight;
 	Out.bRectLight		= false;
 	Out.ShadowedBits	= 1;
 
diff --git a/Engine/Source/Runtime/Engine/Private/Components/RectLightComponent.cpp b/Engine/Source/Runtime/Engine/Private/Components/RectLightComponent.cpp
--- a/Engine/Source/Runtime/Engine/Private/Components/RectLightComponent.cpp
+++ b/Engine/Source/Runtime/Engine/Private/Components/RectLightComponent.cpp
@@ -250,6 +250,25 @@ void FRectLightSceneProxy::GetLightShaderParameters(FLightRenderParameters& Ligh
 	{
 		GetSceneInterface()->GetRectLightAtlasSlot(this, &LightParameters);
 	}
+	
+	// Render RectLight approximately as SpotLight on mobile
+	const bool bRenderAsSpotLight = (SceneInterface && IsMobilePlatform(SceneInterface->GetShaderPlatform()));
+	if (bRenderAsSpotLight)
+	{
+		float ClampedOuterConeAngle = FMath::DegreesToRadians(89.001f);
+		float ClampedInnerConeAngle = FMath::DegreesToRadians(70.0f);
+		float CosOuterCone = FMath::Cos(ClampedOuterConeAngle);
+		float CosInnerCone = FMath::Cos(ClampedInnerConeAngle);
+		float InvCosConeDifference = 1.0f / (CosInnerCone - CosOuterCone);
+
+		LightParameters.Color = GetColor();
+		LightParameters.FalloffExponent = 8.0f;
+		LightParameters.SpotAngles = FVector2f(CosOuterCone, InvCosConeDifference);
+		LightParameters.SourceRadius = (SourceWidth + SourceHeight) * 0.5 * 0.5f;
+		LightParameters.SourceLength = 0.0f;
+		LightParameters.RectLightBarnCosAngle = 0.0f;
+		LightParameters.RectLightBarnLength = -2.0f;
+	}
 }
 
 /**
