PR #9641: UTimelineComponent: Expose "AddInterp" functions accept static/native? (Contributed by [REDACTED_PERSON])

diff --git a/Engine/Source/Runtime/Engine/Classes/Components/TimelineComponent.h b/Engine/Source/Runtime/Engine/Classes/Components/TimelineComponent.h
--- a/Engine/Source/Runtime/Engine/Classes/Components/TimelineComponent.h
+++ b/Engine/Source/Runtime/Engine/Classes/Components/TimelineComponent.h
@@ -528,12 +528,15 @@ public:
 	
 	/** Add a vector interpolation to the timeline */
 	ENGINE_API void AddInterpVector(UCurveVector* VectorCurve, FOnTimelineVector InterpFunc, FName PropertyName = NAME_None, FName TrackName = NAME_None);
+	ENGINE_API void AddInterpVector(UCurveVector* VectorCurve, FOnTimelineVectorStatic InterpFunc);
 	
 	/** Add a float interpolation to the timeline */
 	ENGINE_API void AddInterpFloat(UCurveFloat* FloatCurve, FOnTimelineFloat InterpFunc, FName PropertyName = NAME_None, FName TrackName = NAME_None);
+	ENGINE_API void AddInterpFloat(UCurveFloat* FloatCurve, FOnTimelineFloatStatic InterpFunc);
 
 	/** Add a linear color interpolation to the timeline */
 	ENGINE_API void AddInterpLinearColor(UCurveLinearColor* LinearColorCurve, FOnTimelineLinearColor InterpFunc, FName PropertyName = NAME_None, FName TrackName = NAME_None);
+	ENGINE_API void AddInterpLinearColor(UCurveLinearColor* LinearColorCurve, FOnTimelineLinearColorStatic InterpFunc);
 
 	/** Optionally provide an object to automatically update properties on */
 	ENGINE_API void SetPropertySetObject(UObject* NewPropertySetObject);
@@ -552,6 +555,3 @@ public:
 	/** Get all curves used by the Timeline */
 	ENGINE_API void GetAllCurves(TSet<class UCurveBase*>& InOutCurves) const;
 };
-
-
-
diff --git a/Engine/Source/Runtime/Engine/Private/Timeline.cpp b/Engine/Source/Runtime/Engine/Private/Timeline.cpp
--- a/Engine/Source/Runtime/Engine/Private/Timeline.cpp
+++ b/Engine/Source/Runtime/Engine/Private/Timeline.cpp
@@ -771,16 +771,31 @@ void UTimelineComponent::AddInterpVector(UCurveVector* VectorCurve, FOnTimelineV
 	TheTimeline.AddInterpVector(VectorCurve, InterpFunc, PropertyName, TrackName);
 }
 
+void UTimelineComponent::AddInterpVector(UCurveVector* VectorCurve, const FOnTimelineVectorStatic InterpFunc)
+{
+	TheTimeline.AddInterpVector(VectorCurve, InterpFunc);
+}
+
 void UTimelineComponent::AddInterpFloat(UCurveFloat* FloatCurve, FOnTimelineFloat InterpFunc, FName PropertyName, FName TrackName)
 {
 	TheTimeline.AddInterpFloat(FloatCurve, InterpFunc, PropertyName, TrackName);
 }
 
+void UTimelineComponent::AddInterpFloat(UCurveFloat* FloatCurve, const FOnTimelineFloatStatic InterpFunc)
+{
+	TheTimeline.AddInterpFloat(FloatCurve, InterpFunc);
+}
+
 void UTimelineComponent::AddInterpLinearColor(UCurveLinearColor* LinearColorCurve, FOnTimelineLinearColor InterpFunc, FName PropertyName, FName TrackName)
 {
 	TheTimeline.AddInterpLinearColor(LinearColorCurve, InterpFunc, PropertyName, TrackName);
 }
 
+void UTimelineComponent::AddInterpLinearColor(UCurveLinearColor* LinearColorCurve, const FOnTimelineLinearColorStatic InterpFunc)
+{
+	TheTimeline.AddInterpLinearColor(LinearColorCurve, InterpFunc);
+}
+
 void UTimelineComponent::SetPlaybackPosition(float NewPosition, bool bFireEvents, bool bFireUpdate)
 {
 	TheTimeline.SetPlaybackPosition(NewPosition, bFireEvents, bFireUpdate);
@@ -977,4 +992,4 @@ void UTimelineComponent::GetAllCurves(TSet<class UCurveBase*>& InOutCurves) cons
 bool UTimelineComponent::IsPostLoadThreadSafe() const
 {
 	return true;
-}
\ No newline at end of file
+}
