Set TotalSamples variable on USoundwaves created by FSoundWavePCMWriter

diff --git a/Engine/Source/Runtime/Engine/Private/SampleBufferIO.cpp b/Engine/Source/Runtime/Engine/Private/SampleBufferIO.cpp
--- a/Engine/Source/Runtime/Engine/Private/SampleBufferIO.cpp
+++ b/Engine/Source/Runtime/Engine/Private/SampleBufferIO.cpp
@@ -460,6 +460,7 @@ namespace Audio
 		CurrentSoundWave->NumChannels = CurrentBuffer.GetNumChannels();
 		CurrentSoundWave->RawPCMDataSize = CurrentBuffer.GetNumSamples() * sizeof(int16);
 		CurrentSoundWave->Duration = (float) CurrentBuffer.GetNumFrames() / CurrentBuffer.GetSampleRate();
+		CurrentSoundWave->TotalSamples = CurrentBuffer.GetNumSamples();
 
 		if (CurrentSoundWave->RawPCMData != nullptr)
 		{
