UE-112211 - Crash in UKismetSystemLibrary::execSetStructurePropertyByName

Getting a reference to an element via UKismetArrayLibrary::Array_Get in an array of user structs can lead to a null source pointer in Generic_SetStructurePropertyByName. Specifically, this happens if the array is empty, or an invalid index is used.

We correctly throw a BP exception when this scenario occurs. However, execArrayGetByRef will set Stack.MostRecentPropertyAddress to null, which is then used as the source parameter to Generic_SetStructurePropertyByName. Returning a copy of the element doesn't have the same issue since it creates a temporary user struct on the stack as a source struct.

The most straightforward fix is to simply soft-fail in Generic_SetStructurePropertyByName. The user will have to rely on the BP exception info for the actual error.

diff --git a/Engine/Source/Runtime/Engine/Classes/Kismet/KismetSystemLibrary.h b/Engine/Source/Runtime/Engine/Classes/Kismet/KismetSystemLibrary.h
--- a/Engine/Source/Runtime/Engine/Classes/Kismet/KismetSystemLibrary.h
+++ b/Engine/Source/Runtime/Engine/Classes/Kismet/KismetSystemLibrary.h
@@ -927,40 +927,17 @@ class ENGINE_API UKismetSystemLibrary : public UBlueprintFunctionLibrary
 	UFUNCTION(BlueprintCallable, meta = (BlueprintInternalUseOnly = "true", AutoCreateRefTerm = "Value"))
 	static void SetFieldPathPropertyByName(UObject* Object, FName PropertyName, const TFieldPath<FField>& Value);
 
-	DECLARE_FUNCTION(execSetCollisionProfileNameProperty)
-	{
-		P_GET_OBJECT(UObject, OwnerObject);
-		P_GET_PROPERTY(FNameProperty, StructPropertyName);
-
-		Stack.StepCompiledIn<FStructProperty>(NULL);
-		void* SrcStructAddr = Stack.MostRecentPropertyAddress;
-
-		P_FINISH;
-		P_NATIVE_BEGIN;
-		Generic_SetStructurePropertyByName(OwnerObject, StructPropertyName, SrcStructAddr);
-		P_NATIVE_END;
-	}
+	DECLARE_FUNCTION(execSetCollisionProfileNameProperty);
 
 	/** Set a custom structure property by name */
 	UFUNCTION(BlueprintCallable, CustomThunk, meta = (BlueprintInternalUseOnly = "true", CustomStructureParam = "Value", AutoCreateRefTerm = "Value"))
 	static void SetStructurePropertyByName(UObject* Object, FName PropertyName, const FGenericStruct& Value);
 
+	UE_DEPRECATED(5.2, "Function has been deprecated.")
 	static void Generic_SetStructurePropertyByName(UObject* OwnerObject, FName StructPropertyName, const void* SrcStructAddr);
 
 	/** Based on UKismetArrayLibrary::execSetArrayPropertyByName */
-	DECLARE_FUNCTION(execSetStructurePropertyByName)
-	{
-		P_GET_OBJECT(UObject, OwnerObject);
-		P_GET_PROPERTY(FNameProperty, StructPropertyName);
-
-		Stack.StepCompiledIn<FStructProperty>(NULL);
-		void* SrcStructAddr = Stack.MostRecentPropertyAddress;
-
-		P_FINISH;
-		P_NATIVE_BEGIN;
-		Generic_SetStructurePropertyByName(OwnerObject, StructPropertyName, SrcStructAddr);
-		P_NATIVE_END;
-	}
+	DECLARE_FUNCTION(execSetStructurePropertyByName);
 
 	// --- Collision functions ------------------------------
 
diff --git a/Engine/Source/Runtime/Engine/Private/KismetSystemLibrary.cpp b/Engine/Source/Runtime/Engine/Private/KismetSystemLibrary.cpp
--- a/Engine/Source/Runtime/Engine/Private/KismetSystemLibrary.cpp
+++ b/Engine/Source/Runtime/Engine/Private/KismetSystemLibrary.cpp
@@ -56,6 +56,30 @@ namespace UE::Blueprint::Private
 	FAutoConsoleVariableRef CVarBlamePrintString(TEXT("bp.BlamePrintString"), 
 		bBlamePrintString,
 		TEXT("When true, prints the Blueprint Asset and Function that generated calls to Print String. Useful for tracking down screen message spam."));
+
+	void Generic_SetStructurePropertyByName(UObject* OwnerObject, FName StructPropertyName, FStructProperty* SrcStructProperty, const void* SrcStructAddr)
+	{
+		if (OwnerObject != nullptr)
+		{
+			FStructProperty* DestStructProperty = FindFProperty<FStructProperty>(OwnerObject->GetClass(), StructPropertyName);
+
+			// SrcStructAddr and SrcStructProperty can be null in certain scenarios.
+			// For example, retrieving an element reference from an array of user structs in BP can result in a null source.
+			// We'll report a BP exception in that case, but we also need to soft-fail here to prevent an assert in CopyValuesInternal.
+
+			bool bCanSetStructureProperty =
+				(DestStructProperty != nullptr) &&
+				(SrcStructProperty != nullptr) &&
+				(SrcStructAddr != nullptr) &&
+				SrcStructProperty->SameType(DestStructProperty);
+
+			if (bCanSetStructureProperty)
+			{
+				void* DestStructAddr = DestStructProperty->ContainerPtrToValuePtr<void>(OwnerObject);
+				DestStructProperty->CopyValuesInternal(DestStructAddr, SrcStructAddr, 1);
+			}
+		}
+	}
 }
 
 UKismetSystemLibrary::UKismetSystemLibrary(const FObjectInitializer& ObjectInitializer)
@@ -1129,6 +1153,36 @@ void UKismetSystemLibrary::SetFieldPathPropertyByName(UObject* Object, FName Pro
 	}
 }
 
+DEFINE_FUNCTION(UKismetSystemLibrary::execSetCollisionProfileNameProperty)
+{
+	P_GET_OBJECT(UObject, OwnerObject);
+	P_GET_PROPERTY(FNameProperty, StructPropertyName);
+
+	Stack.StepCompiledIn<FStructProperty>(nullptr);
+	FStructProperty* SrcStructProperty = CastField<FStructProperty>(Stack.MostRecentProperty);
+	void* SrcStructAddr = Stack.MostRecentPropertyAddress;
+
+	P_FINISH;
+	P_NATIVE_BEGIN;
+	UE::Blueprint::Private::Generic_SetStructurePropertyByName(OwnerObject, StructPropertyName, SrcStructProperty, SrcStructAddr);
+	P_NATIVE_END;
+}
+
+DEFINE_FUNCTION(UKismetSystemLibrary::execSetStructurePropertyByName)
+{
+	P_GET_OBJECT(UObject, OwnerObject);
+	P_GET_PROPERTY(FNameProperty, StructPropertyName);
+
+	Stack.StepCompiledIn<FStructProperty>(nullptr);
+	FStructProperty* SrcStructProperty = CastField<FStructProperty>(Stack.MostRecentProperty);
+	void* SrcStructAddr = Stack.MostRecentPropertyAddress;
+
+	P_FINISH;
+	P_NATIVE_BEGIN;
+	UE::Blueprint::Private::Generic_SetStructurePropertyByName(OwnerObject, StructPropertyName, SrcStructProperty, SrcStructAddr);
+	P_NATIVE_END;
+}
+
 void UKismetSystemLibrary::SetSoftClassPropertyByName(UObject* Object, FName PropertyName, const TSoftClassPtr<UObject>& Value)
 {
 	if (Object != NULL)
