PhysicsAsset/SkeletalMesh level set collision objects:
- try to make the output grid resolution better match the specified input resolution
- don't try to render the grid if resolution is too high (debug rendering can really start to slow things down)
- fix a bug generating a mesh for solid rendering of the level set cells

diff --git a/Engine/Source/Developer/PhysicsUtilities/Private/PhysicsAssetUtils.cpp b/Engine/Source/Developer/PhysicsUtilities/Private/PhysicsAssetUtils.cpp
--- a/Engine/Source/Developer/PhysicsUtilities/Private/PhysicsAssetUtils.cpp
+++ b/Engine/Source/Developer/PhysicsUtilities/Private/PhysicsAssetUtils.cpp
@@ -473,14 +473,21 @@ namespace LevelSetHelpers
 	}
 
 	// Copied from FMeshSimpleShapeApproximation::Generate_LevelSets (Plugins/Runtime/GeometryProcessing)
-	bool CreateLevelSetForMesh(const UE::Geometry::FDynamicMesh3& InMesh, int32 LevelSetGridResolution, FKLevelSetElem& OutElement)
+	bool CreateLevelSetForMesh(const UE::Geometry::FDynamicMesh3& InMesh, int32 InLevelSetGridResolution, FKLevelSetElem& OutElement)
 	{
 		using UE::Geometry::FDynamicMesh3;
 		using UE::Geometry::TSweepingMeshSDF;
 
+		constexpr int32 NumNarrowBandCells = 2;
+		constexpr int32 NumExpandCells = 1;
+
+		// Inside SDF.Compute(), extra grid cell are added to account for the narrow band ("NarrowBandMaxDistance") as well as an 
+		// outer buffer ("ExpandBounds"). So here we adjust the cell size to make the final output resolution closer to the user-specified resolution.
+		const int32 LevelSetGridResolution = FMath::Max(1, InLevelSetGridResolution - 2*(NumNarrowBandCells + NumExpandCells));
+
 		const UE::Geometry::FAxisAlignedBox3d Bounds = InMesh.GetBounds();
 		const double CellSize = Bounds.MaxDim() / LevelSetGridResolution;
-		const double ExpandBounds = 2.0 * CellSize;
+		const double ExpandBounds = NumExpandCells * CellSize;
 
 		UE::Geometry::TMeshAABBTree3<FDynamicMesh3> Spatial(&InMesh);
 
@@ -501,7 +508,7 @@ namespace LevelSetHelpers
 		SDF.Spatial = &Spatial;
 		SDF.ComputeMode = TSweepingMeshSDF<FDynamicMesh3>::EComputeModes::NarrowBand_SpatialFloodFill;
 		SDF.CellSize = (float)CellSize;
-		SDF.NarrowBandMaxDistance = 2.0 * CellSize;
+		SDF.NarrowBandMaxDistance = NumNarrowBandCells * CellSize;
 		SDF.ExactBandWidth = FMath::CeilToInt32(SDF.NarrowBandMaxDistance / CellSize);
 		SDF.ExpandBounds = FVector3d(ExpandBounds);
 
diff --git a/Engine/Source/Developer/PhysicsUtilities/Public/PhysicsAssetUtils.h b/Engine/Source/Developer/PhysicsUtilities/Public/PhysicsAssetUtils.h
--- a/Engine/Source/Developer/PhysicsUtilities/Public/PhysicsAssetUtils.h
+++ b/Engine/Source/Developer/PhysicsUtilities/Public/PhysicsAssetUtils.h
@@ -105,7 +105,7 @@ struct FPhysAssetCreateParams
 
 	/** When creating level sets, the grid resolution to use */
 	UPROPERTY(EditAnywhere, AdvancedDisplay, Category = "Body Creation", 
-		meta = (ClampMin = 1, UIMin = 5, UIMax = 100, ClampMax = 500, EditCondition = "GeomType == EPhysAssetFitGeomType::EFG_LevelSet"))
+		meta = (ClampMin = 1, UIMin = 10, UIMax = 100, ClampMax = 500, EditCondition = "GeomType == EPhysAssetFitGeomType::EFG_LevelSet"))
 	int32								LevelSetResolution;
 };
 
diff --git a/Engine/Source/Runtime/Engine/Classes/PhysicsEngine/LevelSetElem.h b/Engine/Source/Runtime/Engine/Classes/PhysicsEngine/LevelSetElem.h
--- a/Engine/Source/Runtime/Engine/Classes/PhysicsEngine/LevelSetElem.h
+++ b/Engine/Source/Runtime/Engine/Classes/PhysicsEngine/LevelSetElem.h
@@ -64,6 +64,7 @@ struct FKLevelSetElem : public FKShapeElem
 
 	ENGINE_API FBox CalcAABB(const FTransform& BoneTM, const FVector& Scale3D) const;
 	ENGINE_API FBox UntransformedAABB() const;
+	ENGINE_API FIntVector3 GridResolution() const;
 
 	bool Serialize(FArchive& Ar);
 
diff --git a/Engine/Source/Runtime/Engine/Private/Collision/KAggregateGeom.cpp b/Engine/Source/Runtime/Engine/Private/Collision/KAggregateGeom.cpp
--- a/Engine/Source/Runtime/Engine/Private/Collision/KAggregateGeom.cpp
+++ b/Engine/Source/Runtime/Engine/Private/Collision/KAggregateGeom.cpp
@@ -514,6 +514,17 @@ FBox FKLevelSetElem::UntransformedAABB() const
 	return Box;
 }
 
+FIntVector3 FKLevelSetElem::GridResolution() const
+{
+	if (LevelSet.IsValid())
+	{
+		const Chaos::TVec3<int32>& Dim = LevelSet->GetGrid().Counts();
+		return FIntVector3(Dim[0], Dim[1], Dim[2]);
+	}
+
+	return FIntVector3(0,0,0);
+}
+
 void FKLevelSetElem::BuildLevelSet(const FTransform& GridTransform, const TArray<double>& GridValues, const FIntVector& GridDims, float GridCellSize)
 {
 	const Chaos::FVec3 Min(0, 0, 0);
@@ -592,24 +603,25 @@ void FKLevelSetElem::GetZeroIsosurfaceGridCellFaces(TArray<FVector3f>& Vertices,
 					const double SignNextI = FMath::Sign(LevelSet->GetPhiArray()(i+1, j, k));
 					const double SignNextJ = FMath::Sign(LevelSet->GetPhiArray()(i, j+1, k));
 					const double SignNextK = FMath::Sign(LevelSet->GetPhiArray()(i, j, k+1));
+					const FVector3d Dx(Grid.Dx());
 
-					const FVector3d CellMin = Grid.MinCorner() + Grid.Dx() * FVector3d(i, j, k);
+					const FVector3d CellMin = Grid.MinCorner() + Dx * FVector3d(i, j, k);
 
 					if (Sign > SignNextI)
 					{
-						const int32 V0 = Vertices.Emplace( CellMin + FVector3d(1, 0, 0) );
-						const int32 V1 = Vertices.Emplace( CellMin + FVector3d(1, 1, 0) );
-						const int32 V2 = Vertices.Emplace( CellMin + FVector3d(1, 1, 1) );
-						const int32 V3 = Vertices.Emplace( CellMin + FVector3d(1, 0, 1) );
+						const int32 V0 = Vertices.Emplace( CellMin + Dx * FVector3d(1, 0, 0) );
+						const int32 V1 = Vertices.Emplace( CellMin + Dx * FVector3d(1, 1, 0) );
+						const int32 V2 = Vertices.Emplace( CellMin + Dx * FVector3d(1, 1, 1) );
+						const int32 V3 = Vertices.Emplace( CellMin + Dx * FVector3d(1, 0, 1) );
 						Tris.Emplace(FIntVector(V0, V1, V2));
 						Tris.Emplace(FIntVector(V2, V3, V0));
 					}
 					else if (Sign < SignNextI)
 					{
-						const int32 V0 = Vertices.Emplace(CellMin + FVector3d(1, 0, 0));
-						const int32 V1 = Vertices.Emplace(CellMin + FVector3d(1, 1, 0));
-						const int32 V2 = Vertices.Emplace(CellMin + FVector3d(1, 1, 1));
-						const int32 V3 = Vertices.Emplace(CellMin + FVector3d(1, 0, 1));
+						const int32 V0 = Vertices.Emplace(CellMin + Dx * FVector3d(1, 0, 0));
+						const int32 V1 = Vertices.Emplace(CellMin + Dx * FVector3d(1, 1, 0));
+						const int32 V2 = Vertices.Emplace(CellMin + Dx * FVector3d(1, 1, 1));
+						const int32 V3 = Vertices.Emplace(CellMin + Dx * FVector3d(1, 0, 1));
 						Tris.Emplace(FIntVector(V0, V2, V1));
 						Tris.Emplace(FIntVector(V2, V0, V3));
 					}
@@ -617,38 +629,38 @@ void FKLevelSetElem::GetZeroIsosurfaceGridCellFaces(TArray<FVector3f>& Vertices,
 					
 					if (Sign > SignNextJ)
 					{
-						const int32 V0 = Vertices.Emplace(CellMin + FVector3d(0, 1, 0));
-						const int32 V1 = Vertices.Emplace(CellMin + FVector3d(1, 1, 0));
-						const int32 V2 = Vertices.Emplace(CellMin + FVector3d(1, 1, 1));
-						const int32 V3 = Vertices.Emplace(CellMin + FVector3d(0, 1, 1));
+						const int32 V0 = Vertices.Emplace(CellMin + Dx * FVector3d(0, 1, 0));
+						const int32 V1 = Vertices.Emplace(CellMin + Dx * FVector3d(1, 1, 0));
+						const int32 V2 = Vertices.Emplace(CellMin + Dx * FVector3d(1, 1, 1));
+						const int32 V3 = Vertices.Emplace(CellMin + Dx * FVector3d(0, 1, 1));
 						Tris.Emplace(FIntVector(V0, V2, V1));
 						Tris.Emplace(FIntVector(V2, V0, V3));
 					}
 					else if (Sign < SignNextJ)
 					{
-						const int32 V0 = Vertices.Emplace(CellMin + FVector3d(0, 1, 0));
-						const int32 V1 = Vertices.Emplace(CellMin + FVector3d(1, 1, 0));
-						const int32 V2 = Vertices.Emplace(CellMin + FVector3d(1, 1, 1));
-						const int32 V3 = Vertices.Emplace(CellMin + FVector3d(0, 1, 1));
+						const int32 V0 = Vertices.Emplace(CellMin + Dx * FVector3d(0, 1, 0));
+						const int32 V1 = Vertices.Emplace(CellMin + Dx * FVector3d(1, 1, 0));
+						const int32 V2 = Vertices.Emplace(CellMin + Dx * FVector3d(1, 1, 1));
+						const int32 V3 = Vertices.Emplace(CellMin + Dx * FVector3d(0, 1, 1));
 						Tris.Emplace(FIntVector(V0, V1, V2));
 						Tris.Emplace(FIntVector(V2, V3, V0));
 					}
 
 					if (Sign > SignNextK)
 					{
-						const int32 V0 = Vertices.Emplace(CellMin + FVector3d(0, 0, 1));
-						const int32 V1 = Vertices.Emplace(CellMin + FVector3d(1, 0, 1));
-						const int32 V2 = Vertices.Emplace(CellMin + FVector3d(1, 1, 1));
-						const int32 V3 = Vertices.Emplace(CellMin + FVector3d(0, 1, 1));
+						const int32 V0 = Vertices.Emplace(CellMin + Dx * FVector3d(0, 0, 1));
+						const int32 V1 = Vertices.Emplace(CellMin + Dx * FVector3d(1, 0, 1));
+						const int32 V2 = Vertices.Emplace(CellMin + Dx * FVector3d(1, 1, 1));
+						const int32 V3 = Vertices.Emplace(CellMin + Dx * FVector3d(0, 1, 1));
 						Tris.Emplace(FIntVector(V0, V1, V2));
 						Tris.Emplace(FIntVector(V2, V3, V0));
 					}
 					else if (Sign < SignNextK)
 					{
-						const int32 V0 = Vertices.Emplace(CellMin + FVector3d(0, 0, 1));
-						const int32 V1 = Vertices.Emplace(CellMin + FVector3d(1, 0, 1));
-						const int32 V2 = Vertices.Emplace(CellMin + FVector3d(1, 1, 1));
-						const int32 V3 = Vertices.Emplace(CellMin + FVector3d(0, 1, 1));
+						const int32 V0 = Vertices.Emplace(CellMin + Dx * FVector3d(0, 0, 1));
+						const int32 V1 = Vertices.Emplace(CellMin + Dx * FVector3d(1, 0, 1));
+						const int32 V2 = Vertices.Emplace(CellMin + Dx * FVector3d(1, 1, 1));
+						const int32 V3 = Vertices.Emplace(CellMin + Dx * FVector3d(0, 1, 1));
 						Tris.Emplace(FIntVector(V0, V2, V1));
 						Tris.Emplace(FIntVector(V2, V0, V3));
 					}
diff --git a/Engine/Source/Runtime/Engine/Private/PhysicsEngine/PhysDrawing.cpp b/Engine/Source/Runtime/Engine/Private/PhysicsEngine/PhysDrawing.cpp
--- a/Engine/Source/Runtime/Engine/Private/PhysicsEngine/PhysDrawing.cpp
+++ b/Engine/Source/Runtime/Engine/Private/PhysicsEngine/PhysDrawing.cpp
@@ -791,16 +791,23 @@ void FKLevelSetElem::DrawElemWire(class FPrimitiveDrawInterface* PDI, const FTra
 {
 	TArray<FBox> Boxes;
 
-	// Bounding box
-	Boxes.Add(UntransformedAABB());
-
-	// Cells with negative Phi values
-	const double Threshold = UE_KINDA_SMALL_NUMBER;		// allow slightly greater than zero for visualization purposes
-	GetInteriorGridCells(Boxes, UE_KINDA_SMALL_NUMBER);
+	if (GridResolution().GetMax() < 30)	// Don't render individual cells if the resolution is too high
+	{
+		// Cells with negative Phi values
+		const double Threshold = UE_KINDA_SMALL_NUMBER;		// allow slightly greater than zero for visualization purposes
+		GetInteriorGridCells(Boxes, UE_KINDA_SMALL_NUMBER);
+	}
+	else
+	{
+		// Just render the bounding box of the grid itself
+		Boxes.Add(UntransformedAABB());
+	}
 
+	PDI->AddReserveLines(SDPG_World, Boxes.Num() * 12);
 	for (const FBox& Box : Boxes)
 	{
-		DrawWireBox(PDI, ElemTM.ToMatrixWithScale(), Box, Color, SDPG_World, 0.f);
+		constexpr float Thickness = 0.f;
+		DrawWireBox(PDI, ElemTM.ToMatrixWithScale(), Box, Color, SDPG_World, Thickness);
 	}
 }
 
