Reduce stack usage in debug builds

diff --git a/Engine/Source/Runtime/Engine/Private/Animation/AnimNodeBase.cpp b/Engine/Source/Runtime/Engine/Private/Animation/AnimNodeBase.cpp
--- a/Engine/Source/Runtime/Engine/Private/Animation/AnimNodeBase.cpp
+++ b/Engine/Source/Runtime/Engine/Private/Animation/AnimNodeBase.cpp
@@ -10,6 +10,55 @@
 
 #include UE_INLINE_GENERATED_CPP_BY_NAME(AnimNodeBase)
 
+namespace UE::Anim::Private
+{
+#if !(UE_BUILD_SHIPPING || UE_BUILD_TEST)
+	// Don't inline this function to keep the stack usage down
+	FORCENOINLINE void ValidatePose(const FCompactPose& Pose, const FAnimInstanceProxy* AnimInstanceProxy, const FAnimNode_Base* LinkedNode)
+	{
+		if (Pose.ContainsNaN())
+		{
+			// Show bone transform with some useful debug info
+			const auto& Bones = Pose.GetBones();
+			for (int32 CPIndex = 0; CPIndex < Bones.Num(); ++CPIndex)
+			{
+				const FTransform& Bone = Bones[CPIndex];
+				if (Bone.ContainsNaN())
+				{
+					const FBoneContainer& BoneContainer = Pose.GetBoneContainer();
+					const FReferenceSkeleton& RefSkel = BoneContainer.GetReferenceSkeleton();
+					const FMeshPoseBoneIndex MeshBoneIndex = BoneContainer.MakeMeshPoseIndex(FCompactPoseBoneIndex(CPIndex));
+					ensureMsgf(!Bone.ContainsNaN(), TEXT("Bone (%s) contains NaN from AnimInstance:[%s] Node:[%s] Value:[%s]"),
+						*RefSkel.GetBoneName(MeshBoneIndex.GetInt()).ToString(),
+						*AnimInstanceProxy->GetAnimInstanceName(), LinkedNode ? *LinkedNode->StaticStruct()->GetName() : TEXT("NULL"),
+						*Bone.ToString());
+				}
+			}
+		}
+
+		if (!Pose.IsNormalized())
+		{
+			// Show bone transform with some useful debug info
+			const auto& Bones = Pose.GetBones();
+			for (int32 CPIndex = 0; CPIndex < Bones.Num(); ++CPIndex)
+			{
+				const FTransform& Bone = Bones[CPIndex];
+				if (!Bone.IsRotationNormalized())
+				{
+					const FBoneContainer& BoneContainer = Pose.GetBoneContainer();
+					const FReferenceSkeleton& RefSkel = BoneContainer.GetReferenceSkeleton();
+					const FMeshPoseBoneIndex MeshBoneIndex = BoneContainer.MakeMeshPoseIndex(FCompactPoseBoneIndex(CPIndex));
+					ensureMsgf(Bone.IsRotationNormalized(), TEXT("Bone (%s) Rotation not normalized from AnimInstance:[%s] Node:[%s] Rotation:[%s]"),
+						*RefSkel.GetBoneName(MeshBoneIndex.GetInt()).ToString(),
+						*AnimInstanceProxy->GetAnimInstanceName(), LinkedNode ? *LinkedNode->StaticStruct()->GetName() : TEXT("NULL"),
+						*Bone.GetRotation().ToString());
+				}
+			}
+		}
+	}
+#endif
+}
+
 /////////////////////////////////////////////////////
 // FAnimationBaseContext
 
@@ -388,45 +437,7 @@ void FPoseLink::Evaluate(FPoseContext& Output)
 
 	// Detect non valid output
 #if !(UE_BUILD_SHIPPING || UE_BUILD_TEST)
-	if (Output.ContainsNaN())
-	{
-		// Show bone transform with some useful debug info
-		const auto& Bones = Output.Pose.GetBones();
-		for (int32 CPIndex = 0; CPIndex < Bones.Num(); ++CPIndex)
-		{
-			const FTransform& Bone = Bones[CPIndex];
-			if (Bone.ContainsNaN())
-			{
-				const FBoneContainer& BoneContainer = Output.Pose.GetBoneContainer();
-				const FReferenceSkeleton& RefSkel = BoneContainer.GetReferenceSkeleton();
-				const FMeshPoseBoneIndex MeshBoneIndex = BoneContainer.MakeMeshPoseIndex(FCompactPoseBoneIndex(CPIndex));
-				ensureMsgf(!Bone.ContainsNaN(), TEXT("Bone (%s) contains NaN from AnimInstance:[%s] Node:[%s] Value:[%s]"),
-					*RefSkel.GetBoneName(MeshBoneIndex.GetInt()).ToString(),
-					*Output.AnimInstanceProxy->GetAnimInstanceName(), LinkedNode ? *LinkedNode->StaticStruct()->GetName() : TEXT("NULL"), 
-					*Bone.ToString());
-			}
-		}
-	}
-
-	if (!Output.IsNormalized())
-	{
-		// Show bone transform with some useful debug info
-		const auto& Bones = Output.Pose.GetBones();
-		for (int32 CPIndex = 0; CPIndex < Bones.Num(); ++CPIndex)
-		{
-			const FTransform& Bone = Bones[CPIndex];
-			if (!Bone.IsRotationNormalized())
-			{
-				const FBoneContainer& BoneContainer = Output.Pose.GetBoneContainer();
-				const FReferenceSkeleton& RefSkel = BoneContainer.GetReferenceSkeleton();
-				const FMeshPoseBoneIndex MeshBoneIndex = BoneContainer.MakeMeshPoseIndex(FCompactPoseBoneIndex(CPIndex));
-				ensureMsgf(Bone.IsRotationNormalized(), TEXT("Bone (%s) Rotation not normalized from AnimInstance:[%s] Node:[%s] Rotation:[%s]"), 
-					*RefSkel.GetBoneName(MeshBoneIndex.GetInt()).ToString(),
-					*Output.AnimInstanceProxy->GetAnimInstanceName(), LinkedNode ? *LinkedNode->StaticStruct()->GetName() : TEXT("NULL"), 
-					*Bone.GetRotation().ToString());
-			}
-		}
-	}
+	UE::Anim::Private::ValidatePose(Output.Pose, Output.AnimInstanceProxy, LinkedNode);
 #endif
 }
 
@@ -494,31 +505,7 @@ void FComponentSpacePoseLink::EvaluateComponentSpace(FComponentSpacePoseContext&
 
 	// Detect non valid output
 #if !(UE_BUILD_SHIPPING || UE_BUILD_TEST)
-	if (Output.ContainsNaN())
-	{
-		// Show bone transform with some useful debug info
-		for (const FTransform& Bone : Output.Pose.GetPose().GetBones())
-		{
-			if (Bone.ContainsNaN())
-			{
-				ensureMsgf(!Bone.ContainsNaN(), TEXT("Bone transform contains NaN from AnimInstance:[%s] Node:[%s] Value:[%s]")
-					, *Output.AnimInstanceProxy->GetAnimInstanceName(), LinkedNode ? *LinkedNode->StaticStruct()->GetName() : TEXT("NULL"), *Bone.ToString());
-			}
-		}
-	}
-
-	if (!Output.IsNormalized())
-	{
-		// Show bone transform with some useful debug info
-		for (const FTransform& Bone : Output.Pose.GetPose().GetBones())
-		{
-			if (!Bone.IsRotationNormalized())
-			{
-				ensureMsgf(Bone.IsRotationNormalized(), TEXT("Bone Rotation not normalized from AnimInstance:[%s] Node:[%s] Value:[%s]")
-					, *Output.AnimInstanceProxy->GetAnimInstanceName(), LinkedNode ? *LinkedNode->StaticStruct()->GetName() : TEXT("NULL"), *Bone.ToString());
-			}
-		}
-	}
+	UE::Anim::Private::ValidatePose(Output.Pose.GetPose(), Output.AnimInstanceProxy, LinkedNode);
 #endif
 }
 
