Don't crash when we can't create a debug directory.

- Log a warning message and get the last error, which should hopefully give us more information on what failed.
- Remove our previous code to try and catch paths longer than max_path.  On windows this is 260, so we would allow any directory path under that length to make it through to the make directory call.  Windows only allows paths that are (260 - 12), so this check wasn't accurately catching failure cases so it has been removed.
- I opted to not modify GetMaxPathLength and instead just check for failure on making the directory.
- Provide a better error message so users know how to fix the warning if they encounter it.

diff --git a/Engine/Source/Runtime/Engine/Private/ShaderCompiler/ShaderCompiler.cpp b/Engine/Source/Runtime/Engine/Private/ShaderCompiler/ShaderCompiler.cpp
--- a/Engine/Source/Runtime/Engine/Private/ShaderCompiler/ShaderCompiler.cpp
+++ b/Engine/Source/Runtime/Engine/Private/ShaderCompiler/ShaderCompiler.cpp
@@ -3985,16 +3985,14 @@ FString FShaderCompilingManager::CreateShaderDebugInfoPath(const FShaderCompiler
 	DumpDebugInfoPath.ReplaceInline(TEXT("?"), TEXT("!"));
 	DumpDebugInfoPath.ReplaceInline(TEXT("\""), TEXT("\'"));
 
-	// Return empty string if we exceeded the platform path limitations and emit warning
-	if (DumpDebugInfoPath.Len() > FPlatformMisc::GetMaxPathLength())
-	{
-		UE_LOG(LogShaderCompilers, Warning, TEXT("Shader debug path exceeded platform limits: %s"), *DumpDebugInfoPath);
-		return TEXT("");
-	}
-
 	if (!IFileManager::Get().DirectoryExists(*DumpDebugInfoPath))
 	{
-		verifyf(IFileManager::Get().MakeDirectory(*DumpDebugInfoPath, true), TEXT("Failed to create directory for shader debug info '%s'"), *DumpDebugInfoPath);
+		if (!IFileManager::Get().MakeDirectory(*DumpDebugInfoPath, true))
+		{
+			const uint32 ErrorCode = FPlatformMisc::GetLastError();
+			UE_LOG(LogShaderCompilers, Warning, TEXT("Last Error %u: Failed to create directory for shader debug info '%s'. Try enabling large file paths or r.DumpShaderDebugShortNames."), ErrorCode, *DumpDebugInfoPath);
+			return FString();
+		}
 	}
 
 	return DumpDebugInfoPath;
