World Partition: fix issue with new maps returning a world size of 0,0,0 until saved.

#rnx

diff --git a/Engine/Source/Runtime/Engine/Private/WorldPartition/WorldPartition.cpp b/Engine/Source/Runtime/Engine/Private/WorldPartition/WorldPartition.cpp
--- a/Engine/Source/Runtime/Engine/Private/WorldPartition/WorldPartition.cpp
+++ b/Engine/Source/Runtime/Engine/Private/WorldPartition/WorldPartition.cpp
@@ -1730,13 +1730,35 @@ void UWorldPartition::RemapSoftObjectPath(FSoftObjectPath& ObjectPath)
 FBox UWorldPartition::GetEditorWorldBounds() const
 {
 	check(EditorHash);
-	return IsStreamingEnabled() ? EditorHash->GetEditorWorldBounds() : EditorHash->GetNonSpatialBounds();
+
+	if (IsStreamingEnabled())
+	{
+		const FBox EditorWorldBounds = EditorHash->GetEditorWorldBounds();
+		
+		if (EditorWorldBounds.IsValid)
+		{
+			return EditorWorldBounds;
+		}
+	}
+
+	return EditorHash->GetNonSpatialBounds();
 }
 
 FBox UWorldPartition::GetRuntimeWorldBounds() const
 {
 	check(EditorHash);
-	return IsStreamingEnabled() ? EditorHash->GetRuntimeWorldBounds() : EditorHash->GetNonSpatialBounds();
+
+	if (IsStreamingEnabled())
+	{
+		const FBox RuntimeWorldBounds = EditorHash->GetRuntimeWorldBounds();
+		
+		if (RuntimeWorldBounds.IsValid)
+		{
+			return RuntimeWorldBounds;
+		}
+	}
+
+	return EditorHash->GetNonSpatialBounds();
 }
 #endif
 
