Fix for content browser preview failing if FStreamedAudioPlatformData's async task is in flight.
#rnx

diff --git a/Engine/Source/Runtime/Engine/Private/SoundWave.cpp b/Engine/Source/Runtime/Engine/Private/SoundWave.cpp
--- a/Engine/Source/Runtime/Engine/Private/SoundWave.cpp
+++ b/Engine/Source/Runtime/Engine/Private/SoundWave.cpp
@@ -1163,6 +1163,7 @@ void USoundWave::SetSoundAssetCompressionType(ESoundAssetCompressionType InSound
 #if WITH_EDITOR
 	SoundAssetCompressionType = InSoundAssetCompressionType;
 	SoundWaveDataPtr->bIsSeekable = IsSeekable();
+	SoundWaveDataPtr->RuntimeFormat = Audio::ToName(InSoundAssetCompressionType);
 	UpdateAsset();
 #endif // #if WITH_EDITOR
 }
@@ -2509,10 +2510,14 @@ void USoundWave::PostEditChangeProperty(FPropertyChangedEvent& PropertyChangedEv
 				}
 			}
 
+			if(Name == SoundAssetCompressionTypeFName)
+			{
+				SetSoundAssetCompressionType(SoundAssetCompressionType);
+			}
+
 			if (Name == CompressionQualityFName
 				|| Name == SampleRateFName
 				|| Name == StreamingFName
-				|| Name == SoundAssetCompressionTypeFName
 				|| Name == LoadingBehaviorFName
 				|| Name == InitialChunkSizeFName
 				|| Name == TransformationsFName)
