Fix contentbundle unable to inject their content after launching & PIEing in a non-wp world with playlist that deactivate plugins when exiting PIE

diff --git a/Engine/Source/Runtime/Engine/Private/WorldPartition/ContentBundle/ContentBundleClient.cpp b/Engine/Source/Runtime/Engine/Private/WorldPartition/ContentBundle/ContentBundleClient.cpp
--- a/Engine/Source/Runtime/Engine/Private/WorldPartition/ContentBundle/ContentBundleClient.cpp
+++ b/Engine/Source/Runtime/Engine/Private/WorldPartition/ContentBundle/ContentBundleClient.cpp
@@ -37,8 +37,17 @@ void FContentBundleClient::RequestRemoveContent()
 {
 	if (State == EContentBundleClientState::ContentInjectionRequested)
 	{
-		SetState(EContentBundleClientState::ContentRemovalRequested);
-		GEngine->GetEngineSubsystem<UContentBundleEngineSubsystem>()->RequestContentRemoval(*this);
+		if (HasInjectedAnyContent())
+		{
+			SetState(EContentBundleClientState::ContentRemovalRequested);
+			GEngine->GetEngineSubsystem<UContentBundleEngineSubsystem>()->RequestContentRemoval(*this);
+		}
+		else
+		{
+			// No ContentBundleBundles were injected (Non-WP Worlds). Reset the state to Register.
+			// Since the client will never receive OnContentRemovedFromWorld.
+			SetState(EContentBundleClientState::Registered);
+		}
 	}
 	else
 	{
@@ -97,11 +106,16 @@ void FContentBundleClient::OnContentRemovedFromWorld(EContentBundleStatus Remova
 	{
 		SetWorldContentState(InjectedWorld, EWorldContentState::NoContent);
 
-		// If the client request content removal and removed everything set the client state as registered.
-		// If the client did not request content removal but removed everything keep the client state. We just changed levels. The next level will have its content bundle injected.
-		if (GetState() == EContentBundleClientState::ContentRemovalRequested && !HasInjectedAnyContent())
+		if (!HasInjectedAnyContent())
 		{
-			SetState(EContentBundleClientState::Registered);
+			WorldContentStates.Empty();
+
+			// If the client did not request content removal but removed everything keep the client state. We just changed levels. The next level will have its content bundle injected.
+			// If the client request content removal and removed everything set the client state as registered.
+			if (GetState() == EContentBundleClientState::ContentRemovalRequested)
+			{
+				SetState(EContentBundleClientState::Registered);
+			}
 		}
 	}
 	else
