Content Bundle: changed several TWeakObjectPtr to the owning world to raw pointers as this can be accessed during UWorld::BeginDestroy after the world is unhashed, resulting in a crash.
This is valid since FContentBundleBase and FContentBundleContainer are both owned by the world being destructed.

#rnx

diff --git a/Engine/Source/Runtime/Engine/Private/WorldPartition/ContentBundle/ContentBundleBase.cpp b/Engine/Source/Runtime/Engine/Private/WorldPartition/ContentBundle/ContentBundleBase.cpp
--- a/Engine/Source/Runtime/Engine/Private/WorldPartition/ContentBundle/ContentBundleBase.cpp
+++ b/Engine/Source/Runtime/Engine/Private/WorldPartition/ContentBundle/ContentBundleBase.cpp
@@ -9,13 +9,11 @@
 #include "Engine/World.h"
 
 FContentBundleBase::FContentBundleBase(TSharedPtr<FContentBundleClient>& InClient, UWorld* InWorld)
-	: Client(InClient),
-	InjectedWorld(InWorld),
-	Descriptor(InClient->GetDescriptor()),
-	Status(EContentBundleStatus::Unknown)
-{
-
-}
+	: Client(InClient)
+	, InjectedWorld(InWorld)
+	, Descriptor(InClient->GetDescriptor())
+	, Status(EContentBundleStatus::Unknown)
+{}
 
 FContentBundleBase::~FContentBundleBase()
 {
@@ -98,7 +96,7 @@ FString FContentBundleBase::GetExternalStreamingObjectName() const
 
 UWorld* FContentBundleBase::GetInjectedWorld() const
 {
-	return InjectedWorld.Get();
+	return InjectedWorld;
 }
 
 const FString& FContentBundleBase::GetDisplayName() const
diff --git a/Engine/Source/Runtime/Engine/Private/WorldPartition/ContentBundle/ContentBundleContainer.cpp b/Engine/Source/Runtime/Engine/Private/WorldPartition/ContentBundle/ContentBundleContainer.cpp
--- a/Engine/Source/Runtime/Engine/Private/WorldPartition/ContentBundle/ContentBundleContainer.cpp
+++ b/Engine/Source/Runtime/Engine/Private/WorldPartition/ContentBundle/ContentBundleContainer.cpp
@@ -13,10 +13,8 @@
 #include "UObject/UObjectGlobals.h"
 
 FContentBundleContainer::FContentBundleContainer(UWorld* WorldToInjectIn)
-	:InjectedWorld(WorldToInjectIn)
-{
-
-}
+	: InjectedWorld(WorldToInjectIn)
+{}
 
 FContentBundleContainer::~FContentBundleContainer()
 {
@@ -33,7 +31,7 @@ FContentBundleContainer::~FContentBundleContainer()
 
 UWorld* FContentBundleContainer::GetInjectedWorld() const
 {
-	return InjectedWorld.Get();
+	return InjectedWorld;
 }
 
 void FContentBundleContainer::Initialize()
@@ -219,14 +217,14 @@ FContentBundleBase& FContentBundleContainer::InitializeContentBundle(TSharedPtr<
 #if WITH_EDITOR
 	if (UseEditorContentBundle())
 	{
-		ContentBundle = GetEditorContentBundles().Emplace_GetRef(MakeShared<FContentBundleEditor>(ContentBundleClient, InjectedWorld.Get())).Get();
+		ContentBundle = GetEditorContentBundles().Emplace_GetRef(MakeShared<FContentBundleEditor>(ContentBundleClient, GetInjectedWorld())).Get();
 	}
 	else
 	{
-		ContentBundle = GetGameContentBundles().Emplace_GetRef(MakeUnique<FContentBundle>(ContentBundleClient, InjectedWorld.Get())).Get();
+		ContentBundle = GetGameContentBundles().Emplace_GetRef(MakeUnique<FContentBundle>(ContentBundleClient, GetInjectedWorld())).Get();
 	}
 #else
-	 ContentBundle = GetGameContentBundles().Emplace_GetRef(MakeUnique<FContentBundle>(ContentBundleClient, InjectedWorld.Get())).Get();
+	 ContentBundle = GetGameContentBundles().Emplace_GetRef(MakeUnique<FContentBundle>(ContentBundleClient, GetInjectedWorld())).Get();
 #endif
 
 	 ContentBundle->Initialize();
diff --git a/Engine/Source/Runtime/Engine/Public/WorldPartition/ContentBundle/ContentBundleBase.h b/Engine/Source/Runtime/Engine/Public/WorldPartition/ContentBundle/ContentBundleBase.h
--- a/Engine/Source/Runtime/Engine/Public/WorldPartition/ContentBundle/ContentBundleBase.h
+++ b/Engine/Source/Runtime/Engine/Public/WorldPartition/ContentBundle/ContentBundleBase.h
@@ -42,7 +42,7 @@ protected:
 
 private:
 	TWeakPtr<FContentBundleClient> Client;
-	TWeakObjectPtr<UWorld> InjectedWorld;
+	UWorld* InjectedWorld;
 	const UContentBundleDescriptor* Descriptor;
 	EContentBundleStatus Status;
 };
\ No newline at end of file
diff --git a/Engine/Source/Runtime/Engine/Public/WorldPartition/ContentBundle/ContentBundleContainer.h b/Engine/Source/Runtime/Engine/Public/WorldPartition/ContentBundle/ContentBundleContainer.h
--- a/Engine/Source/Runtime/Engine/Public/WorldPartition/ContentBundle/ContentBundleContainer.h
+++ b/Engine/Source/Runtime/Engine/Public/WorldPartition/ContentBundle/ContentBundleContainer.h
@@ -95,5 +95,5 @@ private:
 	ContentBundleGameArray ContentBundles;
 #endif
 
-	TWeakObjectPtr<UWorld> InjectedWorld;
+	UWorld* InjectedWorld;
 };
\ No newline at end of file
