SubObject message serialization:
* Added header bit telling if this message is a destroy message or not.
* DestroyMessages now add a destroy flags instead of a null netguid

diff --git a/Engine/Source/Runtime/Core/Public/Misc/NetworkVersion.h b/Engine/Source/Runtime/Core/Public/Misc/NetworkVersion.h
--- a/Engine/Source/Runtime/Core/Public/Misc/NetworkVersion.h
+++ b/Engine/Source/Runtime/Core/Public/Misc/NetworkVersion.h
@@ -71,6 +71,7 @@ enum EEngineNetworkVersionHistory
 	HISTORY_PLACEHOLDER = 27,						// Bump version to a placeholder.  This version is the same as HISTORY_REPMOVE_SERVERFRAME_AND_HANDLE
 	HISTORY_RUNTIME_FEATURES_COMPATIBILITY = 28,	// Bump version to add network runtime feature compatibility test to handshake (hello/upgrade) control messages
 	HISTORY_SOFTOBJECTPTR_NETGUIDS = 29,			// Bump version to support replicating SoftObjectPtrs by NetGuid instead of raw strings.
+	HISTORY_SUBOBJECT_DESTROY_FLAG = 30,			// Bump version to support subobject destruction message flags
 	// New history items go above here.
 
 	HISTORY_ENGINENETVERSION_PLUS_ONE,
diff --git a/Engine/Source/Runtime/Engine/Private/DataChannel.cpp b/Engine/Source/Runtime/Engine/Private/DataChannel.cpp
--- a/Engine/Source/Runtime/Engine/Private/DataChannel.cpp
+++ b/Engine/Source/Runtime/Engine/Private/DataChannel.cpp
@@ -4031,7 +4031,12 @@ void UActorChannel::WriteContentBlockHeader( UObject* Obj, FNetBitWriter &Bunch,
 		}
 		else
 		{
+			// Not a stable name
 			Bunch.WriteBit( 0 );
+
+			// Not a destroy message
+			Bunch.WriteBit( 0 );
+
 			UClass *ObjClass = Obj->GetClass();
 			Bunch << ObjClass;
 
@@ -4079,11 +4084,13 @@ void UActorChannel::WriteContentBlockForSubObjectDelete( FOutBunch & Bunch, FNet
 	// Send a 0 bit to indicate that this is not a stably named object
 	Bunch.WriteBit( 0 );
 
-	//	-Invalid NetGUID (interpreted as delete)
-	FNetworkGUID InvalidNetGUID;
-	InvalidNetGUID.Reset();
-	Bunch << InvalidNetGUID;
-	NET_CHECKSUM(Bunch); // Matches checksum in UPackageMapClient::InternalWriteObject
+	// Send that this is a delete message
+	Bunch.WriteBit( 1 );
+
+	uint8 DeleteFlags = 0; // flags will be defined later
+	Bunch << DeleteFlags;
+
+	NET_CHECKSUM(Bunch); // Matches checksum in UPackageMapClient::ReadContentBlockHeader - DeleteSubObject message
 
 	// Since the subobject has been deleted, we don't have a valid object to pass to the profiler.
 	NETWORK_PROFILER(GNetworkProfiler.TrackBeginContentBlock(nullptr, Bunch.GetNumBits() - NumStartingBits, Connection));
@@ -4129,7 +4136,7 @@ UObject* UActorChannel::ReadContentBlockHeader(FInBunch& Bunch, bool& bObjectDel
 
 		AddToChainResultPtr(Bunch.ExtendedError, ENetCloseResult::ContentBlockHeaderRepLayoutFail);
 
-		return NULL;
+		return nullptr;
 	}
 
 	const bool bIsActor = Bunch.ReadBit() != 0 ? true : false;
@@ -4140,7 +4147,7 @@ UObject* UActorChannel::ReadContentBlockHeader(FInBunch& Bunch, bool& bObjectDel
 
 		AddToChainResultPtr(Bunch.ExtendedError, ENetCloseResult::ContentBlockHeaderIsActorFail);
 
-		return NULL;
+		return nullptr;
 	}
 
 	if (bIsActor)
@@ -4155,7 +4162,7 @@ UObject* UActorChannel::ReadContentBlockHeader(FInBunch& Bunch, bool& bObjectDel
 
 	// Note this heavily mirrors what happens in UPackageMapClient::SerializeNewActor
 	FNetworkGUID NetGUID;
-	UObject* SubObj = NULL;
+	UObject* SubObj = nullptr;
 
 	// Manually serialize the object so that we can get the NetGUID (in order to assign it if we spawn the object here)
 	Connection->PackageMap->SerializeObject(Bunch, UObject::StaticClass(), SubObj, &NetGUID);
@@ -4169,7 +4176,7 @@ UObject* UActorChannel::ReadContentBlockHeader(FInBunch& Bunch, bool& bObjectDel
 		Bunch.SetError();
 		AddToChainResultPtr(Bunch.ExtendedError, ENetCloseResult::ContentBlockHeaderObjFail);
 
-		return NULL;
+		return nullptr;
 	}
 
 	if (Bunch.AtEnd())
@@ -4179,21 +4186,21 @@ UObject* UActorChannel::ReadContentBlockHeader(FInBunch& Bunch, bool& bObjectDel
 		Bunch.SetError();
 		AddToChainResultPtr(Bunch.ExtendedError, ENetCloseResult::ContentBlockHeaderPrematureEnd);
 
-		return NULL;
+		return nullptr;
 	}
 
 	// Validate existing sub-object
-	if (SubObj != NULL)
+	if (SubObj)
 	{
 		// Sub-objects can't be actors (should just use an actor channel in this case)
-		if (Cast< AActor >(SubObj) != NULL)
+		if (Cast< AActor >(SubObj) != nullptr)
 		{
 			UE_LOG(LogNetTraffic, Error, TEXT("UActorChannel::ReadContentBlockHeader: Sub-object not allowed to be actor type. SubObj: %s, Actor: %s"), *SubObj->GetName(), *Actor->GetName());
 
 			Bunch.SetError();
 			AddToChainResultPtr(Bunch.ExtendedError, ENetCloseResult::ContentBlockHeaderSubObjectActor);
 
-			return NULL;
+			return nullptr;
 		}
 
 		// Sub-objects must reside within their actor parents
@@ -4207,7 +4214,7 @@ UObject* UActorChannel::ReadContentBlockHeader(FInBunch& Bunch, bool& bObjectDel
 				Bunch.SetError();
 				AddToChainResultPtr(Bunch.ExtendedError, ENetCloseResult::ContentBlockHeaderBadParent);
 
-				return NULL;
+				return nullptr;
 			}
 		}
 	}
@@ -4215,14 +4222,14 @@ UObject* UActorChannel::ReadContentBlockHeader(FInBunch& Bunch, bool& bObjectDel
 	if (IsServer)
 	{
 		// The server should never need to create sub objects
-		if (SubObj == NULL)
+		if (!SubObj)
 		{
 			UE_LOG(LogNetTraffic, Error, TEXT("ReadContentBlockHeader: Client attempted to create sub-object. Actor: %s"), *Actor->GetName());
 
 			Bunch.SetError();
 			AddToChainResultPtr(Bunch.ExtendedError, ENetCloseResult::ContentBlockHeaderInvalidCreate);
 
-			return NULL;
+			return nullptr;
 		}
 
 		return SubObj;
@@ -4236,13 +4243,13 @@ UObject* UActorChannel::ReadContentBlockHeader(FInBunch& Bunch, bool& bObjectDel
 
 		AddToChainResultPtr(Bunch.ExtendedError, ENetCloseResult::ContentBlockHeaderStablyNamedFail);
 
-		return NULL;
+		return nullptr;
 	}
 
 	if (bStablyNamed)
 	{
 		// If this is a stably named sub-object, we shouldn't need to create it. Don't raise a bunch error though because this may happen while a level is streaming out.
-		if (SubObj == NULL)
+		if (!SubObj)
 		{
 			// (ignore though if this is for replays)
 			if (!Connection->IsInternalAck())
@@ -4250,20 +4257,50 @@ UObject* UActorChannel::ReadContentBlockHeader(FInBunch& Bunch, bool& bObjectDel
 				UE_LOG(LogNetTraffic, Log, TEXT("ReadContentBlockHeader: Stably named sub-object not found. Its level may have streamed out. Component: %s, Actor: %s"), *Connection->Driver->GuidCache->FullNetGUIDPath(NetGUID), *Actor->GetName());
 			}
 
-			return NULL;
+			return nullptr;
 		}
 
 		return SubObj;
 	}
 
-	// Serialize the class in case we have to spawn it.
-	// Manually serialize the object so that we can get the NetGUID (in order to assign it if we spawn the object here)
-	FNetworkGUID ClassNetGUID;
-	UObject* SubObjClassObj = NULL;
-	Connection->PackageMap->SerializeObject(Bunch, UObject::StaticClass(), SubObjClassObj, &ClassNetGUID);
+	bool bDeleteSubObject = false;
+	bool bSerializeClass = true;
 
-	// Delete sub-object
-	if (!ClassNetGUID.IsValid())
+	if (Bunch.EngineNetVer() >= HISTORY_SUBOBJECT_DESTROY_FLAG)
+	{
+		const bool bIsDestroyMessage = Bunch.ReadBit() != 0 ? true : false;
+
+		if (bIsDestroyMessage)
+		{
+			bDeleteSubObject = true;
+			bSerializeClass = false;
+
+			uint8 DestroyFlags = 0;
+			Bunch << DestroyFlags;  // Destroy flags to be defined later
+
+			NET_CHECKSUM(Bunch);
+		}
+	}
+	else
+	{
+		bSerializeClass = true;
+	}
+	
+
+	UObject* SubObjClassObj = nullptr;
+	if (bSerializeClass)
+	{
+		// Serialize the class in case we have to spawn it.
+		// Manually serialize the object so that we can get the NetGUID (in order to assign it if we spawn the object here)
+		FNetworkGUID ClassNetGUID;
+		Connection->PackageMap->SerializeObject(Bunch, UObject::StaticClass(), SubObjClassObj, &ClassNetGUID);
+
+		// When ClassNetGUID is empty it means the sender requested to delete the subobject
+		bDeleteSubObject = !ClassNetGUID.IsValid();
+	}
+	
+
+	if (bDeleteSubObject)
 	{
 		if (SubObj)
 		{
@@ -4284,17 +4321,17 @@ UObject* UActorChannel::ReadContentBlockHeader(FInBunch& Bunch, bool& bObjectDel
 			SubObj->MarkAsGarbage();
 		}
 		bObjectDeleted = true;
-		return NULL;
+		return nullptr;
 	}
 
 	UClass* SubObjClass = Cast< UClass >(SubObjClassObj);
 
-	if (SubObjClass == NULL)
+	if (!SubObjClass)
 	{
 		UE_LOG(LogNetTraffic, Warning, TEXT("UActorChannel::ReadContentBlockHeader: Unable to read sub-object class. Actor: %s"), *Actor->GetName());
 
 		// Valid NetGUID but no class was resolved - this is an error
-		if (SubObj == NULL)
+		if (!SubObj)
 		{
 			// (unless we're using replays, which could be backwards compatibility kicking in)
 			if (!Connection->IsInternalAck())
@@ -4303,7 +4340,7 @@ UObject* UActorChannel::ReadContentBlockHeader(FInBunch& Bunch, bool& bObjectDel
 
 				Bunch.SetError();
 				AddToChainResultPtr(Bunch.ExtendedError, ENetCloseResult::ContentBlockHeaderNoSubObjectClass);
-				return NULL;
+				return nullptr;
 			}
 		}
 	}
@@ -4316,7 +4353,7 @@ UObject* UActorChannel::ReadContentBlockHeader(FInBunch& Bunch, bool& bObjectDel
 			Bunch.SetError();
 			AddToChainResultPtr(Bunch.ExtendedError, ENetCloseResult::ContentBlockHeaderUObjectSubObject);
 
-			return NULL;
+			return nullptr;
 		}
 
 		if (SubObjClass->IsChildOf(AActor::StaticClass()))
@@ -4326,12 +4363,12 @@ UObject* UActorChannel::ReadContentBlockHeader(FInBunch& Bunch, bool& bObjectDel
 			Bunch.SetError();
 			AddToChainResultPtr(Bunch.ExtendedError, ENetCloseResult::ContentBlockHeaderAActorSubObject);
 
-			return NULL;
+			return nullptr;
 		}
 	}
 
 	UObject* ObjOuter = Actor;
-	if (Connection->EngineNetworkProtocolVersion >= HISTORY_SUBOBJECT_OUTER_CHAIN)
+	if (Bunch.EngineNetVer() >= HISTORY_SUBOBJECT_OUTER_CHAIN)
 	{
 		const bool bActorIsOuter = Bunch.ReadBit() != 0;
 
@@ -4369,7 +4406,7 @@ UObject* UActorChannel::ReadContentBlockHeader(FInBunch& Bunch, bool& bObjectDel
 		return nullptr;
 	}
 
-	if ( SubObj == NULL )
+	if (!SubObj)
 	{
 		check( !IsServer );
 
@@ -4379,9 +4416,9 @@ UObject* UActorChannel::ReadContentBlockHeader(FInBunch& Bunch, bool& bObjectDel
 		SubObj = NewObject< UObject >(ObjOuter, SubObjClass);
 
 		// Sanity check some things
-		checkf(SubObj != NULL, TEXT("UActorChannel::ReadContentBlockHeader: Subobject is NULL after instantiating. Class: %s, Actor %s"), *GetNameSafe(SubObjClass), *Actor->GetName());
+		checkf(SubObj != nullptr, TEXT("UActorChannel::ReadContentBlockHeader: Subobject is NULL after instantiating. Class: %s, Actor %s"), *GetNameSafe(SubObjClass), *Actor->GetName());
 		checkf(SubObj->IsIn(ObjOuter), TEXT("UActorChannel::ReadContentBlockHeader: Subobject is not in Outer. SubObject: %s, Actor: %s, Outer: %s"), *SubObj->GetName(), *Actor->GetName(), *ObjOuter->GetName());
-		checkf(Cast< AActor >(SubObj) == NULL, TEXT("UActorChannel::ReadContentBlockHeader: Subobject is an Actor. SubObject: %s, Actor: %s"), *SubObj->GetName(), *Actor->GetName());
+		checkf(Cast< AActor >(SubObj) == nullptr, TEXT("UActorChannel::ReadContentBlockHeader: Subobject is an Actor. SubObject: %s, Actor: %s"), *SubObj->GetName(), *Actor->GetName());
 
 		// Notify actor that we created a component from replication
 		Actor->OnSubobjectCreatedFromReplication( SubObj );
