Added modulus for 64-bit ints.

diff --git a/Engine/Source/Runtime/Engine/Classes/Kismet/KismetMathLibrary.h b/Engine/Source/Runtime/Engine/Classes/Kismet/KismetMathLibrary.h
--- a/Engine/Source/Runtime/Engine/Classes/Kismet/KismetMathLibrary.h
+++ b/Engine/Source/Runtime/Engine/Classes/Kismet/KismetMathLibrary.h
@@ -407,6 +407,10 @@ class ENGINE_API UKismetMathLibrary : public UBlueprintFunctionLibrary
 	UFUNCTION(BlueprintPure, meta=(DisplayName = "integer64 / integer64", CompactNodeTitle = "/", Keywords = "/ divide division"), Category="Math|Integer64")
 	static int64 Divide_Int64Int64(int64 A, int64 B = 1);
 	
+	/** Modulo (A % B) */
+	UFUNCTION(BlueprintPure, meta = (DisplayName = "% (integer64)", CompactNodeTitle = "%", Keywords = "% modulus"), Category = "Math|Integer64")
+	static int64 Percent_Int64Int64(int64 A, int64 B = 1);
+
 	/** Addition (A + B) */
 	UFUNCTION(BlueprintPure, meta=(DisplayName = "integer64 + integer64", CompactNodeTitle = "+", Keywords = "+ add plus", CommutativeAssociativeBinaryOperator = "true"), Category="Math|Integer64")
 	static int64 Add_Int64Int64(int64 A, int64 B = 1);
@@ -4313,6 +4317,7 @@ private:
 	static void ReportError_Divide_DoubleDouble();
 	static void ReportError_Divide_Int64Int64();
 	static void ReportError_Percent_IntInt();
+	static void ReportError_Percent_Int64Int64();
 	static void ReportError_Sqrt();
 	static void ReportError_Divide_VectorFloat();
 	static void ReportError_Divide_VectorInt();
diff --git a/Engine/Source/Runtime/Engine/Classes/Kismet/KismetMathLibrary.inl b/Engine/Source/Runtime/Engine/Classes/Kismet/KismetMathLibrary.inl
--- a/Engine/Source/Runtime/Engine/Classes/Kismet/KismetMathLibrary.inl
+++ b/Engine/Source/Runtime/Engine/Classes/Kismet/KismetMathLibrary.inl
@@ -457,6 +457,18 @@ int32 UKismetMathLibrary::Percent_IntInt(int32 A, int32 B)
 	return (A % B);
 }
 
+KISMET_MATH_FORCEINLINE
+int64 UKismetMathLibrary::Percent_Int64Int64(int64 A, int64 B)
+{
+	if (B == 0)
+	{
+		ReportError_Percent_Int64Int64();
+		return 0;
+	}
+
+	return (A % B);
+}
+
 KISMET_MATH_FORCEINLINE
 double UKismetMathLibrary::Fraction(double A)
 {
diff --git a/Engine/Source/Runtime/Engine/Private/KismetMathLibrary.cpp b/Engine/Source/Runtime/Engine/Private/KismetMathLibrary.cpp
--- a/Engine/Source/Runtime/Engine/Private/KismetMathLibrary.cpp
+++ b/Engine/Source/Runtime/Engine/Private/KismetMathLibrary.cpp
@@ -78,7 +78,7 @@ void UKismetMathLibrary::ReportError_Divide_ByteByte()
 
 void UKismetMathLibrary::ReportError_Percent_ByteByte()
 {
-	FFrame::KismetExecutionMessage(TEXT("Modulo by zero"), ELogVerbosity::Warning, DivideByZeroWarning);
+	FFrame::KismetExecutionMessage(TEXT("Modulo by zero: Percent_ByteByte"), ELogVerbosity::Warning, DivideByZeroWarning);
 }
 
 void UKismetMathLibrary::ReportError_Divide_IntInt()
@@ -103,7 +103,12 @@ void UKismetMathLibrary::ReportError_Divide_Int64Int64()
 
 void UKismetMathLibrary::ReportError_Percent_IntInt()
 {
-	FFrame::KismetExecutionMessage(TEXT("Modulo by zero"), ELogVerbosity::Warning, DivideByZeroWarning);
+	FFrame::KismetExecutionMessage(TEXT("Modulo by zero: Percent_IntInt"), ELogVerbosity::Warning, DivideByZeroWarning);
+}
+
+void UKismetMathLibrary::ReportError_Percent_Int64Int64()
+{
+	FFrame::KismetExecutionMessage(TEXT("Modulo by zero: Percent_Int64Int64"), ELogVerbosity::Warning, DivideByZeroWarning);
 }
 
 void UKismetMathLibrary::ReportError_Sqrt()
