Command line: Map: Allowed more flexibility when providing the map override parameter.

Previously, the override map name was expected to be the first positional parameter (or the second one when the game name/project was specified as the first positional parameter). It created problems because we often combine command-line parameters with additional ones from uecommandline.txt and in some cases the order could not be guaranteed.

It was the only configuration that worked:

MyGame.exe ../../../MyGame/MyGame.uproject MapToLoad

Now, we're looking for the first paramter without the leading '-' (it handles the existing use case with the added flexibility). Alternatively, we may now use a named parameter -map=MapToLoad. If both a named and positional parameters are provided, the first one is used.

Now, all the versions below will work:

MyGame.exe ../../../MyGame/MyGame.uproject MapToLoad
MyGame.exe ../../../MyGame/MyGame.uproject -foo MapToLoad
MyGame.exe ../../../MyGame/MyGame.uproject -map=MapToLoad
MyGame.exe ../../../MyGame/MyGame.uproject -foo -map=MapToLoad
MyGame.exe ../../../MyGame/MyGame.uproject -map=MapToLoadABC MapToLoadXYZ     (MapToLoadABC will be loaded)

#rn minor core Allowed more flexibility when parsing the initial map name on the command line. Previously, it was expected to be the first positional parameter (the second if a project name was provided). Now, we're looking for the first non-switch parameter (one not starting with '-'). It allows more launch scenarios to work correctly e.g. when mixing parameters provided to the debugger with ones from uecommandline.txt. Alternatively, map name can be provided as a named parameter -map=.

diff --git a/Engine/Source/Runtime/Engine/Classes/Engine/GameInstance.h b/Engine/Source/Runtime/Engine/Classes/Engine/GameInstance.h
--- a/Engine/Source/Runtime/Engine/Classes/Engine/GameInstance.h
+++ b/Engine/Source/Runtime/Engine/Classes/Engine/GameInstance.h
@@ -650,6 +650,9 @@ protected:
 	/** Called when the game instance is started either normally or through PIE. */
 	virtual void OnStart();
 
+	/** Find a map override argument on the command-line string (the first argument without a leading '-' or -map=..., whichever comes first). */
+	static bool GetMapOverrideName(const TCHAR* CmdLine, FString& OverrideMapName);
+
 private:
 
 	FObjectSubsystemCollection<UGameInstanceSubsystem> SubsystemCollection;
diff --git a/Engine/Source/Runtime/Engine/Private/GameInstance.cpp b/Engine/Source/Runtime/Engine/Private/GameInstance.cpp
--- a/Engine/Source/Runtime/Engine/Private/GameInstance.cpp
+++ b/Engine/Source/Runtime/Engine/Private/GameInstance.cpp
@@ -622,7 +622,7 @@ void UGameInstance::StartGameInstance()
 	const FString& DefaultMap = GameMapsSettings->GetGameDefaultMap();
 
 	FString PackageName;
-	if (!FParse::Token(Tmp, PackageName, 0) || **PackageName == '-')
+	if (!GetMapOverrideName(Tmp, PackageName))
 	{
 		PackageName = DefaultMap + GameMapsSettings->LocalMapOptions;
 	}
@@ -687,6 +687,34 @@ void UGameInstance::OnStart()
 
 }
 
+bool UGameInstance::GetMapOverrideName(const TCHAR* CmdLine, FString& OverrideMapName)
+{
+	const TCHAR* ParsedCmdLine = CmdLine;
+
+	while (*ParsedCmdLine)
+	{
+		FString Token = FParse::Token(ParsedCmdLine, 0);
+
+		if (Token.IsEmpty())
+		{
+			continue;
+		}
+
+		if (Token[0] != TCHAR('-'))
+		{
+			OverrideMapName = Token;
+
+			return true;
+		}
+		else if (FParse::Value(*Token, TEXT("-map="), OverrideMapName))
+		{
+			return true;
+		}
+	}
+
+	return false;
+}
+
 bool UGameInstance::HandleOpenCommand(const TCHAR* Cmd, FOutputDevice& Ar, UWorld* InWorld)
 {
 	check(WorldContext && WorldContext->World() == InWorld);
