float/double -> FVector2D implicit promotion in blueprints

diff --git a/Engine/Source/Runtime/Engine/Classes/Kismet/KismetMathLibrary.h b/Engine/Source/Runtime/Engine/Classes/Kismet/KismetMathLibrary.h
--- a/Engine/Source/Runtime/Engine/Classes/Kismet/KismetMathLibrary.h
+++ b/Engine/Source/Runtime/Engine/Classes/Kismet/KismetMathLibrary.h
@@ -3604,13 +3604,21 @@ class ENGINE_API UKismetMathLibrary : public UBlueprintFunctionLibrary
 	static FVector Conv_IntVectorToVector(const FIntVector& InIntVector);
 
 	/** Convert a float into a vector, where each element is that float */
-	UFUNCTION(BlueprintPure, meta=(DisplayName = "To Vector (Float)", CompactNodeTitle = "->", Keywords="cast convert", BlueprintAutocast), Category="Math|Conversions")
+	UFUNCTION(BlueprintPure, BlueprintInternalUseOnly, meta=(DisplayName = "Single Precision To Vector", CompactNodeTitle = "->", Keywords="cast convert", BlueprintAutocast), Category="Math|Conversions")
 	static FVector Conv_FloatToVector(float InFloat);
 
 	/** Convert a double into a vector, where each element is that double */
-	UFUNCTION(BlueprintPure, meta = (DisplayName = "To Vector (Double)", CompactNodeTitle = "->", Keywords = "cast convert", BlueprintAutocast), Category = "Math|Conversions")
+	UFUNCTION(BlueprintPure, meta = (DisplayName = "To Vector", CompactNodeTitle = "->", Keywords = "cast convert", BlueprintAutocast), Category = "Math|Conversions")
 	static FVector Conv_DoubleToVector(double InDouble);
 
+	/** Convert a float into a vector, where each element is that float */
+	UFUNCTION(BlueprintPure, BlueprintInternalUseOnly, meta=(DisplayName = "Single Precision To Vector2D", CompactNodeTitle = "->", Keywords="cast convert", BlueprintAutocast), Category="Math|Conversions")
+	static FVector2D Conv_FloatToVector2D(float InFloat);
+
+	/** Convert a float into a vector, where each element is that float */
+	UFUNCTION(BlueprintPure, meta=(DisplayName = "To Vector2D", CompactNodeTitle = "->", Keywords="cast convert", BlueprintAutocast), Category="Math|Conversions")
+	static FVector2D Conv_DoubleToVector2D(double InDouble);
+	
 	//
 	// Box functions
 	//
diff --git a/Engine/Source/Runtime/Engine/Classes/Kismet/KismetMathLibrary.inl b/Engine/Source/Runtime/Engine/Classes/Kismet/KismetMathLibrary.inl
--- a/Engine/Source/Runtime/Engine/Classes/Kismet/KismetMathLibrary.inl
+++ b/Engine/Source/Runtime/Engine/Classes/Kismet/KismetMathLibrary.inl
@@ -2862,6 +2862,18 @@ FVector UKismetMathLibrary::Conv_DoubleToVector(double InDouble)
 	return FVector(InDouble);
 }
 
+KISMET_MATH_FORCEINLINE
+FVector2D UKismetMathLibrary::Conv_FloatToVector2D(float InFloat)
+{
+	return FVector2D(InFloat);
+}
+
+KISMET_MATH_FORCEINLINE
+FVector2D UKismetMathLibrary::Conv_DoubleToVector2D(double InDouble)
+{
+	return FVector2D(InDouble);
+}
+
 KISMET_MATH_FORCEINLINE
 FLinearColor UKismetMathLibrary::Conv_FloatToLinearColor(float InFloat)
 {
