Implement support for adding textures with different aspect ratio to a Texture2DArray.

diff --git a/Engine/Source/Runtime/Engine/Classes/Engine/Texture2DArray.cpp b/Engine/Source/Runtime/Engine/Classes/Engine/Texture2DArray.cpp
--- a/Engine/Source/Runtime/Engine/Classes/Engine/Texture2DArray.cpp
+++ b/Engine/Source/Runtime/Engine/Classes/Engine/Texture2DArray.cpp
@@ -325,35 +325,14 @@ EPixelFormat UTexture2DArray::GetPixelFormat() const
 
 ENGINE_API bool UTexture2DArray::CheckArrayTexturesCompatibility()
 {
-	if (!SourceTextures.Num())
+	for (UTexture2D* SourceTexture : SourceTextures)
 	{
-		return true;
-	}
-
-	if (!SourceTextures[0])
-	{
-		// Do not create array till all texture slots are filled.
-		return false;
-	}
-
-	const FTextureSource& BaseTextureSource = SourceTextures[0]->Source;
-	const FString BaseTextureName = SourceTextures[0]->GetFName().ToString();
-
-	for (int32 TextureIndex = 1; TextureIndex < SourceTextures.Num(); ++TextureIndex)
-	{
-		if (!SourceTextures[TextureIndex])
+		if (!SourceTexture)
 		{
 			// Do not create array till all texture slots are filled.
 			return false;
 		}
-
-		if (SourceTextures[TextureIndex]->Source.GetSizeX() * BaseTextureSource.GetSizeY() != SourceTextures[TextureIndex]->Source.GetSizeY() * BaseTextureSource.GetSizeX())
-		{
-			UE_LOG(LogTexture, Warning, TEXT("Texture2DArray creation failed. Textures %s and %s have different aspect ratios."), *SourceTextures[TextureIndex]->GetFName().ToString(), *BaseTextureName);
-			return false;
-		}
 	}
-
 	return true;
 }
 
@@ -369,49 +348,61 @@ ENGINE_API bool UTexture2DArray::UpdateSourceFromSourceTextures(bool bCreatingNe
 	{
 		Modify();
 
-		ETextureSourceFormat SourceFormat = SourceTextures[0]->Source.Format;
-		bool bUseSRGBSource = SourceTextures[0]->Source.GetGammaSpace(0) == EGammaSpace::sRGB;
+		int32 SizeX = SourceTextures[0]->Source.GetSizeX();
+		int32 SizeY = SourceTextures[0]->Source.GetSizeY();
+		ETextureSourceFormat Format = SourceTextures[0]->Source.Format;
+		bool bSRGB = SourceTextures[0]->Source.GetGammaSpace(0) == EGammaSpace::sRGB;
+
+		bool bMismatchedSize = false;
+		bool bMismatchedAspectRatio = false;
 		bool bMismatchedFormats = false;
 		bool bMismatchedGammaSpace = false;
-		int32 InitialSourceTextureIndex = 0;
+
 		for (int32 SourceTextureIndex = 1; SourceTextureIndex < SourceTextures.Num(); ++SourceTextureIndex)
 		{
-			// Dimensions of the texture array source are calculated based on the dimensions of the largest texture source.
-			// It is sufficient to compare SizeX because CheckArrayTexturesCompatibility() guarantees that all sources have the same aspect ratio.
-			if (SourceTextures[SourceTextureIndex]->Source.GetSizeX() > SourceTextures[InitialSourceTextureIndex]->Source.GetSizeX())
+			if (SourceTextures[SourceTextureIndex]->Source.GetSizeX() != SizeX || SourceTextures[SourceTextureIndex]->Source.GetSizeY() != SizeY)
 			{
-				InitialSourceTextureIndex = SourceTextureIndex;
+				bMismatchedSize = true;
+				if (SourceTextures[SourceTextureIndex]->Source.GetSizeX() * SizeY != SourceTextures[SourceTextureIndex]->Source.GetSizeY() * SizeX)
+				{
+					bMismatchedAspectRatio = true;
+				}
+				// Dimensions of the texture array source are set to the maximum SizeX and SizeY among the array element's sources in order to minimize quality loss.
+				SizeX = FMath::Max(SizeX, SourceTextures[SourceTextureIndex]->Source.GetSizeX());
+				SizeY = FMath::Max(SizeY, SourceTextures[SourceTextureIndex]->Source.GetSizeY());
 			}
 
-			if ((SourceTextures[SourceTextureIndex]->Source.GetGammaSpace(0) == EGammaSpace::sRGB) != bUseSRGBSource)
+			if ((SourceTextures[SourceTextureIndex]->Source.GetGammaSpace(0) == EGammaSpace::sRGB) != bSRGB)
 			{
 				bMismatchedGammaSpace = true;
-				bUseSRGBSource = false;
-				SourceFormat = TSF_RGBA32F;
+				bSRGB = false;
+				Format = TSF_RGBA32F;
 			}
 
-			if (SourceFormat != SourceTextures[SourceTextureIndex]->Source.Format)
+			if (Format != SourceTextures[SourceTextureIndex]->Source.Format)
 			{
 				bMismatchedFormats = true;
-				SourceFormat = FImageCoreUtils::GetCommonSourceFormat(SourceFormat, SourceTextures[SourceTextureIndex]->Source.Format);
+				Format = FImageCoreUtils::GetCommonSourceFormat(Format, SourceTextures[SourceTextureIndex]->Source.Format);
 			}
 		}
 
+		if (bMismatchedSize)
+		{
+			UE_LOG(LogTexture, Log, TEXT("Mismatched sizes of the source textures, resizing all the array elements to %dx%d%s ..."),
+				SizeX, SizeY, bMismatchedAspectRatio ? TEXT(" (this will also affect aspect ratio of some source textures)") : TEXT(""));
+		}
+
 		if (bMismatchedGammaSpace)
 		{
-			UE_LOG(LogTexture, Warning, TEXT("Mismatched source gamma spaces, converting all to RGBA32F/Linear ..."));
+			UE_LOG(LogTexture, Log, TEXT("Mismatched source gamma spaces, converting all to RGBA32F/Linear ..."));
 		}
 		else if (bMismatchedFormats)
 		{
-			UE_LOG(LogTexture, Warning, TEXT("Mismatched source pixel formats, converting all to %s/%s ..."),
-				ERawImageFormat::GetName(FImageCoreUtils::ConvertToRawImageFormat(SourceFormat)), bUseSRGBSource ? TEXT("sRGB") : TEXT("Linear"));
+			UE_LOG(LogTexture, Log, TEXT("Mismatched source pixel formats, converting all to %s/%s ..."),
+				ERawImageFormat::GetName(FImageCoreUtils::ConvertToRawImageFormat(Format)), bSRGB ? TEXT("sRGB") : TEXT("Linear"));
 		}
 
-		FTextureSource& InitialSource = SourceTextures[InitialSourceTextureIndex]->Source;
-		int32 FormatDataSize = FTextureSource::GetBytesPerPixel(SourceFormat);
-		// X,Y,Z size of the array. Switched to use Source sizes - previously would use PlatformData sizes which caused an error in headless commandlet MHC export.
-		int32 SizeX = InitialSource.GetSizeX();
-		int32 SizeY = InitialSource.GetSizeY();
+		int32 FormatDataSize = FTextureSource::GetBytesPerPixel(Format);
 		uint32 ArraySize = SourceTextures.Num();
 
 		// This should be false when texture is updated to avoid overriding user settings.
@@ -424,11 +415,11 @@ ENGINE_API bool UTexture2DArray::UpdateSourceFromSourceTextures(bool bCreatingNe
 			NeverStream = true;
 		}
 
-		SRGB = bUseSRGBSource;
+		SRGB = bSRGB;
 
 		// Create the source texture for this Texture2DArray.
 		// Currently only single-mip Texture2DArray's are supported, therefore only the first mip is copied from each element source.
-		Source.Init(SizeX, SizeY, ArraySize, 1, SourceFormat);
+		Source.Init(SizeX, SizeY, ArraySize, 1, Format);
 
 		uint8* DestMipData = Source.LockMip(0);
 		int64 MipSizeBytes = Source.CalcMipSize(0) / ArraySize;
@@ -437,7 +428,7 @@ ENGINE_API bool UTexture2DArray::UpdateSourceFromSourceTextures(bool bCreatingNe
 		{
 			FTextureSource& TextureSource = SourceTextures[SourceTexIndex]->Source;
 			void* DestSliceData = DestMipData + MipSizeBytes * SourceTexIndex;
-			if (TextureSource.SizeX == SizeX && TextureSource.SizeY == SizeY && TextureSource.Format == SourceFormat)
+			if (TextureSource.SizeX == SizeX && TextureSource.SizeY == SizeY && TextureSource.Format == Format)
 			{
 				void* SourceData = TextureSource.LockMip(0);
 				check(TextureSource.CalcMipSize(0) == MipSizeBytes);
@@ -449,11 +440,11 @@ ENGINE_API bool UTexture2DArray::UpdateSourceFromSourceTextures(bool bCreatingNe
 				FImage SourceImage;
 				SourceTextures[SourceTexIndex]->Source.GetMipImage(SourceImage, 0);
 
-				if (TextureSource.Format != SourceFormat)
+				if (TextureSource.Format != Format)
 				{
 					// note: there is no need to check if the gamma space is different, because in such case we would fallback to TSF_RGBA32F, which is always linear
 					FImage ConvertedSourceImage;
-					SourceImage.CopyTo(ConvertedSourceImage, FImageCoreUtils::ConvertToRawImageFormat(SourceFormat), bUseSRGBSource ? EGammaSpace::sRGB : EGammaSpace::Linear);
+					SourceImage.CopyTo(ConvertedSourceImage, FImageCoreUtils::ConvertToRawImageFormat(Format), bSRGB ? EGammaSpace::sRGB : EGammaSpace::Linear);
 					ConvertedSourceImage.Swap(SourceImage);
 				}
 
