Fix for incorrect check().

diff --git a/Engine/Source/Runtime/Engine/Private/TextureCube.cpp b/Engine/Source/Runtime/Engine/Private/TextureCube.cpp
--- a/Engine/Source/Runtime/Engine/Private/TextureCube.cpp
+++ b/Engine/Source/Runtime/Engine/Private/TextureCube.cpp
@@ -461,7 +461,7 @@ public:
 
 		// Read the mip-levels into the RHI texture.
 		int32 NumMips = Owner->GetNumMips();
-		check(NumMips < MAX_TEXTURE_MIP_COUNT);
+		check(NumMips <= MAX_TEXTURE_MIP_COUNT);
 		for( int32 FaceIndex=0; FaceIndex<6; FaceIndex++ )
 		{
 			for(int32 MipIndex=0; MipIndex < NumMips; MipIndex++)
