Fixed 'bad machine ptr' warnings when compiling animation BPs with certain layer/state machine configurations

diff --git a/Engine/Source/Runtime/Engine/Private/Animation/AnimNode_LinkedAnimLayer.cpp b/Engine/Source/Runtime/Engine/Private/Animation/AnimNode_LinkedAnimLayer.cpp
--- a/Engine/Source/Runtime/Engine/Private/Animation/AnimNode_LinkedAnimLayer.cpp
+++ b/Engine/Source/Runtime/Engine/Private/Animation/AnimNode_LinkedAnimLayer.cpp
@@ -2,6 +2,7 @@
 
 #include "Animation/AnimNode_LinkedAnimLayer.h"
 #include "Animation/AnimClassInterface.h"
+#include "Animation/AnimInstanceProxy.h"
 
 #include UE_INLINE_GENERATED_CPP_BY_NAME(AnimNode_LinkedAnimLayer)
 
@@ -158,7 +159,10 @@ void FAnimNode_LinkedAnimLayer::HandleObjectsReinstanced_Impl(UObject* InSourceO
 		DynamicUnlink(SourceAnimInstance);
 		DynamicLink(SourceAnimInstance);
 
-		FAnimationInitializeContext Context(&SourceAnimInstance->GetProxyOnAnyThread<FAnimInstanceProxy>());
+		FAnimInstanceProxy& SourceProxy = SourceAnimInstance->GetProxyOnAnyThread<FAnimInstanceProxy>();
+		SourceProxy.InitializeCachedClassData();
+		
+		FAnimationInitializeContext Context(&SourceProxy);
 		InitializeSubGraph_AnyThread(Context);
 	}
 }
diff --git a/Engine/Source/Runtime/Engine/Public/Animation/AnimInstanceProxy.h b/Engine/Source/Runtime/Engine/Public/Animation/AnimInstanceProxy.h
--- a/Engine/Source/Runtime/Engine/Public/Animation/AnimInstanceProxy.h
+++ b/Engine/Source/Runtime/Engine/Public/Animation/AnimInstanceProxy.h
@@ -555,6 +555,7 @@ public:
 	friend class UAnimSingleNodeInstance;
 	friend class USkeletalMeshComponent;
 	friend struct FAnimNode_LinkedAnimGraph;
+	friend struct FAnimNode_LinkedAnimLayer;
 	friend struct FAnimationBaseContext;
 	friend struct FAnimTrace;
 	friend struct UE::Anim::FAnimSync;
