#UE Give CheatManagerExtensions Add/Remove events so they can listen to delegates

diff --git a/Engine/Source/Runtime/Engine/Classes/GameFramework/CheatManager.h b/Engine/Source/Runtime/Engine/Classes/GameFramework/CheatManager.h
--- a/Engine/Source/Runtime/Engine/Classes/GameFramework/CheatManager.h
+++ b/Engine/Source/Runtime/Engine/Classes/GameFramework/CheatManager.h
@@ -82,6 +82,12 @@ public:
 
 	UFUNCTION(BlueprintPure, Category = "Cheat Manager")
 	APlayerController* GetPlayerController() const;
+
+	UFUNCTION(BlueprintImplementableEvent, meta=(DisplayName = "Added To Cheat Manager", Keywords="Begin Play"))
+	void AddedToCheatManager();
+
+	UFUNCTION(BlueprintImplementableEvent, meta = (DisplayName = "Removed From Cheat Manager", Keywords = "End Play"))
+	void RemovedFromCheatManager();
 };
 
 /** 
diff --git a/Engine/Source/Runtime/Engine/Private/CheatManager.cpp b/Engine/Source/Runtime/Engine/Private/CheatManager.cpp
--- a/Engine/Source/Runtime/Engine/Private/CheatManager.cpp
+++ b/Engine/Source/Runtime/Engine/Private/CheatManager.cpp
@@ -81,6 +81,10 @@ UCheatManager::UCheatManager(const FObjectInitializer& ObjectInitializer)
 
 void UCheatManager::OnPlayerEndPlayed(AActor* Player, EEndPlayReason::Type EndPlayReason)
 {
+	for (UCheatManagerExtension* CheatExtension : CheatManagerExtensions)
+	{
+		CheatExtension->RemovedFromCheatManager();
+	}
 	CheatManagerExtensions.Empty();
 }
 
@@ -1369,7 +1373,11 @@ void UCheatManager::AddCheatManagerExtension(UCheatManagerExtension* CheatObject
 {
 	if (ensure(CheatObject))
 	{
-		CheatManagerExtensions.AddUnique(CheatObject);
+		if (!CheatManagerExtensions.Contains(CheatObject))
+		{
+			CheatManagerExtensions.Add(CheatObject);
+			CheatObject->AddedToCheatManager();
+		}
 	}
 }
 
@@ -1377,7 +1385,12 @@ void UCheatManager::RemoveCheatManagerExtension(UCheatManagerExtension* CheatObj
 {
 	if (ensure(CheatObject))
 	{
-		CheatManagerExtensions.Remove(CheatObject);
+		int32 CheatExtensionIdx = CheatManagerExtensions.IndexOfByKey(CheatObject);
+		if (CheatExtensionIdx != INDEX_NONE)
+		{
+			CheatManagerExtensions.RemoveAt(CheatExtensionIdx);
+			CheatObject->RemovedFromCheatManager();
+		}
 	}
 }
 
