PlayerController self register as a streaming source

diff --git a/Engine/Source/Runtime/Engine/Classes/GameFramework/PlayerController.h b/Engine/Source/Runtime/Engine/Classes/GameFramework/PlayerController.h
--- a/Engine/Source/Runtime/Engine/Classes/GameFramework/PlayerController.h
+++ b/Engine/Source/Runtime/Engine/Classes/GameFramework/PlayerController.h
@@ -26,6 +26,7 @@
 #include "GenericPlatform/IInputInterface.h"
 #include "GameFramework/PlayerInput.h"
 #include "Physics/AsyncPhysicsData.h"
+#include "WorldPartition/WorldPartitionStreamingSource.h"
 #include "PlayerController.generated.h"
 
 class ACameraActor;
@@ -236,7 +237,7 @@ struct FAsyncPhysicsTimestamp
  * @see https://docs.unrealengine.com/latest/INT/Gameplay/Framework/Controller/PlayerController/
  */
 UCLASS(config=Game, BlueprintType, Blueprintable, meta=(ShortTooltip="A Player Controller is an actor responsible for controlling a Pawn used by the player."))
-class ENGINE_API APlayerController : public AController
+class ENGINE_API APlayerController : public AController, public IWorldPartitionStreamingSourceProvider
 {
 	GENERATED_BODY()
 
@@ -561,6 +562,14 @@ public:
 	UPROPERTY(EditAnywhere, BlueprintReadWrite, Category = WorldPartition, meta=(EditCondition="bEnableStreamingSource"))
 	EStreamingSourcePriority StreamingSourcePriority;
 
+	/** Color used for debugging. */
+	UPROPERTY(EditAnywhere, BlueprintReadWrite, Category = WorldPartition, meta = (EditCondition = "bEnableStreamingSource"))
+	FColor StreamingSourceDebugColor;
+
+	/** Optional aggregated shape list used to build a custom shape for the streaming source. When empty, fallbacks sphere shape with a radius equal to grid's loading range. */
+	UPROPERTY(EditAnywhere, BlueprintReadWrite, Category = WorldPartition, meta = (EditCondition = "bEnableStreamingSource"))
+	TArray<FStreamingSourceShape> Shapes;
+
 	/** Scale applied to force feedback values */
 	UPROPERTY(config)
 	float ForceFeedbackScale;
@@ -785,6 +794,12 @@ public:
 	UFUNCTION(BlueprintCallable, Category = WorldPartition)
 	virtual EStreamingSourcePriority GetStreamingSourcePriority() const { return StreamingSourcePriority; }
 
+	/**
+	 * Gets the PlayerController's streaming source
+	 * @return the streaming source.
+	 */
+	virtual bool GetStreamingSource(FWorldPartitionStreamingSource& OutStreamingSource) const final;
+
 protected:
 	/** Pawn has been possessed, so changing state to NAME_Playing. Start it walking and begin playing with it. */
 	virtual void BeginPlayingState();
@@ -839,6 +854,16 @@ public:
 	 */
 	virtual void PostSeamlessTravel();
 
+	/**
+	 * Called when player controller gets added to its owning world player controller list. 
+	 */
+	void OnAddedToPlayerControllerList();
+
+	/**
+	 * Called when player controller gets removed from its owning world player controller list.
+	 */
+	void OnRemovedFromPlayerControllerList();
+
 	/** 
 	 * Tell the client to enable or disable voice chat (not muting)
 	 * @param bEnable enable or disable voice chat
diff --git a/Engine/Source/Runtime/Engine/Private/PlayerController.cpp b/Engine/Source/Runtime/Engine/Private/PlayerController.cpp
--- a/Engine/Source/Runtime/Engine/Private/PlayerController.cpp
+++ b/Engine/Source/Runtime/Engine/Private/PlayerController.cpp
@@ -68,7 +68,8 @@
 #include "GameMapsSettings.h"
 #include "Particles/EmitterCameraLensEffectBase.h"
 #include "LevelUtils.h"
-#include "WorldPartition/WorldPartitionStreamingSource.h"
+#include "WorldPartition/WorldPartition.h"
+#include "WorldPartition/WorldPartitionSubsystem.h"
 #include "Physics/AsyncPhysicsInputComponent.h"
 #include "GenericPlatform/GenericPlatformInputDeviceMapper.h"
 #include "PBDRigidsSolver.h"
@@ -3597,6 +3598,48 @@ void APlayerController::PostSeamlessTravel()
 
 }
 
+void APlayerController::OnAddedToPlayerControllerList()
+{
+	UWorld* World = GetWorld();
+	UWorldPartitionSubsystem* WorldPartitionSubsystem = World->GetSubsystem<UWorldPartitionSubsystem>();
+	if (ensure(WorldPartitionSubsystem))
+	{
+		WorldPartitionSubsystem->RegisterStreamingSourceProvider(this);
+	}
+}
+
+void APlayerController::OnRemovedFromPlayerControllerList()
+{
+	if (UWorld* World = GetWorld())
+	{
+		if (UWorldPartitionSubsystem* WorldPartitionSubsystem = World->GetSubsystem<UWorldPartitionSubsystem>())
+		{
+			verify(WorldPartitionSubsystem->UnregisterStreamingSourceProvider(this));
+		}
+	}
+}
+
+bool APlayerController::GetStreamingSource(FWorldPartitionStreamingSource& OutStreamingSource) const
+{
+	const ENetMode NetMode = GetNetMode();
+	const bool bIsServer = (NetMode == NM_DedicatedServer || NetMode == NM_ListenServer);
+	if (IsStreamingSourceEnabled() && (IsLocalController() || bIsServer))
+	{
+		GetPlayerViewPoint(OutStreamingSource.Location, OutStreamingSource.Rotation);
+		OutStreamingSource.Name = GetFName();
+		OutStreamingSource.TargetState = StreamingSourceShouldActivate() ? EStreamingSourceTargetState::Activated : EStreamingSourceTargetState::Loaded;
+		OutStreamingSource.bBlockOnSlowLoading = StreamingSourceShouldBlockOnSlowStreaming();
+		OutStreamingSource.DebugColor = StreamingSourceDebugColor;
+		OutStreamingSource.Priority = GetStreamingSourcePriority();
+		if (Shapes.Num())
+		{
+			OutStreamingSource.Shapes = Shapes;
+		}
+		return true;
+	}
+	return false;
+}
+
 /// @cond DOXYGEN_WARNINGS
 
 void APlayerController::ClientEnableNetworkVoice_Implementation(bool bEnable)
diff --git a/Engine/Source/Runtime/Engine/Private/World.cpp b/Engine/Source/Runtime/Engine/Private/World.cpp
--- a/Engine/Source/Runtime/Engine/Private/World.cpp
+++ b/Engine/Source/Runtime/Engine/Private/World.cpp
@@ -5335,9 +5335,13 @@ void UWorld::AddController( AController* Controller )
 {	
 	check( Controller );
 	ControllerList.AddUnique( Controller );	
-	if( Cast<APlayerController>(Controller) )
+	if( APlayerController* PlayerController = Cast<APlayerController>(Controller) )
 	{
-		PlayerControllerList.AddUnique( Cast<APlayerController>(Controller) );
+		if(!PlayerControllerList.Contains( PlayerController ))
+		{
+			PlayerControllerList.Add(PlayerController);
+			PlayerController->OnAddedToPlayerControllerList();
+		}
 	}
 }
 
@@ -5347,9 +5351,12 @@ void UWorld::RemoveController( AController* Controller )
 	check( Controller );
 	if( ControllerList.Remove( Controller ) > 0 )
 	{
-		if( Cast<APlayerController>(Controller) )
+		if( APlayerController* PlayerController = Cast<APlayerController>(Controller) )
 		{
-			PlayerControllerList.Remove( Cast<APlayerController>(Controller) );
+			if (PlayerControllerList.Remove(PlayerController) > 0)
+			{
+				PlayerController->OnRemovedFromPlayerControllerList();
+			}
 		}
 	}
 }
@@ -7295,6 +7302,7 @@ UWorld* FSeamlessTravelHandler::Tick()
 			}
 
 			bool bCreateNewGameMode = !bIsClient;
+			TArray<AController*> KeptControllers;
 			{
 				// scope because after GC the kept pointers will be bad
 				AGameModeBase* KeptGameMode = nullptr;
@@ -7306,11 +7314,12 @@ UWorld* FSeamlessTravelHandler::Tick()
 				for (AActor* const TheActor : ActuallyKeptActors)
 				{
 					KeepAnnotation.Clear(TheActor);
-					TheActor->Rename(nullptr, LoadedWorld->PersistentLevel);
-					// if it's a Controller, add it to the appropriate list in the new world's WorldSettings
-					if (TheActor->IsA<AController>())
+					
+					// if it's a Controller, remove it from the appropriate list in the current world's WorldSettings
+					if (AController* Controller = Cast<AController>(TheActor))
 					{
-						LoadedWorld->AddController(static_cast<AController*>(TheActor));
+						CurrentWorld->RemoveController(Controller);
+						KeptControllers.Add(Controller);
 					}
 					else if (TheActor->IsA<AGameModeBase>())
 					{
@@ -7321,6 +7330,8 @@ UWorld* FSeamlessTravelHandler::Tick()
 						KeptGameState = static_cast<AGameStateBase*>(TheActor);
 					}
 
+					TheActor->Rename(nullptr, LoadedWorld->PersistentLevel);
+
 					TheActor->bActorSeamlessTraveled = true;
 				}
 
@@ -7426,6 +7437,13 @@ UWorld* FSeamlessTravelHandler::Tick()
 			{
 				LoadedWorld->InitWorld();
 			}
+
+			// add controllers to initialized world 
+			for (AController* Controller : KeptControllers)
+			{
+				LoadedWorld->AddController(Controller);
+			}
+
 			bWorldChanged = true;
 			// Track session change on seamless travel.
 			NETWORK_PROFILER(GNetworkProfiler.TrackSessionChange(true, LoadedWorld->URL));
diff --git a/Engine/Source/Runtime/Engine/Private/WorldPartition/WorldPartitionStreamingPolicy.cpp b/Engine/Source/Runtime/Engine/Private/WorldPartition/WorldPartitionStreamingPolicy.cpp
--- a/Engine/Source/Runtime/Engine/Private/WorldPartition/WorldPartitionStreamingPolicy.cpp
+++ b/Engine/Source/Runtime/Engine/Private/WorldPartition/WorldPartitionStreamingPolicy.cpp
@@ -16,8 +16,6 @@
 #include "WorldPartition/HLOD/HLODSubsystem.h"
 #include "WorldPartition/ContentBundle/ContentBundle.h"
 #include "WorldPartition/WorldPartitionDebugHelper.h"
-#include "GameFramework/PlayerController.h"
-#include "Engine/LocalPlayer.h"
 #include "Engine/World.h"
 #include "Engine/Engine.h"
 #include "Engine/Canvas.h"
@@ -119,31 +117,6 @@ void UWorldPartitionStreamingPolicy::UpdateStreamingSources()
 #endif
 		if (!bIsServer || bIsServerStreamingEnabled || AWorldPartitionReplay::IsRecordingEnabled(World))
 		{
-			for (FConstPlayerControllerIterator It = World->GetPlayerControllerIterator(); It; ++It)
-			{
-				if (APlayerController* PlayerController = It->Get())
-				{
-					if (bIsServer || PlayerController->IsLocalController())
-					{
-						if (PlayerController->IsStreamingSourceEnabled())
-						{
-							FVector ViewLocation;
-							FRotator ViewRotation;
-							PlayerController->GetPlayerViewPoint(ViewLocation, ViewRotation);
-
-							// Transform to Local
-							ViewLocation = WorldToLocal.TransformPosition(ViewLocation);
-							ViewRotation = WorldToLocal.TransformRotation(ViewRotation.Quaternion()).Rotator();
-							const EStreamingSourceTargetState TargetState = PlayerController->StreamingSourceShouldActivate() ? EStreamingSourceTargetState::Activated : EStreamingSourceTargetState::Loaded;
-							const bool bBlockOnSlowLoading = PlayerController->StreamingSourceShouldBlockOnSlowStreaming();
-							const EStreamingSourcePriority StreamingSourcePriority = PlayerController->GetStreamingSourcePriority();
-							const bool bRemote = !PlayerController->IsLocalController();
-							StreamingSources.Add(FWorldPartitionStreamingSource(PlayerController->GetFName(), ViewLocation, ViewRotation, TargetState, bBlockOnSlowLoading, StreamingSourcePriority, bRemote));
-						}
-					}
-				}
-			}
-
 			UWorldPartitionSubsystem* WorldPartitionSubsystem = GetWorld()->GetSubsystem<UWorldPartitionSubsystem>();
 			check(WorldPartitionSubsystem);
 
