Fixed canvas text outlines being offset incorrectly when depth is not 1.

Depth was being applied before the offset which was messing things up.

diff --git a/Engine/Source/Runtime/Engine/Private/UserInterface/CanvasItem.cpp b/Engine/Source/Runtime/Engine/Private/UserInterface/CanvasItem.cpp
--- a/Engine/Source/Runtime/Engine/Private/UserInterface/CanvasItem.cpp
+++ b/Engine/Source/Runtime/Engine/Private/UserInterface/CanvasItem.cpp
@@ -1083,30 +1083,30 @@ void FCanvasTextItem::DrawStringInternal_OfflineCache(FCanvas* InCanvas, const F
 			const float SizeU	= Char.USize * InvTextureSize.X;
 			const float SizeV	= Char.VSize * InvTextureSize.Y;				
 
-			const float Left = X * Depth;
-			const float Top = Y * Depth;
-			const float Right = (X + SizeX) * Depth;
-			const float Bottom = (Y + SizeY) * Depth;
-
 			const auto AddTriangles = [&](const FVector2f& Offset, const FLinearColor& InColor)
 			{
+				const float Left = (X + Offset.X) * Depth;
+				const float Top = (Y + Offset.Y) * Depth;
+				const float Right = (X + SizeX) * Depth;
+				const float Bottom = (Y + SizeY) * Depth;
+
 				int32 V00 = BatchedElements->AddVertexf(
-					FVector4f(Left + Offset.X, Top + Offset.Y, 0.f, Depth),
+					FVector4f(Left, Top, 0.f, Depth),
 					FVector2f(U, V),
 					InColor,
 					HitProxyId);
 				int32 V10 = BatchedElements->AddVertexf(
-					FVector4f(Right + Offset.X, Top + Offset.Y, 0.0f, Depth),
+					FVector4f(Right, Top, 0.0f, Depth),
 					FVector2f(U + SizeU, V),
 					InColor,
 					HitProxyId);
 				int32 V01 = BatchedElements->AddVertexf(
-					FVector4f(Left + Offset.X, Bottom + Offset.Y, 0.0f, Depth),
+					FVector4f(Left, Bottom, 0.0f, Depth),
 					FVector2f(U, V + SizeV),
 					InColor,
 					HitProxyId);
 				int32 V11 = BatchedElements->AddVertexf(
-					FVector4f(Right + Offset.X, Bottom + Offset.Y, 0.0f, Depth),
+					FVector4f(Right, Bottom, 0.0f, Depth),
 					FVector2f(U + SizeU, V + SizeV),
 					InColor,
 					HitProxyId);
@@ -1278,30 +1278,30 @@ void FCanvasTextItem::DrawStringInternal_RuntimeCache(FCanvas* InCanvas, const F
 				const float SizeU = Entry.USize * InvTextureSizeX;
 				const float SizeV = Entry.VSize * InvTextureSizeY;
 
-				const float Left = X * Depth;
-				const float Top = Y * Depth;
-				const float Right = (X + SizeX) * Depth;
-				const float Bottom = (Y + SizeY) * Depth;
-
 				const auto AddTriangles = [&](const FVector2f& Offset, const FLinearColor& InColor)
 				{
+					const float Left = (X + Offset.X) * Depth;
+					const float Top = (Y + Offset.Y) * Depth;
+					const float Right = (X + SizeX) * Depth;
+					const float Bottom = (Y + SizeY) * Depth;
+
 					int32 V00 = BatchedElements->AddVertexf(
-						FVector4f(Left + Offset.X, Top + Offset.Y, 0.f, Depth),
+						FVector4f(Left, Top, 0.f, Depth),
 						FVector2f(U, V),
 						InColor,
 						HitProxyId);
 					int32 V10 = BatchedElements->AddVertexf(
-						FVector4f(Right + Offset.X, Top + Offset.Y, 0.0f, Depth),
+						FVector4f(Right, Top, 0.0f, Depth),
 						FVector2f(U + SizeU, V),
 						InColor,
 						HitProxyId);
 					int32 V01 = BatchedElements->AddVertexf(
-						FVector4f(Left + Offset.X, Bottom + Offset.Y, 0.0f, Depth),
+						FVector4f(Left, Bottom, 0.0f, Depth),
 						FVector2f(U, V + SizeV),
 						InColor,
 						HitProxyId);
 					int32 V11 = BatchedElements->AddVertexf(
-						FVector4f(Right + Offset.X, Bottom + Offset.Y, 0.0f, Depth),
+						FVector4f(Right, Bottom, 0.0f, Depth),
 						FVector2f(U + SizeU, V + SizeV),
 						InColor,
 						HitProxyId);
@@ -1441,30 +1441,30 @@ void FCanvasShapedTextItem::DrawStringInternal(FCanvas* InCanvas, const FVector2
 				const float SizeU = GlyphAtlasData.USize * InvTextureSizeX;
 				const float SizeV = GlyphAtlasData.VSize * InvTextureSizeY;
 
-				const float Left = X * Depth;
-				const float Top = Y * Depth;
-				const float Right = (X + SizeX) * Depth;
-				const float Bottom = (Y + SizeY) * Depth;
-
 				const auto AddTriangles = [&](const FVector2f& Offset, const FLinearColor& InColor)
 				{
+					const float Left = (X + Offset.X) * Depth;
+					const float Top = (Y + Offset.Y) * Depth;
+					const float Right = (X + SizeX) * Depth;
+					const float Bottom = (Y + SizeY) * Depth;
+
 					int32 V00 = BatchedElements->AddVertexf(
-						FVector4f(Left + Offset.X, Top + Offset.Y, 0.f, Depth),
+						FVector4f(Left, Top, 0.f, Depth),
 						FVector2f(U, V),
 						InColor,
 						HitProxyId);
 					int32 V10 = BatchedElements->AddVertexf(
-						FVector4f(Right + Offset.X, Top + Offset.Y, 0.0f, Depth),
+						FVector4f(Right, Top, 0.0f, Depth),
 						FVector2f(U + SizeU, V),
 						InColor,
 						HitProxyId);
 					int32 V01 = BatchedElements->AddVertexf(
-						FVector4f(Left + Offset.X, Bottom + Offset.Y, 0.0f, Depth),
+						FVector4f(Left, Bottom, 0.0f, Depth),
 						FVector2f(U, V + SizeV),
 						InColor,
 						HitProxyId);
 					int32 V11 = BatchedElements->AddVertexf(
-						FVector4f(Right + Offset.X, Bottom + Offset.Y, 0.0f, Depth),
+						FVector4f(Right, Bottom, 0.0f, Depth),
 						FVector2f(U + SizeU, V + SizeV),
 						InColor,
 						HitProxyId);
