Fix a wrong logic in the USkeletalMeshComponent::TickAnimation function.

#rnx

diff --git a/Engine/Source/Runtime/Engine/Private/Components/SkeletalMeshComponent.cpp b/Engine/Source/Runtime/Engine/Private/Components/SkeletalMeshComponent.cpp
--- a/Engine/Source/Runtime/Engine/Private/Components/SkeletalMeshComponent.cpp
+++ b/Engine/Source/Runtime/Engine/Private/Components/SkeletalMeshComponent.cpp
@@ -1233,7 +1233,7 @@ void USkeletalMeshComponent::TickAnimation(float DeltaTime, bool bNeedsValidRoot
 
 	bool bIsCompiling = false;
 #if WITH_EDITOR
-	bIsCompiling = GetSkeletalMeshAsset() && !GetSkeletalMeshAsset()->IsCompiling();
+	bIsCompiling = GetSkeletalMeshAsset() && GetSkeletalMeshAsset()->IsCompiling();
 #endif
 	if (GetSkeletalMeshAsset() != nullptr && !bIsCompiling)
 	{
