Marked ApplyRootMotionToVelocity as virtual

diff --git a/Engine/Source/Runtime/Engine/Classes/GameFramework/CharacterMovementComponent.h b/Engine/Source/Runtime/Engine/Classes/GameFramework/CharacterMovementComponent.h
--- a/Engine/Source/Runtime/Engine/Classes/GameFramework/CharacterMovementComponent.h
+++ b/Engine/Source/Runtime/Engine/Classes/GameFramework/CharacterMovementComponent.h
@@ -2628,7 +2628,7 @@ protected:
 	void RestorePreAdditiveRootMotionVelocity();
 
 	/** Applies root motion from root motion sources to velocity (override and additive) */
-	void ApplyRootMotionToVelocity(float deltaTime);
+	virtual void ApplyRootMotionToVelocity(float deltaTime);
 
 	/** Reduces former base velocity according to FormerBaseVelocityDecayHalfLife */
 	void DecayFormerBaseVelocity(float deltaTime);
