[CheatManager]
* Moved UE_WITH_CHEAT_MANAGER out of CheatManager.h and in to its own file in order to be able to use the define without including the code

diff --git a/Engine/Source/Runtime/Engine/Classes/GameFramework/CheatManager.h b/Engine/Source/Runtime/Engine/Classes/GameFramework/CheatManager.h
--- a/Engine/Source/Runtime/Engine/Classes/GameFramework/CheatManager.h
+++ b/Engine/Source/Runtime/Engine/Classes/GameFramework/CheatManager.h
@@ -9,6 +9,7 @@
 
 #pragma once
 
+#include "CheatManagerDefines.h"
 #include "CoreMinimal.h"
 #include "UObject/ObjectMacros.h"
 #include "UObject/Object.h"
@@ -19,11 +20,6 @@
 
 #include "CheatManager.generated.h"
 
-/** If set to 0, cheat manager functionality will be entirely disabled. If cheats are desired in shipping mode this can be overridden in the target definitions */
-#ifndef UE_WITH_CHEAT_MANAGER
-#define UE_WITH_CHEAT_MANAGER (1 && !UE_BUILD_SHIPPING)
-#endif
-
 /** Debug Trace info for capturing **/
 struct FDebugTraceInfo
 {
diff --git a/Engine/Source/Runtime/Engine/Classes/GameFramework/CheatManagerDefines.h b/Engine/Source/Runtime/Engine/Classes/GameFramework/CheatManagerDefines.h
new file mode 100644
--- /dev/null
+++ b/Engine/Source/Runtime/Engine/Classes/GameFramework/CheatManagerDefines.h
@@ -0,0 +1,9 @@
+// Copyright Epic Games, Inc. All Rights Reserved.
+
+#pragma once
+
+
+/** If set to 0, cheat manager functionality will be entirely disabled. If cheats are desired in shipping mode this can be overridden in the target definitions */
+#ifndef UE_WITH_CHEAT_MANAGER
+#define UE_WITH_CHEAT_MANAGER (1 && !UE_BUILD_SHIPPING)
+#endif
