Add custom setters to mark an arrow component's render state as dirty when writing to attributes cached by its scene proxy from a Blueprint function.

diff --git a/Engine/Source/Runtime/Engine/Classes/Components/ArrowComponent.h b/Engine/Source/Runtime/Engine/Classes/Components/ArrowComponent.h
--- a/Engine/Source/Runtime/Engine/Classes/Components/ArrowComponent.h
+++ b/Engine/Source/Runtime/Engine/Classes/Components/ArrowComponent.h
@@ -21,27 +21,27 @@ class ENGINE_API UArrowComponent : public UPrimitiveComponent
 	GENERATED_UCLASS_BODY()
 
 	/** Color to draw arrow */
-	UPROPERTY(EditAnywhere, BlueprintReadWrite, Category=ArrowComponent)
+	UPROPERTY(EditAnywhere, BlueprintReadWrite, BlueprintSetter=SetArrowFColor, Category=ArrowComponent)
 	FColor ArrowColor;
 
 	/** Relative size to scale drawn arrow by */
-	UPROPERTY(EditAnywhere, BlueprintReadWrite, Category=ArrowComponent)
+	UPROPERTY(EditAnywhere, BlueprintReadWrite, BlueprintSetter=SetArrowSize, Category=ArrowComponent)
 	float ArrowSize;
 
 	/** Total length of drawn arrow including head */
-	UPROPERTY(EditAnywhere, BlueprintReadWrite, Category=ArrowComponent)
+	UPROPERTY(EditAnywhere, BlueprintReadWrite, BlueprintSetter=SetArrowLength, Category=ArrowComponent)
 	float ArrowLength;
 
 	/** The size on screen to limit this arrow to (in screen space) */
-	UPROPERTY(EditAnywhere, BlueprintReadWrite, Category = ArrowComponent)
+	UPROPERTY(EditAnywhere, BlueprintReadWrite, BlueprintSetter=SetScreenSize, Category = ArrowComponent)
 	float ScreenSize;
 
 	/** Set to limit the screen size of this arrow */
-	UPROPERTY(EditAnywhere, BlueprintReadWrite, Category=ArrowComponent)
+	UPROPERTY(EditAnywhere, BlueprintReadWrite, BlueprintSetter=SetIsScreenSizeScaled, Category=ArrowComponent)
 	uint8 bIsScreenSizeScaled : 1;
 
 	/** If true, don't show the arrow when EngineShowFlags.BillboardSprites is disabled. */
-	UPROPERTY(EditAnywhere, BlueprintReadWrite, Category=ArrowComponent)
+	UPROPERTY(EditAnywhere, BlueprintReadWrite, BlueprintSetter=SetTreatAsASprite, Category=ArrowComponent)
 	uint8 bTreatAsASprite : 1;
 
 #if WITH_EDITORONLY_DATA
@@ -58,13 +58,34 @@ class ENGINE_API UArrowComponent : public UPrimitiveComponent
 	uint32 bLightAttachment:1;
 
 	/** Whether to use in-editor arrow scaling (i.e. to be affected by the global arrow scale) */
-	UPROPERTY(EditAnywhere, BlueprintReadWrite, Category=ArrowComponent)
+	UPROPERTY(EditAnywhere, BlueprintReadWrite, BlueprintSetter=SetUseInEditorScaling, Category=ArrowComponent)
 	bool bUseInEditorScaling;	
 #endif // WITH_EDITORONLY_DATA
 	/** Updates the arrow's colour, and tells it to refresh */
-	UFUNCTION(BlueprintCallable, DisplayName="SetArrowColor", Category="Components|Arrow")
+	UFUNCTION(BlueprintCallable, DisplayName="Set Arrow Color (Linear Color)", Category="Components|Arrow")
 	virtual void SetArrowColor(FLinearColor NewColor);
 
+	UFUNCTION(BlueprintSetter, Category=ArrowComponent)
+	void SetArrowFColor(FColor NewColor);
+
+	UFUNCTION(BlueprintSetter, Category=ArrowComponent)
+	void SetArrowSize(float NewSize);
+
+	UFUNCTION(BlueprintSetter, Category=ArrowComponent)
+	void SetArrowLength(float NewLength);
+
+	UFUNCTION(BlueprintSetter, Category=ArrowComponent)
+	void SetScreenSize(float NewScreenSize);
+
+	UFUNCTION(BlueprintSetter, Category=ArrowComponent)
+	void SetIsScreenSizeScaled(bool bNewValue);
+
+	UFUNCTION(BlueprintSetter, Category=ArrowComponent)
+	void SetTreatAsASprite(bool bNewValue);
+
+	UFUNCTION(BlueprintSetter, Category = ArrowComponent)
+	void SetUseInEditorScaling(bool bNewValue);
+
 	//~ Begin UPrimitiveComponent Interface.
 	virtual FPrimitiveSceneProxy* CreateSceneProxy() override;
 #if WITH_EDITOR
diff --git a/Engine/Source/Runtime/Engine/Private/Components/ArrowComponent.cpp b/Engine/Source/Runtime/Engine/Private/Components/ArrowComponent.cpp
--- a/Engine/Source/Runtime/Engine/Private/Components/ArrowComponent.cpp
+++ b/Engine/Source/Runtime/Engine/Private/Components/ArrowComponent.cpp
@@ -278,10 +278,56 @@ FBoxSphereBounds UArrowComponent::CalcBounds(const FTransform& LocalToWorld) con
 
 void UArrowComponent::SetArrowColor(FLinearColor NewColor)
 {
-	ArrowColor = NewColor.ToFColor(true);
+	SetArrowFColor(NewColor.ToFColor(true));
+}
+
+void UArrowComponent::SetArrowFColor(FColor NewColor)
+{
+	ArrowColor = NewColor;
+	MarkRenderStateDirty();
+}
+
+void UArrowComponent::SetArrowSize(float NewSize)
+{
+	ArrowSize = NewSize;
+	MarkRenderStateDirty();
+}
+
+void UArrowComponent::SetArrowLength(float NewLength)
+{
+	ArrowLength = NewLength;
+	MarkRenderStateDirty();
+}
+
+void UArrowComponent::SetScreenSize(float NewScreenSize)
+{
+	ScreenSize = NewScreenSize;
 	MarkRenderStateDirty();
 }
 
+void UArrowComponent::SetIsScreenSizeScaled(bool bNewValue)
+{
+	bIsScreenSizeScaled = bNewValue;
+	MarkRenderStateDirty();
+}
+
+void UArrowComponent::SetTreatAsASprite(bool bNewValue)
+{
+	bTreatAsASprite = bNewValue;
+	MarkRenderStateDirty();
+}
+
+void UArrowComponent::SetUseInEditorScaling(bool bNewValue)
+{
+#if WITH_EDITORONLY_DATA
+	bUseInEditorScaling = bNewValue;
+	MarkRenderStateDirty();
+#else
+	const bool bCallOutsideOf_WithEditorOnlyData = false;
+	ensure(bCallOutsideOf_WithEditorOnlyData);
+#endif
+}
+
 #if WITH_EDITORONLY_DATA
 void UArrowComponent::SetEditorScale(float InEditorScale)
 {
