World Partition Content Bundle: propagate the content bundle GUID to actors in subcontainers.

#rnx

diff --git a/Engine/Source/Runtime/Engine/Private/WorldPartition/WorldPartitionStreamingGeneration.cpp b/Engine/Source/Runtime/Engine/Private/WorldPartition/WorldPartitionStreamingGeneration.cpp
--- a/Engine/Source/Runtime/Engine/Private/WorldPartition/WorldPartitionStreamingGeneration.cpp
+++ b/Engine/Source/Runtime/Engine/Private/WorldPartition/WorldPartitionStreamingGeneration.cpp
@@ -126,7 +126,10 @@ class FWorldPartitionStreamingGenerator
 					ActorSetInstance.ActorSet = &ActorSet;
 					ActorSetInstance.ContainerID = ContainerInstanceDescriptor.ID;
 					ActorSetInstance.Transform = ContainerInstanceDescriptor.Transform;
-					ActorSetInstance.ContentBundleID = ReferenceActorDescView.GetContentBundleGuid();
+
+					// Since Content Bundles streaming generation happens in its own context, all actor set instances must have the same content bundle GUID for now, so Level Instances
+					// placed inside a Content Bundle will propagate their Content Bundle GUID to child instances.
+					ActorSetInstance.ContentBundleID = MainWorldContainer->GetContentBundleGuid();
 
 					if (ContainerInstanceDescriptor.ID.IsMainContainer())
 					{
@@ -144,7 +147,7 @@ class FWorldPartitionStreamingGenerator
 						TSet<FName> CombinedDataLayers = ContainerInstanceDescriptor.RuntimeDataLayers;
 						CombinedDataLayers.Append(ReferenceActorDescView.GetRuntimeDataLayers());
 
-						ActorSetInstance.DataLayers = UDataLayerSubsystem::GetRuntimeDataLayerInstances(ContainerDescriptor.Container->GetWorld(), CombinedDataLayers.Array());
+						ActorSetInstance.DataLayers = UDataLayerSubsystem::GetRuntimeDataLayerInstances(ContainerDescriptor.Container->GetWorld(), CombinedDataLayers.Array());						
 					}
 
 					ActorSetInstance.Bounds.Init();
