Rollback behavior of IsBlueprintNotify to before CL#22113548.

diff --git a/Engine/Source/Runtime/Engine/Private/Animation/AnimTypes.cpp b/Engine/Source/Runtime/Engine/Private/Animation/AnimTypes.cpp
--- a/Engine/Source/Runtime/Engine/Private/Animation/AnimTypes.cpp
+++ b/Engine/Source/Runtime/Engine/Private/Animation/AnimTypes.cpp
@@ -110,14 +110,6 @@ bool FAnimNotifyEvent::IsBranchingPoint() const
 	return GetLinkedMontage() && ((MontageTickType == EMontageNotifyTickType::BranchingPoint) || (Notify && Notify->bIsNativeBranchingPoint) || (NotifyStateClass && NotifyStateClass->bIsNativeBranchingPoint));
 }
 
-bool FAnimNotifyEvent::IsBlueprintNotify() const
-{
-	return (
-		(Notify && Notify->GetClass() && !Notify->GetClass()->IsNative()) ||
-		(NotifyStateClass && NotifyStateClass->GetClass() && !NotifyStateClass->GetClass()->IsNative())
-		);
-}
-
 void FAnimNotifyEvent::SetTime(float NewTime, EAnimLinkMethod::Type ReferenceFrame /*= EAnimLinkMethod::Absolute*/)
 {
 	FAnimLinkableElement::SetTime(NewTime, ReferenceFrame);
diff --git a/Engine/Source/Runtime/Engine/Public/Animation/AnimTypes.h b/Engine/Source/Runtime/Engine/Public/Animation/AnimTypes.h
--- a/Engine/Source/Runtime/Engine/Public/Animation/AnimTypes.h
+++ b/Engine/Source/Runtime/Engine/Public/Animation/AnimTypes.h
@@ -423,7 +423,10 @@ public:
 	ENGINE_API bool IsBranchingPoint() const;
 
 	/** Returns true if this is blueprint derived notifies **/
-	bool IsBlueprintNotify() const;
+	bool IsBlueprintNotify() const
+	{
+		return Notify != nullptr || NotifyStateClass != nullptr;
+	}
 
 	bool operator ==(const FAnimNotifyEvent& Other) const
 	{
