Added some logging when cooking content bundles
Output an error when generating an always loaded cell inside a content bundle and avoid populating the cell

diff --git a/Engine/Source/Runtime/Engine/Private/WorldPartition/ContentBundle/ContentBundleEditor.cpp b/Engine/Source/Runtime/Engine/Private/WorldPartition/ContentBundle/ContentBundleEditor.cpp
--- a/Engine/Source/Runtime/Engine/Private/WorldPartition/ContentBundle/ContentBundleEditor.cpp
+++ b/Engine/Source/Runtime/Engine/Private/WorldPartition/ContentBundle/ContentBundleEditor.cpp
@@ -379,6 +379,8 @@ bool FContentBundleEditor::PopulateGeneratorPackageForCook(class IWorldPartition
 
 	if (HasCookedContent())
 	{
+		UE_LOG(LogContentBundle, Log, TEXT("[CB: %s][Cook] Populating Generator Package. %u Packages"), *GetDescriptor()->GetDisplayName(), PackagesToCook.Num());
+
 		for (const FWorldPartitionCookPackage* CookPackage : PackagesToCook)
 		{
 			if (CookPackage->Type == FWorldPartitionCookPackage::EType::Level)
@@ -414,6 +416,8 @@ bool FContentBundleEditor::PopulateGeneratedPackageForCook(class IWorldPartition
 {
 	bool bIsSuccess = true;
 
+	UE_LOG(LogContentBundle, Log, TEXT("[CB: %s][Cook] Populating Generated Package %s"), *GetDescriptor()->GetDisplayName(), *PackageToCook.RelativePath);
+
 	if (PackageToCook.Type == FWorldPartitionCookPackage::EType::Level)
 	{
 		if (UWorldPartitionRuntimeCell** MatchingCell = const_cast<UWorldPartitionRuntimeCell**>(CookPackageIdsToCell.Find(PackageToCook.PackageId)))
@@ -421,10 +425,20 @@ bool FContentBundleEditor::PopulateGeneratedPackageForCook(class IWorldPartition
 			UWorldPartitionRuntimeCell* Cell = *MatchingCell;
 			if (ensure(Cell))
 			{
-				TArray<UPackage*> ModifiedPackages;
-				if (!Cell->PopulateGeneratedPackageForCook(PackageToCook.GetPackage(), OutModifiedPackages))
+				if (!Cell->IsAlwaysLoaded())
+				{
+					TArray<UPackage*> ModifiedPackages;
+					if (!Cell->PopulateGeneratedPackageForCook(PackageToCook.GetPackage(), OutModifiedPackages))
+					{
+						UE_LOG(LogContentBundle, Error, TEXT("[CB: %s][Cook] Failed to populate cell package %s."), *GetDescriptor()->GetDisplayName(), *PackageToCook.RelativePath);
+						bIsSuccess = false;
+					}
+					
+				}
+				else
 				{
-					UE_LOG(LogContentBundle, Error, TEXT("[CB: %s][Cook] Failed to populate cell package %s."), *GetDescriptor()->GetDisplayName(), *PackageToCook.RelativePath);
+					UE_LOG(LogContentBundle, Error, TEXT("[CB: %s][Cook] Cell %s is flagged always loaded. Content Bundles cells should never be always loaded. It will not be populated and be empty at runtime."),
+						*GetDescriptor()->GetDisplayName(), *PackageToCook.RelativePath);
 					bIsSuccess = false;
 				}
 			}
