Fix bad merging before submit

diff --git a/Engine/Source/Runtime/Engine/Private/SkeletalMeshLODRenderData.cpp b/Engine/Source/Runtime/Engine/Private/SkeletalMeshLODRenderData.cpp
--- a/Engine/Source/Runtime/Engine/Private/SkeletalMeshLODRenderData.cpp
+++ b/Engine/Source/Runtime/Engine/Private/SkeletalMeshLODRenderData.cpp
@@ -215,6 +215,7 @@ void FSkeletalMeshLODRenderData::InitResources(bool bNeedsVertexColors, int32 LO
 		const FSkeletalMeshLODInfo* SkeletalMeshLODInfo = Owner->GetLODInfo(LODIndex);
 		MorphTargetVertexInfoBuffers.InitMorphResources(GMaxRHIShaderPlatform, RenderSections, Owner->GetMorphTargets(), StaticVertexBuffers.StaticMeshVertexBuffer.GetNumVertices(), LODIndex, SkeletalMeshLODInfo->MorphTargetPositionErrorTolerance);
 	}
+	}
 
 	if (!MorphTargetVertexInfoBuffers.IsRHIIntialized() && MorphTargetVertexInfoBuffers.IsMorphCPUDataValid() && (MorphTargetVertexInfoBuffers.NumTotalBatches > 0 || !bStreamedDataInlined))
 	{
