Fix telemetry bug where the ForceNoAudioMixer telemetry was always called, even if the flag was not set to true

diff --git a/Engine/Source/Runtime/Engine/Private/AudioDeviceManager.cpp b/Engine/Source/Runtime/Engine/Private/AudioDeviceManager.cpp
--- a/Engine/Source/Runtime/Engine/Private/AudioDeviceManager.cpp
+++ b/Engine/Source/Runtime/Engine/Private/AudioDeviceManager.cpp
@@ -720,8 +720,10 @@ bool FAudioDeviceManager::LoadDefaultAudioDeviceModule()
 		IsUsingAudioMixerCvar->Set(0, ECVF_SetByConstructor);
 	}
 
-	Audio::Analytics::RecordEvent_Usage(TEXT("ForceNoAudioMixer"));
-
+	if (bForceNoAudioMixer)
+	{
+		Audio::Analytics::RecordEvent_Usage(TEXT("ForceNoAudioMixer"));
+	}
 	return AudioDeviceModule != nullptr;
 }
 
