Cloth - Deprecated an unused bLocalSpaceSimulation property on the SkeletalMeshComponent.

diff --git a/Engine/Source/Runtime/Engine/Classes/Components/SkeletalMeshComponent.h b/Engine/Source/Runtime/Engine/Classes/Components/SkeletalMeshComponent.h
--- a/Engine/Source/Runtime/Engine/Classes/Components/SkeletalMeshComponent.h
+++ b/Engine/Source/Runtime/Engine/Classes/Components/SkeletalMeshComponent.h
@@ -699,14 +699,9 @@ public:
 	UPROPERTY(EditAnywhere, BlueprintReadWrite, Category=Clothing)
 	uint8 bForceCollisionUpdate : 1;
 
-	/**
-	 * It's worth trying this option when you feel that the current cloth simulation is unstable.
-	 * The scale of the actor is maintained during the simulation. 
-	 * It is possible to add the inertia effects to the simulation, through the inertiaScale parameter of the clothing material. 
-	 * So with an inertiaScale of 1.0 there should be no visible difference between local space and global space simulation. 
-	 * Known issues: - Currently there's simulation issues when this feature is used in 3.x (DE4076) So if localSpaceSim is enabled there's no inertia effect when the global pose of the clothing actor changes.
-	 */
-	UPROPERTY(EditAnywhere, BlueprintReadWrite, Category = Clothing)
+	/** Deprecated. */
+	UE_DEPRECATED(5.2, "Cloth simulation always happens in local space.")
+	UPROPERTY(BlueprintReadWrite, Category = Clothing, meta = (DeprecationMessage = "This property is deprecated. Cloth simulation always happens in local space."))
 	uint8 bLocalSpaceSimulation : 1;
 
 	/** reset the clothing after moving the clothing position (called teleport) */
