Have a ProgressDialog when shaders are recached because of scalability change

diff --git a/Engine/Source/Runtime/Engine/Private/UnrealEngine.cpp b/Engine/Source/Runtime/Engine/Private/UnrealEngine.cpp
--- a/Engine/Source/Runtime/Engine/Private/UnrealEngine.cpp
+++ b/Engine/Source/Runtime/Engine/Private/UnrealEngine.cpp
@@ -833,8 +833,8 @@ void ScalabilityCVarsSinkCallback()
 			if (bCacheResourceShaders)
 			{
 				// For all materials, UMaterial::CacheResourceShadersForRendering
-				UMaterial::AllMaterialsCacheResourceShadersForRendering();
-				UMaterialInstance::AllMaterialsCacheResourceShadersForRendering();
+				UMaterial::AllMaterialsCacheResourceShadersForRendering(true);
+				UMaterialInstance::AllMaterialsCacheResourceShadersForRendering(true);
 			}
 		}
 		else
