Iris - Removed hardcoded prioritizer for PlayerState as it is ini configurable.

diff --git a/Engine/Source/Runtime/Engine/Classes/GameFramework/PlayerState.h b/Engine/Source/Runtime/Engine/Classes/GameFramework/PlayerState.h
--- a/Engine/Source/Runtime/Engine/Classes/GameFramework/PlayerState.h
+++ b/Engine/Source/Runtime/Engine/Classes/GameFramework/PlayerState.h
@@ -339,10 +339,6 @@ protected:
 	/** called after receiving player name */
 	virtual void HandleWelcomeMessage();
 
-#if UE_WITH_IRIS
-	virtual void BeginReplication() override;
-#endif
-
 private:
 	// Hidden functions that don't make sense to use on this class.
 	HIDE_ACTOR_TRANSFORM_FUNCTIONS();
diff --git a/Engine/Source/Runtime/Engine/Private/PlayerState.cpp b/Engine/Source/Runtime/Engine/Private/PlayerState.cpp
--- a/Engine/Source/Runtime/Engine/Private/PlayerState.cpp
+++ b/Engine/Source/Runtime/Engine/Private/PlayerState.cpp
@@ -606,39 +606,3 @@ void APlayerState::OnPawnPrivateDestroyed(AActor* InActor)
 }
 
 PRAGMA_ENABLE_DEPRECATION_WARNINGS
-
-#if UE_WITH_IRIS
-void APlayerState::BeginReplication()
-{
-	using namespace UE::Net;
-
-	Super::BeginReplication();
-
-	/**
-	 * If we're replicating using Iris we set the prioritizer to one named PlayerState if there is one.
-	 * For example add the following to your DefaultEngine.ini to limit the number of player states
-	 * replicated each frame, but still allow the owning connection's to be replicated each frame.
-	 * You can also implement your own prioritzer of course.
-	 * 
-	 * [/Script/IrisCore.NetObjectPrioritizerDefinitions]
-	 * +NetObjectPrioritizerDefinitions=(PrioritizerName=PlayerState, ClassName=/Script/IrisCore.NetObjectCountLimiter, ConfigClassName=/Script/Engine.PlayerStateCountLimiterConfig)
-	 */
-	const FNetHandle NetHandle = FReplicationSystemUtil::GetNetHandle(this);
-	if (!NetHandle.IsValid())
-	{
-		return;
-	}
-
-	UReplicationSystem* ReplicationSystem = GetReplicationSystem(NetHandle.GetReplicationSystemId());
-	if (ReplicationSystem == nullptr)
-	{
-		return;
-	}
-
-	FNetObjectPrioritizerHandle PrioritizerHandle = ReplicationSystem->GetPrioritizerHandle(FName("PlayerState", FNAME_Find));
-	if (PrioritizerHandle != InvalidNetObjectPrioritizerHandle)
-	{
-		ReplicationSystem->SetPrioritizer(NetHandle, PrioritizerHandle);
-	}	
-}
-#endif
