Engine module & dependencies pass: ANIM-RT truncation fixes

diff --git a/Engine/Source/Editor/UnrealEd/Public/SSkeletonWidget.h b/Engine/Source/Editor/UnrealEd/Public/SSkeletonWidget.h
--- a/Engine/Source/Editor/UnrealEd/Public/SSkeletonWidget.h
+++ b/Engine/Source/Editor/UnrealEd/Public/SSkeletonWidget.h
@@ -71,7 +71,7 @@ public:
 
 		if (ColumnName == TEXT("Curretly Selected"))
 		{
-			return SAssignNew(Border1, SBorder) .Padding(2)
+			return SAssignNew(Border1, SBorder) .Padding(2.f)
 				.Content()
 				[
 					SNew(STextBlock)
@@ -82,7 +82,7 @@ public:
 		{
 			if (BonePair->Bone2 == NAME_None)
 			{
-				return SAssignNew(Border2, SBorder) .Padding(2)
+				return SAssignNew(Border2, SBorder) .Padding(2.f)
 					.ColorAndOpacity(FLinearColor(1.f, 0.f, 0.f))
 					.Content()
 					[
@@ -92,7 +92,7 @@ public:
 			}
 			else
 			{
-				return SAssignNew(Border2, SBorder) .Padding(2)
+				return SAssignNew(Border2, SBorder) .Padding(2.f)
 					.Content()
 					[
 						SNew(STextBlock)
diff --git a/Engine/Source/Runtime/Engine/Public/Animation/AnimSequenceDecompressionContext.h b/Engine/Source/Runtime/Engine/Public/Animation/AnimSequenceDecompressionContext.h
--- a/Engine/Source/Runtime/Engine/Public/Animation/AnimSequenceDecompressionContext.h
+++ b/Engine/Source/Runtime/Engine/Public/Animation/AnimSequenceDecompressionContext.h
@@ -45,7 +45,7 @@ struct FAnimSequenceDecompressionContext
 	FAnimSequenceDecompressionContext(const FFrameRate& InSamplingRate, const int32 InNumberOfFrames, EAnimInterpolationType Interpolation_, const FName& AnimName_, const ICompressedAnimData& CompressedAnimData_, const TArray<FTransform>& RefPoses_, const TArray<FTrackToSkeletonMap>& TrackToSkeletonMap_, const USkeleton* InSourceSkeleton, bool bIsBakedAdditive)
 		:
 		PRAGMA_DISABLE_DEPRECATION_WARNINGS
-		SequenceLength(InSamplingRate.AsSeconds(InNumberOfFrames)),
+		SequenceLength(static_cast<float>(InSamplingRate.AsSeconds(InNumberOfFrames))),
 		PRAGMA_ENABLE_DEPRECATION_WARNINGS
 		Interpolation(Interpolation_), AnimName(AnimName_), CompressedAnimData(CompressedAnimData_),
 		RefPoses(RefPoses_.GetData(), RefPoses_.Num()), TrackToSkeletonMap(TrackToSkeletonMap_.GetData(), TrackToSkeletonMap_.Num()),
@@ -109,8 +109,8 @@ public:
 	void Seek(double SampleAtTime)
 	{
 		PRAGMA_DISABLE_DEPRECATION_WARNINGS
-		Time = SampleAtTime;
-		RelativePos = SequenceLength != 0.f ? SampleAtTime / SequenceLength : 0.f;
+		Time = static_cast<float>(SampleAtTime);
+		RelativePos = SequenceLength != 0.f ? static_cast<float>(SampleAtTime) / SequenceLength : 0.f;
 		PRAGMA_ENABLE_DEPRECATION_WARNINGS
 		
 		SamplingTime = SamplingRate.AsFrameTime(SampleAtTime);
