PlayerController add GetStreamingSourceLocationAndRotation + GetStreamingSourceShapes virtuals that can be implemented by child class

diff --git a/Engine/Source/Runtime/Engine/Classes/GameFramework/PlayerController.h b/Engine/Source/Runtime/Engine/Classes/GameFramework/PlayerController.h
--- a/Engine/Source/Runtime/Engine/Classes/GameFramework/PlayerController.h
+++ b/Engine/Source/Runtime/Engine/Classes/GameFramework/PlayerController.h
@@ -568,7 +568,7 @@ public:
 
 	/** Optional aggregated shape list used to build a custom shape for the streaming source. When empty, fallbacks sphere shape with a radius equal to grid's loading range. */
 	UPROPERTY(EditAnywhere, BlueprintReadWrite, Category = WorldPartition, meta = (EditCondition = "bEnableStreamingSource"))
-	TArray<FStreamingSourceShape> Shapes;
+	TArray<FStreamingSourceShape> StreamingSourceShapes;
 
 	/** Scale applied to force feedback values */
 	UPROPERTY(config)
@@ -794,6 +794,21 @@ public:
 	UFUNCTION(BlueprintCallable, Category = WorldPartition)
 	virtual EStreamingSourcePriority GetStreamingSourcePriority() const { return StreamingSourcePriority; }
 
+	/**
+	* Gets the streaming source location and rotation.
+	* Default implementation returns APlayerController::GetPlayerViewPoint but can be overriden in child classes.
+	*/
+	UFUNCTION(BlueprintCallable, Category = WorldPartition)
+	virtual void GetStreamingSourceLocationAndRotation(FVector& OutLocation, FRotator& OutRotation) const;
+	
+	/**
+	* Gets the streaming source priority.
+	* Default implementation returns StreamingSourceShapes but can be overriden in child classes.
+	* @return the streaming source priority.
+	*/
+	UFUNCTION(BlueprintCallable, Category = WorldPartition)
+	virtual void GetStreamingSourceShapes(TArray<FStreamingSourceShape>& OutShapes) const;
+
 	/**
 	 * Gets the PlayerController's streaming source
 	 * @return the streaming source.
diff --git a/Engine/Source/Runtime/Engine/Private/PlayerController.cpp b/Engine/Source/Runtime/Engine/Private/PlayerController.cpp
--- a/Engine/Source/Runtime/Engine/Private/PlayerController.cpp
+++ b/Engine/Source/Runtime/Engine/Private/PlayerController.cpp
@@ -3619,22 +3619,32 @@ void APlayerController::OnRemovedFromPlayerControllerList()
 	}
 }
 
+void APlayerController::GetStreamingSourceLocationAndRotation(FVector& OutLocation, FRotator& OutRotation) const
+{
+	GetPlayerViewPoint(OutLocation, OutRotation);
+}
+
+void APlayerController::GetStreamingSourceShapes(TArray<FStreamingSourceShape>& OutShapes) const
+{
+	if (StreamingSourceShapes.Num())
+	{
+		OutShapes.Append(StreamingSourceShapes);
+	}
+}
+
 bool APlayerController::GetStreamingSource(FWorldPartitionStreamingSource& OutStreamingSource) const
 {
 	const ENetMode NetMode = GetNetMode();
 	const bool bIsServer = (NetMode == NM_DedicatedServer || NetMode == NM_ListenServer);
 	if (IsStreamingSourceEnabled() && (IsLocalController() || bIsServer))
 	{
-		GetPlayerViewPoint(OutStreamingSource.Location, OutStreamingSource.Rotation);
+		GetStreamingSourceLocationAndRotation(OutStreamingSource.Location, OutStreamingSource.Rotation);
 		OutStreamingSource.Name = GetFName();
 		OutStreamingSource.TargetState = StreamingSourceShouldActivate() ? EStreamingSourceTargetState::Activated : EStreamingSourceTargetState::Loaded;
 		OutStreamingSource.bBlockOnSlowLoading = StreamingSourceShouldBlockOnSlowStreaming();
 		OutStreamingSource.DebugColor = StreamingSourceDebugColor;
 		OutStreamingSource.Priority = GetStreamingSourcePriority();
-		if (Shapes.Num())
-		{
-			OutStreamingSource.Shapes = Shapes;
-		}
+		GetStreamingSourceShapes(OutStreamingSource.Shapes);
 		return true;
 	}
 	return false;
