World Partition Content Bundle: don't load all actors when disabling a bundle.

#rnx

diff --git a/Engine/Source/Runtime/Engine/Private/WorldPartition/WorldPartition.cpp b/Engine/Source/Runtime/Engine/Private/WorldPartition/WorldPartition.cpp
--- a/Engine/Source/Runtime/Engine/Private/WorldPartition/WorldPartition.cpp
+++ b/Engine/Source/Runtime/Engine/Private/WorldPartition/WorldPartition.cpp
@@ -1621,7 +1621,14 @@ bool UWorldPartition::UnregisterActorDescContainer(UActorDescContainer* InActorD
 			if (ActorHandle.IsValid())
 			{
 				ActorGuids.Add(It->GetGuid());
-				DirtyActors.Remove(ActorHandle.ToReference());
+
+				for (TMap<FWorldPartitionReference, AActor*>::TIterator DirtyActorIt(DirtyActors); DirtyActorIt; ++DirtyActorIt)
+				{
+					if (DirtyActorIt.Key() == ActorHandle)
+					{
+						DirtyActorIt.RemoveCurrent();
+					}
+				}
 			}
 		}
 
