Load Special Engine Materials from the DDC even when JobCacheDDC is on since they are required complete.

diff --git a/Engine/Source/Runtime/Engine/Private/Materials/MaterialShared.cpp b/Engine/Source/Runtime/Engine/Private/Materials/MaterialShared.cpp
--- a/Engine/Source/Runtime/Engine/Private/Materials/MaterialShared.cpp
+++ b/Engine/Source/Runtime/Engine/Private/Materials/MaterialShared.cpp
@@ -2478,7 +2478,7 @@ bool FMaterial::CacheShaders(const FMaterialShaderMapId& ShaderMapId, EShaderPla
 			// If we have an incomplete shadermap, continue with it to prevent creation of duplicate shadermaps for the same ShaderMapId
 			if (!ShaderMap && !FPlatformProperties::RequiresCookedData())
 			{
-				if (bSkipCompilationOnPostLoad == false || IsDefaultMaterial())
+				if (bSkipCompilationOnPostLoad == false || IsRequiredComplete())
 				{
 					TRefCountPtr<FMaterialShaderMap> LoadedShaderMap;
 					CacheShadersPending = FMaterialShaderMap::BeginLoadFromDerivedDataCache(this, ShaderMapId, Platform, TargetPlatform, LoadedShaderMap, DDCKeyHash);
