Change Strata to use DiffuseAlbedo/F0 paramterization when creating a new Strata material.

diff --git a/Engine/Source/Runtime/Engine/Private/Materials/Material.cpp b/Engine/Source/Runtime/Engine/Private/Materials/Material.cpp
--- a/Engine/Source/Runtime/Engine/Private/Materials/Material.cpp
+++ b/Engine/Source/Runtime/Engine/Private/Materials/Material.cpp
@@ -3194,6 +3194,7 @@ void UMaterial::ConvertMaterialToStrataMaterial()
 		// Empty material: Create by default a slab node
 		const int32 NewNodeOffsetX = -300;
 		UMaterialExpressionStrataSlabBSDF* SlabNode = NewObject<UMaterialExpressionStrataSlabBSDF>(this);
+		SlabNode->bUseMetalness = false;
 		SlabNode->MaterialExpressionEditorX = EditorX + NewNodeOffsetX;
 		EditorOnly->FrontMaterial.Connect(0, SlabNode);
 		bRelinkCustomOutputNodes = false;
