Cloth - Fixed the USkeletalMeshComponent::GetTeleportRotationThreshold() method that was returning the distance instead of the rotation threshold.

diff --git a/Engine/Source/Runtime/Engine/Private/Components/SkeletalMeshComponent.cpp b/Engine/Source/Runtime/Engine/Private/Components/SkeletalMeshComponent.cpp
--- a/Engine/Source/Runtime/Engine/Private/Components/SkeletalMeshComponent.cpp
+++ b/Engine/Source/Runtime/Engine/Private/Components/SkeletalMeshComponent.cpp
@@ -4423,7 +4423,7 @@ void USkeletalMeshComponent::SetUpdateClothInEditor(const bool NewUpdateState)
 
 float USkeletalMeshComponent::GetTeleportRotationThreshold() const
 {
-	return TeleportDistanceThreshold;
+	return TeleportRotationThreshold;
 }
 
 void USkeletalMeshComponent::SetTeleportRotationThreshold(float Threshold)
