Moved FLevelStreamingGCHelper in its own set of files

diff --git a/Engine/Source/Editor/UnrealEd/Private/WorldPartition/WorldPartitionHLODsBuilder.cpp b/Engine/Source/Editor/UnrealEd/Private/WorldPartition/WorldPartitionHLODsBuilder.cpp
--- a/Engine/Source/Editor/UnrealEd/Private/WorldPartition/WorldPartitionHLODsBuilder.cpp
+++ b/Engine/Source/Editor/UnrealEd/Private/WorldPartition/WorldPartitionHLODsBuilder.cpp
@@ -27,6 +27,7 @@
 #include "WorldPartition/HLOD/HLODLayer.h"
 
 #include "Engine/StaticMesh.h"
+#include "Engine/LevelStreamingGCHelper.h"
 #include "Materials/Material.h"
 #include "Materials/MaterialInstance.h"
 
diff --git a/Engine/Source/Runtime/Engine/Classes/Engine/LevelStreamingGCHelper.h b/Engine/Source/Runtime/Engine/Classes/Engine/LevelStreamingGCHelper.h
new file mode 100644
--- /dev/null
+++ b/Engine/Source/Runtime/Engine/Classes/Engine/LevelStreamingGCHelper.h
@@ -0,0 +1,65 @@
+// Copyright Epic Games, Inc. All Rights Reserved.
+
+#pragma once
+
+#include "CoreMinimal.h"
+
+class ULevel;
+class UWorld;
+
+/**
+ * Helper structure encapsulating functionality used to defer marking actors and their components as pending
+ * kill till right before garbage collection by registering a callback.
+ */
+struct ENGINE_API FLevelStreamingGCHelper
+{
+	/** Called when streamed out levels are going to be garbage collected  */
+	DECLARE_MULTICAST_DELEGATE(FOnGCStreamedOutLevelsEvent);
+	static FOnGCStreamedOutLevelsEvent OnGCStreamedOutLevels;
+
+	/**
+	 * Register with the garbage collector to receive callbacks pre and post garbage collection
+	 */
+	static void AddGarbageCollectorCallback();
+
+	/**
+	 * Request to be unloaded.
+	 *
+	 * @param InLevel	Level that should be unloaded
+	 */
+	static void RequestUnload( ULevel* InLevel );
+
+	/**
+	 * Cancel any pending unload requests for passed in Level.
+	 */
+	static void CancelUnloadRequest( ULevel* InLevel );
+
+	/** 
+	 * Prepares levels that are marked for unload for the GC call by marking their actors and components as
+	 * pending kill.
+	 */
+	static void PrepareStreamedOutLevelsForGC();
+
+	/**
+	 * Verify that the level packages are no longer around.
+	 */
+	static void VerifyLevelsGotRemovedByGC();
+	
+	/**
+	 * @return	The number of levels pending a purge by the garbage collector
+	 */
+	static int32 GetNumLevelsPendingPurge();
+
+	/**
+	 * Allows FLevelStreamingGCHelper to be used in a commandlet.
+	 */
+	static void EnableForCommandlet();
+	
+private:
+	/** Static array of levels that should be unloaded */
+	static TArray<TWeakObjectPtr<ULevel> > LevelsPendingUnload;
+	/** Static array of level packages that have been marked by PrepareStreamedOutLevelsForGC */
+	static TArray<FName> LevelPackageNames;
+	/** Static bool allows FLevelStreamingGCHelper to be used in a commandlet */
+	static bool bEnabledForCommandlet;
+};
\ No newline at end of file
diff --git a/Engine/Source/Runtime/Engine/Classes/Engine/World.h b/Engine/Source/Runtime/Engine/Classes/Engine/World.h
--- a/Engine/Source/Runtime/Engine/Classes/Engine/World.h
+++ b/Engine/Source/Runtime/Engine/Classes/Engine/World.h
@@ -432,63 +432,6 @@ public:
 };
 
 
-/**
- * Helper structure encapsulating functionality used to defer marking actors and their components as pending
- * kill till right before garbage collection by registering a callback.
- */
-struct ENGINE_API FLevelStreamingGCHelper
-{
-	/** Called when streamed out levels are going to be garbage collected  */
-	DECLARE_MULTICAST_DELEGATE(FOnGCStreamedOutLevelsEvent);
-	static FOnGCStreamedOutLevelsEvent OnGCStreamedOutLevels;
-
-	/**
-	 * Register with the garbage collector to receive callbacks pre and post garbage collection
-	 */
-	static void AddGarbageCollectorCallback();
-
-	/**
-	 * Request to be unloaded.
-	 *
-	 * @param InLevel	Level that should be unloaded
-	 */
-	static void RequestUnload( ULevel* InLevel );
-
-	/**
-	 * Cancel any pending unload requests for passed in Level.
-	 */
-	static void CancelUnloadRequest( ULevel* InLevel );
-
-	/** 
-	 * Prepares levels that are marked for unload for the GC call by marking their actors and components as
-	 * pending kill.
-	 */
-	static void PrepareStreamedOutLevelsForGC();
-
-	/**
-	 * Verify that the level packages are no longer around.
-	 */
-	static void VerifyLevelsGotRemovedByGC();
-	
-	/**
-	 * @return	The number of levels pending a purge by the garbage collector
-	 */
-	static int32 GetNumLevelsPendingPurge();
-
-	/**
-	 * Allows FLevelStreamingGCHelper to be used in a commandlet.
-	 */
-	static void EnableForCommandlet();
-	
-private:
-	/** Static array of levels that should be unloaded */
-	static TArray<TWeakObjectPtr<ULevel> > LevelsPendingUnload;
-	/** Static array of level packages that have been marked by PrepareStreamedOutLevelsForGC */
-	static TArray<FName> LevelPackageNames;
-	/** Static bool allows FLevelStreamingGCHelper to be used in a commandlet */
-	static bool bEnabledForCommandlet;
-};
-
 /** Saved editor viewport state information */
 USTRUCT()
 struct ENGINE_API FLevelViewportInfo
diff --git a/Engine/Source/Runtime/Engine/Private/LevelStreaming.cpp b/Engine/Source/Runtime/Engine/Private/LevelStreaming.cpp
--- a/Engine/Source/Runtime/Engine/Private/LevelStreaming.cpp
+++ b/Engine/Source/Runtime/Engine/Private/LevelStreaming.cpp
@@ -1,6 +1,7 @@
 // Copyright Epic Games, Inc. All Rights Reserved.
 
 #include "Engine/LevelStreaming.h"
+#include "Engine/LevelStreamingGCHelper.h"
 #include "ContentStreaming.h"
 #include "Misc/App.h"
 #include "UObject/Package.h"
diff --git a/Engine/Source/Runtime/Engine/Private/LevelStreamingGCHelper.cpp b/Engine/Source/Runtime/Engine/Private/LevelStreamingGCHelper.cpp
new file mode 100644
--- /dev/null
+++ b/Engine/Source/Runtime/Engine/Private/LevelStreamingGCHelper.cpp
@@ -0,0 +1,196 @@
+// Copyright Epic Games, Inc. All Rights Reserved.
+
+#include "Engine/LevelStreamingGCHelper.h"
+#include "Engine/LevelStreaming.h"
+#include "Engine/CoreSettings.h"
+#include "Engine/NetDriver.h"
+#include "Engine/Engine.h"
+#include "UObject/UObjectGlobals.h"
+#include "UObject/UObjectGlobalsInternal.h"
+#include "UObject/ReferenceChainSearch.h"
+#include "UObject/UObjectIterator.h"
+#include "CoreGlobals.h"
+
+/************************************************************************/
+/* FLevelStreamingGCHelper implementation                               */
+/************************************************************************/
+FLevelStreamingGCHelper::FOnGCStreamedOutLevelsEvent FLevelStreamingGCHelper::OnGCStreamedOutLevels;
+TArray<TWeakObjectPtr<ULevel> > FLevelStreamingGCHelper::LevelsPendingUnload;
+TArray<FName> FLevelStreamingGCHelper::LevelPackageNames;
+bool FLevelStreamingGCHelper::bEnabledForCommandlet = false;
+
+void FLevelStreamingGCHelper::AddGarbageCollectorCallback()
+{
+	// Only register for garbage collection once
+	static bool GarbageCollectAdded = false;
+	if ( GarbageCollectAdded == false )
+	{
+		FCoreUObjectDelegates::GetPreGarbageCollectDelegate().AddStatic( FLevelStreamingGCHelper::PrepareStreamedOutLevelsForGC );
+		FCoreUObjectDelegates::GetPostGarbageCollect().AddStatic( FLevelStreamingGCHelper::VerifyLevelsGotRemovedByGC );
+		GarbageCollectAdded = true;
+	}
+}
+
+void FLevelStreamingGCHelper::EnableForCommandlet()
+{
+	check(IsRunningCommandlet());
+	bEnabledForCommandlet = true;
+}
+
+void FLevelStreamingGCHelper::RequestUnload( ULevel* InLevel )
+{
+	if (!IsRunningCommandlet() || bEnabledForCommandlet)
+	{
+		check( InLevel );
+		check( InLevel->bIsVisible == false );
+		LevelsPendingUnload.AddUnique( InLevel );
+	}
+}
+
+void FLevelStreamingGCHelper::CancelUnloadRequest( ULevel* InLevel )
+{
+	LevelsPendingUnload.Remove( InLevel );
+}
+
+void FLevelStreamingGCHelper::PrepareStreamedOutLevelsForGC()
+{
+	TRACE_CPUPROFILER_EVENT_SCOPE(FLevelStreamingGCHelper::PrepareStreamedOutLevelsForGC);
+	if (LevelsPendingUnload.Num() > 0)
+	{
+		OnGCStreamedOutLevels.Broadcast();
+	}
+	
+	// Iterate over all level objects that want to be unloaded.
+	for( int32 LevelIndex=0; LevelIndex < LevelsPendingUnload.Num(); LevelIndex++ )
+	{
+		ULevel*	Level = LevelsPendingUnload[LevelIndex].Get();
+
+		if( Level && ((!GIsEditor || bEnabledForCommandlet) || Level->GetOutermost()->HasAnyPackageFlags(PKG_PlayInEditor)))
+		{
+			UPackage* LevelPackage = Level->GetOutermost();
+			UE_LOG(LogStreaming, Log, TEXT("PrepareStreamedOutLevelsForGC called on '%s'"), *LevelPackage->GetName() );
+			
+			for (const FWorldContext& Context : GEngine->GetWorldContexts())
+			{
+				UWorld* World = Context.World();
+				if (World)
+				{
+					// This can never ever be called during tick; same goes for GC in general.
+					check( !World->bInTick );
+
+					FWorldContext& MutableContext = GEngine->GetWorldContextFromWorldChecked(World);
+					for (FNamedNetDriver& Driver : MutableContext.ActiveNetDrivers)
+					{
+						if (Driver.NetDriver != nullptr)
+						{
+							// The net driver must remove this level and its actors from the packagemap, or else
+							// the client package map will keep hard refs to them and prevent GC
+							Driver.NetDriver->NotifyStreamingLevelUnload(Level);
+						}
+					}
+
+					// Broadcast level unloaded event to blueprints through level streaming objects
+					ULevelStreaming::BroadcastLevelLoadedStatus(World, LevelPackage->GetFName(), false);
+				}
+			}
+
+			// Make sure that this package has been unloaded after GC pass.
+			LevelPackageNames.Add( LevelPackage->GetFName() );
+
+			Level->CleanupLevel();
+
+			// Mark world and all other package subobjects as pending kill
+			// This will destroy metadata objects and any other objects left behind
+			TSet<UPackage*> Packages;
+			ForEachObjectWithOuter(Level->GetOutermostObject(), [&Packages](UObject* Object)
+			{
+				bool bIsAlreadyInSet = false;
+				UPackage* Package = Object->GetPackage();
+				Packages.Add(Package, &bIsAlreadyInSet);
+				if (!bIsAlreadyInSet)
+				{
+					ForEachObjectWithPackage(Package, [](UObject* PackageObject) { PackageObject->MarkAsGarbage(); return true; }, true, RF_NoFlags, EInternalObjectFlags::Garbage);
+					Package->MarkAsGarbage();
+				}
+			}, true, RF_NoFlags, EInternalObjectFlags::Garbage);
+
+			if (!UObjectBaseUtility::IsPendingKillEnabled())
+			{
+				// Rename the packages that we are streaming out so that we can possibly reload another copy of them
+				for (UPackage* Package : Packages)
+				{
+					FCoreUObjectInternalDelegates::GetOnLeakedPackageRenameDelegate().Broadcast(Package);
+					FName NewName = MakeUniqueObjectName(nullptr, UPackage::StaticClass(), Package->GetFName());
+					Package->Rename(*NewName.ToString(), nullptr, REN_ForceNoResetLoaders | REN_DontCreateRedirectors | REN_NonTransactional);
+				}
+
+#if !WITH_EDITOR
+				// Clear the level actors array to maximize the memory savings from GC if there are outstanding references to some actors.
+				// Don't do this outside the editor now til we can validate it works properly with external packages.
+				Level->Actors.Empty();
+				Level->ActorsForGC.Empty();
+				Level->ActorCluster = nullptr; // The actor cluster has been marked garbage so will be dissolved, also drop a reference to it here since it has an internal array of pointers
+#endif
+			}
+
+			Level->CleanupReferences();
+		}
+	}
+
+	LevelsPendingUnload.Empty();
+}
+
+void FLevelStreamingGCHelper::VerifyLevelsGotRemovedByGC()
+{
+	if (LevelPackageNames.Num())
+	{
+#if DO_GUARD_SLOW
+		const bool bShouldVerify = !GIsEditor;
+#else
+		const bool bShouldVerify = false;
+#endif
+
+#if !UE_BUILD_SHIPPING
+		if (bShouldVerify || GLevelStreamingForceVerifyLevelsGotRemovedByGC)
+		{
+			int32 FailCount = 0;
+			const bool bIsAsyncLoading = IsAsyncLoading();
+			// Iterate over all objects and find out whether they reside in a GC'ed level package.
+			for (FThreadSafeObjectIterator It; It; ++It)
+			{
+				UObject* Object = *It;
+				// Check whether object's outermost is in the list.
+				if (LevelPackageNames.Find(Object->GetOutermost()->GetFName()) != INDEX_NONE
+					// But disregard package object itself.
+					&& !Object->IsA(UPackage::StaticClass()))
+				{
+					if (bIsAsyncLoading && Object->HasAnyInternalFlags(EInternalObjectFlags::Async|EInternalObjectFlags::AsyncLoading))
+					{
+						UE_LOG(LogStreaming, Display, TEXT("Level object %s isn't released by async loading yet, "
+								"it will get garbage collected next time instead."),
+								*Object->GetFullName());
+					}
+					else
+					{
+						UE_LOG(LogStreaming, Warning, TEXT("Level object %s didn't get garbage collected! Trying to find culprit, though this might crash. Try increasing stack size if it does. Referenced by:"), *Object->GetFullName());
+						FReferenceChainSearch RefChainSearch(Object, EReferenceChainSearchMode::Shortest | EReferenceChainSearchMode::PrintResults);
+						FailCount++;
+					}
+				}
+			}
+			if (FailCount > 0)
+			{
+				UE_LOG(LogStreaming, Error, TEXT("Streamed out levels were not completely garbage collected! Please see previous log entries."));
+				// If not forcing verify, consider errors as fatal
+				check(GLevelStreamingForceVerifyLevelsGotRemovedByGC);
+			}
+		}
+#endif
+		LevelPackageNames.Empty();
+	}
+}
+
+int32 FLevelStreamingGCHelper::GetNumLevelsPendingPurge()
+{
+	return LevelsPendingUnload.Num();
+}
\ No newline at end of file
diff --git a/Engine/Source/Runtime/Engine/Private/UnrealEngine.cpp b/Engine/Source/Runtime/Engine/Private/UnrealEngine.cpp
--- a/Engine/Source/Runtime/Engine/Private/UnrealEngine.cpp
+++ b/Engine/Source/Runtime/Engine/Private/UnrealEngine.cpp
@@ -18,6 +18,7 @@ UnrealEngine.cpp: Implements the UEngine class and helpers.
 #include "RenderingThread.h"
 #include "RHIStaticStates.h"
 #include "Engine/TextureStreamingTypes.h"
+#include "Engine/LevelStreamingGCHelper.h"
 #include "Components/PrimitiveComponent.h"
 #include "AI/NavigationSystemBase.h"
 #include "AI/NavigationSystemConfig.h"
diff --git a/Engine/Source/Runtime/Engine/Private/World.cpp b/Engine/Source/Runtime/Engine/Private/World.cpp
--- a/Engine/Source/Runtime/Engine/Private/World.cpp
+++ b/Engine/Source/Runtime/Engine/Private/World.cpp
@@ -48,6 +48,7 @@
 #include "RendererInterface.h"
 #include "Engine/Texture2D.h"
 #include "Engine/LevelStreaming.h"
+#include "Engine/LevelStreamingGCHelper.h"
 #include "Engine/StaticMesh.h"
 #include "Materials/MaterialInstance.h"
 #include "Engine/LocalPlayer.h"
@@ -3342,190 +3343,6 @@ FGameTime FGameTime::GetTimeSinceAppStart()
 	return FGameTime::CreateUndilated(FApp::GetCurrentTime() - GStartTime, FApp::GetDeltaTime());
 }
 
-/************************************************************************/
-/* FLevelStreamingGCHelper implementation                               */
-/************************************************************************/
-FLevelStreamingGCHelper::FOnGCStreamedOutLevelsEvent FLevelStreamingGCHelper::OnGCStreamedOutLevels;
-TArray<TWeakObjectPtr<ULevel> > FLevelStreamingGCHelper::LevelsPendingUnload;
-TArray<FName> FLevelStreamingGCHelper::LevelPackageNames;
-bool FLevelStreamingGCHelper::bEnabledForCommandlet = false;
-
-void FLevelStreamingGCHelper::AddGarbageCollectorCallback()
-{
-	// Only register for garbage collection once
-	static bool GarbageCollectAdded = false;
-	if ( GarbageCollectAdded == false )
-	{
-		FCoreUObjectDelegates::GetPreGarbageCollectDelegate().AddStatic( FLevelStreamingGCHelper::PrepareStreamedOutLevelsForGC );
-		FCoreUObjectDelegates::GetPostGarbageCollect().AddStatic( FLevelStreamingGCHelper::VerifyLevelsGotRemovedByGC );
-		GarbageCollectAdded = true;
-	}
-}
-
-void FLevelStreamingGCHelper::EnableForCommandlet()
-{
-	check(IsRunningCommandlet());
-	bEnabledForCommandlet = true;
-}
-
-void FLevelStreamingGCHelper::RequestUnload( ULevel* InLevel )
-{
-	if (!IsRunningCommandlet() || bEnabledForCommandlet)
-	{
-		check( InLevel );
-		check( InLevel->bIsVisible == false );
-		LevelsPendingUnload.AddUnique( InLevel );
-	}
-}
-
-void FLevelStreamingGCHelper::CancelUnloadRequest( ULevel* InLevel )
-{
-	LevelsPendingUnload.Remove( InLevel );
-}
-
-void FLevelStreamingGCHelper::PrepareStreamedOutLevelsForGC()
-{
-	TRACE_CPUPROFILER_EVENT_SCOPE(FLevelStreamingGCHelper::PrepareStreamedOutLevelsForGC);
-	if (LevelsPendingUnload.Num() > 0)
-	{
-		OnGCStreamedOutLevels.Broadcast();
-	}
-	
-	// Iterate over all level objects that want to be unloaded.
-	for( int32 LevelIndex=0; LevelIndex < LevelsPendingUnload.Num(); LevelIndex++ )
-	{
-		ULevel*	Level = LevelsPendingUnload[LevelIndex].Get();
-
-		if( Level && ((!GIsEditor || bEnabledForCommandlet) || Level->GetOutermost()->HasAnyPackageFlags(PKG_PlayInEditor)))
-		{
-			UPackage* LevelPackage = Level->GetOutermost();
-			UE_LOG(LogStreaming, Log, TEXT("PrepareStreamedOutLevelsForGC called on '%s'"), *LevelPackage->GetName() );
-			
-			for (const FWorldContext& Context : GEngine->GetWorldContexts())
-			{
-				UWorld* World = Context.World();
-				if (World)
-				{
-					// This can never ever be called during tick; same goes for GC in general.
-					check( !World->bInTick );
-
-					FWorldContext& MutableContext = GEngine->GetWorldContextFromWorldChecked(World);
-					for (FNamedNetDriver& Driver : MutableContext.ActiveNetDrivers)
-					{
-						if (Driver.NetDriver != nullptr)
-						{
-							// The net driver must remove this level and its actors from the packagemap, or else
-							// the client package map will keep hard refs to them and prevent GC
-							Driver.NetDriver->NotifyStreamingLevelUnload(Level);
-						}
-					}
-
-					// Broadcast level unloaded event to blueprints through level streaming objects
-					ULevelStreaming::BroadcastLevelLoadedStatus(World, LevelPackage->GetFName(), false);
-				}
-			}
-
-			// Make sure that this package has been unloaded after GC pass.
-			LevelPackageNames.Add( LevelPackage->GetFName() );
-
-			Level->CleanupLevel();
-
-			// Mark world and all other package subobjects as pending kill
-			// This will destroy metadata objects and any other objects left behind
-			TSet<UPackage*> Packages;
-			ForEachObjectWithOuter(Level->GetOutermostObject(), [&Packages](UObject* Object)
-			{
-				bool bIsAlreadyInSet = false;
-				UPackage* Package = Object->GetPackage();
-				Packages.Add(Package, &bIsAlreadyInSet);
-				if (!bIsAlreadyInSet)
-				{
-					ForEachObjectWithPackage(Package, [](UObject* PackageObject) { PackageObject->MarkAsGarbage(); return true; }, true, RF_NoFlags, EInternalObjectFlags::Garbage);
-					Package->MarkAsGarbage();
-				}
-			}, true, RF_NoFlags, EInternalObjectFlags::Garbage);
-
-			if (!UObjectBaseUtility::IsPendingKillEnabled())
-			{
-				// Rename the packages that we are streaming out so that we can possibly reload another copy of them
-				for (UPackage* Package : Packages)
-				{
-					FCoreUObjectInternalDelegates::GetOnLeakedPackageRenameDelegate().Broadcast(Package);
-					FName NewName = MakeUniqueObjectName(nullptr, UPackage::StaticClass(), Package->GetFName());
-					Package->Rename(*NewName.ToString(), nullptr, REN_ForceNoResetLoaders | REN_DontCreateRedirectors | REN_NonTransactional);
-				}
-
-#if !WITH_EDITOR
-				// Clear the level actors array to maximize the memory savings from GC if there are outstanding references to some actors.
-				// Don't do this outside the editor now til we can validate it works properly with external packages.
-				Level->Actors.Empty();
-				Level->ActorsForGC.Empty();
-				Level->ActorCluster = nullptr; // The actor cluster has been marked garbage so will be dissolved, also drop a reference to it here since it has an internal array of pointers
-#endif
-			}
-
-			Level->CleanupReferences();
-		}
-	}
-
-	LevelsPendingUnload.Empty();
-}
-
-void FLevelStreamingGCHelper::VerifyLevelsGotRemovedByGC()
-{
-	if (LevelPackageNames.Num())
-	{
-#if DO_GUARD_SLOW
-		const bool bShouldVerify = !GIsEditor;
-#else
-		const bool bShouldVerify = false;
-#endif
-
-#if !UE_BUILD_SHIPPING
-		if (bShouldVerify || GLevelStreamingForceVerifyLevelsGotRemovedByGC)
-		{
-			int32 FailCount = 0;
-			const bool bIsAsyncLoading = IsAsyncLoading();
-			// Iterate over all objects and find out whether they reside in a GC'ed level package.
-			for (FThreadSafeObjectIterator It; It; ++It)
-			{
-				UObject* Object = *It;
-				// Check whether object's outermost is in the list.
-				if (LevelPackageNames.Find(Object->GetOutermost()->GetFName()) != INDEX_NONE
-					// But disregard package object itself.
-					&& !Object->IsA(UPackage::StaticClass()))
-				{
-					if (bIsAsyncLoading && Object->HasAnyInternalFlags(EInternalObjectFlags::Async|EInternalObjectFlags::AsyncLoading))
-					{
-						UE_LOG(LogWorld, Display, TEXT("Level object %s isn't released by async loading yet, "
-								"it will get garbage collected next time instead."),
-								*Object->GetFullName());
-					}
-					else
-					{
-						UE_LOG(LogWorld, Warning, TEXT("Level object %s didn't get garbage collected! Trying to find culprit, though this might crash. Try increasing stack size if it does. Referenced by:"), *Object->GetFullName());
-						FReferenceChainSearch RefChainSearch(Object, EReferenceChainSearchMode::Shortest | EReferenceChainSearchMode::PrintResults);
-						FailCount++;
-					}
-				}
-			}
-			if (FailCount > 0)
-			{
-				UE_LOG(LogWorld, Error, TEXT("Streamed out levels were not completely garbage collected! Please see previous log entries."));
-				// If not forcing verify, consider errors as fatal
-				check(GLevelStreamingForceVerifyLevelsGotRemovedByGC);
-			}
-		}
-#endif
-		LevelPackageNames.Empty();
-	}
-}
-
-int32 FLevelStreamingGCHelper::GetNumLevelsPendingPurge()
-{
-	return LevelsPendingUnload.Num();
-}
-
 void UWorld::RenameToPIEWorld(int32 PIEInstanceID)
 {
 #if WITH_EDITOR
diff --git a/Engine/Source/Runtime/Engine/Private/WorldPartition/WorldPartitionLevelStreamingPolicy.cpp b/Engine/Source/Runtime/Engine/Private/WorldPartition/WorldPartitionLevelStreamingPolicy.cpp
--- a/Engine/Source/Runtime/Engine/Private/WorldPartition/WorldPartitionLevelStreamingPolicy.cpp
+++ b/Engine/Source/Runtime/Engine/Private/WorldPartition/WorldPartitionLevelStreamingPolicy.cpp
@@ -15,6 +15,7 @@
 #include "Engine/Level.h"
 #include "Engine/World.h"
 #include "Engine/Canvas.h"
+#include "Engine/LevelStreamingGCHelper.h"
 
 #include UE_INLINE_GENERATED_CPP_BY_NAME(WorldPartitionLevelStreamingPolicy)
 
