World Partition: fix FWorldPartitionActorDescUtils::GetActorNativeClassFromAssetData to detect short class names and handle them properly, to avoid a warning.

#rnx

diff --git a/Engine/Source/Runtime/Engine/Private/WorldPartition/WorldPartitionActorDescUtils.cpp b/Engine/Source/Runtime/Engine/Private/WorldPartition/WorldPartitionActorDescUtils.cpp
--- a/Engine/Source/Runtime/Engine/Private/WorldPartition/WorldPartitionActorDescUtils.cpp
+++ b/Engine/Source/Runtime/Engine/Private/WorldPartition/WorldPartitionActorDescUtils.cpp
@@ -32,8 +32,12 @@ UClass* FWorldPartitionActorDescUtils::GetActorNativeClassFromAssetData(const FA
 		// Look for a class redirectors
 		const FCoreRedirectObjectName OldClassName = FCoreRedirectObjectName(*ActorClassName, NAME_None, *ActorPackageName);
 		const FCoreRedirectObjectName NewClassName = FCoreRedirects::GetRedirectedName(ECoreRedirectFlags::Type_Class, OldClassName);
+		
+		// Handle deprecated short class names
+		const FTopLevelAssetPath ClassPath = FAssetData::TryConvertShortClassNameToPathName(*NewClassName.ToString(), ELogVerbosity::Log);
 
-		return UClass::TryFindTypeSlow<UClass>(NewClassName.ToString(), EFindFirstObjectOptions::ExactClass);
+		// Lookup the native class
+		return UClass::TryFindTypeSlow<UClass>(ClassPath.ToString(), EFindFirstObjectOptions::ExactClass);
 	}
 	return nullptr;
 }
