Moved r.ExposureOffset out of #if !(UE_BUILD_SHIPPING || UE_BUILD_TEST) block
- So it can be used in Test builds.
- Still tagged as ECVF_Cheat so it can't be changed in a final release build.

diff --git a/Engine/Source/Runtime/Engine/Private/SceneView.cpp b/Engine/Source/Runtime/Engine/Private/SceneView.cpp
--- a/Engine/Source/Runtime/Engine/Private/SceneView.cpp
+++ b/Engine/Source/Runtime/Engine/Private/SceneView.cpp
@@ -88,12 +88,6 @@ static TAutoConsoleVariable<int32> CVarShadowFreezeCamera(
 	TEXT("1: freeze camera at current location"),
 	ECVF_Cheat);
 
-static TAutoConsoleVariable<float> CVarExposureOffset(
-	TEXT("r.ExposureOffset"),
-	0.0f,
-	TEXT("For adjusting the exposure on top of post process settings and eye adaptation. For developers only. 0:default"),
-	ECVF_Cheat);
-
 static TAutoConsoleVariable<int32> CVarRenderTimeFrozen(
 	TEXT("r.RenderTimeFrozen"),
 	0,
@@ -139,6 +133,12 @@ static TAutoConsoleVariable<float> CVarSSAOFadeRadiusScale(
 	TEXT(" 0.01:smallest .. 1.0:normal (default), <1:smaller, >1:larger"),
 	ECVF_Cheat | ECVF_RenderThreadSafe);
 
+static TAutoConsoleVariable<float> CVarExposureOffset(
+	TEXT("r.ExposureOffset"),
+	0.0f,
+	TEXT("For adjusting the exposure on top of post process settings and eye adaptation. 0: default"),
+	ECVF_Cheat);
+
 // Engine default (project settings):
 
 static TAutoConsoleVariable<int32> CVarDefaultBloom(
@@ -2130,12 +2130,12 @@ void FSceneView::EndFinalPostprocessSettings(const FSceneViewInitOptions& ViewIn
 		FinalPostProcessSettings.ToneCurveAmount = 0;
 	}
 
-#if !(UE_BUILD_SHIPPING || UE_BUILD_TEST)
 	{
 		float Value = CVarExposureOffset.GetValueOnGameThread();
 		FinalPostProcessSettings.AutoExposureBias += Value;
 	}
 
+#if !(UE_BUILD_SHIPPING || UE_BUILD_TEST)
 	{
 		float& DepthBlurAmount = FinalPostProcessSettings.DepthOfFieldDepthBlurAmount;
 
