Fix a deadlock when skeletalmesh component is tick and we are compiling the skeletalmesh.

#rnx

diff --git a/Engine/Source/Runtime/Engine/Private/Components/SkeletalMeshComponent.cpp b/Engine/Source/Runtime/Engine/Private/Components/SkeletalMeshComponent.cpp
--- a/Engine/Source/Runtime/Engine/Private/Components/SkeletalMeshComponent.cpp
+++ b/Engine/Source/Runtime/Engine/Private/Components/SkeletalMeshComponent.cpp
@@ -1509,6 +1509,14 @@ void USkeletalMeshComponent::TickComponent(float DeltaTime, enum ELevelTick Tick
 {
 	CSV_SCOPED_TIMING_STAT_EXCLUSIVE(Animation);
 
+#if WITH_EDITOR
+	//Do not tick skeletal mesh component if the skeletal mesh asset is compiling
+	if (GetSkeletalMeshAsset() && GetSkeletalMeshAsset()->IsCompiling())
+	{
+		return;
+	}
+#endif
+
 	if (ClothingSimulation)
 	{
 		ClothingSimulation->UpdateWorldForces(this);
