Fixed pose watch elements being added in the transient package instead of the pose watch owner.

diff --git a/Engine/Source/Runtime/Engine/Classes/Engine/PoseWatch.h b/Engine/Source/Runtime/Engine/Classes/Engine/PoseWatch.h
--- a/Engine/Source/Runtime/Engine/Classes/Engine/PoseWatch.h
+++ b/Engine/Source/Runtime/Engine/Classes/Engine/PoseWatch.h
@@ -440,7 +440,7 @@ protected:
 
 template< class TElementType > TObjectPtr<TElementType> UPoseWatch::AddElement(const FText InLabel, const FName InIconName)
 {
-	TObjectPtr<TElementType> PoseWatchElement(NewObject<TElementType>());
+	TObjectPtr<TElementType> PoseWatchElement(NewObject<TElementType>(this));
 	PoseWatchElement->SetParent(this);
 	PoseWatchElement->SetUniqueLabel(InLabel);
 	PoseWatchElement->SetIconName(InIconName);
diff --git a/Engine/Source/Runtime/Engine/Private/Engine/PoseWatch.cpp b/Engine/Source/Runtime/Engine/Private/Engine/PoseWatch.cpp
--- a/Engine/Source/Runtime/Engine/Private/Engine/PoseWatch.cpp
+++ b/Engine/Source/Runtime/Engine/Private/Engine/PoseWatch.cpp
@@ -778,6 +778,12 @@ void UPoseWatch::Serialize(FArchive& Ar)
 		NewPoseWatchElement->ViewportOffset = ViewportOffset_DEPRECATED;
 		Elements.Add(NewPoseWatchElement);
 	}
+
+	// UE-162694 Remove pose watch elements that were previously incorrectly serialized
+	Elements.RemoveAll([](const TObjectPtr<UPoseWatchElement>& Element)
+	{
+		return Element == nullptr;
+	});
 }
 #endif // WITH_EDITORONLY_DATA
 
