Fix Rendering.ShaderModels.Decal.Animated

- UMaterialInterface::SubmitRemainingJobsForWorld uses a FObjectCacheContextScope which only processes Primitive Components when looking for incomplete materials.
- The DecalComponent is not a PrimitiveComponent so the materials were never being compiled causing the shaders to be missing during the test.
- Similar to Add UI and PP Materials we need to explicitly iterate over DecalComponents, get their materials, and submit compilation for them.  This is the easiest and most straight forward fix.

diff --git a/Engine/Source/Runtime/Engine/Private/Materials/MaterialInterface.cpp b/Engine/Source/Runtime/Engine/Private/Materials/MaterialInterface.cpp
--- a/Engine/Source/Runtime/Engine/Private/Materials/MaterialInterface.cpp
+++ b/Engine/Source/Runtime/Engine/Private/Materials/MaterialInterface.cpp
@@ -31,6 +31,7 @@
 #include "ObjectCacheContext.h"
 #include "MaterialCachedData.h"
 #include "Misc/ScopedSlowTask.h"
+#include "Components/DecalComponent.h"
 
 #if WITH_EDITOR
 #include "ObjectCacheEventSink.h"
@@ -580,6 +581,24 @@ void UMaterialInterface::SubmitRemainingJobsForWorld(UWorld* World, EMaterialSha
 		}
 	}
 
+	// Add Decal Component Materials.
+	for (TObjectIterator<UDecalComponent> It; It; ++It)
+	{
+		if (It)
+		{
+			TArray<UMaterialInterface*> OutMaterials;
+			It->GetUsedMaterials(OutMaterials);
+
+			for (UMaterialInterface* MaterialInterface : OutMaterials)
+			{
+				if (MaterialInterface && !MaterialInterface->IsComplete())
+				{
+					MaterialsToCache.Add(MaterialInterface);
+				}
+			}
+		}
+	}
+
 	if (MaterialsToCache.Num())
 	{
 		FMaterialUpdateContext UpdateContext(FMaterialUpdateContext::EOptions::SyncWithRenderingThread);
