Fix WP streaming status showing (Standalone, Listen Server, Client) at the wrong place.

diff --git a/Engine/Source/Runtime/Engine/Private/WorldPartition/WorldPartitionSubsystem.cpp b/Engine/Source/Runtime/Engine/Private/WorldPartition/WorldPartitionSubsystem.cpp
--- a/Engine/Source/Runtime/Engine/Private/WorldPartition/WorldPartitionSubsystem.cpp
+++ b/Engine/Source/Runtime/Engine/Private/WorldPartition/WorldPartitionSubsystem.cpp
@@ -378,13 +378,13 @@ void UWorldPartitionSubsystem::Draw(UCanvas* Canvas, class APlayerController* PC
 			if (IsAsyncLoading()) { StatusText += TEXT("(AsyncLoading) "); }
 			if (StatusText.IsEmpty()) { StatusText = TEXT("(Idle) "); }
 
-			StatusText += FString::Printf(TEXT("(%s) "), *GetDebugStringForWorld(GetWorld()));
+			FString DebugWorldText = FString::Printf(TEXT("(%s)"), *GetDebugStringForWorld(GetWorld()));
 			if (WorldPartition->IsServer())
 			{
-				StatusText += FString::Printf(TEXT("(Server Streaming %s)"), WorldPartition->IsServerStreamingEnabled() ? TEXT("Enabled") : TEXT("Disabled"));
+				DebugWorldText += FString::Printf(TEXT(" (Server Streaming %s)"), WorldPartition->IsServerStreamingEnabled() ? TEXT("Enabled") : TEXT("Disabled"));
 			}
 			
-			const FString Text = FString::Printf(TEXT("Streaming Status: %s"), *StatusText);
+			const FString Text = FString::Printf(TEXT("Streaming Status for %s: %s"), *DebugWorldText, *StatusText);
 			FWorldPartitionDebugHelper::DrawText(Canvas, Text, GEngine->GetSmallFont(), FColor::White, CurrentOffset);
 		}
 
