Fixed issue where simulated proxy deviates from the expected path when playing root motion animation with rotation in it

diff --git a/Engine/Source/Runtime/Engine/Private/Components/CharacterMovementComponent.cpp b/Engine/Source/Runtime/Engine/Private/Components/CharacterMovementComponent.cpp
--- a/Engine/Source/Runtime/Engine/Private/Components/CharacterMovementComponent.cpp
+++ b/Engine/Source/Runtime/Engine/Private/Components/CharacterMovementComponent.cpp
@@ -1627,10 +1627,14 @@ void UCharacterMovementComponent::SimulatedTick(float DeltaSeconds)
 			}
 		}
 
+		const FQuat OldRotationQuat = UpdatedComponent->GetComponentQuat();
+		const FVector OldLocation = UpdatedComponent->GetComponentLocation();
+
+		USkeletalMeshComponent* Mesh = CharacterOwner->GetMesh();
+		const FVector SavedMeshRelativeLocation = Mesh ? Mesh->GetRelativeLocation() : FVector::ZeroVector;
+
 		if( RootMotionParams.bHasRootMotion )
 		{
-			const FQuat OldRotationQuat = UpdatedComponent->GetComponentQuat();
-			const FVector OldLocation = UpdatedComponent->GetComponentLocation();
 			SimulateRootMotion(DeltaSeconds, RootMotionParams.GetRootMotionTransform());
 
 #if !(UE_BUILD_SHIPPING)
@@ -1655,6 +1659,18 @@ void UCharacterMovementComponent::SimulatedTick(float DeltaSeconds)
 		{
 			CharacterOwner->SimulatedRootMotionPositionFixup(DeltaSeconds);
 		}
+
+		if (!bNetworkSmoothingComplete && (NetworkSmoothingMode == ENetworkSmoothingMode::Linear))
+		{
+			// Same mesh with different rotation?
+			const FQuat NewCapsuleRotation = UpdatedComponent->GetComponentQuat();
+			if (Mesh == CharacterOwner->GetMesh() && !NewCapsuleRotation.Equals(OldRotationQuat, 1e-6f) && ClientPredictionData)
+			{
+				// Smoothing should lerp toward this new rotation target, otherwise it will just try to go back toward the old rotation.
+				ClientPredictionData->MeshRotationTarget = NewCapsuleRotation;
+				Mesh->SetRelativeLocationAndRotation(SavedMeshRelativeLocation, CharacterOwner->GetBaseRotationOffset());
+			}
+		}
 	}
 	else if (CurrentRootMotion.HasActiveRootMotionSources())
 	{
