Fix Sequence Evaluators not working with State Machine Relevancy Nodes

diff --git a/Engine/Source/Runtime/Engine/Classes/Animation/AnimNode_AssetPlayerBase.h b/Engine/Source/Runtime/Engine/Classes/Animation/AnimNode_AssetPlayerBase.h
--- a/Engine/Source/Runtime/Engine/Classes/Animation/AnimNode_AssetPlayerBase.h
+++ b/Engine/Source/Runtime/Engine/Classes/Animation/AnimNode_AssetPlayerBase.h
@@ -55,6 +55,7 @@ struct ENGINE_API FAnimNode_AssetPlayerBase : public FAnimNode_AssetPlayerReleva
 	virtual void SetAccumulatedTime(float NewTime) override;
 	virtual float GetCachedBlendWeight() const override;
 	virtual void ClearCachedBlendWeight() override;
+	virtual float GetCurrentAssetTimePlayRateAdjusted() const;
 	// --- End of FAnimNode_RelevantAssetPlayerBase ---
 
 private:
diff --git a/Engine/Source/Runtime/Engine/Private/Animation/AnimNode_AssetPlayerBase.cpp b/Engine/Source/Runtime/Engine/Private/Animation/AnimNode_AssetPlayerBase.cpp
--- a/Engine/Source/Runtime/Engine/Private/Animation/AnimNode_AssetPlayerBase.cpp
+++ b/Engine/Source/Runtime/Engine/Private/Animation/AnimNode_AssetPlayerBase.cpp
@@ -71,3 +71,7 @@ void FAnimNode_AssetPlayerBase::ClearCachedBlendWeight()
 	BlendWeight = 0.0f;
 }
 
+float FAnimNode_AssetPlayerBase::GetCurrentAssetTimePlayRateAdjusted() const
+{
+	return  GetCurrentAssetTime();
+}
