Relands multiple dyn res heuristic

diff --git a/Engine/Source/Runtime/Engine/Private/DynamicResolution.cpp b/Engine/Source/Runtime/Engine/Private/DynamicResolution.cpp
--- a/Engine/Source/Runtime/Engine/Private/DynamicResolution.cpp
+++ b/Engine/Source/Runtime/Engine/Private/DynamicResolution.cpp
@@ -101,12 +101,10 @@ DynamicRenderScaling::FHeuristicSettings GetPrimaryDynamicResolutionSettings()
 	BudgetSetting.MinResolutionFraction      = DynamicRenderScaling::GetPercentageCVarToFraction(CVarDynamicResMinSP);
 	BudgetSetting.MaxResolutionFraction      = DynamicRenderScaling::GetPercentageCVarToFraction(CVarDynamicResMaxSP);
 	BudgetSetting.UpperBoundQuantization     = CVarUpperBoundQuantization.GetValueOnAnyThread();
-	// BudgetSetting.BudgetMs depends on the cost of other buckets.
+	BudgetSetting.BudgetMs                   = CVarFrameTimeBudget.GetValueOnAnyThread();
 	BudgetSetting.ChangeThreshold            = DynamicRenderScaling::GetPercentageCVarToFraction(CVarChangeThreshold);
+	BudgetSetting.TargetedHeadRoom           = DynamicRenderScaling::GetPercentageCVarToFraction(CVarTargetedGPUHeadRoomPercentage);
 	BudgetSetting.IncreaseAmortizationFactor = CVarIncreaseAmortizationFactor.GetValueOnAnyThread();
-
-	// CVarTargetedGPUHeadRoomPercentage is taken into account on the entire frame instead.
-	BudgetSetting.TargetedHeadRoom = 0.0f;
 	return BudgetSetting;
 }
 
@@ -217,30 +215,43 @@ void FDynamicResolutionHeuristicProxy::CommitPreviousFrameGPUTimings_RenderThrea
 
 void FDynamicResolutionHeuristicProxy::RefreshCurentFrameResolutionFraction_RenderThread()
 {
-	// GPU time budget per frame=.
-	const float FrameTimeBudgetMs = CVarFrameTimeBudget.GetValueOnRenderThread();
-
-	// Targeted GPU time, lower than budget to diggest noise.
-	const float TargetedGPUBusyTimeMs = FrameTimeBudgetMs * (1.0f - CVarTargetedGPUHeadRoomPercentage.GetValueOnRenderThread() / 100.0f);
+	// Global constants.
+	const float FrameWeightExponent = CVarFrameWeightExponent.GetValueOnRenderThread();
+	const int32 MaxConsecutiveOverbudgetGPUFrameCount = FMath::Max(CVarMaxConsecutiveOverbudgetGPUFrameCount.GetValueOnRenderThread(), 2);
 
-	// Resolution fraction to downscale to when CPU bound.
-	const float CPUBoundResolutionFraction = 1.0f; // CVarCPUBoundScreenPercentage.GetValueOnRenderThread() / 100.0f;
+	const bool bCanChangeResolution = NumberOfFramesSinceScreenPercentageChange >= CVarFrameChangePeriod.GetValueOnRenderThread();
 
 	// New ResolutionFraction to use for this frame.
-	float NewFrameResolutionFraction = CurrentFrameResolutionFractions[GDynamicPrimaryResolutionFraction];
+	DynamicRenderScaling::TMap<float> NewFrameResolutionFractions = CurrentFrameResolutionFractions;
 
 	// Whether there is a GPU over budget panic.
-	bool bGPUOverbugetPanic = false;
+	bool bGlobalGPUOverbugetPanic = false;
 
-	// Compute NewFrameResolutionFraction if have an history to work with.
-	if (HistorySize > 0)
+	for (TLinkedList<DynamicRenderScaling::FBudget*>::TIterator BudgetIt(DynamicRenderScaling::FBudget::GetGlobalList()); BudgetIt; BudgetIt.Next())
 	{
-		const float FrameWeightExponent = CVarFrameWeightExponent.GetValueOnRenderThread();
-		const int32 MaxConsecutiveOverbudgetGPUFrameCount = FMath::Max(CVarMaxConsecutiveOverbudgetGPUFrameCount.GetValueOnRenderThread(), 2);
+		const DynamicRenderScaling::FBudget& Budget = **BudgetIt;
+		const DynamicRenderScaling::FHeuristicSettings& BudgetSettings = Budget.GetSettings();
+
+		const int32 BudgetHistorySize = FMath::Min(HistorySize, BudgetHistorySizes[Budget]);
 
-		const float OutLierTimeThreshold = 0.0f; //FrameTimeBudgetMs * CVarOutlierThreshold.GetValueOnRenderThread();
+		bool bEarlyReturn = BudgetHistorySize == 0;
 
-		NewFrameResolutionFraction = 0.0f;
+		if (Budget == GDynamicPrimaryResolutionFraction)
+		{
+			// NOP
+		}
+		else if (!BudgetSettings.IsEnabled())
+		{
+			bEarlyReturn = true;
+		}
+
+		// Early returns if not enough data to work with.
+		if (bEarlyReturn)
+		{
+			continue;
+		}
+
+		float NewFrameResolutionFraction = 0.0f;
 
 		// Total weight of NewFrameResolutionFraction.
 		float TotalWeight = 0.0f;
@@ -257,53 +268,31 @@ void FDynamicResolutionHeuristicProxy::RefreshCurentFrameResolutionFraction_Rend
 		// Number of frame browsed.
 		int32 FrameCount = 0;
 		int32 BrowsingFrameId = 0;
-		for (; BrowsingFrameId < HistorySize && bCaryOnBrowsingFrameHistory; BrowsingFrameId++)
+		for (; BrowsingFrameId < BudgetHistorySize && bCaryOnBrowsingFrameHistory; BrowsingFrameId++)
 		{
 			const FrameHistoryEntry& NextFrameEntry = GetPreviousFrameEntry(BrowsingFrameId - 1);
 			const FrameHistoryEntry& FrameEntry = GetPreviousFrameEntry(BrowsingFrameId + 0);
 			const FrameHistoryEntry& PreviousFrameEntry = GetPreviousFrameEntry(BrowsingFrameId + 1);
 
+			const float TotalFrameGPUBusyTimeMs = Budget == GDynamicPrimaryResolutionFraction ? FrameEntry.TotalFrameGPUBusyTimeMs : FrameEntry.BudgetTimingMs[Budget];
+
 			// Ignores frames that does not have any GPU timing yet. 
-			if (FrameEntry.TotalFrameGPUBusyTimeMs < 0)
+			if (TotalFrameGPUBusyTimeMs < 0)
 			{
 				continue;
 			}
 
-			// Compute the maximum thread time between game and render thread for this <i>th frame.
-			// We add a head room for render thread to detect when the GPU could be starving.
-			const float MaxThreadTimeMs = FMath::Max(
-				FrameEntry.RenderThreadTimeMs,
-				FrameEntry.GameThreadTimeMs);
-
 			// Whether bound by game thread.
-			const bool bIsGameThreadBound = FrameEntry.GameThreadTimeMs > FrameTimeBudgetMs;
+			const bool bIsGameThreadBound = FrameEntry.GameThreadTimeMs > GDynamicPrimaryResolutionFraction.GetSettings().BudgetMs;
 
 			// Whether bound by render thread.
-			const bool bIsRenderThreadBound = FrameEntry.RenderThreadTimeMs > TargetedGPUBusyTimeMs;
+			const bool bIsRenderThreadBound = FrameEntry.RenderThreadTimeMs > GDynamicPrimaryResolutionFraction.GetSettings().GetTargetedMs();
 
 			// Whether the frame is CPU bound.
 			const bool bIsCPUBound = bIsGameThreadBound || bIsRenderThreadBound;
 
 			// Whether GPU is over budget, when not CPU bound.
-			const bool bHasOverbudgetGPU = !bIsCPUBound && FrameEntry.TotalFrameGPUBusyTimeMs > FrameTimeBudgetMs;
-
-			// Whether this is possibly an outlier compared to previous and next frame.
-			// Normally should look for out lier compared to previous and next frames timings, however
-			// only interested in detected them if they are over budget, when cruising at targeted
-			// GPU busy time.
-			const bool bIsGPUOutlier =
-				bHasOverbudgetGPU &&
-				ConsecutiveOverbudgetGPUFramCount == 0 &&
-				OutLierTimeThreshold > 0.0f &&
-				FrameEntry.TotalFrameGPUBusyTimeMs > OutLierTimeThreshold &&
-				(!NextFrameEntry.HasGPUTimings() || (
-					NextFrameEntry.TotalFrameGPUBusyTimeMs < OutLierTimeThreshold&&
-					NextFrameEntry.GameThreadTimeMs < OutLierTimeThreshold&&
-					NextFrameEntry.RenderThreadTimeMs < OutLierTimeThreshold)) &&
-				(!PreviousFrameEntry.HasGPUTimings() || (
-					PreviousFrameEntry.TotalFrameGPUBusyTimeMs < OutLierTimeThreshold&&
-					PreviousFrameEntry.GameThreadTimeMs < OutLierTimeThreshold&&
-					PreviousFrameEntry.RenderThreadTimeMs < OutLierTimeThreshold));
+			const bool bHasOverbudgetGPU = !bIsCPUBound && TotalFrameGPUBusyTimeMs > BudgetSettings.BudgetMs;
 
 			// Look if this is multiple consecutive GPU over budget frames.
 			if (bHasOverbudgetGPU)
@@ -316,8 +305,6 @@ void FDynamicResolutionHeuristicProxy::RefreshCurentFrameResolutionFraction_Rend
 				{
 					// We ignore frames in history that happen before consecutive GPU overbudget frames.
 					bCaryOnBrowsingFrameHistory = false;
-
-					ensureMsgf(!bIsGPUOutlier, TEXT("Does not make any sens to have ConsecutiveOverbudgetGPUFramCount >= 2 if this frame is detected as being an outlier."));
 				}
 			}
 			else
@@ -325,12 +312,6 @@ void FDynamicResolutionHeuristicProxy::RefreshCurentFrameResolutionFraction_Rend
 				ConsecutiveOverbudgetGPUFramCount = 0;
 			}
 
-			// Ignore if an outlier.
-			if (bIsGPUOutlier)
-			{
-				continue;
-			}
-
 			float SuggestedResolutionFraction = 1.0f;
 
 			// If have reliable GPU times, or guess there is not GPU bubbles -> estimate the suggested resolution fraction that could have been used.
@@ -346,8 +327,8 @@ void FDynamicResolutionHeuristicProxy::RefreshCurentFrameResolutionFraction_Rend
 				// At resolution change that happen every N frames, amortized over time and scaled down as the
 				// standard variation of the GPU timing over non resolution changing frames increases.
 				SuggestedResolutionFraction = (
-					FMath::Sqrt(TargetedGPUBusyTimeMs / FrameEntry.TotalFrameGPUBusyTimeMs)
-					* FrameEntry.ResolutionFractions[GDynamicPrimaryResolutionFraction]);
+					BudgetSettings.EstimateResolutionFactor(TotalFrameGPUBusyTimeMs)
+					* FrameEntry.ResolutionFractions[Budget]);
 			}
 
 			NewFrameResolutionFraction += SuggestedResolutionFraction * Weight;
@@ -360,65 +341,75 @@ void FDynamicResolutionHeuristicProxy::RefreshCurentFrameResolutionFraction_Rend
 		NewFrameResolutionFraction /= TotalWeight;
 
 		// If immediate previous frames where over budget, react immediately.
-		bGPUOverbugetPanic = FrameCount > 0 && ConsecutiveOverbudgetGPUFramCount == FrameCount;
+		bool bGPUOverbugetPanic = FrameCount > 0 && ConsecutiveOverbudgetGPUFramCount == FrameCount;
 
 		// If over budget, reset history size to 0 so that this frame really behave as a first frame after an history reset.
 		if (bGPUOverbugetPanic)
 		{
-			HistorySize = 0;
+			BudgetHistorySizes[Budget] = 0;
+
+			if (Budget == GDynamicPrimaryResolutionFraction)
+			{
+				HistorySize = 0;
+				bGlobalGPUOverbugetPanic = true;
+			}
 		}
 		// If not immediately over budget, refine the new resolution fraction.
 		else
 		{
 			// If scaling the resolution, look if this is above a threshold compared to current res.
-			if (CVarChangeThreshold.GetValueOnRenderThread())
+			if (!BudgetSettings.DoesResolutionChangeEnough(CurrentFrameResolutionFractions[Budget], NewFrameResolutionFraction, bCanChangeResolution))
 			{
-				float IncreaseResolutionFractionThreshold = CurrentFrameResolutionFractions[GDynamicPrimaryResolutionFraction] * (1.0f + CVarChangeThreshold.GetValueOnRenderThread() / 100.0f);
-
-				float DecreaseResolutionFractionThreshold = CurrentFrameResolutionFractions[GDynamicPrimaryResolutionFraction] * (1.0f - CVarChangeThreshold.GetValueOnRenderThread() / 100.0f);
-
-				if (NewFrameResolutionFraction > CurrentFrameResolutionFractions[GDynamicPrimaryResolutionFraction] && NewFrameResolutionFraction < IncreaseResolutionFractionThreshold)
-				{
-					NewFrameResolutionFraction = CurrentFrameResolutionFractions[GDynamicPrimaryResolutionFraction];
-				}
+				NewFrameResolutionFraction = CurrentFrameResolutionFractions[Budget];
 			}
 
 			// If scaling the resolution up, amortize to avoid oscillations.
-			if (NewFrameResolutionFraction > CurrentFrameResolutionFractions[GDynamicPrimaryResolutionFraction])
+			if (NewFrameResolutionFraction > CurrentFrameResolutionFractions[Budget])
 			{
 				NewFrameResolutionFraction = FMath::Lerp(
-					CurrentFrameResolutionFractions[GDynamicPrimaryResolutionFraction],
+					CurrentFrameResolutionFractions[Budget],
 					NewFrameResolutionFraction,
-					CVarIncreaseAmortizationFactor.GetValueOnRenderThread());
+					BudgetSettings.IncreaseAmortizationFactor);
 			}
 		}
 
 		// Clamp resolution fraction.
 		NewFrameResolutionFraction = FMath::Clamp(
 			NewFrameResolutionFraction,
-			CVarDynamicResMinSP.GetValueOnRenderThread() / 100.0f,
-			GDynamicPrimaryResolutionFraction.GetSettings().MaxResolutionFraction);
+			BudgetSettings.MinResolutionFraction,
+			BudgetSettings.MaxResolutionFraction);
+
+		NewFrameResolutionFractions[Budget] = NewFrameResolutionFraction;
 	}
 
 	// Update the current frame's resolution fraction.
 	{
-		// CVarChangeThreshold avoids very small changes.
-		bool bWouldBeWorthChangingRes = CurrentFrameResolutionFractions[GDynamicPrimaryResolutionFraction] != NewFrameResolutionFraction;
+		bool bWouldBeWorthChangingRes = false;
+		for (TLinkedList<DynamicRenderScaling::FBudget*>::TIterator BudgetIt(DynamicRenderScaling::FBudget::GetGlobalList()); BudgetIt; BudgetIt.Next())
+		{
+			const DynamicRenderScaling::FBudget& Budget = **BudgetIt;
 
+			// CVarChangeThreshold avoids very small changes.
+			bWouldBeWorthChangingRes = bWouldBeWorthChangingRes || CurrentFrameResolutionFractions[Budget] != NewFrameResolutionFractions[Budget];
+		}
+		
 		// We do not change resolution too often to avoid interferences with temporal sub pixel in TAA upsample.
-		if ((bWouldBeWorthChangingRes &&
-			NumberOfFramesSinceScreenPercentageChange >= CVarFrameChangePeriod.GetValueOnRenderThread()) ||
-			bGPUOverbugetPanic)
+		if ((bWouldBeWorthChangingRes && bCanChangeResolution) || bGlobalGPUOverbugetPanic)
 		{
 			NumberOfFramesSinceScreenPercentageChange = 0;
-			CurrentFrameResolutionFractions[GDynamicPrimaryResolutionFraction] = NewFrameResolutionFraction;
+			CurrentFrameResolutionFractions = NewFrameResolutionFractions;
 		}
 		else
 		{
 			NumberOfFramesSinceScreenPercentageChange++;
 		}
 	}
+	
+	return RefreshCurentFrameResolutionFractionUpperBound_RenderThread();
+}
 
+void FDynamicResolutionHeuristicProxy::RefreshCurentFrameResolutionFractionUpperBound_RenderThread()
+{
 	// Compute max resolution for each budget by quantizing the new resolution fraction (falls back to the MaxResolution setting if BudgetSetting.UpperBoundQuantization==0)
 	DynamicRenderScaling::TMap<float> NewMaxResolutionFractions;
 	for (TLinkedList<DynamicRenderScaling::FBudget*>::TIterator BudgetIt(DynamicRenderScaling::FBudget::GetGlobalList()); BudgetIt; BudgetIt.Next())
diff --git a/Engine/Source/Runtime/Engine/Public/DynamicResolutionProxy.h b/Engine/Source/Runtime/Engine/Public/DynamicResolutionProxy.h
--- a/Engine/Source/Runtime/Engine/Public/DynamicResolutionProxy.h
+++ b/Engine/Source/Runtime/Engine/Public/DynamicResolutionProxy.h
@@ -126,6 +126,8 @@ private:
 
 	DynamicRenderScaling::TMap<float> QueryCurentFrameResolutionFractions_Internal() const;
 
+	void RefreshCurentFrameResolutionFractionUpperBound_RenderThread();
+
 	void ResetInternal();
 
 	void ResizeHistoryIfNeeded();
diff --git a/Engine/Source/Runtime/RenderCore/Private/DynamicRenderScaling.cpp b/Engine/Source/Runtime/RenderCore/Private/DynamicRenderScaling.cpp
--- a/Engine/Source/Runtime/RenderCore/Private/DynamicRenderScaling.cpp
+++ b/Engine/Source/Runtime/RenderCore/Private/DynamicRenderScaling.cpp
@@ -88,6 +88,11 @@ float FHeuristicSettings::CorrectNewResolutionFraction(float CurrentResolutionFr
 
 bool FHeuristicSettings::DoesResolutionChangeEnough(float CurrentResolutionFraction, float NewResolutionFraction, bool bCanChangeResolution) const
 {
+	if (ChangeThreshold == 0.0f)
+	{
+		return true;
+	}
+
 	float IncreaseResolutionFractionThreshold = CurrentResolutionFraction * (1.0f + ChangeThreshold);
 	float DecreaseResolutionFractionThreshold = CurrentResolutionFraction * (1.0f - ChangeThreshold);
 	return bCanChangeResolution && (
