#UE Fix crash when passing in an empty list of points to a spline component

diff --git a/Engine/Source/Runtime/Engine/Private/Components/SplineComponent.cpp b/Engine/Source/Runtime/Engine/Private/Components/SplineComponent.cpp
--- a/Engine/Source/Runtime/Engine/Private/Components/SplineComponent.cpp
+++ b/Engine/Source/Runtime/Engine/Private/Components/SplineComponent.cpp
@@ -193,7 +193,7 @@ void FSplineCurves::UpdateSpline(bool bClosedLoop, bool bStationaryEndpoints, in
 	Scale.AutoSetTangents(0.0f, bStationaryEndpoints);
 
 	// Now initialize the spline reparam table
-	const int32 NumSegments = bClosedLoop ? NumPoints : NumPoints - 1;
+	const int32 NumSegments = bClosedLoop ? NumPoints : FMath::Max(0, NumPoints - 1);
 
 	// Start by clearing it
 	ReparamTable.Points.Reset(NumSegments * ReparamStepsPerSegment + 1);
