Removing check from FlushActorDormancy (condition can occur in PIE + replay)

diff --git a/Engine/Source/Runtime/Engine/Private/NetDriver.cpp b/Engine/Source/Runtime/Engine/Private/NetDriver.cpp
--- a/Engine/Source/Runtime/Engine/Private/NetDriver.cpp
+++ b/Engine/Source/Runtime/Engine/Private/NetDriver.cpp
@@ -3511,7 +3511,6 @@ void UNetDriver::FlushActorDormancy(AActor* Actor, bool bWasDormInitial)
 		return;
 
 	check(Actor);
-	check(ServerConnection == NULL);
 
 	if (ReplicationDriver)
 	{
