Fix for changing filter on a texture doesn't update materials in the editor

diff --git a/Engine/Source/Runtime/Engine/Private/Texture.cpp b/Engine/Source/Runtime/Engine/Private/Texture.cpp
--- a/Engine/Source/Runtime/Engine/Private/Texture.cpp
+++ b/Engine/Source/Runtime/Engine/Private/Texture.cpp
@@ -564,6 +564,7 @@ void UTexture::PostEditChangeProperty(FPropertyChangedEvent& PropertyChangedEven
 		static const FName DeferCompressionName = GET_MEMBER_NAME_CHECKED(UTexture, DeferCompression);
 		static const FName SrgbName = GET_MEMBER_NAME_CHECKED(UTexture, SRGB);
 		static const FName VirtualTextureStreamingName = GET_MEMBER_NAME_CHECKED(UTexture, VirtualTextureStreaming);
+		static const FName FilterName = GET_MEMBER_NAME_CHECKED(UTexture, Filter);
 #if WITH_EDITORONLY_DATA
 		static const FName SourceColorSpaceName = GET_MEMBER_NAME_CHECKED(FTextureSourceColorSettings, ColorSpace);
 		static const FName CompressionQualityName = GET_MEMBER_NAME_CHECKED(UTexture, CompressionQuality);
@@ -573,6 +574,7 @@ void UTexture::PostEditChangeProperty(FPropertyChangedEvent& PropertyChangedEven
 		const FName PropertyName = PropertyThatChanged->GetFName();
 
 		if ((PropertyName == CompressionSettingsName) ||
+			(PropertyName == FilterName) ||
 			(PropertyName == LODGroupName) ||
 			(PropertyName == SrgbName))
 		{
