Refactor common shadermap DDC key logic into a shared function

There were multiple copies of very similar logic that contained subtle inconsistencies. There is an upcoming need to add some extra information to the DDC key, so it makes sense to centralize the logic first.

Keep the publically exposed function SerializeUniformBufferInfo  but mark is as deprecated to signal to any users of the old API that there is a better mechanism now that can save them some boilerplate.

diff --git a/Engine/Plugins/Compositing/OpenColorIO/Source/OpenColorIO/Private/OpenColorIOShaderMap.cpp b/Engine/Plugins/Compositing/OpenColorIO/Source/OpenColorIO/Private/OpenColorIOShaderMap.cpp
--- a/Engine/Plugins/Compositing/OpenColorIO/Source/OpenColorIO/Private/OpenColorIOShaderMap.cpp
+++ b/Engine/Plugins/Compositing/OpenColorIO/Source/OpenColorIO/Private/OpenColorIOShaderMap.cpp
@@ -134,51 +134,19 @@ bool FOpenColorIOShaderMapId::operator==(const FOpenColorIOShaderMapId& InRefere
 void FOpenColorIOShaderMapId::AppendKeyString(FString& OutKeyString) const
 {
 #if WITH_EDITOR
-	OutKeyString += ShaderCodeAndConfigHash;
-	OutKeyString += TEXT("_");
-
-	FString FeatureLevelString;
-	GetFeatureLevelName(FeatureLevel, FeatureLevelString);
+	OutKeyString.Append(ShaderCodeAndConfigHash);
+	OutKeyString.AppendChar('_');
 
 	{
-		const FSHAHash LayoutHash = Freeze::HashLayout(StaticGetTypeLayoutDesc<FOpenColorIOShaderMapContent>(), LayoutParams);
-		OutKeyString += TEXT("_");
-		OutKeyString += LayoutHash.ToString();
-		OutKeyString += TEXT("_");
+		const FSHAHash LayoutHash = GetShaderTypeLayoutHash(StaticGetTypeLayoutDesc<FOpenColorIOShaderMapContent>(), LayoutParams);
+		OutKeyString.AppendChar('_');
+		LayoutHash.AppendString(OutKeyString);
+		OutKeyString.AppendChar('_');
 	}
 
-	TSortedMap<const TCHAR*, FCachedUniformBufferDeclaration, FDefaultAllocator, FUniformBufferNameSortOrder> ReferencedUniformBuffers;
-
 	// Add the inputs for any shaders that are stored inline in the shader map
-	for (const FShaderTypeDependency& ShaderTypeDependency : ShaderTypeDependencies)
-	{
-		const FShaderType* ShaderType = FindShaderTypeByName(ShaderTypeDependency.ShaderTypeName);
-		OutKeyString += TEXT("_");
-		OutKeyString += ShaderType->GetName();
-		OutKeyString += ShaderTypeDependency.SourceHash.ToString();
-		
-		const FSHAHash LayoutHash = Freeze::HashLayout(ShaderType->GetLayout(), LayoutParams);
-		OutKeyString += LayoutHash.ToString();
-
-		const TMap<const TCHAR*, FCachedUniformBufferDeclaration>& ReferencedUniformBufferStructsCache = ShaderType->GetReferencedUniformBufferStructsCache();
-
-		for (TMap<const TCHAR*, FCachedUniformBufferDeclaration>::TConstIterator It(ReferencedUniformBufferStructsCache); It; ++It)
-		{
-			ReferencedUniformBuffers.Add(It.Key(), It.Value());
-		}
-	}
+	AppendKeyStringShaderDependencies(MakeArrayView(ShaderTypeDependencies), LayoutParams, OutKeyString);
 
-	{
-		TArray<uint8> TempData;
-		FSerializationHistory SerializationHistory;
-		FMemoryWriter Ar(TempData, true);
-		FShaderSaveArchive SaveArchive(Ar, SerializationHistory);
-
-		// Save uniform buffer member info so we can detect when layout has changed
-		SerializeUniformBufferInfo(SaveArchive, ReferencedUniformBuffers);
-
-		SerializationHistory.AppendKeyString(OutKeyString);
-	}
 #endif //WITH_EDITOR
 }
 
diff --git a/Engine/Plugins/FX/Niagara/Source/NiagaraShader/Private/NiagaraShader.cpp b/Engine/Plugins/FX/Niagara/Source/NiagaraShader/Private/NiagaraShader.cpp
--- a/Engine/Plugins/FX/Niagara/Source/NiagaraShader/Private/NiagaraShader.cpp
+++ b/Engine/Plugins/FX/Niagara/Source/NiagaraShader/Private/NiagaraShader.cpp
@@ -250,29 +250,30 @@ bool FNiagaraShaderMapId::operator==(const FNiagaraShaderMapId& ReferenceSet) co
 
 void FNiagaraShaderMapId::AppendKeyString(FString& KeyString) const
 {
-	KeyString += BaseCompileHash.ToString();
-	KeyString += TEXT("_");
-
-	FString FeatureLevelString;
-	GetFeatureLevelName(FeatureLevel, FeatureLevelString);
+	BaseCompileHash.AppendString(KeyString);
+	KeyString.AppendChar('_');
 
 	{
-		const FSHAHash LayoutHash = Freeze::HashLayout(StaticGetTypeLayoutDesc<FNiagaraShaderMapContent>(), LayoutParams);
-		KeyString += TEXT("_");
-		KeyString += LayoutHash.ToString();
-		KeyString += TEXT("_");
+		const FSHAHash LayoutHash = GetShaderTypeLayoutHash(StaticGetTypeLayoutDesc<FNiagaraShaderMapContent>(), LayoutParams);
+		KeyString.AppendChar('_');
+		LayoutHash.AppendString(KeyString);
+		KeyString.AppendChar('_');
 	}
 
 	{
-		const FSHAHash LayoutHash = Freeze::HashLayout(StaticGetTypeLayoutDesc<FNiagaraShader>(), LayoutParams);
-		KeyString += TEXT("_");
-		KeyString += LayoutHash.ToString();
-		KeyString += TEXT("_");
+		const FSHAHash LayoutHash = GetShaderTypeLayoutHash(StaticGetTypeLayoutDesc<FNiagaraShader>(), LayoutParams);
+		KeyString.AppendChar('_');
+		LayoutHash.AppendString(KeyString);
+		KeyString.AppendChar('_');
 	}
 
-	KeyString += FeatureLevelString + TEXT("_");
-	KeyString += CompilerVersionID.ToString();
-	KeyString += TEXT("_");
+	FName FeatureLevelName;
+	GetFeatureLevelName(FeatureLevel, FeatureLevelName);
+
+	FeatureLevelName.AppendString(KeyString);
+	KeyString.AppendChar('_');
+	CompilerVersionID.AppendString(KeyString);
+	KeyString.AppendChar('_');
 
 	if (bUsesRapidIterationParams)
 	{
@@ -284,67 +285,31 @@ void FNiagaraShaderMapId::AppendKeyString(FString& KeyString) const
 	}
 
 	// Add base parameters structure
-	KeyString += FString::Printf(TEXT("%08x"), TShaderParameterStructTypeInfo<FNiagaraShader::FParameters>::GetStructMetadata()->GetLayoutHash());
+	KeyString.Appendf(TEXT("%08x"), TShaderParameterStructTypeInfo<FNiagaraShader::FParameters>::GetStructMetadata()->GetLayoutHash());
 
 	// Add additional defines
-	for (int32 DefinesIndex = 0; DefinesIndex < AdditionalDefines.Num(); DefinesIndex++)
+	for (const FMemoryImageString& Define : AdditionalDefines)
 	{
-		KeyString += AdditionalDefines[DefinesIndex];
-		if (DefinesIndex < AdditionalDefines.Num() - 1)
-		{
-			KeyString += TEXT("_");
-		}
+		KeyString.AppendChar('_');
+		KeyString.Append(Define);
 	}
 
 	// Add additional variables
-	for (int32 VariablesIndex = 0; VariablesIndex < AdditionalVariables.Num(); VariablesIndex++)
+	for (const FMemoryImageString& Var : AdditionalVariables)
 	{
-		KeyString += AdditionalVariables[VariablesIndex];
-		if (VariablesIndex < AdditionalVariables.Num() - 1)
-		{
-			KeyString += TEXT("_");
-		}
+		KeyString.AppendChar('_');
+		KeyString.Append(Var);
 	}
 
 	// Add any referenced top level compile hashes to the key so that we will recompile when they are changed
-	for (int32 HashIndex = 0; HashIndex < ReferencedCompileHashes.Num(); HashIndex++)
+	for (const FSHAHash& Hash : ReferencedCompileHashes)
 	{
-		KeyString += ReferencedCompileHashes[HashIndex].ToString();
-		if (HashIndex < ReferencedCompileHashes.Num() - 1)
-		{
-			KeyString += TEXT("_");
-		}
+		KeyString.AppendChar('_');
+		Hash.AppendString(KeyString);
 	}
 
-	TSortedMap<const TCHAR*, FCachedUniformBufferDeclaration, FDefaultAllocator, FUniformBufferNameSortOrder> ReferencedUniformBuffers;
-	for (const FShaderTypeDependency& ShaderTypeDependency : ShaderTypeDependencies)
-	{
-		const FShaderType* ShaderType = FindShaderTypeByName(ShaderTypeDependency.ShaderTypeName);
-		KeyString += TEXT("_");
-		KeyString += ShaderType->GetName();
-		KeyString += ShaderTypeDependency.SourceHash.ToString();
-
-		const FSHAHash LayoutHash = Freeze::HashLayout(ShaderType->GetLayout(), LayoutParams);
-		KeyString += LayoutHash.ToString();
-
-		const TMap<const TCHAR*, FCachedUniformBufferDeclaration>& ReferencedUniformBufferStructsCache = ShaderType->GetReferencedUniformBufferStructsCache();
-		for (TMap<const TCHAR*, FCachedUniformBufferDeclaration>::TConstIterator It(ReferencedUniformBufferStructsCache); It; ++It)
-		{
-			ReferencedUniformBuffers.Add(It.Key(), It.Value());
-		}
-	}
-
-	{
-		TArray<uint8> TempData;
-		FSerializationHistory SerializationHistory;
-		FMemoryWriter Ar(TempData, true);
-		FShaderSaveArchive SaveArchive(Ar, SerializationHistory);
-
-		// Save uniform buffer member info so we can detect when layout has changed
-		SerializeUniformBufferInfo(SaveArchive, ReferencedUniformBuffers);
-
-		SerializationHistory.AppendKeyString(KeyString);
-	}
+	// Add the inputs for any shaders that are stored inline in the shader map
+	AppendKeyStringShaderDependencies(MakeArrayView(ShaderTypeDependencies), LayoutParams, KeyString);
 }
 
 
diff --git a/Engine/Plugins/Runtime/ComputeFramework/Source/ComputeFramework/Private/ComputeKernelShaderMap.cpp b/Engine/Plugins/Runtime/ComputeFramework/Source/ComputeFramework/Private/ComputeKernelShaderMap.cpp
--- a/Engine/Plugins/Runtime/ComputeFramework/Source/ComputeFramework/Private/ComputeKernelShaderMap.cpp
+++ b/Engine/Plugins/Runtime/ComputeFramework/Source/ComputeFramework/Private/ComputeKernelShaderMap.cpp
@@ -108,51 +108,18 @@ bool FComputeKernelShaderMapId::operator==(const FComputeKernelShaderMapId& InRe
 void FComputeKernelShaderMapId::AppendKeyString(FString& OutKeyString) const
 {
 #if WITH_EDITOR
-	OutKeyString += FString::Printf(TEXT("%llX"), ShaderCodeHash);
-	OutKeyString += TEXT("_");
-
-	FString FeatureLevelString;
-	GetFeatureLevelName(FeatureLevel, FeatureLevelString);
+	OutKeyString.Appendf(TEXT("%llX"), ShaderCodeHash);
+	OutKeyString.AppendChar('_');
 
 	{
-		const FSHAHash LayoutHash = Freeze::HashLayout(StaticGetTypeLayoutDesc<FComputeKernelShaderMapContent>(), LayoutParams);
-		OutKeyString += TEXT("_");
-		OutKeyString += LayoutHash.ToString();
-		OutKeyString += TEXT("_");
+		const FSHAHash LayoutHash = GetShaderTypeLayoutHash(StaticGetTypeLayoutDesc<FComputeKernelShaderMapContent>(), LayoutParams);
+		OutKeyString.AppendChar('_');
+		LayoutHash.AppendString(OutKeyString);
+		OutKeyString.AppendChar('_');
 	}
 
-	TSortedMap<const TCHAR*, FCachedUniformBufferDeclaration, FDefaultAllocator, FUniformBufferNameSortOrder> ReferencedUniformBuffers;
-
 	// Add the inputs for any shaders that are stored inline in the shader map
-	for (const FShaderTypeDependency& ShaderTypeDependency : ShaderTypeDependencies)
-	{
-		const FShaderType* ShaderType = FindShaderTypeByName(ShaderTypeDependency.ShaderTypeName);
-		OutKeyString += TEXT("_");
-		OutKeyString += ShaderType->GetName();
-		OutKeyString += ShaderTypeDependency.SourceHash.ToString();
-		
-		const FSHAHash LayoutHash = Freeze::HashLayout(ShaderType->GetLayout(), LayoutParams);
-		OutKeyString += LayoutHash.ToString();
-
-		const TMap<const TCHAR*, FCachedUniformBufferDeclaration>& ReferencedUniformBufferStructsCache = ShaderType->GetReferencedUniformBufferStructsCache();
-
-		for (TMap<const TCHAR*, FCachedUniformBufferDeclaration>::TConstIterator It(ReferencedUniformBufferStructsCache); It; ++It)
-		{
-			ReferencedUniformBuffers.Add(It.Key(), It.Value());
-		}
-	}
-
-	{
-		TArray<uint8> TempData;
-		FSerializationHistory SerializationHistory;
-		FMemoryWriter Ar(TempData, true);
-		FShaderSaveArchive SaveArchive(Ar, SerializationHistory);
-
-		// Save uniform buffer member info so we can detect when layout has changed
-		SerializeUniformBufferInfo(SaveArchive, ReferencedUniformBuffers);
-
-		SerializationHistory.AppendKeyString(OutKeyString);
-	}
+	AppendKeyStringShaderDependencies(MakeArrayView(ShaderTypeDependencies), LayoutParams, OutKeyString);
 #endif //WITH_EDITOR
 }
 
diff --git a/Engine/Source/Runtime/Engine/Private/Materials/MaterialShader.cpp b/Engine/Source/Runtime/Engine/Private/Materials/MaterialShader.cpp
--- a/Engine/Source/Runtime/Engine/Private/Materials/MaterialShader.cpp
+++ b/Engine/Source/Runtime/Engine/Private/Materials/MaterialShader.cpp
@@ -1,7 +1,6 @@
 // Copyright Epic Games, Inc. All Rights Reserved.
 
 #include "MaterialShader.h"
-#include "Algo/Find.h"
 #include "Containers/StringConv.h"
 #include "DerivedDataCache.h"
 #include "DerivedDataRequestOwner.h"
@@ -997,112 +996,20 @@ void FMaterialShaderMapId::UpdateFromParameterSet(const FStaticParameterSet& Sta
 
 #if WITH_EDITOR
 
-static TSortedMap<const TCHAR*, FCachedUniformBufferDeclaration, FDefaultAllocator, FUniformBufferNameSortOrder> MergeBufferMaps(TArrayView<const TMap<const TCHAR*,FCachedUniformBufferDeclaration>*> Maps)
-{
-	int32 Num = 0;
-	for (const TMap<const TCHAR*, FCachedUniformBufferDeclaration>* Map : Maps)
-	{
-		Num += Map->Num();
-	}
-
-	// Deduplicate identical pointers (the common case)
-	using PairType = TPair<const TCHAR*, FCachedUniformBufferDeclaration>;
-	TMap<const void*, FCachedUniformBufferDeclaration> UniquePointerPairs;
-	UniquePointerPairs.Reserve(Num);
-	for (const TMap<const TCHAR*, FCachedUniformBufferDeclaration>* Map : Maps)
-	{
-		for (const PairType& Pair : *Map)
-		{
-			UniquePointerPairs.Add(Pair);
-		}
-	}
-
-	// Presort pairs so that they can be efficiently inserted into the TSortedMap
-	TArray<PairType> PresortedPairs;
-	PresortedPairs.Reserve(UniquePointerPairs.Num());
-	for (TPair<const void*, FCachedUniformBufferDeclaration>& UniquePointerPair : UniquePointerPairs)
-	{
-		PresortedPairs.Add(PairType(static_cast<const TCHAR*>(UniquePointerPair.Key), MoveTemp(UniquePointerPair.Value)));
-	}
-
-	Algo::SortBy(PresortedPairs, &PairType::Key, FUniformBufferNameSortOrder());
-
-	TSortedMap<const TCHAR*, FCachedUniformBufferDeclaration, FDefaultAllocator, FUniformBufferNameSortOrder> Out;
-	Out.Reserve(PresortedPairs.Num());
-	for (PairType& PresortedPair : PresortedPairs)
-	{
-		Out.Add(PresortedPair.Key, MoveTemp(PresortedPair.Value));
-	}
-
-	return Out;
-}
-
-
 uint32 GetTypeHash(FPlatformTypeLayoutParameters Params) { return HashCombine(Params.Flags, Params.MaxFieldAlignment); }
 
-class FFrozenMaterialLayoutHashCache
-{
-public:
-	FSHAHash Get(const FTypeLayoutDesc& TypeDesc, FPlatformTypeLayoutParameters LayoutParams)
-	{
-		{
-			FReadScopeLock ReadScope(Lock);
-
-			if (const FPlatformCache* Platform = Algo::FindBy(Platforms, LayoutParams, &FPlatformCache::Parameters))
-			{
-				if (const FSHAHash* Hash = Platform->Cache.Find(&TypeDesc))
-				{
-					return *Hash;
-				}
-			}
-		}
-
-		FSHAHash Hash = Freeze::HashLayout(TypeDesc, LayoutParams);
-
-		{
-			FWriteScopeLock WriteScope(Lock);
-
-			FPlatformCache* Platform = Algo::FindBy(Platforms, LayoutParams, &FPlatformCache::Parameters);
-			if (!Platform)
-			{
-				Platform = &Platforms.AddDefaulted_GetRef();
-				Platform->Parameters = LayoutParams;
-			}
-
-
-			Platform->Cache.FindOrAdd(&TypeDesc, Hash);
-		}
-
-		return Hash;
-	}
-
-private:
-	struct FPlatformCache
-	{
-		FPlatformTypeLayoutParameters Parameters;
-		TMap<const FTypeLayoutDesc*, FSHAHash> Cache;
-	};
-
-	FRWLock Lock;
-	TArray<FPlatformCache, TInlineAllocator<8>> Platforms;
-};
-
-static FFrozenMaterialLayoutHashCache GFrozenMaterialLayoutHashes;
-
 void FMaterialShaderMapId::AppendKeyString(FString& KeyString) const
 {
 	check(IsContentValid());
 	BaseMaterialId.AppendString(KeyString);
 	KeyString.AppendChar('_');
 
-	FString QualityLevelName;
-	GetMaterialQualityLevelName(QualityLevel, QualityLevelName);
-	KeyString += QualityLevelName;
+	GetMaterialQualityLevelFName(QualityLevel).AppendString(KeyString);
 	KeyString.AppendChar('_');
 
-	FString FeatureLevelString;
-	GetFeatureLevelName(FeatureLevel, FeatureLevelString);
-	KeyString += FeatureLevelString;
+	FName FeatureLevelName;
+	GetFeatureLevelName(FeatureLevel, FeatureLevelName);
+	FeatureLevelName.AppendString(KeyString);
 	KeyString.AppendChar('_');
 
 	LayoutParams.AppendKeyString(KeyString);
@@ -1141,7 +1048,7 @@ void FMaterialShaderMapId::AppendKeyString(FString& KeyString) const
 	}
 
 	{
-		const FSHAHash LayoutHash = GFrozenMaterialLayoutHashes.Get(StaticGetTypeLayoutDesc<FMaterialShaderMapContent>(), LayoutParams);
+		const FSHAHash LayoutHash = GetShaderTypeLayoutHash(StaticGetTypeLayoutDesc<FMaterialShaderMapContent>(), LayoutParams);
 		KeyString.AppendChar('_');
 		LayoutHash.AppendString(KeyString);
 		KeyString.AppendChar('_');
@@ -1154,92 +1061,13 @@ void FMaterialShaderMapId::AppendKeyString(FString& KeyString) const
 		ReferencedParameterCollections[CollectionIndex].AppendString(KeyString);
 	}
 
-	TArray<const TMap<const TCHAR*,FCachedUniformBufferDeclaration>*> ReferencedUniformBufferMaps;
-
-	// Add the inputs for any shaders that are stored inline in the shader map
-	for (int32 ShaderIndex = 0; ShaderIndex < ShaderTypeDependencies.Num(); ShaderIndex++)
-	{
-		const FShaderTypeDependency& ShaderTypeDependency = ShaderTypeDependencies[ShaderIndex];
-		const FShaderType* ShaderType = FindShaderTypeByName(ShaderTypeDependency.ShaderTypeName);
-#if WITH_EDITORONLY_DATA
-		checkf(ShaderType != nullptr, TEXT("Failed to find FShaderType for dependency %s (total in the NameToTypeMap: %d)"), ShaderTypeDependency.ShaderTypeName.GetDebugString().String.Get(), FShaderType::GetNameToTypeMap().Num());
-#endif
-
-		KeyString.AppendChar('_');
-		KeyString += ShaderType->GetName();
-		ShaderTypeDependency.SourceHash.AppendString(KeyString);
-
-		if (const FShaderParametersMetadata* ParameterStructMetadata = ShaderType->GetRootParametersMetadata())
-		{
-			KeyString.Appendf(TEXT("%08x"), ParameterStructMetadata->GetLayoutHash());
-		}
-
-		const FSHAHash LayoutHash = GFrozenMaterialLayoutHashes.Get(ShaderType->GetLayout(), LayoutParams);
-		LayoutHash.AppendString(KeyString);
-
-		ReferencedUniformBufferMaps.Add(&ShaderType->GetReferencedUniformBufferStructsCache());
-	}
-
-	// Add the inputs for any shader pipelines that are stored inline in the shader map
-	for (int32 TypeIndex = 0; TypeIndex < ShaderPipelineTypeDependencies.Num(); TypeIndex++)
-	{
-		const FShaderPipelineTypeDependency& Dependency = ShaderPipelineTypeDependencies[TypeIndex];
-		const FShaderPipelineType* ShaderPipelineType = FShaderPipelineType::GetShaderPipelineTypeByName(Dependency.ShaderPipelineTypeName);
-#if WITH_EDITORONLY_DATA
-		checkf(ShaderPipelineType != nullptr, TEXT("Failed to find FShaderPipelineType for dependency %s (total in the NameToTypeMap: %d)"), Dependency.ShaderPipelineTypeName.GetDebugString().String.Get(), FShaderType::GetNameToTypeMap().Num());
-#endif
-
-		KeyString.AppendChar('_');
-		KeyString += ShaderPipelineType->GetName();
-		Dependency.StagesSourceHash.AppendString(KeyString);
-
-		for (const FShaderType* ShaderType : ShaderPipelineType->GetStages())
-		{
-			if (const FShaderParametersMetadata* ParameterStructMetadata = ShaderType->GetRootParametersMetadata())
-			{
-				KeyString.Appendf(TEXT("%08x"), ParameterStructMetadata->GetLayoutHash());
-			}
-
-			ReferencedUniformBufferMaps.Add(&ShaderType->GetReferencedUniformBufferStructsCache());
-		}
-	}
-
 	// Add the inputs for any shaders that are stored inline in the shader map
-	for (int32 VFIndex = 0; VFIndex < VertexFactoryTypeDependencies.Num(); VFIndex++)
-	{
-		KeyString.AppendChar('_');
-
-		const FVertexFactoryTypeDependency& VFDependency = VertexFactoryTypeDependencies[VFIndex];
-		const FVertexFactoryType* VertexFactoryType = FVertexFactoryType::GetVFByName(VFDependency.VertexFactoryTypeName);
-
-		KeyString += VertexFactoryType->GetName();
-		VFDependency.VFSourceHash.AppendString(KeyString);
-
-		for (int32 Frequency = 0; Frequency < SF_NumFrequencies; Frequency++)
-		{
-			const FTypeLayoutDesc* ParameterLayout = VertexFactoryType->GetShaderParameterLayout((EShaderFrequency)Frequency);
-			if (ParameterLayout)
-			{
-				const FSHAHash LayoutHash = GFrozenMaterialLayoutHashes.Get(*ParameterLayout, LayoutParams);
-				LayoutHash.AppendString(KeyString);
-			}
-		}
-
-		ReferencedUniformBufferMaps.Add(&VertexFactoryType->GetReferencedUniformBufferStructsCache());
-	}
-
-	{
-		TArray<uint8> TempData;
-		FSerializationHistory SerializationHistory;
-		FMemoryWriter Ar(TempData, true);
-		FShaderSaveArchive SaveArchive(Ar, SerializationHistory);
-		TSortedMap<const TCHAR*, FCachedUniformBufferDeclaration, FDefaultAllocator, FUniformBufferNameSortOrder> ReferencedUniformBuffers = MergeBufferMaps(ReferencedUniformBufferMaps);
-
-		// Save uniform buffer member info so we can detect when layout has changed
-		SerializeUniformBufferInfo(SaveArchive, ReferencedUniformBuffers);
-
-		SerializationHistory.AppendKeyString(KeyString);
-	}
+	AppendKeyStringShaderDependencies(
+		MakeArrayView(ShaderTypeDependencies),
+		MakeArrayView(ShaderPipelineTypeDependencies),
+		MakeArrayView(VertexFactoryTypeDependencies),
+		LayoutParams, 
+		KeyString);
 
 	BytesToHex(&TextureReferencesHash.Hash[0], sizeof(TextureReferencesHash.Hash), KeyString);
 
diff --git a/Engine/Source/Runtime/RenderCore/Private/GlobalShader.cpp b/Engine/Source/Runtime/RenderCore/Private/GlobalShader.cpp
--- a/Engine/Source/Runtime/RenderCore/Private/GlobalShader.cpp
+++ b/Engine/Source/Runtime/RenderCore/Private/GlobalShader.cpp
@@ -462,86 +462,22 @@ void FGlobalShaderMapId::AppendKeyString(FString& KeyString, const TArray<FShade
 	LayoutParams.AppendKeyString(KeyString);
 
 	{
-		const FSHAHash LayoutHash = Freeze::HashLayout(StaticGetTypeLayoutDesc<FGlobalShaderMapContent>(), LayoutParams);
-		KeyString += TEXT("_");
-		KeyString += LayoutHash.ToString();
-		KeyString += TEXT("_");
+		const FSHAHash LayoutHash = GetShaderTypeLayoutHash(StaticGetTypeLayoutDesc<FGlobalShaderMapContent>(), LayoutParams);
+		KeyString.AppendChar('_');
+		LayoutHash.AppendString(KeyString);
+		KeyString.AppendChar('_');
 	}
 
-	TSortedMap<const TCHAR*, FCachedUniformBufferDeclaration, FDefaultAllocator, FUniformBufferNameSortOrder> ReferencedUniformBuffers;
+	AppendKeyStringShaderDependencies(Dependencies, ShaderPipelineTypeDependencies, TConstArrayView<FVertexFactoryTypeDependency>(), LayoutParams, KeyString);
 
-	for (int32 ShaderIndex = 0; ShaderIndex < Dependencies.Num(); ShaderIndex++)
+	// Extra logic, just for Global shaders
+	for (const FShaderTypeDependency& ShaderTypeDependency : Dependencies)
 	{
-		const FShaderTypeDependency& ShaderTypeDependency = Dependencies[ShaderIndex];
 		const FShaderType* ShaderType = FindShaderTypeByName(ShaderTypeDependency.ShaderTypeName);
-
-		KeyString += TEXT("_");
-		KeyString += ShaderType->GetName();
-		KeyString += FString::Printf(TEXT("%i"), ShaderTypeDependency.PermutationId);
-
-		// Add the type's source hash so that we can invalidate cached shaders when .usf changes are made
-		KeyString += ShaderTypeDependency.SourceHash.ToString();
-
 		// Add the config define hash, if any defines exist in config
 		FGlobalShaderConfigDefines::AppendKeyString(KeyString, ShaderType->GetFName(), ShaderTypeDependency.PermutationId, IniPlatformName, ShaderPlatform);
-
-		if (const FShaderParametersMetadata* ParameterStructMetadata = ShaderType->GetRootParametersMetadata())
-		{
-			KeyString += FString::Printf(TEXT("%08x"), ParameterStructMetadata->GetLayoutHash());
-		}
-
-		// Add the serialization history to the key string so that we can detect changes to global shader serialization without a corresponding .usf change
-		const FSHAHash LayoutHash = Freeze::HashLayout(ShaderType->GetLayout(), LayoutParams);
-		KeyString += LayoutHash.ToString();
-
-		const TMap<const TCHAR*, FCachedUniformBufferDeclaration>& ReferencedUniformBufferStructsCache = ShaderType->GetReferencedUniformBufferStructsCache();
-
-		// Gather referenced uniform buffers
-		for (TMap<const TCHAR*, FCachedUniformBufferDeclaration>::TConstIterator It(ReferencedUniformBufferStructsCache); It; ++It)
-		{
-			ReferencedUniformBuffers.Add(It.Key(), It.Value());
-		}
 	}
 
-	for (int32 Index = 0; Index < ShaderPipelineTypeDependencies.Num(); ++Index)
-	{
-		const FShaderPipelineTypeDependency& Dependency = ShaderPipelineTypeDependencies[Index];
-		const FShaderPipelineType* ShaderPipelineType = FShaderPipelineType::GetShaderPipelineTypeByName(Dependency.ShaderPipelineTypeName);
-
-		KeyString += TEXT("_");
-		KeyString += ShaderPipelineType->GetName();
-
-		// Add the type's source hash so that we can invalidate cached shaders when .usf changes are made
-		KeyString += Dependency.StagesSourceHash.ToString();
-
-		for (const FShaderType* ShaderType : ShaderPipelineType->GetStages())
-		{
-			if (const FShaderParametersMetadata* ParameterStructMetadata = ShaderType->GetRootParametersMetadata())
-			{
-				KeyString += FString::Printf(TEXT("%08x"), ParameterStructMetadata->GetLayoutHash());
-			}
-
-			const TMap<const TCHAR*, FCachedUniformBufferDeclaration>& ReferencedUniformBufferStructsCache = ShaderType->GetReferencedUniformBufferStructsCache();
-
-			// Gather referenced uniform buffers
-			for (TMap<const TCHAR*, FCachedUniformBufferDeclaration>::TConstIterator It(ReferencedUniformBufferStructsCache); It; ++It)
-			{
-				ReferencedUniformBuffers.Add(It.Key(), It.Value());
-			}
-		}
-	}
-
-	{
-		TArray<uint8> TempData;
-		FSerializationHistory SerializationHistory;
-		FMemoryWriter Ar(TempData, true);
-		FShaderSaveArchive SaveArchive(Ar, SerializationHistory);
-
-		// Save uniform buffer member info so we can detect when layout has changed
-		SerializeUniformBufferInfo(SaveArchive, ReferencedUniformBuffers);
-
-		SerializationHistory.AppendKeyString(KeyString);
-	}
 #endif // WITH_EDITOR
 }
 
diff --git a/Engine/Source/Runtime/RenderCore/Private/ShaderCore.cpp b/Engine/Source/Runtime/RenderCore/Private/ShaderCore.cpp
--- a/Engine/Source/Runtime/RenderCore/Private/ShaderCore.cpp
+++ b/Engine/Source/Runtime/RenderCore/Private/ShaderCore.cpp
@@ -5,6 +5,7 @@
 =============================================================================*/
 
 #include "ShaderCore.h"
+#include "Algo/Find.h"
 #include "HAL/FileManager.h"
 #include "Misc/FileHelper.h"
 #include "Misc/Paths.h"
@@ -23,6 +24,7 @@
 #include "Interfaces/IShaderFormatModule.h"
 #include "RHIShaderFormatDefinitions.inl"
 #include "Compression/OodleDataCompression.h"
+#include "Serialization/MemoryWriter.h"
 #if WITH_EDITOR
 #include "Misc/CoreMisc.h"
 #include "Interfaces/ITargetPlatform.h"
@@ -1519,6 +1521,55 @@ void GenerateReferencedUniformBuffers(
 	}
 }
 
+namespace {
+// anonymous namespace
+
+class FFrozenMaterialLayoutHashCache
+{
+public:
+	FSHAHash Get(const FTypeLayoutDesc& TypeDesc, FPlatformTypeLayoutParameters LayoutParams)
+	{
+		{
+			FReadScopeLock ReadScope(Lock);
+
+			if (const FPlatformCache* Platform = Algo::FindBy(Platforms, LayoutParams, &FPlatformCache::Parameters))
+			{
+				if (const FSHAHash* Hash = Platform->Cache.Find(&TypeDesc))
+				{
+					return *Hash;
+				}
+			}
+		}
+
+		FSHAHash Hash = Freeze::HashLayout(TypeDesc, LayoutParams);
+
+		{
+			FWriteScopeLock WriteScope(Lock);
+
+			FPlatformCache* Platform = Algo::FindBy(Platforms, LayoutParams, &FPlatformCache::Parameters);
+			if (!Platform)
+			{
+				Platform = &Platforms.AddDefaulted_GetRef();
+				Platform->Parameters = LayoutParams;
+			}
+
+
+			Platform->Cache.FindOrAdd(&TypeDesc, Hash);
+		}
+
+		return Hash;
+	}
+
+private:
+	struct FPlatformCache
+	{
+		FPlatformTypeLayoutParameters Parameters;
+		TMap<const FTypeLayoutDesc*, FSHAHash> Cache;
+	};
+
+	FRWLock Lock;
+	TArray<FPlatformCache, TInlineAllocator<8>> Platforms;
+};
 
 /** Efficient lookup to find FShaderParametersMetadata members by name pointer */
 class FShaderParameterMemberLookup 
@@ -1575,8 +1626,15 @@ private:
 
 static FFShaderParameterPointerLookupCache GShaderParameterMemberLookupCache;
 
-void SerializeUniformBufferInfo(FShaderSaveArchive& Ar, const TSortedMap<const TCHAR*, FCachedUniformBufferDeclaration, FDefaultAllocator, FUniformBufferNameSortOrder>& UniformBufferEntries)
+
+// This copy is only used internally once - it could be inlined once the public API version is removed
+void SerializeUniformBufferInfo_Internal(FShaderSaveArchive& Ar, const TSortedMap<const TCHAR*, FCachedUniformBufferDeclaration, FDefaultAllocator, FUniformBufferNameSortOrder>& UniformBufferEntries)
 {
+	if (UniformBufferEntries.IsEmpty())
+	{
+		return;
+	}
+
 	TSharedPtr<const FShaderParameterMemberLookup> ShaderParameterMembers = GShaderParameterMemberLookupCache.Get();
 
 	for (const TPair<const TCHAR*, FCachedUniformBufferDeclaration>& Entry : UniformBufferEntries)
@@ -1600,6 +1658,135 @@ void SerializeUniformBufferInfo(FShaderSaveArchive& Ar, const TSortedMap<const T
 	}
 }
 
+} // anonymous namespace
+
+FSHAHash GetShaderTypeLayoutHash(const FTypeLayoutDesc& TypeDesc, FPlatformTypeLayoutParameters LayoutParameters)
+{
+	static FFrozenMaterialLayoutHashCache GFrozenMaterialLayoutHashes;
+	return GFrozenMaterialLayoutHashes.Get(TypeDesc, LayoutParameters);
+}
+
+void AppendKeyStringShaderDependencies(
+	TConstArrayView<FShaderTypeDependency> ShaderTypeDependencies,
+	FPlatformTypeLayoutParameters LayoutParams,
+	FString& OutKeyString)
+{
+	// Simplified interface if we only have ShaderTypeDependencies
+	AppendKeyStringShaderDependencies(
+		ShaderTypeDependencies,
+		TConstArrayView<FShaderPipelineTypeDependency>(),
+		TConstArrayView<FVertexFactoryTypeDependency>(),
+		LayoutParams,
+		OutKeyString);
+}
+
+void AppendKeyStringShaderDependencies(
+	TConstArrayView<FShaderTypeDependency> ShaderTypeDependencies,
+	TConstArrayView<FShaderPipelineTypeDependency> ShaderPipelineTypeDependencies,
+	TConstArrayView<FVertexFactoryTypeDependency> VertexFactoryTypeDependencies,
+	FPlatformTypeLayoutParameters LayoutParams,
+	FString& OutKeyString)
+{
+	FSortedMapUniformBufferDeclaration ReferencedUniformBuffers;
+
+	for (const FShaderTypeDependency& ShaderTypeDependency : ShaderTypeDependencies)
+	{
+		const FShaderType* ShaderType = FindShaderTypeByName(ShaderTypeDependency.ShaderTypeName);
+#if WITH_EDITORONLY_DATA
+		checkf(ShaderType != nullptr, TEXT("Failed to find FShaderType for dependency %s (total in the NameToTypeMap: %d)"), ShaderTypeDependency.ShaderTypeName.GetDebugString().String.Get(), FShaderType::GetNameToTypeMap().Num());
+#endif
+
+		OutKeyString.AppendChar('_');
+		OutKeyString.Append(ShaderType->GetName());
+		OutKeyString.AppendInt(ShaderTypeDependency.PermutationId);
+		OutKeyString.AppendChar('_');
+
+		// Add the type's source hash so that we can invalidate cached shaders when .usf changes are made
+		ShaderTypeDependency.SourceHash.AppendString(OutKeyString);
+
+		if (const FShaderParametersMetadata* ParameterStructMetadata = ShaderType->GetRootParametersMetadata())
+		{
+			OutKeyString.Appendf(TEXT("%08x"), ParameterStructMetadata->GetLayoutHash());
+		}
+
+		const FSHAHash LayoutHash = GetShaderTypeLayoutHash(ShaderType->GetLayout(), LayoutParams);
+		LayoutHash.AppendString(OutKeyString);
+
+		for (const auto& Entry : ShaderType->GetReferencedUniformBufferStructsCache())
+		{
+			ReferencedUniformBuffers.Add(Entry.Key, Entry.Value);
+		}
+	}
+
+	// Add the inputs for any shader pipelines that are stored inline in the shader map
+	for (const FShaderPipelineTypeDependency& Dependency : ShaderPipelineTypeDependencies)
+	{
+		const FShaderPipelineType* ShaderPipelineType = FShaderPipelineType::GetShaderPipelineTypeByName(Dependency.ShaderPipelineTypeName);
+#if WITH_EDITORONLY_DATA
+		checkf(ShaderPipelineType != nullptr, TEXT("Failed to find FShaderPipelineType for dependency %s (total in the NameToTypeMap: %d)"), Dependency.ShaderPipelineTypeName.GetDebugString().String.Get(), FShaderType::GetNameToTypeMap().Num());
+#endif
+
+		OutKeyString.AppendChar('_');
+		OutKeyString.Append(ShaderPipelineType->GetName());
+		Dependency.StagesSourceHash.AppendString(OutKeyString);
+
+		for (const FShaderType* ShaderType : ShaderPipelineType->GetStages())
+		{
+			if (const FShaderParametersMetadata* ParameterStructMetadata = ShaderType->GetRootParametersMetadata())
+			{
+				OutKeyString.Appendf(TEXT("%08x"), ParameterStructMetadata->GetLayoutHash());
+			}
+
+			for (const auto& Entry : ShaderType->GetReferencedUniformBufferStructsCache())
+			{
+				ReferencedUniformBuffers.Add(Entry.Key, Entry.Value);
+			}
+		}
+	}
+
+	for (const FVertexFactoryTypeDependency& VFDependency : VertexFactoryTypeDependencies)
+	{
+		OutKeyString.AppendChar('_');
+
+		const FVertexFactoryType* VertexFactoryType = FVertexFactoryType::GetVFByName(VFDependency.VertexFactoryTypeName);
+
+		OutKeyString.Append(VertexFactoryType->GetName());
+		VFDependency.VFSourceHash.AppendString(OutKeyString);
+
+		for (int32 Frequency = 0; Frequency < SF_NumFrequencies; Frequency++)
+		{
+			const FTypeLayoutDesc* ParameterLayout = VertexFactoryType->GetShaderParameterLayout((EShaderFrequency)Frequency);
+			if (ParameterLayout)
+			{
+				const FSHAHash LayoutHash = GetShaderTypeLayoutHash(*ParameterLayout, LayoutParams);
+				LayoutHash.AppendString(OutKeyString);
+			}
+		}
+
+		for (const auto& Entry : VertexFactoryType->GetReferencedUniformBufferStructsCache())
+		{
+			ReferencedUniformBuffers.Add(Entry.Key, Entry.Value);
+		}
+	}
+
+	{
+		TArray<uint8> TempData;
+		FSerializationHistory SerializationHistory;
+		FMemoryWriter Ar(TempData, true);
+		FShaderSaveArchive SaveArchive(Ar, SerializationHistory);
+
+		// Save uniform buffer member info so we can detect when layout has changed
+		SerializeUniformBufferInfo_Internal(SaveArchive, ReferencedUniformBuffers);
+
+		SerializationHistory.AppendKeyString(OutKeyString);
+	}
+}
+
+void SerializeUniformBufferInfo(FShaderSaveArchive& Ar, const TSortedMap<const TCHAR*, FCachedUniformBufferDeclaration, FDefaultAllocator, FUniformBufferNameSortOrder>& UniformBufferEntries)
+{
+	SerializeUniformBufferInfo_Internal(Ar, UniformBufferEntries);
+}
+
 FString MakeInjectedShaderCodeBlock(const TCHAR* BlockName, const FString& CodeToInject)
 {
 	return FString::Printf(TEXT("#line 1 \"%s\"\n%s"), BlockName, *CodeToInject);
diff --git a/Engine/Source/Runtime/RenderCore/Private/ShaderMap.cpp b/Engine/Source/Runtime/RenderCore/Private/ShaderMap.cpp
--- a/Engine/Source/Runtime/RenderCore/Private/ShaderMap.cpp
+++ b/Engine/Source/Runtime/RenderCore/Private/ShaderMap.cpp
@@ -8,7 +8,6 @@
 #include "Misc/CoreMisc.h"
 #include "Misc/StringBuilder.h"
 #include "Stats/StatsMisc.h"
-#include "Serialization/MemoryWriter.h"
 #include "VertexFactory.h"
 #include "ProfilingDebugging/DiagnosticTable.h"
 #include "Interfaces/ITargetPlatform.h"
diff --git a/Engine/Source/Runtime/RenderCore/Public/ShaderCore.h b/Engine/Source/Runtime/RenderCore/Public/ShaderCore.h
--- a/Engine/Source/Runtime/RenderCore/Public/ShaderCore.h
+++ b/Engine/Source/Runtime/RenderCore/Public/ShaderCore.h
@@ -1104,8 +1104,34 @@ struct FUniformBufferNameSortOrder
 	}
 };
 
-/** Records information about all the uniform buffer layouts referenced by UniformBufferEntries. */
-extern RENDERCORE_API void SerializeUniformBufferInfo(class FShaderSaveArchive& Ar, const TSortedMap<const TCHAR*,FCachedUniformBufferDeclaration, FDefaultAllocator, FUniformBufferNameSortOrder>& UniformBufferEntries);
+using FSortedMapUniformBufferDeclaration = TSortedMap<const TCHAR*, FCachedUniformBufferDeclaration, FDefaultAllocator, FUniformBufferNameSortOrder>;
+
+/** Records information about all the uniform buffer layouts referenced by UniformBufferEntries. This function is now deprecated as there now a unified way of preparing
+ * shadermap keys that includes this and more: AppendKeyStringShaderDependencies.
+ */
+UE_DEPRECATED(5.2, "SerializeUniformBufferInfo is depreceated. For creating shadermap keys please use AppendKeyStringShaderDependencies") extern RENDERCORE_API void SerializeUniformBufferInfo(class FShaderSaveArchive& Ar, const FSortedMapUniformBufferDeclaration& UniformBufferEntries);
+
+/**
+ * Return the hash of the given type layout for a partical platform type layout. This function employs caching to avoid re-hashing the same parameters several times.
+ */
+extern RENDERCORE_API FSHAHash GetShaderTypeLayoutHash(const FTypeLayoutDesc& TypeDesc, FPlatformTypeLayoutParameters LayoutParameters);
+
+// Forward declarations
+class FShaderTypeDependency;
+class FShaderPipelineTypeDependency;
+class FVertexFactoryTypeDependency;
+
+/** Appends information to a KeyString for a given shader to reflect its dependencies */
+extern RENDERCORE_API void AppendKeyStringShaderDependencies(
+	TConstArrayView<FShaderTypeDependency> ShaderTypeDependencies,
+	FPlatformTypeLayoutParameters LayoutParams,
+	FString& OutKeyString);
+extern RENDERCORE_API void AppendKeyStringShaderDependencies(
+	TConstArrayView<FShaderTypeDependency> ShaderTypeDependencies,
+	TConstArrayView<FShaderPipelineTypeDependency> ShaderPipelineTypeDependencies,
+	TConstArrayView<FVertexFactoryTypeDependency> VertexFactoryTypeDependencies,
+	FPlatformTypeLayoutParameters LayoutParams,
+	FString& OutKeyString);
 
 /** Create a block of source code to be injected in the preprocessed shader code. The Block will be put into a #line directive
  * to show up in case shader compilation failures happen in this code block.
