Fix editor crash when a BP is open and right click on Content folder and select Sync.

The crash happens because a series of packages are reloaded, and the associated skeletal mesh has its render data recreated, but the skinned mesh component doesn't recreate render state until later. So the mesh object holds the old render data pointer that contains invalid data temporarily, causing it not to release its render resource, and the crash happens momently later on render thread for destroying a render resource belonging to mesh object that hasn't been released.

Removing the render data check allows mesh object to release its render resource, and later on the skinned mesh component recreates the mesh object pointing to the new render data.

diff --git a/Engine/Source/Runtime/Engine/Private/SkeletalRenderCPUSkin.cpp b/Engine/Source/Runtime/Engine/Private/SkeletalRenderCPUSkin.cpp
--- a/Engine/Source/Runtime/Engine/Private/SkeletalRenderCPUSkin.cpp
+++ b/Engine/Source/Runtime/Engine/Private/SkeletalRenderCPUSkin.cpp
@@ -101,12 +101,7 @@ void FSkeletalMeshObjectCPUSkin::ReleaseResources()
 	for( int32 LODIndex=0;LODIndex < LODs.Num();LODIndex++ )
 	{
 		FSkeletalMeshObjectLOD& SkelLOD = LODs[LODIndex];
-
-		// Skip LODs that have their render data stripped
-		if (SkelLOD.SkelMeshRenderData->LODRenderData[LODIndex].GetNumVertices() > 0)
-		{
-			SkelLOD.ReleaseResources();
-		}
+		SkelLOD.ReleaseResources();
 	}
 }
 
diff --git a/Engine/Source/Runtime/Engine/Private/SkeletalRenderGPUSkin.cpp b/Engine/Source/Runtime/Engine/Private/SkeletalRenderGPUSkin.cpp
--- a/Engine/Source/Runtime/Engine/Private/SkeletalRenderGPUSkin.cpp
+++ b/Engine/Source/Runtime/Engine/Private/SkeletalRenderGPUSkin.cpp
@@ -315,12 +315,7 @@ void FSkeletalMeshObjectGPUSkin::ReleaseResources()
 	for( int32 LODIndex=0;LODIndex < LODs.Num();LODIndex++ )
 	{
 		FSkeletalMeshObjectLOD& SkelLOD = LODs[LODIndex];
-		
-		// Skip LODs that have their render data stripped
-		if (SkelLOD.SkelMeshRenderData && SkelLOD.SkelMeshRenderData->LODRenderData.IsValidIndex(LODIndex) && SkelLOD.SkelMeshRenderData->LODRenderData[LODIndex].GetNumVertices() > 0)
-		{
-			SkelLOD.ReleaseResources();
-		}
+		SkelLOD.ReleaseResources();
 	}
 	// also release morph resources
 	ReleaseMorphResources();
