Texture apply configurable new defaults to properties on import
so that we can set standard new properites values for new texures and leave old values on legacy textures
also update oodle sdk version to latest on reimport

diff --git a/Engine/Plugins/Interchange/Runtime/Source/Import/Private/Texture/InterchangeTextureFactory.cpp b/Engine/Plugins/Interchange/Runtime/Source/Import/Private/Texture/InterchangeTextureFactory.cpp
--- a/Engine/Plugins/Interchange/Runtime/Source/Import/Private/Texture/InterchangeTextureFactory.cpp
+++ b/Engine/Plugins/Interchange/Runtime/Source/Import/Private/Texture/InterchangeTextureFactory.cpp
@@ -1432,6 +1432,9 @@ namespace UE::Interchange::Private::InterchangeTextureFactory
 			}
 		}
 #endif // WITH_EDITORONLY_DATA
+
+		// at this point the texture is mostly set up
+		// NormalMapIdentification is not run yet
 	}
  }
 
@@ -1569,6 +1572,7 @@ UObject* UInterchangeTextureFactory::CreateAsset(const FCreateAssetParams& Argum
 	}
 
 	// create an asset if it doesn't exist
+	// typically ExistingAsset should already exist, made by CreateEmptyAsset on the main thread
 	UObject* ExistingAsset = StaticFindObject(nullptr, Arguments.Parent, *Arguments.AssetName);
 
 	UTexture* Texture = nullptr;
@@ -1582,7 +1586,8 @@ UObject* UInterchangeTextureFactory::CreateAsset(const FCreateAssetParams& Argum
 	}
 	else if(ExistingAsset->GetClass()->IsChildOf(TextureClass))
 	{
-		//This is a reimport, we are just re-updating the source data
+		//This is a reimport, we are just re-updating the source data  
+		// <- pretty sure this comment is wrong, this is hit in regular imports, because ExistingAsset was previously made by CreateNewAsset
 		Texture = static_cast<UTexture*>(ExistingAsset);
 	}
 
@@ -1670,6 +1675,9 @@ void UInterchangeTextureFactory::PreImportPreCompletedCallback(const FImportPreC
 		UInterchangeFactoryBaseNode* TextureFactoryNode = Arguments.FactoryNode;
 		if (!Arguments.bIsReimport)
 		{
+			// Finalize texture properties, before ApplyAllCustomAttributeToObject :
+			UE::TextureUtilitiesCommon::ApplyDefaultsForNewlyImportedTextures(Texture,Arguments.bIsReimport);
+
 			/** Apply all TextureNode custom attributes to the texture asset */
 			TextureFactoryNode->ApplyAllCustomAttributeToObject(Texture);
 		}
@@ -1687,7 +1695,14 @@ void UInterchangeTextureFactory::PreImportPreCompletedCallback(const FImportPreC
 			CurrentNode->FillAllCustomAttributeFromObject(Texture);
 			//Apply reimport strategy
 			UE::Interchange::FFactoryCommon::ApplyReimportStrategyToAsset(Texture, PreviousNode, CurrentNode, TextureFactoryNode);
+			
+			// ApplyDefaultsForNewlyImportedTextures after ApplyReimportStrategyToAsset
+			//	so we can override values
+			UE::TextureUtilitiesCommon::ApplyDefaultsForNewlyImportedTextures(Texture,Arguments.bIsReimport);
 		}
+
+		// Texture is mostly done setting properties now (exception: FTaskPipelinePostImport NormalmapIdentification)
+		// PostEditChange will be called subseqently, in FTaskPreCompletion::DoTask
 #endif // WITH_EDITOR
 	}
 	else
diff --git a/Engine/Plugins/Interchange/Runtime/Source/Pipelines/Private/InterchangeGenericTexturePipeline.cpp b/Engine/Plugins/Interchange/Runtime/Source/Pipelines/Private/InterchangeGenericTexturePipeline.cpp
--- a/Engine/Plugins/Interchange/Runtime/Source/Pipelines/Private/InterchangeGenericTexturePipeline.cpp
+++ b/Engine/Plugins/Interchange/Runtime/Source/Pipelines/Private/InterchangeGenericTexturePipeline.cpp
@@ -83,7 +83,8 @@ namespace UE::Interchange::Private
 					Texture->bFlipGreenChannel = true;
 				}
 			}
-			Texture->PostEditChange();
+			// this will rebuild the texture if it changed to normal map
+			Texture->PostEditChange(); 
 		}
 	}
 #endif
@@ -332,16 +333,23 @@ UInterchangeTextureFactoryNode* UInterchangeGenericTexturePipeline::CreateTextur
 void UInterchangeGenericTexturePipeline::PostImportTextureAssetImport(UObject* CreatedAsset, bool bIsAReimport)
 {
 #if WITH_EDITOR
+
+	// this is run on main thread
+	//	after texture may have started compiling
+
 	if (!bIsAReimport && bDetectNormalMapTexture)
 	{
-		// Verify if the texture is a normal map
 		if (UTexture* Texture = Cast<UTexture>(CreatedAsset))
 		{
+			// if it's already a normal map, no need to run NormalMapIdentification
 			if (!Texture->IsNormalMap())
 			{
-				// This can create 2 build of the texture (we should revisit this at some point)
+				// @todo Interchange: This can create 2 build of the texture (we should revisit this at some point)
+
 				if (FTextureCompilingManager::Get().IsCompilingTexture(Texture))
 				{
+					// install a lambda to do AdjustTextureForNormalMap which will run after compile is done
+
 					TWeakObjectPtr<UTexture> WeakTexturePtr = Texture;
 					TSharedRef<FDelegateHandle> HandlePtr = MakeShared<FDelegateHandle>();
 					HandlePtr.Get() = FTextureCompilingManager::Get().OnTexturePostCompileEvent().AddLambda([this, WeakTexturePtr, HandlePtr](const TArrayView<UTexture* const>&)
@@ -350,17 +358,26 @@ void UInterchangeGenericTexturePipeline::PostImportTextureAssetImport(UObject* C
 							{
 								if (FTextureCompilingManager::Get().IsCompilingTexture(TextureToTest))
 								{
+									// I'm still in compile queue, this lambda will get called again when I finish the next compile
+									//	OnTexturePostCompileEvent().Remove is not called, leave the lambda
+									// this can happen if multiple compiles of this texture got queued
 									return;
 								}
 
 								UE::Interchange::Private::AdjustTextureForNormalMap(TextureToTest, bFlipNormalMapGreenChannel);
+								// ?? does the asset registry need to be fixed?
+								//	asset import notifications were sent earlier
 							}
 
+							// remove myself:
 							FTextureCompilingManager::Get().OnTexturePostCompileEvent().Remove(HandlePtr.Get());
 						});
 				}
 				else
 				{
+					//Texture->Modify(); // Modify blocks on async build task being done
+
+					// AdjustTextureForNormalMap does a PostEditChange which triggers a rebuild
 					UE::Interchange::Private::AdjustTextureForNormalMap(Texture, bFlipNormalMapGreenChannel);
 				}
 			}
diff --git a/Engine/Source/Editor/UnrealEd/Private/Factories/EditorFactories.cpp b/Engine/Source/Editor/UnrealEd/Private/Factories/EditorFactories.cpp
--- a/Engine/Source/Editor/UnrealEd/Private/Factories/EditorFactories.cpp
+++ b/Engine/Source/Editor/UnrealEd/Private/Factories/EditorFactories.cpp
@@ -4062,48 +4062,41 @@ UObject* UTextureFactory::FactoryCreateBinary
 	//	notice above that CompressionSettings is primed from Texture->CompressionSettings
 	//	but LODGroup is not pulled from Texture->LODGroup
 
-	// Figure out whether we're using a normal map LOD group.
-	bool bIsNormalMapLODGroup = false;
-	if( LODGroup == TEXTUREGROUP_WorldNormalMap 
-	||	LODGroup == TEXTUREGROUP_CharacterNormalMap
-	||	LODGroup == TEXTUREGROUP_VehicleNormalMap
-	||	LODGroup == TEXTUREGROUP_WeaponNormalMap )
+	// If the TextureFactory LODGroup is set to a normal map LOD group, then set CompressionSettings to Normalmap
+	bool bIsNormalMapLODGroup = ( LODGroup == TEXTUREGROUP_WorldNormalMap 
+		||	LODGroup == TEXTUREGROUP_CharacterNormalMap
+		||	LODGroup == TEXTUREGROUP_VehicleNormalMap
+		||	LODGroup == TEXTUREGROUP_WeaponNormalMap );
+
+	if ( bIsNormalMapLODGroup )
 	{
 		// Change from default to normal map.
 		if( CompressionSettings == TC_Default )
 		{
 			CompressionSettings = TC_Normalmap;
 		}
-		bIsNormalMapLODGroup = true;
+	}
+		
+	// Propagate options.
+	Texture->CompressionSettings	= CompressionSettings;
+	Texture->LODGroup				= LODGroup;
+	
+	if(!FCString::Stricmp(Type, TEXT("ies")))
+	{
+		Texture->LODGroup = TEXTUREGROUP_IESLightProfile;
 	}
 
-	// Packed normal map
-	if( Texture->IsNormalMap() )
+	// note that NormalmapIdentification has not run yet
+	//	so I think the only way you get in this branch is if bIsNormalMapLODGroup == true
+	if( Texture->IsNormalMap() ) // TC == TC_Normalmap
 	{
 		Texture->SRGB = 0;
 		if( !bIsNormalMapLODGroup )
 		{
-			LODGroup = TEXTUREGROUP_WorldNormalMap;
+			Texture->LODGroup = TEXTUREGROUP_WorldNormalMap;
 		}
 	}
 
-	if(!FCString::Stricmp(Type, TEXT("ies")))
-	{
-		LODGroup = TEXTUREGROUP_IESLightProfile;
-	}
-	
-	// Propagate options.
-	Texture->CompressionSettings	= CompressionSettings;
-	Texture->LODGroup				= LODGroup;
-
-	// Revert the LODGroup to the default if it was forcibly set by the texture being a normal map.
-	// This handles the case where multiple textures are being imported consecutively and
-	// LODGroup unexpectedly changes because some textures were normal maps and others weren't.
-	if ( LODGroup == TEXTUREGROUP_WorldNormalMap && !bIsNormalMapLODGroup )
-	{
-		LODGroup = TEXTUREGROUP_World;
-	}
-
 	Texture->CompressionNone				= NoCompression;
 	Texture->CompressionNoAlpha				= NoAlpha;
 	Texture->DeferCompression				= bDeferCompression;
@@ -4143,6 +4136,9 @@ UObject* UTextureFactory::FactoryCreateBinary
 				Texture->CompressionSettings = TC_EditorIcon; // "UserInterface2D"
 			}
 		}
+
+		// @@CB move this to ApplyDefaultsForNewlyImportedTextures so it can be shared with Interchange, not code-duped
+		// also set LODGroup to UI ?
 	}
 			
 	static const auto CVarVirtualTexturesEnabled = IConsoleManager::Get().FindTConsoleVariableDataInt(TEXT("r.VirtualTextures"));
@@ -4255,6 +4251,11 @@ UObject* UTextureFactory::FactoryCreateBinary
 			Texture->bFlipGreenChannel = bFlipNormalMapGreenChannel;
 		}
 	}
+	
+	// Texture property setup is almost done
+	//  ApplyDefaultsForNewlyImportedTextures before ApplyAutoImportSettings so user settings can override
+	bool bIsReimport = ExistingTexture && bUsingExistingSettings;
+	UE::TextureUtilitiesCommon::ApplyDefaultsForNewlyImportedTextures(Texture,bIsReimport);
 
 	if(IsAutomatedImport())
 	{
diff --git a/Engine/Source/Runtime/Engine/Classes/Engine/Texture.h b/Engine/Source/Runtime/Engine/Classes/Engine/Texture.h
--- a/Engine/Source/Runtime/Engine/Classes/Engine/Texture.h
+++ b/Engine/Source/Runtime/Engine/Classes/Engine/Texture.h
@@ -1210,8 +1210,8 @@ public:
 	UPROPERTY(EditAnywhere, BlueprintReadWrite, Category = Texture, meta=(ClampMin = "0", ClampMax = "1.0", EditCondition="bDoScaleMipsForAlphaCoverage"), AdvancedDisplay)
 	FVector4 AlphaCoverageThresholds = FVector4(0,0,0,0);
 
-	/** Whether to use newer & faster mip generation filter, same quality but produces slightly different results from previous implementation */
-	UPROPERTY(EditAnywhere, BlueprintReadWrite, Category=Texture, AdvancedDisplay)
+	/** Use faster mip generation filter, usually the same result but occasionally causes color shift in high contrast areas. */
+	UPROPERTY(EditAnywhere, BlueprintReadWrite, Category=Texture, meta=(DisplayName = "Use Fast MipGen Filter"), AdvancedDisplay)
 	bool bUseNewMipFilter = false;
 
 	/** When true the texture's border will be preserved during mipmap generation. */
@@ -1260,7 +1260,7 @@ public:
 
 	/**
 	 * default 1, high values result in a stronger effect e.g 1, 2, 4, 8
-	 * this is no slider because the texture update would not be fast enough
+	 * this is not a slider because the texture update would not be fast enough
 	 */
 	UPROPERTY(EditAnywhere, BlueprintReadWrite, Category=Compositing, AdvancedDisplay)
 	float CompositePower;
@@ -1289,7 +1289,7 @@ public:
 	UPROPERTY(transient, duplicatetransient, NonTransactional)
 	int32 LevelIndex = INDEX_NONE;
 
-	/** A bias to the index of the top mip level to use. */
+	/** A bias to the index of the top mip level to use.  That is, number of mip levels to drop when cooking. */
 	UPROPERTY(EditAnywhere, BlueprintReadWrite, Category=LevelOfDetail, meta=(DisplayName="LOD Bias"), AssetRegistrySearchable)
 	int32 LODBias;
 
@@ -1309,7 +1309,7 @@ public:
 	UPROPERTY(EditAnywhere, BlueprintReadWrite, Category=LevelOfDetail, meta=(DisplayName="Texture Group"), AssetRegistrySearchable)
 	TEnumAsByte<enum TextureGroup> LODGroup;
 
-	/** Downscale source texture, applied only to textures without mips 
+	/** Downscale source texture, applied only to 2d textures without mips 
 	 * 0.0 - use scale value from texture group
 	 * 1.0 - do not scale texture
 	 * > 1.0 - scale texure
@@ -1322,12 +1322,12 @@ public:
 	ETextureDownscaleOptions DownscaleOptions;
 
 	
-	/** This should be unchecked if using alpha channels individually as masks. */
+	/** Whether Texture and its source are in SRGB Gamma color space.  Can only be used with 8-bit and compressed formats.  This should be unchecked if using alpha channels individually as masks. */
 	UPROPERTY(EditAnywhere, BlueprintReadWrite, Category=Texture, meta=(DisplayName="sRGB"), AssetRegistrySearchable)
 	uint8 SRGB:1;
 
 #if WITH_EDITORONLY_DATA
-	/* Should Normals in Normal Maps be normalized after mip processing?  true gives better quality normals, but false is required to match legacy behavior. */
+	/* Normalize colors in Normal Maps after mip generation for better and sharper quality; recommended on if not required to match legacy behavior. */
 	UPROPERTY(EditAnywhere, BlueprintReadWrite, Category=Texture, meta=(DisplayName="Normalize after making mips", EditCondition="CompressionSettings==1"), AdvancedDisplay)
 	uint8 bNormalizeNormals:1;
 
@@ -1611,6 +1611,11 @@ public:
 		
 	/** Ensure settings are valid after import or edit; this is called by PostEditChange. */
 	ENGINE_API virtual void ValidateSettingsAfterImportOrEdit(bool * pRequiresNotifyMaterials = nullptr);
+
+	/* Change the Oodle Texture Sdk Version used to encode this texture to latest.
+	*  You should do this any time the texture is modified, such as on reimport, since the bits are changing anyway.
+	*/
+	ENGINE_API virtual void UpdateOodleTextureSdkVersionToLatest(void);
 #endif // WITH_EDITOR
 
 	ENGINE_API EGammaSpace GetGammaSpace() const
diff --git a/Engine/Source/Runtime/Engine/Private/Texture.cpp b/Engine/Source/Runtime/Engine/Private/Texture.cpp
--- a/Engine/Source/Runtime/Engine/Private/Texture.cpp
+++ b/Engine/Source/Runtime/Engine/Private/Texture.cpp
@@ -365,6 +365,14 @@ bool UTexture::CanEditChange(const FProperty* InProperty) const
 	return true;
 }
 
+void UTexture::UpdateOodleTextureSdkVersionToLatest(void)
+{
+	// OodleTextureSdkVersion = get latest sdk version
+	//	this needs to get the actual version number so it will be IO'd frozen (not just "latest")
+	OodleTextureSdkVersion = CachedGetLatestOodleSdkVersion();
+}
+
+
 // we're in WITH_EDITOR but not sure that's right, maybe move out?
 // Beware: while ValidateSettingsAfterImportOrEdit should have been called on all Textures,
 //	 it is not called on load at runtime
@@ -924,9 +932,10 @@ void UTexture::PostInitProperties()
 	{
 		AssetImportData = NewObject<UAssetImportData>(this, TEXT("AssetImportData"));
 
-		// OodleTextureSdkVersion = get latest sdk version
-		//	this needs to get the actual version number so it will be IO'd frozen (not just "latest")
-		OodleTextureSdkVersion = CachedGetLatestOodleSdkVersion();
+		// this is for textures that are not being loaded
+		// eg. created from code, eg. lightmaps
+		// we want them to go ahead and use the latest oodle sdk, since their content is new anyway
+		UpdateOodleTextureSdkVersionToLatest();
 	}
 #endif
 	Super::PostInitProperties();
diff --git a/Engine/Source/Runtime/TextureUtilitiesCommon/Private/NormalMapIdentification.cpp b/Engine/Source/Runtime/TextureUtilitiesCommon/Private/NormalMapIdentification.cpp
--- a/Engine/Source/Runtime/TextureUtilitiesCommon/Private/NormalMapIdentification.cpp
+++ b/Engine/Source/Runtime/TextureUtilitiesCommon/Private/NormalMapIdentification.cpp
@@ -457,6 +457,7 @@ public:
 				Texture2D->WaitForPendingInitOrStreaming();
 
 				Texture2D->SetFlags(RF_Transactional);
+				// Modify calls FinishCachePlatformData to wait on any async build of this texture
 				Texture2D->Modify();
 				Texture2D->PreEditChange(NULL);
 				{
@@ -485,12 +486,13 @@ bool UE::NormalMapIdentification::HandleAssetPostImport( UTexture* Texture )
 	if( Texture != NULL)
 	{
 		// Try to automatically identify a normal map
+		//	this only reads Texture->Source
 		if ( IsTextureANormalMap( Texture ) )
 		{
 			// Set the compression settings and no gamma correction for a normal map
 			{
 				Texture->SetFlags(RF_Transactional);
-				Texture->Modify();
+				Texture->Modify(); // Modify blocks on pending texture builds
 				Texture->CompressionSettings = TC_Normalmap;
 				Texture->SRGB = false;
 				Texture->LODGroup = TEXTUREGROUP_WorldNormalMap;
@@ -498,6 +500,7 @@ bool UE::NormalMapIdentification::HandleAssetPostImport( UTexture* Texture )
 
 			// Show the user a notification indicating that this texture will be imported as a normal map.
 			// Offer two options to the user, "OK" dismisses the notification early, "Revert" reverts the settings to that of a diffuse map.
+			// ?? Guess?? this has to be done from main thread only??
 			TSharedPtr<NormalMapImportNotificationHandler> NormalMapNotificationDelegate(new NormalMapImportNotificationHandler);
 			{
 				NormalMapNotificationDelegate->Texture = Texture;
diff --git a/Engine/Source/Runtime/TextureUtilitiesCommon/Private/TextureImportSettings.cpp b/Engine/Source/Runtime/TextureUtilitiesCommon/Private/TextureImportSettings.cpp
--- a/Engine/Source/Runtime/TextureUtilitiesCommon/Private/TextureImportSettings.cpp
+++ b/Engine/Source/Runtime/TextureUtilitiesCommon/Private/TextureImportSettings.cpp
@@ -5,13 +5,16 @@
 #include "UObject/UnrealType.h"
 #include "UObject/PropertyPortFlags.h"
 
+#if WITH_EDITOR
+#include "Engine/Texture.h"
+#endif
+
 #include UE_INLINE_GENERATED_CPP_BY_NAME(TextureImportSettings)
 
 UTextureImportSettings::UTextureImportSettings(const FObjectInitializer& ObjectInitializer)
 	: Super(ObjectInitializer)
 {
 	SectionName = TEXT("Importing");
-	AutoVTSize = 4096;
 }
 
 void UTextureImportSettings::PostInitProperties()
@@ -37,3 +40,50 @@ void UTextureImportSettings::PostEditChangeProperty(FPropertyChangedEvent& Prope
 #endif // #if WITH_EDITOR
 
 
+
+#if WITH_EDITOR
+namespace UE::TextureUtilitiesCommon
+{
+	/* Set default properties on Texture for newly imported textures, or reimports.
+	*  Should be called after all texture properties are set, before PostEditChange() 
+	*/
+	TEXTUREUTILITIESCOMMON_API void ApplyDefaultsForNewlyImportedTextures(UTexture * Texture, bool bIsReimport)
+	{
+	
+		// things that are done for both fresh import and reimport :
+
+		Texture->UpdateOodleTextureSdkVersionToLatest();
+
+		if ( bIsReimport )
+		{
+			return;
+		}
+
+		// things that are for fresh import only :
+
+		// here we can change values that must have different defaults for backwards compatibility
+		// we set them to the new desired value here, the Texture constructor sets the legacy value
+	
+		const UTextureImportSettings* Settings = GetDefault<UTextureImportSettings>();
+
+		if ( Settings->bEnableNormalizeNormals )
+		{
+			// cannot check for TC_Normalmap here
+			//	because of the way NormalmapIdentification is delayed in Interchange
+			//	it's harmless to just always turn it on
+			//	it will be ignored if we are not TC_Normalmap
+			//	OutBuildSettings.bNormalizeNormals = Texture.bNormalizeNormals && Texture.IsNormalMap();
+			//if ( Texture->CompressionSettings == TC_Normalmap )
+			{
+				Texture->bNormalizeNormals = true;
+			}
+		}
+
+		if ( Settings->bEnableFastMipFilter )
+		{
+			Texture->bUseNewMipFilter = true;
+		}
+	}
+
+}
+#endif
\ No newline at end of file
diff --git a/Engine/Source/Runtime/TextureUtilitiesCommon/Public/TextureImportSettings.h b/Engine/Source/Runtime/TextureUtilitiesCommon/Public/TextureImportSettings.h
--- a/Engine/Source/Runtime/TextureUtilitiesCommon/Public/TextureImportSettings.h
+++ b/Engine/Source/Runtime/TextureUtilitiesCommon/Public/TextureImportSettings.h
@@ -15,13 +15,23 @@ UCLASS(config=Editor, defaultconfig, meta=(DisplayName="Texture Import"))
 class TEXTUREUTILITIESCOMMON_API UTextureImportSettings : public UDeveloperSettings
 {
 	GENERATED_UCLASS_BODY()
+	
+public:
 
-	UPROPERTY(config, EditAnywhere, Category = VirtualTextures, meta = (
+	UPROPERTY(config, EditAnywhere, Category=VirtualTextures, meta = (
 		DisplayName = "Auto Virtual Texturing Size",
 		ToolTip = "Automatically enable the 'Virtual Texture Streaming' texture setting for textures larger than or equal to this size. This setting will not affect existing textures in the project."))
-	int32 AutoVTSize;
-
-public:
+	int32 AutoVTSize = 4096;
+	
+	UPROPERTY(config, EditAnywhere, Category=ImportSettings, meta = (
+		DisplayName = "Turn on NormalizeNormals for normal maps",
+		ToolTip = "NormalizeNormals makes more correct normals in mip maps; it is recommended, but can be turned off to maintain legacy behavior. This setting is applied to newly imported textures, it does not affect existing textures in the project."))
+	bool bEnableNormalizeNormals = true;
+	
+	UPROPERTY(config, EditAnywhere, Category=ImportSettings, meta = (
+		DisplayName = "Turn on fast mip generation filter",
+		ToolTip = "Use the fast mip filter on new textures; it is recommended, but can be turned off to maintain legacy behavior. This setting is applied to newly imported textures, it does not affect existing textures in the project."))
+	bool bEnableFastMipFilter = true;
 
 	//~ Begin UObject Interface
 
@@ -33,3 +43,17 @@ public:
 
 	//~ End UObject Interface
 };
+
+
+#if WITH_EDITOR
+class UTexture;
+
+namespace UE::TextureUtilitiesCommon
+{
+	/* Set default properties on Texture for newly imported textures, or reimports.
+	*  Should be called after all texture properties are set, before PostEditChange() 
+	*/
+	TEXTUREUTILITIESCOMMON_API void ApplyDefaultsForNewlyImportedTextures(UTexture * Texture, bool bIsReimport);
+
+}
+#endif
\ No newline at end of file
