Clarify meaning of SSS diffusion option on slab for Strata users.

diff --git a/Engine/Source/Runtime/Engine/Classes/Materials/MaterialExpressionStrata.h b/Engine/Source/Runtime/Engine/Classes/Materials/MaterialExpressionStrata.h
--- a/Engine/Source/Runtime/Engine/Classes/Materials/MaterialExpressionStrata.h
+++ b/Engine/Source/Runtime/Engine/Classes/Materials/MaterialExpressionStrata.h
@@ -327,7 +327,7 @@ class UMaterialExpressionStrataSlabBSDF : public UMaterialExpressionStrataBSDF
 	UPROPERTY(EditAnywhere, Category = Mode)
 	uint32 bUseMetalness : 1;
 
-	/** Whether to use light diffusion (i.e., SSS diffusion) or wrap-approximation for material with scattering behavior. */
+	/** Whether to use light diffusion (i.e., SSS diffusion) or wrap-approximation for material with scattering behavior. This option trades quality over performance and will result into visual differences. */
 	UPROPERTY(EditAnywhere, Category = Mode, meta = (DisplayName = "Use Subsurface Diffusion"))
 	uint32 bUseSSSDiffusion : 1;
 
