ObjectTraceAnnotation: just return null if the default object id is searched for, to avoid checkSlow

diff --git a/Engine/Source/Runtime/Engine/Private/ObjectTrace.cpp b/Engine/Source/Runtime/Engine/Private/ObjectTrace.cpp
--- a/Engine/Source/Runtime/Engine/Private/ObjectTrace.cpp
+++ b/Engine/Source/Runtime/Engine/Private/ObjectTrace.cpp
@@ -218,6 +218,11 @@ UObject* FObjectTrace::GetObjectFromId(uint64 Id)
 	FObjectIdAnnotation FindAnnotation;
 	// Id used for annotation map doesn't include the parent id in the upper bits, so zero those first
 	FindAnnotation.Id = Id & 0x00000000FFFFFFFFll;
+	if (FindAnnotation.IsDefault())
+	{
+		return nullptr;
+	}
+	
 	return GObjectIdAnnotations.Find(FindAnnotation);
 }
 
