Github PR#9223: adding overridable tick function to allow server-side functionality on remote autonomous proxy characters.

diff --git a/Engine/Source/Runtime/Engine/Classes/GameFramework/CharacterMovementComponent.h b/Engine/Source/Runtime/Engine/Classes/GameFramework/CharacterMovementComponent.h
--- a/Engine/Source/Runtime/Engine/Classes/GameFramework/CharacterMovementComponent.h
+++ b/Engine/Source/Runtime/Engine/Classes/GameFramework/CharacterMovementComponent.h
@@ -2124,6 +2124,12 @@ protected:
 	/** Simulate movement on a non-owning client. Called by SimulatedTick(). */
 	virtual void SimulateMovement(float DeltaTime);
 
+	/** Special Tick to allow custom server-side functionality on Autonomous Proxies. 
+	 * Called for all remote APs, including APs controlled on Listen Servers such as the hosting player's Character.
+	 * If full server-side control is desired, you may need to override ControlledCharacterMove as well.
+	 */
+	virtual void ServerAutonomousProxyTick(float DeltaSeconds) { }
+
 public:
 
 	/** Force a client update by making it appear on the server that the client hasn't updated in a long time. */
diff --git a/Engine/Source/Runtime/Engine/Private/Components/CharacterMovementComponent.cpp b/Engine/Source/Runtime/Engine/Private/Components/CharacterMovementComponent.cpp
--- a/Engine/Source/Runtime/Engine/Private/Components/CharacterMovementComponent.cpp
+++ b/Engine/Source/Runtime/Engine/Private/Components/CharacterMovementComponent.cpp
@@ -1491,10 +1491,23 @@ void UCharacterMovementComponent::TickComponent(float DeltaTime, enum ELevelTick
 			}
 		}
 
-		// Allow root motion to move characters that have no controller.
-		if (CharacterOwner->IsLocallyControlled() || (!CharacterOwner->Controller && bRunPhysicsWithNoController) || (!CharacterOwner->Controller && CharacterOwner->IsPlayingRootMotion()))
+		// Perform input-driven move for any locally-controlled character, and also
+		// allow animation root motion or physics to move characters even if they have no controller
+		const bool bShouldPerformControlledCharMove = CharacterOwner->IsLocallyControlled() 
+													  || (!CharacterOwner->Controller && bRunPhysicsWithNoController)		
+													  || (!CharacterOwner->Controller && CharacterOwner->IsPlayingRootMotion());
+		
+		if (bShouldPerformControlledCharMove)
 		{
 			ControlledCharacterMove(InputVector, DeltaTime);
+
+			const bool bIsaListenServerAutonomousProxy = CharacterOwner->IsLocallyControlled()
+													 	 && (CharacterOwner->GetRemoteRole() == ROLE_AutonomousProxy);
+
+			if (bIsaListenServerAutonomousProxy)
+			{
+				ServerAutonomousProxyTick(DeltaTime);
+			}
 		}
 		else if (CharacterOwner->GetRemoteRole() == ROLE_AutonomousProxy)
 		{
@@ -1504,6 +1517,8 @@ void UCharacterMovementComponent::TickComponent(float DeltaTime, enum ELevelTick
 			MaybeUpdateBasedMovement(DeltaTime);
 			MaybeSaveBaseLocation();
 
+			ServerAutonomousProxyTick(DeltaTime);
+
 			// Smooth on listen server for local view of remote clients. We may receive updates at a rate different than our own tick rate.
 			if (CharacterMovementCVars::NetEnableListenServerSmoothing && !bNetworkSmoothingComplete && IsNetMode(NM_ListenServer))
 			{
