PR #9218: Fixes bug where Root Motion Source's "move to dynamic" force would not be updated to match a server correction, reducing severity of corrections while target is moving. This greatly improves perceived movement quality under bad networking conditions.

diff --git a/Engine/Source/Runtime/Engine/Classes/GameFramework/RootMotionSource.h b/Engine/Source/Runtime/Engine/Classes/GameFramework/RootMotionSource.h
--- a/Engine/Source/Runtime/Engine/Classes/GameFramework/RootMotionSource.h
+++ b/Engine/Source/Runtime/Engine/Classes/GameFramework/RootMotionSource.h
@@ -610,6 +610,7 @@ struct ENGINE_API FRootMotionSource_MoveToDynamicForce : public FRootMotionSourc
 	UPROPERTY()
 	FVector InitialTargetLocation;
 
+	// Dynamically-changing location of target, which may be altered while this movement is ongoing 
 	UPROPERTY()
 	FVector TargetLocation;
 
diff --git a/Engine/Source/Runtime/Engine/Private/GameFramework/RootMotionSource.cpp b/Engine/Source/Runtime/Engine/Private/GameFramework/RootMotionSource.cpp
--- a/Engine/Source/Runtime/Engine/Private/GameFramework/RootMotionSource.cpp
+++ b/Engine/Source/Runtime/Engine/Private/GameFramework/RootMotionSource.cpp
@@ -882,8 +882,12 @@ bool FRootMotionSource_MoveToDynamicForce::MatchesAndHasSameState(const FRootMot
 	{
 		return false;
 	}
+	
+	// We can cast safely here since in FRootMotionSource::Matches() we ensured ScriptStruct equality
+	const FRootMotionSource_MoveToDynamicForce* OtherCast = static_cast<const FRootMotionSource_MoveToDynamicForce*>(Other);
 
-	return true; // MoveToDynamicForce has no unique state
+	return (StartLocation.Equals(OtherCast->StartLocation) &&
+			TargetLocation.Equals(OtherCast->TargetLocation));
 }
 
 bool FRootMotionSource_MoveToDynamicForce::UpdateStateFrom(const FRootMotionSource* SourceToTakeStateFrom, bool bMarkForSimulatedCatchup)
@@ -893,7 +897,13 @@ bool FRootMotionSource_MoveToDynamicForce::UpdateStateFrom(const FRootMotionSour
 		return false;
 	}
 
-	return true; // MoveToDynamicForce has no unique state other than Time which is handled by FRootMotionSource
+	// We can cast safely here since in FRootMotionSource::UpdateStateFrom() we ensured ScriptStruct equality
+	const FRootMotionSource_MoveToDynamicForce* OtherCast = static_cast<const FRootMotionSource_MoveToDynamicForce*>(SourceToTakeStateFrom);
+
+	StartLocation = OtherCast->StartLocation;
+	TargetLocation = OtherCast->TargetLocation;
+
+	return true;
 }
 
 void FRootMotionSource_MoveToDynamicForce::SetTime(float NewTime)
