Added optional earlier validation of client URL. Fixed clientside handling of encryption failure customization.

diff --git a/Engine/Plugins/Online/OnlineSubsystemUtils/Source/OnlineSubsystemUtils/Private/OnlineBeaconClient.cpp b/Engine/Plugins/Online/OnlineSubsystemUtils/Source/OnlineSubsystemUtils/Private/OnlineBeaconClient.cpp
--- a/Engine/Plugins/Online/OnlineSubsystemUtils/Source/OnlineSubsystemUtils/Private/OnlineBeaconClient.cpp
+++ b/Engine/Plugins/Online/OnlineSubsystemUtils/Source/OnlineSubsystemUtils/Private/OnlineBeaconClient.cpp
@@ -251,24 +251,51 @@ void AOnlineBeaconClient::SetEncryptionData(const FEncryptionData& InEncryptionD
 
 void AOnlineBeaconClient::SendInitialJoin()
 {
-	if (ensure(NetDriver != nullptr && NetDriver->ServerConnection != nullptr))
+	UNetConnection* ServerConn = NetDriver != nullptr ? NetDriver->ServerConnection : nullptr;
+
+	if (ensure(ServerConn != nullptr))
 	{
 		uint8 IsLittleEndian = uint8(PLATFORM_LITTLE_ENDIAN);
 		check(IsLittleEndian == !!IsLittleEndian); // should only be one or zero
 
+		const int32 AllowEncryption = CVarNetAllowEncryption.GetValueOnGameThread();
 		uint32 LocalNetworkVersion = FNetworkVersion::GetLocalNetworkVersion();
 
-		if (CVarNetAllowEncryption.GetValueOnGameThread() == 0)
+		if (AllowEncryption == 0)
 		{
 			EncryptionData.Identifier.Empty();
 		}
 
+		bool bEncryptionRequirementsFailure = false;
+
+		if (EncryptionData.Identifier.IsEmpty())
+		{
+			EEncryptionFailureAction FailureResult = EEncryptionFailureAction::Default;
+
+			if (FNetDelegates::OnReceivedNetworkEncryptionFailure.IsBound())
+			{
+				FailureResult = FNetDelegates::OnReceivedNetworkEncryptionFailure.Execute(ServerConn);
+			}
 
-		EEngineNetworkRuntimeFeatures LocalNetworkFeatures = NetDriver->GetNetworkRuntimeFeatures();
-		FNetControlMessage<NMT_Hello>::Send(NetDriver->ServerConnection, IsLittleEndian, LocalNetworkVersion, EncryptionData.Identifier, LocalNetworkFeatures);
-		
+			const bool bGameplayDisableEncryptionCheck = FailureResult == EEncryptionFailureAction::AllowConnection;
+
+			bEncryptionRequirementsFailure = NetDriver->IsEncryptionRequired() && !bGameplayDisableEncryptionCheck;
+		}
 
-		NetDriver->ServerConnection->FlushNet();
+		if (!bEncryptionRequirementsFailure)
+		{
+			EEngineNetworkRuntimeFeatures LocalNetworkFeatures = NetDriver->GetNetworkRuntimeFeatures();
+			FNetControlMessage<NMT_Hello>::Send(NetDriver->ServerConnection, IsLittleEndian, LocalNetworkVersion, EncryptionData.Identifier, LocalNetworkFeatures);
+
+			ServerConn->FlushNet();
+		}
+		else if (GetConnectionState() == EBeaconConnectionState::Pending)
+		{
+			UE_LOG(LogNet, Error, TEXT("AOnlineBeaconClient::SendInitialJoin: EncryptionToken is empty when 'net.AllowEncryption' requires it."));
+
+			SetConnectionState(EBeaconConnectionState::Invalid);
+			OnFailure();
+		}
 	}
 }
 
diff --git a/Engine/Plugins/Online/OnlineSubsystemUtils/Source/OnlineSubsystemUtils/Private/OnlineBeaconHost.cpp b/Engine/Plugins/Online/OnlineSubsystemUtils/Source/OnlineSubsystemUtils/Private/OnlineBeaconHost.cpp
--- a/Engine/Plugins/Online/OnlineSubsystemUtils/Source/OnlineSubsystemUtils/Private/OnlineBeaconHost.cpp
+++ b/Engine/Plugins/Online/OnlineSubsystemUtils/Source/OnlineSubsystemUtils/Private/OnlineBeaconHost.cpp
@@ -206,7 +206,25 @@ bool AOnlineBeaconHost::HandleControlMessage(UNetConnection* Connection, uint8 M
 			// if the client didn't specify an encryption token we're done with Hello
 			if (EncryptionToken.IsEmpty())
 			{
-				OnHelloSequenceComplete(Connection);
+				EEncryptionFailureAction FailureResult = EEncryptionFailureAction::Default;
+							
+				if (FNetDelegates::OnReceivedNetworkEncryptionFailure.IsBound())
+				{
+					FailureResult = FNetDelegates::OnReceivedNetworkEncryptionFailure.Execute(Connection);
+				}
+
+				const bool bGameplayDisableEncryptionCheck = FailureResult == EEncryptionFailureAction::AllowConnection;
+				const bool bEncryptionRequired = NetDriver != nullptr && NetDriver->IsEncryptionRequired() && !bGameplayDisableEncryptionCheck;
+
+				if (!bEncryptionRequired)
+				{
+					OnHelloSequenceComplete(Connection);
+				}
+				else
+				{
+					SendFailurePacket(Connection, ENetCloseResult::EncryptionTokenMissing, FText::FromString(TEXT("Encryption token missing")));
+					return false;
+				}
 			}
 			else
 			{
@@ -465,11 +483,14 @@ void AOnlineBeaconHost::SendFailurePacket(UNetConnection* Connection, FNetCloseR
 	if (Connection != nullptr)
 	{
 		FString ErrorMsg = ErrorText.ToString();
+		FNetCloseResult CloseReasonCopy = CloseReason;
 
 		UE_LOG(LogBeacon, Log, TEXT("%s: Send failure: %s"), ToCStr(GetDebugName(Connection)), ToCStr(ErrorMsg));
 
-		Connection->SendCloseReason(MoveTemp(CloseReason));
+		Connection->SendCloseReason(MoveTemp(CloseReasonCopy));
 		FNetControlMessage<NMT_Failure>::Send(Connection, ErrorMsg);
+		Connection->FlushNet(true);
+		Connection->Close(MoveTemp(CloseReason));
 	}
 }
 
@@ -478,8 +499,12 @@ void AOnlineBeaconHost::SendFailurePacket(UNetConnection* Connection, const FTex
 	if (Connection != nullptr)
 	{
 		FString ErrorMsg = ErrorText.ToString();
-		UE_LOG(LogBeacon, Log, TEXT("%s: Send failure: %s"), *GetDebugName(Connection), *ErrorMsg);
+
+		UE_LOG(LogBeacon, Log, TEXT("%s: Send failure: %s"), ToCStr(GetDebugName(Connection)), ToCStr(ErrorMsg));
+
 		FNetControlMessage<NMT_Failure>::Send(Connection, ErrorMsg);
+		Connection->FlushNet(true);
+		Connection->Close();
 	}
 }
 
@@ -647,9 +672,8 @@ void AOnlineBeaconHost::OnEncryptionResponse(const FEncryptionKeyResponse& Respo
 
 				Connection->SendCloseReason(ENetCloseResult::EncryptionFailure);
 				FNetControlMessage<NMT_Failure>::Send(Connection, ResponseStr);
-				Connection->FlushNet();
-				// Can't close the connection here since it will leave the failure message in the send buffer and just close the socket. 
-				// Connection->Close(ENetCloseResult::EncryptionFailure);
+				Connection->FlushNet(true);
+				Connection->Close(ENetCloseResult::EncryptionFailure);
 			}
 		}
 		else
diff --git a/Engine/Source/Runtime/Engine/Classes/Engine/DemoNetDriver.h b/Engine/Source/Runtime/Engine/Classes/Engine/DemoNetDriver.h
--- a/Engine/Source/Runtime/Engine/Classes/Engine/DemoNetDriver.h
+++ b/Engine/Source/Runtime/Engine/Classes/Engine/DemoNetDriver.h
@@ -362,6 +362,8 @@ public:
 	virtual void LowLevelSend(TSharedPtr<const FInternetAddr> Address, void* Data, int32 CountBits, FOutPacketTraits& Traits) override {}
 	virtual class ISocketSubsystem* GetSocketSubsystem() override { return nullptr; }
 
+	virtual bool DoesSupportEncryption() const override { return false; }
+
 protected:
 	virtual UChannel* InternalCreateChannelByName(const FName& ChName) override;
 
diff --git a/Engine/Source/Runtime/Engine/Classes/Engine/GameInstance.h b/Engine/Source/Runtime/Engine/Classes/Engine/GameInstance.h
--- a/Engine/Source/Runtime/Engine/Classes/Engine/GameInstance.h
+++ b/Engine/Source/Runtime/Engine/Classes/Engine/GameInstance.h
@@ -555,6 +555,9 @@ public:
 	/** Called when a client receives the EncryptionAck control message from the server, will generally enable encryption. */
 	virtual void ReceivedNetworkEncryptionAck(const FOnEncryptionKeyResponse& Delegate);
 
+	/** Called when a connecting client fails to setup encryption */
+	virtual EEncryptionFailureAction ReceivedNetworkEncryptionFailure(UNetConnection* Connection);
+
 	/** Call to preload any content before loading a map URL, used during seamless travel as well as map loading */
 	virtual void PreloadContentForURL(FURL InURL);
 
diff --git a/Engine/Source/Runtime/Engine/Classes/Engine/NetDriver.h b/Engine/Source/Runtime/Engine/Classes/Engine/NetDriver.h
--- a/Engine/Source/Runtime/Engine/Classes/Engine/NetDriver.h
+++ b/Engine/Source/Runtime/Engine/Classes/Engine/NetDriver.h
@@ -1925,6 +1925,23 @@ public:
 	/** Whether or not this driver has an IsReplay() connection, updated in Add/RemoveClientConnection */
 	bool HasReplayConnection() const { return bHasReplayConnection; }
 
+	/**
+	 * Whether or not this NetDriver supports encryption. Does not signify that encryption is actually enabled, nor setup by the PacketHandler.
+	 *
+	 * @return		Whether or not this NetDriver supports encryption.
+	 */
+	ENGINE_API virtual bool DoesSupportEncryption() const
+	{
+		return true;
+	}
+
+	/**
+	 * Whether or not this NetDriver requires encryption. Does signify that encryption is enabled, but does not signify that it's setup properly.
+	 *
+	 * @return		Whether or not encryption is presently required for connections.
+	 */
+	ENGINE_API virtual bool IsEncryptionRequired() const;
+
 protected:
 	
 	/** Stream of random numbers to be used by this instance of UNetDriver */
diff --git a/Engine/Source/Runtime/Engine/Classes/Engine/NetworkDelegates.h b/Engine/Source/Runtime/Engine/Classes/Engine/NetworkDelegates.h
--- a/Engine/Source/Runtime/Engine/Classes/Engine/NetworkDelegates.h
+++ b/Engine/Source/Runtime/Engine/Classes/Engine/NetworkDelegates.h
@@ -147,6 +147,17 @@ struct FEncryptionKeyResponse
 	}
 };
 
+/** Specifies how to handle encryption failures */
+enum class EEncryptionFailureAction : uint8
+{
+	/** Default handling of encryption failures - net.AllowEncryption determines whether to reject or accept, with exceptions for PIE etc. */
+	Default,
+	/** Reject the connection */
+	RejectConnection,
+	/** Allow the connection */
+	AllowConnection
+};
+
 /** Types of punishment to apply to a cheating client */
 enum class ECheatPunishType : uint8
 {
@@ -250,6 +261,18 @@ public:
 	DECLARE_DELEGATE_OneParam(FReceivedNetworkEncryptionAck, const FOnEncryptionKeyResponse& /*Delegate*/);
 	static FReceivedNetworkEncryptionAck OnReceivedNetworkEncryptionAck;
 
+	/**
+	 * Delegate fired when encryption has failed for a specific connection (client OR server) - allowing gameplay code to override how this is handled
+	 * (e.g. to accept connections without encryption in development environments, and reject them in shipping environments).
+	 *
+	 * Binding to this delegate overrides UGameInstance's handling of encryption (@see ReceivedNetworkEncryptionFailure)
+	 *
+	 * @param Connection	The server or client NetConnection which failed encryption.
+	 * @return				Returns how to handle the encryption failure.
+	 */
+	DECLARE_DELEGATE_RetVal_OneParam(EEncryptionFailureAction, FReceivedNetworkEncryptionFailure, UNetConnection* /* Connection */);
+	static FReceivedNetworkEncryptionFailure OnReceivedNetworkEncryptionFailure;
+
 
 	/**
 	 * Delegate fires whenever a client cheater is detected in the networking code
diff --git a/Engine/Source/Runtime/Engine/Private/GameInstance.cpp b/Engine/Source/Runtime/Engine/Private/GameInstance.cpp
--- a/Engine/Source/Runtime/Engine/Private/GameInstance.cpp
+++ b/Engine/Source/Runtime/Engine/Private/GameInstance.cpp
@@ -111,6 +111,7 @@ void UGameInstance::Init()
 
 		FNetDelegates::OnReceivedNetworkEncryptionToken.BindUObject(this, &ThisClass::ReceivedNetworkEncryptionToken);
 		FNetDelegates::OnReceivedNetworkEncryptionAck.BindUObject(this, &ThisClass::ReceivedNetworkEncryptionAck);
+		FNetDelegates::OnReceivedNetworkEncryptionFailure.BindUObject(this, &ThisClass::ReceivedNetworkEncryptionFailure);
 
 		IPlatformInputDeviceMapper& PlatformInputMapper = IPlatformInputDeviceMapper::Get();
 		PlatformInputMapper.GetOnInputDeviceConnectionChange().AddUObject(this, &UGameInstance::HandleInputDeviceConnectionChange);
@@ -1325,6 +1326,11 @@ void UGameInstance::ReceivedNetworkEncryptionAck(const FOnEncryptionKeyResponse&
 	Delegate.ExecuteIfBound(Response);
 }
 
+EEncryptionFailureAction UGameInstance::ReceivedNetworkEncryptionFailure(UNetConnection* Connection)
+{
+	return EEncryptionFailureAction::Default;
+}
+
 TSubclassOf<UOnlineSession> UGameInstance::GetOnlineSessionClass()
 {
 	return UOnlineSession::StaticClass();
diff --git a/Engine/Source/Runtime/Engine/Private/NetConnection.cpp b/Engine/Source/Runtime/Engine/Private/NetConnection.cpp
--- a/Engine/Source/Runtime/Engine/Private/NetConnection.cpp
+++ b/Engine/Source/Runtime/Engine/Private/NetConnection.cpp
@@ -5180,9 +5180,8 @@ void UNetConnection::SendChallengeControlMessage(const FEncryptionKeyResponse& R
 
 			SendCloseReason(ENetCloseResult::EncryptionFailure);
 			FNetControlMessage<NMT_Failure>::Send(this, ResponseStr);
-			FlushNet();
-			// Can't close the connection here since it will leave the failure message in the send buffer and just close the socket. 
-			// Close(ENetCloseResult::EncryptionFailure);
+			FlushNet(true);
+			Close(ENetCloseResult::EncryptionFailure);
 		}
 	}
 	else
diff --git a/Engine/Source/Runtime/Engine/Private/NetDriver.cpp b/Engine/Source/Runtime/Engine/Private/NetDriver.cpp
--- a/Engine/Source/Runtime/Engine/Private/NetDriver.cpp
+++ b/Engine/Source/Runtime/Engine/Private/NetDriver.cpp
@@ -369,7 +369,10 @@ static TAutoConsoleVariable<int32> CVarUseAdaptiveNetUpdateFrequency(
 TAutoConsoleVariable<int32> CVarNetAllowEncryption(
 	TEXT("net.AllowEncryption"),
 	1,
-	TEXT("If true, the engine will attempt to load an encryption PacketHandler component and fill in the EncryptionToken parameter of the NMT_Hello message based on the ?EncryptionToken= URL option and call callbacks if it's non-empty."));
+	TEXT("If true, the engine will attempt to load an encryption PacketHandler component and fill in the EncryptionToken parameter ")
+	TEXT("of the NMT_Hello message based on the ?EncryptionToken= URL option and call callbacks if it's non-empty. ")
+	TEXT("Additionally, a value of '2' will make the EncryptionToken required - which is enforced serverside. ")
+	TEXT("(0 = Disabled, 1 = Allowed (default), 2 = Required)"));
 
 static TAutoConsoleVariable<int32> CVarActorChannelPool(
 	TEXT("net.ActorChannelPool"),
@@ -396,6 +399,49 @@ static FAutoConsoleVariableRef CVarNetResetAckStatePostSeamlessTravel(
 	TEXT("If 1, the server will reset the ack state of the package map after seamless travel. Increases bandwidth usage, but may resolve some issues with GUIDs not being available on clients after seamlessly traveling."),
 	ECVF_Default);
 
+namespace UE::Net
+{
+	static FString GRequiredEncryptionNetDriverDefNames_Internal = TEXT("all");
+	static TArray<FName> GRequiredEncryptionNetDriverDefNames;
+
+	// Safe to update static here, but not CVar system cached value
+	static void ParseRequiredEncryptionCVars(IConsoleVariable* =nullptr)
+	{
+		GRequiredEncryptionNetDriverDefNames.Reset();
+
+		if (GRequiredEncryptionNetDriverDefNames_Internal.TrimStartAndEnd() == TEXT("all"))
+		{
+			if (GEngine != nullptr)
+			{
+				for (const FNetDriverDefinition& CurDef : GEngine->NetDriverDefinitions)
+				{
+					if (CurDef.DefName != NAME_DemoNetDriver)
+					{
+						GRequiredEncryptionNetDriverDefNames.Add(CurDef.DefName);
+					}
+				}
+			}
+		}
+		else
+		{
+			TArray<FString> NetDriverNames;
+			GRequiredEncryptionNetDriverDefNames_Internal.ParseIntoArray(NetDriverNames, TEXT(","));
+
+			for (const FString& CurName : NetDriverNames)
+			{
+				GRequiredEncryptionNetDriverDefNames.Add(FName(CurName.TrimStartAndEnd()));
+			};
+		}
+	}
+
+	static FAutoConsoleVariableRef CVarNetRequiredEncryptionNetDriverDefNames(
+		TEXT("net.RequiredEncryptionNetDriverDefNames"),
+		GRequiredEncryptionNetDriverDefNames_Internal,
+		TEXT("Comma-delimited list of NetDriverDefinition's where 'IsEncryptionRequired' will return true, when 'net.AllowEncryption' is 2. ")
+		TEXT("(specifying 'all' will enable this for all NetDriverDefinition's)"),
+		FConsoleVariableDelegate::CreateStatic(&ParseRequiredEncryptionCVars));
+}
+
 
 /*-----------------------------------------------------------------------------
 	UNetDriver implementation.
@@ -505,8 +551,8 @@ void UNetDriver::InitPacketSimulationSettings()
 
 	// Read the settings from .ini and command line, with the command line taking precedence	
 	PacketSimulationSettings.ResetSettings();
-	PacketSimulationSettings.LoadConfig(*NetDriverName.ToString());
-	PacketSimulationSettings.ParseSettings(FCommandLine::Get(), *NetDriverName.ToString());
+	PacketSimulationSettings.LoadConfig(*NetDriverDefinition.ToString());
+	PacketSimulationSettings.ParseSettings(FCommandLine::Get(), *NetDriverDefinition.ToString());
 #endif
 }
 
@@ -1431,7 +1477,7 @@ void UNetDriver::InitConnectionlessHandler()
 		if (ConnectionlessHandler.IsValid())
 		{
 			ConnectionlessHandler->NotifyAnalyticsProvider(AnalyticsProvider, AnalyticsAggregator);
-			ConnectionlessHandler->Initialize(Handler::Mode::Server, MAX_PACKET_SIZE, true, nullptr, nullptr, NetDriverName);
+			ConnectionlessHandler->Initialize(Handler::Mode::Server, MAX_PACKET_SIZE, true, nullptr, nullptr, NetDriverDefinition);
 
 			// Add handling for the stateless connect handshake, for connectionless packets, as the outermost layer
 			TSharedPtr<HandlerComponent> NewComponent =
@@ -2388,7 +2434,7 @@ void UNetDriver::SetAnalyticsProvider(TSharedPtr<IAnalyticsProvider> InProvider)
 		}
 		else
 		{
-			AnalyticsAggregator = MakeShareable(new FNetAnalyticsAggregator(InProvider, NetDriverName));
+			AnalyticsAggregator = MakeShareable(new FNetAnalyticsAggregator(InProvider, NetDriverDefinition));
 
 			AnalyticsAggregator->Init();
 		}
@@ -2750,6 +2796,8 @@ bool UNetDriver::HandleNetDisconnectCommand( const TCHAR* Cmd, FOutputDevice& Ar
 
 		ServerConnection->SendCloseReason(ENetCloseResult::Disconnect);
 		FNetControlMessage<NMT_Failure>::Send(ServerConnection, Msg);
+		ServerConnection->FlushNet(true);
+		ServerConnection->Close(ENetCloseResult::Disconnect);
 	}
 #if WITH_SERVER_CODE
 	else
@@ -2767,6 +2815,17 @@ bool UNetDriver::HandleNetDisconnectCommand( const TCHAR* Cmd, FOutputDevice& Ar
 				Connection->FlushNet(true);
 			}
 		}
+
+		// Copy ClientConnections to close each connection, in case closing modifies the ClientConnnection list
+		const TArray<UNetConnection*> ClientConnCopy = ClientConnections;
+
+		for (UNetConnection* Connection : ClientConnCopy)
+		{
+			if (Connection != nullptr)
+			{
+				Connection->Close(ENetCloseResult::Disconnect);
+			}
+		}
 	}
 #endif
 	return true;
@@ -3099,7 +3158,7 @@ bool UNetDriver::Exec( UWorld* InWorld, const TCHAR* Cmd, FOutputDevice& Ar )
 #if DO_ENABLE_NET_TEST
 	// This will allow changing the Pkt* options at runtime
 	else if (bNeverApplyNetworkEmulationSettings == false && 
-			 PacketSimulationSettings.ParseSettings(Cmd, *NetDriverName.ToString()))
+			 PacketSimulationSettings.ParseSettings(Cmd, *NetDriverDefinition.ToString()))
 	{
 		OnPacketSimulationSettingsChanged();
 		return true;
@@ -5292,7 +5351,6 @@ void UNetDriver::ReleaseToChannelPool(UChannel* Channel)
 void UNetDriver::SetNetDriverName(FName NewNetDriverNamed)
 {
 	NetDriverName = NewNetDriverNamed;
-	InitPacketSimulationSettings();
 }
 
 void UNetDriver::SetNetDriverDefinition(FName NewNetDriverDefinition)
@@ -5303,6 +5361,8 @@ void UNetDriver::SetNetDriverDefinition(FName NewNetDriverDefinition)
 	{
 		MaxChannelsOverride = DriverDef->MaxChannelsOverride;
 	}
+
+	InitPacketSimulationSettings();
 }
 
 void UNetDriver::PostCreation(bool bInitializeWithIris)
@@ -6794,6 +6854,26 @@ void UNetDriver::MoveMappedObjectToUnmapped(const UObject* Object)
 	}
 }
 
+bool UNetDriver::IsEncryptionRequired() const
+{
+	using namespace UE::Net;
+
+	if (GRequiredEncryptionNetDriverDefNames.Num() == 0)
+	{
+		ParseRequiredEncryptionCVars();
+	}
+
+	const FWorldContext* const WorldContext = (NetDriverName == NAME_PendingNetDriver ? GEngine->GetWorldContextFromPendingNetGameNetDriver(this) :
+												GEngine->GetWorldContextFromWorld(World));
+	const int32 AllowEncryption = CVarNetAllowEncryption.GetValueOnGameThread();
+	const bool bDevDisableEncryptionCheck = !IsServer() && !(UE_BUILD_SHIPPING || UE_BUILD_TEST); // Clientside only
+	const bool bPIEDisableEncryptionCheck = WorldContext != nullptr && WorldContext->WorldType == EWorldType::PIE;
+	const bool bNetDriverDefNameRequiresEncryption = GRequiredEncryptionNetDriverDefNames.Contains(NetDriverDefinition);
+
+	return DoesSupportEncryption() && bNetDriverDefNameRequiresEncryption && AllowEncryption == 2 && !bDevDisableEncryptionCheck &&
+			!bPIEDisableEncryptionCheck;
+}
+
 void UNetDriver::UpdateNetworkStats()
 {
 	using namespace UE::Net::Private;
diff --git a/Engine/Source/Runtime/Engine/Private/NetworkDelegates.cpp b/Engine/Source/Runtime/Engine/Private/NetworkDelegates.cpp
--- a/Engine/Source/Runtime/Engine/Private/NetworkDelegates.cpp
+++ b/Engine/Source/Runtime/Engine/Private/NetworkDelegates.cpp
@@ -4,6 +4,7 @@
 
 FNetDelegates::FReceivedNetworkEncryptionToken FNetDelegates::OnReceivedNetworkEncryptionToken;
 FNetDelegates::FReceivedNetworkEncryptionAck FNetDelegates::OnReceivedNetworkEncryptionAck;
+FNetDelegates::FReceivedNetworkEncryptionFailure FNetDelegates::OnReceivedNetworkEncryptionFailure;
 FNetDelegates::FOnPendingNetGameConnectionCreated FNetDelegates::OnPendingNetGameConnectionCreated;
 FNetDelegates::FNetworkCheatDetected FNetDelegates::OnNetworkCheatDetected;
 FNetDelegates::FOnSyncLoadDetected FNetDelegates::OnSyncLoadDetected;
\ No newline at end of file
diff --git a/Engine/Source/Runtime/Engine/Private/PendingNetGame.cpp b/Engine/Source/Runtime/Engine/Private/PendingNetGame.cpp
--- a/Engine/Source/Runtime/Engine/Private/PendingNetGame.cpp
+++ b/Engine/Source/Runtime/Engine/Private/PendingNetGame.cpp
@@ -102,21 +102,49 @@ void UPendingNetGame::SendInitialJoin()
 		{
 			uint8 IsLittleEndian = uint8(PLATFORM_LITTLE_ENDIAN);
 			check(IsLittleEndian == !!IsLittleEndian); // should only be one or zero
-			
-			uint32 LocalNetworkVersion = FNetworkVersion::GetLocalNetworkVersion();
-
-			UE_LOG(LogNet, Log, TEXT( "UPendingNetGame::SendInitialJoin: Sending hello. %s" ), *ServerConn->Describe());
 
+			const int32 AllowEncryption = CVarNetAllowEncryption.GetValueOnGameThread();
 			FString EncryptionToken;
-			if (CVarNetAllowEncryption.GetValueOnGameThread() != 0)
+
+			if (AllowEncryption != 0)
 			{
 				EncryptionToken = URL.GetOption(TEXT("EncryptionToken="), TEXT(""));
 			}
 
-			EEngineNetworkRuntimeFeatures LocalNetworkFeatures = NetDriver->GetNetworkRuntimeFeatures();
-			FNetControlMessage<NMT_Hello>::Send(ServerConn, IsLittleEndian, LocalNetworkVersion, EncryptionToken, LocalNetworkFeatures);
+			bool bEncryptionRequirementsFailure = false;
+
+			if (EncryptionToken.IsEmpty())
+			{
+				EEncryptionFailureAction FailureResult = EEncryptionFailureAction::Default;
+
+				if (FNetDelegates::OnReceivedNetworkEncryptionFailure.IsBound())
+				{
+					FailureResult = FNetDelegates::OnReceivedNetworkEncryptionFailure.Execute(ServerConn);
+				}
+
+				const bool bGameplayDisableEncryptionCheck = FailureResult == EEncryptionFailureAction::AllowConnection;
+
+				bEncryptionRequirementsFailure = NetDriver->IsEncryptionRequired() && !bGameplayDisableEncryptionCheck;
+			}
 			
-			ServerConn->FlushNet();
+			if (!bEncryptionRequirementsFailure)
+			{
+				uint32 LocalNetworkVersion = FNetworkVersion::GetLocalNetworkVersion();
+
+				UE_LOG(LogNet, Log, TEXT("UPendingNetGame::SendInitialJoin: Sending hello. %s"), *ServerConn->Describe());
+
+				EEngineNetworkRuntimeFeatures LocalNetworkFeatures = NetDriver->GetNetworkRuntimeFeatures();
+				FNetControlMessage<NMT_Hello>::Send(ServerConn, IsLittleEndian, LocalNetworkVersion, EncryptionToken, LocalNetworkFeatures);
+
+
+				ServerConn->FlushNet();
+			}
+			else
+			{
+				UE_LOG(LogNet, Error, TEXT("UPendingNetGame::SendInitialJoin: EncryptionToken is empty when 'net.AllowEncryption' requires it."));
+
+				ConnectionError = TEXT("EncryptionToken not set.");
+			}
 		}
 	}
 }
diff --git a/Engine/Source/Runtime/Engine/Private/World.cpp b/Engine/Source/Runtime/Engine/Private/World.cpp
--- a/Engine/Source/Runtime/Engine/Private/World.cpp
+++ b/Engine/Source/Runtime/Engine/Private/World.cpp
@@ -5897,13 +5897,38 @@ void UWorld::NotifyControlMessage(UNetConnection* Connection, uint8 MessageType,
 					{
 						if (EncryptionToken.IsEmpty())
 						{
-							Connection->SendChallengeControlMessage();
+							EEncryptionFailureAction FailureResult = EEncryptionFailureAction::Default;
+							
+							if (FNetDelegates::OnReceivedNetworkEncryptionFailure.IsBound())
+							{
+								FailureResult = FNetDelegates::OnReceivedNetworkEncryptionFailure.Execute(Connection);
+							}
+
+							const bool bGameplayDisableEncryptionCheck = FailureResult == EEncryptionFailureAction::AllowConnection;
+							const bool bEncryptionRequired = NetDriver->IsEncryptionRequired() && !bGameplayDisableEncryptionCheck;
+
+							if (!bEncryptionRequired)
+							{
+								Connection->SendChallengeControlMessage();
+							}
+							else
+							{
+								FString FailureMsg(TEXT("Encryption token missing"));
+
+								UE_LOG(LogNet, Warning, TEXT("%s: No EncryptionToken specified, disconnecting."), ToCStr(Connection->GetName()));
+
+								Connection->SendCloseReason(ENetCloseResult::EncryptionTokenMissing);
+								FNetControlMessage<NMT_Failure>::Send(Connection, FailureMsg);
+								Connection->FlushNet(true);
+								Connection->Close(ENetCloseResult::EncryptionTokenMissing);
+							}
 						}
 						else
 						{
 							if (FNetDelegates::OnReceivedNetworkEncryptionToken.IsBound())
 							{
-								FNetDelegates::OnReceivedNetworkEncryptionToken.Execute(EncryptionToken, FOnEncryptionKeyResponse::CreateUObject(Connection, &UNetConnection::SendChallengeControlMessage));
+								FNetDelegates::OnReceivedNetworkEncryptionToken.Execute(EncryptionToken,
+									FOnEncryptionKeyResponse::CreateUObject(Connection, &UNetConnection::SendChallengeControlMessage));
 							}
 							else
 							{
@@ -5915,6 +5940,7 @@ void UWorld::NotifyControlMessage(UNetConnection* Connection, uint8 MessageType,
 								Connection->SendCloseReason(ENetCloseResult::EncryptionFailure);
 								FNetControlMessage<NMT_Failure>::Send(Connection, FailureMsg);
 								Connection->FlushNet(true);
+								Connection->Close(ENetCloseResult::EncryptionFailure);
 							}
 						}
 					}
@@ -6017,9 +6043,7 @@ void UWorld::NotifyControlMessage(UNetConnection* Connection, uint8 MessageType,
 						Connection->SendCloseReason(ENetCloseResult::PreLoginFailure);
 						FNetControlMessage<NMT_Failure>::Send(Connection, ErrorMsg);
 						Connection->FlushNet(true);
-						//@todo sz - can't close the connection here since it will leave the failure message 
-						// in the send buffer and just close the socket. 
-						//Connection->Close(ENetCloseResult::PreLoginFailure);
+						Connection->Close(ENetCloseResult::PreLoginFailure);
 					}
 					else
 					{
@@ -6061,9 +6085,7 @@ void UWorld::NotifyControlMessage(UNetConnection* Connection, uint8 MessageType,
 						Connection->SendCloseReason(ENetCloseResult::JoinFailure);
 						FNetControlMessage<NMT_Failure>::Send(Connection, ErrorMsg);
 						Connection->FlushNet(true);
-						//@todo sz - can't close the connection here since it will leave the failure message 
-						// in the send buffer and just close the socket. 
-						//Connection->Close(ENetCloseResult::JoinFailure);
+						Connection->Close(ENetCloseResult::JoinFailure);
 					}
 					else
 					{
@@ -6161,9 +6183,7 @@ void UWorld::NotifyControlMessage(UNetConnection* Connection, uint8 MessageType,
 						Connection->SendCloseReason(ENetCloseResult::PreLoginFailure);
 						FNetControlMessage<NMT_Failure>::Send(Connection, ErrorMsg);
 						Connection->FlushNet(true);
-						//@todo sz - can't close the connection here since it will leave the failure message 
-						// in the send buffer and just close the socket. 
-						//Connection->Close(ENetCloseResult::PreLoginFailure);
+						Connection->Close(ENetCloseResult::PreLoginFailure);
 					}
 					else
 					{
@@ -6198,9 +6218,7 @@ void UWorld::NotifyControlMessage(UNetConnection* Connection, uint8 MessageType,
 							Connection->SendCloseReason(ENetCloseResult::JoinSplitFailure);
 							FNetControlMessage<NMT_Failure>::Send(Connection, ErrorMsg);
 							Connection->FlushNet(true);
-							//@todo sz - can't close the connection here since it will leave the failure message 
-							// in the send buffer and just close the socket. 
-							//Connection->Close(ENetCloseResult::JoinSplitFailure);
+							Connection->Close(ENetCloseResult::JoinSplitFailure);
 						}
 						else
 						{
diff --git a/Engine/Source/Runtime/Net/Core/Public/Net/Core/Connection/NetCloseResult.h b/Engine/Source/Runtime/Net/Core/Public/Net/Core/Connection/NetCloseResult.h
--- a/Engine/Source/Runtime/Net/Core/Public/Net/Core/Connection/NetCloseResult.h
+++ b/Engine/Source/Runtime/Net/Core/Public/Net/Core/Connection/NetCloseResult.h
@@ -186,6 +186,9 @@ enum class ENetCloseResult : uint32
 	/** Encryption failure */
 	EncryptionFailure,
 
+	/** EncryptionToken is missing */
+	EncryptionTokenMissing,
+
 
 	/** Channel Error result types */
 
