Added option LevelStreaming.ShouldReuseUnloadedButStillAroundLevels to allow disabling reuse of unloaded levels (defaulted to true to keep actual behavior)

diff --git a/Engine/Source/Runtime/Engine/Classes/Engine/LevelStreaming.h b/Engine/Source/Runtime/Engine/Classes/Engine/LevelStreaming.h
--- a/Engine/Source/Runtime/Engine/Classes/Engine/LevelStreaming.h
+++ b/Engine/Source/Runtime/Engine/Classes/Engine/LevelStreaming.h
@@ -560,6 +560,9 @@ public:
 	/** If true server will wait for client acknowledgment before making treating streaming levels as visible for the client */
 	static bool ShouldServerUseMakingVisibleTransactionRequest();
 
+	/** If true level streaming can reuse an unloaded level that wasn't GC'd yet. */
+	static bool ShouldReuseUnloadedButStillAroundLevels(const ULevel* InLevel);
+
 	/** 
 	 * Traverses all streaming level objects in the persistent world and in all inner worlds and calls appropriate delegate for streaming objects that refer specified level 
 	 *
diff --git a/Engine/Source/Runtime/Engine/Classes/Engine/LevelStreamingGCHelper.h b/Engine/Source/Runtime/Engine/Classes/Engine/LevelStreamingGCHelper.h
--- a/Engine/Source/Runtime/Engine/Classes/Engine/LevelStreamingGCHelper.h
+++ b/Engine/Source/Runtime/Engine/Classes/Engine/LevelStreamingGCHelper.h
@@ -3,6 +3,7 @@
 #pragma once
 
 #include "CoreMinimal.h"
+#include "Engine/EngineBaseTypes.h"
 
 class ULevel;
 class UWorld;
@@ -34,17 +35,6 @@ struct ENGINE_API FLevelStreamingGCHelper
 	 */
 	static void CancelUnloadRequest( ULevel* InLevel );
 
-	/** 
-	 * Prepares levels that are marked for unload for the GC call by marking their actors and components as
-	 * pending kill.
-	 */
-	static void PrepareStreamedOutLevelsForGC();
-
-	/**
-	 * Verify that the level packages are no longer around.
-	 */
-	static void VerifyLevelsGotRemovedByGC();
-	
 	/**
 	 * @return	The number of levels pending a purge by the garbage collector
 	 */
@@ -56,10 +46,33 @@ struct ENGINE_API FLevelStreamingGCHelper
 	static void EnableForCommandlet();
 	
 private:
+	/** Prepares levels that are marked for unload for the next GC call by marking their actors and components as garbage. */
+	static void PrepareStreamedOutLevelsForGC();
+
+	/** Verify that the level packages are no longer around. */
+	static void VerifyLevelsGotRemovedByGC();
+
+	/** Called before garbage collect. */
+	static void OnPreGarbageCollect();
+
+	/** Called at the end of a world tick. */
+	static void OnWorldTickEnd(UWorld* InWorld, ELevelTick InTickType, float InDeltaSeconds);
+
+	/** Prepares a level that is marked for unload for the next GC call by marking its actors and components as garbage.  */
+	static void PrepareStreamedOutLevelForGC(ULevel* Level);
+
 	/** Static array of levels that should be unloaded */
 	static TArray<TWeakObjectPtr<ULevel> > LevelsPendingUnload;
-	/** Static array of level packages that have been marked by PrepareStreamedOutLevelsForGC */
-	static TArray<FName> LevelPackageNames;
+
+	/** Static set of level packages that have been marked by PrepareStreamedOutLevelsForGC */
+	static TSet<FName> LevelPackageNames;
+
 	/** Static bool allows FLevelStreamingGCHelper to be used in a commandlet */
 	static bool bEnabledForCommandlet;
+
+	/** Whether RequestUnload delayed its call to PrepareStreamedOutLevelForGC after world tick. */
+	static bool bIsPrepareStreamedOutLevelForGCDelayedToWorldTickEnd;
+
+	/** The number of unloaded levels prepared for the next GC (used by GetNumLevelsPendingPurge). */
+	static int32 NumberOfPreparedStreamedOutLevelsForGC;
 };
\ No newline at end of file
diff --git a/Engine/Source/Runtime/Engine/Classes/Engine/World.h b/Engine/Source/Runtime/Engine/Classes/Engine/World.h
--- a/Engine/Source/Runtime/Engine/Classes/Engine/World.h
+++ b/Engine/Source/Runtime/Engine/Classes/Engine/World.h
@@ -4170,6 +4170,9 @@ public:
 	DECLARE_MULTICAST_DELEGATE_ThreeParams(FOnWorldTickStart, UWorld*, ELevelTick, float);
 	static FOnWorldTickStart OnWorldTickStart;
 
+	DECLARE_MULTICAST_DELEGATE_ThreeParams(FOnWorldTickEnd, UWorld*, ELevelTick, float);
+	static FOnWorldTickEnd OnWorldTickEnd;
+
 	// Delegate called before actors are ticked for each world. Delta seconds is already dilated and clamped.
 	DECLARE_MULTICAST_DELEGATE_ThreeParams(FOnWorldPreActorTick, UWorld* /*World*/, ELevelTick/**Tick Type*/, float/**Delta Seconds*/);
 	static FOnWorldPreActorTick OnWorldPreActorTick;
diff --git a/Engine/Source/Runtime/Engine/Private/Level.cpp b/Engine/Source/Runtime/Engine/Private/Level.cpp
--- a/Engine/Source/Runtime/Engine/Private/Level.cpp
+++ b/Engine/Source/Runtime/Engine/Private/Level.cpp
@@ -496,22 +496,46 @@ void ULevel::CleanupLevel(bool bCleanupResources, bool bUnloadFromEditor)
 {
 	OnCleanupLevel.Broadcast();
 
+	const bool bTrashPackage = !ULevelStreaming::ShouldReuseUnloadedButStillAroundLevels(this);
+	TSet<UPackage*> ProcessedPackages;
+	auto ProcessPackage = [&ProcessedPackages, bTrashPackage](UPackage* InPackage, bool bInClearStandaloneFlag = false)
+	{
+		bool bWasAlreadyInSet;
+		ProcessedPackages.Add(InPackage, &bWasAlreadyInSet);
+		if (!bWasAlreadyInSet)
+		{
+#if WITH_EDITOR
+			if (bInClearStandaloneFlag || bTrashPackage)
+			{
+				// Clear RF_Standalone flag on objects in package
+				ForEachObjectWithPackage(InPackage, [](UObject* Object) { Object->ClearFlags(RF_Standalone); return true; }, false);
+			}
+#endif
+			if (bTrashPackage)
+			{
+				// Rename package to make sure it won't be reused
+				FName NewPackageName = MakeUniqueObjectName(nullptr, UPackage::StaticClass(), FName(*FString::Printf(TEXT("%s_Trashed"), *InPackage->GetName())));
+				InPackage->Rename(*NewPackageName.ToString(), nullptr, REN_ForceNoResetLoaders | REN_DontCreateRedirectors | REN_NonTransactional | REN_DoNotDirty);
+			}
+		}
+	};
+
+	UPackage* LevelPackage = GetPackage();
+	if (bTrashPackage)
+	{
+		ProcessPackage(LevelPackage);
+	}
+
 #if WITH_EDITOR
-	if (bUnloadFromEditor)
+	// Process objects outered to this level but in a different package (currently only possible in editor)
+	check(!bUnloadFromEditor || bCleanupResources);
+	if (bUnloadFromEditor || bTrashPackage)
 	{
-		check(bCleanupResources);
-		// if the level contains any actor with an external package, clear their metadata standalone flag so that the packages can be properly unloaded.
-		// Do so for any actors outered to the level and not just from the level actors array
-		ForEachObjectWithOuter(this, [](UObject* InObject)
+		ForEachObjectWithOuter(this, [LevelPackage, ProcessPackage](UObject* InObject)
 		{
-			// Ask directly on all objects since the tests validate against object flags before doing a hash lookup
-			if (UPackage* ExternalPackage = InObject->GetExternalPackage())
+			if (UPackage* ObjectPackage = InObject->GetPackage(); ObjectPackage && LevelPackage != ObjectPackage)
 			{
-				ForEachObjectWithPackage(ExternalPackage, [](UObject* Object)
-				{
-					Object->ClearFlags(RF_Standalone);
-					return true;
-				}, false);
+				ProcessPackage(ObjectPackage, true);
 			}
 		}, false);
 	}
diff --git a/Engine/Source/Runtime/Engine/Private/LevelStreaming.cpp b/Engine/Source/Runtime/Engine/Private/LevelStreaming.cpp
--- a/Engine/Source/Runtime/Engine/Private/LevelStreaming.cpp
+++ b/Engine/Source/Runtime/Engine/Private/LevelStreaming.cpp
@@ -80,6 +80,14 @@ namespace LevelStreamingCVars
 		TEXT("Whether server should wait for client to acknowledge visibility update before treating streaming levels as visible by the client.\n")
 		TEXT("0: Disable, 1: Enable"),
 		ECVF_Default);
+
+	static bool bShouldReuseUnloadedButStillAroundLevels = true;
+	FAutoConsoleVariableRef CVarShouldReuseUnloadedButStillAroundLevels(
+		TEXT("LevelStreaming.ShouldReuseUnloadedButStillAroundLevels"),
+		bShouldReuseUnloadedButStillAroundLevels,
+		TEXT("Whether level streaming will reuse the unloaded levels that aren't GC'd yet.\n")
+		TEXT("0: Disable, 1: Enable"),
+		ECVF_ReadOnly);
 }
 
 bool ULevelStreaming::DefaultAllowClientUseMakingInvisibleTransactionRequests()
@@ -121,6 +129,16 @@ bool ULevelStreaming::ShouldServerUseMakingVisibleTransactionRequest()
 	return LevelStreamingCVars::bShouldServerUseMakingVisibleTransactionRequest;
 }
 
+bool ULevelStreaming::ShouldReuseUnloadedButStillAroundLevels(const ULevel* InLevel)
+{
+	UWorld* OuterWorld = InLevel ? InLevel->GetTypedOuter<UWorld>() : nullptr;
+	if (OuterWorld && OuterWorld->IsGameWorld() && !LevelStreamingCVars::bShouldReuseUnloadedButStillAroundLevels)
+	{
+		return false;
+	}
+	return true;
+}
+
 int32 ULevelStreamingDynamic::UniqueLevelInstanceId = 0;
 
 /**
@@ -1481,6 +1499,7 @@ bool ULevelStreaming::RequestLevel(UWorld* PersistentWorld, bool bAllowLevelLoad
 				return false;
 			}
 #endif
+			check(ULevelStreaming::ShouldReuseUnloadedButStillAroundLevels(World->PersistentLevel));
 			if (World->PersistentLevel != LoadedLevel)
 			{
 				// Level already exists but may have the wrong type due to being inactive before, so copy data over
diff --git a/Engine/Source/Runtime/Engine/Private/LevelStreamingGCHelper.cpp b/Engine/Source/Runtime/Engine/Private/LevelStreamingGCHelper.cpp
--- a/Engine/Source/Runtime/Engine/Private/LevelStreamingGCHelper.cpp
+++ b/Engine/Source/Runtime/Engine/Private/LevelStreamingGCHelper.cpp
@@ -16,8 +16,10 @@
 /************************************************************************/
 FLevelStreamingGCHelper::FOnGCStreamedOutLevelsEvent FLevelStreamingGCHelper::OnGCStreamedOutLevels;
 TArray<TWeakObjectPtr<ULevel> > FLevelStreamingGCHelper::LevelsPendingUnload;
-TArray<FName> FLevelStreamingGCHelper::LevelPackageNames;
+TSet<FName> FLevelStreamingGCHelper::LevelPackageNames;
 bool FLevelStreamingGCHelper::bEnabledForCommandlet = false;
+bool FLevelStreamingGCHelper::bIsPrepareStreamedOutLevelForGCDelayedToWorldTickEnd = false;
+int32 FLevelStreamingGCHelper::NumberOfPreparedStreamedOutLevelsForGC = 0;
 
 void FLevelStreamingGCHelper::AddGarbageCollectorCallback()
 {
@@ -25,8 +27,9 @@ void FLevelStreamingGCHelper::AddGarbageCollectorCallback()
 	static bool GarbageCollectAdded = false;
 	if ( GarbageCollectAdded == false )
 	{
-		FCoreUObjectDelegates::GetPreGarbageCollectDelegate().AddStatic( FLevelStreamingGCHelper::PrepareStreamedOutLevelsForGC );
-		FCoreUObjectDelegates::GetPostGarbageCollect().AddStatic( FLevelStreamingGCHelper::VerifyLevelsGotRemovedByGC );
+		FCoreUObjectDelegates::GetPreGarbageCollectDelegate().AddStatic(FLevelStreamingGCHelper::OnPreGarbageCollect);
+		FCoreUObjectDelegates::GetPostGarbageCollect().AddStatic(FLevelStreamingGCHelper::VerifyLevelsGotRemovedByGC);
+		FWorldDelegates::OnWorldTickEnd.AddStatic(FLevelStreamingGCHelper::OnWorldTickEnd);
 		GarbageCollectAdded = true;
 	}
 }
@@ -39,11 +42,43 @@ void FLevelStreamingGCHelper::EnableForCommandlet()
 
 void FLevelStreamingGCHelper::RequestUnload( ULevel* InLevel )
 {
+	auto IsInsideWorldTick = []()
+	{
+		for (const FWorldContext& Context : GEngine->GetWorldContexts())
+		{
+			UWorld* World = Context.World();
+			if (World && World->bInTick)
+			{
+				return true;
+			}
+		}
+		return false;
+	};
+
 	if (!IsRunningCommandlet() || bEnabledForCommandlet)
 	{
 		check( InLevel );
 		check( InLevel->bIsVisible == false );
-		LevelsPendingUnload.AddUnique( InLevel );
+		
+		if (!ULevelStreaming::ShouldReuseUnloadedButStillAroundLevels(InLevel))
+		{
+			if (!IsInsideWorldTick())
+			{
+				// Prepare unloaded level
+				PrepareStreamedOutLevelForGC(InLevel);
+			}
+			else
+			{
+				// Since PrepareStreamedOutLevelForGC cannot be called inside a world tick, delay it to the next OnWorldTickEnd
+				// TODO: Investigate why NotifyStreamingLevelUnload can't be called inside world's tick
+				bIsPrepareStreamedOutLevelForGCDelayedToWorldTickEnd = true;
+				LevelsPendingUnload.AddUnique(InLevel);
+			}
+		}
+		else
+		{
+			LevelsPendingUnload.AddUnique(InLevel);
+		}
 	}
 }
 
@@ -52,92 +87,116 @@ void FLevelStreamingGCHelper::CancelUnloadRequest( ULevel* InLevel )
 	LevelsPendingUnload.Remove( InLevel );
 }
 
-void FLevelStreamingGCHelper::PrepareStreamedOutLevelsForGC()
+void FLevelStreamingGCHelper::OnWorldTickEnd(UWorld* InWorld, ELevelTick InTickType, float InDeltaSeconds)
 {
-	TRACE_CPUPROFILER_EVENT_SCOPE(FLevelStreamingGCHelper::PrepareStreamedOutLevelsForGC);
-	if (LevelsPendingUnload.Num() > 0)
+	if (bIsPrepareStreamedOutLevelForGCDelayedToWorldTickEnd)
+	{
+		TRACE_CPUPROFILER_EVENT_SCOPE(FLevelStreamingGCHelper::OnWorldTickEnd);
+		bIsPrepareStreamedOutLevelForGCDelayedToWorldTickEnd = false;
+		PrepareStreamedOutLevelsForGC();
+	}
+}
+
+void FLevelStreamingGCHelper::OnPreGarbageCollect()
+{
+	TRACE_CPUPROFILER_EVENT_SCOPE(FLevelStreamingGCHelper::OnPreGarbageCollect);
+	if (GetNumLevelsPendingPurge() > 0)
 	{
 		OnGCStreamedOutLevels.Broadcast();
 	}
-	
+
+	PrepareStreamedOutLevelsForGC();
+}
+
+void FLevelStreamingGCHelper::PrepareStreamedOutLevelsForGC()
+{
+	TRACE_CPUPROFILER_EVENT_SCOPE(FLevelStreamingGCHelper::PrepareStreamedOutLevelsForGC);
+
 	// Iterate over all level objects that want to be unloaded.
 	for( int32 LevelIndex=0; LevelIndex < LevelsPendingUnload.Num(); LevelIndex++ )
 	{
-		ULevel*	Level = LevelsPendingUnload[LevelIndex].Get();
+		ULevel* Level = LevelsPendingUnload[LevelIndex].Get();
+		PrepareStreamedOutLevelForGC(Level);
+	}
+
+	LevelsPendingUnload.Empty();
+	NumberOfPreparedStreamedOutLevelsForGC = 0;
+}
+
+void FLevelStreamingGCHelper::PrepareStreamedOutLevelForGC(ULevel* InLevel)
+{
+	if (InLevel && ((!GIsEditor || bEnabledForCommandlet) || InLevel->GetPackage()->HasAnyPackageFlags(PKG_PlayInEditor)))
+	{
+		++NumberOfPreparedStreamedOutLevelsForGC;
+
+		UPackage* LevelPackage = InLevel->GetOutermost();
+		UE_LOG(LogStreaming, Log, TEXT("PrepareStreamedOutLevelsForGC called on '%s'"), *LevelPackage->GetName());
 
-		if( Level && ((!GIsEditor || bEnabledForCommandlet) || Level->GetOutermost()->HasAnyPackageFlags(PKG_PlayInEditor)))
+		for (const FWorldContext& Context : GEngine->GetWorldContexts())
 		{
-			UPackage* LevelPackage = Level->GetOutermost();
-			UE_LOG(LogStreaming, Log, TEXT("PrepareStreamedOutLevelsForGC called on '%s'"), *LevelPackage->GetName() );
-			
-			for (const FWorldContext& Context : GEngine->GetWorldContexts())
+			UWorld* World = Context.World();
+			if (World)
 			{
-				UWorld* World = Context.World();
-				if (World)
-				{
-					// This can never ever be called during tick; same goes for GC in general.
-					check( !World->bInTick );
+				// This can never ever be called during tick; same goes for GC in general.
+				check(!World->bInTick);
 
-					FWorldContext& MutableContext = GEngine->GetWorldContextFromWorldChecked(World);
-					for (FNamedNetDriver& Driver : MutableContext.ActiveNetDrivers)
+				FWorldContext& MutableContext = GEngine->GetWorldContextFromWorldChecked(World);
+				for (FNamedNetDriver& Driver : MutableContext.ActiveNetDrivers)
+				{
+					if (Driver.NetDriver != nullptr)
 					{
-						if (Driver.NetDriver != nullptr)
-						{
-							// The net driver must remove this level and its actors from the packagemap, or else
-							// the client package map will keep hard refs to them and prevent GC
-							Driver.NetDriver->NotifyStreamingLevelUnload(Level);
-						}
+						// The net driver must remove this level and its actors from the packagemap, or else
+						// the client package map will keep hard refs to them and prevent GC
+						Driver.NetDriver->NotifyStreamingLevelUnload(InLevel);
 					}
-
-					// Broadcast level unloaded event to blueprints through level streaming objects
-					ULevelStreaming::BroadcastLevelLoadedStatus(World, LevelPackage->GetFName(), false);
 				}
+
+				// Broadcast level unloaded event to blueprints through level streaming objects
+				ULevelStreaming::BroadcastLevelLoadedStatus(World, LevelPackage->GetFName(), false);
 			}
+		}
 
-			// Make sure that this package has been unloaded after GC pass.
-			LevelPackageNames.Add( LevelPackage->GetFName() );
+		InLevel->CleanupLevel();
 
-			Level->CleanupLevel();
+		// Make sure that this package has been unloaded after GC pass.
+		LevelPackageNames.Add(InLevel->GetPackage()->GetFName());
 
-			// Mark world and all other package subobjects as pending kill
-			// This will destroy metadata objects and any other objects left behind
-			TSet<UPackage*> Packages;
-			ForEachObjectWithOuter(Level->GetOutermostObject(), [&Packages](UObject* Object)
+		// Mark world and all other package subobjects as pending kill
+		// This will destroy metadata objects and any other objects left behind
+		TSet<UPackage*> Packages;
+		ForEachObjectWithOuter(InLevel->GetOutermostObject(), [&Packages](UObject* Object)
+		{
+			bool bIsAlreadyInSet = false;
+			UPackage* Package = Object->GetPackage();
+			Packages.Add(Package, &bIsAlreadyInSet);
+			if (!bIsAlreadyInSet)
 			{
-				bool bIsAlreadyInSet = false;
-				UPackage* Package = Object->GetPackage();
-				Packages.Add(Package, &bIsAlreadyInSet);
-				if (!bIsAlreadyInSet)
-				{
-					ForEachObjectWithPackage(Package, [](UObject* PackageObject) { PackageObject->MarkAsGarbage(); return true; }, true, RF_NoFlags, EInternalObjectFlags::Garbage);
-					Package->MarkAsGarbage();
-				}
-			}, true, RF_NoFlags, EInternalObjectFlags::Garbage);
+				ForEachObjectWithPackage(Package, [](UObject* PackageObject) { PackageObject->MarkAsGarbage(); return true; }, true, RF_NoFlags, EInternalObjectFlags::Garbage);
+				Package->MarkAsGarbage();
+			}
+		}, true, RF_NoFlags, EInternalObjectFlags::Garbage);
 
-			if (!UObjectBaseUtility::IsPendingKillEnabled())
+		if (!UObjectBaseUtility::IsPendingKillEnabled())
+		{
+			// Rename the packages that we are streaming out so that we can possibly reload another copy of them
+			for (UPackage* Package : Packages)
 			{
-				// Rename the packages that we are streaming out so that we can possibly reload another copy of them
-				for (UPackage* Package : Packages)
-				{
-					FCoreUObjectInternalDelegates::GetOnLeakedPackageRenameDelegate().Broadcast(Package);
-					FName NewName = MakeUniqueObjectName(nullptr, UPackage::StaticClass(), Package->GetFName());
-					Package->Rename(*NewName.ToString(), nullptr, REN_ForceNoResetLoaders | REN_DontCreateRedirectors | REN_NonTransactional);
-				}
+				FCoreUObjectInternalDelegates::GetOnLeakedPackageRenameDelegate().Broadcast(Package);
+				FName NewName = MakeUniqueObjectName(nullptr, UPackage::StaticClass(), Package->GetFName());
+				Package->Rename(*NewName.ToString(), nullptr, REN_ForceNoResetLoaders | REN_DontCreateRedirectors | REN_NonTransactional);
+			}
 
 #if !WITH_EDITOR
-				// Clear the level actors array to maximize the memory savings from GC if there are outstanding references to some actors.
-				// Don't do this outside the editor now til we can validate it works properly with external packages.
-				Level->Actors.Empty();
-				Level->ActorsForGC.Empty();
-				Level->ActorCluster = nullptr; // The actor cluster has been marked garbage so will be dissolved, also drop a reference to it here since it has an internal array of pointers
+			// Clear the level actors array to maximize the memory savings from GC if there are outstanding references to some actors.
+			// Don't do this outside the editor now til we can validate it works properly with external packages.
+			InLevel->Actors.Empty();
+			InLevel->ActorsForGC.Empty();
+			InLevel->ActorCluster = nullptr; // The actor cluster has been marked garbage so will be dissolved, also drop a reference to it here since it has an internal array of pointers
 #endif
-			}
-
-			Level->CleanupReferences();
 		}
-	}
 
-	LevelsPendingUnload.Empty();
+		InLevel->CleanupReferences();
+	}
 }
 
 void FLevelStreamingGCHelper::VerifyLevelsGotRemovedByGC()
@@ -160,19 +219,19 @@ void FLevelStreamingGCHelper::VerifyLevelsGotRemovedByGC()
 			{
 				UObject* Object = *It;
 				// Check whether object's outermost is in the list.
-				if (LevelPackageNames.Find(Object->GetOutermost()->GetFName()) != INDEX_NONE
+				if (LevelPackageNames.Contains(Object->GetOutermost()->GetFName())
 					// But disregard package object itself.
 					&& !Object->IsA(UPackage::StaticClass()))
 				{
-					if (bIsAsyncLoading && Object->HasAnyInternalFlags(EInternalObjectFlags::Async|EInternalObjectFlags::AsyncLoading))
+					if (bIsAsyncLoading && Object->HasAnyInternalFlags(EInternalObjectFlags::Async | EInternalObjectFlags::AsyncLoading))
 					{
-						UE_LOG(LogStreaming, Display, TEXT("Level object %s isn't released by async loading yet, "
-								"it will get garbage collected next time instead."),
-								*Object->GetFullName());
+						UE_LOG(LogGarbage, Display, TEXT("Level object %s isn't released by async loading yet, "
+							"it will get garbage collected next time instead."),
+							*Object->GetFullName());
 					}
 					else
 					{
-						UE_LOG(LogStreaming, Warning, TEXT("Level object %s didn't get garbage collected! Trying to find culprit, though this might crash. Try increasing stack size if it does. Referenced by:"), *Object->GetFullName());
+						UE_LOG(LogGarbage, Warning, TEXT("Level object %s didn't get garbage collected! Trying to find culprit, though this might crash. Try increasing stack size if it does. Referenced by:"), *Object->GetFullName());
 						FReferenceChainSearch RefChainSearch(Object, EReferenceChainSearchMode::Shortest | EReferenceChainSearchMode::PrintResults);
 						FailCount++;
 					}
@@ -180,7 +239,7 @@ void FLevelStreamingGCHelper::VerifyLevelsGotRemovedByGC()
 			}
 			if (FailCount > 0)
 			{
-				UE_LOG(LogStreaming, Error, TEXT("Streamed out levels were not completely garbage collected! Please see previous log entries."));
+				UE_LOG(LogGarbage, Error, TEXT("Streamed out levels were not completely garbage collected! Please see previous log entries."));
 				// If not forcing verify, consider errors as fatal
 				check(GLevelStreamingForceVerifyLevelsGotRemovedByGC);
 			}
@@ -192,5 +251,5 @@ void FLevelStreamingGCHelper::VerifyLevelsGotRemovedByGC()
 
 int32 FLevelStreamingGCHelper::GetNumLevelsPendingPurge()
 {
-	return LevelsPendingUnload.Num();
+	return LevelsPendingUnload.Num() + NumberOfPreparedStreamedOutLevelsForGC;
 }
\ No newline at end of file
diff --git a/Engine/Source/Runtime/Engine/Private/LevelTick.cpp b/Engine/Source/Runtime/Engine/Private/LevelTick.cpp
--- a/Engine/Source/Runtime/Engine/Private/LevelTick.cpp
+++ b/Engine/Source/Runtime/Engine/Private/LevelTick.cpp
@@ -1810,6 +1810,8 @@ void UWorld::Tick( ELevelTick TickType, float DeltaSeconds )
 				WorldParam->PerfTrackers->GetInGamePerformanceTracker((EInGamePerfTrackers)Tracker, EInGamePerfTrackerThreads::RenderThread).Tick();
 			}
 		});
+
+	FWorldDelegates::OnWorldTickEnd.Broadcast(this, TickType, DeltaSeconds);
 }
 
 void UWorld::CleanupActors()
diff --git a/Engine/Source/Runtime/Engine/Private/World.cpp b/Engine/Source/Runtime/Engine/Private/World.cpp
--- a/Engine/Source/Runtime/Engine/Private/World.cpp
+++ b/Engine/Source/Runtime/Engine/Private/World.cpp
@@ -375,6 +375,7 @@ FWorldDelegates::FOnLevelChanged FWorldDelegates::LevelRemovedFromWorld;
 FWorldDelegates::FLevelOffsetEvent FWorldDelegates::PostApplyLevelOffset;
 FWorldDelegates::FWorldGetAssetTags FWorldDelegates::GetAssetTags;
 FWorldDelegates::FOnWorldTickStart FWorldDelegates::OnWorldTickStart;
+FWorldDelegates::FOnWorldTickEnd FWorldDelegates::OnWorldTickEnd;
 FWorldDelegates::FOnWorldPreActorTick FWorldDelegates::OnWorldPreActorTick;
 FWorldDelegates::FOnWorldPostActorTick FWorldDelegates::OnWorldPostActorTick;
 FWorldDelegates::FOnWorldPreSendAllEndOfFrameUpdates FWorldDelegates::OnWorldPreSendAllEndOfFrameUpdates;
diff --git a/Engine/Source/Runtime/Engine/Private/WorldPartition/WorldPartitionLevelStreamingDynamic.cpp b/Engine/Source/Runtime/Engine/Private/WorldPartition/WorldPartitionLevelStreamingDynamic.cpp
--- a/Engine/Source/Runtime/Engine/Private/WorldPartition/WorldPartitionLevelStreamingDynamic.cpp
+++ b/Engine/Source/Runtime/Engine/Private/WorldPartition/WorldPartitionLevelStreamingDynamic.cpp
@@ -219,6 +219,7 @@ bool UWorldPartitionLevelStreamingDynamic::RequestLevel(UWorld* InPersistentWorl
 	check(!FoundWorld || FoundWorld->PersistentLevel);
 	if (FoundWorld && FoundWorld->PersistentLevel != RuntimeLevel)
 	{
+		check(ULevelStreaming::ShouldReuseUnloadedButStillAroundLevels(FoundWorld->PersistentLevel));
 		check(RuntimeLevel == nullptr);
 		check(LoadedLevel == nullptr);
 		RuntimeLevel = FoundWorld->PersistentLevel;
@@ -226,6 +227,7 @@ bool UWorldPartitionLevelStreamingDynamic::RequestLevel(UWorld* InPersistentWorl
 
 	if (RuntimeLevel)
 	{
+		check(ULevelStreaming::ShouldReuseUnloadedButStillAroundLevels(RuntimeLevel));
 		// Reuse existing Level
 		UPackage* CellLevelPackage = RuntimeLevel->GetPackage();
 		check(CellLevelPackage);
@@ -504,29 +506,34 @@ void UWorldPartitionLevelStreamingDynamic::OnCleanupLevel()
 
 		RuntimeLevel->OnCleanupLevel.Remove(OnCleanupLevelDelegateHandle);
 
-		TSet<UPackage*> TrashedPackages;
-		auto TrashPackage = [&TrashedPackages](UPackage* Package)
+		// If reusing levels is enabled, trash world partition level/actor packages
+		// as it won't be done by ULevel::CleanupLevel
+		if (ShouldReuseUnloadedButStillAroundLevels(RuntimeLevel))
 		{
-			bool bWasAlreadyInSet;
-			TrashedPackages.Add(Package, &bWasAlreadyInSet);
-
-			if (!bWasAlreadyInSet)
+			TSet<UPackage*> TrashedPackages;
+			auto TrashPackage = [&TrashedPackages](UPackage* Package)
 			{
-				// Clears RF_Standalone flag on objects in package (UMetaData)
-				ForEachObjectWithPackage(Package, [](UObject* Object) { Object->ClearFlags(RF_Standalone); return true; }, false);
+				bool bWasAlreadyInSet;
+				TrashedPackages.Add(Package, &bWasAlreadyInSet);
 
-				// Rename package to avoid having to deal with pending kill objects in subsequent RequestLevel calls
-				FName NewPackageName = MakeUniqueObjectName(nullptr, UPackage::StaticClass(), FName(*FString::Printf(TEXT("%s_Trashed"), *Package->GetName())));
-				Package->Rename(*NewPackageName.ToString(), nullptr, REN_ForceNoResetLoaders | REN_DontCreateRedirectors | REN_NonTransactional | REN_DoNotDirty);
-			}
-		};
-		
-		TrashPackage(RuntimeLevel->GetPackage());
-		for (AActor* Actor : RuntimeLevel->Actors)
-		{
-			if (UPackage* ActorPackage = Actor ? Actor->GetExternalPackage() : nullptr)
+				if (!bWasAlreadyInSet)
+				{
+					// Clears RF_Standalone flag on objects in package (UMetaData)
+					ForEachObjectWithPackage(Package, [](UObject* Object) { Object->ClearFlags(RF_Standalone); return true; }, false);
+
+					// Rename package to avoid having to deal with pending kill objects in subsequent RequestLevel calls
+					FName NewPackageName = MakeUniqueObjectName(nullptr, UPackage::StaticClass(), FName(*FString::Printf(TEXT("%s_Trashed"), *Package->GetName())));
+					Package->Rename(*NewPackageName.ToString(), nullptr, REN_ForceNoResetLoaders | REN_DontCreateRedirectors | REN_NonTransactional | REN_DoNotDirty);
+				}
+			};
+
+			TrashPackage(RuntimeLevel->GetPackage());
+			for (AActor* Actor : RuntimeLevel->Actors)
 			{
-				TrashPackage(ActorPackage);
+				if (UPackage* ActorPackage = Actor ? Actor->GetExternalPackage() : nullptr)
+				{
+					TrashPackage(ActorPackage);
+				}
 			}
 		}
 
