Downgrade error to warning so build can be published even when it occurs

diff --git a/Engine/Source/Runtime/Engine/Private/WorldPartition/ContentBundle/ContentBundleEditor.cpp b/Engine/Source/Runtime/Engine/Private/WorldPartition/ContentBundle/ContentBundleEditor.cpp
--- a/Engine/Source/Runtime/Engine/Private/WorldPartition/ContentBundle/ContentBundleEditor.cpp
+++ b/Engine/Source/Runtime/Engine/Private/WorldPartition/ContentBundle/ContentBundleEditor.cpp
@@ -437,7 +437,7 @@ bool FContentBundleEditor::PopulateGeneratedPackageForCook(class IWorldPartition
 				}
 				else
 				{
-					UE_LOG(LogContentBundle, Error, TEXT("[CB: %s][Cook] Cell %s is flagged always loaded. Content Bundles cells should never be always loaded. It will not be populated and be empty at runtime."),
+					UE_LOG(LogContentBundle, Warning, TEXT("[CB: %s][Cook] Cell %s is flagged always loaded. Content Bundles cells should never be always loaded. It will not be populated and be empty at runtime."),
 						*GetDescriptor()->GetDisplayName(), *PackageToCook.RelativePath);
 					bIsSuccess = false;
 				}
