Make material compilation debug output Verbose.

- On materials w/ many static switches this can create walls of text which spam the log.

diff --git a/Engine/Source/Runtime/Engine/Private/Materials/MaterialShader.cpp b/Engine/Source/Runtime/Engine/Private/Materials/MaterialShader.cpp
--- a/Engine/Source/Runtime/Engine/Private/Materials/MaterialShader.cpp
+++ b/Engine/Source/Runtime/Engine/Private/Materials/MaterialShader.cpp
@@ -2110,7 +2110,7 @@ void FMaterialShaderMap::Compile(
 	// If we aren't actually compiling shaders don't print the debug message that we are compiling shaders.
 	if (PrecompileMode != EMaterialShaderPrecompileMode::None)
 	{
-		UE_LOG(LogShaders, Display, TEXT("	%s"), *WorkingDebugDescription);
+		UE_LOG(LogShaders, Verbose, TEXT("	%s"), *WorkingDebugDescription);
 	}
 	NewContent->DebugDescription = *WorkingDebugDescription;
 #else
