Fixed a crash when attempting to fill a mirror data table with 0 bones in the compact pose

diff --git a/Engine/Source/Runtime/Engine/Private/Animation/MirrorDataTable.cpp b/Engine/Source/Runtime/Engine/Private/Animation/MirrorDataTable.cpp
--- a/Engine/Source/Runtime/Engine/Private/Animation/MirrorDataTable.cpp
+++ b/Engine/Source/Runtime/Engine/Private/Animation/MirrorDataTable.cpp
@@ -365,8 +365,11 @@ void UMirrorDataTable::FillCompactPoseAndComponentRefRotations(
 
 	const int32 NumBones = BoneContainer.GetCompactPoseNumBones();
 	OutComponentSpaceRefRotations.SetNumUninitialized(NumBones);
-	OutComponentSpaceRefRotations[FCompactPoseBoneIndex(0)] = 
-		BoneContainer.GetRefPoseTransform(FCompactPoseBoneIndex(0)).GetRotation();
+	if (NumBones > 0 && BoneContainer.GetRefPoseArray().Num() > 0)
+	{
+		OutComponentSpaceRefRotations[FCompactPoseBoneIndex(0)] =
+			BoneContainer.GetRefPoseTransform(FCompactPoseBoneIndex(0)).GetRotation();
+	}
 	for (FCompactPoseBoneIndex BoneIndex(1); BoneIndex < NumBones; ++BoneIndex)
 	{
 		const FCompactPoseBoneIndex ParentBoneIndex = BoneContainer.GetParentBoneIndex(BoneIndex);
