Follow up fix for editor-time codec mismatch (edge case w/ "Platform Specific" and "Project Defined" formats)
#rnx

diff --git a/Engine/Source/Runtime/Engine/Private/SoundWave.cpp b/Engine/Source/Runtime/Engine/Private/SoundWave.cpp
--- a/Engine/Source/Runtime/Engine/Private/SoundWave.cpp
+++ b/Engine/Source/Runtime/Engine/Private/SoundWave.cpp
@@ -1163,7 +1163,7 @@ void USoundWave::SetSoundAssetCompressionType(ESoundAssetCompressionType InSound
 #if WITH_EDITOR
 	SoundAssetCompressionType = InSoundAssetCompressionType;
 	SoundWaveDataPtr->bIsSeekable = IsSeekable();
-	SoundWaveDataPtr->RuntimeFormat = Audio::ToName(InSoundAssetCompressionType);
+	SoundWaveDataPtr->RuntimeFormat = SoundWaveDataPtr->FindRuntimeFormat(*this);
 	UpdateAsset();
 #endif // #if WITH_EDITOR
 }
