Adding a method to change the IsSpatiallyLoaded flag in actor descriptors.

#rnx

diff --git a/Engine/Source/Runtime/Engine/Public/WorldPartition/WorldPartitionActorDesc.h b/Engine/Source/Runtime/Engine/Public/WorldPartition/WorldPartitionActorDesc.h
--- a/Engine/Source/Runtime/Engine/Public/WorldPartition/WorldPartitionActorDesc.h
+++ b/Engine/Source/Runtime/Engine/Public/WorldPartition/WorldPartitionActorDesc.h
@@ -92,6 +92,9 @@ PRAGMA_ENABLE_DEPRECATION_WARNINGS
 	inline bool GetActorIsEditorOnly() const { return bActorIsEditorOnly; }
 	inline bool GetActorIsRuntimeOnly() const { return bActorIsRuntimeOnly; }
 
+	UE_DEPRECATED(5.1, "SetIsSpatiallyLoadedRaw is deprecated and should not be used.")
+	inline void SetIsSpatiallyLoadedRaw(bool bNewIsSpatiallyLoaded) { bIsSpatiallyLoaded = bNewIsSpatiallyLoaded; }
+
 	UE_DEPRECATED(5.1, "GetLevelBoundsRelevant is deprecated.")
 	inline bool GetLevelBoundsRelevant() const { return false; }
 
