IncrementalCooking: Add DeclareConstructClasses for classes that need it. Fix some transient UObjects that assert they are never persistently serialized to support Serialize being called on their ClassDefaultObject to support DeclareCustomVersions.

#rnx

diff --git a/Engine/Source/Runtime/Engine/Private/Animation/AnimSequenceBase.cpp b/Engine/Source/Runtime/Engine/Private/Animation/AnimSequenceBase.cpp
--- a/Engine/Source/Runtime/Engine/Private/Animation/AnimSequenceBase.cpp
+++ b/Engine/Source/Runtime/Engine/Private/Animation/AnimSequenceBase.cpp
@@ -52,6 +52,18 @@ UAnimSequenceBase::UAnimSequenceBase(const FObjectInitializer& ObjectInitializer
 #endif // WITH_EDITOR
 }
 
+{
+#if WITH_EDITOR
+	// Editor time we update the source data
+	Controller->SetPlayLength(NewLength);
+#else 
+	// In case this is called during runtime, update the sequence length directly. Current use-case is runtime created Montages.
+	PRAGMA_DISABLE_DEPRECATION_WARNINGS
+	SequenceLength = NewLength;
+	PRAGMA_ENABLE_DEPRECATION_WARNINGS
+#endif
+}
+
 bool UAnimSequenceBase::IsPostLoadThreadSafe() const
 {
 	return false;	// PostLoad is not thread safe because of the call to VerifyCurveNames() (calling USkeleton::VerifySmartName) that can mutate a shared map in the skeleton.
