Fix for FMaterialRenderProxy::EvaluateUniformExpressions() being called every frame (when in editor) and not caching the results.

#rnx

diff --git a/Engine/Source/Runtime/Engine/Private/Materials/MaterialShared.cpp b/Engine/Source/Runtime/Engine/Private/Materials/MaterialShared.cpp
--- a/Engine/Source/Runtime/Engine/Private/Materials/MaterialShared.cpp
+++ b/Engine/Source/Runtime/Engine/Private/Materials/MaterialShared.cpp
@@ -4016,7 +4016,9 @@ void FMaterialRenderProxy::EvaluateUniformExpressions(FUniformExpressionCache& O
 	OutUniformExpressionCache.ParameterCollections = UniformExpressionSet.ParameterCollections;
 
 	++UniformExpressionCacheSerialNumber;
-	OutUniformExpressionCache.bUpToDate = Context.Material.IsRenderingThreadShaderMapComplete();
+	// The cache is invalidated whenever new shaders are added, so it's fine to always set this to true here
+	// instead of checking if the shader map is complete.
+	OutUniformExpressionCache.bUpToDate = true;
 }
 
 void FMaterialRenderProxy::CacheUniformExpressions(bool bRecreateUniformBuffer)
