Only add actor name to the LevelInstanceID hash if the name is stable for networking. LevelInstance Actor with an unstable name have been spawned at runtime and will have a unique GUID.
#rnx

diff --git a/Engine/Source/Runtime/Engine/Private/LevelInstance/LevelInstanceSubsystem.cpp b/Engine/Source/Runtime/Engine/Private/LevelInstance/LevelInstanceSubsystem.cpp
--- a/Engine/Source/Runtime/Engine/Private/LevelInstance/LevelInstanceSubsystem.cpp
+++ b/Engine/Source/Runtime/Engine/Private/LevelInstance/LevelInstanceSubsystem.cpp
@@ -130,12 +130,16 @@ FLevelInstanceID::FLevelInstanceID(ULevelInstanceSubsystem* LevelInstanceSubsyst
 	});
 	check(!Guids.IsEmpty());
 	
+	uint64 NameHash = 0;
+	ActorName = LevelInstanceActor->GetFName();
 	// Add Actor Name to hash because with World Partition Embedding top level of Level Instance hierarchy can get stripped leaving us with clashing ids.
 	// When embedding actors we make sure their names are unique (they get suffixed with their parent container id)
-	ActorName = LevelInstanceActor->GetFName();
-	
-	FString NameStr = ActorName.ToString();
-	const uint64 NameHash = CityHash64((const char*)*NameStr, NameStr.Len() * sizeof(TCHAR));
+	// Only do it for actor with a stable name. Actor with an unstable name are dynamically spawned and should have an unique replicated GUID.
+	if (LevelInstanceActor->IsNameStableForNetworking())
+	{
+		FString NameStr = ActorName.ToString();
+		NameHash = CityHash64((const char*)*NameStr, NameStr.Len() * sizeof(TCHAR));
+	}
 	
 	Hash = CityHash64WithSeed((const char*)Guids.GetData(), Guids.Num() * sizeof(FGuid), NameHash);
 }
