World Partition Spatial Hash: sanitize editable grid values for grid name, cell size and loading range.

#rnx

diff --git a/Engine/Source/Runtime/Engine/Private/WorldPartition/WorldPartitionRuntimeSpatialHash.cpp b/Engine/Source/Runtime/Engine/Private/WorldPartition/WorldPartitionRuntimeSpatialHash.cpp
--- a/Engine/Source/Runtime/Engine/Private/WorldPartition/WorldPartitionRuntimeSpatialHash.cpp
+++ b/Engine/Source/Runtime/Engine/Private/WorldPartition/WorldPartitionRuntimeSpatialHash.cpp
@@ -871,6 +871,39 @@ URuntimeHashExternalStreamingObjectBase* UWorldPartitionRuntimeSpatialHash::Stor
 	return StreamingObject;
 }
 
+void UWorldPartitionRuntimeSpatialHash::PostEditChangeProperty(FPropertyChangedEvent& PropertyChangedEvent)
+{
+	const FName PropertyName = PropertyChangedEvent.GetPropertyName();
+
+	if (PropertyName == GET_MEMBER_NAME_CHECKED(FSpatialHashStreamingGrid, CellSize))
+	{
+		for (FSpatialHashRuntimeGrid& Grid : Grids)
+		{
+			Grid.CellSize = FMath::Max(Grid.CellSize, 1600);
+		}
+	}
+	else if (PropertyName == GET_MEMBER_NAME_CHECKED(FSpatialHashStreamingGrid, GridName))
+	{
+		int GridIndex = 0;
+		for (FSpatialHashRuntimeGrid& Grid : Grids)
+		{
+			if (Grid.GridName.IsNone())
+			{
+				Grid.GridName = *FString::Printf(TEXT("Grid%02d"), GridIndex++);
+			}
+		}
+	}
+	else if (PropertyName == GET_MEMBER_NAME_CHECKED(FSpatialHashStreamingGrid, LoadingRange))
+	{
+		for (FSpatialHashRuntimeGrid& Grid : Grids)
+		{
+			Grid.LoadingRange = FMath::Max(Grid.LoadingRange, 0);
+		}
+	}
+
+	Super::PostEditChangeProperty(PropertyChangedEvent);
+}
+
 void UWorldPartitionRuntimeSpatialHash::SetDefaultValues()
 {
 	check(!Grids.Num());
diff --git a/Engine/Source/Runtime/Engine/Public/WorldPartition/WorldPartitionRuntimeSpatialHash.h b/Engine/Source/Runtime/Engine/Public/WorldPartition/WorldPartitionRuntimeSpatialHash.h
--- a/Engine/Source/Runtime/Engine/Public/WorldPartition/WorldPartitionRuntimeSpatialHash.h
+++ b/Engine/Source/Runtime/Engine/Public/WorldPartition/WorldPartitionRuntimeSpatialHash.h
@@ -255,6 +255,8 @@ public:
 	virtual void PreSave(FObjectPreSaveContext ObjectSaveContext) override;
 
 #if WITH_EDITOR
+	virtual void PostEditChangeProperty(FPropertyChangedEvent& PropertyChangedEvent) override;
+
 	virtual void SetDefaultValues() override;
 	virtual bool SupportsHLODs() const override { return true; }
 	virtual bool PrepareGeneratorPackageForCook(TArray<UPackage*>& OutModifiedPackages) override;
