Skeletalmesh fix issue when doing a posteditchange and we are compiling, the Postedit change blocking

#rnx

diff --git a/Engine/Source/Runtime/Engine/Private/SkeletalMesh.cpp b/Engine/Source/Runtime/Engine/Private/SkeletalMesh.cpp
--- a/Engine/Source/Runtime/Engine/Private/SkeletalMesh.cpp
+++ b/Engine/Source/Runtime/Engine/Private/SkeletalMesh.cpp
@@ -324,8 +324,10 @@ FScopedSkeletalMeshPostEditChange::~FScopedSkeletalMeshPostEditChange()
 
 void FScopedSkeletalMeshPostEditChange::SetSkeletalMesh(USkeletalMesh* InSkeletalMesh)
 {
+	//Skip only if we are calling post edit change
+	bool bSkipCompiling = InSkeletalMesh->IsCompiling() && bCallPostEditChange;
 	//Some parallel task may try to call post edit change, we must prevent it
-	if (!IsInGameThread() || InSkeletalMesh->IsCompiling())
+	if (!IsInGameThread() || bSkipCompiling)
 	{
 		return;
 	}
