Texture.cpp : nop comment about inconsistency in UTexture::SetLayerFormatSettings(0)

diff --git a/Engine/Source/Runtime/Engine/Private/Texture.cpp b/Engine/Source/Runtime/Engine/Private/Texture.cpp
--- a/Engine/Source/Runtime/Engine/Private/Texture.cpp
+++ b/Engine/Source/Runtime/Engine/Private/Texture.cpp
@@ -2904,6 +2904,12 @@ void UTexture::SetLayerFormatSettings(int32 LayerIndex, const FTextureFormatSett
 			}
 		}
 		LayerFormatSettings[LayerIndex] = InSettings;
+
+		// @todo Oodle : inconsistency in SetLayerFormatSettings(0) and possible bug?
+		// should SetLayerFormatSettings(0,Settings) always set the base Texture properties?
+		// if you call this when you have a LayerFormatSettings[] array, it does not
+		// if you query via GetLayerFormatSettings(0) then these settings are seen
+		//	but you just get them directly from the texture they are not!
 	}
 }
 
