Disable display-based screen percentage modes when stereo HMD is active

Display-based and DPI screen percentage modes don't make sense for stereo HMDs.
We use FEngineShowFlags to decide if we need to disable these screen percentage modes when we are rendering in stereo.

diff --git a/Engine/Source/Editor/UnrealEd/Private/EditorViewportClient.cpp b/Engine/Source/Editor/UnrealEd/Private/EditorViewportClient.cpp
--- a/Engine/Source/Editor/UnrealEd/Private/EditorViewportClient.cpp
+++ b/Engine/Source/Editor/UnrealEd/Private/EditorViewportClient.cpp
@@ -2750,7 +2750,7 @@ bool FEditorViewportClient::SupportsPreviewResolutionFraction() const
 
 float FEditorViewportClient::GetDefaultPrimaryResolutionFractionTarget() const
 {
-	FStaticResolutionFractionHeuristic StaticHeuristic;
+	FStaticResolutionFractionHeuristic StaticHeuristic(EngineShowFlags);
 	StaticHeuristic.Settings.PullEditorRenderingSettings(bIsRealtime, /* bIsPathTraced = */ EngineShowFlags.PathTracing);
 
 	if (SupportsLowDPIPreview() && IsLowDPIPreview()) // TODO: && ViewFamily.SupportsScreenPercentage())
diff --git a/Engine/Source/Runtime/Engine/Private/GameViewportClient.cpp b/Engine/Source/Runtime/Engine/Private/GameViewportClient.cpp
--- a/Engine/Source/Runtime/Engine/Private/GameViewportClient.cpp
+++ b/Engine/Source/Runtime/Engine/Private/GameViewportClient.cpp
@@ -1650,7 +1650,7 @@ void UGameViewportClient::Draw(FViewport* InViewport, FCanvas* SceneCanvas)
 			if (ViewFamily.EngineShowFlags.ScreenPercentage && !bDisableWorldRendering && ViewFamily.Views.Num() > 0)
 			{
 				// Get global view fraction.
-				FStaticResolutionFractionHeuristic StaticHeuristic;
+				FStaticResolutionFractionHeuristic StaticHeuristic(ViewFamily.EngineShowFlags);
 
 #if WITH_EDITOR
 				if (FStaticResolutionFractionHeuristic::FUserSettings::EditorOverridePIESettings())
diff --git a/Engine/Source/Runtime/Engine/Private/LegacyScreenPercentageDriver.cpp b/Engine/Source/Runtime/Engine/Private/LegacyScreenPercentageDriver.cpp
--- a/Engine/Source/Runtime/Engine/Private/LegacyScreenPercentageDriver.cpp
+++ b/Engine/Source/Runtime/Engine/Private/LegacyScreenPercentageDriver.cpp
@@ -177,6 +177,11 @@ float FLegacyScreenPercentageDriver::GetCVarResolutionFraction()
 	return GlobalFraction;
 }
 
+FStaticResolutionFractionHeuristic::FStaticResolutionFractionHeuristic(const FEngineShowFlags& EngineShowFlags)
+{
+	Settings.bAllowDisplayBasedScreenPercentageMode = (EngineShowFlags.StereoRendering == 0) && (EngineShowFlags.VREditing == 0);
+}
+
 #if WITH_EDITOR
 // static
 bool FStaticResolutionFractionHeuristic::FUserSettings::EditorOverridePIESettings()
@@ -259,7 +264,10 @@ float FStaticResolutionFractionHeuristic::ResolveResolutionFraction() const
 	float LocalTotalDisplayedPixelCount = FMath::Max(TotalDisplayedPixelCount, 1) * SecondaryViewFraction * SecondaryViewFraction;
 
 	float GlobalResolutionFraction = 1.0f;
-	if (Settings.Mode == EScreenPercentageMode::BasedOnDisplayResolution)
+
+	const EScreenPercentageMode EffectiveScreenPercentageMode = Settings.bAllowDisplayBasedScreenPercentageMode ? Settings.Mode : EScreenPercentageMode::Manual;
+
+	if (EffectiveScreenPercentageMode == EScreenPercentageMode::BasedOnDisplayResolution)
 	{
 		static bool bInitPixelCount = false;
 		static float AutoMinDisplayResolution;
@@ -314,13 +322,13 @@ float FStaticResolutionFractionHeuristic::ResolveResolutionFraction() const
 
 		GlobalResolutionFraction = FMath::Sqrt(Settings.AutoPixelCountMultiplier * LerpedRenderingPixelCount / LocalTotalDisplayedPixelCount);
 	}
-	else if (Settings.Mode == EScreenPercentageMode::BasedOnDPIScale)
+	else if (EffectiveScreenPercentageMode == EScreenPercentageMode::BasedOnDPIScale)
 	{
 		GlobalResolutionFraction = 1.0f / DPIScale;
 	}
 	else
 	{
-		ensure(Settings.Mode == EScreenPercentageMode::Manual);
+		ensure(EffectiveScreenPercentageMode == EScreenPercentageMode::Manual);
 		GlobalResolutionFraction = Settings.GlobalResolutionFraction;
 	}
 
diff --git a/Engine/Source/Runtime/Engine/Public/LegacyScreenPercentageDriver.h b/Engine/Source/Runtime/Engine/Public/LegacyScreenPercentageDriver.h
--- a/Engine/Source/Runtime/Engine/Public/LegacyScreenPercentageDriver.h
+++ b/Engine/Source/Runtime/Engine/Public/LegacyScreenPercentageDriver.h
@@ -65,6 +65,8 @@ enum class EScreenPercentageMode
  */
 struct ENGINE_API FStaticResolutionFractionHeuristic
 {
+	FStaticResolutionFractionHeuristic(const FEngineShowFlags& EngineShowFlags);
+
 	// User configurable settings
 	struct ENGINE_API FUserSettings
 	{
@@ -80,6 +82,9 @@ struct ENGINE_API FStaticResolutionFractionHeuristic
 		// r.ScreenPercentage.Auto.* Mode = EMode::BasedOnDisplayResolution.
 		float AutoPixelCountMultiplier = 1.0f;
 
+		// stereo HMDs cannot use percentage modes based on 2D monitor
+		bool bAllowDisplayBasedScreenPercentageMode = true;
+
 		/** Return whether should use the editor settings for PIE. */
 #if WITH_EDITOR
 		static bool EditorOverridePIESettings();
