Fix ShooterGame DedicatedServerTest warning:
- It is possible for PlayerController to be moved in a temp transition world (FSeamlessTravelHandler::StartTravel) with no WorldParittionSubsystem

#rnx

diff --git a/Engine/Source/Runtime/Engine/Private/PlayerController.cpp b/Engine/Source/Runtime/Engine/Private/PlayerController.cpp
--- a/Engine/Source/Runtime/Engine/Private/PlayerController.cpp
+++ b/Engine/Source/Runtime/Engine/Private/PlayerController.cpp
@@ -3601,8 +3601,8 @@ void APlayerController::PostSeamlessTravel()
 void APlayerController::OnAddedToPlayerControllerList()
 {
 	UWorld* World = GetWorld();
-	UWorldPartitionSubsystem* WorldPartitionSubsystem = World->GetSubsystem<UWorldPartitionSubsystem>();
-	if (ensure(WorldPartitionSubsystem))
+	// Possible we are moved into a world with no WorldPartitionSubsystem by the seamless travel (FSeamlessTravelHandler::StartTravel with no TransitionMap)
+	if (UWorldPartitionSubsystem* WorldPartitionSubsystem = World->GetSubsystem<UWorldPartitionSubsystem>())
 	{
 		WorldPartitionSubsystem->RegisterStreamingSourceProvider(this);
 	}
