Unshelved from pending changelist '22698707'

Default constructor for FStaticResolutionFractionHeuristic

diff --git a/Engine/Source/Runtime/Engine/Public/LegacyScreenPercentageDriver.h b/Engine/Source/Runtime/Engine/Public/LegacyScreenPercentageDriver.h
--- a/Engine/Source/Runtime/Engine/Public/LegacyScreenPercentageDriver.h
+++ b/Engine/Source/Runtime/Engine/Public/LegacyScreenPercentageDriver.h
@@ -65,6 +65,8 @@ enum class EScreenPercentageMode
  */
 struct ENGINE_API FStaticResolutionFractionHeuristic
 {
+	FStaticResolutionFractionHeuristic() = default;
+
 	FStaticResolutionFractionHeuristic(const FEngineShowFlags& EngineShowFlags);
 
 	// User configurable settings
