Fix MS Platform codecs when audiolink is on

diff --git a/Engine/Source/Runtime/AudioLink/AudioLinkEngine/Private/IAudioLinkFactory.cpp b/Engine/Source/Runtime/AudioLink/AudioLinkEngine/Private/IAudioLinkFactory.cpp
--- a/Engine/Source/Runtime/AudioLink/AudioLinkEngine/Private/IAudioLinkFactory.cpp
+++ b/Engine/Source/Runtime/AudioLink/AudioLinkEngine/Private/IAudioLinkFactory.cpp
@@ -2,6 +2,7 @@
 
 #include "IAudioLinkFactory.h"
 #include "Algo/Transform.h"
+#include "AudioLinkLog.h"
 
 // Concrete Buffer Listeners.
 #include "BufferedSubmixListener.h" 
@@ -23,6 +24,10 @@ IAudioLinkFactory::IAudioLinkFactory()
 	{
 		IModularFeatures::Get().RegisterModularFeature(GetModularFeatureName(), this);
 	}
+	else
+	{
+		UE_LOG(LogAudioLink, Warning, TEXT("AudioLink is disabled, au.audiolink.enabled=0. Not registering factory."));
+	}
 }
 
 IAudioLinkFactory::~IAudioLinkFactory()
diff --git a/Engine/Source/Runtime/Engine/Private/AudioFormatSettings.cpp b/Engine/Source/Runtime/Engine/Private/AudioFormatSettings.cpp
--- a/Engine/Source/Runtime/Engine/Private/AudioFormatSettings.cpp
+++ b/Engine/Source/Runtime/Engine/Private/AudioFormatSettings.cpp
@@ -10,6 +10,17 @@
 
 namespace Audio
 {
+	bool ShouldAllowPlatformSpecificFormats()
+	{
+		static bool IsAudioLinkEnabled = []() -> bool
+		{
+			bool bAvailable = IModularFeatures::Get().IsModularFeatureAvailable(TEXT("AudioLink Factory"));
+			UE_CLOG(bAvailable,LogAudio, Display, TEXT("AudioLink is enabled, disabling platform specific AudioFormats."));
+			return bAvailable;
+		}();
+		return !IsAudioLinkEnabled;
+	}
+
 	FAudioFormatSettings::FAudioFormatSettings(FConfigCacheIni* InConfigSystem, const FString& InConfigFilename, const FString& InPlatformIdentifierForLogging)
 	{
 		ReadConfiguration(InConfigSystem, InConfigFilename, InPlatformIdentifierForLogging);
@@ -20,13 +31,20 @@ namespace Audio
 		FName FormatName = Audio::ToName(Wave->GetSoundAssetCompressionType());
 		if (FormatName == Audio::NAME_PLATFORM_SPECIFIC)
 		{
-			if (Wave->IsStreaming())
+			if (ShouldAllowPlatformSpecificFormats())
 			{
-				return PlatformStreamingFormat;
+				if (Wave->IsStreaming())
+				{
+					return PlatformStreamingFormat;
+				}
+				else
+				{
+					return PlatformFormat;
+				}
 			}
 			else
 			{
-				return PlatformFormat;
+				return FallbackFormat;
 			}
 		}
 		return FormatName;
