Fix for updating USoundWaves in editor when the default compression codec changes.

diff --git a/Engine/Source/Runtime/Engine/Classes/Sound/SoundWave.h b/Engine/Source/Runtime/Engine/Classes/Sound/SoundWave.h
--- a/Engine/Source/Runtime/Engine/Classes/Sound/SoundWave.h
+++ b/Engine/Source/Runtime/Engine/Classes/Sound/SoundWave.h
@@ -456,9 +456,12 @@ public:
 	UFUNCTION(BlueprintPure, Category = "Audio")
 	ESoundAssetCompressionType GetSoundAssetCompressionType() const;
 
+	/** will return the raw value, (i.e. does not resolve options such as "Project Defined" to the correct codec) */
+	ESoundAssetCompressionType GetSoundAssetCompressionTypeEnum() const;
+
 	/** Procedurally set the compression type. */
 	UFUNCTION(BlueprintCallable, Category = "Audio")
-	void SetSoundAssetCompressionType(ESoundAssetCompressionType InSoundAssetCompressionType);
+	void SetSoundAssetCompressionType(ESoundAssetCompressionType InSoundAssetCompressionType, bool bMarkDirty = true);
 
 	/** Returns the Runtime format of the wave */
 	FName GetRuntimeFormat() const;
@@ -599,7 +602,7 @@ private:
 
 	// Called when we change any properties about the underlying audio asset
 #if WITH_EDITOR
-	void UpdateAsset();
+	void UpdateAsset(bool bMarkDirty = true);
 #endif
 
 public:
diff --git a/Engine/Source/Runtime/Engine/Private/AudioSettings.cpp b/Engine/Source/Runtime/Engine/Private/AudioSettings.cpp
--- a/Engine/Source/Runtime/Engine/Private/AudioSettings.cpp
+++ b/Engine/Source/Runtime/Engine/Private/AudioSettings.cpp
@@ -101,6 +101,20 @@ void UAudioSettings::PostEditChangeChainProperty(FPropertyChangedChainEvent& Pro
 		{
 			bPromptRestartRequired = true;
 		}
+		else if (PropertyName == GET_MEMBER_NAME_CHECKED(UAudioSettings, DefaultAudioCompressionType))
+		{
+			// loop through all USoundWaves and update the compression type w/ the new defualt
+			// (if they are set to "Project Default")
+			for (TObjectIterator<USoundWave> It; It; ++It)
+			{
+				USoundWave* SoundWave = *It;
+				if(SoundWave && SoundWave->GetSoundAssetCompressionTypeEnum() == ESoundAssetCompressionType::ProjectDefined)
+				{
+					// this will query the correct compression type and update the asset accordingly
+					SoundWave->SetSoundAssetCompressionType(SoundWave->GetSoundAssetCompressionType(), /*bMarkDirty*/false);
+				}
+			}
+		}
 		else if (PropertyName == GET_MEMBER_NAME_CHECKED(UAudioSettings, QualityLevels))
 		{
 			if (QualityLevels.Num() == 0)
diff --git a/Engine/Source/Runtime/Engine/Private/SoundWave.cpp b/Engine/Source/Runtime/Engine/Private/SoundWave.cpp
--- a/Engine/Source/Runtime/Engine/Private/SoundWave.cpp
+++ b/Engine/Source/Runtime/Engine/Private/SoundWave.cpp
@@ -774,8 +774,12 @@ void USoundWave::GetResourceSizeEx(FResourceSizeEx& CumulativeResourceSize)
 
 	CumulativeResourceSize.AddDedicatedSystemMemoryBytes(CookedEnvelopeTimeData.Num() * sizeof(FSoundWaveEnvelopeTimeData));
 
-	// Add zeroth chunk data, if it's used (if this USoundWave isn't streaming, this won't report).
-	CumulativeResourceSize.AddDedicatedSystemMemoryBytes(SoundWaveDataPtr->GetZerothChunkData().GetView().Num());
+	// Add zeroth chunk data, if it's used (if this USoundWave isn't streaming,
+	// or the bulk data is culled on this platform, this won't report).
+	if(FApp::CanEverRenderAudio())
+	{
+		CumulativeResourceSize.AddDedicatedSystemMemoryBytes(SoundWaveDataPtr->GetZerothChunkData().GetView().Num());
+	}
 
 	// Finally, report the actual audio memory being used, if this asset isn't using the stream cache.
 	if (FAudioDevice* LocalAudioDevice = GEngine->GetMainAudioDeviceRaw())
@@ -1158,13 +1162,18 @@ ESoundAssetCompressionType USoundWave::GetSoundAssetCompressionType() const
 	}
 }
 
-void USoundWave::SetSoundAssetCompressionType(ESoundAssetCompressionType InSoundAssetCompressionType)
+ESoundAssetCompressionType USoundWave::GetSoundAssetCompressionTypeEnum() const
+{
+	return SoundAssetCompressionType;
+}
+
+void USoundWave::SetSoundAssetCompressionType(ESoundAssetCompressionType InSoundAssetCompressionType, bool bMarkDirty)
 {
 #if WITH_EDITOR
 	SoundAssetCompressionType = InSoundAssetCompressionType;
 	SoundWaveDataPtr->bIsSeekable = IsSeekable();
 	SoundWaveDataPtr->RuntimeFormat = SoundWaveDataPtr->FindRuntimeFormat(*this);
-	UpdateAsset();
+	UpdateAsset(bMarkDirty);
 #endif // #if WITH_EDITOR
 }
 
@@ -3878,12 +3887,15 @@ void USoundWave::OverrideLoadingBehavior(ESoundWaveLoadingBehavior InLoadingBeha
 }
 
 #if WITH_EDITOR
-void USoundWave::UpdateAsset()
+void USoundWave::UpdateAsset(bool bMarkDirty)
 {
 	InvalidateCompressedData();
 	FreeResources();
 	UpdatePlatformData();
-	MarkPackageDirty();
+	if(bMarkDirty)
+	{
+		MarkPackageDirty();
+	}
 
 	// if we are force inline, we need to make sure the shared data is pulled from the DDC
 	// before we attempt to use a decoder on the proxy (not using stream caching)
