Animation Compression codecs per platform settings
- Add ITargetPlatform to PopulateDDCKey function parameters.
- Update ACLPlugin to use new parameters.

#rnx
#tests Cooked single assets with different TargetPlatforms to verify platform-specific settings were used.

diff --git a/Engine/Source/Runtime/Engine/Classes/Animation/AnimBoneCompressionCodec.h b/Engine/Source/Runtime/Engine/Classes/Animation/AnimBoneCompressionCodec.h
--- a/Engine/Source/Runtime/Engine/Classes/Animation/AnimBoneCompressionCodec.h
+++ b/Engine/Source/Runtime/Engine/Classes/Animation/AnimBoneCompressionCodec.h
@@ -33,10 +33,17 @@ class ENGINE_API UAnimBoneCompressionCodec : public UObject
 	/** Compresses the curve data from an animation sequence. */
 	virtual bool Compress(const FCompressibleAnimData& CompressibleAnimData, FCompressibleAnimDataResult& OutResult) PURE_VIRTUAL(UAnimCurveCompressionCodec::Compress, return false;);
 
+	/*
+	 * Called to generate a unique DDC key for this codec instance and input anim sequence and TargetPlatform
+	 * A suitable key should be generated from: the InstanceGuid, a codec version, and all relevant properties that drive the behavior.
+	 */
+	virtual void PopulateDDCKey(const UE::Anim::Compression::FAnimDDCKeyArgs& KeyArgs, FArchive& Ar);
+
 	/*
 	 * Called to generate a unique DDC key for this codec instance and input anim sequence.
 	 * A suitable key should be generated from: the InstanceGuid, a codec version, and all relevant properties that drive the behavior.
 	 */
+	UE_DEPRECATED(5.2, "This function has been deprecated. Override the one above instead.")
 	virtual void PopulateDDCKey(const UAnimSequenceBase& AnimSeq, FArchive& Ar);
 
 	/*
diff --git a/Engine/Source/Runtime/Engine/Classes/Animation/AnimBoneCompressionSettings.h b/Engine/Source/Runtime/Engine/Classes/Animation/AnimBoneCompressionSettings.h
--- a/Engine/Source/Runtime/Engine/Classes/Animation/AnimBoneCompressionSettings.h
+++ b/Engine/Source/Runtime/Engine/Classes/Animation/AnimBoneCompressionSettings.h
@@ -9,6 +9,7 @@ class UAnimBoneCompressionCodec;
 class UAnimSequenceBase;
 struct FCompressibleAnimData;
 struct FCompressibleAnimDataResult;
+namespace UE::Anim::Compression { struct FAnimDDCKeyArgs; }
 
 /*
  * This object is used to wrap a bone compression codec. It allows a clean integration in the editor by avoiding the need
@@ -51,7 +52,11 @@ class ENGINE_API UAnimBoneCompressionSettings : public UObject
 	 */
 	bool Compress(const FCompressibleAnimData& AnimSeq, FCompressibleAnimDataResult& OutCompressedData) const;
 
+	/** Generates a DDC key that takes into account the current settings, selected codec, input anim sequence and TargetPlatform */
+	void PopulateDDCKey(const UE::Anim::Compression::FAnimDDCKeyArgs& KeyArgs, FArchive& Ar);
+
 	/** Generates a DDC key that takes into account the current settings, selected codec, and input anim sequence. */
+	UE_DEPRECATED(5.2, "This function has been deprecated. Override the one above instead.")
 	void PopulateDDCKey(const UAnimSequenceBase& AnimSeq, FArchive& Ar);
 
 	/** Generates a DDC key that takes into account the current settings and selected codec. */
diff --git a/Engine/Source/Runtime/Engine/Classes/Animation/AnimCompress.h b/Engine/Source/Runtime/Engine/Classes/Animation/AnimCompress.h
--- a/Engine/Source/Runtime/Engine/Classes/Animation/AnimCompress.h
+++ b/Engine/Source/Runtime/Engine/Classes/Animation/AnimCompress.h
@@ -509,7 +509,7 @@ public:
 	void PopulateDDCKeyArchive(FArchive& Ar) {}
 
 protected:
-	virtual void PopulateDDCKey(const UAnimSequenceBase& AnimSeq, FArchive& Ar) override;
+	virtual void PopulateDDCKey(const UE::Anim::Compression::FAnimDDCKeyArgs& KeyArgs, FArchive& Ar) override;
 #endif // WITH_EDITOR
 
 	/**
diff --git a/Engine/Source/Runtime/Engine/Classes/Animation/AnimCompress_PerTrackCompression.h b/Engine/Source/Runtime/Engine/Classes/Animation/AnimCompress_PerTrackCompression.h
--- a/Engine/Source/Runtime/Engine/Classes/Animation/AnimCompress_PerTrackCompression.h
+++ b/Engine/Source/Runtime/Engine/Classes/Animation/AnimCompress_PerTrackCompression.h
@@ -137,7 +137,7 @@ protected:
 #if WITH_EDITOR
 	virtual void PostEditChangeProperty(struct FPropertyChangedEvent& PropertyChangedEvent) override;
 	virtual bool DoReduction(const FCompressibleAnimData& CompressibleAnimData, FCompressibleAnimDataResult& OutResult) override;
-	virtual void PopulateDDCKey(const UAnimSequenceBase& AnimSeq, FArchive& Ar) override;
+	virtual void PopulateDDCKey(const UE::Anim::Compression::FAnimDDCKeyArgs& KeyArgs, FArchive& Ar) override;
 #endif // WITH_EDITOR
 	//~ Begin UAnimCompress Interface
 
diff --git a/Engine/Source/Runtime/Engine/Classes/Animation/AnimCompress_RemoveEverySecondKey.h b/Engine/Source/Runtime/Engine/Classes/Animation/AnimCompress_RemoveEverySecondKey.h
--- a/Engine/Source/Runtime/Engine/Classes/Animation/AnimCompress_RemoveEverySecondKey.h
+++ b/Engine/Source/Runtime/Engine/Classes/Animation/AnimCompress_RemoveEverySecondKey.h
@@ -33,7 +33,7 @@ protected:
 	//~ Begin UAnimCompress Interface
 #if WITH_EDITOR
 	virtual bool DoReduction(const FCompressibleAnimData& CompressibleAnimData, FCompressibleAnimDataResult& OutResult) override;
-	virtual void PopulateDDCKey(const UAnimSequenceBase& AnimSeq, FArchive& Ar) override;
+	virtual void PopulateDDCKey(const UE::Anim::Compression::FAnimDDCKeyArgs& KeyArgs, FArchive& Ar) override;
 #endif // WITH_EDITOR
 	//~ Begin UAnimCompress Interface
 };
diff --git a/Engine/Source/Runtime/Engine/Classes/Animation/AnimCompress_RemoveLinearKeys.h b/Engine/Source/Runtime/Engine/Classes/Animation/AnimCompress_RemoveLinearKeys.h
--- a/Engine/Source/Runtime/Engine/Classes/Animation/AnimCompress_RemoveLinearKeys.h
+++ b/Engine/Source/Runtime/Engine/Classes/Animation/AnimCompress_RemoveLinearKeys.h
@@ -80,7 +80,7 @@ protected:
 	//~ Begin UAnimCompress Interface
 #if WITH_EDITOR
 	virtual bool DoReduction(const FCompressibleAnimData& CompressibleAnimData, FCompressibleAnimDataResult& OutResult) override;
-	virtual void PopulateDDCKey(const UAnimSequenceBase& AnimSeq, FArchive& Ar) override;
+	virtual void PopulateDDCKey(const UE::Anim::Compression::FAnimDDCKeyArgs& KeyArgs, FArchive& Ar) override;
 #endif // WITH_EDITOR
 	//~ Begin UAnimCompress Interface
 
diff --git a/Engine/Source/Runtime/Engine/Classes/Animation/AnimCompress_RemoveTrivialKeys.h b/Engine/Source/Runtime/Engine/Classes/Animation/AnimCompress_RemoveTrivialKeys.h
--- a/Engine/Source/Runtime/Engine/Classes/Animation/AnimCompress_RemoveTrivialKeys.h
+++ b/Engine/Source/Runtime/Engine/Classes/Animation/AnimCompress_RemoveTrivialKeys.h
@@ -32,7 +32,7 @@ protected:
 	//~ Begin UAnimCompress Interface
 #if WITH_EDITOR
 	virtual bool DoReduction(const FCompressibleAnimData& CompressibleAnimData, FCompressibleAnimDataResult& OutResult) override;
-	virtual void PopulateDDCKey(const UAnimSequenceBase& AnimSeq, FArchive& Ar) override;
+	virtual void PopulateDDCKey(const UE::Anim::Compression::FAnimDDCKeyArgs& KeyArgs, FArchive& Ar) override;
 #endif // WITH_EDITOR
 	//~ Begin UAnimCompress Interface
 };
diff --git a/Engine/Source/Runtime/Engine/Classes/Animation/AnimSequence.h b/Engine/Source/Runtime/Engine/Classes/Animation/AnimSequence.h
--- a/Engine/Source/Runtime/Engine/Classes/Animation/AnimSequence.h
+++ b/Engine/Source/Runtime/Engine/Classes/Animation/AnimSequence.h
@@ -35,6 +35,7 @@ class USkeletalMesh;
 struct FAnimCompressContext;
 struct FAnimSequenceDecompressionContext;
 struct FCompactPose;
+class ITargetPlatform;
 
 namespace UE { namespace Anim { namespace Compression { struct FScopedCompressionGuard; } } }
 
@@ -195,12 +196,14 @@ struct ENGINE_API FRequestAnimCompressionParams
 	TSharedPtr<FAnimCompressContext> CompressContext;
 
 	// Constructors
-	FRequestAnimCompressionParams(bool bInAsyncCompression, bool bInAllowAlternateCompressor = false, bool bInOutput = false);
+	FRequestAnimCompressionParams(bool bInAsyncCompression, bool bInAllowAlternateCompressor = false, bool bInOutput = false, const ITargetPlatform* = nullptr);
 	FRequestAnimCompressionParams(bool bInAsyncCompression, TSharedPtr<FAnimCompressContext> InCompressContext);
 
 	// Frame stripping initialization funcs (allow stripping per platform)
 	void InitFrameStrippingFromCVar();
-	void InitFrameStrippingFromPlatform(const class ITargetPlatform* TargetPlatform);
+	void InitFrameStrippingFromPlatform();
+
+	const ITargetPlatform* TargetPlatform;
 };
 
 UCLASS(config=Engine, hidecategories=(UObject, Length), BlueprintType)
@@ -540,7 +543,7 @@ public:
 	// Get compressed data for this UAnimSequence. May be built directly or pulled from DDC
 #if WITH_EDITOR
 	bool ShouldPerformStripping(const bool bPerformFrameStripping, const bool bPerformStrippingOnOddFramedAnims) const;
-	FString GetDDCCacheKeySuffix(const bool bPerformStripping) const;
+	FString GetDDCCacheKeySuffix(const bool bPerformStripping, const ITargetPlatform* TargetPlatform) const;
 	void ApplyCompressedData(const FString& DataCacheKeySuffix, const bool bPerformFrameStripping, const TArray<uint8>& Data);
 
 	void WaitOnExistingCompression(const bool bWantResults=true);
@@ -822,6 +825,8 @@ protected:
 
 	bool bBlockCompressionRequests;
 
+	FString RequestedCompression_DDCKey;
+
 private:
 	PRAGMA_DISABLE_DEPRECATION_WARNINGS
 	UE_DEPRECATED(5.0, "PerBoneCustomAttributeData has been deprecated see UAnimDataModel::AnimatedBoneAttributes")
diff --git a/Engine/Source/Runtime/Engine/Classes/Animation/AnimStreamable.h b/Engine/Source/Runtime/Engine/Classes/Animation/AnimStreamable.h
--- a/Engine/Source/Runtime/Engine/Classes/Animation/AnimStreamable.h
+++ b/Engine/Source/Runtime/Engine/Classes/Animation/AnimStreamable.h
@@ -236,9 +236,9 @@ public:
 
 	void UpdateRawData();
 
-	FString GetBaseDDCKey(uint32 NumChunks) const;
+	FString GetBaseDDCKey(uint32 NumChunks, const ITargetPlatform* TargetPlatform) const;
 
-	void RequestCompressedDataForChunk(const FString& ChunkDDCKey, FAnimStreamableChunk& Chunk, const int32 ChunkIndex, const uint32 FrameStart, const uint32 FrameEnd, TSharedRef<FAnimCompressContext> CompressContext);
+	void RequestCompressedDataForChunk(const FString& ChunkDDCKey, FAnimStreamableChunk& Chunk, const int32 ChunkIndex, const uint32 FrameStart, const uint32 FrameEnd, TSharedRef<FAnimCompressContext> CompressContext, const ITargetPlatform* TargetPlatform);
 #endif
 
 	bool bUseRawDataOnly;
diff --git a/Engine/Source/Runtime/Engine/Private/Animation/AnimBoneCompressionCodec.cpp b/Engine/Source/Runtime/Engine/Private/Animation/AnimBoneCompressionCodec.cpp
--- a/Engine/Source/Runtime/Engine/Private/Animation/AnimBoneCompressionCodec.cpp
+++ b/Engine/Source/Runtime/Engine/Private/Animation/AnimBoneCompressionCodec.cpp
@@ -45,6 +45,10 @@ FString UAnimBoneCompressionCodec::GetCodecDDCHandle() const
 }
 
 #if WITH_EDITORONLY_DATA
+void UAnimBoneCompressionCodec::PopulateDDCKey(const UE::Anim::Compression::FAnimDDCKeyArgs& KeyArgs, FArchive& Ar)
+{
+}
+
 void UAnimBoneCompressionCodec::PopulateDDCKey(const UAnimSequenceBase& AnimSeq, FArchive& Ar)
 {
 }
diff --git a/Engine/Source/Runtime/Engine/Private/Animation/AnimBoneCompressionSettings.cpp b/Engine/Source/Runtime/Engine/Private/Animation/AnimBoneCompressionSettings.cpp
--- a/Engine/Source/Runtime/Engine/Private/Animation/AnimBoneCompressionSettings.cpp
+++ b/Engine/Source/Runtime/Engine/Private/Animation/AnimBoneCompressionSettings.cpp
@@ -14,6 +14,7 @@
 
 #define DEBUG_DUMP_ANIM_COMPRESSION_STATS 0
 
+
 UAnimBoneCompressionSettings::UAnimBoneCompressionSettings(const FObjectInitializer& ObjectInitializer)
 	: Super(ObjectInitializer)
 {
@@ -273,6 +274,46 @@ bool UAnimBoneCompressionSettings::Compress(const FCompressibleAnimData& AnimSeq
 	return Success;
 }
 
+void UAnimBoneCompressionSettings::PopulateDDCKey(const UE::Anim::Compression::FAnimDDCKeyArgs& KeyArgs, FArchive& Ar)
+{
+	Ar << ErrorThreshold;
+	Ar << bForceBelowThreshold;
+
+	int32 NumValidCodecs = 0;
+	for (UAnimBoneCompressionCodec* Codec : Codecs)
+	{
+		if (Codec != nullptr)
+		{
+			const int64 ArchiveOffset = Ar.Tell();
+			Codec->PopulateDDCKey(KeyArgs, Ar);
+
+			if (ArchiveOffset == Ar.Tell())
+			{
+				// If nothing was written, perhaps the codec implements the old deprecated API, call it just in case
+				PRAGMA_DISABLE_DEPRECATION_WARNINGS
+				Codec->PopulateDDCKey(KeyArgs.AnimSequence, Ar);
+				PRAGMA_ENABLE_DEPRECATION_WARNINGS
+
+				// Again, If nothing was written, call the older deprecated API
+				if (ArchiveOffset == Ar.Tell())
+				{
+					PRAGMA_DISABLE_DEPRECATION_WARNINGS
+					Codec->PopulateDDCKey(Ar);
+					PRAGMA_ENABLE_DEPRECATION_WARNINGS
+				}
+			}
+
+			NumValidCodecs++;
+		}
+	}
+
+	if (NumValidCodecs == 0)
+	{
+		static FString NoCodecString(TEXT("<Missing Codec>"));
+		Ar << NoCodecString;
+	}
+}
+
 void UAnimBoneCompressionSettings::PopulateDDCKey(const UAnimSequenceBase& AnimSeq, FArchive& Ar)
 {
 	Ar << ErrorThreshold;
@@ -284,11 +325,12 @@ void UAnimBoneCompressionSettings::PopulateDDCKey(const UAnimSequenceBase& AnimS
 		if (Codec != nullptr)
 		{
 			const int64 archiveOffset = Ar.Tell();
+			PRAGMA_DISABLE_DEPRECATION_WARNINGS
 			Codec->PopulateDDCKey(AnimSeq, Ar);
+			PRAGMA_ENABLE_DEPRECATION_WARNINGS
 
 			if (archiveOffset == Ar.Tell())
 			{
-				// If nothing was written, perhaps the codec implements the old deprecated API, call it just in case
 				PRAGMA_DISABLE_DEPRECATION_WARNINGS
 				Codec->PopulateDDCKey(Ar);
 				PRAGMA_ENABLE_DEPRECATION_WARNINGS
diff --git a/Engine/Source/Runtime/Engine/Private/Animation/AnimCompress.cpp b/Engine/Source/Runtime/Engine/Private/Animation/AnimCompress.cpp
--- a/Engine/Source/Runtime/Engine/Private/Animation/AnimCompress.cpp
+++ b/Engine/Source/Runtime/Engine/Private/Animation/AnimCompress.cpp
@@ -705,9 +705,9 @@ void UAnimCompress::DecompressBone(FAnimSequenceDecompressionContext& DecompCont
 
 #if WITH_EDITOR
 
-void UAnimCompress::PopulateDDCKey(const UAnimSequenceBase& AnimSeq, FArchive& Ar)
+void UAnimCompress::PopulateDDCKey(const UE::Anim::Compression::FAnimDDCKeyArgs& KeyArgs, FArchive& Ar)
 {
-	Super::PopulateDDCKey(AnimSeq, Ar);
+	Super::PopulateDDCKey(KeyArgs, Ar);
 
 	uint8 TCF, RCF, SCF;
 	TCF = (uint8)TranslationCompressionFormat.GetValue();
diff --git a/Engine/Source/Runtime/Engine/Private/Animation/AnimCompress_PerTrackCompression.cpp b/Engine/Source/Runtime/Engine/Private/Animation/AnimCompress_PerTrackCompression.cpp
--- a/Engine/Source/Runtime/Engine/Private/Animation/AnimCompress_PerTrackCompression.cpp
+++ b/Engine/Source/Runtime/Engine/Private/Animation/AnimCompress_PerTrackCompression.cpp
@@ -1596,9 +1596,9 @@ void WriteEnumArrayToKey(FArchive& Ar, TArray<TEnumAsByte<enum AnimationCompress
 		Ar << Val;
 	}
 }
-void UAnimCompress_PerTrackCompression::PopulateDDCKey(const UAnimSequenceBase& AnimSeq, FArchive& Ar)
+void UAnimCompress_PerTrackCompression::PopulateDDCKey(const UE::Anim::Compression::FAnimDDCKeyArgs& KeyArgs, FArchive& Ar)
 {
-	Super::PopulateDDCKey(AnimSeq, Ar);
+	Super::PopulateDDCKey(KeyArgs, Ar);
 
 	Ar << MaxZeroingThreshold;
 	Ar << MaxPosDiffBitwise;
diff --git a/Engine/Source/Runtime/Engine/Private/Animation/AnimCompress_RemoveEverySecondKey.cpp b/Engine/Source/Runtime/Engine/Private/Animation/AnimCompress_RemoveEverySecondKey.cpp
--- a/Engine/Source/Runtime/Engine/Private/Animation/AnimCompress_RemoveEverySecondKey.cpp
+++ b/Engine/Source/Runtime/Engine/Private/Animation/AnimCompress_RemoveEverySecondKey.cpp
@@ -61,9 +61,9 @@ bool UAnimCompress_RemoveEverySecondKey::DoReduction(const FCompressibleAnimData
 	return true;
 }
 
-void UAnimCompress_RemoveEverySecondKey::PopulateDDCKey(const UAnimSequenceBase& AnimSeq, FArchive& Ar)
+void UAnimCompress_RemoveEverySecondKey::PopulateDDCKey(const UE::Anim::Compression::FAnimDDCKeyArgs& KeyArgs, FArchive& Ar)
 {
-	Super::PopulateDDCKey(AnimSeq, Ar);
+	Super::PopulateDDCKey(KeyArgs, Ar);
 	Ar << MinKeys;
 	bool bVal = bStartAtSecondKey;
 	Ar << bVal;
diff --git a/Engine/Source/Runtime/Engine/Private/Animation/AnimCompress_RemoveLinearKeys.cpp b/Engine/Source/Runtime/Engine/Private/Animation/AnimCompress_RemoveLinearKeys.cpp
--- a/Engine/Source/Runtime/Engine/Private/Animation/AnimCompress_RemoveLinearKeys.cpp
+++ b/Engine/Source/Runtime/Engine/Private/Animation/AnimCompress_RemoveLinearKeys.cpp
@@ -1107,9 +1107,9 @@ bool UAnimCompress_RemoveLinearKeys::DoReduction(const FCompressibleAnimData& Co
 	return true;
 }
 
-void UAnimCompress_RemoveLinearKeys::PopulateDDCKey(const UAnimSequenceBase& AnimSeq, FArchive& Ar)
+void UAnimCompress_RemoveLinearKeys::PopulateDDCKey(const UE::Anim::Compression::FAnimDDCKeyArgs& KeyArgs, FArchive& Ar)
 {
-	Super::PopulateDDCKey(AnimSeq, Ar);
+	Super::PopulateDDCKey(KeyArgs, Ar);
 	Ar << MaxPosDiff;
 	Ar << MaxAngleDiff;
 	Ar << MaxScaleDiff;
diff --git a/Engine/Source/Runtime/Engine/Private/Animation/AnimCompress_RemoveTrivialKeys.cpp b/Engine/Source/Runtime/Engine/Private/Animation/AnimCompress_RemoveTrivialKeys.cpp
--- a/Engine/Source/Runtime/Engine/Private/Animation/AnimCompress_RemoveTrivialKeys.cpp
+++ b/Engine/Source/Runtime/Engine/Private/Animation/AnimCompress_RemoveTrivialKeys.cpp
@@ -55,9 +55,9 @@ bool UAnimCompress_RemoveTrivialKeys::DoReduction(const FCompressibleAnimData& C
 	return true;
 }
 
-void UAnimCompress_RemoveTrivialKeys::PopulateDDCKey(const UAnimSequenceBase& AnimSeq, FArchive& Ar)
+void UAnimCompress_RemoveTrivialKeys::PopulateDDCKey(const UE::Anim::Compression::FAnimDDCKeyArgs& KeyArgs, FArchive& Ar)
 {
-	Super::PopulateDDCKey(AnimSeq, Ar);
+	Super::PopulateDDCKey(KeyArgs, Ar);
 	Ar << MaxPosDiff;
 	Ar << MaxAngleDiff;
 	Ar << MaxScaleDiff;
diff --git a/Engine/Source/Runtime/Engine/Private/Animation/AnimCompressionTypes.cpp b/Engine/Source/Runtime/Engine/Private/Animation/AnimCompressionTypes.cpp
--- a/Engine/Source/Runtime/Engine/Private/Animation/AnimCompressionTypes.cpp
+++ b/Engine/Source/Runtime/Engine/Private/Animation/AnimCompressionTypes.cpp
@@ -302,7 +302,7 @@ void FCompressibleAnimData::WriteCompressionDataToJSON(TArrayView<FName> Origina
 	}
 }
 
-FCompressibleAnimData::FCompressibleAnimData(class UAnimSequence* InSeq, const bool bPerformStripping)
+FCompressibleAnimData::FCompressibleAnimData(class UAnimSequence* InSeq, const bool bPerformStripping, const ITargetPlatform* InTargetPlatform)
 	: CurveCompressionSettings(InSeq->CurveCompressionSettings)
 	, BoneCompressionSettings(InSeq->BoneCompressionSettings)
 	, Interpolation(InSeq->Interpolation)
@@ -317,6 +317,7 @@ FCompressibleAnimData::FCompressibleAnimData(class UAnimSequence* InSeq, const b
 	, Name(InSeq->GetName())
 	, FullName(InSeq->GetFullName())
 	, AnimFName(InSeq->GetFName())
+	, TargetPlatform(InTargetPlatform)
 {
 #if WITH_EDITOR
 	TRACE_CPUPROFILER_EVENT_SCOPE(FCompressibleAnimData::FCompressibleAnimData);
@@ -499,7 +500,7 @@ FCompressibleAnimData::FCompressibleAnimData(class UAnimSequence* InSeq, const b
 #endif
 }
 
-FCompressibleAnimData::FCompressibleAnimData(UAnimBoneCompressionSettings* InBoneCompressionSettings, UAnimCurveCompressionSettings* InCurveCompressionSettings, USkeleton* InSkeleton, EAnimInterpolationType InInterpolation, float InSequenceLength, int32 InNumberOfKeys)
+FCompressibleAnimData::FCompressibleAnimData(UAnimBoneCompressionSettings* InBoneCompressionSettings, UAnimCurveCompressionSettings* InCurveCompressionSettings, USkeleton* InSkeleton, EAnimInterpolationType InInterpolation, float InSequenceLength, int32 InNumberOfKeys, const ITargetPlatform* InTargetPlatform)
 	: CurveCompressionSettings(InCurveCompressionSettings)
 	, BoneCompressionSettings(InBoneCompressionSettings)
 	, Interpolation(InInterpolation)
@@ -510,6 +511,7 @@ FCompressibleAnimData::FCompressibleAnimData(UAnimBoneCompressionSettings* InBon
 	, NumberOfKeys(InNumberOfKeys)
 	, bIsValidAdditive(false)
 	, ErrorThresholdScale(1.f)
+	, TargetPlatform(InTargetPlatform)
 {
 #if WITH_EDITOR
 	RefLocalPoses = InSkeleton->GetRefLocalPoses();
@@ -529,6 +531,7 @@ PRAGMA_ENABLE_DEPRECATION_WARNINGS
 , NumberOfKeys(0)
 , bIsValidAdditive(false)
 , ErrorThresholdScale(1.f)
+, TargetPlatform(nullptr)
 {
 }
 
@@ -1064,3 +1067,18 @@ FArchive& operator<<(FArchive& Ar, FAnimationErrorStats& ErrorStats)
 	Ar << ErrorStats.MaxErrorBone;
 	return Ar;
 }
+
+#if WITH_EDITORONLY_DATA
+UE::Anim::Compression::FAnimDDCKeyArgs::FAnimDDCKeyArgs(const UAnimSequenceBase& AnimSequence)
+	: AnimSequence(AnimSequence)
+	, TargetPlatform(nullptr)
+{
+}
+
+UE::Anim::Compression::FAnimDDCKeyArgs::FAnimDDCKeyArgs(const UAnimSequenceBase& AnimSequence, const ITargetPlatform* TargetPlatform)
+	: AnimSequence(AnimSequence)
+	, TargetPlatform(TargetPlatform)
+{
+}
+#endif
+
diff --git a/Engine/Source/Runtime/Engine/Private/Animation/AnimSequence.cpp b/Engine/Source/Runtime/Engine/Private/Animation/AnimSequence.cpp
--- a/Engine/Source/Runtime/Engine/Private/Animation/AnimSequence.cpp
+++ b/Engine/Source/Runtime/Engine/Private/Animation/AnimSequence.cpp
@@ -37,6 +37,7 @@
 #include "Logging/MessageLog.h"
 #include "DerivedDataCacheInterface.h"
 #include "Interfaces/ITargetPlatform.h"
+#include "Interfaces/ITargetPlatformManagerModule.h"
 #include "Animation/AnimCompressionDerivedData.h"
 #include "Animation/AnimCompressionDerivedDataPublic.h"
 #include "UObject/UObjectThreadContext.h"
@@ -269,7 +270,7 @@ void OnCVarsChanged()
 
 FAutoConsoleVariableSink AnimationCVarSink(FConsoleCommandDelegate::CreateStatic(&OnCVarsChanged));
 
-FString GetAnimSequenceSpecificCacheKeySuffix(const UAnimSequence& Seq, bool bPerformStripping, float CompressionErrorThresholdScale)
+FString GetAnimSequenceSpecificCacheKeySuffix(const UAnimSequence& Seq, bool bPerformStripping, float CompressionErrorThresholdScale, const ITargetPlatform* TargetPlatform)
 {
 	//Make up our content key consisting of:
 	//	* Global animation compression version
@@ -290,7 +291,7 @@ FString GetAnimSequenceSpecificCacheKeySuffix(const UAnimSequence& Seq, bool bPe
 
 	ArcToHexString.Ar << CompressionErrorThresholdScale;
 	ArcToHexString.Ar << bPerformStripping;
-	Seq.BoneCompressionSettings->PopulateDDCKey(Seq, ArcToHexString.Ar);
+	Seq.BoneCompressionSettings->PopulateDDCKey(UE::Anim::Compression::FAnimDDCKeyArgs(Seq, TargetPlatform), ArcToHexString.Ar);
 	Seq.CurveCompressionSettings->PopulateDDCKey(ArcToHexString.Ar);
 
 	FString Ret = FString::Printf(TEXT("%i_%s%s%s_%c%c%i_%s_%s_%i"),
@@ -316,9 +317,10 @@ FString GetAnimSequenceSpecificCacheKeySuffix(const UAnimSequence& Seq, bool bPe
 
 /////////////////////////////////////////////////////
 // FRequestAnimCompressionParams
-FRequestAnimCompressionParams::FRequestAnimCompressionParams(bool bInAsyncCompression, bool bInAllowAlternateCompressor, bool bInOutput)
+FRequestAnimCompressionParams::FRequestAnimCompressionParams(bool bInAsyncCompression, bool bInAllowAlternateCompressor, bool bInOutput, const ITargetPlatform* InTargetPlatform)
 	: bAsyncCompression(bInAsyncCompression)
 	, CompressContext(MakeShared<FAnimCompressContext>(bInAllowAlternateCompressor, bInOutput))
+	, TargetPlatform(InTargetPlatform)
 {
 	InitFrameStrippingFromCVar();
 }
@@ -326,6 +328,7 @@ FRequestAnimCompressionParams::FRequestAnimCompressionParams(bool bInAsyncCompre
 FRequestAnimCompressionParams::FRequestAnimCompressionParams(bool bInAsyncCompression, TSharedPtr<FAnimCompressContext> InCompressContext)
 	: bAsyncCompression(bInAsyncCompression)
 	, CompressContext(InCompressContext)
+	, TargetPlatform(nullptr)
 {
 	InitFrameStrippingFromCVar();
 }
@@ -336,23 +339,26 @@ void FRequestAnimCompressionParams::InitFrameStrippingFromCVar()
 	bPerformFrameStrippingOnOddNumberedFrames = (GPerformFrameStrippingOddFramedAnimations == 1);
 }
 
-void FRequestAnimCompressionParams::InitFrameStrippingFromPlatform(const class ITargetPlatform* TargetPlatform)
+void FRequestAnimCompressionParams::InitFrameStrippingFromPlatform()
 {
 #if WITH_EDITOR
 	bPerformFrameStripping = false;
 
-	if (UDeviceProfile* DeviceProfile = UDeviceProfileManager::Get().FindProfile(TargetPlatform->IniPlatformName()))
+	if (TargetPlatform)
 	{
-		int32 CVarPlatformFrameStrippingValue = 0;
-		if (DeviceProfile->GetConsolidatedCVarValue(StripFrameCVarName, CVarPlatformFrameStrippingValue))
+		if (UDeviceProfile* DeviceProfile = UDeviceProfileManager::Get().FindProfile(TargetPlatform->IniPlatformName()))
 		{
-			bPerformFrameStripping = CVarPlatformFrameStrippingValue == 1;
-		}
+			int32 CVarPlatformFrameStrippingValue = 0;
+			if (DeviceProfile->GetConsolidatedCVarValue(StripFrameCVarName, CVarPlatformFrameStrippingValue))
+			{
+				bPerformFrameStripping = CVarPlatformFrameStrippingValue == 1;
+			}
 
-		int32 CVarPlatformOddAnimFrameStrippingValue = 0;
-		if (DeviceProfile->GetConsolidatedCVarValue(OddFrameStripStrippingCVarName, CVarPlatformOddAnimFrameStrippingValue))
-		{
-			bPerformFrameStrippingOnOddNumberedFrames = CVarPlatformOddAnimFrameStrippingValue == 1;
+			int32 CVarPlatformOddAnimFrameStrippingValue = 0;
+			if (DeviceProfile->GetConsolidatedCVarValue(OddFrameStripStrippingCVarName, CVarPlatformOddAnimFrameStrippingValue))
+			{
+				bPerformFrameStrippingOnOddNumberedFrames = CVarPlatformOddAnimFrameStrippingValue == 1;
+			}
 		}
 	}
 #endif
@@ -893,8 +899,8 @@ void UAnimSequence::PreSave(FObjectPreSaveContext ObjectSaveContext)
 	if (TargetPlatform)
 	{
 		// Update compressed data for platform
-		FRequestAnimCompressionParams Params(false, false, false);
-		Params.InitFrameStrippingFromPlatform(TargetPlatform);
+		FRequestAnimCompressionParams Params(false, false, false, TargetPlatform);
+		Params.InitFrameStrippingFromPlatform();
 		RequestAnimCompression(Params);
 	}
 
@@ -1822,9 +1828,9 @@ bool UAnimSequence::ShouldPerformStripping(const bool bPerformFrameStripping, co
 	return bStripCandidate && bShouldPerformFrameStripping;
 }
 
-FString UAnimSequence::GetDDCCacheKeySuffix(const bool bPerformStripping) const
+FString UAnimSequence::GetDDCCacheKeySuffix(const bool bPerformStripping, const ITargetPlatform* TargetPlatform) const
 {
-	return GetAnimSequenceSpecificCacheKeySuffix(*this, bPerformStripping, CompressionErrorThresholdScale);
+	return GetAnimSequenceSpecificCacheKeySuffix(*this, bPerformStripping, CompressionErrorThresholdScale, TargetPlatform);
 }
 #endif
 
@@ -1910,7 +1916,8 @@ void UAnimSequence::RequestAnimCompression(FRequestAnimCompressionParams Params)
 	const double CompressionStartTime = FPlatformTime::Seconds();
 
 	const bool bPerformStripping = ShouldPerformStripping(Params.bPerformFrameStripping, Params.bPerformFrameStrippingOnOddNumberedFrames);
-	const FString AssetDDCKey = GetDDCCacheKeySuffix(bPerformStripping);
+	const FString AssetDDCKey = GetDDCCacheKeySuffix(bPerformStripping, Params.TargetPlatform);
+	RequestedCompression_DDCKey = FString(AssetDDCKey);
 
 	bool bCompressedDataFromDDC = false;
 
@@ -1924,7 +1931,7 @@ void UAnimSequence::RequestAnimCompression(FRequestAnimCompressionParams Params)
 		const bool bSkipDDC = GSkipDDC == 1;
 
 		// Data does not exist, need to build it.
-		FCompressibleAnimRef CompressibleData = MakeShared<FCompressibleAnimData, ESPMode::ThreadSafe>(this, bPerformStripping);
+		FCompressibleAnimRef CompressibleData = MakeShared<FCompressibleAnimData, ESPMode::ThreadSafe>(this, bPerformStripping, Params.TargetPlatform);
 		AnimCompressor->SetCompressibleData(CompressibleData);
 
 		if (bSkipDDC || (CompressCommandletVersion == INDEX_NONE))
@@ -1978,7 +1985,7 @@ void UAnimSequence::RequestAnimCompression(FRequestAnimCompressionParams Params)
 
 void UAnimSequence::ApplyCompressedData(const FString& DataCacheKeySuffix, const bool bPerformFrameStripping, const TArray<uint8>& Data)
 {
-	if (GetDDCCacheKeySuffix(bPerformFrameStripping) == DataCacheKeySuffix)
+	if (RequestedCompression_DDCKey == DataCacheKeySuffix)
 	{
 		ApplyCompressedData(Data);
 	}
@@ -2032,7 +2039,14 @@ void UAnimSequence::OnAnimModelLoaded()
 
 		if (FAnimationUtils::GetForcedRecompressionSetting())
 		{
-			RequestAnimCompression(FRequestAnimCompressionParams(true, false, false));
+			const ITargetPlatform* TargetPlatform = nullptr;
+			ITargetPlatformManagerModule* TargetPlatformManager = GetTargetPlatformManager();
+			if (TargetPlatformManager)
+			{
+				TargetPlatform = TargetPlatformManager->GetRunningTargetPlatform();
+			}
+
+			RequestAnimCompression(FRequestAnimCompressionParams(true, false, false, TargetPlatform));
 		}
 
 		if ( DataModelInterface->GetNumberOfKeys() == 0 && DataModelInterface->GetNumberOfFloatCurves() == 0 )
diff --git a/Engine/Source/Runtime/Engine/Private/Animation/AnimStreamable.cpp b/Engine/Source/Runtime/Engine/Private/Animation/AnimStreamable.cpp
--- a/Engine/Source/Runtime/Engine/Private/Animation/AnimStreamable.cpp
+++ b/Engine/Source/Runtime/Engine/Private/Animation/AnimStreamable.cpp
@@ -601,7 +601,7 @@ void UAnimStreamable::RequestCompressedData(const ITargetPlatform* Platform)
 
 	PlatformData.Chunks.AddDefaulted(NumChunks);
 
-	const FString BaseDDCKey = GetBaseDDCKey(NumChunks);
+	const FString BaseDDCKey = GetBaseDDCKey(NumChunks, Platform);
 
 	const bool bInAllowAlternateCompressor = false;
 	const bool bInOutput				   = false;
@@ -616,7 +616,7 @@ void UAnimStreamable::RequestCompressedData(const ITargetPlatform* Platform)
 		const uint32 FrameStart = ChunkIndex * FramesPerChunk;
 		const uint32 FrameEnd = bLastChunk ? NumFramesToChunk : (ChunkIndex + 1) * FramesPerChunk;
 
-		RequestCompressedDataForChunk(ChunkDDCKey, PlatformData.Chunks[ChunkIndex], ChunkIndex, FrameStart, FrameEnd, CompressContext);
+		RequestCompressedDataForChunk(ChunkDDCKey, PlatformData.Chunks[ChunkIndex], ChunkIndex, FrameStart, FrameEnd, CompressContext, Platform);
 	}
 
 	if (bIsRunningPlatform)
@@ -661,7 +661,7 @@ void MakeKeyChunk(const TArray<KeyType>& SrcKeys, TArray<KeyType>& DestKeys, int
 	}
 }
 
-void UAnimStreamable::RequestCompressedDataForChunk(const FString& ChunkDDCKey, FAnimStreamableChunk& Chunk, const int32 ChunkIndex, const uint32 FrameStart, const uint32 FrameEnd, TSharedRef<FAnimCompressContext> CompressContext)
+void UAnimStreamable::RequestCompressedDataForChunk(const FString& ChunkDDCKey, FAnimStreamableChunk& Chunk, const int32 ChunkIndex, const uint32 FrameStart, const uint32 FrameEnd, TSharedRef<FAnimCompressContext> CompressContext, const ITargetPlatform* Platform)
 {
 	// Need to unify with Anim Sequence!
 
@@ -688,7 +688,7 @@ void UAnimStreamable::RequestCompressedDataForChunk(const FString& ChunkDDCKey,
 		}
 		else
 		{
-			FCompressibleAnimRef CompressibleData = MakeShared<FCompressibleAnimData, ESPMode::ThreadSafe>(BoneCompressionSettings, CurveCompressionSettings, GetSkeleton(), Interpolation, Chunk.SequenceLength, ChunkNumFrames+1);
+			FCompressibleAnimRef CompressibleData = MakeShared<FCompressibleAnimData, ESPMode::ThreadSafe>(BoneCompressionSettings, CurveCompressionSettings, GetSkeleton(), Interpolation, Chunk.SequenceLength, ChunkNumFrames+1, Platform);
 
 			const TArray<FBoneAnimationTrack>& BoneAnimationTracks = DataModelInterface->GetBoneAnimationTracks();
 			CompressibleData->TrackToSkeletonMapTable.Empty();
@@ -762,7 +762,7 @@ void UAnimStreamable::UpdateRawData()
 	}
 }
 
-FString UAnimStreamable::GetBaseDDCKey(uint32 NumChunks) const
+FString UAnimStreamable::GetBaseDDCKey(uint32 NumChunks, const ITargetPlatform* TargetPlatform) const
 {
 	//Make up our content key consisting of:
 	//  * Streaming Anim Chunk logic version
@@ -775,7 +775,7 @@ FString UAnimStreamable::GetBaseDDCKey(uint32 NumChunks) const
 	FArcToHexString ArcToHexString;
 
 	ArcToHexString.Ar << NumChunks;
-	BoneCompressionSettings->PopulateDDCKey(*this, ArcToHexString.Ar);
+	BoneCompressionSettings->PopulateDDCKey(UE::Anim::Compression::FAnimDDCKeyArgs(*this, TargetPlatform), ArcToHexString.Ar);
 	CurveCompressionSettings->PopulateDDCKey(ArcToHexString.Ar);
 
 	FString Ret = FString::Printf(TEXT("%s%s%s%s_%s"),
diff --git a/Engine/Source/Runtime/Engine/Public/Animation/AnimCompressionTypes.h b/Engine/Source/Runtime/Engine/Public/Animation/AnimCompressionTypes.h
--- a/Engine/Source/Runtime/Engine/Public/Animation/AnimCompressionTypes.h
+++ b/Engine/Source/Runtime/Engine/Public/Animation/AnimCompressionTypes.h
@@ -192,9 +192,9 @@ struct ENGINE_API FCompressibleAnimData
 public:
 	FCompressibleAnimData();
 
-	FCompressibleAnimData(UAnimBoneCompressionSettings* InBoneCompressionSettings, UAnimCurveCompressionSettings* InCurveCompressionSettings, USkeleton* InSkeleton, EAnimInterpolationType InInterpolation, float InSequenceLength, int32 InNumberOfKeys);
+	FCompressibleAnimData(UAnimBoneCompressionSettings* InBoneCompressionSettings, UAnimCurveCompressionSettings* InCurveCompressionSettings, USkeleton* InSkeleton, EAnimInterpolationType InInterpolation, float InSequenceLength, int32 InNumberOfKeys, const ITargetPlatform* InTargetPlatform);
 
-	FCompressibleAnimData(class UAnimSequence* InSeq, const bool bPerformStripping);
+	FCompressibleAnimData(class UAnimSequence* InSeq, const bool bPerformStripping, const ITargetPlatform* InTargetPlatform);
 
 	FCompressibleAnimData(const FCompressibleAnimData&);
 	FCompressibleAnimData& operator=(const FCompressibleAnimData&);
@@ -242,6 +242,8 @@ public:
 
 	FCancelCompressionSignal IsCancelledSignal;
 
+	const ITargetPlatform* TargetPlatform;
+
 	static int32 GetApproxRawDataArraySize(const TArray<FRawAnimSequenceTrack>& AnimData)
 	{
 		int32 Total = sizeof(FRawAnimSequenceTrack) * AnimData.Num();
@@ -811,6 +813,19 @@ struct FRootMotionReset
 	}
 };
 
+#if WITH_EDITOR
+namespace UE::Anim::Compression {
+	struct ENGINE_API FAnimDDCKeyArgs
+	{
+		FAnimDDCKeyArgs(const UAnimSequenceBase& AnimSequence);
+		FAnimDDCKeyArgs(const UAnimSequenceBase& AnimSequence, const ITargetPlatform* TargetPlatform);
+
+		const UAnimSequenceBase& AnimSequence;
+		const ITargetPlatform* TargetPlatform;
+	};
+}
+#endif // WITH_EDITOR
+
 UE_DEPRECATED(5.1, "Signature of DecompressPose has been deprecated, use UE::Anim::Decompression::DecompressPose instead")
 extern void DecompressPose(FCompactPose& OutPose,
 							const FCompressedAnimSequence& CompressedData,
