[Engine]
* MaterialLayersFunctions - Added #if WITH_EDITOR around declared methods that only are implemented when WITH_EDITOR is set

(This caused errors when building import-header-units)

diff --git a/Engine/Source/Runtime/Engine/Classes/Materials/MaterialLayersFunctions.h b/Engine/Source/Runtime/Engine/Classes/Materials/MaterialLayersFunctions.h
--- a/Engine/Source/Runtime/Engine/Classes/Materials/MaterialLayersFunctions.h
+++ b/Engine/Source/Runtime/Engine/Classes/Materials/MaterialLayersFunctions.h
@@ -28,6 +28,7 @@ struct FMaterialLayersFunctionsID
 	TArray<FGuid> BlendIDs;
 	TArray<bool> LayerStates;
 
+	#if WITH_EDITOR
 	bool operator==(const FMaterialLayersFunctionsID& Reference) const;
 	inline bool operator!=(const FMaterialLayersFunctionsID& Reference) const { return !operator==(Reference); }
 
@@ -43,6 +44,7 @@ struct FMaterialLayersFunctionsID
 
 	//TODO: Investigate whether this is really required given it is only used by FMaterialShaderMapId AND that one also uses UpdateHash
 	void AppendKeyString(FString& KeyString) const;
+	#endif
 };
 
 USTRUCT()
