fix FTextureSourceData::Init setting wrong source gamma

diff --git a/Engine/Source/Runtime/Engine/Private/Texture.cpp b/Engine/Source/Runtime/Engine/Private/Texture.cpp
--- a/Engine/Source/Runtime/Engine/Private/Texture.cpp
+++ b/Engine/Source/Runtime/Engine/Private/Texture.cpp
@@ -2443,7 +2443,12 @@ EGammaSpace FTextureSource::GetGammaSpace(int LayerIndex) const
 	// TextureSource does not know its own gamma, but its owning Texture does :
 	if ( Owner != nullptr )
 	{
-		return Owner->GetGammaSpace();
+		// same as Owner->GetGammaSpace , but uses LayerFormatSettings for SRGB flag
+		FTextureFormatSettings FormatSettings;
+		Owner->GetLayerFormatSettings(LayerIndex, FormatSettings);
+
+		EGammaSpace GammaSpace = FormatSettings.SRGB ? ( Owner->bUseLegacyGamma ? EGammaSpace::Pow22 : EGammaSpace::sRGB ) : EGammaSpace::Linear;
+		return GammaSpace;
 	}
 	else
 	{
diff --git a/Engine/Source/Runtime/Engine/Private/TextureDerivedDataTask.cpp b/Engine/Source/Runtime/Engine/Private/TextureDerivedDataTask.cpp
--- a/Engine/Source/Runtime/Engine/Private/TextureDerivedDataTask.cpp
+++ b/Engine/Source/Runtime/Engine/Private/TextureDerivedDataTask.cpp
@@ -195,10 +195,7 @@ void FTextureSourceData::Init(UTexture& InTexture, TextureMipGenSettings InMipGe
 
 		LayerData->ImageFormat = FImageCoreUtils::ConvertToRawImageFormat( InTexture.Source.GetFormat(LayerIndex) );
 
-		FTextureFormatSettings FormatSettings;
-		InTexture.GetLayerFormatSettings(LayerIndex, FormatSettings);
-
-		LayerData->GammaSpace = FormatSettings.SRGB ? (InTexture.bUseLegacyGamma ? EGammaSpace::Pow22 : EGammaSpace::sRGB) : EGammaSpace::Linear;
+		LayerData->SourceGammaSpace = InTexture.Source.GetGammaSpace(LayerIndex);
 	}
 
 	Blocks.Reserve(NumBlocks);
@@ -291,7 +288,7 @@ void FTextureSourceData::GetSourceMips(FTextureSource& Source, IImageWrapperModu
 						FImage* SourceMip = new(BlockData.MipsPerLayer[LayerIndex]) FImage(
 							MipSizeX, MipSizeY, MipSizeZ,
 							LayerData.ImageFormat,
-							LayerData.GammaSpace
+							LayerData.SourceGammaSpace
 						);
 
 						if (!ScopedMipData.GetMipData(SourceMip->RawData, BlockIndex, LayerIndex, MipIndex))
diff --git a/Engine/Source/Runtime/Engine/Private/TextureDerivedDataTask.h b/Engine/Source/Runtime/Engine/Private/TextureDerivedDataTask.h
--- a/Engine/Source/Runtime/Engine/Private/TextureDerivedDataTask.h
+++ b/Engine/Source/Runtime/Engine/Private/TextureDerivedDataTask.h
@@ -61,8 +61,8 @@ ENUM_CLASS_FLAGS(ETextureCacheFlags);
 // Everything required to get the texture source data.
 struct FTextureSourceLayerData
 {
-	ERawImageFormat::Type ImageFormat;
-	EGammaSpace GammaSpace;
+	ERawImageFormat::Type ImageFormat = ERawImageFormat::Invalid;
+	EGammaSpace SourceGammaSpace = EGammaSpace::Invalid;
 };
 
 struct FTextureSourceBlockData
