[NO NON-PII COMMIT MESSAGE]

diff --git a/Engine/Source/Runtime/Engine/Classes/Kismet/KismetSystemLibrary.h b/Engine/Source/Runtime/Engine/Classes/Kismet/KismetSystemLibrary.h
--- a/Engine/Source/Runtime/Engine/Classes/Kismet/KismetSystemLibrary.h
+++ b/Engine/Source/Runtime/Engine/Classes/Kismet/KismetSystemLibrary.h
@@ -483,6 +483,15 @@ class ENGINE_API UKismetSystemLibrary : public UBlueprintFunctionLibrary
 	UFUNCTION(BlueprintCallable, Category="Development",meta=(WorldContext="WorldContextObject", CallableWithoutWorldContext))
 	static void ExecuteConsoleCommand(const UObject* WorldContextObject, const FString& Command, class APlayerController* SpecificPlayer = NULL );
 
+	/**
+	 * Attempts to retrieve the value of the specified string console variable, if it exists.
+	 * 
+	 * @param	VariableName	Name of the console variable to find.
+	 * @return	The value if found, empty string otherwise.
+	 */
+	UFUNCTION(BlueprintCallable, Category="Development")
+	static FString GetConsoleVariableStringValue(const FString& VariableName);
+
 	/**
 	 * Attempts to retrieve the value of the specified float console variable, if it exists.
 	 * 
diff --git a/Engine/Source/Runtime/Engine/Private/KismetSystemLibrary.cpp b/Engine/Source/Runtime/Engine/Private/KismetSystemLibrary.cpp
--- a/Engine/Source/Runtime/Engine/Private/KismetSystemLibrary.cpp
+++ b/Engine/Source/Runtime/Engine/Private/KismetSystemLibrary.cpp
@@ -454,6 +454,23 @@ void UKismetSystemLibrary::ExecuteConsoleCommand(const UObject* WorldContextObje
 	}
 }
 
+FString UKismetSystemLibrary::GetConsoleVariableStringValue(const FString& VariableName)
+{
+	FString Value = "";
+	
+	IConsoleVariable* Variable = IConsoleManager::Get().FindConsoleVariable(*VariableName);
+	if (Variable)
+	{
+		Value = Variable->GetString();
+	}
+	else
+	{
+		UE_LOG(LogBlueprintUserMessages, Warning, TEXT("Failed to find console variable '%s'."), *VariableName);
+	}
+	
+	return Value;
+}
+
 float UKismetSystemLibrary::GetConsoleVariableFloatValue(const FString& VariableName)
 {
 	float Value = 0.0f;
