Add BuildActorDescContainerPackage to content bundle paths to be accessed from outside a content bundle editor

diff --git a/Engine/Source/Runtime/Engine/Private/WorldPartition/ContentBundle/ContentBundleEditor.cpp b/Engine/Source/Runtime/Engine/Private/WorldPartition/ContentBundle/ContentBundleEditor.cpp
--- a/Engine/Source/Runtime/Engine/Private/WorldPartition/ContentBundle/ContentBundleEditor.cpp
+++ b/Engine/Source/Runtime/Engine/Private/WorldPartition/ContentBundle/ContentBundleEditor.cpp
@@ -53,7 +53,8 @@ void FContentBundleEditor::DoUninitialize()
 void FContentBundleEditor::DoInjectContent()
 {
 	FString ActorDescContainerPackage;
-	if (BuildContentBundleContainerPackagePath(ActorDescContainerPackage))
+	bool bCreatedContainerPath = ContentBundlePaths::BuildActorDescContainerPackgePath(GetDescriptor()->GetPackageRoot(), GetDescriptor()->GetGuid(), GetInjectedWorld()->GetPackage()->GetName(), ActorDescContainerPackage);
+	if (bCreatedContainerPath)
 	{
 		UnsavedActorMonitor = NewObject<UContentBundleUnsavedActorMonitor>(GetTransientPackage(), NAME_None, RF_Transactional);
 		UnsavedActorMonitor->Initialize(*this);
@@ -87,6 +88,7 @@ void FContentBundleEditor::DoInjectContent()
 	}
 	else
 	{
+		UE_LOG(LogContentBundle, Error, TEXT("[CB: %s] Failed to build Container Package Path using %s"), *GetDescriptor()->GetDisplayName(), *GetInjectedWorld()->GetPackage()->GetName());
 		SetStatus(EContentBundleStatus::FailedToInject);
 	}
 
@@ -470,29 +472,6 @@ void FContentBundleEditor::BroadcastChanged()
 	}
 }
 
-bool FContentBundleEditor::BuildContentBundleContainerPackagePath(FString& ContainerPackagePath) const
-{
-	FString PackageRoot, PackagePath, PackageName;
-	FString LongPackageName = GetInjectedWorld()->GetPackage()->GetName();
-	if (FPackageName::SplitLongPackageName(LongPackageName, PackageRoot, PackagePath, PackageName))
-	{
-		TStringBuilderWithBuffer<TCHAR, NAME_SIZE> PluginLeveldPackagePath;
-		PluginLeveldPackagePath += TEXT("/");
-		PluginLeveldPackagePath += GetDescriptor()->GetPackageRoot();
-		PluginLeveldPackagePath += TEXT("/ContentBundle/");
-		PluginLeveldPackagePath += GetDescriptor()->GetGuid().ToString();
-		PluginLeveldPackagePath += TEXT("/");
-		PluginLeveldPackagePath += PackagePath;
-		PluginLeveldPackagePath += PackageName;
-
-		ContainerPackagePath = UPackageTools::SanitizePackageName(*PluginLeveldPackagePath);
-		return true;
-	}
-
-	UE_LOG(LogContentBundle, Error, TEXT("[CB: %s] Failed to build Container Package Path using %s"), *GetDescriptor()->GetDisplayName(), *LongPackageName);
-	return false;
-}
-
 UPackage* FContentBundleEditor::CreateActorPackage(const FName& ActorName) const
 {
 	FString ActorPackagePath = ULevel::GetActorPackageName(ActorDescContainer->GetExternalActorPath(), EActorPackagingScheme::Reduced, ActorName.ToString());
diff --git a/Engine/Source/Runtime/Engine/Private/WorldPartition/ContentBundle/ContentBundlePaths.cpp b/Engine/Source/Runtime/Engine/Private/WorldPartition/ContentBundle/ContentBundlePaths.cpp
--- a/Engine/Source/Runtime/Engine/Private/WorldPartition/ContentBundle/ContentBundlePaths.cpp
+++ b/Engine/Source/Runtime/Engine/Private/WorldPartition/ContentBundle/ContentBundlePaths.cpp
@@ -6,9 +6,11 @@
 #include "WorldPartition/ContentBundle/ContentBundleBase.h"
 #include "WorldPartition/ContentBundle/ContentBundleDescriptor.h"
 #include "Misc/Paths.h"
+#include "Misc/PathViews.h"
 
 #if WITH_EDITOR
 #include "String/Find.h"
+#include "PackageTools.h"
 #endif
 
 namespace ContentBundlePaths
@@ -118,4 +120,30 @@ FStringView ContentBundlePaths::GetActorPathRelativeToExternalActors(FStringView
 	return FStringView();
 }
 
+bool ContentBundlePaths::BuildActorDescContainerPackgePath(const FString& InContenBundleMountPoint, const FGuid& InContentBundleGuid, const FString& InLevelPackagePath, FString& OutContainerPackagePath)
+{
+	if (InContenBundleMountPoint.IsEmpty() || !InContentBundleGuid.IsValid())
+	{
+		return false;
+	}
+
+	FString LevelRoot, LevelPackagePath, LevelName;
+	if (FPackageName::SplitLongPackageName(InLevelPackagePath, LevelRoot, LevelPackagePath, LevelName))
+	{
+		TStringBuilderWithBuffer<TCHAR, NAME_SIZE> ContentBundleActorDescContainerPackage;
+		ContentBundleActorDescContainerPackage += TEXT("/");
+		ContentBundleActorDescContainerPackage += InContenBundleMountPoint;
+		ContentBundleActorDescContainerPackage += GetContentBundleFolder();
+		ContentBundleActorDescContainerPackage += InContentBundleGuid.ToString();
+		ContentBundleActorDescContainerPackage += TEXT("/");
+		ContentBundleActorDescContainerPackage += LevelPackagePath;
+		ContentBundleActorDescContainerPackage += LevelName;
+
+		OutContainerPackagePath = UPackageTools::SanitizePackageName(*ContentBundleActorDescContainerPackage);
+		return true;
+	}
+
+	return false;
+}
+
 #endif
\ No newline at end of file
diff --git a/Engine/Source/Runtime/Engine/Public/WorldPartition/ContentBundle/ContentBundleEditor.h b/Engine/Source/Runtime/Engine/Public/WorldPartition/ContentBundle/ContentBundleEditor.h
--- a/Engine/Source/Runtime/Engine/Public/WorldPartition/ContentBundle/ContentBundleEditor.h
+++ b/Engine/Source/Runtime/Engine/Public/WorldPartition/ContentBundle/ContentBundleEditor.h
@@ -75,7 +75,6 @@ private:
 	void OnUnsavedActorDeleted(AActor* Actor);
 
 	void BroadcastChanged();
-	bool BuildContentBundleContainerPackagePath(FString& ContainerPackagePath) const;
 	UPackage* CreateActorPackage(const FName& ActorName) const;
 	FName BuildWorlDataLayersName() const;
 
diff --git a/Engine/Source/Runtime/Engine/Public/WorldPartition/ContentBundle/ContentBundlePaths.h b/Engine/Source/Runtime/Engine/Public/WorldPartition/ContentBundle/ContentBundlePaths.h
--- a/Engine/Source/Runtime/Engine/Public/WorldPartition/ContentBundle/ContentBundlePaths.h
+++ b/Engine/Source/Runtime/Engine/Public/WorldPartition/ContentBundle/ContentBundlePaths.h
@@ -39,6 +39,9 @@ namespace ContentBundlePaths
 	FString GetRelativeLevelFolder(const FContentBundleBase& ContentBundle);
 
 #if WITH_EDITOR
+	// return format is /{MountPoint}/{ExternalActorFolder}/{ContentBundleFolder}/{ContentBundleUID}/{LevelPath}/
+	ENGINE_API bool BuildActorDescContainerPackgePath(const FString& InContenBundleMountPoint, const FGuid& InContentBundleGuid, const FString& InLevelPackagePath, FString& OutContainerPackagePath);
+
 	// return an ExternalActor path following the format : /{MountPoint}/{ExternalActorFolder}/{ContentBundleFolder}/{ContentBundleUID}/{LevelPath}/{ExternalActorPackagePath}
 	ENGINE_API FString MakeExternalActorPackagePath(const FString& ContentBundleExternalActorFolder, const FString& ActorName);
 
