Map checks on map load now validate content bundle

diff --git a/Engine/Source/Runtime/Engine/Private/WorldPartition/WorldPartitionStreamingGeneration.cpp b/Engine/Source/Runtime/Engine/Private/WorldPartition/WorldPartitionStreamingGeneration.cpp
--- a/Engine/Source/Runtime/Engine/Private/WorldPartition/WorldPartitionStreamingGeneration.cpp
+++ b/Engine/Source/Runtime/Engine/Private/WorldPartition/WorldPartitionStreamingGeneration.cpp
@@ -960,7 +960,10 @@ void UWorldPartition::GenerateHLOD(ISourceControlHelper* SourceControlHelper, bo
 void UWorldPartition::CheckForErrors(IStreamingGenerationErrorHandler* ErrorHandler) const
 {
 	const bool bIsChangelistValidation = false;
-	CheckForErrors(ErrorHandler, ActorDescContainer, IsStreamingEnabled(), bIsChangelistValidation);
+	ForEachActorDescContainer([this, ErrorHandler, bIsChangelistValidation](const UActorDescContainer* ActorDescContainer)
+	{
+		CheckForErrors(ErrorHandler, ActorDescContainer, IsStreamingEnabled(), bIsChangelistValidation);
+	});
 }
 
 void UWorldPartition::CheckForErrors(IStreamingGenerationErrorHandler* ErrorHandler, const UActorDescContainer* ActorDescContainer, bool bEnableStreaming, bool bIsChangelistValidation)
