Fix linux editor compile

diff --git a/Engine/Source/Runtime/Engine/Private/WorldPartition/WorldPartitionStreamingGeneration.cpp b/Engine/Source/Runtime/Engine/Private/WorldPartition/WorldPartitionStreamingGeneration.cpp
--- a/Engine/Source/Runtime/Engine/Private/WorldPartition/WorldPartitionStreamingGeneration.cpp
+++ b/Engine/Source/Runtime/Engine/Private/WorldPartition/WorldPartitionStreamingGeneration.cpp
@@ -960,9 +960,9 @@ void UWorldPartition::GenerateHLOD(ISourceControlHelper* SourceControlHelper, bo
 void UWorldPartition::CheckForErrors(IStreamingGenerationErrorHandler* ErrorHandler) const
 {
 	const bool bIsChangelistValidation = false;
-	ForEachActorDescContainer([this, ErrorHandler, bIsChangelistValidation](const UActorDescContainer* ActorDescContainer)
+	ForEachActorDescContainer([this, ErrorHandler, bIsChangelistValidation](const UActorDescContainer* InActorDescContainer)
 	{
-		CheckForErrors(ErrorHandler, ActorDescContainer, IsStreamingEnabled(), bIsChangelistValidation);
+		CheckForErrors(ErrorHandler, InActorDescContainer, IsStreamingEnabled(), bIsChangelistValidation);
 	});
 }
 
