Update ULocalPlayer comments to be more clear as to where these will exist across server/clients.

diff --git a/Engine/Source/Runtime/Engine/Classes/Engine/LocalPlayer.h b/Engine/Source/Runtime/Engine/Classes/Engine/LocalPlayer.h
--- a/Engine/Source/Runtime/Engine/Classes/Engine/LocalPlayer.h
+++ b/Engine/Source/Runtime/Engine/Classes/Engine/LocalPlayer.h
@@ -159,9 +159,9 @@ private:
 };
 
 /**
- *	Each player that is active on the current client has a LocalPlayer. It stays active across maps
- *	There may be several spawned in the case of splitscreen/coop.
- *	There may be 0 spawned on servers.
+ *	Each player that is active on the current client/listen server has a LocalPlayer.
+ *	It stays active across maps, and there may be several spawned in the case of splitscreen/coop.
+ *	There will be 0 spawned on dedicated servers.
  */
 UCLASS(Within=Engine, config=Engine, transient)
 class ENGINE_API ULocalPlayer : public UPlayer
