Fix int32 overflow issues, and additionally protect against int64 overflows by malicious content.

diff --git a/Engine/Plugins/Interchange/Runtime/Source/Import/Private/Texture/InterchangeSlicedTexturePayloadData.cpp b/Engine/Plugins/Interchange/Runtime/Source/Import/Private/Texture/InterchangeSlicedTexturePayloadData.cpp
--- a/Engine/Plugins/Interchange/Runtime/Source/Import/Private/Texture/InterchangeSlicedTexturePayloadData.cpp
+++ b/Engine/Plugins/Interchange/Runtime/Source/Import/Private/Texture/InterchangeSlicedTexturePayloadData.cpp
@@ -2,6 +2,7 @@
 
 #include "Texture/InterchangeSlicedTexturePayloadData.h"
 
+#include "Math/CheckedInt.h"
 #include "CoreMinimal.h"
 
 void UE::Interchange::FImportSlicedImage::Init(int32 InSizeX, int32 InSizeY, int32 InNumSlice, int32 InNumMips, ETextureSourceFormat InFormat, bool InSRGB)
@@ -53,9 +54,10 @@ int64 UE::Interchange::FImportSlicedImage::GetMipSize(int32 InMipIndex) const
 	int32 MipSizeY = FMath::Max<int32>(SizeY >> InMipIndex, 1);
 	int32 MipSizeZ = bIsVolume ? FMath::Max(NumSlice >> InMipIndex, 1) : 1;
 
-	const int64 MipSize = MipSizeX * MipSizeY * MipSizeZ * BytesPerPixel;
+	FCheckedInt64 MipSize = FCheckedInt64(MipSizeX) * MipSizeY * MipSizeZ * BytesPerPixel;
+	checkf(MipSize.IsValid(), TEXT("Invalid (overflowing) mip sizes made it in to FImportSlicedImage::GetMipSize! Check import locations for mip size validation"));
 
-	return MipSize;
+	return MipSize.Get(0);
 }
 
 int64 UE::Interchange::FImportSlicedImage::ComputeBufferSize() const
diff --git a/Engine/Source/Runtime/Engine/Private/Texture.cpp b/Engine/Source/Runtime/Engine/Private/Texture.cpp
--- a/Engine/Source/Runtime/Engine/Private/Texture.cpp
+++ b/Engine/Source/Runtime/Engine/Private/Texture.cpp
@@ -2,6 +2,7 @@
 
 #include "Engine/Texture.h"
 
+#include "Math/CheckedInt.h"
 #include "Misc/App.h"
 #include "Modules/ModuleManager.h"
 #include "Materials/MaterialInterface.h"
@@ -2711,16 +2712,19 @@ int64 FTextureSource::CalcLayerSize(const FTextureSourceBlock& Block, int32 Laye
 {
 	int64 BytesPerPixel = GetBytesPerPixel(LayerIndex);
 
-	int64 TotalSize = 0;
+	// This is used for memory allocation, so use FCheckedInt to rigorously check against overflow issues.
+	FCheckedInt64 TotalSize(0);
 	for (int32 MipIndex = 0; MipIndex < Block.NumMips; ++MipIndex)
 	{
 		int32 MipSizeX = FMath::Max<int32>(Block.SizeX >> MipIndex, 1);
 		int32 MipSizeY = FMath::Max<int32>(Block.SizeY >> MipIndex, 1);
 		int32 MipSizeZ = GetMippedNumSlices(Block.NumSlices,MipIndex);
 
-		TotalSize += MipSizeX * MipSizeY * MipSizeZ * BytesPerPixel;
+		TotalSize += FCheckedInt64(MipSizeX) * MipSizeY * MipSizeZ * BytesPerPixel;
 	}
-	return TotalSize;
+
+	checkf(TotalSize.IsValid(), TEXT("Invalid (overflowing) mip sizes made it in to FTextureSource::CalcLayerSize! Check import locations for mip size validation"));
+	return TotalSize.Get(0);
 }
 
 int64 FTextureSource::CalcMipOffset(int32 BlockIndex, int32 LayerIndex, int32 OffsetToMipIndex) const
@@ -2729,7 +2733,8 @@ int64 FTextureSource::CalcMipOffset(int32 BlockIndex, int32 LayerIndex, int32 Of
 	GetBlock(BlockIndex, Block);
 	check(OffsetToMipIndex < Block.NumMips);
 
-	int64 MipOffset = BlockDataOffsets[BlockIndex];
+	// This is used for memory indexing, so use FCheckedInt to rigorously check against overflow issues.
+	FCheckedInt64 MipOffset(BlockDataOffsets[BlockIndex]);
 
 	// Skip over the initial layers within the tile
 	for (int i = 0; i < LayerIndex; ++i)
@@ -2745,10 +2750,11 @@ int64 FTextureSource::CalcMipOffset(int32 BlockIndex, int32 LayerIndex, int32 Of
 		int32 MipSizeY = FMath::Max<int32>(Block.SizeY >> MipIndex, 1);
 		int32 MipSizeZ = GetMippedNumSlices(Block.NumSlices,MipIndex);
 
-		MipOffset += MipSizeX * MipSizeY * MipSizeZ * BytesPerPixel;
+		MipOffset += FCheckedInt64(MipSizeX) * MipSizeY * MipSizeZ * BytesPerPixel;
 	}
 
-	return MipOffset;
+	checkf(MipOffset.IsValid(), TEXT("Invalid (overflowing) mip sizes made it in to FTextureSource::CalcMipOffset! Check import locations for mip size validation"));
+	return MipOffset.Get(0);
 }
 
 void FTextureSource::UseHashAsGuid()
