fix VT GammaSpace for Intermediate Linear texture format
re-enable ImageCore IsValid check

diff --git a/Engine/Source/Runtime/Engine/Private/VT/VirtualTextureDataBuilder.cpp b/Engine/Source/Runtime/Engine/Private/VT/VirtualTextureDataBuilder.cpp
--- a/Engine/Source/Runtime/Engine/Private/VT/VirtualTextureDataBuilder.cpp
+++ b/Engine/Source/Runtime/Engine/Private/VT/VirtualTextureDataBuilder.cpp
@@ -633,6 +633,7 @@ void FVirtualTextureDataBuilder::BuildTiles(const TArray<FVTSourceTileEntry>& Ti
 			const FTextureSourceBlockData& Block = SourceBlocks[Tile.BlockIndex];
 			const FImage& SourceMip = Block.MipsPerLayer[LayerIndex][Tile.MipIndexInBlock];
 			check(SourceMip.Format == LayerData.ImageFormat);
+			check(SourceMip.GammaSpace == LayerData.GammaSpace);
 
 			FPixelDataRectangle SourceData(LayerData.SourceFormat,
 				SourceMip.SizeX,
@@ -640,7 +641,7 @@ void FVirtualTextureDataBuilder::BuildTiles(const TArray<FVTSourceTileEntry>& Ti
 				const_cast<uint8*>(SourceMip.RawData.GetData()));
 
 			TArray<FImage> TileImages;
-			FImage* TileImage = new(TileImages) FImage(PhysicalTileSize, PhysicalTileSize, LayerData.ImageFormat, BuildSettingsForLayer.GetDestGammaSpace());
+			FImage* TileImage = new(TileImages) FImage(PhysicalTileSize, PhysicalTileSize, LayerData.ImageFormat, LayerData.GammaSpace);
 			FPixelDataRectangle TileData(LayerData.SourceFormat, PhysicalTileSize, PhysicalTileSize, TileImage->RawData.GetData());
 
 			TileData.Clear();
@@ -828,10 +829,18 @@ void FVirtualTextureDataBuilder::BuildSourcePixels(const FTextureSourceData& Sou
 		const FTextureBuildSettings& BuildSettingsForLayer = SettingsPerLayer[LayerIndex];
 		FVirtualTextureSourceLayerData& LayerData = SourceLayers[LayerIndex];
 
+		// Specify the format we are processing to in this step :
+		LayerData.ImageFormat = UE::TextureBuildUtilities::GetVirtualTextureBuildIntermediateFormat(BuildSettingsForLayer);
 		LayerData.GammaSpace = BuildSettingsForLayer.GetDestGammaSpace();
-		LayerData.bHasAlpha = BuildSettingsForLayer.bForceAlphaChannel;
+		// Gamma correction can either be applied in step 1 or step 2 of the VT build
+		//	depending on whether the Intermediate format is U8 or not
+		if ( ! ERawImageFormat::GetFormatNeedsGammaSpace(LayerData.ImageFormat) )
+		{
+			LayerData.GammaSpace = EGammaSpace::Linear;
+		}
 
-		LayerData.ImageFormat = UE::TextureBuildUtilities::GetVirtualTextureBuildIntermediateFormat(BuildSettingsForLayer);
+		// LayerData.bHasAlpha will be updated after the Build to detect if there is actual alpha in the layers
+		LayerData.bHasAlpha = BuildSettingsForLayer.bForceAlphaChannel;
 		
 		LayerData.FormatName = FImageCoreUtils::ConvertToUncompressedTextureFormatName(LayerData.ImageFormat);
 		LayerData.PixelFormat = FImageCoreUtils::GetPixelFormatForRawImageFormat(LayerData.ImageFormat);
@@ -995,8 +1004,9 @@ void FVirtualTextureDataBuilder::BuildSourcePixels(const FTextureSourceData& Sou
 				Image->SizeX = CompressedMip.SizeX;
 				Image->SizeY = CompressedMip.SizeY;
 				Image->Format = LayerData.ImageFormat;
-				Image->GammaSpace = BuildSettingsForLayer.GetDestGammaSpace();
+				Image->GammaSpace = LayerData.GammaSpace;
 				Image->NumSlices = 1;
+				check( Image->IsImageInfoValid() );
 				Image->RawData = MoveTemp(CompressedMip.RawData);
 			}
 
@@ -1117,8 +1127,9 @@ void FVirtualTextureDataBuilder::BuildSourcePixels(const FTextureSourceData& Sou
 				Image->SizeX = CompressedMip.SizeX;
 				Image->SizeY = CompressedMip.SizeY;
 				Image->Format = LayerData.ImageFormat;
-				Image->GammaSpace = BuildSettingsForLayer.GetDestGammaSpace();
+				Image->GammaSpace = LayerData.GammaSpace;
 				Image->NumSlices = 1;
+				check( Image->IsImageInfoValid() );
 				Image->RawData = MoveTemp(CompressedMip.RawData);
 			}
 		}
diff --git a/Engine/Source/Runtime/ImageCore/Private/ImageCore.cpp b/Engine/Source/Runtime/ImageCore/Private/ImageCore.cpp
--- a/Engine/Source/Runtime/ImageCore/Private/ImageCore.cpp
+++ b/Engine/Source/Runtime/ImageCore/Private/ImageCore.cpp
@@ -21,7 +21,7 @@ IMPLEMENT_MODULE(FDefaultModuleImpl, ImageCore);
  */
 static void InitImageStorage(FImage& Image)
 {
-	//check( Image.IsImageInfoValid() );
+	check( Image.IsImageInfoValid() );
 
 	int64 NumBytes = Image.GetImageSizeBytes();
 	Image.RawData.Empty(NumBytes);
