Wrap deprecated properties with WITH_EDITORONLY_DATA to avoid runtime waste

diff --git a/Engine/Plugins/Animation/ControlRig/Source/ControlRig/Private/Rigs/RigControlHierarchy.cpp b/Engine/Plugins/Animation/ControlRig/Source/ControlRig/Private/Rigs/RigControlHierarchy.cpp
--- a/Engine/Plugins/Animation/ControlRig/Source/ControlRig/Private/Rigs/RigControlHierarchy.cpp
+++ b/Engine/Plugins/Animation/ControlRig/Source/ControlRig/Private/Rigs/RigControlHierarchy.cpp
@@ -77,6 +77,7 @@ FRigControl& FRigControlHierarchy::Add(
 
 void FRigControlHierarchy::PostLoad()
 {
+#if WITH_EDITORONLY_DATA
 	for (FRigControl& Control : Controls)
 	{
 		for (int32 ValueType = 0; ValueType <= (int32)ERigControlValueType::Maximum; ValueType++)
@@ -88,4 +89,5 @@ void FRigControlHierarchy::PostLoad()
 			}
 		}
 	}
+#endif
 }
diff --git a/Engine/Plugins/Animation/ControlRig/Source/ControlRig/Private/Rigs/RigHierarchyController.cpp b/Engine/Plugins/Animation/ControlRig/Source/ControlRig/Private/Rigs/RigHierarchyController.cpp
--- a/Engine/Plugins/Animation/ControlRig/Source/ControlRig/Private/Rigs/RigHierarchyController.cpp
+++ b/Engine/Plugins/Animation/ControlRig/Source/ControlRig/Private/Rigs/RigHierarchyController.cpp
@@ -1315,10 +1315,13 @@ TArray<FRigElementKey> URigHierarchyController::ImportFromHierarchyContainer(con
 		}
 
 		FRigControlValue InitialValue = Control.InitialValue;
+
+#if WITH_EDITORONLY_DATA
 		if(!InitialValue.IsValid())
 		{
 			InitialValue.SetFromTransform(InitialValue.Storage_DEPRECATED, Settings.ControlType, Settings.PrimaryAxis);
 		}
+#endif
 		
 		const FRigElementKey Key = AddControl(
 			Control.Name,
diff --git a/Engine/Plugins/Animation/ControlRig/Source/ControlRig/Public/ControlRig.h b/Engine/Plugins/Animation/ControlRig/Source/ControlRig/Public/ControlRig.h
--- a/Engine/Plugins/Animation/ControlRig/Source/ControlRig/Public/ControlRig.h
+++ b/Engine/Plugins/Animation/ControlRig/Source/ControlRig/Public/ControlRig.h
@@ -520,8 +520,10 @@ protected:
 	UPROPERTY()
 	TObjectPtr<URigHierarchy> DynamicHierarchy;
 
+#if WITH_EDITORONLY_DATA
 	UPROPERTY()
 	TSoftObjectPtr<UControlRigShapeLibrary> GizmoLibrary_DEPRECATED;
+#endif
 
 	UPROPERTY()
 	TArray<TSoftObjectPtr<UControlRigShapeLibrary>> ShapeLibraries;
@@ -534,12 +536,14 @@ protected:
 	bool bEnableControlRigLogging;
 #endif
 
+#if WITH_EDITORONLY_DATA
 	// you either go Input or Output, currently if you put it in both place, Output will override
 	UPROPERTY()
 	TMap<FName, FCachedPropertyPath> InputProperties_DEPRECATED;
 
 	UPROPERTY()
 	TMap<FName, FCachedPropertyPath> OutputProperties_DEPRECATED;
+#endif
 
 	FRigNameCache NameCache;
 
diff --git a/Engine/Plugins/Animation/ControlRig/Source/ControlRig/Public/Rigs/RigHierarchyDefines.h b/Engine/Plugins/Animation/ControlRig/Source/ControlRig/Public/Rigs/RigHierarchyDefines.h
--- a/Engine/Plugins/Animation/ControlRig/Source/ControlRig/Public/Rigs/RigHierarchyDefines.h
+++ b/Engine/Plugins/Animation/ControlRig/Source/ControlRig/Public/Rigs/RigHierarchyDefines.h
@@ -392,7 +392,9 @@ public:
 
 	FRigControlValue()
 		: FloatStorage()
+#if WITH_EDITORONLY_DATA
 		, Storage_DEPRECATED(FTransform::Identity)
+#endif
 	{
 	}
 
@@ -1133,8 +1135,10 @@ private:
 	UPROPERTY()
 	FRigControlValueStorage FloatStorage;
 	
+#if WITH_EDITORONLY_DATA
 	UPROPERTY()
 	FTransform Storage_DEPRECATED;
+#endif
 
 	friend struct FRigControlHierarchy;
 	friend class UControlRigBlueprint;
diff --git a/Engine/Plugins/Animation/ControlRig/Source/ControlRig/Public/Rigs/RigHierarchyElements.h b/Engine/Plugins/Animation/ControlRig/Source/ControlRig/Public/Rigs/RigHierarchyElements.h
--- a/Engine/Plugins/Animation/ControlRig/Source/ControlRig/Public/Rigs/RigHierarchyElements.h
+++ b/Engine/Plugins/Animation/ControlRig/Source/ControlRig/Public/Rigs/RigHierarchyElements.h
@@ -1126,6 +1126,7 @@ struct CONTROLRIG_API FRigControlSettings
 	UPROPERTY(BlueprintReadOnly, EditAnywhere, Category = Animation)
 	bool bRestrictSpaceSwitching;
 
+#if WITH_EDITORONLY_DATA
 	/**
 	 * Deprecated properties.
 	 */
@@ -1134,6 +1135,7 @@ struct CONTROLRIG_API FRigControlSettings
 	
 	UPROPERTY(meta=(DeprecatedProperty, DeprecationMessage = "Use animation_type or shape_visible instead."))
 	bool bShapeEnabled_DEPRECATED = true;
+#endif
 	
 	/** Applies the limits expressed by these settings to a value */
 	FORCEINLINE void ApplyLimits(FRigControlValue& InOutValue) const
diff --git a/Engine/Plugins/Animation/ControlRig/Source/ControlRigDeveloper/Public/ControlRigBlueprint.h b/Engine/Plugins/Animation/ControlRig/Source/ControlRigDeveloper/Public/ControlRigBlueprint.h
--- a/Engine/Plugins/Animation/ControlRig/Source/ControlRigDeveloper/Public/ControlRigBlueprint.h
+++ b/Engine/Plugins/Animation/ControlRig/Source/ControlRigDeveloper/Public/ControlRigBlueprint.h
@@ -395,11 +395,13 @@ public:
 
 protected:
 
+#if WITH_EDITORONLY_DATA
 	UPROPERTY()
 	TObjectPtr<URigVMGraph> Model_DEPRECATED;
 
 	UPROPERTY()
 	TObjectPtr<URigVMFunctionLibrary> FunctionLibrary_DEPRECATED;
+#endif
 
 	UPROPERTY()
 	FRigVMClient RigVMClient;
@@ -475,10 +477,10 @@ public:
 	TArray<TSoftObjectPtr<UControlRigShapeLibrary>> ShapeLibraries;
 
 	const FControlRigShapeDefinition* GetControlShapeByName(const FName& InName) const;
-#endif
 
 	UPROPERTY(transient, DuplicateTransient, meta = (DisplayName = "VM Statistics", DisplayAfter = "VMCompileSettings"))
 	FRigVMStatistics Statistics_DEPRECATED;
+#endif
 
 	UPROPERTY(EditAnywhere, Category = "Drawing")
 	FControlRigDrawContainer DrawContainer;
@@ -506,8 +508,10 @@ public:
 
 public:
 
+#if WITH_EDITORONLY_DATA
 	UPROPERTY()
 	FRigHierarchyContainer HierarchyContainer_DEPRECATED;
+#endif
 
 	UPROPERTY(BlueprintReadOnly, Category = "Hierarchy")
 	TObjectPtr<URigHierarchy> Hierarchy;
diff --git a/Engine/Plugins/Animation/ControlRig/Source/ControlRigDeveloper/Public/Graph/ControlRigGraphNode.h b/Engine/Plugins/Animation/ControlRig/Source/ControlRigDeveloper/Public/Graph/ControlRigGraphNode.h
--- a/Engine/Plugins/Animation/ControlRig/Source/ControlRigDeveloper/Public/Graph/ControlRigGraphNode.h
+++ b/Engine/Plugins/Animation/ControlRig/Source/ControlRigDeveloper/Public/Graph/ControlRigGraphNode.h
@@ -46,6 +46,7 @@ private:
 	UPROPERTY(transient)
 	TMap<FString, TWeakObjectPtr<URigVMPin>> PinPathToModelPin;
 
+#if WITH_EDITORONLY_DATA
 	/** The property we represent. For template nodes this represents the struct/property type name. */
 	UPROPERTY()
 	FName PropertyName_DEPRECATED;
@@ -64,6 +65,7 @@ private:
 	/** Expanded pins */
 	UPROPERTY()
 	TArray<FString> ExpandedPins_DEPRECATED;
+#endif
 
 	/** Cached dimensions of this node (used for auto-layout) */
 	FVector2D Dimensions;
diff --git a/Engine/Plugins/Animation/IKRig/Source/IKRig/Private/Retargeter/IKRetargeter.cpp b/Engine/Plugins/Animation/IKRig/Source/IKRig/Private/Retargeter/IKRetargeter.cpp
--- a/Engine/Plugins/Animation/IKRig/Source/IKRig/Private/Retargeter/IKRetargeter.cpp
+++ b/Engine/Plugins/Animation/IKRig/Source/IKRig/Private/Retargeter/IKRetargeter.cpp
@@ -164,6 +164,7 @@ void URetargetChainSettings::Serialize(FArchive& Ar)
 	Super::Serialize(Ar);
 	Ar.UsingCustomVersion(FIKRigObjectVersion::GUID);
 
+#if WITH_EDITORONLY_DATA
 	if (Ar.IsLoading())
 	{
 		// load the old chain settings into the new struct format
@@ -190,6 +191,7 @@ void URetargetChainSettings::Serialize(FArchive& Ar)
 			PRAGMA_ENABLE_DEPRECATION_WARNINGS
 		}
 	}
+#endif
 }
 
 void URetargetRootSettings::Serialize(FArchive& Ar)
@@ -197,6 +199,7 @@ void URetargetRootSettings::Serialize(FArchive& Ar)
 	Super::Serialize(Ar);
 	Ar.UsingCustomVersion(FIKRigObjectVersion::GUID);
 
+#if WITH_EDITORONLY_DATA
 	if (Ar.IsLoading())
 	{
 		// load the old root settings into the new struct format
@@ -211,6 +214,7 @@ void URetargetRootSettings::Serialize(FArchive& Ar)
 			PRAGMA_ENABLE_DEPRECATION_WARNINGS
 		}
 	}
+#endif
 }
 
 FQuat FIKRetargetPose::GetDeltaRotationForBone(const FName BoneName) const
diff --git a/Engine/Plugins/Animation/IKRig/Source/IKRig/Public/Retargeter/IKRetargeter.h b/Engine/Plugins/Animation/IKRig/Source/IKRig/Public/Retargeter/IKRetargeter.h
--- a/Engine/Plugins/Animation/IKRig/Source/IKRig/Public/Retargeter/IKRetargeter.h
+++ b/Engine/Plugins/Animation/IKRig/Source/IKRig/Public/Retargeter/IKRetargeter.h
@@ -58,6 +58,7 @@ public:
 	UPROPERTY(EditAnywhere, Category = "Settings")
 	FTargetChainSettings Settings;
 
+#if WITH_EDITORONLY_DATA
 	/** Deprecated properties from before FTargetChainSettings / profile refactor  (July 2022)*/
 	UPROPERTY()
 	bool CopyPoseUsingFK_DEPRECATED = true;
@@ -94,6 +95,7 @@ public:
 	UPROPERTY()
 	float UnplantCriticalDamping_DEPRECATED = 1.0f;
 	/** END deprecated properties */
+#endif
 
 	// pointer to editor for details customization
 	#if WITH_EDITOR
@@ -122,6 +124,7 @@ public:
 	#endif
 
 private:
+#if WITH_EDITORONLY_DATA
 	/** Deprecated properties from before FTargetRootSettings / profile refactor */
 	UPROPERTY()
 	bool RetargetRootTranslation_DEPRECATED = true;
@@ -136,6 +139,7 @@ private:
 	UPROPERTY()
 	FRotator StaticRotationOffset_DEPRECATED = FRotator::ZeroRotator;
 	/** END deprecated properties */
+#endif
 };
 
 
diff --git a/Engine/Source/Developer/RigVMDeveloper/Public/RigVMModel/Nodes/RigVMTemplateNode.h b/Engine/Source/Developer/RigVMDeveloper/Public/RigVMModel/Nodes/RigVMTemplateNode.h
--- a/Engine/Source/Developer/RigVMDeveloper/Public/RigVMModel/Nodes/RigVMTemplateNode.h
+++ b/Engine/Source/Developer/RigVMDeveloper/Public/RigVMModel/Nodes/RigVMTemplateNode.h
@@ -182,10 +182,12 @@ protected:
 	UPROPERTY()
 	FString ResolvedFunctionName;
 
+#if WITH_EDITORONLY_DATA
 	// Indicates a preferred permutation using the types of the arguments
 	// Each element is in the format "ArgumentName:CPPType"
 	UPROPERTY()
 	TArray<FString> PreferredPermutationTypes_DEPRECATED;
+#endif
 
 	UPROPERTY()
 	TArray<FRigVMTemplatePreferredType> PreferredPermutationPairs;
diff --git a/Engine/Source/Developer/RigVMDeveloper/Public/RigVMModel/Nodes/RigVMUnitNode.h b/Engine/Source/Developer/RigVMDeveloper/Public/RigVMModel/Nodes/RigVMUnitNode.h
--- a/Engine/Source/Developer/RigVMDeveloper/Public/RigVMModel/Nodes/RigVMUnitNode.h
+++ b/Engine/Source/Developer/RigVMDeveloper/Public/RigVMModel/Nodes/RigVMUnitNode.h
@@ -79,11 +79,13 @@ protected:
 
 private:
 
+#if WITH_EDITORONLY_DATA
 	UPROPERTY()
 	TObjectPtr<UScriptStruct> ScriptStruct_DEPRECATED;
 
 	UPROPERTY()
 	FName MethodName_DEPRECATED;
+#endif
 
 	friend class URigVMController;
 };
diff --git a/Engine/Source/Developer/RigVMDeveloper/Public/RigVMModel/RigVMFunctionLibrary.h b/Engine/Source/Developer/RigVMDeveloper/Public/RigVMModel/RigVMFunctionLibrary.h
--- a/Engine/Source/Developer/RigVMDeveloper/Public/RigVMModel/RigVMFunctionLibrary.h
+++ b/Engine/Source/Developer/RigVMDeveloper/Public/RigVMModel/RigVMFunctionLibrary.h
@@ -72,8 +72,10 @@ public:
 	
 private:
 
+#if WITH_EDITORONLY_DATA
 	UPROPERTY()
 	TMap< TObjectPtr<URigVMLibraryNode>, FRigVMFunctionReferenceArray > FunctionReferences_DEPRECATED;
+#endif
 
 	// A map which stores a library node per original pathname.
 	// The source pathname is the full path of the source function that was localized
diff --git a/Engine/Source/Developer/RigVMDeveloper/Public/RigVMModel/RigVMPin.h b/Engine/Source/Developer/RigVMDeveloper/Public/RigVMModel/RigVMPin.h
--- a/Engine/Source/Developer/RigVMDeveloper/Public/RigVMModel/RigVMPin.h
+++ b/Engine/Source/Developer/RigVMDeveloper/Public/RigVMModel/RigVMPin.h
@@ -35,8 +35,10 @@ public:
 		bInjectedAsInput = true;
 	}
 
+#if WITH_EDITORONLY_DATA
 	UPROPERTY()
 	TObjectPtr<URigVMUnitNode> UnitNode_DEPRECATED;
+#endif
 
 	UPROPERTY()
 	TObjectPtr<URigVMNode> Node;
@@ -513,8 +515,10 @@ private:
 	UPROPERTY()
 	TArray<TObjectPtr<URigVMInjectionInfo>> InjectionInfos;
 
+#if WITH_EDITORONLY_DATA
 	UPROPERTY()
 	FString BoundVariablePath_DEPRECATED;
+#endif
 
 	mutable FString LastKnownCPPType;
 	mutable TRigVMTypeIndex LastKnownTypeIndex;
diff --git a/Engine/Source/Runtime/AnimGraphRuntime/Private/BoneControllers/AnimNode_RigidBody.cpp b/Engine/Source/Runtime/AnimGraphRuntime/Private/BoneControllers/AnimNode_RigidBody.cpp
--- a/Engine/Source/Runtime/AnimGraphRuntime/Private/BoneControllers/AnimNode_RigidBody.cpp
+++ b/Engine/Source/Runtime/AnimGraphRuntime/Private/BoneControllers/AnimNode_RigidBody.cpp
@@ -111,7 +111,9 @@ FSimSpaceSettings::FSimSpaceSettings()
 	, MaxAngularVelocity(10000)
 	, MaxLinearAcceleration(10000)
 	, MaxAngularAcceleration(10000)
+#if WITH_EDITORONLY_DATA
 	, ExternalLinearDrag_DEPRECATED(0)
+#endif
 	, ExternalLinearDragV(FVector::ZeroVector)
 	, ExternalLinearVelocity(FVector::ZeroVector)
 	, ExternalAngularVelocity(FVector::ZeroVector)
@@ -120,6 +122,7 @@ FSimSpaceSettings::FSimSpaceSettings()
 
 void FSimSpaceSettings::PostSerialize(const FArchive& Ar)
 {
+#if WITH_EDITORONLY_DATA
 	if (Ar.IsLoading())
 	{
 		if (ExternalLinearDrag_DEPRECATED != 0.0f)
@@ -127,6 +130,7 @@ void FSimSpaceSettings::PostSerialize(const FArchive& Ar)
 			ExternalLinearDragV = FVector(ExternalLinearDrag_DEPRECATED, ExternalLinearDrag_DEPRECATED, ExternalLinearDrag_DEPRECATED);
 		}
 	}
+#endif
 }
 
 
diff --git a/Engine/Source/Runtime/AnimGraphRuntime/Public/BoneControllers/AnimNode_RigidBody.h b/Engine/Source/Runtime/AnimGraphRuntime/Public/BoneControllers/AnimNode_RigidBody.h
--- a/Engine/Source/Runtime/AnimGraphRuntime/Public/BoneControllers/AnimNode_RigidBody.h
+++ b/Engine/Source/Runtime/AnimGraphRuntime/Public/BoneControllers/AnimNode_RigidBody.h
@@ -99,8 +99,10 @@ struct ANIMGRAPHRUNTIME_API FSimSpaceSettings
 	UPROPERTY(EditAnywhere, BlueprintReadWrite, Category = Settings, meta = (ClampMin = "0.0"))
 	float MaxAngularAcceleration;
 
+#if WITH_EDITORONLY_DATA
 	UPROPERTY(meta = (DeprecatedProperty, DeprecationMessage = "ExternalLinearDrag is deprecated. Please use ExternalLinearDragV instead."))
 	float ExternalLinearDrag_DEPRECATED;
+#endif
 
 	// Additional linear drag applied to every body in addition to linear drag specified on them in the physics asset. 
 	// When combined with ExternalLinearVelocity, this can be used to add a temporary wind-blown effect without having to tune linear drag on 
diff --git a/Engine/Source/Runtime/Engine/Classes/Animation/AnimMontage.h b/Engine/Source/Runtime/Engine/Classes/Animation/AnimMontage.h
--- a/Engine/Source/Runtime/Engine/Classes/Animation/AnimMontage.h
+++ b/Engine/Source/Runtime/Engine/Classes/Animation/AnimMontage.h
@@ -44,9 +44,11 @@ struct FCompositeSection : public FAnimLinkableElement
 	UPROPERTY(EditAnywhere, Category=Section)
 	FName SectionName;
 
+#if WITH_EDITORONLY_DATA
 	/** Start Time **/
 	UPROPERTY()
 	float StartTime_DEPRECATED;
+#endif
 
 	/** Should this animation loop. */
 	UPROPERTY(VisibleAnywhere, Category=Section)
@@ -63,7 +65,9 @@ public:
 	FCompositeSection()
 		: FAnimLinkableElement()
 		, SectionName(NAME_None)
+#if WITH_EDITORONLY_DATA
 		, StartTime_DEPRECATED(0.0f)
+#endif
 		, NextSectionName(NAME_None)
 	{
 	}
@@ -104,8 +108,10 @@ struct FBranchingPoint : public FAnimLinkableElement
 	UPROPERTY(EditAnywhere, Category=BranchingPoint)
 	FName EventName;
 
+#if WITH_EDITORONLY_DATA
 	UPROPERTY()
 	float DisplayTime_DEPRECATED = 0.f;
+#endif
 
 	/** An offset from the DisplayTime to the actual time we will trigger the notify, as we cannot always trigger it exactly at the time the user wants */
 	UPROPERTY()
@@ -622,15 +628,19 @@ class UAnimMontage : public UAnimCompositeBase
 	UPROPERTY(EditAnywhere, Category=BlendOption)
 	FAlphaBlend BlendIn;
 
+#if WITH_EDITORONLY_DATA
 	UPROPERTY()
 	float BlendInTime_DEPRECATED;
+#endif
 
 	/** Blend out option. This is only used when it blends out itself. If it's interrupted by other montages, it will use new montage's BlendIn option to blend out. */
 	UPROPERTY(EditAnywhere, Category=BlendOption)
 	FAlphaBlend BlendOut;
 
+#if WITH_EDITORONLY_DATA
 	UPROPERTY()
 	float BlendOutTime_DEPRECATED;
+#endif
 
 	/** Time from Sequence End to trigger blend out.
 	 * <0 means using BlendOutTime, so BlendOut finishes as Montage ends.
@@ -669,9 +679,11 @@ class UAnimMontage : public UAnimCompositeBase
 	UPROPERTY()
 	TArray<struct FSlotAnimationTrack> SlotAnimTracks;
 
+#if WITH_EDITORONLY_DATA
 	// Remove this when VER_UE4_MONTAGE_BRANCHING_POINT_REMOVAL is removed.
 	UPROPERTY()
 	TArray<struct FBranchingPoint> BranchingPoints_DEPRECATED;
+#endif
 
 	/** If this is on, it will allow extracting root motion translation. DEPRECATED in 4.5 root motion is controlled by anim sequences **/
 	UPROPERTY()
diff --git a/Engine/Source/Runtime/Engine/Classes/Animation/PreviewAssetAttachComponent.h b/Engine/Source/Runtime/Engine/Classes/Animation/PreviewAssetAttachComponent.h
--- a/Engine/Source/Runtime/Engine/Classes/Animation/PreviewAssetAttachComponent.h
+++ b/Engine/Source/Runtime/Engine/Classes/Animation/PreviewAssetAttachComponent.h
@@ -25,12 +25,18 @@ private:
 	UPROPERTY()
 	TSoftObjectPtr<class UObject> AttachedObject;
 
+#if WITH_EDITORONLY_DATA
 	UPROPERTY()
 	TObjectPtr<UObject> Object_DEPRECATED;
+#endif
 
 public:
 
-	FPreviewAttachedObjectPair() : Object_DEPRECATED(NULL) {}
+	FPreviewAttachedObjectPair()
+#if WITH_EDITORONLY_DATA
+		: Object_DEPRECATED(NULL)
+#endif
+	{}
 
 	/** The name of the attach point of the Object (for example a bone or socket name) */
 	UPROPERTY()
@@ -38,11 +44,13 @@ public:
 
 	void SaveAttachedObjectFromDeprecatedProperty()
 	{
+#if WITH_EDITORONLY_DATA
 		if (Object_DEPRECATED)
 		{
 			AttachedObject = Object_DEPRECATED;
 			Object_DEPRECATED = NULL;
 		}
+#endif
 	}
 
 	UObject* GetAttachedObject() const
diff --git a/Engine/Source/Runtime/Engine/Classes/Animation/SkeletalMeshActor.h b/Engine/Source/Runtime/Engine/Classes/Animation/SkeletalMeshActor.h
--- a/Engine/Source/Runtime/Engine/Classes/Animation/SkeletalMeshActor.h
+++ b/Engine/Source/Runtime/Engine/Classes/Animation/SkeletalMeshActor.h
@@ -29,8 +29,10 @@ class ENGINE_API ASkeletalMeshActor : public AActor
 	UPROPERTY(EditAnywhere, BlueprintReadWrite, Category=Animation, AdvancedDisplay)
 	uint32 bShouldDoAnimNotifies:1;
 
+#if WITH_EDITORONLY_DATA
 	UPROPERTY()
 	uint32 bWakeOnLevelStart_DEPRECATED:1;
+#endif
 
 private:
 	UPROPERTY(Category = SkeletalMeshActor, VisibleAnywhere, BlueprintReadOnly, meta = (ExposeFunctionCategories = "Mesh,Components|SkeletalMesh,Animation,Physics", AllowPrivateAccess = "true"))
diff --git a/Engine/Source/Runtime/Engine/Classes/Animation/Skeleton.h b/Engine/Source/Runtime/Engine/Classes/Animation/Skeleton.h
--- a/Engine/Source/Runtime/Engine/Classes/Animation/Skeleton.h
+++ b/Engine/Source/Runtime/Engine/Classes/Animation/Skeleton.h
@@ -83,6 +83,7 @@ struct FBoneNode
 {
 	GENERATED_USTRUCT_BODY()
 
+#if WITH_EDITORONLY_DATA
 	/** Name of bone, this is the search criteria to match with mesh bone. This will be NAME_None if deleted. */
 	UPROPERTY()
 	FName Name_DEPRECATED;
@@ -90,21 +91,28 @@ struct FBoneNode
 	/** Parent Index. -1 if not used. The root has 0 as its parent. Do not delete the element but set this to -1. If it is revived by other reason, fix up this link. */
 	UPROPERTY()
 	int32 ParentIndex_DEPRECATED;
+#endif
 
 	/** Retargeting Mode for Translation Component. */
 	UPROPERTY(EditAnywhere, Category=BoneNode)
 	TEnumAsByte<EBoneTranslationRetargetingMode::Type> TranslationRetargetingMode;
 
 	FBoneNode()
-		: ParentIndex_DEPRECATED(INDEX_NONE)
-		, TranslationRetargetingMode(EBoneTranslationRetargetingMode::Animation)
+		:
+#if WITH_EDITORONLY_DATA
+		ParentIndex_DEPRECATED(INDEX_NONE),
+#endif
+		TranslationRetargetingMode(EBoneTranslationRetargetingMode::Animation)
 	{
 	}
 
 	FBoneNode(FName InBoneName, int32 InParentIndex)
-		: Name_DEPRECATED(InBoneName)
-		, ParentIndex_DEPRECATED(InParentIndex)
-		, TranslationRetargetingMode(EBoneTranslationRetargetingMode::Animation)
+		:
+#if WITH_EDITORONLY_DATA
+		Name_DEPRECATED(InBoneName),
+		ParentIndex_DEPRECATED(InParentIndex),
+#endif
+		TranslationRetargetingMode(EBoneTranslationRetargetingMode::Animation)
 	{
 	}
 };
@@ -457,9 +465,11 @@ protected:
 	UPROPERTY(VisibleAnywhere, Category=Skeleton)
 	TArray<struct FBoneNode> BoneTree;
 
+#if WITH_EDITORONLY_DATA
 	/** Reference skeleton poses in local space */
 	UPROPERTY()
 	TArray<FTransform> RefLocalPoses_DEPRECATED;
+#endif
 
 	/** Reference Skeleton */
 	FReferenceSkeleton ReferenceSkeleton;
diff --git a/Engine/Source/Runtime/Engine/Classes/Components/SkeletalMeshComponent.h b/Engine/Source/Runtime/Engine/Classes/Components/SkeletalMeshComponent.h
--- a/Engine/Source/Runtime/Engine/Classes/Components/SkeletalMeshComponent.h
+++ b/Engine/Source/Runtime/Engine/Classes/Components/SkeletalMeshComponent.h
@@ -670,10 +670,12 @@ private:
 	UPROPERTY(EditAnywhere, AdvancedDisplay, Category = SkeletalMesh)
 	uint8 bAllowAnimCurveEvaluation : 1;
 
+#if WITH_EDITORONLY_DATA
 	/** DEPRECATED. Use bAllowAnimCurveEvaluation instead */
 	UE_DEPRECATED(4.18, "This property is deprecated. Please use bAllowAnimCurveEvaluatiuon instead. Note that the meaning is reversed.")	
 	UPROPERTY()
 	uint8 bDisableAnimCurves_DEPRECATED : 1;
+#endif
 
 	/** Whether or not we're taking cloth sim information from our leader component */
 	uint8 bBindClothToLeaderComponent:1;
diff --git a/Engine/Source/Runtime/Engine/Classes/Components/SkinnedMeshComponent.h b/Engine/Source/Runtime/Engine/Classes/Components/SkinnedMeshComponent.h
--- a/Engine/Source/Runtime/Engine/Classes/Components/SkinnedMeshComponent.h
+++ b/Engine/Source/Runtime/Engine/Classes/Components/SkinnedMeshComponent.h
@@ -629,9 +629,11 @@ public:
 	UPROPERTY()
 	uint8 bForceWireframe:1;
 
+#if WITH_EDITORONLY_DATA
 	/** Draw the skeleton hierarchy for this skel mesh. */
 	UPROPERTY()
 	uint8 bDisplayBones_DEPRECATED:1;
+#endif
 
 	/** Disable Morphtarget for this component. */
 	UPROPERTY(EditAnywhere, AdvancedDisplay, BlueprintReadWrite, Category = SkeletalMesh)
diff --git a/Engine/Source/Runtime/Engine/Private/Animation/AnimCurveTypes.cpp b/Engine/Source/Runtime/Engine/Private/Animation/AnimCurveTypes.cpp
--- a/Engine/Source/Runtime/Engine/Private/Animation/AnimCurveTypes.cpp
+++ b/Engine/Source/Runtime/Engine/Private/Animation/AnimCurveTypes.cpp
@@ -27,12 +27,17 @@ void FAnimCurveBase::PostSerialize(FArchive& Ar)
 				Ar << CurveUid;
 
 				Name.UID = CurveUid;
+
+#if WITH_EDITORONLY_DATA
 				Name.DisplayName = LastObservedName_DEPRECATED;
+#endif
 			}
+#if WITH_EDITORONLY_DATA
 			else
 			{
 				Name.DisplayName = LastObservedName_DEPRECATED;
 			}
+#endif
 		}
 
 #if WITH_EDITORONLY_DATA
diff --git a/Engine/Source/Runtime/Engine/Private/Animation/AnimMontage.cpp b/Engine/Source/Runtime/Engine/Private/Animation/AnimMontage.cpp
--- a/Engine/Source/Runtime/Engine/Private/Animation/AnimMontage.cpp
+++ b/Engine/Source/Runtime/Engine/Private/Animation/AnimMontage.cpp
@@ -65,8 +65,10 @@ UAnimMontage::UAnimMontage(const FObjectInitializer& ObjectInitializer)
 	BlendProfileIn = nullptr;
 	BlendProfileOut = nullptr;
 
+#if WITH_EDITORONLY_DATA
 	BlendInTime_DEPRECATED = -1.f;
 	BlendOutTime_DEPRECATED = -1.f;
+#endif
 
 	AddSlot(FAnimSlotGroup::DefaultSlotName);
 
@@ -402,12 +404,14 @@ void UAnimMontage::PostLoad()
 
 	for(FCompositeSection& Composite : CompositeSections)
 	{
+#if WITH_EDITORONLY_DATA
 		if(Composite.StartTime_DEPRECATED != 0.0f)
 		{
 			Composite.Clear();
 			Composite.Link(this, Composite.StartTime_DEPRECATED);
 		}
 		else
+#endif
 		{
 			Composite.RefreshSegmentOnLoad();
 			Composite.Link(this, Composite.GetTime());
@@ -484,12 +488,14 @@ void UAnimMontage::PostLoad()
 
 	for(FAnimNotifyEvent& Notify : Notifies)
 	{
+#if WITH_EDITORONLY_DATA
 		if(Notify.DisplayTime_DEPRECATED != 0.0f)
 		{
 			Notify.Clear();
 			Notify.Link(this, Notify.DisplayTime_DEPRECATED);
 		}
 		else
+#endif
 		{
 			Notify.Link(this, Notify.GetTime());
 		}
@@ -506,6 +512,7 @@ void UAnimMontage::PostLoad()
 		ConvertBranchingPointsToAnimNotifies();
 	}
 
+#if WITH_EDITORONLY_DATA
 	// fix up blending time deprecated variable
 	if (BlendInTime_DEPRECATED != -1.f)
 	{
@@ -518,6 +525,7 @@ void UAnimMontage::PostLoad()
 		BlendOut.SetBlendTime(BlendOutTime_DEPRECATED);
 		BlendOutTime_DEPRECATED = -1.f;
 	}
+#endif
 
 	// collect markers if it's valid
 	CollectMarkers();
@@ -525,6 +533,7 @@ void UAnimMontage::PostLoad()
 
 void UAnimMontage::ConvertBranchingPointsToAnimNotifies()
 {
+#if WITH_EDITORONLY_DATA
 	if (BranchingPoints_DEPRECATED.Num() > 0)
 	{
 		// Handle deprecated DisplayTime first
@@ -544,7 +553,6 @@ void UAnimMontage::ConvertBranchingPointsToAnimNotifies()
 		// Then convert to AnimNotifies
 		USkeleton * MySkeleton = GetSkeleton();
 
-#if WITH_EDITORONLY_DATA
 		// Add a new AnimNotifyTrack, and place all branching points in there.
 		int32 TrackIndex = AnimNotifyTracks.Num();
 
@@ -552,7 +560,6 @@ void UAnimMontage::ConvertBranchingPointsToAnimNotifies()
 		NewItem.TrackName = *FString::FromInt(TrackIndex + 1);
 		NewItem.TrackColor = FLinearColor::White;
 		AnimNotifyTracks.Add(NewItem);
-#endif
 
 		for (auto BranchingPoint : BranchingPoints_DEPRECATED)
 		{
@@ -565,26 +572,23 @@ void UAnimMontage::ConvertBranchingPointsToAnimNotifies()
 #if WITH_EDITOR
 			NewEvent.TriggerTimeOffset = GetTriggerTimeOffsetForType(CalculateOffsetForNotify(TriggerTime));
 #endif
-#if WITH_EDITORONLY_DATA
 			NewEvent.TrackIndex = TrackIndex;
-#endif
 			NewEvent.Notify = nullptr;
 			NewEvent.NotifyStateClass = nullptr;
 			NewEvent.bConvertedFromBranchingPoint = true;
 			NewEvent.MontageTickType = EMontageNotifyTickType::BranchingPoint;
 
-#if WITH_EDITORONLY_DATA
 			// Add as a custom AnimNotify event to Skeleton.
 			if (MySkeleton)
 			{
 				MySkeleton->AnimationNotifies.AddUnique(NewEvent.NotifyName);
 			}
-#endif
 		}
 
 		BranchingPoints_DEPRECATED.Empty();
 		RefreshBranchingPointMarkers();
 	}
+#endif
 }
 
 void UAnimMontage::RefreshBranchingPointMarkers()
diff --git a/Engine/Source/Runtime/Engine/Private/Animation/AnimSequence.cpp b/Engine/Source/Runtime/Engine/Private/Animation/AnimSequence.cpp
--- a/Engine/Source/Runtime/Engine/Private/Animation/AnimSequence.cpp
+++ b/Engine/Source/Runtime/Engine/Private/Animation/AnimSequence.cpp
@@ -963,12 +963,14 @@ void UAnimSequence::PostLoad()
 
 	for(FAnimNotifyEvent& Notify : Notifies)
 	{
+#if WITH_EDITORONLY_DATA
 		if(Notify.DisplayTime_DEPRECATED != 0.0f)
 		{
 			Notify.Clear();
 			Notify.Link(this, Notify.DisplayTime_DEPRECATED);
 		}
 		else
+#endif
 		{
 			Notify.Link(this, Notify.GetTime());
 		}
diff --git a/Engine/Source/Runtime/Engine/Private/Animation/Skeleton.cpp b/Engine/Source/Runtime/Engine/Private/Animation/Skeleton.cpp
--- a/Engine/Source/Runtime/Engine/Private/Animation/Skeleton.cpp
+++ b/Engine/Source/Runtime/Engine/Private/Animation/Skeleton.cpp
@@ -828,6 +828,7 @@ void USkeleton::PostEditUndo()
 /** Remove this function when VER_UE4_REFERENCE_SKELETON_REFACTOR is removed. */
 void USkeleton::ConvertToFReferenceSkeleton()
 {
+#if WITH_EDITORONLY_DATA
 	FReferenceSkeletonModifier RefSkelModifier(ReferenceSkeleton, this);
 
 	check( BoneTree.Num() == RefLocalPoses_DEPRECATED.Num() );
@@ -849,6 +850,7 @@ void USkeleton::ConvertToFReferenceSkeleton()
 	// VER_UE4_REFERENCE_SKELETON_REFACTOR, this shouldn't be needed. It shouldn't have any 
 	// AnimatedRetargetSources
 	ensure (AnimRetargetSources.Num() == 0);
+#endif
 }
 
 void USkeleton::ClearSkeletonRemappings()
diff --git a/Engine/Source/Runtime/Engine/Private/Components/SkeletalMeshComponent.cpp b/Engine/Source/Runtime/Engine/Private/Components/SkeletalMeshComponent.cpp
--- a/Engine/Source/Runtime/Engine/Private/Components/SkeletalMeshComponent.cpp
+++ b/Engine/Source/Runtime/Engine/Private/Components/SkeletalMeshComponent.cpp
@@ -358,10 +358,13 @@ void USkeletalMeshComponent::Serialize(FArchive& Ar)
 	}
 
 	Ar.UsingCustomVersion(FAnimPhysObjectVersion::GUID);
+
+#if WITH_EDITORONLY_DATA
 	if (Ar.IsLoading() && Ar.CustomVer(FAnimPhysObjectVersion::GUID) < FAnimPhysObjectVersion::RenameDisableAnimCurvesToAllowAnimCurveEvaluation)
 	{
 		bAllowAnimCurveEvaluation = !bDisableAnimCurves_DEPRECATED;
 	}
+#endif
 
 	PRAGMA_ENABLE_DEPRECATION_WARNINGS
 }
diff --git a/Engine/Source/Runtime/Engine/Public/Animation/AnimCurveTypes.h b/Engine/Source/Runtime/Engine/Public/Animation/AnimCurveTypes.h
--- a/Engine/Source/Runtime/Engine/Public/Animation/AnimCurveTypes.h
+++ b/Engine/Source/Runtime/Engine/Public/Animation/AnimCurveTypes.h
@@ -77,11 +77,13 @@ struct ENGINE_API FAnimCurveBase
 {
 	GENERATED_USTRUCT_BODY()
 
+#if WITH_EDITORONLY_DATA
 	// Last observed name of the curve. We store this so we can recover from situations that
 	// mean the skeleton doesn't have a mapped name for our UID (such as a user saving the an
 	// animation but not the skeleton).
 	UPROPERTY()
 	FName		LastObservedName_DEPRECATED;
+#endif
 
 	UPROPERTY()
 	FSmartName	Name;
diff --git a/Engine/Source/Runtime/Engine/Public/Animation/AnimTypes.h b/Engine/Source/Runtime/Engine/Public/Animation/AnimTypes.h
--- a/Engine/Source/Runtime/Engine/Public/Animation/AnimTypes.h
+++ b/Engine/Source/Runtime/Engine/Public/Animation/AnimTypes.h
@@ -289,9 +289,11 @@ struct FAnimNotifyEvent : public FAnimLinkableElement
 {
 	GENERATED_USTRUCT_BODY()
 
+#if WITH_EDITORONLY_DATA
 	/** The user requested time for this notify */
 	UPROPERTY()
 	float DisplayTime_DEPRECATED;
+#endif
 
 	/** An offset from the DisplayTime to the actual time we will trigger the notify, as we cannot always trigger it exactly at the time the user wants */
 	UPROPERTY()
@@ -371,7 +373,9 @@ private:
 public:
 	FAnimNotifyEvent()
 		: FAnimLinkableElement()
+#if WITH_EDITORONLY_DATA
 		, DisplayTime_DEPRECATED(0)
+#endif
 		, TriggerTimeOffset(0)
 		, EndTriggerTimeOffset(0)
 		, TriggerWeightThreshold(ZERO_ANIMWEIGHT_THRESH)
