Workarounds for StressTestUnity errors.

diff --git a/Engine/Source/Runtime/Engine/Private/Net/Experimental/Iris/DataStreamChannel.cpp b/Engine/Source/Runtime/Engine/Private/Net/Experimental/Iris/DataStreamChannel.cpp
--- a/Engine/Source/Runtime/Engine/Private/Net/Experimental/Iris/DataStreamChannel.cpp
+++ b/Engine/Source/Runtime/Engine/Private/Net/Experimental/Iris/DataStreamChannel.cpp
@@ -158,7 +158,7 @@ void UDataStreamChannel::Tick()
 		return;
 	}
 
-	if (!IsNetReady(Private::bIrisSaturateBandwidth) || IsPacketWindowFull() || !Connection->HasReceivedClientPacket() || (Connection->Handler != nullptr && !Connection->Handler->IsFullyInitialized()))
+	if (!IsNetReady(UE::Net::Private::bIrisSaturateBandwidth) || IsPacketWindowFull() || !Connection->HasReceivedClientPacket() || (Connection->Handler != nullptr && !Connection->Handler->IsFullyInitialized()))
 	{
 		return;
 	}
@@ -276,7 +276,7 @@ void UDataStreamChannel::Tick()
 	{
 		// Write data until we are not allowed to write more
 	}
-	while ((WriteDataFunction() == UDataStream::EWriteResult::HasMoreData) && IsNetReady(Private::bIrisSaturateBandwidth) && (!IsPacketWindowFull()));
+	while ((WriteDataFunction() == UDataStream::EWriteResult::HasMoreData) && IsNetReady(UE::Net::Private::bIrisSaturateBandwidth) && (!IsPacketWindowFull()));
 
 	// call end write to cleanup data initialized in BeginWrite	
 	DataStreamManager->EndWrite();
diff --git a/Engine/Source/Runtime/Engine/Private/Net/Iris/ReplicationSystem/ActorReplicationBridge.cpp b/Engine/Source/Runtime/Engine/Private/Net/Iris/ReplicationSystem/ActorReplicationBridge.cpp
--- a/Engine/Source/Runtime/Engine/Private/Net/Iris/ReplicationSystem/ActorReplicationBridge.cpp
+++ b/Engine/Source/Runtime/Engine/Private/Net/Iris/ReplicationSystem/ActorReplicationBridge.cpp
@@ -256,7 +256,7 @@ UE::Net::FNetHandle UActorReplicationBridge::BeginReplication(AActor* Actor, con
 	const uint32 ClampedPollPeriod = FMath::Clamp<uint32>(UnclampedPollPeriod, 1U, static_cast<uint32>(std::numeric_limits<uint8>::max()) + 1U);
 	CreateNetHandleParams.PollFramePeriod = (ClampedPollPeriod - 1U) & 255U;
 #if !UE_BUILD_SHIPPING
-	ensureAlwaysMsgf(!Private::bIrisEnableLowNetUpdateFrequencyEnsure || (ClampedPollPeriod >= UnclampedPollPeriod), TEXT("Very low NetUpdateFrequency %f for Actor %s. Suggest setting it to %f or higher."), Actor->NetUpdateFrequency, ToCStr(Actor->GetName()), GetMinSupportedNetUpdateFrequency());
+	ensureAlwaysMsgf(!UE::Net::Private::bIrisEnableLowNetUpdateFrequencyEnsure || (ClampedPollPeriod >= UnclampedPollPeriod), TEXT("Very low NetUpdateFrequency %f for Actor %s. Suggest setting it to %f or higher."), Actor->NetUpdateFrequency, ToCStr(Actor->GetName()), GetMinSupportedNetUpdateFrequency());
 	ensureAlwaysMsgf(!(Actor->bAlwaysRelevant || Actor->bOnlyRelevantToOwner) || CreateNetHandleParams.StaticPriority >= 1.0f, TEXT("Very low NetPriority %.02f for always relevant or owner relevant Actor %s. Set it to 1.0f or higher."), Actor->NetPriority, ToCStr(Actor->GetName()));
 #endif
 
