Fix for volume texture mip upload. If the alignment is 64kb for the mip and the Z has mipped down to 1, then this check fails and the mip doesn't upload. Fuond with a 2048x2048x8 DXT1 texture, mip index 3 (256x256x1)

diff --git a/Engine/Source/Runtime/Engine/Private/Streaming/Texture2DMipDataProvider_DDC.cpp b/Engine/Source/Runtime/Engine/Private/Streaming/Texture2DMipDataProvider_DDC.cpp
--- a/Engine/Source/Runtime/Engine/Private/Streaming/Texture2DMipDataProvider_DDC.cpp
+++ b/Engine/Source/Runtime/Engine/Private/Streaming/Texture2DMipDataProvider_DDC.cpp
@@ -126,7 +126,10 @@ bool FTexture2DMipDataProvider_DDC::SerializeMipInfo(const FTextureUpdateContext
 		Ar.Serialize(OutMipInfo.DestData, MipSize);
 		return true;
 	}
-	else if (uint64(MipSize) < OutMipInfo.DataSize && DepthOrArraySize > 1 && OutMipInfo.DataSize % DepthOrArraySize == 0 && MipSize % DepthOrArraySize == 0)
+	// This used to be DepthOrArraySize > 1, which fails when a volume texture gets mipped down to WxHx1 that also has
+	// allocation alignment such that the size is larger than its bulk data size. The original author is gone, but this looks
+	// like it should just be a 0 check to prevent divide by zero with erroneous data. (test case was 2048x2048x8 DXT1 on dx12)
+	else if (uint64(MipSize) < OutMipInfo.DataSize && DepthOrArraySize > 0 && OutMipInfo.DataSize % DepthOrArraySize == 0 && MipSize % DepthOrArraySize == 0)
 	{
 		UE_LOG(LogTexture, Verbose, TEXT("Cached mip size is smaller than streaming buffer size for mip %d of %s: %d KiB / %d KiB."),
 			ResourceState.MaxNumLODs - MipIndex, *Context.Resource->GetTextureName().ToString(), OutMipInfo.DataSize / 1024, MipSize / 1024);
