Remove previous transform hack left over from Frosty (decreasing compressed per instance data from 5 -> 3 floats4s, and uncompressed per instance data from 7 -> 4 float4s. Also saves a lot of ALU)

diff --git a/Engine/Plugins/Experimental/GPULightmass/Source/GPULightmass/Private/LightmapRenderer.cpp b/Engine/Plugins/Experimental/GPULightmass/Source/GPULightmass/Private/LightmapRenderer.cpp
--- a/Engine/Plugins/Experimental/GPULightmass/Source/GPULightmass/Private/LightmapRenderer.cpp
+++ b/Engine/Plugins/Experimental/GPULightmass/Source/GPULightmass/Private/LightmapRenderer.cpp
@@ -520,8 +520,7 @@ void FCachedRayTracingSceneData::SetupViewUniformBufferFromSceneRenderState(FSce
 				INVALID_LAST_UPDATE_FRAME,
 				0, /* Custom Data Count */
 				0.0f, /* Random ID */
-				PrimitiveUniformShaderParameters.LocalToRelativeWorld,
-				PrimitiveUniformShaderParameters.LocalToRelativeWorld // TODO: Temporary PrevVelocityHack
+				PrimitiveUniformShaderParameters.LocalToRelativeWorld
 			);
 
 			PrimitiveSceneData.Add(FPrimitiveSceneShaderData(PrimitiveUniformShaderParameters));
diff --git a/Engine/Shaders/Private/SceneData.ush b/Engine/Shaders/Private/SceneData.ush
--- a/Engine/Shaders/Private/SceneData.ush
+++ b/Engine/Shaders/Private/SceneData.ush
@@ -951,11 +951,6 @@ FInstanceSceneData GetInstanceSceneData(uint InstanceId, uint SOAStride, bool bC
 		float3 Scale = 0;
 		float4x4 LocalToRelativeWorld							= DecodeTransform( RotationScale, Translation, Scale );
 
-		uint4  PrevRotationScale								= asuint(LoadInstanceSceneDataElement(3 * SOAStride + InstanceId));
-		float3 PrevTranslation									= LoadInstanceSceneDataElement(4 * SOAStride + InstanceId).xyz;
-		float3 PrevScale = 0;
-		float4x4 PrevLocalToRelativeWorld						= DecodeTransform( PrevRotationScale, PrevTranslation, PrevScale );
-
 		InstanceData.NonUniformScale.xyz						= abs(Scale);
 		InstanceData.InvNonUniformScale							= rcp(InstanceData.NonUniformScale.xyz);
 #else
@@ -963,17 +958,10 @@ FInstanceSceneData GetInstanceSceneData(uint InstanceId, uint SOAStride, bool bC
 																			LoadInstanceSceneDataElement(2 * SOAStride + InstanceId),
 																	                LoadInstanceSceneDataElement(3 * SOAStride + InstanceId),
 																	                float4(0.0f, 0.0f, 0.0f, 1.0f)));
-
-		float4x4 PrevLocalToRelativeWorld						= transpose(float4x4(LoadInstanceSceneDataElement(4 * SOAStride + InstanceId),
-																	                LoadInstanceSceneDataElement(5 * SOAStride + InstanceId),
-																	                LoadInstanceSceneDataElement(6 * SOAStride + InstanceId),
-																	                float4(0.0f, 0.0f, 0.0f, 1.0f)));
-
 #endif
 
 
 		InstanceData.LocalToWorld	= MakeLWCMatrix(TilePosition, LocalToRelativeWorld);
-		InstanceData.PrevLocalToWorld	= MakeLWCMatrix(TilePosition, PrevLocalToRelativeWorld);
 		ComputeInstanceDerivedData(InstanceData, TilePosition, LocalToRelativeWorld);
 
 		InstanceData.NaniteRuntimeResourceID	= PrimitiveData.NaniteResourceID;
@@ -1030,6 +1018,10 @@ FInstanceSceneData GetInstanceSceneData(uint InstanceId, uint SOAStride, bool bC
 			float3 TilePosition 				= PrimitiveData.TilePosition;
 			InstanceData.PrevLocalToWorld		= MakeLWCMatrix(TilePosition, PrevLocalToRelativeWorld);
 		}
+		else
+		{
+			InstanceData.PrevLocalToWorld		= MakeLWCMatrix(TilePosition, LocalToRelativeWorld);
+		}
 
 	#if 1 //NEEDS_LIGHTMAP_COORDINATE
 		BRANCH
diff --git a/Engine/Source/Runtime/Engine/Private/InstanceUniformShaderParameters.cpp b/Engine/Source/Runtime/Engine/Private/InstanceUniformShaderParameters.cpp
--- a/Engine/Source/Runtime/Engine/Private/InstanceUniformShaderParameters.cpp
+++ b/Engine/Source/Runtime/Engine/Private/InstanceUniformShaderParameters.cpp
@@ -14,7 +14,7 @@ void FInstanceSceneShaderData::Build
 	float RandomID
 )
 {
-	BuildInternal(PrimitiveId, RelativeId, InstanceFlags, LastUpdateFrame, CustomDataCount, RandomID, FRenderTransform::Identity, FRenderTransform::Identity);
+	BuildInternal(PrimitiveId, RelativeId, InstanceFlags, LastUpdateFrame, CustomDataCount, RandomID, FRenderTransform::Identity);
 }
 
 void FInstanceSceneShaderData::Build
@@ -26,18 +26,15 @@ void FInstanceSceneShaderData::Build
 	uint32 CustomDataCount,
 	float RandomID,
 	const FRenderTransform& LocalToPrimitive,
-	const FRenderTransform& PrimitiveToWorld,
-	const FRenderTransform& PrevPrimitiveToWorld // TODO: Temporary PrevVelocityHack
+	const FRenderTransform& PrimitiveToWorld
 )
 {
 	FRenderTransform LocalToWorld = LocalToPrimitive * PrimitiveToWorld;
-	FRenderTransform PrevLocalToWorld = LocalToPrimitive * PrevPrimitiveToWorld; // TODO: Temporary PrevVelocityHack
 
 	// Remove shear
 	LocalToWorld.Orthogonalize();
-	PrevLocalToWorld.Orthogonalize(); // TODO: Temporary PrevVelocityHack
 
-	BuildInternal(PrimitiveId, RelativeId, InstanceFlags, LastUpdateFrame, CustomDataCount, RandomID, LocalToWorld, PrevLocalToWorld);
+	BuildInternal(PrimitiveId, RelativeId, InstanceFlags, LastUpdateFrame, CustomDataCount, RandomID, LocalToWorld);
 }
 
 void FInstanceSceneShaderData::BuildInternal
@@ -48,8 +45,7 @@ void FInstanceSceneShaderData::BuildInternal
 	uint32 LastUpdateFrame,
 	uint32 CustomDataCount,
 	float RandomID,
-	const FRenderTransform& LocalToWorld, // Assumes shear has been removed already
-	const FRenderTransform& PrevLocalToWorld // Assumes shear has been removed already // TODO: Temporary PrevVelocityHack
+	const FRenderTransform& LocalToWorld // Assumes shear has been removed already
 )
 {
 	// Note: layout must match GetInstanceData in SceneData.ush and InitializeInstanceSceneData in GPUSceneWriter.ush
@@ -76,21 +72,15 @@ void FInstanceSceneShaderData::BuildInternal
 	Data[0].Z  = *(const float*)&LastUpdateFrame;
 	Data[0].W  = *(const float*)&RandomID;
 
-	// TODO: Temporary PrevVelocityHack
 	if (FDataDrivenShaderPlatformInfo::GetSupportSceneDataCompressedTransforms(GMaxRHIShaderPlatform))
 	{
 		FCompressedTransform CompressedLocalToWorld(LocalToWorld);
 		Data[1] = *(const FVector4f*)&CompressedLocalToWorld.Rotation[0];
 		Data[2] = *(const FVector3f*)&CompressedLocalToWorld.Translation;
-
-		FCompressedTransform CompressedPrevLocalToWorld(PrevLocalToWorld);
-		Data[3] = *(const FVector4f*)&CompressedPrevLocalToWorld.Rotation[0];
-		Data[4] = *(const FVector3f*)&CompressedPrevLocalToWorld.Translation;
 	}
 	else
 	{
 		// Note: writes 3x float4s
 		LocalToWorld.To3x4MatrixTranspose((float*)&Data[1]);
-		PrevLocalToWorld.To3x4MatrixTranspose((float*)&Data[4]);
 	}
 }
diff --git a/Engine/Source/Runtime/Engine/Public/InstanceUniformShaderParameters.h b/Engine/Source/Runtime/Engine/Public/InstanceUniformShaderParameters.h
--- a/Engine/Source/Runtime/Engine/Public/InstanceUniformShaderParameters.h
+++ b/Engine/Source/Runtime/Engine/Public/InstanceUniformShaderParameters.h
@@ -93,9 +93,8 @@ struct FInstanceSceneShaderData
 private:
 	// Must match GetInstanceSceneData() in SceneData.ush
 	// Allocate the max Float4s usage when compressed transform is used.
-	// TODO: Temporary PrevVelocityHack (last float4s when compressed)
-	static constexpr uint32 CompressedTransformDataStrideInFloat4s = 5;
-	static constexpr uint32 UnCompressedTransformDataStrideInFloat4s = 7;
+	static constexpr uint32 CompressedTransformDataStrideInFloat4s = 3;
+	static constexpr uint32 UnCompressedTransformDataStrideInFloat4s = 4;
 
 public:
 
@@ -139,8 +138,7 @@ public:
 		uint32 CustomDataCount,
 		float RandomID,
 		const FRenderTransform& LocalToPrimitive,
-		const FRenderTransform& PrimitiveToWorld,
-		const FRenderTransform& PrevPrimitiveToWorld // TODO: Temporary PrevVelocityHack
+		const FRenderTransform& PrimitiveToWorld
 	);
 
 	ENGINE_API void BuildInternal
@@ -151,8 +149,7 @@ public:
 		uint32 LastUpdateFrame,
 		uint32 CustomDataCount,
 		float RandomID,
-		const FRenderTransform& LocalToWorld,
-		const FRenderTransform& PrevLocalToWorld // Assumes shear has been removed already // TODO: Temporary PrevVelocityHack
+		const FRenderTransform& LocalToWorld
 	);
 
 	TStaticArray<FVector4f, UnCompressedTransformDataStrideInFloat4s> Data;
diff --git a/Engine/Source/Runtime/Renderer/Private/GPUScene.cpp b/Engine/Source/Runtime/Renderer/Private/GPUScene.cpp
--- a/Engine/Source/Runtime/Renderer/Private/GPUScene.cpp
+++ b/Engine/Source/Runtime/Renderer/Private/GPUScene.cpp
@@ -1149,8 +1149,7 @@ void FGPUScene::UploadGeneral(FRDGBuilder& GraphBuilder, FScene& Scene, FRDGExte
 							UploadInfo.InstanceCustomDataCount,
 							RandomID,
 							SceneData.LocalToPrimitive,
-							UploadInfo.PrimitiveToWorld,
-							UploadInfo.PrevPrimitiveToWorld
+							UploadInfo.PrimitiveToWorld
 						);
 
 						// RefIndex* BufferState.InstanceSceneDataSOAStride + UploadInfo.InstanceSceneDataOffset + InstanceIndex
diff --git a/Engine/Source/Runtime/Renderer/Private/Renderer.cpp b/Engine/Source/Runtime/Renderer/Private/Renderer.cpp
--- a/Engine/Source/Runtime/Renderer/Private/Renderer.cpp
+++ b/Engine/Source/Runtime/Renderer/Private/Renderer.cpp
@@ -191,8 +191,7 @@ void FRendererModule::DrawTileMesh(FCanvasRenderContext& RenderContext, FMeshPas
 					INVALID_LAST_UPDATE_FRAME,
 					0 /* Custom Data Count */,
 					0.0f /* Random ID */,
-					PrimitiveParams.LocalToRelativeWorld,
-					PrimitiveParams.PreviousLocalToRelativeWorld
+					PrimitiveParams.LocalToRelativeWorld
 				);
 
 				// TODO: Payload dummy?
