Crash fix when calling UAnimInstance::IsBeingDebugged, no longer defer to proxy

diff --git a/Engine/Source/Runtime/Engine/Private/Animation/AnimInstance.cpp b/Engine/Source/Runtime/Engine/Private/Animation/AnimInstance.cpp
--- a/Engine/Source/Runtime/Engine/Private/Animation/AnimInstance.cpp
+++ b/Engine/Source/Runtime/Engine/Private/Animation/AnimInstance.cpp
@@ -3942,7 +3942,15 @@ bool UAnimInstance::QueryAndMarkTransitionEvent(int32 MachineIndex, int32 Transi
 #if WITH_EDITOR
 bool UAnimInstance::IsBeingDebugged() const
 {
-	return GetProxyOnAnyThread<FAnimInstanceProxy>().IsBeingDebugged();
+	if (IAnimClassInterface* AnimBlueprintClass = IAnimClassInterface::GetFromClass(GetClass()))
+	{
+		if (UAnimBlueprint* AnimBlueprint = Cast<UAnimBlueprint>(Cast<UAnimBlueprintGeneratedClass>(AnimBlueprintClass)->ClassGeneratedBy))
+		{
+			return AnimBlueprint->IsObjectBeingDebugged(this);
+		}
+	}
+
+	return false;
 }
 #endif
 
