Fix crash in FSkeletalMeshSceneProxy::UpdateLooseParametersUniformBuffer using CPU skinning mode.

diff --git a/Engine/Source/Runtime/Engine/Private/SkeletalMesh.cpp b/Engine/Source/Runtime/Engine/Private/SkeletalMesh.cpp
--- a/Engine/Source/Runtime/Engine/Private/SkeletalMesh.cpp
+++ b/Engine/Source/Runtime/Engine/Private/SkeletalMesh.cpp
@@ -6302,8 +6302,14 @@ void FSkeletalMeshSceneProxy::CreateBaseMeshBatch(const FSceneView* View, const
 
 void FSkeletalMeshSceneProxy::UpdateLooseParametersUniformBuffer(const FSceneView* View, const int32 SectionIndex, const FMeshBatch& Mesh, const FGPUSkinBatchElementUserData* BatchUserData) const
 {
-	// Loose parameters uniform buffer is only needed for PassThroughVF
-	if (!Mesh.VertexFactory || Mesh.VertexFactory->GetType() != &FGPUSkinPassthroughVertexFactory::StaticType)
+	// Loose parameters uniform buffer is only needed for PassThroughVF which is a type of LocalVF.
+	// Check it is a LocalVF StaticType and bGPUSkinPassThrough flag is set. Early out if either condition is not met.
+	if (!Mesh.VertexFactory || Mesh.VertexFactory->GetType() != &FLocalVertexFactory::StaticType)
+	{
+		return;
+	}
+	const FLocalVertexFactory* LocalVertexFactory = static_cast<const FLocalVertexFactory*>(Mesh.VertexFactory);
+	if (!LocalVertexFactory->bGPUSkinPassThrough)
 	{
 		return;
 	}
