Remove miscellaneous unneeded matinee-related code.

diff --git a/Engine/Source/Runtime/Engine/Public/ShowFlagsValues.inl b/Engine/Source/Runtime/Engine/Public/ShowFlagsValues.inl
--- a/Engine/Source/Runtime/Engine/Public/ShowFlagsValues.inl
+++ b/Engine/Source/Runtime/Engine/Public/ShowFlagsValues.inl
@@ -102,7 +102,7 @@ SHOWFLAG_FIXED_IN_SHIPPING(0, VertexColors, SFG_Advanced, NSLOCTEXT("UnrealEd",
 SHOWFLAG_FIXED_IN_SHIPPING(0, PhysicalMaterialMasks, SFG_Advanced, NSLOCTEXT("UnrealEd", "PhysicalMaterialMasksSF", "Physical Material Masks"))
 /** Render Post process (screen space) distortion/refraction */
 SHOWFLAG_FIXED_IN_SHIPPING(1, Refraction, SFG_Developer, NSLOCTEXT("UnrealEd", "RefractionSF", "Refraction"))
-/** Usually set in game or when previewing Matinee but not in editor, used for motion blur or any kind of rendering features that rely on the former frame */
+/** Usually set in game or when previewing cinematics but not in editor, used for motion blur or any kind of rendering features that rely on the former frame */
 SHOWFLAG_ALWAYS_ACCESSIBLE(CameraInterpolation, SFG_Hidden, NSLOCTEXT("UnrealEd", "CameraInterpolationSF", "Camera Interpolation"))
 /** Post processing color fringe (chromatic aberration) */
 SHOWFLAG_FIXED_IN_SHIPPING(1, SceneColorFringe, SFG_PostProcess, NSLOCTEXT("UnrealEd", "SceneColorFringeSF", "Scene Color Fringe"))
