Added DummyViewport.h placeholder file to fix build farm issue

diff --git a/Engine/Source/Runtime/Engine/Public/DummyViewport.h b/Engine/Source/Runtime/Engine/Public/DummyViewport.h
new file mode 100644
--- /dev/null
+++ b/Engine/Source/Runtime/Engine/Public/DummyViewport.h
@@ -0,0 +1,5 @@
+// Copyright Epic Games, Inc. All Rights Reserved.
+
+#pragma once
+
+// Placeholder for coming changelist
\ No newline at end of file
