- Fix crash when using Niagara high precision sorting in post render opaque

#rnx

diff --git a/Engine/Source/Runtime/Engine/Private/Particles/ParticleGpuSimulation.cpp b/Engine/Source/Runtime/Engine/Private/Particles/ParticleGpuSimulation.cpp
--- a/Engine/Source/Runtime/Engine/Private/Particles/ParticleGpuSimulation.cpp
+++ b/Engine/Source/Runtime/Engine/Private/Particles/ParticleGpuSimulation.cpp
@@ -4504,7 +4504,12 @@ bool FFXSystem::AddSortedGPUSimulation(FParticleSimulationGPU* Simulation, const
 void FFXSystem::GenerateSortKeys(FRHICommandListImmediate& RHICmdList, int32 BatchId, int32 NumElementsInBatch, EGPUSortFlags Flags, FRHIUnorderedAccessView* KeysUAV, FRHIUnorderedAccessView* ValuesUAV)
 {
 	check(EnumHasAnyFlags(Flags, EGPUSortFlags::KeyGenAfterPostRenderOpaque));
-	check(EnumHasAnyFlags(Flags, EGPUSortFlags::LowPrecisionKeys));
+
+	// Cascade does not support high precision keys so we can safely ignore this as we will not have a batch which conatins anything but low precision
+	if (EnumHasAnyFlags(Flags, EGPUSortFlags::LowPrecisionKeys) == false)
+	{
+		return;
+	}
 
 	// First generate keys for each emitter to be sorted.
 	const int32 TotalParticleCount = GenerateParticleSortKeys(
