Reduce the refresh rate CVAR for AudioCookOverrides

diff --git a/Engine/Source/Runtime/Engine/Private/AudioCompressionSettingsUtils.cpp b/Engine/Source/Runtime/Engine/Private/AudioCompressionSettingsUtils.cpp
--- a/Engine/Source/Runtime/Engine/Private/AudioCompressionSettingsUtils.cpp
+++ b/Engine/Source/Runtime/Engine/Private/AudioCompressionSettingsUtils.cpp
@@ -21,7 +21,7 @@
 
 #include "Misc/ConfigCacheIni.h"
 
-static float CookOverrideCachingIntervalCvar = 1.0f;
+static float CookOverrideCachingIntervalCvar = 60.0f;
 FAutoConsoleVariableRef CVarCookOverrideCachingIntervalCVar(
 	TEXT("au.editor.CookOverrideCachingInterval"),
 	CookOverrideCachingIntervalCvar,
