Add PerQualityLeveleProperty to LandscapeGrassType on StartCullingDistance, EndCullingDistance and GrassDensity. This will allow to have per instance value on platforms that have  consoles with different specs.

-Add PerQualityLeveleFloat type
-Add customization in LandscapeGrassType to enable or disable perquality/perplaform properties depending on the Project settings()
-Add the possibility to define the perQuality cvar in device profile
-Add new commandlet to convert grassType perPlatform data to perQuality levels

diff --git a/Engine/Source/Editor/DetailCustomizations/Private/DetailCustomizations.cpp b/Engine/Source/Editor/DetailCustomizations/Private/DetailCustomizations.cpp
--- a/Engine/Source/Editor/DetailCustomizations/Private/DetailCustomizations.cpp
+++ b/Engine/Source/Editor/DetailCustomizations/Private/DetailCustomizations.cpp
@@ -159,11 +159,13 @@
 #include "VectorStructCustomization.h"
 #include "WindowsTargetSettingsDetails.h"
 #include "WorldSettingsDetails.h"
+#include "LandscapeGrassTypeDetails.h"
 
 struct FPerPlatformBool;
 struct FPerPlatformFloat;
 struct FPerPlatformInt;
 struct FPerQualityLevelInt;
+struct FPerQualityLevelFloat;
 
 IMPLEMENT_MODULE( FDetailCustomizationsModule, DetailCustomizations );
 
@@ -342,6 +344,7 @@ void FDetailCustomizationsModule::RegisterPropertyTypeCustomizations()
 	RegisterCustomPropertyTypeLayout("PerPlatformFloat", FOnGetPropertyTypeCustomizationInstance::CreateStatic(&FPerPlatformPropertyCustomization<FPerPlatformFloat>::MakeInstance));
 	RegisterCustomPropertyTypeLayout("PerPlatformBool", FOnGetPropertyTypeCustomizationInstance::CreateStatic(&FPerPlatformPropertyCustomization<FPerPlatformBool>::MakeInstance));
 	RegisterCustomPropertyTypeLayout("PerQualityLevelInt", FOnGetPropertyTypeCustomizationInstance::CreateStatic(&FPerQualityLevelPropertyCustomization<FPerQualityLevelInt>::MakeInstance));
+	RegisterCustomPropertyTypeLayout("PerQualityLevelFloat", FOnGetPropertyTypeCustomizationInstance::CreateStatic(&FPerQualityLevelPropertyCustomization<FPerQualityLevelFloat>::MakeInstance));
 	RegisterCustomPropertyTypeLayout("SkeletalMeshOptimizationSettings", FOnGetPropertyTypeCustomizationInstance::CreateStatic(&FSkeletalMeshReductionSettingsDetails::MakeInstance));
 	RegisterCustomPropertyTypeLayout("GrassInput", FOnGetPropertyTypeCustomizationInstance::CreateStatic(&FMaterialExpressionLandscapeGrassInputCustomization::MakeInstance));
 	RegisterCustomPropertyTypeLayout("ComponentReference", FOnGetPropertyTypeCustomizationInstance::CreateStatic(&FComponentReferenceCustomization::MakeInstance));
@@ -465,6 +468,7 @@ void FDetailCustomizationsModule::RegisterObjectCustomizations()
 
 	RegisterCustomClassLayout("Landscape", FOnGetDetailCustomizationInstance::CreateStatic(&FLandscapeUIDetails::MakeInstance));
 	RegisterCustomClassLayout("LandscapeProxy", FOnGetDetailCustomizationInstance::CreateStatic(&FLandscapeProxyUIDetails::MakeInstance));
+	RegisterCustomClassLayout("LandscapeGrassType", FOnGetDetailCustomizationInstance::CreateStatic(&FLandscapeGrassTypeDetails::MakeInstance));
 
 	RegisterCustomClassLayout("BoundsCopyComponent", FOnGetDetailCustomizationInstance::CreateStatic(&FBoundsCopyComponentDetailsCustomization::MakeInstance));
 }
diff --git a/Engine/Source/Editor/DetailCustomizations/Private/LandscapeGrassTypeDetails.cpp b/Engine/Source/Editor/DetailCustomizations/Private/LandscapeGrassTypeDetails.cpp
new file mode 100644
--- /dev/null
+++ b/Engine/Source/Editor/DetailCustomizations/Private/LandscapeGrassTypeDetails.cpp
@@ -0,0 +1,101 @@
+// Copyright Epic Games, Inc. All Rights Reserved.
+
+#include "LandscapeGrassTypeDetails.h"
+
+#include "Containers/Array.h"
+#include "Delegates/Delegate.h"
+#include "DetailLayoutBuilder.h"
+#include "Engine/EngineTypes.h"
+#include "Misc/AssertionMacros.h"
+#include "PropertyHandle.h"
+#include "PropertyRestriction.h"
+#include "Templates/Casts.h"
+#include "UObject/Class.h"
+#include "UObject/NameTypes.h"
+#include "UObject/ObjectPtr.h"
+#include "UObject/ReflectedTypeAccessors.h"
+#include "UObject/UObjectGlobals.h"
+#include "UObject/UnrealType.h"
+#include "LandscapeGrassType.h"
+
+class UObject;
+
+static bool GShowBothPerQualityAndPerPlaformProperties = false;
+static FAutoConsoleVariableRef CVarMinFramesToKeepGrass(
+	TEXT("r.grass.ShowBothPerQualityAndPerPlaformProperties"),
+	GShowBothPerQualityAndPerPlaformProperties,
+	TEXT("Show both per platform and per quality properties in the editor."));
+
+#define LOCTEXT_NAMESPACE "LandscapeGrassTypeDetails"
+
+TSharedRef<IDetailCustomization> FLandscapeGrassTypeDetails::MakeInstance()
+{
+	return MakeShareable(new FLandscapeGrassTypeDetails);
+}
+
+FLandscapeGrassTypeDetails::~FLandscapeGrassTypeDetails()
+{
+}
+
+void FLandscapeGrassTypeDetails::CustomizeDetails( IDetailLayoutBuilder& DetailLayout )
+{
+	if (GShowBothPerQualityAndPerPlaformProperties)
+	{
+		return;
+	}
+
+	TArray<TWeakObjectPtr<UObject>> Objects;
+	DetailLayout.GetObjectsBeingCustomized(Objects);
+
+	const TArray<FName> PerPlatformProperties =
+	{
+		GET_MEMBER_NAME_CHECKED(FGrassVariety, GrassDensity),
+		GET_MEMBER_NAME_CHECKED(FGrassVariety, StartCullDistance),
+		GET_MEMBER_NAME_CHECKED(FGrassVariety, EndCullDistance),
+	};
+
+	const TArray<FName> PerQualityProperties =
+	{
+		GET_MEMBER_NAME_CHECKED(FGrassVariety, GrassDensityQuality),
+		GET_MEMBER_NAME_CHECKED(FGrassVariety, StartCullDistanceQuality),
+		GET_MEMBER_NAME_CHECKED(FGrassVariety, EndCullDistanceQuality),
+	};
+
+	if (Objects.Num() > 0)
+	{
+		// get the grass variety array
+		TSharedPtr<IPropertyHandle> GrassVariety = DetailLayout.GetProperty(GET_MEMBER_NAME_CHECKED(ULandscapeGrassType, GrassVarieties));
+		TSharedPtr<IPropertyHandleArray> GrassVarietyArrayHandle = GrassVariety->AsArray();
+
+		if (GrassVarietyArrayHandle.IsValid())
+		{
+			uint32 NumElements;
+			GrassVarietyArrayHandle->GetNumElements(NumElements);
+
+			for (uint32 Index = 0; Index < NumElements; ++Index)
+			{
+				TSharedRef<IPropertyHandle> ElementHandle = GrassVarietyArrayHandle->GetElement(Index);
+
+				uint32 NumChildren;
+				ElementHandle->GetNumChildren(NumChildren);
+
+				for (uint32 ChildIndex = 0; ChildIndex < NumChildren; ++ChildIndex)
+				{
+					TSharedRef<IPropertyHandle> ChildHandle = ElementHandle->GetChildHandle(ChildIndex).ToSharedRef();
+					const FName& ChildPropertyName = ChildHandle->GetProperty()->GetFName();
+
+					if (PerPlatformProperties.Contains(ChildPropertyName) && GEngine && GEngine->UseGrassVarityPerQualityLevels)
+					{
+						DetailLayout.HideProperty(ChildHandle);
+					}
+
+					if (PerQualityProperties.Contains(ChildPropertyName) && GEngine && !GEngine->UseGrassVarityPerQualityLevels)
+					{
+						DetailLayout.HideProperty(ChildHandle);
+					}
+				}
+			}
+		}
+	}
+}
+#undef LOCTEXT_NAMESPACE
diff --git a/Engine/Source/Editor/DetailCustomizations/Private/LandscapeGrassTypeDetails.h b/Engine/Source/Editor/DetailCustomizations/Private/LandscapeGrassTypeDetails.h
new file mode 100644
--- /dev/null
+++ b/Engine/Source/Editor/DetailCustomizations/Private/LandscapeGrassTypeDetails.h
@@ -0,0 +1,27 @@
+// Copyright Epic Games, Inc. All Rights Reserved.
+
+#pragma once
+
+#include "Delegates/IDelegateInstance.h"
+#include "IDetailCustomization.h"
+#include "Templates/SharedPointer.h"
+#include "UObject/WeakObjectPtr.h"
+#include "UObject/WeakObjectPtrTemplates.h"
+
+class FPropertyRestriction;
+class IDetailLayoutBuilder;
+class UObject;
+struct FPropertyChangedEvent;
+
+class FLandscapeGrassTypeDetails : public IDetailCustomization
+{
+public:
+	/** Makes a new instance of this detail layout class for a specific detail view requesting it */
+	static TSharedRef<IDetailCustomization> MakeInstance();
+
+	virtual ~FLandscapeGrassTypeDetails();
+
+	/** IDetailCustomization interface */
+	virtual void CustomizeDetails( IDetailLayoutBuilder& DetailLayout ) override;
+};
+
diff --git a/Engine/Source/Editor/DetailCustomizations/Private/PerQualityLevelPropertyCustomization.cpp b/Engine/Source/Editor/DetailCustomizations/Private/PerQualityLevelPropertyCustomization.cpp
--- a/Engine/Source/Editor/DetailCustomizations/Private/PerQualityLevelPropertyCustomization.cpp
+++ b/Engine/Source/Editor/DetailCustomizations/Private/PerQualityLevelPropertyCustomization.cpp
@@ -251,6 +251,6 @@ TSharedRef<IPropertyTypeCustomization> FPerQualityLevelPropertyCustomization<Ove
 /* Only explicitly instantiate the types which are supported
 *****************************************************************************/
 template class FPerQualityLevelPropertyCustomization<FPerQualityLevelInt>;
-
+template class FPerQualityLevelPropertyCustomization<FPerQualityLevelFloat>;
 
 #undef LOCTEXT_NAMESPACE
diff --git a/Engine/Source/Editor/UnrealEd/Classes/Commandlets/LandscapeGrassTypeCommandlet.h b/Engine/Source/Editor/UnrealEd/Classes/Commandlets/LandscapeGrassTypeCommandlet.h
new file mode 100644
--- /dev/null
+++ b/Engine/Source/Editor/UnrealEd/Classes/Commandlets/LandscapeGrassTypeCommandlet.h
@@ -0,0 +1,21 @@
+// Copyright Epic Games, Inc. All Rights Reserved.
+
+#pragma once
+
+#include "CoreMinimal.h"
+#include "UObject/ObjectMacros.h"
+#include "Commandlets/Commandlet.h"
+#include "LandscapeGrassTypeCommandlet.generated.h"
+
+/** Commandlet  */
+UCLASS()
+class ULandscapeGrassTypeCommandlet : public UCommandlet
+{
+	GENERATED_UCLASS_BODY()
+
+public:
+	//~ Begin UCommandlet Interface
+	virtual int32 Main(const FString& Params) override;
+
+	//~ End UCommandlet Interface
+};
\ No newline at end of file
diff --git a/Engine/Source/Editor/UnrealEd/Private/Commandlets/ContentCommandlets.cpp b/Engine/Source/Editor/UnrealEd/Private/Commandlets/ContentCommandlets.cpp
--- a/Engine/Source/Editor/UnrealEd/Private/Commandlets/ContentCommandlets.cpp
+++ b/Engine/Source/Editor/UnrealEd/Private/Commandlets/ContentCommandlets.cpp
@@ -32,6 +32,7 @@
 #include "Commandlets/ListMaterialsUsedWithMeshEmittersCommandlet.h"
 #include "Commandlets/ListStaticMeshesImportedFromSpeedTreesCommandlet.h"
 #include "Commandlets/StaticMeshMinLodCommandlet.h"
+#include "Commandlets/LandscapeGrassTypeCommandlet.h"
 #include "Particles/ParticleSystem.h"
 #include "Commandlets/ResavePackagesCommandlet.h"
 #include "Commandlets/WrangleContentCommandlet.h"
@@ -39,6 +40,7 @@
 #include "Particles/ParticleEmitter.h"
 #include "GameFramework/WorldSettings.h"
 #include "Engine/StaticMesh.h"
+#include "Landscape/Classes/LandscapeGrassType.h"
 #include "AssetRegistry/AssetData.h"
 #include "Engine/Brush.h"
 #include "Editor.h"
@@ -3514,30 +3516,6 @@ int32 UStaticMeshMinLodCommandlet::Main(const FString& Params)
 	AssetRegistryModule.Get().GetAssetsByClass(UStaticMesh::StaticClass()->GetClassPathName(), AssetList, true);
 	TArray<UPackage*> PackagesToSave;
 
-	// Platform (group) names
-	const TArray<FString> Filters = { TEXT("NoEditor"), TEXT("Client"), TEXT("Server"), TEXT("AllDesktop") };
-	TMultiMap<FName, FName> GroupToPlatform;
-
-	// sanitize all vanilla platform names
-	// generate a list of all supported platform (from the datadrivenplatform files)
-	const TArray<FName>& SanitizedPlatformNameArray = PlatformInfo::GetAllVanillaPlatformNames().FilterByPredicate([&Filters, &GroupToPlatform](const FName& PlatformName)
-		{
-			for (const FString& Filter : Filters)
-			{
-				const int32 Position = PlatformName.ToString().Find(Filter);
-				if (Position != INDEX_NONE)
-				{
-					return false;
-				}
-			}
-
-			if (const PlatformInfo::FTargetPlatformInfo* PlatformInfo = PlatformInfo::FindPlatformInfo(PlatformName))
-			{
-				GroupToPlatform.AddUnique(PlatformInfo->DataDrivenPlatformInfo->PlatformGroupName, PlatformName);
-			}
-			return true;
-		});
-
 	for (int32 AssetIdx = 0; AssetIdx < AssetList.Num(); ++AssetIdx)
 	{
 		bool SavePackage = false;
@@ -3699,3 +3677,137 @@ int32 UStaticMeshMinLodCommandlet::Main(const FString& Params)
 	
 	return 0;
 }
+
+
+/* ==========================================================================================================
+	ULandscapeGrassTypeCommandlet
+========================================================================================================== */
+
+ULandscapeGrassTypeCommandlet::ULandscapeGrassTypeCommandlet(const FObjectInitializer& ObjectInitializer)
+	: Super(ObjectInitializer)
+{
+}
+
+int32 ULandscapeGrassTypeCommandlet::Main(const FString& Params)
+{
+	TArray<FString> FilesInPath;
+	FEditorFileUtils::FindAllPackageFiles(FilesInPath);
+	if (FilesInPath.Num() == 0)
+	{
+		UE_LOG(LogContentCommandlet, Warning, TEXT("No packages found"));
+		return 1;
+	}
+
+	// parse the mapping from PerPlatform to PerQualityLevel
+	// ex:-mapping=Mobile:Low,Switch:Medium,Desktop:High,PS4:High,XboxOne:High,PS5:Epic,XSX:Epic,console:medium;high;epic,Desktop:high
+	TMultiMap<FName, FName> PerPlatformToQualityLevel;
+	FString MappingStr;
+	if (FParse::Value(*Params, TEXT("-mapping="), MappingStr, false))
+	{
+		TArray<FString> Mappings;
+		MappingStr.ParseIntoArray(Mappings, TEXT(","), true);
+		for (FString& PlatformToQualityLevel : Mappings)
+		{
+			TArray<FString> Entries;
+			PlatformToQualityLevel.ParseIntoArray(Entries, TEXT(":"), false);
+
+			if (Entries.Num() != 2)
+			{
+				UE_LOG(LogContentCommandlet, Error, TEXT("Error bad -mapping argument: %s"), *MappingStr);
+				return 1;
+			}
+
+			TArray<FString> Values;
+			Entries[1].ParseIntoArray(Values, TEXT(";"), false);
+
+			for (const FString& Value : Values)
+			{
+				PerPlatformToQualityLevel.AddUnique(FName(*Entries[0]), FName(*Value));
+			}
+		}
+	}
+
+	if (PerPlatformToQualityLevel.Num() == 0)
+	{
+		UE_LOG(LogContentCommandlet, Warning, TEXT("No Mapping rules found"));
+		return 1;
+	}
+
+	bool bNoSourceControl = FParse::Param(*Params, TEXT("nosourcecontrol"));
+	bool bGenerateCollections = FParse::Param(*Params, TEXT("collections"));
+	ISourceControlProvider* SourceControlProvider = bNoSourceControl ? nullptr : &ISourceControlModule::Get().GetProvider();
+	ICollectionManager& CollectionManager = FCollectionManagerModule::GetModule().Get();
+
+	// Load the asset registry module
+	FAssetRegistryModule& AssetRegistryModule = FModuleManager::LoadModuleChecked<FAssetRegistryModule>(TEXT("AssetRegistry"));
+
+	// Update Registry Module
+	UE_LOG(LogContentCommandlet, Display, TEXT("Searching Asset Registry for static mesh "));
+	AssetRegistryModule.Get().SearchAllAssets(true);
+
+	// Retrieve list of all assets, used to find unreferenced ones.
+	TArray<FAssetData> AssetList;
+	int32 TotalPackagesChecked = 0;
+	AssetRegistryModule.Get().GetAssetsByClass(ULandscapeGrassType::StaticClass()->GetClassPathName(), AssetList, true);
+	TArray<UPackage*> PackagesToSave;
+
+	for (int32 AssetIdx = 0; AssetIdx < AssetList.Num(); ++AssetIdx)
+	{
+		bool SavePackage = false;
+		const FString Filename = AssetList[AssetIdx].GetObjectPathString();
+		UE_LOG(LogContentCommandlet, Display, TEXT("Processing static mesh (%i/%i):  %s "), AssetIdx, AssetList.Num(), *Filename);
+
+		UPackage* Package = LoadPackage(NULL, *Filename, LOAD_Quiet);
+		if (Package == NULL)
+		{
+			UE_LOG(LogContentCommandlet, Error, TEXT("Error loading %s!"), *Filename);
+			continue;
+		}
+
+		TotalPackagesChecked++;
+
+		for (TObjectIterator<ULandscapeGrassType> It; It; ++It)
+		{
+			ULandscapeGrassType* GrassType = *It;
+			if (GrassType->IsIn(Package) && !GrassType->IsTemplate())
+			{
+				for (FGrassVariety& GrassVariety : GrassType->GrassVarieties)
+				{
+					GrassVariety.GrassDensityQuality.ConvertQualtiyLevelData(GrassVariety.GrassDensity.PerPlatform, PerPlatformToQualityLevel, GrassVariety.GrassDensity.Default);
+					GrassVariety.StartCullDistanceQuality.ConvertQualtiyLevelData(GrassVariety.StartCullDistance.PerPlatform, PerPlatformToQualityLevel, GrassVariety.StartCullDistance.Default);
+					GrassVariety.EndCullDistanceQuality.ConvertQualtiyLevelData(GrassVariety.EndCullDistance.PerPlatform, PerPlatformToQualityLevel, GrassVariety.EndCullDistance.Default);
+				}
+			}
+		}
+
+		if (!SavePackage)
+		{
+			PackagesToSave.AddUnique(Package);
+			SavePackage = true;
+		}
+	}
+
+	// save quality level modifications 
+	if (SourceControlProvider)
+	{
+		FEditorFileUtils::CheckoutPackages(PackagesToSave, nullptr, false);
+	}
+	else
+	{
+		for (UPackage* Package : PackagesToSave)
+		{
+			FString PackageFilename = SourceControlHelpers::PackageFilename(Package);
+			if (IPlatformFile::GetPlatformPhysical().FileExists(*PackageFilename))
+			{
+				if (!IPlatformFile::GetPlatformPhysical().SetReadOnly(*PackageFilename, false))
+				{
+					UE_LOG(LogContentCommandlet, Error, TEXT("Error setting %s writable"), *PackageFilename);
+					return 1;
+				}
+			}
+		}
+	}
+	FEditorFileUtils::PromptForCheckoutAndSave(PackagesToSave, false, false, nullptr, true, false);
+	
+	return 0;
+}
\ No newline at end of file
diff --git a/Engine/Source/Runtime/Engine/Classes/Engine/Engine.h b/Engine/Source/Runtime/Engine/Classes/Engine/Engine.h
--- a/Engine/Source/Runtime/Engine/Classes/Engine/Engine.h
+++ b/Engine/Source/Runtime/Engine/Classes/Engine/Engine.h
@@ -1552,6 +1552,9 @@ public:
 	UPROPERTY(config, EditAnywhere, Category = PerQualityLevelProperty, AdvancedDisplay)
 	bool UseSkeletalMeshMinLODPerQualityLevels;
 
+	UPROPERTY(config, EditAnywhere, Category = PerQualityLevelProperty, AdvancedDisplay)
+	bool UseGrassVarityPerQualityLevels;
+
 	/** The state of the current map transition.  */
 	UPROPERTY()
 	ETransitionType TransitionType;
diff --git a/Engine/Source/Runtime/Engine/Private/PerQualityLevelProperties.cpp b/Engine/Source/Runtime/Engine/Private/PerQualityLevelProperties.cpp
--- a/Engine/Source/Runtime/Engine/Private/PerQualityLevelProperties.cpp
+++ b/Engine/Source/Runtime/Engine/Private/PerQualityLevelProperties.cpp
@@ -9,6 +9,8 @@
 #if WITH_EDITOR
 #include "Interfaces/ITargetPlatform.h"
 #include "PlatformInfo.h"
+#include "DeviceProfiles/DeviceProfile.h"
+#include "DeviceProfiles/DeviceProfileManager.h"
 #endif
 
 
@@ -17,9 +19,10 @@ namespace QualityLevelProperty
 	static TArray<FName> QualityLevelNames = { FName("Low"), FName("Medium"), FName("High"), FName("Epic"), FName("Cinematic") };
 	static FString QualityLevelMappingStr = TEXT("QualityLevelMapping");
 
-	TMap<int32, int32> ConvertQualtiyLevelData(const TMap<EPerQualityLevels, int32>& Data)
+	template<typename _ValueType>
+	TMap<int32, _ValueType> ConvertQualtiyLevelData(const TMap<EPerQualityLevels, _ValueType>& Data)
 	{
-		TMap<int32, int32> ConvertedData;
+		TMap<int32, _ValueType> ConvertedData;
 
 		for (const TPair<EPerQualityLevels, int32>& Pair : Data)
 		{
@@ -29,11 +32,12 @@ namespace QualityLevelProperty
 		return ConvertedData;
 	}
 
-	TMap<EPerQualityLevels, int32> ConvertQualtiyLevelData(const TMap<int32, int32>& Data)
+	template<typename _ValueType>
+	TMap<EPerQualityLevels, _ValueType> ConvertQualtiyLevelData(const TMap<int32, _ValueType>& Data)
 	{
-		TMap<EPerQualityLevels, int32> ConvertedData;
+		TMap<EPerQualityLevels, _ValueType> ConvertedData;
 
-		for (const TPair<int32, int32>& Pair : Data)
+		for (const TPair<int32, _ValueType>& Pair : Data)
 		{
 			ConvertedData.Add((EPerQualityLevels)Pair.Key, Pair.Value);
 		}
@@ -142,44 +146,81 @@ namespace QualityLevelProperty
 static TMap<FString, FSupportedQualityLevelArray> GSupportedQualityLevels;
 static FCriticalSection GCookCriticalSection;
 
+template<typename StructType, typename ValueType, EName _BasePropertyName>
+void FPerQualityLevelProperty<StructType, ValueType, _BasePropertyName>::ConvertQualtiyLevelData(TMap<FName, ValueType>& PlaformData, TMultiMap<FName, FName>& PerPlatformToQualityLevel, ValueType Default)
+{
+	StructType* This = StaticCast<StructType*>(this);
+	
+	This->Default = Default;
+
+	// convert each platform overrides
+	for (const TPair<FName, ValueType>& Pair : PlaformData)
+	{
+		// get all quality levels associated with the PerPlatform override 
+		TArray<FName> QLNames;
+		PerPlatformToQualityLevel.MultiFind(Pair.Key, QLNames);
 
-template<typename _StructType, typename _ValueType, EName _BasePropertyName>
-int32 FPerQualityLevelProperty<_StructType, _ValueType, _BasePropertyName>::GetValueForPlatform(const ITargetPlatform* TargetPlatform) const
+		for (const FName& QLName : QLNames)
+		{
+			int32 QLKey = QualityLevelProperty::FNameToQualityLevel(QLName);
+			if (QLKey != INDEX_NONE)
+			{
+				ValueType* Value = This->PerQuality.Find(QLKey);
+
+				// if the quality level already as a value, only change it if the value is lower
+				// this can happen if two mapping key as the same quality level but different values
+				if (Value != nullptr && Pair.Value < *Value)
+				{
+					// only change the override if its a smaller 
+					*Value = Pair.Value;
+				}
+				else
+				{
+					This->PerQuality.Add(QLKey, Pair.Value);
+				}
+			}
+		}
+	}
+}
+
+template<typename StructType, typename ValueType, EName _BasePropertyName>
+int32 FPerQualityLevelProperty<StructType, ValueType, _BasePropertyName>::GetValueForPlatform(const ITargetPlatform* TargetPlatform) const
 {
-	const _StructType* This = StaticCast<const _StructType*>(this);
+	const StructType* This = StaticCast<const StructType*>(this);
 	// get all supported quality level from scalability + engine ini files
 	FSupportedQualityLevelArray SupportedQualityLevels = GetSupportedQualityLevels(*TargetPlatform->GetPlatformInfo().IniPlatformName.ToString());
 
-	// loop through all the supported quality level to find the min lod index
-	int32 MinLodIdx = MAX_int32;
+	// loop through all the supported quality level to find the min value
+	ValueType MinValue = This->MaxType();
+
 	for (int32& QL : SupportedQualityLevels)
 	{
 		// check if have data for the supported quality level
 		if (IsQualityLevelValid(QL))
 		{
-			MinLodIdx = FMath::Min(GetValueForQualityLevel(QL), MinLodIdx);
+			MinValue = FMath::Min(GetValueForQualityLevel(QL), MinValue);
 		}
 	}
 
-	if (MinLodIdx == MAX_int32)
+	if (MinValue == This->MaxType())
 	{
-		MinLodIdx = This->Default;
+		MinValue = This->Default;
 	}
 
-	return MinLodIdx;
+	return MinValue;
 }
 
-template<typename _StructType, typename _ValueType, EName _BasePropertyName>
-FSupportedQualityLevelArray FPerQualityLevelProperty<_StructType, _ValueType, _BasePropertyName>::GetSupportedQualityLevels(const TCHAR* InPlatformName) const
+template<typename StructType, typename ValueType, EName _BasePropertyName>
+FSupportedQualityLevelArray FPerQualityLevelProperty<StructType, ValueType, _BasePropertyName>::GetSupportedQualityLevels(const TCHAR* InPlatformName) const
 {
 	const FString PlatformNameStr = FDataDrivenPlatformInfoRegistry::GetPlatformInfo(FName(InPlatformName)).IniPlatformName.ToString();
 	InPlatformName = *PlatformNameStr;
 
 	FSupportedQualityLevelArray* CachedCookingQualitLevelInfo = nullptr;
-	
+	FString UniqueName = FString(InPlatformName) + CVarName;
 	{
 		FScopeLock ScopeLock(&GCookCriticalSection);
-		CachedCookingQualitLevelInfo = GSupportedQualityLevels.Find(InPlatformName);
+		CachedCookingQualitLevelInfo = GSupportedQualityLevels.Find(UniqueName);
 		if (CachedCookingQualitLevelInfo)
 		{
 			return *CachedCookingQualitLevelInfo;
@@ -188,53 +229,74 @@ FSupportedQualityLevelArray FPerQualityLevelProperty<_StructType, _ValueType, _B
 
 	FSupportedQualityLevelArray CookingQualitLevelInfo;
 
-	// check the Engine file
-	FConfigFile EngineSettings;
-	FConfigCacheIni::LoadLocalIniFile(EngineSettings, TEXT("Engine"), true, InPlatformName);
+	//get the platform config cache ini
+	FConfigCacheIni* ConfigSystemPlatform = FConfigCacheIni::ForPlatform(FName(InPlatformName));
 
-	int32 PropertyQualityLevel = -1;
-	if (EngineSettings.GetInt(TEXT("SystemSettings"), *CVarName, PropertyQualityLevel))
+	// check the Engine file
+	if (FConfigFile* PlatformEngine = ConfigSystemPlatform->FindConfigFile(GEngineIni))
 	{
-		CookingQualitLevelInfo.Add(PropertyQualityLevel);
+		int32 PropertyQualityLevel = -1;
+		if (PlatformEngine->GetInt(TEXT("SystemSettings"), *CVarName, PropertyQualityLevel))
+		{
+			CookingQualitLevelInfo.Add(PropertyQualityLevel);
+		}
 	}
 
 	// Load the scalability platform file
-	FConfigFile ScalabilitySettings;
-	FConfigCacheIni::LoadLocalIniFile(ScalabilitySettings, TEXT("Scalability"), true, InPlatformName);
+	if (FConfigFile* PlatformScalability = ConfigSystemPlatform->FindConfigFile(GScalabilityIni))
+	{
+		//check all possible quality levels specify in the scalability ini 
+		for (int32 QualityLevel = 0; QualityLevel < (int32)EPerQualityLevels::Num; ++QualityLevel)
+		{
+			FString QualitLevelSectionName = Scalability::GetScalabilitySectionString(*ScalabilitySection, QualityLevel, (int32)EPerQualityLevels::Num);
+			int32 PropertyQualityLevel = -1;
+			PlatformScalability->GetInt(*QualitLevelSectionName, *CVarName, PropertyQualityLevel);
+
+			// add supported quality level to the property map
+			if (PropertyQualityLevel != -1)
+			{
+				CookingQualitLevelInfo.Add(PropertyQualityLevel);
+			}
+		}
+	}
 
-	//check all possible quality levels specify in the scalability ini 
-	for (int32 QualityLevel = 0; QualityLevel < (int32)EPerQualityLevels::Num; ++QualityLevel)
+	// get all the device profile related to the platform we want to cook
+	TArray<TObjectPtr<UDeviceProfile>> DeviceProfiles = UDeviceProfileManager::Get().Profiles.FilterByPredicate([&InPlatformName](const TObjectPtr<UDeviceProfile>& DeviceProfile)
 	{
-		FString QualitLevelSectionName = Scalability::GetScalabilitySectionString(*ScalabilitySection, QualityLevel, (int32)EPerQualityLevels::Num);
-		PropertyQualityLevel = -1;
-		ScalabilitySettings.GetInt(*QualitLevelSectionName, *CVarName, PropertyQualityLevel);
+		return (DeviceProfile->DeviceType == InPlatformName);
+	});
 
-		// add supported quality level to the property map
-		if (PropertyQualityLevel != -1)
+	// iterate through the DP sections to find the ones relevant to the platform
+	for (TObjectPtr<UDeviceProfile>& DeviceProfile : DeviceProfiles)
+	{
+		const TMap<FString, FString>& ExpandedCVars = DeviceProfile->GetAllExpandedCVars();
+		const FString* const CVarValue = ExpandedCVars.Find(CVarName);
+		if (CVarValue)
 		{
-			CookingQualitLevelInfo.Add(PropertyQualityLevel);
+			int32 Result = FCString::Atoi(*(* CVarValue));
+			CookingQualitLevelInfo.Add(Result);
 		}
 	}
 
 	// Cache the Scalability setting for this platform
 	{
 		FScopeLock ScopeLock(&GCookCriticalSection);
-		CachedCookingQualitLevelInfo = &GSupportedQualityLevels.Add(FString(InPlatformName), CookingQualitLevelInfo);
+		CachedCookingQualitLevelInfo = &GSupportedQualityLevels.Add(UniqueName, CookingQualitLevelInfo);
 		return *CachedCookingQualitLevelInfo;
 	}
 }
 
-template<typename _StructType, typename _ValueType, EName _BasePropertyName>
-void FPerQualityLevelProperty<_StructType, _ValueType, _BasePropertyName>::StripQualtiyLevelForCooking(const TCHAR* InPlatformName)
+template<typename StructType, typename ValueType, EName _BasePropertyName>
+void FPerQualityLevelProperty<StructType, ValueType, _BasePropertyName>::StripQualtiyLevelForCooking(const TCHAR* InPlatformName)
 {
-	_StructType* This = StaticCast<_StructType*>(this);
+	StructType* This = StaticCast<StructType*>(this);
 	if (This->PerQuality.Num() > 0)
 	{
 		FSupportedQualityLevelArray CookQualityLevelInfo = This->GetSupportedQualityLevels(InPlatformName);
 		CookQualityLevelInfo.Sort([&](const int32& A, const int32& B) { return (A > B); });
 
 		// remove unsupported quality levels 
-		for (TMap<int32, int32>::TIterator It(This->PerQuality); It; ++It)
+		for (typename TMap<int32, ValueType>::TIterator It(This->PerQuality); It; ++It)
 		{
 			if (!CookQualityLevelInfo.Contains(It.Key()))
 			{
@@ -248,16 +310,16 @@ void FPerQualityLevelProperty<_StructType, _ValueType, _BasePropertyName>::Strip
 
 			for (TSet<int32>::TIterator It(CookQualityLevelInfo); It; ++It)
 			{
-				int32* QualityLevelMinLod = This->PerQuality.Find(*It);
+				ValueType* QualityLevelMinValue = This->PerQuality.Find(*It);
 			
 				// add quality level supported by the platform
-				if (!QualityLevelMinLod)
+				if (!QualityLevelMinValue)
 				{
 					This->PerQuality.Add(*It, PreviousQualityLevel);
 				}
 				else
 				{
-					PreviousQualityLevel = *QualityLevelMinLod;
+					PreviousQualityLevel = *QualityLevelMinValue;
 				}
 			}
 
@@ -266,10 +328,10 @@ void FPerQualityLevelProperty<_StructType, _ValueType, _BasePropertyName>::Strip
 	}
 }
 
-template<typename _StructType, typename _ValueType, EName _BasePropertyName>
-bool FPerQualityLevelProperty<_StructType, _ValueType, _BasePropertyName>::IsQualityLevelValid(int32 QualityLevel) const
+template<typename StructType, typename ValueType, EName _BasePropertyName>
+bool FPerQualityLevelProperty<StructType, ValueType, _BasePropertyName>::IsQualityLevelValid(int32 QualityLevel) const
 {
-	const _StructType* This = StaticCast<const _StructType*>(this);
+	const StructType* This = StaticCast<const StructType*>(this);
 	int32* Value = (int32*)This->PerQuality.Find(QualityLevel);
 
 	if (Value != nullptr)
@@ -283,29 +345,12 @@ bool FPerQualityLevelProperty<_StructType, _ValueType, _BasePropertyName>::IsQua
 }
 #endif
 
-template<typename _StructType, typename _ValueType, EName _BasePropertyName>
-_ValueType FPerQualityLevelProperty<_StructType, _ValueType, _BasePropertyName>::GetLowestValue() const
-{
-	const _StructType* This = StaticCast<const _StructType*>(this);
-	_ValueType Value = This->Default;
-
-	for (const TPair<int32, int32>& Pair : This->PerQuality)
-	{
-		if (Pair.Value < Value)
-		{
-			Value = Pair.Value;
-		}
-	}
-	return Value;
-}
-
 /** Serializer to cook out the most appropriate platform override */
-template<typename _StructType, typename _ValueType, EName _BasePropertyName>
-ENGINE_API FArchive& operator<<(FArchive& Ar, FPerQualityLevelProperty<_StructType, _ValueType, _BasePropertyName>& Property)
+template<typename StructType, typename ValueType, EName _BasePropertyName>
+ENGINE_API FArchive& operator<<(FArchive& Ar, FPerQualityLevelProperty<StructType, ValueType, _BasePropertyName>& Property)
 {
 	bool bCooked = false;
-	_StructType* This = StaticCast<_StructType*>(&Property);
-
+	StructType* This = StaticCast<StructType*>(&Property);
 #if WITH_EDITOR
 	if (Ar.IsCooking())
 	{
@@ -323,14 +368,14 @@ ENGINE_API FArchive& operator<<(FArchive& Ar, FPerQualityLevelProperty<_StructTy
 }
 
 /** Serializer to cook out the most appropriate platform override */
-template<typename _StructType, typename _ValueType, EName _BasePropertyName>
-ENGINE_API void operator<<(FStructuredArchive::FSlot Slot, FPerQualityLevelProperty<_StructType, _ValueType, _BasePropertyName>& Property)
+template<typename StructType, typename ValueType, EName _BasePropertyName>
+ENGINE_API void operator<<(FStructuredArchive::FSlot Slot, FPerQualityLevelProperty<StructType, ValueType, _BasePropertyName>& Property)
 {
 	FArchive& UnderlyingArchive = Slot.GetUnderlyingArchive();
 	FStructuredArchive::FRecord Record = Slot.EnterRecord();
 
 	bool bCooked = false;
-	_StructType* This = StaticCast<_StructType*>(&Property);
+	StructType* This = StaticCast<StructType*>(&Property);
 
 #if WITH_EDITOR
 	if (UnderlyingArchive.IsCooking())
@@ -354,8 +399,12 @@ template int32 FPerQualityLevelProperty<FPerQualityLevelInt, int32, NAME_IntProp
 template FSupportedQualityLevelArray FPerQualityLevelProperty<FPerQualityLevelInt, int32, NAME_IntProperty>::GetSupportedQualityLevels(const TCHAR* InPlatformName) const;
 template void FPerQualityLevelProperty<FPerQualityLevelInt, int32, NAME_IntProperty>::StripQualtiyLevelForCooking(const TCHAR* InPlatformName);
 template bool FPerQualityLevelProperty<FPerQualityLevelInt, int32, NAME_IntProperty>::IsQualityLevelValid(int32 QualityLevel) const;
+template void FPerQualityLevelProperty<FPerQualityLevelInt, int32, NAME_IntProperty>::ConvertQualtiyLevelData(TMap<FName, int32>& PlaformData, TMultiMap<FName, FName>& PerPlatformToQualityLevel, int32 Default);
 #endif
-template int32 FPerQualityLevelProperty<FPerQualityLevelInt, int32, NAME_IntProperty>::GetLowestValue() const;
+template TMap<int32, int32> QualityLevelProperty::ConvertQualtiyLevelData(const TMap<EPerQualityLevels, int32>& Data);
+template TMap<EPerQualityLevels, int32> QualityLevelProperty::ConvertQualtiyLevelData(const TMap<int32, int32>& Data);
+template ENGINE_API void FPerQualityLevelProperty<FPerQualityLevelInt, int32, NAME_IntProperty>::Init(const TCHAR* InCVarName, const TCHAR* InSection);
+template ENGINE_API int32 FPerQualityLevelProperty<FPerQualityLevelInt, int32, NAME_IntProperty>::GetValue(int32 QualityLevel) const;
 
 FString FPerQualityLevelInt::ToString() const
 {
@@ -375,3 +424,35 @@ FString FPerQualityLevelInt::ToString() const
 	return Result;
 }
 
+template ENGINE_API FArchive& operator<<(FArchive&, FPerQualityLevelProperty<FPerQualityLevelFloat, float, NAME_FloatProperty>&);
+template ENGINE_API void operator<<(FStructuredArchive::FSlot Slot, FPerQualityLevelProperty<FPerQualityLevelFloat, float, NAME_FloatProperty>&);
+
+#if WITH_EDITOR
+template int32 FPerQualityLevelProperty<FPerQualityLevelFloat, float, NAME_FloatProperty>::GetValueForPlatform(const ITargetPlatform* TargetPlatform) const;
+template FSupportedQualityLevelArray FPerQualityLevelProperty<FPerQualityLevelFloat, float, NAME_FloatProperty>::GetSupportedQualityLevels(const TCHAR* InPlatformName) const;
+template void FPerQualityLevelProperty<FPerQualityLevelFloat, float, NAME_FloatProperty>::StripQualtiyLevelForCooking(const TCHAR* InPlatformName);
+template bool FPerQualityLevelProperty<FPerQualityLevelFloat, float, NAME_FloatProperty>::IsQualityLevelValid(int32 QualityLevel) const;
+template void FPerQualityLevelProperty<FPerQualityLevelFloat, float, NAME_FloatProperty>::ConvertQualtiyLevelData(TMap<FName, float>& PlaformData, TMultiMap<FName, FName>& PerPlatformToQualityLevel, float Default);
+#endif
+template TMap<int32, float> QualityLevelProperty::ConvertQualtiyLevelData(const TMap<EPerQualityLevels, float>&Data);
+template TMap<EPerQualityLevels, float> QualityLevelProperty::ConvertQualtiyLevelData(const TMap<int32, float>& Data);
+template ENGINE_API void FPerQualityLevelProperty<FPerQualityLevelFloat, float, NAME_FloatProperty>::Init(const TCHAR* InCVarName, const TCHAR* InSection);
+template ENGINE_API float FPerQualityLevelProperty<FPerQualityLevelFloat, float, NAME_FloatProperty>::GetValue(int32 QualityLevel) const;
+
+FString FPerQualityLevelFloat::ToString() const
+{
+	FString Result = FString::FromInt(Default);
+
+#if WITH_EDITORONLY_DATA
+	TArray<int32> QualityLevels;
+	PerQuality.GetKeys(QualityLevels);
+	QualityLevels.Sort();
+
+	for (int32 QL : QualityLevels)
+	{
+		Result = FString::Printf(TEXT("%s, %s=%d"), *Result, *QualityLevelProperty::QualityLevelToFName(QL).ToString(), PerQuality.FindChecked(QL));
+	}
+#endif
+
+	return Result;
+}
\ No newline at end of file
diff --git a/Engine/Source/Runtime/Engine/Public/PerQualityLevelProperties.h b/Engine/Source/Runtime/Engine/Public/PerQualityLevelProperties.h
--- a/Engine/Source/Runtime/Engine/Public/PerQualityLevelProperties.h
+++ b/Engine/Source/Runtime/Engine/Public/PerQualityLevelProperties.h
@@ -50,8 +50,11 @@ namespace QualityLevelProperty
 
 	ENGINE_API FName QualityLevelToFName(int32 QL);
 	ENGINE_API int32 FNameToQualityLevel(FName QL);
-	ENGINE_API TMap<int32, int32> ConvertQualtiyLevelData(const TMap<EPerQualityLevels, int32>& Data);
-	ENGINE_API TMap<EPerQualityLevels, int32> ConvertQualtiyLevelData(const TMap<int32, int32>& Data);
+	template<typename _ValueType>
+	ENGINE_API TMap<int32, _ValueType> ConvertQualtiyLevelData(const TMap<EPerQualityLevels, _ValueType>& Data);
+	template<typename _ValueType>
+	ENGINE_API TMap<EPerQualityLevels, _ValueType> ConvertQualtiyLevelData(const TMap<int32, _ValueType>& Data);
+	
 #if WITH_EDITOR
 	ENGINE_API FSupportedQualityLevelArray PerPlatformOverrideMapping(FString& InPlatformName);
 #endif
@@ -61,14 +64,13 @@ template<typename _StructType, typename _ValueType, EName _BasePropertyName>
 struct ENGINE_API FPerQualityLevelProperty
 {
 	typedef _ValueType ValueType;
+	typedef _StructType StructType;
 
 	FPerQualityLevelProperty() 
 	{
 	}
 	~FPerQualityLevelProperty() {}
 
-	bool operator==(const FPerQualityLevelProperty& Other) const;
-
 	_ValueType GetValueForQualityLevel(int32 QualityLevel) const
 	{
 		const _StructType* This = StaticCast<const _StructType*>(this);
@@ -77,7 +79,7 @@ struct ENGINE_API FPerQualityLevelProperty
 			return This->Default;
 		}
 
-		int32* Value = (int32*)This->PerQuality.Find(QualityLevel);
+		_ValueType* Value = (_ValueType*)This->PerQuality.Find(QualityLevel);
 		
 		if (Value)
 		{
@@ -94,6 +96,7 @@ struct ENGINE_API FPerQualityLevelProperty
 	FSupportedQualityLevelArray GetSupportedQualityLevels(const TCHAR* InPlatformName = nullptr) const;
 	void StripQualtiyLevelForCooking(const TCHAR* InPlatformName = nullptr);
 	bool IsQualityLevelValid(int32 QualityLevel) const;
+	void ConvertQualtiyLevelData(TMap<FName, _ValueType>& PlaformData, TMultiMap<FName, FName>& PerPlatformToQualityLevel, _ValueType Default);
 #endif
 
 	void Init(const TCHAR* InCVarName, const TCHAR* InSection)
@@ -115,7 +118,20 @@ struct ENGINE_API FPerQualityLevelProperty
 		return GetValueForQualityLevel(QualityLevel);
 	}
 
-	_ValueType GetLowestValue() const;
+	_ValueType GetLowestValue() const
+	{
+		const _StructType* This = StaticCast<const _StructType*>(this);
+		_ValueType Value = This->Default;
+
+		for (const TPair<int32, _ValueType>& Pair : This->PerQuality)
+		{
+			if (Pair.Value < Value)
+			{
+				Value = Pair.Value;
+			}
+		}
+		return Value;
+	}
 
 	/* Load old properties that have been converted to FPerQualityLevel */
 	bool SerializeFromMismatchedTag(const FPropertyTag& Tag, FArchive& Ar)
@@ -182,6 +198,7 @@ struct ENGINE_API FPerQualityLevelInt
 	}
 
 	FString ToString() const;
+	int32 MaxType() const { return MAX_int32; }
 };
 
 extern template ENGINE_API FArchive& operator<<(FArchive&, FPerQualityLevelProperty<FPerQualityLevelInt, int32, NAME_IntProperty>&);
@@ -196,4 +213,47 @@ struct TStructOpsTypeTraits<FPerQualityLevelInt>
 		WithSerializeFromMismatchedTag = true,
 		WithSerializer = true
 	};
-};
\ No newline at end of file
+};
+
+USTRUCT()
+struct ENGINE_API FPerQualityLevelFloat
+#if CPP
+	:	public FPerQualityLevelProperty<FPerQualityLevelFloat, float, NAME_FloatProperty>
+#endif
+{
+	GENERATED_USTRUCT_BODY()
+
+	UPROPERTY(EditAnywhere, Category = PerQualityLevel)
+	float Default;
+
+	UPROPERTY(EditAnywhere, Category = PerQualityLevel)
+	TMap<int32, float> PerQuality;
+
+	FPerQualityLevelFloat()
+	{
+		Default = 0.0f;
+	}
+
+	FPerQualityLevelFloat(float InDefaultValue)
+	{
+		Default = InDefaultValue;
+	}
+
+	FString ToString() const;
+	float MaxType() const { return UE_MAX_FLT; }
+};
+
+extern template ENGINE_API FArchive& operator<<(FArchive&, FPerQualityLevelProperty<FPerQualityLevelFloat, float, NAME_FloatProperty>&);
+extern template ENGINE_API void operator<<(FStructuredArchive::FSlot Slot, FPerQualityLevelProperty<FPerQualityLevelFloat, float, NAME_FloatProperty>&);
+
+template<>
+struct TStructOpsTypeTraits<FPerQualityLevelFloat>
+	: public TStructOpsTypeTraitsBase2<FPerQualityLevelFloat>
+{
+	enum
+	{
+		WithSerializeFromMismatchedTag = true,
+		WithSerializer = true
+	};
+};
+
diff --git a/Engine/Source/Runtime/Landscape/Classes/LandscapeGrassType.h b/Engine/Source/Runtime/Landscape/Classes/LandscapeGrassType.h
--- a/Engine/Source/Runtime/Landscape/Classes/LandscapeGrassType.h
+++ b/Engine/Source/Runtime/Landscape/Classes/LandscapeGrassType.h
@@ -7,6 +7,7 @@
 #include "UObject/Object.h"
 #include "Engine/EngineTypes.h"
 #include "PerPlatformProperties.h"
+#include "PerQualityLevelProperties.h"
 #include "LandscapeGrassType.generated.h"
 
 class UStaticMesh;
@@ -28,6 +29,8 @@ struct FGrassVariety
 {
 	GENERATED_USTRUCT_BODY()
 
+	FGrassVariety();
+
 	UPROPERTY(EditAnywhere, Category=Grass)
 	TObjectPtr<UStaticMesh> GrassMesh;
 
@@ -38,6 +41,9 @@ struct FGrassVariety
 	UPROPERTY(EditAnywhere, Category=Grass, meta = (UIMin = 0, ClampMin = 0, UIMax = 1000, ClampMax = 1000))
 	FPerPlatformFloat GrassDensity;
 
+	UPROPERTY(EditAnywhere, Category=Grass, meta = (UIMin = 0, ClampMin = 0, UIMax = 1000, ClampMax = 1000))
+	FPerQualityLevelFloat GrassDensityQuality;
+
 	/* If true, use a jittered grid sequence for placement, otherwise use a halton sequence. */
 	UPROPERTY(EditAnywhere, Category=Grass)
 	bool bUseGrid;
@@ -49,6 +55,9 @@ struct FGrassVariety
 	UPROPERTY(EditAnywhere, Category=Grass, meta = (UIMin = 0, ClampMin = 0, UIMax = 1000000, ClampMax = 1000000))
 	FPerPlatformInt StartCullDistance;
 
+	UPROPERTY(EditAnywhere, Category = Grass, meta = (UIMin = 0, ClampMin = 0, UIMax = 1000000, ClampMax = 1000000))
+	FPerQualityLevelInt StartCullDistanceQuality;
+
 	/**
 	 * The distance where instances will have completely faded out when using a PerInstanceFadeAmount material node. 0 disables. 
 	 * When the entire cluster is beyond this distance, the cluster is completely culled and not rendered at all.
@@ -56,6 +65,8 @@ struct FGrassVariety
 	UPROPERTY(EditAnywhere, Category = Grass, meta = (UIMin = 0, ClampMin = 0, UIMax = 1000000, ClampMax = 1000000))
 	FPerPlatformInt EndCullDistance;
 
+	UPROPERTY(EditAnywhere, Category = Grass, meta = (UIMin = 0, ClampMin = 0, UIMax = 1000000, ClampMax = 1000000))
+	FPerQualityLevelInt EndCullDistanceQuality;
 	/** 
 	 * Specifies the smallest LOD that will be used for this component.
 	 * If -1 (default), the MinLOD of the static mesh asset will be used instead.
@@ -118,28 +129,14 @@ struct FGrassVariety
 	UPROPERTY(EditAnywhere, Category = Grass)
 	uint32 InstanceWorldPositionOffsetDisableDistance;
 
-	FGrassVariety()
-		: GrassMesh(nullptr)
-		, GrassDensity(400)
-		, bUseGrid(true)
-		, PlacementJitter(1.0f)
-		, StartCullDistance(10000)
-		, EndCullDistance(10000)
-		, MinLOD(-1)
-		, Scaling(EGrassScaling::Uniform)
-		, ScaleX(1.0f, 1.0f)
-		, ScaleY(1.0f, 1.0f)
-		, ScaleZ(1.0f, 1.0f)
-		, RandomRotation(true)
-		, AlignToSurface(true)
-		, bUseLandscapeLightmap(false)
-		, bReceivesDecals(true)
-		, bCastDynamicShadow(true)
-		, bCastContactShadow(true)
-		, bKeepInstanceBufferCPUCopy(false)
-		, InstanceWorldPositionOffsetDisableDistance(0)
-	{
-	}
+	bool IsGrassQualityLevelEnable() const;
+
+	int32 GetStartCullDistance() const;
+
+	int32 GetEndCullDistance() const;
+
+	float GetDensity() const;
+
 };
 
 UCLASS(MinimalAPI)
diff --git a/Engine/Source/Runtime/Landscape/Private/LandscapeGrass.cpp b/Engine/Source/Runtime/Landscape/Private/LandscapeGrass.cpp
--- a/Engine/Source/Runtime/Landscape/Private/LandscapeGrass.cpp
+++ b/Engine/Source/Runtime/Landscape/Private/LandscapeGrass.cpp
@@ -201,6 +201,19 @@ static TAutoConsoleVariable<int32> CVarRayTracingLandscapeGrass(
 	TEXT("Include landscapes grass in ray tracing effects (default = 1)"));
 #endif
 
+const TCHAR* GGrassQualityLevelCVarName = TEXT("r.grass.DensityQualityLevel");
+const TCHAR* GGrassQualityLevelScalabilitySection = TEXT("ViewDistanceQuality");
+
+int32 GGrassQualityLevel = -1;
+static FAutoConsoleVariableRef CVarGGrassDensityQualityLevelCVar(
+	GGrassQualityLevelCVarName,
+	GGrassQualityLevel,
+	TEXT("The quality level for grass (low, medium, high, epic). \n"),
+	ECVF_Scalability);
+	
+struct FPerQualityLevelInt;
+struct FPerQualityLevelFloat;
+
 DECLARE_CYCLE_STAT(TEXT("Grass Async Build Time"), STAT_FoliageGrassAsyncBuildTime, STATGROUP_Foliage);
 DECLARE_CYCLE_STAT(TEXT("Grass Start Comp"), STAT_FoliageGrassStartComp, STATGROUP_Foliage);
 DECLARE_CYCLE_STAT(TEXT("Grass End Comp"), STAT_FoliageGrassEndComp, STATGROUP_Foliage);
@@ -1304,6 +1317,80 @@ void UMaterialExpressionLandscapeGrassOutput::ValidateInputName(FGrassInput& InI
 #endif
 
 //
+// FGrassVariety
+//
+
+FGrassVariety::FGrassVariety()
+	: GrassMesh(nullptr)
+	, GrassDensity(400)
+	, GrassDensityQuality(400.0f)
+	, bUseGrid(true)
+	, PlacementJitter(1.0f)
+	, StartCullDistance(10000)
+	, StartCullDistanceQuality(10000)
+	, EndCullDistance(10000)
+	, EndCullDistanceQuality(10000)
+	, MinLOD(-1)
+	, Scaling(EGrassScaling::Uniform)
+	, ScaleX(1.0f, 1.0f)
+	, ScaleY(1.0f, 1.0f)
+	, ScaleZ(1.0f, 1.0f)
+	, RandomRotation(true)
+	, AlignToSurface(true)
+	, bUseLandscapeLightmap(false)
+	, bReceivesDecals(true)
+	, bCastDynamicShadow(true)
+	, bCastContactShadow(true)
+	, bKeepInstanceBufferCPUCopy(false)
+	, InstanceWorldPositionOffsetDisableDistance(0)
+{
+	GrassDensityQuality.Init(GGrassQualityLevelCVarName, GGrassQualityLevelScalabilitySection);
+	StartCullDistanceQuality.Init(GGrassQualityLevelCVarName, GGrassQualityLevelScalabilitySection);
+	EndCullDistanceQuality.Init(GGrassQualityLevelCVarName, GGrassQualityLevelScalabilitySection);
+}
+
+bool FGrassVariety::IsGrassQualityLevelEnable() const
+{
+	return (GEngine && GEngine->UseGrassVarityPerQualityLevels);
+}
+
+int32 FGrassVariety::GetStartCullDistance() const
+{
+	if (IsGrassQualityLevelEnable())
+	{
+		return StartCullDistanceQuality.GetValue(GGrassQualityLevel);
+	}
+	else
+	{
+		return StartCullDistance.GetValue();
+	}
+}
+
+int32 FGrassVariety::GetEndCullDistance() const
+{
+	if (IsGrassQualityLevelEnable())
+	{
+		return EndCullDistanceQuality.GetValue(GGrassQualityLevel);
+	}
+	else
+	{
+		return EndCullDistance.GetValue();
+	}
+}
+
+float FGrassVariety::GetDensity() const
+{
+	if (IsGrassQualityLevelEnable())
+	{
+		return GrassDensityQuality.GetValue(GGrassQualityLevel);
+	}
+	else
+	{
+		return GrassDensity.GetValue();
+	}
+}
+
+
 // ULandscapeGrassType
 //
 
@@ -1669,7 +1756,7 @@ struct FGrassBuilderBase
 	FGrassBuilderBase(ALandscapeProxy* Landscape, ULandscapeComponent* Component, const FGrassVariety& GrassVariety, ERHIFeatureLevel::Type FeatureLevel, int32 SqrtSubsections = 1, int32 SubX = 0, int32 SubY = 0, bool bEnableDensityScaling = true)
 	{
 		const float DensityScale = bEnableDensityScaling ? GGrassDensityScale : 1.0f;
-		GrassDensity = GrassVariety.GrassDensity.GetValue() * DensityScale;
+		GrassDensity = GrassVariety.GetDensity() * DensityScale;
 
 		DrawScale = Landscape->GetRootComponent()->GetRelativeScale3D();
 		DrawLoc = Landscape->GetActorLocation();
@@ -2352,7 +2439,7 @@ void ALandscapeProxy::GetGrassTypes(const UWorld* World, UMaterialInterface* Lan
 			{
 				for (auto& GrassVariety : GrassType->GrassVarieties)
 				{
-					const int32 EndCullDistance = GrassVariety.EndCullDistance.GetValue();
+					int32 EndCullDistance = GrassVariety.GetEndCullDistance();
 					if (EndCullDistance > OutMaxDiscardDistance)
 					{
 						OutMaxDiscardDistance = EndCullDistance;
@@ -2833,8 +2920,8 @@ void ALandscapeProxy::UpdateGrass(const TArray<FVector>& Cameras, int32& InOutNu
 						for (auto& GrassVariety : GrassType->GrassVarieties)
 						{
 							GrassVarietyIndex++;
-							int32 EndCullDistance = GrassVariety.EndCullDistance.GetValue();
-							if (GrassVariety.GrassMesh && GrassVariety.GrassDensity.GetValue() > 0.0f && EndCullDistance > 0)
+							int32 EndCullDistance = GrassVariety.GetEndCullDistance();
+							if (GrassVariety.GrassMesh && GrassVariety.GetDensity() > 0.0f && EndCullDistance > 0)
 							{
 								float MustHaveDistance = GuardBand * (float)EndCullDistance * CullDistanceScale;
 								float DiscardDistance = DiscardGuardBand * (float)EndCullDistance * CullDistanceScale;
@@ -3079,8 +3166,8 @@ void ALandscapeProxy::UpdateGrass(const TArray<FVector>& Cameras, int32& InOutNu
 										}
 										else
 										{
-											GrassInstancedStaticMeshComponent->InstanceStartCullDistance = GrassVariety.StartCullDistance.GetValue() * CullDistanceScale;
-											GrassInstancedStaticMeshComponent->InstanceEndCullDistance = GrassVariety.EndCullDistance.GetValue() * CullDistanceScale;
+											GrassInstancedStaticMeshComponent->InstanceStartCullDistance = GrassVariety.GetStartCullDistance() * CullDistanceScale;
+											GrassInstancedStaticMeshComponent->InstanceEndCullDistance = GrassVariety.GetEndCullDistance() * CullDistanceScale;
 										}
 
 										//@todo - take the settings from a UFoliageType object.  For now, disable distance field lighting on grass so we don't hitch.
