Removed lots of includes in UE based on compile profiling

diff --git a/Engine/Source/Runtime/Engine/Classes/Engine/GameInstance.h b/Engine/Source/Runtime/Engine/Classes/Engine/GameInstance.h
--- a/Engine/Source/Runtime/Engine/Classes/Engine/GameInstance.h
+++ b/Engine/Source/Runtime/Engine/Classes/Engine/GameInstance.h
@@ -16,9 +16,11 @@
 #include "GameFramework/OnlineReplStructs.h"
 #include "ReplayTypes.h"
 
+#if UE_ENABLE_INCLUDE_ORDER_DEPRECATED_IN_5_2
 #if WITH_EDITOR
 #include "Settings/LevelEditorPlaySettings.h"
 #endif 
+#endif
 
 #include "GameInstance.generated.h"
 
@@ -26,8 +28,10 @@ class AGameModeBase;
 class APlayerController;
 class FOnlineSessionSearchResult;
 class FTimerManager;
+class UGameViewportClient;
 class ULocalPlayer;
 class UOnlineSession;
+enum EPlayNetMode : int;
 struct FLatentActionManager;
 class ULevelEditorPlaySettings;
 class IAnalyticsProvider;
@@ -102,16 +106,7 @@ private:
 //@TODO: Some of these are really mutually exclusive and should be refactored (put into a struct to make this easier in the future)
 struct FGameInstancePIEParameters
 {
-	FGameInstancePIEParameters()
-		: bSimulateInEditor(false)
-		, bAnyBlueprintErrors(false)
-		, bStartInSpectatorMode(false)
-		, bRunAsDedicated(false)
-		, bIsPrimaryPIEClient(false)
-		, WorldFeatureLevel(ERHIFeatureLevel::Num)
-		, EditorPlaySettings(nullptr)
-		, NetMode(EPlayNetMode::PIE_Standalone)
-	{}
+	ENGINE_API FGameInstancePIEParameters();
 
 	// Are we doing SIE instead of PIE?
 	bool bSimulateInEditor;
diff --git a/Engine/Source/Runtime/Engine/Classes/Engine/LocalPlayer.h b/Engine/Source/Runtime/Engine/Classes/Engine/LocalPlayer.h
--- a/Engine/Source/Runtime/Engine/Classes/Engine/LocalPlayer.h
+++ b/Engine/Source/Runtime/Engine/Classes/Engine/LocalPlayer.h
@@ -11,7 +11,9 @@
 #include "Templates/SubclassOf.h"
 #include "Engine/EngineTypes.h"
 #include "Input/Reply.h"
+#if UE_ENABLE_INCLUDE_ORDER_DEPRECATED_IN_5_2
 #include "Engine/GameViewportClient.h"
+#endif
 #include "Online/CoreOnline.h"
 #include "SceneTypes.h"
 #include "Engine/Player.h"
@@ -28,7 +30,9 @@
 class AActor;
 class FSceneView;
 class FSlateUser;
+class FViewport;
 class UGameInstance;
+class UGameViewportClient;
 class ULocalPlayer;
 struct FMinimalViewInfo;
 struct FSceneViewProjectionData;
diff --git a/Engine/Source/Runtime/Engine/Classes/Engine/ScriptViewportClient.h b/Engine/Source/Runtime/Engine/Classes/Engine/ScriptViewportClient.h
--- a/Engine/Source/Runtime/Engine/Classes/Engine/ScriptViewportClient.h
+++ b/Engine/Source/Runtime/Engine/Classes/Engine/ScriptViewportClient.h
@@ -10,7 +10,10 @@
 #include "CoreMinimal.h"
 #include "UObject/ObjectMacros.h"
 #include "UObject/Object.h"
+#if UE_ENABLE_INCLUDE_ORDER_DEPRECATED_IN_5_2
 #include "UnrealClient.h"
+#endif
+#include "ViewportClient.h"
 #include "ScriptViewportClient.generated.h"
 
 UCLASS(transient)
diff --git a/Engine/Source/Runtime/Engine/Classes/Engine/SkeletalMesh.h b/Engine/Source/Runtime/Engine/Classes/Engine/SkeletalMesh.h
--- a/Engine/Source/Runtime/Engine/Classes/Engine/SkeletalMesh.h
+++ b/Engine/Source/Runtime/Engine/Classes/Engine/SkeletalMesh.h
@@ -13,24 +13,28 @@
 #include "Interfaces/Interface_AssetUserData.h"
 #include "RenderCommandFence.h"
 #include "EngineDefines.h"
-#include "Components.h"
 #include "ReferenceSkeleton.h"
 #include "GPUSkinPublicDefs.h"
-#include "Rendering/SkeletalMeshRenderData.h"
 #include "Animation/PreviewAssetAttachComponent.h"
-#include "BoneContainer.h"
 #include "Interfaces/Interface_CollisionDataProvider.h"
 #include "EngineTypes.h"
 #include "SkeletalMeshSampling.h"
 #include "PerPlatformProperties.h"
 #include "Misc/EnumClassFlags.h"
-#include "SkeletalMeshLODSettings.h"
 #include "Animation/NodeMappingProviderInterface.h"
-#include "Animation/SkinWeightProfile.h"
 #include "Animation/MorphTarget.h"
 #include "Engine/StreamableRenderAsset.h"
 #include "PerQualityLevelProperties.h"
 #include "SkinnedAsset.h"
+
+#if UE_ENABLE_INCLUDE_ORDER_DEPRECATED_IN_5_2
+#include "Components.h"
+#include "Rendering/SkeletalMeshRenderData.h"
+#include "BoneContainer.h"
+#include "SkeletalMeshLODSettings.h"
+#include "Animation/SkinWeightProfile.h"
+#endif
+
 #include "SkeletalMesh.generated.h"
 
 class UAnimInstance;
@@ -50,7 +54,14 @@ class FSkeletalMeshModel;
 class FSkeletalMeshRenderData;
 class USkeletalMeshSocket;
 class FSkeletalMeshUpdate;
+class FSkinnedAssetAsyncBuildScope;
+class FSkinnedAssetAsyncBuildWorker;
+class FSkinnedAssetCompilationContext;
+class FSkinnedAssetBuildContext;
+class FSkinnedAssetPostLoadContext;
+class USkeletalMeshLODSettings;
 class USkeleton;
+struct FMeshUVChannelInfo;
 struct FSkinWeightProfileInfo;
 class FSkinWeightVertexBuffer;
 enum class ESkeletalMeshGeoImportVersions : uint8;
@@ -426,21 +437,9 @@ private:
 	UE_DEPRECATED(5.0, "This must be protected for async build, always use the accessors even internally.")
 	TUniquePtr<FSkeletalMeshRenderData> SkeletalMeshRenderData;
 
-	FSkeletalMeshRenderData* GetSkeletalMeshRenderData() const
-	{
-		WaitUntilAsyncPropertyReleased(ESkeletalMeshAsyncProperties::SkeletalMeshRenderData);
-		PRAGMA_DISABLE_DEPRECATION_WARNINGS
-		return SkeletalMeshRenderData.Get();
-		PRAGMA_ENABLE_DEPRECATION_WARNINGS
-	}
+	FSkeletalMeshRenderData* GetSkeletalMeshRenderData() const;
 
-	void SetSkeletalMeshRenderData(TUniquePtr<FSkeletalMeshRenderData>&& InSkeletalMeshRenderData)
-	{
-		WaitUntilAsyncPropertyReleased(ESkeletalMeshAsyncProperties::SkeletalMeshRenderData);
-		PRAGMA_DISABLE_DEPRECATION_WARNINGS
-		SkeletalMeshRenderData = MoveTemp(InSkeletalMeshRenderData);
-		PRAGMA_ENABLE_DEPRECATION_WARNINGS
-	}
+	void SetSkeletalMeshRenderData(TUniquePtr<FSkeletalMeshRenderData>&& InSkeletalMeshRenderData);
 
 #if WITH_EDITORONLY_DATA
 public:
@@ -569,13 +568,7 @@ public:
 #endif
 
 	/** Get the data to use for rendering. USkinnedAsset interface. */
-	virtual FSkeletalMeshRenderData* GetResourceForRendering() const override
-	{
-		WaitUntilAsyncPropertyReleased(ESkeletalMeshAsyncProperties::SkeletalMeshRenderData);
-		PRAGMA_DISABLE_DEPRECATION_WARNINGS
-		return SkeletalMeshRenderData.Get();
-		PRAGMA_ENABLE_DEPRECATION_WARNINGS
-	}
+	virtual FSkeletalMeshRenderData* GetResourceForRendering() const override;
 
 	/** Skeleton of this skeletal mesh **/
 	UE_DEPRECATED(4.27, "Please do not access this member directly; use USkeletalMesh::GetSkeleton() or USkeletalMesh::SetSkeleton().")
@@ -775,13 +768,7 @@ public:
 	}
 
 	UFUNCTION(BlueprintSetter)
-	void SetMaterials(const TArray<FSkeletalMaterial>& InMaterials)
-	{
-		WaitUntilAsyncPropertyReleased(ESkeletalMeshAsyncProperties::Materials);
-		PRAGMA_DISABLE_DEPRECATION_WARNINGS
-		Materials = InMaterials;
-		PRAGMA_ENABLE_DEPRECATION_WARNINGS
-	}
+	void SetMaterials(const TArray<FSkeletalMaterial>& InMaterials);
 
 	/** List of bones that should be mirrored. */
 	UE_DEPRECATED(4.27, "Please do not access this member directly; Use UMirrorDataTable for mirroring support")
@@ -2805,13 +2792,7 @@ public:
 	 *
 	 * @param NewLODInfo : new LOD info to be added
 	 */
-	void AddLODInfo(const FSkeletalMeshLODInfo& NewLODInfo) 
-	{
-		WaitUntilAsyncPropertyReleased(ESkeletalMeshAsyncProperties::LODInfo);
-		PRAGMA_DISABLE_DEPRECATION_WARNINGS
-		LODInfo.Add(NewLODInfo);
-		PRAGMA_ENABLE_DEPRECATION_WARNINGS
-	}
+	void AddLODInfo(const FSkeletalMeshLODInfo& NewLODInfo);
 	
 	/* 
 	 * Remove LOD info of given index
@@ -2855,24 +2836,12 @@ public:
 	/* 
 	 * Get LODInfo of the given index non-const. USkinnedAsset interface.
 	 */
-	virtual FSkeletalMeshLODInfo* GetLODInfo(int32 Index) override
-	{
-		WaitUntilAsyncPropertyReleased(ESkeletalMeshAsyncProperties::LODInfo);
-		PRAGMA_DISABLE_DEPRECATION_WARNINGS
-		return LODInfo.IsValidIndex(Index) ? &LODInfo[Index] : nullptr;  
-		PRAGMA_ENABLE_DEPRECATION_WARNINGS
-	}
+	virtual FSkeletalMeshLODInfo* GetLODInfo(int32 Index) override;
 	
 	/* 
 	 * Get LODInfo of the given index const. USkinnedAsset interface.
 	 */	
-	virtual const FSkeletalMeshLODInfo* GetLODInfo(int32 Index) const override
-	{
-		WaitUntilAsyncPropertyReleased(ESkeletalMeshAsyncProperties::LODInfo, ESkinnedAssetAsyncPropertyLockType::ReadOnly);
-		PRAGMA_DISABLE_DEPRECATION_WARNINGS
-		return LODInfo.IsValidIndex(Index) ? &LODInfo[Index] : nullptr;
-		PRAGMA_ENABLE_DEPRECATION_WARNINGS
-	}
+	virtual const FSkeletalMeshLODInfo* GetLODInfo(int32 Index) const override;
 
 	/**
 	 *	Get BakePose for the given LOD
@@ -2887,23 +2856,12 @@ public:
 	/* 
 	 * Return true if given index's LOD is valid. USkinnedAsset interface.
 	 */
-	virtual bool IsValidLODIndex(int32 Index) const override
-	{
-		WaitUntilAsyncPropertyReleased(ESkeletalMeshAsyncProperties::LODInfo, ESkinnedAssetAsyncPropertyLockType::ReadOnly);
-		PRAGMA_DISABLE_DEPRECATION_WARNINGS
-		return LODInfo.IsValidIndex(Index);
-		PRAGMA_ENABLE_DEPRECATION_WARNINGS
-	}
+	virtual bool IsValidLODIndex(int32 Index) const override;
+
 	/* 
 	 * Returns total number of LOD. USkinnedAsset interface.
 	 */
-	virtual int32 GetLODNum() const override
-	{
-		WaitUntilAsyncPropertyReleased(ESkeletalMeshAsyncProperties::LODInfo, ESkinnedAssetAsyncPropertyLockType::ReadOnly);
-		PRAGMA_DISABLE_DEPRECATION_WARNINGS
-		return LODInfo.Num();
-		PRAGMA_ENABLE_DEPRECATION_WARNINGS
-	}
+	virtual int32 GetLODNum() const override;
 
 	/** USkinnedAsset interface. */
 	virtual bool IsMaterialUsed(int32 MaterialIndex) const override;
@@ -2932,13 +2890,8 @@ public:
 		return SkinWeightProfiles; 
 		PRAGMA_ENABLE_DEPRECATION_WARNINGS
 	}
-	void AddSkinWeightProfile(const FSkinWeightProfileInfo& Profile) 
-	{
-		WaitUntilAsyncPropertyReleased(ESkeletalMeshAsyncProperties::SkinWeightProfiles); 
-		PRAGMA_DISABLE_DEPRECATION_WARNINGS
-		SkinWeightProfiles.Add(Profile); 
-		PRAGMA_ENABLE_DEPRECATION_WARNINGS
-	}
+	void AddSkinWeightProfile(const FSkinWeightProfileInfo& Profile);
+
 	int32 GetNumSkinWeightProfiles() const 
 	{ 
 		WaitUntilAsyncPropertyReleased(ESkeletalMeshAsyncProperties::SkinWeightProfiles, ESkinnedAssetAsyncPropertyLockType::ReadOnly);
diff --git a/Engine/Source/Runtime/Engine/Classes/Engine/SkinnedAsset.h b/Engine/Source/Runtime/Engine/Classes/Engine/SkinnedAsset.h
--- a/Engine/Source/Runtime/Engine/Classes/Engine/SkinnedAsset.h
+++ b/Engine/Source/Runtime/Engine/Classes/Engine/SkinnedAsset.h
@@ -12,18 +12,25 @@
 #include "Interfaces/Interface_AsyncCompilation.h"
 #include "ReferenceSkeleton.h"
 #include "PerPlatformProperties.h"
+#if UE_ENABLE_INCLUDE_ORDER_DEPRECATED_IN_5_2
 #include "SkeletalMeshTypes.h"
 #include "SkinnedAssetAsyncCompileUtils.h"
 #include "SkinnedAssetCommon.h"
+#endif
 #include "SkinnedAsset.generated.h"
 
 struct FSkeletalMaterial;
 struct FSkeletalMeshLODInfo;
+class FSkinnedAssetBuildContext;
+class FSkinnedAssetPostLoadContext;
+class FSkinnedAsyncTaskContext;
+class FVertexFactoryType;
 class ITargetPlatform;
 class UMeshDeformer;
 class UMorphTarget;
 class UPhysicsAsset;
 class USkeleton;
+struct FSkinnedAssetAsyncBuildTask;
 
 enum class ESkinnedAssetAsyncPropertyLockType
 {
@@ -46,11 +53,8 @@ class ENGINE_API USkinnedAsset : public UStreamableRenderAsset, public IInterfac
 	GENERATED_BODY()
 
 public:
-	USkinnedAsset(const FObjectInitializer& ObjectInitializer)
-		: Super(ObjectInitializer)
-	{}
-
-	virtual ~USkinnedAsset() {}
+	USkinnedAsset(const FObjectInitializer& ObjectInitializer);
+	virtual ~USkinnedAsset();
 
 	/** Return the reference skeleton. */
 	virtual struct FReferenceSkeleton& GetRefSkeleton()
@@ -65,12 +69,10 @@ public:
 	PURE_VIRTUAL(USkinnedAsset::GetLODInfo, return nullptr;);
 
 	/** Return if the material index is valid. */
-	virtual bool IsValidMaterialIndex(int32 Index) const
-	{ return GetMaterials().IsValidIndex(Index); }
+	virtual bool IsValidMaterialIndex(int32 Index) const;
 
 	/** Return the number of materials of this mesh. */
-	virtual int32 GetNumMaterials() const
-	{ return GetMaterials().Num(); }
+	virtual int32 GetNumMaterials() const;
 
 	/** Return the physics asset whose shapes will be used for shadowing. */
 	virtual UPhysicsAsset* GetShadowPhysicsAsset() const
@@ -110,9 +112,9 @@ public:
 
 	/** Return the whole array of LOD info. */
 	virtual TArray<FSkeletalMeshLODInfo>& GetLODInfoArray()
-	PURE_VIRTUAL(USkinnedAsset::GetLODInfoArray, static TArray<FSkeletalMeshLODInfo> Dummy; return Dummy;);
+	PURE_VIRTUAL(USkinnedAsset::GetLODInfoArray, return GetMeshLodInfoDummyArray(););
 	virtual const TArray<FSkeletalMeshLODInfo>& GetLODInfoArray() const
-	PURE_VIRTUAL(USkinnedAsset::GetLODInfoArray, static const TArray<FSkeletalMeshLODInfo> Dummy; return Dummy;);
+	PURE_VIRTUAL(USkinnedAsset::GetLODInfoArray, return GetMeshLodInfoDummyArray(););
 
 	/** Get the data to use for rendering. */
 	virtual class FSkeletalMeshRenderData* GetResourceForRendering() const
@@ -131,9 +133,9 @@ public:
 	PURE_VIRTUAL(USkinnedAsset::GetPhysicsAsset, return nullptr;);
 
 	virtual TArray<FSkeletalMaterial>& GetMaterials()
-	PURE_VIRTUAL(USkinnedAsset::GetMaterials, static TArray<FSkeletalMaterial> Dummy; return Dummy;);
+	PURE_VIRTUAL(USkinnedAsset::GetMaterials, return GetSkeletalMaterialDummyArray(););
 	virtual const TArray<FSkeletalMaterial>& GetMaterials() const
-	PURE_VIRTUAL(USkinnedAsset::GetMaterials, static const TArray<FSkeletalMaterial> Dummy; return Dummy;);
+	PURE_VIRTUAL(USkinnedAsset::GetMaterials, return GetSkeletalMaterialDummyArray(););
 
 	virtual int32 GetLODNum() const
 	PURE_VIRTUAL(USkinnedAsset::GetLODNum, return 0;);
@@ -213,8 +215,7 @@ public:
 	virtual void SetSkinWeightProfilesData(int32 LODIndex, struct FSkinWeightProfilesData& SkinWeightProfilesData) {}
 
 	/** Computes flags for building vertex buffers. */
-	virtual ESkeletalMeshVertexFlags GetVertexBufferFlags() const
-	{ return GetHasVertexColors() ? ESkeletalMeshVertexFlags::HasVertexColors : ESkeletalMeshVertexFlags::None; }
+	virtual ESkeletalMeshVertexFlags GetVertexBufferFlags() const;
 
 	/**
 	 * Take the BoneSpaceTransforms array (translation vector, rotation quaternion and scale vector) and update the array of component-space bone transformation matrices (ComponentSpaceTransforms).
@@ -329,5 +330,8 @@ private:
 
 	friend class FSkinnedAssetCompilingManager;
 	friend class FSkinnedAssetAsyncBuildWorker;
+
+	static TArray<FSkeletalMeshLODInfo>& GetMeshLodInfoDummyArray();
+	static TArray<FSkeletalMaterial>& GetSkeletalMaterialDummyArray();
 };
 
diff --git a/Engine/Source/Runtime/Engine/Classes/Engine/StaticMesh.h b/Engine/Source/Runtime/Engine/Classes/Engine/StaticMesh.h
--- a/Engine/Source/Runtime/Engine/Classes/Engine/StaticMesh.h
+++ b/Engine/Source/Runtime/Engine/Classes/Engine/StaticMesh.h
@@ -14,20 +14,26 @@
 #include "UObject/ScriptMacros.h"
 #include "Interfaces/Interface_AssetUserData.h"
 #include "RenderCommandFence.h"
-#include "Components.h"
 #include "Interfaces/Interface_CollisionDataProvider.h"
 #include "Interfaces/Interface_AsyncCompilation.h"
+#include "MeshUVChannelInfo.h"
 #include "Engine/StreamableRenderAsset.h"
 #include "Templates/UniquePtr.h"
 #include "StaticMeshSourceData.h"
-#include "StaticMeshResources.h"
 #include "PerPlatformProperties.h"
-#include "RenderAssetUpdate.h"
 #include "MeshTypes.h"
 #include "PerQualityLevelProperties.h"
+
+#if UE_ENABLE_INCLUDE_ORDER_DEPRECATED_IN_5_2
+#include "Components.h"
+#include "StaticMeshResources.h"
+#include "RenderAssetUpdate.h"
+#endif
+
 #include "StaticMesh.generated.h"
 
 class FSpeedTreeWind;
+class FStaticMeshLODGroup;
 class UAssetUserData;
 class UMaterialInterface;
 class UNavCollisionBase;
@@ -477,41 +483,14 @@ struct FStaticMaterial
 {
 	GENERATED_USTRUCT_BODY()
 
-		FStaticMaterial()
-		: MaterialInterface(NULL)
-		, MaterialSlotName(NAME_None)
-#if WITH_EDITORONLY_DATA
-		, ImportedMaterialSlotName(NAME_None)
-#endif //WITH_EDITORONLY_DATA
-	{
+	ENGINE_API FStaticMaterial();
 
-	}
-
-	FStaticMaterial(class UMaterialInterface* InMaterialInterface
+	ENGINE_API FStaticMaterial(class UMaterialInterface* InMaterialInterface
 		, FName InMaterialSlotName = NAME_None
 #if WITH_EDITORONLY_DATA
-		, FName InImportedMaterialSlotName = NAME_None)
-#else
-		)
+		, FName InImportedMaterialSlotName = NAME_None
 #endif
-		: MaterialInterface(InMaterialInterface)
-		, MaterialSlotName(InMaterialSlotName)
-#if WITH_EDITORONLY_DATA
-		, ImportedMaterialSlotName(InImportedMaterialSlotName)
-#endif //WITH_EDITORONLY_DATA
-	{
-		//If not specified add some valid material slot name
-		if (MaterialInterface && MaterialSlotName == NAME_None)
-		{
-			MaterialSlotName = MaterialInterface->GetFName();
-		}
-#if WITH_EDITORONLY_DATA
-		if (ImportedMaterialSlotName == NAME_None)
-		{
-			ImportedMaterialSlotName = MaterialSlotName;
-		}
-#endif
-	}
+		);
 
 	friend FArchive& operator<<(FArchive& Ar, FStaticMaterial& Elem);
 
diff --git a/Engine/Source/Runtime/Engine/Classes/Engine/StreamableRenderAsset.h b/Engine/Source/Runtime/Engine/Classes/Engine/StreamableRenderAsset.h
--- a/Engine/Source/Runtime/Engine/Classes/Engine/StreamableRenderAsset.h
+++ b/Engine/Source/Runtime/Engine/Classes/Engine/StreamableRenderAsset.h
@@ -9,7 +9,9 @@
 #include "Engine/TextureStreamingTypes.h"
 #include "Serialization/BulkData.h"
 #include "Templates/RefCounting.h"
+#if UE_ENABLE_INCLUDE_ORDER_DEPRECATED_IN_5_2
 #include "RenderAssetUpdate.h"
+#endif
 #include "Streaming/StreamableRenderResourceState.h"
 #include "PerQualityLevelProperties.h"
 #include "StreamableRenderAsset.generated.h"
@@ -41,7 +43,7 @@ class UStreamableRenderAsset : public UObject
 public:
 
 	/** Destructor */
-	ENGINE_API virtual ~UStreamableRenderAsset() {};
+	ENGINE_API virtual ~UStreamableRenderAsset();
 
 	/** Get an integer representation of the LOD group */
 	virtual int32 GetLODGroupForStreaming() const
diff --git a/Engine/Source/Runtime/Engine/Classes/Engine/Texture.h b/Engine/Source/Runtime/Engine/Classes/Engine/Texture.h
--- a/Engine/Source/Runtime/Engine/Classes/Engine/Texture.h
+++ b/Engine/Source/Runtime/Engine/Classes/Engine/Texture.h
@@ -11,11 +11,14 @@
 #include "Interfaces/Interface_AssetUserData.h"
 #include "Interfaces/Interface_AsyncCompilation.h"
 #include "RenderCommandFence.h"
-#include "RenderResource.h"
 #include "Serialization/EditorBulkData.h"
 #include "Engine/TextureDefines.h"
+#include "MaterialValueType.h"
+#if UE_ENABLE_INCLUDE_ORDER_DEPRECATED_IN_5_2
 #include "MaterialShared.h"
 #include "TextureResource.h"
+#include "RenderResource.h"
+#endif
 #include "Engine/StreamableRenderAsset.h"
 #include "PerPlatformProperties.h"
 #include "Misc/FieldAccessor.h"
@@ -41,7 +44,7 @@
 
 namespace FOodleDataCompression {enum class ECompressor : uint8; enum class ECompressionLevel : int8; }
 
-
+class FTextureReference;
 class ITargetPlatform;
 class UAssetUserData;
 struct FPropertyChangedEvent;
@@ -1408,9 +1411,8 @@ public:
 	TFieldPtrAccessor<FTextureResource> Resource;
 #endif
 
-	PRAGMA_DISABLE_DEPRECATION_WARNINGS
-	ENGINE_API virtual ~UTexture() {};
-	PRAGMA_ENABLE_DEPRECATION_WARNINGS
+	ENGINE_API UTexture(FVTableHelper& Helper);
+	ENGINE_API virtual ~UTexture();
 	
 	/** Set texture's resource, can be NULL */
 	ENGINE_API void SetResource(FTextureResource* Resource);
@@ -1422,7 +1424,7 @@ public:
 	ENGINE_API const FTextureResource* GetResource() const;
 
 	/** Stable RHI texture reference that refers to the current RHI texture. Note this is manually refcounted! */
-	FTextureReference TextureReference;
+	FTextureReference& TextureReference;
 
 	/** Release fence to know when resources have been freed on the rendering thread. */
 	FRenderCommandFence ReleaseFence;
@@ -1863,34 +1865,3 @@ protected:
 	 */
 	FStreamableRenderResourceState GetResourcePostInitState(const FTexturePlatformData* PlatformData, bool bAllowStreaming, int32 MinRequestMipCount = 0, int32 MaxMipCount = 0, bool bSkipCanBeLoaded = false) const;
 };
-
-/** 
-* Replaces the RHI reference of one texture with another.
-* Allows one texture to be replaced with another at runtime and have all existing references to it remain valid.
-*/
-struct FTextureReferenceReplacer
-{
-	FTextureReferenceRHIRef OriginalRef;
-
-	FTextureReferenceReplacer(UTexture* OriginalTexture)
-	{
-		if (OriginalTexture)
-		{
-			OriginalTexture->ReleaseResource();
-			OriginalRef = OriginalTexture->TextureReference.TextureReferenceRHI;
-		}
-		else
-		{
-			OriginalRef = nullptr;
-		}
-	}
-
-	void Replace(UTexture* NewTexture)
-	{
-		if (OriginalRef)
-		{
-			NewTexture->TextureReference.TextureReferenceRHI = OriginalRef;
-		}
-	}
-};
- 
\ No newline at end of file
diff --git a/Engine/Source/Runtime/Engine/Classes/Engine/Texture2D.h b/Engine/Source/Runtime/Engine/Classes/Engine/Texture2D.h
--- a/Engine/Source/Runtime/Engine/Classes/Engine/Texture2D.h
+++ b/Engine/Source/Runtime/Engine/Classes/Engine/Texture2D.h
@@ -7,13 +7,16 @@
 #include "UObject/ObjectMacros.h"
 #include "UObject/ScriptMacros.h"
 #include "Engine/Texture.h"
+#if UE_ENABLE_INCLUDE_ORDER_DEPRECATED_IN_5_2
 #include "TextureResource.h"
+#endif
 #include "Misc/FieldAccessor.h"
 #include "Serialization/BulkData.h"
 #include "Texture2D.generated.h"
 
 class FTexture2DResourceMem;
 class FTexture2DResource;
+struct FUpdateTextureRegion2D;
 
 UCLASS(hidecategories=Object, MinimalAPI, BlueprintType)
 class UTexture2D : public UTexture
diff --git a/Engine/Source/Runtime/Engine/Classes/Engine/World.h b/Engine/Source/Runtime/Engine/Classes/Engine/World.h
--- a/Engine/Source/Runtime/Engine/Classes/Engine/World.h
+++ b/Engine/Source/Runtime/Engine/Classes/Engine/World.h
@@ -13,6 +13,7 @@
 #include "Delegates/IDelegateInstance.h"
 #include "Engine/EngineTypes.h"
 #include "Engine/EngineBaseTypes.h"
+#include "GameTime.h"
 #include "CollisionQueryParams.h"
 #include "WorldCollision.h"
 #include "GameFramework/Pawn.h"
@@ -25,7 +26,10 @@
 #include "Physics/PhysicsInterfaceDeclares.h"
 #include "Particles/WorldPSCPool.h"
 #include "Containers/SortedMap.h"
+#include "AudioDeviceHandle.h"
+#if UE_ENABLE_INCLUDE_ORDER_DEPRECATED_IN_5_2
 #include "AudioDeviceManager.h"
+#endif
 #include "Subsystems/WorldSubsystem.h"
 #include "Subsystems/SubsystemCollection.h"
 #include "CollisionProfile.h"
@@ -141,86 +145,6 @@ private:
 	friend UWorld;
 };
 
-/** Contains all the timings of a gaming frame, to handle pause and time dilation (for instance bullet time) of the world. */
-struct ENGINE_API FGameTime
-{
-	FORCEINLINE_DEBUGGABLE FGameTime()
-		: RealTimeSeconds(0.0)
-		, WorldTimeSeconds(0.0)
-		, DeltaRealTimeSeconds(0.0f)
-		, DeltaWorldTimeSeconds(0.0f)
-	{ }
-
-	FGameTime(const FGameTime&) = default;
-	FGameTime& operator = (const FGameTime&) = default;
-
-	// Returns the game time since GStartTime.
-	static FGameTime GetTimeSinceAppStart();
-
-	static FORCEINLINE_DEBUGGABLE FGameTime CreateUndilated(double InRealTimeSeconds, float InDeltaRealTimeSeconds)
-	{
-		return FGameTime::CreateDilated(InRealTimeSeconds, InDeltaRealTimeSeconds, InRealTimeSeconds, InDeltaRealTimeSeconds);
-	}
-
-	static FORCEINLINE_DEBUGGABLE FGameTime CreateDilated(double InRealTimeSeconds, float InDeltaRealTimeSeconds, double InWorldTimeSeconds, float InDeltaWorldTimeSeconds)
-	{
-		return FGameTime(InRealTimeSeconds, InDeltaRealTimeSeconds, InWorldTimeSeconds, InDeltaWorldTimeSeconds);
-	}
-
-	/** Returns time in seconds since level began play, but IS NOT paused when the game is paused, and IS NOT dilated/clamped. */
-	FORCEINLINE_DEBUGGABLE double GetRealTimeSeconds() const
-	{
-		return RealTimeSeconds;
-	}
-
-	/** Returns frame delta time in seconds with no adjustment for time dilation and pause. */
-	FORCEINLINE_DEBUGGABLE float GetDeltaRealTimeSeconds() const
-	{
-		return DeltaRealTimeSeconds;
-	}
-
-	/** Returns time in seconds since level began play, but IS paused when the game is paused, and IS dilated/clamped. */
-	FORCEINLINE_DEBUGGABLE double GetWorldTimeSeconds() const
-	{
-		return WorldTimeSeconds;
-	}
-
-	/** Returns frame delta time in seconds adjusted by e.g. time dilation. */
-	FORCEINLINE_DEBUGGABLE float GetDeltaWorldTimeSeconds() const
-	{
-		return DeltaWorldTimeSeconds;
-	}
-
-	/** Returns how much world time is slowed compared to real time. */
-	FORCEINLINE_DEBUGGABLE float GetTimeDilation() const
-	{
-		ensure(DeltaRealTimeSeconds > 0.0f);
-		return DeltaWorldTimeSeconds / DeltaRealTimeSeconds;
-	}
-
-	/** Returns whether the world time is paused. */
-	FORCEINLINE_DEBUGGABLE bool IsPaused() const
-	{
-		return DeltaWorldTimeSeconds == 0.0f;
-	}
-
-private:
-	double RealTimeSeconds;
-	double WorldTimeSeconds;
-
-	float DeltaRealTimeSeconds;
-	float DeltaWorldTimeSeconds;
-
-	FORCEINLINE_DEBUGGABLE FGameTime(double InRealTimeSeconds, float InDeltaRealTimeSeconds, double InWorldTimeSeconds, float InDeltaWorldTimeSeconds)
-		: RealTimeSeconds(InRealTimeSeconds)
-		, WorldTimeSeconds(InWorldTimeSeconds)
-		, DeltaRealTimeSeconds(InDeltaRealTimeSeconds)
-		, DeltaWorldTimeSeconds(InDeltaWorldTimeSeconds)
-	{ }
-
-};
-
-
 ENGINE_API DECLARE_LOG_CATEGORY_EXTERN(LogSpawn, Warning, All);
 
 DECLARE_MULTICAST_DELEGATE_OneParam(FOnActorSpawned, AActor*);
@@ -684,6 +608,13 @@ struct ENGINE_API FActorSpawnUtils
 	static FName GetBaseName(FName Name);
 };
 
+struct FActorsInitializedParams
+{
+	FActorsInitializedParams(UWorld* InWorld, bool InResetTime) : World(InWorld), ResetTime(InResetTime) {}
+	UWorld* World;
+	bool ResetTime;
+};
+
 /**
  *  This encapsulate World's async trace functionality. This contains two buffers of trace data buffer and alternates it for each tick. 
  *  
@@ -2571,12 +2502,8 @@ public:
 	/** Returns true if the actors have been initialized and are ready to start play */
 	bool AreActorsInitialized() const;
 
-	struct FActorsInitializedParams
-	{
-		FActorsInitializedParams(UWorld* InWorld, bool InResetTime) : World(InWorld), ResetTime(InResetTime) {}
-		UWorld* World;
-		bool ResetTime;
-	};
+	/** For backwards compatibility */
+	using FActorsInitializedParams = ::FActorsInitializedParams;
 
 	DECLARE_MULTICAST_DELEGATE_OneParam(FOnWorldInitializedActors, const FActorsInitializedParams&);
 	FOnWorldInitializedActors OnActorsInitialized;
diff --git a/Engine/Source/Runtime/Engine/Classes/GameFramework/PlayerController.h b/Engine/Source/Runtime/Engine/Classes/GameFramework/PlayerController.h
--- a/Engine/Source/Runtime/Engine/Classes/GameFramework/PlayerController.h
+++ b/Engine/Source/Runtime/Engine/Classes/GameFramework/PlayerController.h
@@ -6,13 +6,13 @@
 #include "SlateFwd.h"
 #include "UObject/ObjectMacros.h"
 #include "Containers/SortedMap.h"
+#include "Containers/StaticArray.h"
 #include "Misc/Guid.h"
 #include "InputCoreTypes.h"
 #include "Templates/SubclassOf.h"
 #include "Engine/NetSerialization.h"
 #include "Engine/EngineTypes.h"
 #include "Engine/EngineBaseTypes.h"
-#include "Widgets/SWidget.h"
 #include "Engine/LatentActionManager.h"
 #include "SceneTypes.h"
 #include "GameFramework/Controller.h"
@@ -20,13 +20,19 @@
 #include "GameFramework/OnlineReplStructs.h"
 #include "GameFramework/PlayerMuteList.h"
 #include "Camera/PlayerCameraManager.h"
-#include "Components/InputComponent.h"
 #include "GameFramework/ForceFeedbackEffect.h"
 #include "GameFramework/UpdateLevelVisibilityLevelInfo.h"
+#include "GenericPlatform/ICursor.h"
 #include "GenericPlatform/IInputInterface.h"
-#include "GameFramework/PlayerInput.h"
 #include "Physics/AsyncPhysicsData.h"
 #include "WorldPartition/WorldPartitionStreamingSource.h"
+
+#if UE_ENABLE_INCLUDE_ORDER_DEPRECATED_IN_5_2
+#include "Widgets/SWidget.h"
+#include "Components/InputComponent.h"
+#include "GameFramework/PlayerInput.h"
+#endif
+
 #include "PlayerController.generated.h"
 
 class ACameraActor;
@@ -35,15 +41,21 @@ class APawn;
 class APlayerCameraManager;
 class ASpectatorPawn;
 class FDebugDisplayInfo;
+class SWidget;
 class UActorChannel;
 class UCheatManager;
 class UGameViewportClient;
+class UInputComponent;
 class ULocalMessage;
 class UNetConnection;
 class UPlayer;
+class UPlayerInput;
 class UPrimitiveComponent;
+namespace EControllerAnalogStick { enum Type : int; }
 struct FActiveHapticFeedbackEffect;
 struct FCollisionQueryParams;
+struct FInputKeyParams;
+struct FPlatformUserId;
 class UAsyncPhysicsInputComponent;
 
 /** Default delegate that provides an implementation for those that don't have special needs other than a toggle */
diff --git a/Engine/Source/Runtime/Engine/Classes/Materials/MaterialInterface.h b/Engine/Source/Runtime/Engine/Classes/Materials/MaterialInterface.h
--- a/Engine/Source/Runtime/Engine/Classes/Materials/MaterialInterface.h
+++ b/Engine/Source/Runtime/Engine/Classes/Materials/MaterialInterface.h
@@ -16,25 +16,36 @@
 #include "UObject/ScriptMacros.h"
 #include "RenderCommandFence.h"
 #include "SceneTypes.h"
-#include "RHI.h"
+#include "RHIDefinitions.h"
 #include "Engine/BlendableInterface.h"
 #include "Materials/MaterialLayersFunctions.h"
 #include "Interfaces/Interface_AssetUserData.h"
 #include "MaterialSceneTextureId.h"
 #include "Materials/MaterialRelevance.h"
+#include "MaterialRecursionGuard.h"
+#include "MaterialShaderPrecompileMode.h"
+
+#if UE_ENABLE_INCLUDE_ORDER_DEPRECATED_IN_5_2
+#include "RHI.h"
 #include "Physics/PhysicsInterfaceCore.h"
 #include "MaterialShared.h"
+#endif
+
 #include "MaterialInterface.generated.h"
 
 class FMaterialCompiler;
 class FMaterialRenderProxy;
 class FMaterialResource;
+class FShaderPipelineType;
+class FShaderType;
+class FVertexFactoryType;
 class UMaterial;
 class UPhysicalMaterial;
 class UPhysicalMaterialMask;
 class USubsurfaceProfile;
 class UTexture;
 class UMaterialInstance;
+struct FDebugShaderTypeInfo;
 struct FMaterialParameterInfo;
 struct FMaterialResourceLocOnDisk;
 class FMaterialCachedData;
@@ -964,7 +975,7 @@ public:
 #endif // WITH_EDITOR
 
 	/** Get bitfield indicating which feature levels should be compiled by default */
-	ENGINE_API static uint32 GetFeatureLevelsToCompileForAllMaterials() { return FeatureLevelsForAllMaterials | (1 << GMaxRHIFeatureLevel); }
+	ENGINE_API static uint32 GetFeatureLevelsToCompileForAllMaterials();
 
 	/** Return number of used texture coordinates and whether or not the Vertex data is used in the shader graph */
 	ENGINE_API void AnalyzeMaterialProperty(EMaterialProperty InProperty, int32& OutNumTextureCoordinates, bool& bOutRequiresVertexData);
diff --git a/Engine/Source/Runtime/Engine/Classes/Materials/MaterialParameterCollection.h b/Engine/Source/Runtime/Engine/Classes/Materials/MaterialParameterCollection.h
--- a/Engine/Source/Runtime/Engine/Classes/Materials/MaterialParameterCollection.h
+++ b/Engine/Source/Runtime/Engine/Classes/Materials/MaterialParameterCollection.h
@@ -11,11 +11,14 @@
 #include "UObject/Object.h"
 #include "HAL/ThreadSafeBool.h"
 #include "Misc/Guid.h"
+#if UE_ENABLE_INCLUDE_ORDER_DEPRECATED_IN_5_2
 #include "UniformBuffer.h"
+#endif
 #include "Templates/UniquePtr.h"
 #include "RenderCommandFence.h"
 #include "MaterialParameterCollection.generated.h"
 
+class FShaderParametersMetadata;
 struct FPropertyChangedEvent;
 
 /** Base struct for collection parameters */
diff --git a/Engine/Source/Runtime/Engine/Classes/Particles/ParticleEmitter.h b/Engine/Source/Runtime/Engine/Classes/Particles/ParticleEmitter.h
--- a/Engine/Source/Runtime/Engine/Classes/Particles/ParticleEmitter.h
+++ b/Engine/Source/Runtime/Engine/Classes/Particles/ParticleEmitter.h
@@ -325,27 +325,7 @@ class UParticleEmitter : public UObject
 	*	@return NULL if the requested LODLevel is not valid.
 	*			The pointer to the requested UParticleLODLevel if valid.
 	*/
-	FORCEINLINE UParticleLODLevel* GetCurrentLODLevel(FParticleEmitterInstance* Instance)
-	{
-		if (!FPlatformProperties::HasEditorOnlyData())
-		{
-			return Instance->CurrentLODLevel;
-		}
-		else
-		{
-			// for the game (where we care about perf) we don't branch
-			if (Instance->GetWorld()->IsGameWorld() )
-			{
-				return Instance->CurrentLODLevel;
-			}
-			else
-			{
-				EditorUpdateCurrentLOD( Instance );
-				return Instance->CurrentLODLevel;
-			}
-		}
-	}
-
+	ENGINE_API UParticleLODLevel* GetCurrentLODLevel(FParticleEmitterInstance* Instance);
 
 	/**
 	 * This will update the LOD of the particle in the editor.
diff --git a/Engine/Source/Runtime/Engine/Classes/PhysicsEngine/BodyInstance.h b/Engine/Source/Runtime/Engine/Classes/PhysicsEngine/BodyInstance.h
--- a/Engine/Source/Runtime/Engine/Classes/PhysicsEngine/BodyInstance.h
+++ b/Engine/Source/Runtime/Engine/Classes/PhysicsEngine/BodyInstance.h
@@ -8,7 +8,9 @@
 #include "Engine/EngineTypes.h"
 #include "CollisionQueryParams.h"
 #include "EngineDefines.h"
+#if UE_ENABLE_INCLUDE_ORDER_DEPRECATED_IN_5_2
 #include "PhysxUserData.h"
+#endif
 #include "Physics/PhysicsInterfaceCore.h"
 #include "Physics/PhysicsInterfaceTypes.h"
 #include "PhysicsPublic.h"
@@ -24,6 +26,7 @@ struct FCollisionNotifyInfo;
 struct FCollisionShape;
 struct FConstraintInstance;
 struct FPropertyChangedEvent;
+struct FRigidBodyState;
 struct FShapeData;
 class UPrimitiveComponent;
 
diff --git a/Engine/Source/Runtime/Engine/Private/Animation/AnimInstance.cpp b/Engine/Source/Runtime/Engine/Private/Animation/AnimInstance.cpp
--- a/Engine/Source/Runtime/Engine/Private/Animation/AnimInstance.cpp
+++ b/Engine/Source/Runtime/Engine/Private/Animation/AnimInstance.cpp
@@ -14,6 +14,7 @@
 #include "Animation/AnimBlueprint.h"
 #include "Animation/AnimClassInterface.h"
 #include "Animation/AnimBlueprintGeneratedClass.h"
+#include "Components/SkeletalMeshComponent.h"
 #include "DisplayDebugHelpers.h"
 #include "Animation/BlendSpace.h"
 #include "Animation/AnimNotifies/AnimNotifyState.h"
@@ -1271,6 +1272,11 @@ void UAnimInstance::RecalcRequiredCurves(const FCurveEvaluationOption& CurveEval
 	GetProxyOnGameThread<FAnimInstanceProxy>().RecalcRequiredCurves(CurveEvalOption);
 }
 
+USkeletalMeshComponent* UAnimInstance::GetSkelMeshComponent() const
+{
+	return CastChecked<USkeletalMeshComponent>(GetOuter());
+}
+
 void UAnimInstance::Serialize(FArchive& Ar)
 {
 	Super::Serialize(Ar);
@@ -3833,6 +3839,18 @@ bool UAnimInstance::ShouldTriggerAnimNotifyState(const UAnimNotifyState* AnimNot
 	return false;
 }
 
+bool UAnimInstance::IsSkeletalMeshComponent(const UObject* Object)
+{
+	return Object && Object->IsA<USkeletalMeshComponent>();
+}
+
+void UAnimInstance::HandleExistingParallelEvaluationTask(USkeletalMeshComponent* Component)
+{
+	bool bBlockOnTask = true;
+	bool bPerformPostAnimEvaluation = true;
+	Component->HandleExistingParallelEvaluationTask(bBlockOnTask, bPerformPostAnimEvaluation);
+}
+
 void UAnimInstance::RecordMachineWeight(const int32 InMachineClassIndex, const float InMachineWeight)
 {
 	GetProxyOnAnyThread<FAnimInstanceProxy>().RecordMachineWeight(InMachineClassIndex, InMachineWeight);
diff --git a/Engine/Source/Runtime/Engine/Private/Components/LineBatchComponent.cpp b/Engine/Source/Runtime/Engine/Private/Components/LineBatchComponent.cpp
--- a/Engine/Source/Runtime/Engine/Private/Components/LineBatchComponent.cpp
+++ b/Engine/Source/Runtime/Engine/Private/Components/LineBatchComponent.cpp
@@ -17,6 +17,32 @@
 
 #include UE_INLINE_GENERATED_CPP_BY_NAME(LineBatchComponent)
 
+/** Represents a LineBatchComponent to the scene manager. */
+class ENGINE_API FLineBatcherSceneProxy : public FPrimitiveSceneProxy
+{
+public:
+	SIZE_T GetTypeHash() const override;
+
+	FLineBatcherSceneProxy(const ULineBatchComponent* InComponent);
+
+	virtual void GetDynamicMeshElements(const TArray<const FSceneView*>& Views, const FSceneViewFamily& ViewFamily, uint32 VisibilityMap, FMeshElementCollector& Collector) const override;
+
+	/**
+	*  Returns a struct that describes to the renderer when to draw this proxy.
+	*	@param		Scene view to use to determine our relevence.
+	*  @return		View relevance struct
+	*/
+	virtual FPrimitiveViewRelevance GetViewRelevance(const FSceneView* View) const override;
+	virtual uint32 GetMemoryFootprint( void ) const override;
+	uint32 GetAllocatedSize( void ) const;
+
+private:
+	TArray<FBatchedLine> Lines;
+	TArray<FBatchedPoint> Points;
+	TArray<FBatchedMesh> Meshes;
+};
+
+
 FLineBatcherSceneProxy::FLineBatcherSceneProxy(const ULineBatchComponent* InComponent) :
 	FPrimitiveSceneProxy(InComponent), Lines(InComponent->BatchedLines), 
 	Points(InComponent->BatchedPoints), Meshes(InComponent->BatchedMeshes)
diff --git a/Engine/Source/Runtime/Engine/Private/Components/SkinnedMeshComponent.cpp b/Engine/Source/Runtime/Engine/Private/Components/SkinnedMeshComponent.cpp
--- a/Engine/Source/Runtime/Engine/Private/Components/SkinnedMeshComponent.cpp
+++ b/Engine/Source/Runtime/Engine/Private/Components/SkinnedMeshComponent.cpp
@@ -5,6 +5,7 @@
 =============================================================================*/
 
 #include "Components/SkinnedMeshComponent.h"
+#include "Components/ExternalMorphSet.h"
 #include "Misc/App.h"
 #include "RenderingThread.h"
 #include "GameFramework/PlayerController.h"
@@ -18,11 +19,13 @@
 #include "Animation/AnimStats.h"
 #include "SkeletalMeshDeformerHelpers.h"
 #include "Engine/SkeletalMeshSocket.h"
+#include "Engine/SkinnedAssetCommon.h"
 #include "PhysicsEngine/PhysicsAsset.h"
 #include "EngineGlobals.h"
 #include "PrimitiveSceneProxy.h"
 #include "Engine/CollisionProfile.h"
 #include "Rendering/SkinWeightVertexBuffer.h"
+#include "SkeletalMeshSceneProxy.h"
 #include "SkeletalMeshTypes.h"
 #include "Animation/MeshDeformer.h"
 #include "Animation/MeshDeformerInstance.h"
@@ -33,6 +36,7 @@
 #include "GPUSkinCache.h"
 #include "PipelineStateCache.h"
 #include "SkeletalRender.h"
+#include "UObject/UE5MainStreamObjectVersion.h"
 
 DEFINE_LOG_CATEGORY_STATIC(LogSkinnedMeshComp, Log, All);
 
@@ -440,6 +444,8 @@ USkinnedMeshComponent::USkinnedMeshComponent(const FObjectInitializer& ObjectIni
 	PRAGMA_ENABLE_DEPRECATION_WARNINGS
 }
 
+USkinnedMeshComponent::USkinnedMeshComponent(FVTableHelper& Helper) : Super(Helper) {}
+USkinnedMeshComponent::~USkinnedMeshComponent() = default;
 
 void USkinnedMeshComponent::UpdateMorphMaterialUsageOnProxy()
 {
@@ -664,6 +670,16 @@ void USkinnedMeshComponent::BeginPlay()
 	MeshDeformerInstance = ActiveMeshDeformer != nullptr ? ActiveMeshDeformer->CreateInstance(this, MeshDeformerInstanceSettings) : nullptr;
 }
 
+const FExternalMorphWeightData& USkinnedMeshComponent::GetExternalMorphWeights(int32 LOD) const
+{
+	return ExternalMorphWeightData[LOD];
+}
+
+FExternalMorphWeightData& USkinnedMeshComponent::GetExternalMorphWeights(int32 LOD)
+{
+	return ExternalMorphWeightData[LOD];
+}
+
 void USkinnedMeshComponent::AddExternalMorphSet(int32 LOD, int32 ID, TSharedPtr<FExternalMorphSet> MorphSet)
 {
 	const int32 NumLODs = GetNumLODs();
@@ -1964,6 +1980,11 @@ USkeletalMesh* USkinnedMeshComponent::GetSkeletalMesh_DEPRECATED() const
 	return Cast<class USkeletalMesh>(GetSkinnedAsset());
 }
 
+void USkinnedMeshComponent::SetSkeletalMesh_DEPRECATED(USkeletalMesh* NewMesh)
+{
+	SetSkinnedAssetAndUpdate(Cast<USkinnedAsset>(NewMesh));
+}
+
 void USkinnedMeshComponent::SetSkinnedAssetAndUpdate(USkinnedAsset* InSkinnedAsset, bool /*bReinitPose*/)
 {
 	// NOTE: InSkinnedAsset may be nullptr (useful in the editor for removing the skeletal mesh associated with
diff --git a/Engine/Source/Runtime/Engine/Private/ContentStreaming.cpp b/Engine/Source/Runtime/Engine/Private/ContentStreaming.cpp
--- a/Engine/Source/Runtime/Engine/Private/ContentStreaming.cpp
+++ b/Engine/Source/Runtime/Engine/Private/ContentStreaming.cpp
@@ -1,5 +1,4 @@
 // Copyright Epic Games, Inc. All Rights Reserved.
-// Copyright Epic Games, Inc. All Rights Reserved.
 
 /*=============================================================================
 	ContentStreaming.cpp: Implementation of content streaming classes.
@@ -17,6 +16,7 @@
 #include "EngineGlobals.h"
 #include "Components/MeshComponent.h"
 #include "Engine/Engine.h"
+#include "Engine/Level.h"
 #include "Streaming/TextureStreamingHelpers.h"
 #include "Streaming/StreamingManagerTexture.h"
 #include "AudioStreaming.h"
@@ -1006,6 +1006,11 @@ bool FStreamingManagerCollection::IsStreamingEnabled() const
 	return DisableResourceStreamingCount == 0;
 }
 
+bool FStreamingManagerCollection::IsTextureStreamingEnabled() const
+{
+	return IsRenderAssetStreamingEnabled(EStreamableRenderAssetType::Texture);
+}
+
 bool FStreamingManagerCollection::IsRenderAssetStreamingEnabled(EStreamableRenderAssetType FilteredAssetType) const
 {
 	if (RenderAssetStreamingManager)
diff --git a/Engine/Source/Runtime/Engine/Private/GameInstance.cpp b/Engine/Source/Runtime/Engine/Private/GameInstance.cpp
--- a/Engine/Source/Runtime/Engine/Private/GameInstance.cpp
+++ b/Engine/Source/Runtime/Engine/Private/GameInstance.cpp
@@ -15,6 +15,7 @@
 #include "GameFramework/PlayerController.h"
 #include "Engine/Engine.h"
 #include "Engine/Console.h"
+#include "Engine/GameViewportClient.h"
 #include "Misc/ScopedSlowTask.h"
 #include "Engine/GameEngine.h"
 #include "GameFramework/GameModeBase.h"
@@ -42,6 +43,19 @@
 #include "StudioAnalytics.h"
 #endif
 
+#if WITH_EDITOR
+FGameInstancePIEParameters::FGameInstancePIEParameters()
+	: bSimulateInEditor(false)
+	, bAnyBlueprintErrors(false)
+	, bStartInSpectatorMode(false)
+	, bRunAsDedicated(false)
+	, bIsPrimaryPIEClient(false)
+	, WorldFeatureLevel(ERHIFeatureLevel::Num)
+	, EditorPlaySettings(nullptr)
+	, NetMode(EPlayNetMode::PIE_Standalone)
+{}
+#endif
+
 UGameInstance::UGameInstance(const FObjectInitializer& ObjectInitializer)
 	: Super(ObjectInitializer)
 	, TimerManager(new FTimerManager(this))
diff --git a/Engine/Source/Runtime/Engine/Private/Materials/MaterialInterface.cpp b/Engine/Source/Runtime/Engine/Private/Materials/MaterialInterface.cpp
--- a/Engine/Source/Runtime/Engine/Private/Materials/MaterialInterface.cpp
+++ b/Engine/Source/Runtime/Engine/Private/Materials/MaterialInterface.cpp
@@ -17,12 +17,14 @@
 #include "EditorFramework/AssetImportData.h"
 #include "Engine/AssetUserData.h"
 #include "Engine/Texture2D.h"
+#include "Engine/World.h"
 #include "ObjectCacheEventSink.h"
 #include "Engine/SubsurfaceProfile.h"
 #include "Engine/TextureStreamingTypes.h"
 #include "Algo/BinarySearch.h"
 #include "Interfaces/ITargetPlatform.h"
 #include "Components.h"
+#include "Components/PrimitiveComponent.h"
 #include "ContentStreaming.h"
 #include "MeshBatch.h"
 #include "TextureCompiler.h"
@@ -1404,6 +1406,11 @@ float UMaterialInterface::GetTextureDensity(FName TextureName, const FMeshUVChan
 	return 0;
 }
 
+uint32 UMaterialInterface::GetFeatureLevelsToCompileForAllMaterials()
+{
+	return FeatureLevelsForAllMaterials | (1 << GMaxRHIFeatureLevel);
+}
+
 bool UMaterialInterface::UseAnyStreamingTexture() const
 {
 	TArray<UTexture*> Textures;
diff --git a/Engine/Source/Runtime/Engine/Private/Particles/ParticleComponents.cpp b/Engine/Source/Runtime/Engine/Private/Particles/ParticleComponents.cpp
--- a/Engine/Source/Runtime/Engine/Private/Particles/ParticleComponents.cpp
+++ b/Engine/Source/Runtime/Engine/Private/Particles/ParticleComponents.cpp
@@ -16,6 +16,7 @@
 #include "EngineDefines.h"
 #include "EngineGlobals.h"
 #include "Engine/EngineTypes.h"
+#include "Engine/Level.h"
 #include "Components/ActorComponent.h"
 #include "Components/SceneComponent.h"
 #include "CollisionQueryParams.h"
@@ -1474,6 +1475,26 @@ bool UParticleEmitter::IsLODLevelValid(int32 LODLevel)
 	return false;
 }
 
+UParticleLODLevel* UParticleEmitter::GetCurrentLODLevel(FParticleEmitterInstance* Instance)
+{
+	if (!FPlatformProperties::HasEditorOnlyData())
+	{
+		return Instance->CurrentLODLevel;
+	}
+	else
+	{
+		// for the game (where we care about perf) we don't branch
+		if (Instance->GetWorld()->IsGameWorld() )
+		{
+			return Instance->CurrentLODLevel;
+		}
+		else
+		{
+			EditorUpdateCurrentLOD( Instance );
+			return Instance->CurrentLODLevel;
+		}
+	}
+}
 
 void UParticleEmitter::EditorUpdateCurrentLOD(FParticleEmitterInstance* Instance)
 {
diff --git a/Engine/Source/Runtime/Engine/Private/PerPlatformProperties.cpp b/Engine/Source/Runtime/Engine/Private/PerPlatformProperties.cpp
--- a/Engine/Source/Runtime/Engine/Private/PerPlatformProperties.cpp
+++ b/Engine/Source/Runtime/Engine/Private/PerPlatformProperties.cpp
@@ -2,6 +2,7 @@
 
 #include "PerPlatformProperties.h"
 #include "Serialization/Archive.h"
+#include "Engine/Engine.h"
 
 #include UE_INLINE_GENERATED_CPP_BY_NAME(PerPlatformProperties)
 
@@ -111,3 +112,9 @@ FString FFreezablePerPlatformInt::ToString() const
 	return FPerPlatformInt(*this).ToString();
 }
 
+#if WITH_EDITORONLY_DATA && WITH_EDITOR
+bool GEngine_GetPreviewPlatformName(FName& PlatformName)
+{
+	return GEngine && GEngine->GetPreviewPlatformName(PlatformName);
+}
+#endif
\ No newline at end of file
diff --git a/Engine/Source/Runtime/Engine/Private/PrimitiveUniformShaderParameters.cpp b/Engine/Source/Runtime/Engine/Private/PrimitiveUniformShaderParameters.cpp
--- a/Engine/Source/Runtime/Engine/Private/PrimitiveUniformShaderParameters.cpp
+++ b/Engine/Source/Runtime/Engine/Private/PrimitiveUniformShaderParameters.cpp
@@ -6,7 +6,25 @@
 #include "PrimitiveSceneInfo.h"
 #include "NaniteSceneProxy.h"
 #include "ProfilingDebugging/LoadTimeTracker.h"
+#include "UnrealEngine.h"
 
+FPrimitiveUniformShaderParametersBuilder& FPrimitiveUniformShaderParametersBuilder::InstanceDrawDistance(FVector2f DistanceMinMax)
+{
+	// Only scale the far distance by scalability parameters
+	DistanceMinMax.Y *= GetCachedScalabilityCVars().ViewDistanceScale;
+	Parameters.InstanceDrawDistanceMinMaxSquared = FMath::Square(DistanceMinMax);
+	bHasInstanceDrawDistanceCull = true;
+	return *this;
+}
+
+FPrimitiveUniformShaderParametersBuilder& FPrimitiveUniformShaderParametersBuilder::InstanceWorldPositionOffsetDisableDistance(float WPODisableDistance)
+{
+	WPODisableDistance *= GetCachedScalabilityCVars().ViewDistanceScale;
+	bHasWPODisableDistance = true;
+	Parameters.InstanceWPODisableDistanceSquared = WPODisableDistance * WPODisableDistance;
+
+	return *this;
+}
 
 void FSinglePrimitiveStructured::InitRHI() 
 {
diff --git a/Engine/Source/Runtime/Engine/Private/SceneManagement.cpp b/Engine/Source/Runtime/Engine/Private/SceneManagement.cpp
--- a/Engine/Source/Runtime/Engine/Private/SceneManagement.cpp
+++ b/Engine/Source/Runtime/Engine/Private/SceneManagement.cpp
@@ -14,6 +14,8 @@
 #include "Engine/Engine.h"
 #include "Engine/LightMapTexture2D.h"
 #include "Engine/ShadowMapTexture2D.h"
+#include "Engine/TextureLightProfile.h"
+#include "TextureResource.h"
 #include "VT/LightmapVirtualTexture.h"
 #include "UnrealEngine.h"
 #include "ColorSpace.h"
@@ -24,6 +26,12 @@ static TAutoConsoleVariable<float> CVarLODTemporalLag(
 	TEXT("This controls the the time lag for temporal LOD, in seconds."),
 	ECVF_Scalability | ECVF_Default);
 
+FTexture* FLightSceneProxy::GetIESTextureResource() const
+{
+	return IESTexture ? IESTexture->GetResource() : nullptr;
+}
+
+
 void FTemporalLODState::UpdateTemporalLODTransition(const FViewInfo& View, float LastRenderTime)
 {
 	bool bOk = false;
diff --git a/Engine/Source/Runtime/Engine/Private/SceneView.cpp b/Engine/Source/Runtime/Engine/Private/SceneView.cpp
--- a/Engine/Source/Runtime/Engine/Private/SceneView.cpp
+++ b/Engine/Source/Runtime/Engine/Private/SceneView.cpp
@@ -2809,6 +2809,33 @@ TArray<const FSceneView*> FSceneView::GetSecondaryViews() const
 	return Views;
 }
 
+FSceneViewFamily::ConstructionValues::ConstructionValues(
+	const FRenderTarget* InRenderTarget,
+	FSceneInterface* InScene,
+	const FEngineShowFlags& InEngineShowFlags
+	)
+:	RenderTarget(InRenderTarget)
+,	Scene(InScene)
+,	EngineShowFlags(InEngineShowFlags)
+,	ViewModeParam(-1)
+,	GammaCorrection(1.0f)
+,	bAdditionalViewFamily(false)
+,	bRealtimeUpdate(false)
+,	bDeferClear(false)
+,	bResolveScene(true)			
+,	bTimesSet(false)
+{
+	if( InScene != NULL )			
+	{
+		UWorld* World = InScene->GetWorld();
+		// Ensure the world is valid and that we are being called from a game thread (GetRealTimeSeconds requires this)
+		if( World && IsInGameThread() )
+		{	
+			SetTime(World->GetTime());
+		}
+	}
+}
+
 FSceneViewFamily::FSceneViewFamily(const ConstructionValues& CVS)
 	:
 	ViewMode(VMI_Lit),
diff --git a/Engine/Source/Runtime/Engine/Private/SkeletalMesh.cpp b/Engine/Source/Runtime/Engine/Private/SkeletalMesh.cpp
--- a/Engine/Source/Runtime/Engine/Private/SkeletalMesh.cpp
+++ b/Engine/Source/Runtime/Engine/Private/SkeletalMesh.cpp
@@ -15,6 +15,9 @@
 #include "RawIndexBuffer.h"
 #include "Engine/TextureStreamingTypes.h"
 #include "Engine/Brush.h"
+#include "Engine/SkinnedAssetAsyncCompileUtils.h"
+#include "Engine/SkinnedAssetCommon.h"
+#include "Engine/SkeletalMeshLODSettings.h"
 #include "MaterialShared.h"
 #include "Materials/Material.h"
 #include "Components/SkinnedMeshComponent.h"
@@ -38,6 +41,7 @@
 #include "EngineUtils.h"
 #include "EditorSupportDelegates.h"
 #include "GPUSkinVertexFactory.h"
+#include "SkeletalMeshSceneProxy.h"
 #include "SkeletalRenderPublic.h"
 #include "Logging/TokenizedMessage.h"
 #include "Logging/MessageLog.h"
@@ -437,6 +441,30 @@ PRAGMA_DISABLE_DEPRECATION_WARNINGS
 USkeletalMesh::~USkeletalMesh() = default;
 PRAGMA_ENABLE_DEPRECATION_WARNINGS
 
+FSkeletalMeshRenderData* USkeletalMesh::GetSkeletalMeshRenderData() const
+{
+	WaitUntilAsyncPropertyReleased(ESkeletalMeshAsyncProperties::SkeletalMeshRenderData);
+	PRAGMA_DISABLE_DEPRECATION_WARNINGS
+	return SkeletalMeshRenderData.Get();
+	PRAGMA_ENABLE_DEPRECATION_WARNINGS
+}
+
+void USkeletalMesh::SetSkeletalMeshRenderData(TUniquePtr<FSkeletalMeshRenderData>&& InSkeletalMeshRenderData)
+{
+	WaitUntilAsyncPropertyReleased(ESkeletalMeshAsyncProperties::SkeletalMeshRenderData);
+	PRAGMA_DISABLE_DEPRECATION_WARNINGS
+	SkeletalMeshRenderData = MoveTemp(InSkeletalMeshRenderData);
+	PRAGMA_ENABLE_DEPRECATION_WARNINGS
+}
+
+FSkeletalMeshRenderData* USkeletalMesh::GetResourceForRendering() const
+{
+	WaitUntilAsyncPropertyReleased(ESkeletalMeshAsyncProperties::SkeletalMeshRenderData);
+	PRAGMA_DISABLE_DEPRECATION_WARNINGS
+	return SkeletalMeshRenderData.Get();
+	PRAGMA_ENABLE_DEPRECATION_WARNINGS
+}
+
 void USkeletalMesh::PostInitProperties()
 {
 #if WITH_EDITORONLY_DATA
@@ -595,6 +623,14 @@ FSkeletalMeshOptimizationSettings USkeletalMesh::GetReductionSettings(int32 LODI
 
 #endif
 
+void USkeletalMesh::SetMaterials(const TArray<FSkeletalMaterial>& InMaterials)
+{
+	WaitUntilAsyncPropertyReleased(ESkeletalMeshAsyncProperties::Materials);
+	PRAGMA_DISABLE_DEPRECATION_WARNINGS
+	Materials = InMaterials;
+	PRAGMA_ENABLE_DEPRECATION_WARNINGS
+}
+
 void USkeletalMesh::AddClothingAsset(UClothingAssetBase* InNewAsset)
 {
 	check(IsInGameThread());
@@ -5044,6 +5080,22 @@ class UNodeMappingContainer* USkeletalMesh::GetNodeMappingContainer(class UBluep
 	return nullptr;
 }
 
+FSkeletalMeshLODInfo* USkeletalMesh::GetLODInfo(int32 Index)
+{
+	WaitUntilAsyncPropertyReleased(ESkeletalMeshAsyncProperties::LODInfo);
+	PRAGMA_DISABLE_DEPRECATION_WARNINGS
+	return LODInfo.IsValidIndex(Index) ? &LODInfo[Index] : nullptr;  
+	PRAGMA_ENABLE_DEPRECATION_WARNINGS
+}
+	
+const FSkeletalMeshLODInfo* USkeletalMesh::GetLODInfo(int32 Index) const
+{
+	WaitUntilAsyncPropertyReleased(ESkeletalMeshAsyncProperties::LODInfo, ESkinnedAssetAsyncPropertyLockType::ReadOnly);
+	PRAGMA_DISABLE_DEPRECATION_WARNINGS
+	return LODInfo.IsValidIndex(Index) ? &LODInfo[Index] : nullptr;
+	PRAGMA_ENABLE_DEPRECATION_WARNINGS
+}
+
 const UAnimSequence* USkeletalMesh::GetBakePose(int32 LODIndex) const
 {
 	const FSkeletalMeshLODInfo* LOD = GetLODInfo(LODIndex);
@@ -5076,6 +5128,24 @@ const USkeletalMeshLODSettings* USkeletalMesh::GetDefaultLODSetting() const
 	return GetDefault<USkeletalMeshLODSettings>();
 }
 
+bool USkeletalMesh::IsValidLODIndex(int32 Index) const
+{
+	WaitUntilAsyncPropertyReleased(ESkeletalMeshAsyncProperties::LODInfo, ESkinnedAssetAsyncPropertyLockType::ReadOnly);
+	PRAGMA_DISABLE_DEPRECATION_WARNINGS
+	return LODInfo.IsValidIndex(Index);
+	PRAGMA_ENABLE_DEPRECATION_WARNINGS
+}
+/* 
+	* Returns total number of LOD. USkinnedAsset interface.
+	*/
+int32 USkeletalMesh::GetLODNum() const
+{
+	WaitUntilAsyncPropertyReleased(ESkeletalMeshAsyncProperties::LODInfo, ESkinnedAssetAsyncPropertyLockType::ReadOnly);
+	PRAGMA_DISABLE_DEPRECATION_WARNINGS
+	return LODInfo.Num();
+	PRAGMA_ENABLE_DEPRECATION_WARNINGS
+}
+
 bool USkeletalMesh::IsMaterialUsed(int32 MaterialIndex) const
 {
 	if (GIsEditor || !CVarSkeletalMeshLODMaterialReference.GetValueOnAnyThread())
@@ -5115,6 +5185,16 @@ bool USkeletalMesh::IsMaterialUsed(int32 MaterialIndex) const
 	return false;
 }
 
+#if WITH_EDITOR
+void USkeletalMesh::AddSkinWeightProfile(const FSkinWeightProfileInfo& Profile) 
+{
+	WaitUntilAsyncPropertyReleased(ESkeletalMeshAsyncProperties::SkinWeightProfiles); 
+	PRAGMA_DISABLE_DEPRECATION_WARNINGS
+	SkinWeightProfiles.Add(Profile); 
+	PRAGMA_ENABLE_DEPRECATION_WARNINGS
+}
+#endif
+
 void USkeletalMesh::ReleaseSkinWeightProfileResources()
 {
 	// This assumes that skin weights buffers are not used anywhere
@@ -5174,6 +5254,14 @@ FSkeletalMeshLODInfo& USkeletalMesh::AddLODInfo()
 	return LODInfoArray[NewIndex];
 }
 
+void USkeletalMesh::AddLODInfo(const FSkeletalMeshLODInfo& NewLODInfo) 
+{
+	WaitUntilAsyncPropertyReleased(ESkeletalMeshAsyncProperties::LODInfo);
+	PRAGMA_DISABLE_DEPRECATION_WARNINGS
+	LODInfo.Add(NewLODInfo);
+	PRAGMA_ENABLE_DEPRECATION_WARNINGS
+}
+
 void USkeletalMesh::RemoveLODInfo(int32 Index)
 {
 	TArray<FSkeletalMeshLODInfo>& LODInfoArray = GetLODInfoArray();
diff --git a/Engine/Source/Runtime/Engine/Private/SkinnedAsset.cpp b/Engine/Source/Runtime/Engine/Private/SkinnedAsset.cpp
--- a/Engine/Source/Runtime/Engine/Private/SkinnedAsset.cpp
+++ b/Engine/Source/Runtime/Engine/Private/SkinnedAsset.cpp
@@ -1,11 +1,15 @@
 // Copyright Epic Games, Inc. All Rights Reserved.
 
 #include "Engine/SkinnedAsset.h"
+#include "Engine/SkinnedAssetAsyncCompileUtils.h"
+#include "Engine/SkinnedAssetCommon.h"
 #include "SkinnedAssetCompiler.h"
 #include "Animation/AnimStats.h"
 #include "SkeletalRenderGPUSkin.h"
 #include "PSOPrecache.h"
 
+#include UE_INLINE_GENERATED_CPP_BY_NAME(SkinnedAsset)
+
 #if INTEL_ISPC
 #include "SkinnedAsset.ispc.generated.h"
 static_assert(sizeof(ispc::FTransform) == sizeof(FTransform), "sizeof(ispc::FTransform) != sizeof(FTransform)");
@@ -25,6 +29,22 @@ static FAutoConsoleVariableRef CVarAnimSkinnedAssetISPCEnabled(TEXT("a.SkinnedAs
 
 DEFINE_LOG_CATEGORY_STATIC(LogSkinnedAsset, Log, All);
 
+USkinnedAsset::USkinnedAsset(const FObjectInitializer& ObjectInitializer)
+	: Super(ObjectInitializer)
+{}
+
+USkinnedAsset::~USkinnedAsset() {}
+
+bool USkinnedAsset::IsValidMaterialIndex(int32 Index) const
+{
+	return GetMaterials().IsValidIndex(Index);
+}
+
+int32 USkinnedAsset::GetNumMaterials() const
+{
+	return GetMaterials().Num();
+}
+
 void USkinnedAsset::PostLoad()
 {
 	Super::PostLoad();
@@ -264,6 +284,11 @@ static bool IsISPCEnabled()
 	return bAnim_SkinnedAsset_ISPC_Enabled;
 }
 
+ESkeletalMeshVertexFlags USkinnedAsset::GetVertexBufferFlags() const
+{
+	return GetHasVertexColors() ? ESkeletalMeshVertexFlags::HasVertexColors : ESkeletalMeshVertexFlags::None;
+}
+
 void USkinnedAsset::FillComponentSpaceTransforms(const TArray<FTransform>& InBoneSpaceTransforms,
 												 const TArray<FBoneIndexType>& InFillComponentSpaceTransformsRequiredBones, 
 												 TArray<FTransform>& OutComponentSpaceTransforms) const
@@ -339,4 +364,16 @@ void USkinnedAsset::FillComponentSpaceTransforms(const TArray<FTransform>& InBon
 			checkSlow(!SpaceBase->ContainsNaN());
 		}
 	}
-}
\ No newline at end of file
+}
+
+TArray<FSkeletalMeshLODInfo>& USkinnedAsset::GetMeshLodInfoDummyArray()
+{
+	static TArray<FSkeletalMeshLODInfo> Dummy;
+	return Dummy;
+}
+
+TArray<FSkeletalMaterial>& USkinnedAsset::GetSkeletalMaterialDummyArray()
+{
+	static TArray<FSkeletalMaterial> Dummy;
+	return Dummy;
+}
diff --git a/Engine/Source/Runtime/Engine/Private/StaticMesh.cpp b/Engine/Source/Runtime/Engine/Private/StaticMesh.cpp
--- a/Engine/Source/Runtime/Engine/Private/StaticMesh.cpp
+++ b/Engine/Source/Runtime/Engine/Private/StaticMesh.cpp
@@ -101,6 +101,8 @@
 #include "ContentStreaming.h"
 #include "Streaming/StaticMeshUpdate.h"
 
+#include UE_INLINE_GENERATED_CPP_BY_NAME(StaticMesh)
+
 #define LOCTEXT_NAMESPACE "StaticMesh"
 DEFINE_LOG_CATEGORY(LogStaticMesh);	
 
@@ -3142,6 +3144,42 @@ FArchive& operator<<(FArchive& Ar, FStaticMaterial& Elem)
 	return Ar;
 }
 
+FStaticMaterial::FStaticMaterial()
+: MaterialInterface(NULL)
+, MaterialSlotName(NAME_None)
+#if WITH_EDITORONLY_DATA
+, ImportedMaterialSlotName(NAME_None)
+#endif //WITH_EDITORONLY_DATA
+{
+
+}
+
+FStaticMaterial::FStaticMaterial(class UMaterialInterface* InMaterialInterface
+, FName InMaterialSlotName
+#if WITH_EDITORONLY_DATA
+, FName InImportedMaterialSlotName)
+#else
+)
+#endif
+: MaterialInterface(InMaterialInterface)
+, MaterialSlotName(InMaterialSlotName)
+#if WITH_EDITORONLY_DATA
+, ImportedMaterialSlotName(InImportedMaterialSlotName)
+#endif //WITH_EDITORONLY_DATA
+{
+//If not specified add some valid material slot name
+if (MaterialInterface && MaterialSlotName == NAME_None)
+{
+	MaterialSlotName = MaterialInterface->GetFName();
+}
+#if WITH_EDITORONLY_DATA
+if (ImportedMaterialSlotName == NAME_None)
+{
+	ImportedMaterialSlotName = MaterialSlotName;
+}
+#endif
+}
+
 bool operator== (const FStaticMaterial& LHS, const FStaticMaterial& RHS)
 {
 	return (LHS.MaterialInterface == RHS.MaterialInterface &&
diff --git a/Engine/Source/Runtime/Engine/Private/StaticMeshRender.cpp b/Engine/Source/Runtime/Engine/Private/StaticMeshRender.cpp
--- a/Engine/Source/Runtime/Engine/Private/StaticMeshRender.cpp
+++ b/Engine/Source/Runtime/Engine/Private/StaticMeshRender.cpp
@@ -25,7 +25,9 @@
 #include "Engine/StaticMesh.h"
 #include "ComponentReregisterContext.h"
 #include "EngineUtils.h"
+#include "StaticMeshComponentLODInfo.h"
 #include "StaticMeshResources.h"
+#include "StaticMeshSceneProxy.h"
 #include "Rendering/NaniteResources.h"
 #include "SpeedTreeWind.h"
 #include "PhysicalMaterials/PhysicalMaterialMask.h"
@@ -1386,6 +1388,11 @@ bool FStaticMeshSceneProxy::IsCollisionView(const FEngineShowFlags& EngineShowFl
 	return bInCollisionView;
 }
 
+uint8 FStaticMeshSceneProxy::GetCurrentFirstLODIdx_Internal() const
+{
+	return RenderData->CurrentFirstLODIdx;
+}
+
 void FStaticMeshSceneProxy::GetMeshDescription(int32 LODIndex, TArray<FMeshBatch>& OutMeshElements) const
 {
 	const FStaticMeshLODResources& LODModel = RenderData->LODResources[LODIndex];
diff --git a/Engine/Source/Runtime/Engine/Private/StreamableRenderAsset.cpp b/Engine/Source/Runtime/Engine/Private/StreamableRenderAsset.cpp
--- a/Engine/Source/Runtime/Engine/Private/StreamableRenderAsset.cpp
+++ b/Engine/Source/Runtime/Engine/Private/StreamableRenderAsset.cpp
@@ -2,7 +2,9 @@
 
 #include "Engine/StreamableRenderAsset.h"
 #include "Misc/App.h"
+#include "Containers/Ticker.h"
 #include "ContentStreaming.h"
+#include "RenderAssetUpdate.h"
 #include "Rendering/NaniteCoarseMeshStreamingManager.h"
 
 #include UE_INLINE_GENERATED_CPP_BY_NAME(StreamableRenderAsset)
@@ -27,6 +29,8 @@ UStreamableRenderAsset::UStreamableRenderAsset(const FObjectInitializer& ObjectI
 	NoRefStreamingLODBias.Init(GNoRefBiasQualityLevelCVarName, GNoRefBiasQualityLevelScalabilitySection);
 }
 
+UStreamableRenderAsset::~UStreamableRenderAsset() = default;
+
 void UStreamableRenderAsset::RegisterMipLevelChangeCallback(UPrimitiveComponent* Component, int32 LODIndex, float TimeoutSecs, bool bOnStreamIn, FLODStreamingCallback&& Callback)
 {
 	check(IsInGameThread());
diff --git a/Engine/Source/Runtime/Engine/Private/Texture.cpp b/Engine/Source/Runtime/Engine/Private/Texture.cpp
--- a/Engine/Source/Runtime/Engine/Private/Texture.cpp
+++ b/Engine/Source/Runtime/Engine/Private/Texture.cpp
@@ -48,6 +48,8 @@
 #include "Misc/ScopeRWLock.h"
 #endif
 
+#include UE_INLINE_GENERATED_CPP_BY_NAME(Texture)
+
 #if WITH_EDITORONLY_DATA
 	#include "EditorFramework/AssetImportData.h"
 #endif
@@ -109,6 +111,7 @@ UTexture::UTexture(const FObjectInitializer& ObjectInitializer)
 		[this]()-> FTextureResource* { return GetResource(); },
 		[this](FTextureResource* InTextureResource) { SetResource(InTextureResource); })
 #endif
+	, TextureReference(*new FTextureReference())
 {
 	SRGB = true;
 	Filter = TF_Default;
@@ -190,6 +193,18 @@ FTextureResource* UTexture::GetResource()
 	return nullptr;
 }
 
+UTexture::UTexture(FVTableHelper& Helper)
+	: TextureReference(*new FTextureReference())
+{
+}
+
+PRAGMA_DISABLE_DEPRECATION_WARNINGS
+UTexture::~UTexture()
+{
+	delete& TextureReference;
+}
+PRAGMA_ENABLE_DEPRECATION_WARNINGS
+
 void UTexture::SetResource(FTextureResource* InResource)
 {
 	check (!IsInActualRenderingThread() && !IsInRHIThread());
diff --git a/Engine/Source/Runtime/Engine/Private/UnrealClient.cpp b/Engine/Source/Runtime/Engine/Private/UnrealClient.cpp
--- a/Engine/Source/Runtime/Engine/Private/UnrealClient.cpp
+++ b/Engine/Source/Runtime/Engine/Private/UnrealClient.cpp
@@ -5,6 +5,7 @@
 #include "Misc/FileHelper.h"
 #include "Misc/Paths.h"
 #include "Misc/App.h"
+#include "DummyViewport.h"
 #include "EngineStats.h"
 #include "EngineGlobals.h"
 #include "RenderingThread.h"
@@ -2384,6 +2385,42 @@ void FCommonViewportClient::DrawHighResScreenshotCaptureRegion(FCanvas& Canvas)
 	LineItem.Draw( &Canvas, FVector2D(Config.UnscaledCaptureRegion.Min.X, Config.UnscaledCaptureRegion.Max.Y), FVector2D(Config.UnscaledCaptureRegion.Min.X, Config.UnscaledCaptureRegion.Min.Y));
 }
 
+void FViewportClient::RedrawRequested(FViewport* Viewport)
+{
+	Viewport->Draw();
+}
+
+void FViewportClient::RequestInvalidateHitProxy(FViewport* Viewport)
+{
+	Viewport->InvalidateHitProxy();
+}
+
+bool FViewportClient::InputKey(const FInputKeyEventArgs& EventArgs)
+{
+	PRAGMA_DISABLE_DEPRECATION_WARNINGS
+	return InputKey(EventArgs.Viewport, EventArgs.ControllerId, EventArgs.Key, EventArgs.Event, EventArgs.AmountDepressed, EventArgs.Key.IsGamepadKey());
+	PRAGMA_ENABLE_DEPRECATION_WARNINGS
+}
+
+TOptional<TSharedRef<SWidget>> FViewportClient::MapCursor(FViewport* Viewport, const FCursorReply& CursorReply)
+{
+	return TOptional<TSharedRef<SWidget>>();
+}
+
+FPopupMethodReply FViewportClient::OnQueryPopupMethod() const
+{
+	return FPopupMethodReply::Unhandled();
+}
+
+FCommonViewportClient::~FCommonViewportClient()
+{
+	//make to clean up the global "stat" client when we delete the active one.
+	if (GStatProcessingViewportClient == this)
+	{
+		GStatProcessingViewportClient = NULL;
+	}
+}
+
 void FCommonViewportClient::RequestUpdateDPIScale()
 {
 	bShouldUpdateDPIScale = true;
diff --git a/Engine/Source/Runtime/Engine/Private/VisualLogger/VisualLogger.cpp b/Engine/Source/Runtime/Engine/Private/VisualLogger/VisualLogger.cpp
--- a/Engine/Source/Runtime/Engine/Private/VisualLogger/VisualLogger.cpp
+++ b/Engine/Source/Runtime/Engine/Private/VisualLogger/VisualLogger.cpp
@@ -2,6 +2,7 @@
 
 #include "VisualLogger/VisualLogger.h"
 #include "Misc/CoreMisc.h"
+#include "EngineStats.h"
 #include "GameFramework/Actor.h"
 #include "Modules/ModuleInterface.h"
 #include "Misc/CommandLine.h"
@@ -22,8 +23,230 @@
 
 DEFINE_LOG_CATEGORY(LogVisual);
 #if ENABLE_VISUAL_LOG
+
 DEFINE_STAT(STAT_VisualLog);
 
+// Unfortunately needs to be a #define since it uses GET_VARARGS_RESULT which uses the va_list stuff which operates on the
+// current function, so we can't easily call a function
+#define COLLAPSED_LOGF(SerializeFunc) \
+	SCOPE_CYCLE_COUNTER(STAT_VisualLog); \
+	UWorld *World = nullptr; \
+	FVisualLogEntry *CurrentEntry = nullptr; \
+	if (CheckVisualLogInputInternal(Object, CategoryName, Verbosity, &World, &CurrentEntry) == false) \
+	{ \
+		return; \
+	} \
+	int32	BufferSize	= 1024; \
+	TCHAR*	Buffer		= nullptr; \
+	int32	Result		= -1; \
+	/* allocate some stack space to use on the first pass, which matches most strings */ \
+	TCHAR	StackBuffer[512]; \
+	TCHAR*	AllocatedBuffer = nullptr; \
+	\
+	/* first, try using the stack buffer */ \
+	Buffer = StackBuffer; \
+	GET_VARARGS_RESULT( Buffer, UE_ARRAY_COUNT(StackBuffer), UE_ARRAY_COUNT(StackBuffer) - 1, Fmt, Fmt, Result ); \
+	\
+	/* if that fails, then use heap allocation to make enough space */ \
+			while(Result == -1) \
+						{ \
+		FMemory::SystemFree(AllocatedBuffer); \
+		/* We need to use malloc here directly as GMalloc might not be safe. */ \
+		Buffer = AllocatedBuffer = (TCHAR*) FMemory::SystemMalloc( BufferSize * sizeof(TCHAR) ); \
+		GET_VARARGS_RESULT( Buffer, BufferSize, BufferSize-1, Fmt, Fmt, Result ); \
+		BufferSize *= 2; \
+						}; \
+	Buffer[Result] = 0; \
+	; \
+	\
+	SerializeFunc; \
+	FMemory::SystemFree(AllocatedBuffer);
+
+void FVisualLogger::CategorizedLogfImpl(const UObject* Object, const FLogCategoryBase& Category, ELogVerbosity::Type Verbosity, const TCHAR* Fmt, ...)
+{
+	const FName CategoryName = Category.GetCategoryName(); 
+	COLLAPSED_LOGF(
+		CurrentEntry->AddText(Buffer, CategoryName, Verbosity);
+	);
+}
+void FVisualLogger::CategorizedLogfImpl(const UObject* Object, const FName& CategoryName, ELogVerbosity::Type Verbosity, const TCHAR* Fmt, ...)
+{
+	COLLAPSED_LOGF(
+		CurrentEntry->AddText(Buffer, CategoryName, Verbosity);
+	);
+}
+
+void FVisualLogger::GeometryShapeLogfImpl(const UObject* Object, const FLogCategoryBase& Category, ELogVerbosity::Type Verbosity, const FVector& Start, const FVector& End, const FColor& Color, const uint16 Thickness, const TCHAR* Fmt, ...)
+{
+	const FName CategoryName = Category.GetCategoryName(); 
+	COLLAPSED_LOGF(
+		CurrentEntry->AddElement(Start, End, CategoryName, Verbosity, Color, Buffer, Thickness);
+	);
+}
+void FVisualLogger::GeometryShapeLogfImpl(const UObject* Object, const FName& CategoryName, ELogVerbosity::Type Verbosity, const FVector& Start, const FVector& End, const FColor& Color, const uint16 Thickness, const TCHAR* Fmt, ...)
+{
+	COLLAPSED_LOGF(
+		CurrentEntry->AddElement(Start, End, CategoryName, Verbosity, Color, Buffer, Thickness);
+	);
+}
+
+void FVisualLogger::ArrowLogfImpl(const UObject* Object, const FLogCategoryBase& Category, ELogVerbosity::Type Verbosity, const FVector& Start, const FVector& End, const FColor& Color, const TCHAR* Fmt, ...)
+{
+	const FName CategoryName = Category.GetCategoryName();
+	COLLAPSED_LOGF(
+		CurrentEntry->AddArrow(Start, End, CategoryName, Verbosity, Color, Buffer);
+	);
+}
+void FVisualLogger::ArrowLogfImpl(const UObject* Object, const FName& CategoryName, ELogVerbosity::Type Verbosity, const FVector& Start, const FVector& End, const FColor& Color, const TCHAR* Fmt, ...)
+{
+	COLLAPSED_LOGF(
+		CurrentEntry->AddArrow(Start, End, CategoryName, Verbosity, Color, Buffer);
+	);
+}
+
+void FVisualLogger::CircleLogfImpl(const UObject* Object, const FLogCategoryBase& Category, ELogVerbosity::Type Verbosity, const FVector& Center, const FVector& UpAxis, const float Radius, const FColor& Color, const uint16 Thickness, const TCHAR* Fmt, ...)
+{
+	const FName CategoryName = Category.GetCategoryName();
+	COLLAPSED_LOGF(
+		CurrentEntry->AddCircle(Center, UpAxis, Radius, CategoryName, Verbosity, Color, Buffer, Thickness);
+	);
+}
+void FVisualLogger::CircleLogfImpl(const UObject* Object, const FName& CategoryName, ELogVerbosity::Type Verbosity, const FVector& Center, const FVector& UpAxis, const float Radius, const FColor& Color, const uint16 Thickness, const TCHAR* Fmt, ...)
+{
+	COLLAPSED_LOGF(
+		CurrentEntry->AddCircle(Center, UpAxis, Radius, CategoryName, Verbosity, Color, Buffer, Thickness);
+	);
+}
+
+void FVisualLogger::GeometryShapeLogfImpl(const UObject* Object, const FLogCategoryBase& Category, ELogVerbosity::Type Verbosity, const FVector& Location, float Radius, const FColor& Color, const TCHAR* Fmt, ...)
+{
+	const FName CategoryName = Category.GetCategoryName(); 
+	COLLAPSED_LOGF(
+		CurrentEntry->AddElement(Location, CategoryName, Verbosity, Color, Buffer, (uint16)Radius);
+	);
+}
+void FVisualLogger::GeometryShapeLogfImpl(const UObject* Object, const FName& CategoryName, ELogVerbosity::Type Verbosity, const FVector& Location, float Radius, const FColor& Color, const TCHAR* Fmt, ...)
+{
+	COLLAPSED_LOGF(
+		CurrentEntry->AddElement(Location, CategoryName, Verbosity, Color, Buffer, (uint16)Radius);
+	);
+}
+
+void FVisualLogger::GeometryBoxLogfImpl(const UObject* Object, const FLogCategoryBase& Category, ELogVerbosity::Type Verbosity, const FBox& Box, const FMatrix& Matrix, const FColor& Color, const TCHAR* Fmt, ...)
+{
+	const FName CategoryName = Category.GetCategoryName(); 
+	COLLAPSED_LOGF(
+		CurrentEntry->AddElement(Box, Matrix, CategoryName, Verbosity, Color, Buffer);
+	);
+}
+void FVisualLogger::GeometryBoxLogfImpl(const UObject* Object, const FName& CategoryName, ELogVerbosity::Type Verbosity, const FBox& Box, const FMatrix& Matrix, const FColor& Color, const TCHAR* Fmt, ...)
+{
+	COLLAPSED_LOGF(
+		CurrentEntry->AddElement(Box, Matrix, CategoryName, Verbosity, Color, Buffer);
+	);
+}
+
+void FVisualLogger::GeometryShapeLogfImpl(const UObject* Object, const FLogCategoryBase& Category, ELogVerbosity::Type Verbosity, const FVector& Origin, const FVector& Direction, const float Length, const float Angle, const FColor& Color, const TCHAR* Fmt, ...)
+{
+	const FName CategoryName = Category.GetCategoryName(); 
+	COLLAPSED_LOGF(
+		CurrentEntry->AddElement(Origin, Direction, Length, Angle, Angle, CategoryName, Verbosity, Color, Buffer);
+	);
+}
+void FVisualLogger::GeometryShapeLogfImpl(const UObject* Object, const FName& CategoryName, ELogVerbosity::Type Verbosity, const FVector& Origin, const FVector& Direction, const float Length, const float Angle, const FColor& Color, const TCHAR* Fmt, ...)
+{
+	COLLAPSED_LOGF(
+		CurrentEntry->AddElement(Origin, Direction, Length, Angle, Angle, CategoryName, Verbosity, Color, Buffer);
+	);
+}
+
+void FVisualLogger::GeometryShapeLogfImpl(const UObject* Object, const FLogCategoryBase& Category, ELogVerbosity::Type Verbosity, const FVector& Start, const FVector& End, const float Radius, const FColor& Color, const TCHAR* Fmt, ...)
+{
+	const FName CategoryName = Category.GetCategoryName(); 
+	COLLAPSED_LOGF(
+		CurrentEntry->AddElement(Start, End, Radius, CategoryName, Verbosity, Color, Buffer);
+	);
+}
+void FVisualLogger::GeometryShapeLogfImpl(const UObject* Object, const FName& CategoryName, ELogVerbosity::Type Verbosity, const FVector& Start, const FVector& End, const float Radius, const FColor& Color, const TCHAR* Fmt, ...)
+{
+	COLLAPSED_LOGF(
+		CurrentEntry->AddElement(Start, End, Radius, CategoryName, Verbosity, Color, Buffer);
+	);
+}
+
+void FVisualLogger::GeometryShapeLogfImpl(const UObject* Object, const FLogCategoryBase& Category, ELogVerbosity::Type Verbosity, const FVector& Base, float HalfHeight, float Radius, const FQuat& Rotation, const FColor& Color, const TCHAR* Fmt, ...)
+{
+	const FName CategoryName = Category.GetCategoryName(); 
+	COLLAPSED_LOGF(
+		CurrentEntry->AddElement(Base, HalfHeight, Radius, Rotation, CategoryName, Verbosity, Color, Buffer);
+	);
+}
+void FVisualLogger::GeometryShapeLogfImpl(const UObject* Object, const FName& CategoryName, ELogVerbosity::Type Verbosity, const FVector& Base, float HalfHeight, float Radius, const FQuat& Rotation, const FColor& Color, const TCHAR* Fmt, ...)
+{
+	COLLAPSED_LOGF(
+		CurrentEntry->AddElement(Base, HalfHeight, Radius, Rotation, CategoryName, Verbosity, Color, Buffer);
+	);
+}
+
+void FVisualLogger::NavAreaShapeLogfImpl(const UObject* Object, const FLogCategoryBase& Category, ELogVerbosity::Type Verbosity, const TArray<FVector>& ConvexPoints, float MinZ, float MaxZ, const FColor& Color, const TCHAR* Fmt, ...)
+{
+	const FName CategoryName = Category.GetCategoryName(); 
+	COLLAPSED_LOGF(
+		CurrentEntry->AddElement(ConvexPoints, MinZ, MaxZ, CategoryName, Verbosity, Color, Buffer);
+	);
+}
+void FVisualLogger::NavAreaShapeLogfImpl(const UObject* Object, const FName& CategoryName, ELogVerbosity::Type Verbosity, const TArray<FVector>& ConvexPoints, float MinZ, float MaxZ, const FColor& Color, const TCHAR* Fmt, ...)
+{
+	COLLAPSED_LOGF(
+		CurrentEntry->AddElement(ConvexPoints, MinZ, MaxZ, CategoryName, Verbosity, Color, Buffer);
+	);
+}
+
+void FVisualLogger::GeometryShapeLogfImpl(const UObject* Object, const FLogCategoryBase& Category, ELogVerbosity::Type Verbosity, const TArray<FVector>& Vertices, const TArray<int32>& Indices, const FColor& Color, const TCHAR* Fmt, ...)
+{
+	const FName CategoryName = Category.GetCategoryName(); 
+	COLLAPSED_LOGF(
+		CurrentEntry->AddElement(Vertices, Indices, CategoryName, Verbosity, Color, Buffer);
+	);
+}
+void FVisualLogger::GeometryShapeLogfImpl(const UObject* Object, const FName& CategoryName, ELogVerbosity::Type Verbosity, const TArray<FVector>& Vertices, const TArray<int32>& Indices, const FColor& Color, const TCHAR* Fmt, ...)
+{
+	COLLAPSED_LOGF(
+		CurrentEntry->AddElement(Vertices, Indices, CategoryName, Verbosity, Color, Buffer);
+	);
+}
+
+void FVisualLogger::GeometryConvexLogfImpl(const UObject* Object, const FLogCategoryBase& Category, ELogVerbosity::Type Verbosity, const TArray<FVector>& Points, const FColor& Color, const TCHAR* Fmt, ...)
+{
+	const FName CategoryName = Category.GetCategoryName(); 
+	COLLAPSED_LOGF(
+		CurrentEntry->AddConvexElement(Points, CategoryName, Verbosity, Color, Buffer);
+	);
+}
+void FVisualLogger::GeometryConvexLogfImpl(const UObject* Object, const FName& CategoryName, ELogVerbosity::Type Verbosity, const TArray<FVector>& Points, const FColor& Color, const TCHAR* Fmt, ...)
+{
+	COLLAPSED_LOGF(
+		CurrentEntry->AddConvexElement(Points, CategoryName, Verbosity, Color, Buffer);
+	);
+}
+
+void FVisualLogger::HistogramDataLogfImpl(const UObject* Object, const FLogCategoryBase& Category, ELogVerbosity::Type Verbosity, FName GraphName, FName DataName, const FVector2D& Data, const FColor& Color, const TCHAR* Fmt, ...)
+{
+	const FName CategoryName = Category.GetCategoryName(); 
+	COLLAPSED_LOGF(
+		CurrentEntry->AddHistogramData(Data, CategoryName, Verbosity, GraphName, DataName);
+	);
+}
+void FVisualLogger::HistogramDataLogfImpl(const UObject* Object, const FName& CategoryName, ELogVerbosity::Type Verbosity, FName GraphName, FName DataName, const FVector2D& Data, const FColor& Color, const TCHAR* Fmt, ...)
+{
+	COLLAPSED_LOGF(
+		CurrentEntry->AddHistogramData(Data, CategoryName, Verbosity, GraphName, DataName);
+	);
+}
+
+#undef COLLAPSED_LOGF
+
+
 namespace
 {
 	UWorld* GetWorldForVisualLogger(const UObject* Object)
diff --git a/Engine/Source/Runtime/Engine/Public/Animation/AnimCurveTypes.h b/Engine/Source/Runtime/Engine/Public/Animation/AnimCurveTypes.h
--- a/Engine/Source/Runtime/Engine/Public/Animation/AnimCurveTypes.h
+++ b/Engine/Source/Runtime/Engine/Public/Animation/AnimCurveTypes.h
@@ -11,6 +11,18 @@
 #include "Misc/EnumRange.h"
 #include "AnimCurveTypes.generated.h"
 
+UENUM()
+enum class EAnimCurveType : uint8 
+{
+	AttributeCurve,
+	MaterialCurve, 
+	MorphTargetCurve, 
+	// make sure to update MaxCurve 
+	MaxAnimCurveType UMETA(Hidden)
+};
+
+ENUM_RANGE_BY_COUNT(EAnimCurveType, EAnimCurveType::MaxAnimCurveType);
+
 /** This is curve flags that are saved in asset and **/
 UENUM(BlueprintType, meta=(Bitflags))
 enum EAnimAssetCurveFlags
diff --git a/Engine/Source/Runtime/Engine/Public/Audio/AudioDebug.h b/Engine/Source/Runtime/Engine/Public/Audio/AudioDebug.h
--- a/Engine/Source/Runtime/Engine/Public/Audio/AudioDebug.h
+++ b/Engine/Source/Runtime/Engine/Public/Audio/AudioDebug.h
@@ -1,7 +1,9 @@
 // Copyright Epic Games, Inc. All Rights Reserved.
 #pragma once
 
+#if UE_ENABLE_INCLUDE_ORDER_DEPRECATED_IN_5_2
 #include "CanvasTypes.h"
+#endif
 #include "CoreMinimal.h"
 
 #include "AudioDefines.h"		// For ENABLE_AUDIO_DEBUG
@@ -16,6 +18,7 @@ struct FListener;
 struct FWaveInstance;
 
 class FSoundSource;
+class FViewportClient;
 class USoundWave;
 class UWorld;
 
diff --git a/Engine/Source/Runtime/Engine/Public/CanvasBatchedElementRenderItem.h b/Engine/Source/Runtime/Engine/Public/CanvasBatchedElementRenderItem.h
--- a/Engine/Source/Runtime/Engine/Public/CanvasBatchedElementRenderItem.h
+++ b/Engine/Source/Runtime/Engine/Public/CanvasBatchedElementRenderItem.h
@@ -1,5 +1,143 @@
 // Copyright Epic Games, Inc. All Rights Reserved.
 
+/*=============================================================================
+	Canvas.h: Unreal canvas definition.
+=============================================================================*/
+
 #pragma once
 
-// Placeholder for upcoming change
\ No newline at end of file
+#include "CanvasTypes.h"
+#include "BatchedElements.h"
+
+/**
+* Info needed to render a batched element set
+*/
+class FCanvasBatchedElementRenderItem : public FCanvasBaseRenderItem
+{
+public:
+	/** 
+	* Init constructor 
+	*/
+	FCanvasBatchedElementRenderItem(
+		FBatchedElementParameters* InBatchedElementParameters=NULL,
+		const FTexture* InTexture=NULL,
+		ESimpleElementBlendMode InBlendMode=SE_BLEND_MAX,
+		FCanvas::EElementType InElementType=FCanvas::ET_MAX,
+		const FCanvas::FTransformEntry& InTransform=FCanvas::FTransformEntry(FMatrix::Identity),
+		const FDepthFieldGlowInfo& InGlowInfo=FDepthFieldGlowInfo() )
+		// this data is deleted after rendering has completed
+		: Data(new FRenderData(InBatchedElementParameters, InTexture, InBlendMode, InElementType, InTransform, InGlowInfo))
+	{}
+
+	/**
+	* Destructor to delete data in case nothing rendered
+	*/
+	virtual ~FCanvasBatchedElementRenderItem()
+	{
+		delete Data;
+	}
+
+	/**
+	* FCanvasBatchedElementRenderItem instance accessor
+	*
+	* @return this instance
+	*/
+	virtual class FCanvasBatchedElementRenderItem* GetCanvasBatchedElementRenderItem() override
+	{ 
+		return this; 
+	}
+
+	/**
+	* Renders the canvas item. 
+	* Iterates over all batched elements and draws them with their own transforms
+	*
+	* @param Canvas - canvas currently being rendered
+	* @param RHICmdList - command list to use
+	* @return true if anything rendered
+	*/
+	virtual bool Render_RenderThread(FCanvasRenderContext& RenderContext, FMeshPassProcessorRenderState& DrawRenderState, const FCanvas* Canvas) override;
+	
+	/**
+	* Renders the canvas item.
+	* Iterates over all batched elements and draws them with their own transforms
+	*
+	* @param Canvas - canvas currently being rendered
+	* @return true if anything rendered
+	*/
+	virtual bool Render_GameThread(const FCanvas* Canvas, FCanvasRenderThreadScope& RenderScope) override;
+
+	/**
+	* Determine if this is a matching set by comparing texture,blendmode,elementype,transform. All must match
+	*
+	* @param BatchedElementParameters - parameters for this batched element
+	* @param InTexture - texture resource for the item being rendered
+	* @param InBlendMode - current alpha blend mode 
+	* @param InElementType - type of item being rendered: triangle,line,etc
+	* @param InTransform - the transform for the item being rendered
+	* @param InGlowInfo - the depth field glow of the item being rendered
+	* @return true if the parameters match this render item
+	*/
+	bool IsMatch(FBatchedElementParameters* BatchedElementParameters, const FTexture* InTexture, ESimpleElementBlendMode InBlendMode, FCanvas::EElementType InElementType, const FCanvas::FTransformEntry& InTransform, const FDepthFieldGlowInfo& InGlowInfo)
+	{
+		return(	Data->BatchedElementParameters.GetReference() == BatchedElementParameters &&
+				Data->Texture == InTexture &&
+				Data->BlendMode == InBlendMode &&
+				Data->ElementType == InElementType &&
+				Data->Transform.GetMatrixCRC() == InTransform.GetMatrixCRC() &&
+				Data->GlowInfo == InGlowInfo );
+	}
+
+	/**
+	* Accessor for the batched elements. This can be used for adding triangles and primitives to the batched elements
+	*
+	* @return pointer to batched elements struct
+	*/
+	FORCEINLINE FBatchedElements* GetBatchedElements()
+	{
+		return &Data->BatchedElements;
+	}
+
+private:
+	class FRenderData
+	{
+	public:
+		/**
+		* Init constructor
+		*/
+		FRenderData(
+			FBatchedElementParameters* InBatchedElementParameters=NULL,
+			const FTexture* InTexture=NULL,
+			ESimpleElementBlendMode InBlendMode=SE_BLEND_MAX,
+			FCanvas::EElementType InElementType=FCanvas::ET_MAX,
+			const FCanvas::FTransformEntry& InTransform=FCanvas::FTransformEntry(FMatrix::Identity),
+			const FDepthFieldGlowInfo& InGlowInfo=FDepthFieldGlowInfo() )
+			:	BatchedElementParameters(InBatchedElementParameters)
+			,	Texture(InTexture)
+			,	BlendMode(InBlendMode)
+			,	ElementType(InElementType)
+			,	Transform(InTransform)
+			,	GlowInfo(InGlowInfo)
+		{}
+		/** Current batched elements, destroyed once rendering completes. */
+		FBatchedElements BatchedElements;
+		/** Batched element parameters */
+		TRefCountPtr<FBatchedElementParameters> BatchedElementParameters;
+		/** Current texture being used for batching, set to NULL if it hasn't been used yet. */
+		const FTexture* Texture;
+		/** Current blend mode being used for batching, set to BLEND_MAX if it hasn't been used yet. */
+		ESimpleElementBlendMode BlendMode;
+		/** Current element type being used for batching, set to ET_MAX if it hasn't been used yet. */
+		FCanvas::EElementType ElementType;
+		/** Transform used to render including projection */
+		FCanvas::FTransformEntry Transform;
+		/** info for optional glow effect when using depth field rendering */
+		FDepthFieldGlowInfo GlowInfo;
+	};
+	
+	/**
+	* Render data which is allocated when a new FCanvasBatchedElementRenderItem is added for rendering.
+	* This data is only freed on the rendering thread once the item has finished rendering
+	*/
+	FRenderData* Data;		
+};
+
diff --git a/Engine/Source/Runtime/Engine/Public/CanvasTypes.h b/Engine/Source/Runtime/Engine/Public/CanvasTypes.h
--- a/Engine/Source/Runtime/Engine/Public/CanvasTypes.h
+++ b/Engine/Source/Runtime/Engine/Public/CanvasTypes.h
@@ -10,12 +10,16 @@
 #include "UObject/ObjectMacros.h"
 #include "Engine/EngineTypes.h"
 #include "HitProxies.h"
+#if UE_ENABLE_INCLUDE_ORDER_DEPRECATED_IN_5_2
 #include "BatchedElements.h"
 #include "RendererInterface.h"
+#endif
 #include "StaticMeshResources.h"
 #include "UnrealEngine.h"
 #include "CanvasTypes.generated.h"
 
+class FBatchedElements;
+class FCanvasBatchedElementRenderItem;
 class FCanvasRenderContext;
 class FCanvasRenderThreadScope;
 class FCanvasItem;
@@ -837,139 +841,6 @@ public:
 };
 
 
-/**
-* Info needed to render a batched element set
-*/
-class FCanvasBatchedElementRenderItem : public FCanvasBaseRenderItem
-{
-public:
-	/** 
-	* Init constructor 
-	*/
-	FCanvasBatchedElementRenderItem(
-		FBatchedElementParameters* InBatchedElementParameters=NULL,
-		const FTexture* InTexture=NULL,
-		ESimpleElementBlendMode InBlendMode=SE_BLEND_MAX,
-		FCanvas::EElementType InElementType=FCanvas::ET_MAX,
-		const FCanvas::FTransformEntry& InTransform=FCanvas::FTransformEntry(FMatrix::Identity),
-		const FDepthFieldGlowInfo& InGlowInfo=FDepthFieldGlowInfo() )
-		// this data is deleted after rendering has completed
-		: Data(new FRenderData(InBatchedElementParameters, InTexture, InBlendMode, InElementType, InTransform, InGlowInfo))
-	{}
-
-	/**
-	* Destructor to delete data in case nothing rendered
-	*/
-	virtual ~FCanvasBatchedElementRenderItem()
-	{
-		delete Data;
-	}
-
-	/**
-	* FCanvasBatchedElementRenderItem instance accessor
-	*
-	* @return this instance
-	*/
-	virtual class FCanvasBatchedElementRenderItem* GetCanvasBatchedElementRenderItem() override
-	{ 
-		return this; 
-	}
-
-	/**
-	* Renders the canvas item. 
-	* Iterates over all batched elements and draws them with their own transforms
-	*
-	* @param Canvas - canvas currently being rendered
-	* @param RHICmdList - command list to use
-	* @return true if anything rendered
-	*/
-	virtual bool Render_RenderThread(FCanvasRenderContext& RenderContext, FMeshPassProcessorRenderState& DrawRenderState, const FCanvas* Canvas) override;
-	
-	/**
-	* Renders the canvas item.
-	* Iterates over all batched elements and draws them with their own transforms
-	*
-	* @param Canvas - canvas currently being rendered
-	* @return true if anything rendered
-	*/
-	virtual bool Render_GameThread(const FCanvas* Canvas, FCanvasRenderThreadScope& RenderScope) override;
-
-	/**
-	* Determine if this is a matching set by comparing texture,blendmode,elementype,transform. All must match
-	*
-	* @param BatchedElementParameters - parameters for this batched element
-	* @param InTexture - texture resource for the item being rendered
-	* @param InBlendMode - current alpha blend mode 
-	* @param InElementType - type of item being rendered: triangle,line,etc
-	* @param InTransform - the transform for the item being rendered
-	* @param InGlowInfo - the depth field glow of the item being rendered
-	* @return true if the parameters match this render item
-	*/
-	bool IsMatch(FBatchedElementParameters* BatchedElementParameters, const FTexture* InTexture, ESimpleElementBlendMode InBlendMode, FCanvas::EElementType InElementType, const FCanvas::FTransformEntry& InTransform, const FDepthFieldGlowInfo& InGlowInfo)
-	{
-		return(	Data->BatchedElementParameters.GetReference() == BatchedElementParameters &&
-				Data->Texture == InTexture &&
-				Data->BlendMode == InBlendMode &&
-				Data->ElementType == InElementType &&
-				Data->Transform.GetMatrixCRC() == InTransform.GetMatrixCRC() &&
-				Data->GlowInfo == InGlowInfo );
-	}
-
-	/**
-	* Accessor for the batched elements. This can be used for adding triangles and primitives to the batched elements
-	*
-	* @return pointer to batched elements struct
-	*/
-	FORCEINLINE FBatchedElements* GetBatchedElements()
-	{
-		return &Data->BatchedElements;
-	}
-
-private:
-	class FRenderData
-	{
-	public:
-		/**
-		* Init constructor
-		*/
-		FRenderData(
-			FBatchedElementParameters* InBatchedElementParameters=NULL,
-			const FTexture* InTexture=NULL,
-			ESimpleElementBlendMode InBlendMode=SE_BLEND_MAX,
-			FCanvas::EElementType InElementType=FCanvas::ET_MAX,
-			const FCanvas::FTransformEntry& InTransform=FCanvas::FTransformEntry(FMatrix::Identity),
-			const FDepthFieldGlowInfo& InGlowInfo=FDepthFieldGlowInfo() )
-			:	BatchedElementParameters(InBatchedElementParameters)
-			,	Texture(InTexture)
-			,	BlendMode(InBlendMode)
-			,	ElementType(InElementType)
-			,	Transform(InTransform)
-			,	GlowInfo(InGlowInfo)
-		{}
-		/** Current batched elements, destroyed once rendering completes. */
-		FBatchedElements BatchedElements;
-		/** Batched element parameters */
-		TRefCountPtr<FBatchedElementParameters> BatchedElementParameters;
-		/** Current texture being used for batching, set to NULL if it hasn't been used yet. */
-		const FTexture* Texture;
-		/** Current blend mode being used for batching, set to BLEND_MAX if it hasn't been used yet. */
-		ESimpleElementBlendMode BlendMode;
-		/** Current element type being used for batching, set to ET_MAX if it hasn't been used yet. */
-		FCanvas::EElementType ElementType;
-		/** Transform used to render including projection */
-		FCanvas::FTransformEntry Transform;
-		/** info for optional glow effect when using depth field rendering */
-		FDepthFieldGlowInfo GlowInfo;
-	};
-	
-	/**
-	* Render data which is allocated when a new FCanvasBatchedElementRenderItem is added for rendering.
-	* This data is only freed on the rendering thread once the item has finished rendering
-	*/
-	FRenderData* Data;		
-};
-
-
 /**
 * Info needed to render a single FTileRenderer
 */
diff --git a/Engine/Source/Runtime/Engine/Public/ClothCollisionSource.h b/Engine/Source/Runtime/Engine/Public/ClothCollisionSource.h
--- a/Engine/Source/Runtime/Engine/Public/ClothCollisionSource.h
+++ b/Engine/Source/Runtime/Engine/Public/ClothCollisionSource.h
@@ -2,4 +2,33 @@
 
 #pragma once
 
-// Placeholder for upcoming change
\ No newline at end of file
+#include "ClothCollisionPrim.h"
+#include "Components/SkeletalMeshComponent.h"
+
+/** Helper struct used to store info about a cloth collision source */
+struct FClothCollisionSource
+{
+	FClothCollisionSource(USkeletalMeshComponent* InSourceComponent, UPhysicsAsset* InSourcePhysicsAsset, const FOnBoneTransformsFinalizedMultiCast::FDelegate& InOnBoneTransformsFinalizedDelegate);
+	ENGINE_API ~FClothCollisionSource();
+
+	/** Component that collision data will be copied from */
+	TWeakObjectPtr<USkeletalMeshComponent> SourceComponent;
+
+	/** Physics asset to use to generate collision against the source component */
+	TWeakObjectPtr<UPhysicsAsset> SourcePhysicsAsset;
+
+	/** Callback used to remove the cloth transform updates delegate */
+	FDelegateHandle OnBoneTransformsFinalizedHandle;
+
+	/** Cached skeletal mesh used to invalidate the cache if the skeletal mesh has changed */
+	TWeakObjectPtr<USkeletalMesh> CachedSkeletalMesh;
+
+	/** Cached spheres from physics asset */
+	TArray<FClothCollisionPrim_Sphere> CachedSpheres;
+
+	/** Cached sphere connections from physics asset */
+	TArray<FClothCollisionPrim_SphereConnection> CachedSphereConnections;
+
+	/** Flag whether the cache is valid */
+	bool bCached;
+};
diff --git a/Engine/Source/Runtime/Engine/Public/Components.h b/Engine/Source/Runtime/Engine/Public/Components.h
--- a/Engine/Source/Runtime/Engine/Public/Components.h
+++ b/Engine/Source/Runtime/Engine/Public/Components.h
@@ -4,16 +4,16 @@
 
 #include "CoreMinimal.h"
 #include "UObject/ObjectMacros.h"
+#include "MeshUVChannelInfo.h"
 #include "RenderResource.h"
 #include "VertexFactory.h"
-#include "Components.generated.h"
 
 /*=============================================================================
 	Components.h: Forward declarations of object components of actors
 =============================================================================*/
 
 // Constants.
-enum { MAX_TEXCOORDS = 4, MAX_STATIC_TEXCOORDS = 8 };
+enum { MAX_STATIC_TEXCOORDS = 8 };
 
 /** The information used to build a static-mesh vertex. */
 struct FStaticMeshBuildVertex
@@ -60,56 +60,3 @@ struct FStaticMeshDataType
 	uint32 ColorIndexMask = ~0u;
 	uint32 LODLightmapDataIndex = 0;
 };
-
-/** The world size for each texcoord mapping. Used by the texture streaming. */
-USTRUCT(BlueprintType)
-struct FMeshUVChannelInfo
-{
-	GENERATED_USTRUCT_BODY()
-
-	/** Default constructor (no initialization). */
-	FORCEINLINE FMeshUVChannelInfo() { FMemory::Memzero(*this); }
-
-	/** Constructor which initializes all components to zero. */
-	FMeshUVChannelInfo(ENoInit) { }
-
-	FMeshUVChannelInfo(float DefaultDensity) : bInitialized(true), bOverrideDensities(false)
-	{
-		for (float& Density : LocalUVDensities)
-		{
-			Density = DefaultDensity;
-		}
-	}
-
-	/** Returns whether the structure contains any valid LocalUVDensities. */
-	bool IsInitialized() const
-	{
-		if (bInitialized)
-		{
-			for (float Density : LocalUVDensities)
-			{
-				if (Density != 0)
-				{
-					return true;
-				}
-			}
-		}
-		return false;
-	}
-
-	UPROPERTY()
-	bool bInitialized;
-
-	/** Whether this values was set manually or is auto generated. */
-	UPROPERTY(EditAnywhere, AdvancedDisplay, Category = Basic)
-	bool bOverrideDensities;
-
-	/**
-	 * The UV density in the mesh, before any transform scaling, in world unit per UV.
-	 * This value represents the length taken to cover a full UV unit.
-	 */
-	UPROPERTY(EditAnywhere, AdvancedDisplay, Category = Basic, meta = (EditCondition = "bOverrideDensities"))
-	float LocalUVDensities[MAX_TEXCOORDS];
-
-	friend FArchive& operator<<(FArchive& Ar, FMeshUVChannelInfo& Info);
-};
diff --git a/Engine/Source/Runtime/Engine/Public/ContentStreaming.h b/Engine/Source/Runtime/Engine/Public/ContentStreaming.h
--- a/Engine/Source/Runtime/Engine/Public/ContentStreaming.h
+++ b/Engine/Source/Runtime/Engine/Public/ContentStreaming.h
@@ -8,12 +8,17 @@
 
 #include "CoreMinimal.h"
 #include "UObject/WeakObjectPtr.h"
+#if UE_ENABLE_INCLUDE_ORDER_DEPRECATED_IN_5_2
 #include "CanvasTypes.h"
 #include "UnrealClient.h"
+#endif
+#include "Serialization/BulkData.h"
 
 class AActor;
 class FSoundSource;
 class UPrimitiveComponent;
+class FCanvas;
+class FViewport;
 class FSoundWaveData;
 class FSoundWaveProxy;
 class ICompressedAudioInfo;
@@ -21,6 +26,7 @@ class UTexture2D;
 struct FRenderAssetStreamingManager;
 struct FWaveInstance;
 class UAnimStreamable;
+enum class EStreamableRenderAssetType : uint8;
 struct FCompressedAnimSequence;
 
 using FSoundWaveProxyPtr = TSharedPtr<FSoundWaveProxy, ESPMode::ThreadSafe>;
@@ -695,12 +701,12 @@ struct FStreamingManagerCollection : public IStreamingManager
 	/**
 	 * Checks whether texture streaming is enabled. 
 	 */
-	FORCEINLINE bool IsTextureStreamingEnabled() const { return IsRenderAssetStreamingEnabled(EStreamableRenderAssetType::Texture); }
+	ENGINE_API bool IsTextureStreamingEnabled() const;
 
 	/**
 	 * Checks whether texture/mesh streaming is enabled
 	 */
-	ENGINE_API bool IsRenderAssetStreamingEnabled(EStreamableRenderAssetType FilteredAssetType = EStreamableRenderAssetType::None) const;
+	ENGINE_API bool IsRenderAssetStreamingEnabled(EStreamableRenderAssetType FilteredAssetType) const;
 
 	/**
 	 * Gets a reference to the Texture Streaming Manager interface
diff --git a/Engine/Source/Runtime/Engine/Public/DummyViewport.h b/Engine/Source/Runtime/Engine/Public/DummyViewport.h
--- a/Engine/Source/Runtime/Engine/Public/DummyViewport.h
+++ b/Engine/Source/Runtime/Engine/Public/DummyViewport.h
@@ -2,4 +2,72 @@
 
 #pragma once
 
-// Placeholder for coming changelist
\ No newline at end of file
+#include "UnrealClient.h"
+
+/**
+ * Minimal viewport for assisting with taking screenshots (also used within a plugin)
+ * @todo: This should be refactored
+ */
+class ENGINE_API FDummyViewport : public FViewport
+{
+public:
+	FDummyViewport(FViewportClient* InViewportClient);
+
+	virtual ~FDummyViewport();
+
+	//~ Begin FViewport Interface
+	virtual void BeginRenderFrame(FRHICommandListImmediate& RHICmdList) override
+	{
+		check( IsInRenderingThread() );
+		//SetRenderTarget(RHICmdList,  RenderTargetTextureRHI,  FTexture2DRHIRef() );
+	};
+
+	virtual void EndRenderFrame(FRHICommandListImmediate& RHICmdList, bool bPresent, bool bLockToVsync) override
+	{
+		check( IsInRenderingThread() );
+	}
+
+	void SetupHDR(EDisplayColorGamut InDisplayColorGamut, EDisplayOutputFormat InDisplayOutputFormat, bool bInSceneHDREnabled)
+	{
+		DisplayColorGamut = InDisplayColorGamut;
+		DisplayOutputFormat = InDisplayOutputFormat;
+		bSceneHDREnabled = bInSceneHDREnabled;
+	}
+
+	virtual void*	GetWindow() override { return 0; }
+	virtual void	MoveWindow(int32 NewPosX, int32 NewPosY, int32 NewSizeX, int32 NewSizeY) override {}
+	virtual void	Destroy() override {}
+	virtual bool SetUserFocus(bool bFocus) override { return false; }
+	virtual bool	KeyState(FKey Key) const override { return false; }
+	virtual int32	GetMouseX() const override { return 0; }
+	virtual int32	GetMouseY() const override { return 0; }
+	virtual void	GetMousePos( FIntPoint& MousePosition, const bool bLocalPosition = true) override { MousePosition = FIntPoint(0, 0); }
+	virtual void	SetMouse(int32 x, int32 y) override { }
+	virtual void	ProcessInput( float DeltaTime ) override { }
+	virtual FVector2D VirtualDesktopPixelToViewport(FIntPoint VirtualDesktopPointPx) const override { return FVector2D::ZeroVector; }
+	virtual FIntPoint ViewportToVirtualDesktopPixel(FVector2D ViewportCoordinate) const override { return FIntPoint::ZeroValue; }
+	virtual void InvalidateDisplay() override { }
+	virtual void DeferInvalidateHitProxy() override { }
+	virtual FViewportFrame* GetViewportFrame() override { return 0; }
+	virtual FCanvas* GetDebugCanvas() override { return DebugCanvas; }
+	virtual EDisplayColorGamut GetDisplayColorGamut() const;
+	virtual EDisplayOutputFormat GetDisplayOutputFormat() const;
+	virtual bool GetSceneHDREnabled() const;
+	//~ End FViewport Interface
+
+	//~ Begin FRenderResource Interface
+	virtual void InitDynamicRHI() override;
+
+	// @todo DLL: Without these functions we get unresolved linker errors with FRenderResource
+	virtual void InitRHI() override{}
+	virtual void ReleaseRHI() override{}
+	virtual void InitResource() override{ FViewport::InitResource(); }
+	virtual void ReleaseResource() override { FViewport::ReleaseResource(); }
+	virtual FString GetFriendlyName() const override { return FString(TEXT("FDummyViewport"));}
+	//~ End FRenderResource Interface
+private:
+	FCanvas* DebugCanvas;
+	EDisplayColorGamut DisplayColorGamut = EDisplayColorGamut::sRGB_D65;
+	EDisplayOutputFormat DisplayOutputFormat = EDisplayOutputFormat::SDR_sRGB;
+	bool bSceneHDREnabled = false;
+};
diff --git a/Engine/Source/Runtime/Engine/Public/DynamicMeshBuilder.h b/Engine/Source/Runtime/Engine/Public/DynamicMeshBuilder.h
--- a/Engine/Source/Runtime/Engine/Public/DynamicMeshBuilder.h
+++ b/Engine/Source/Runtime/Engine/Public/DynamicMeshBuilder.h
@@ -11,11 +11,15 @@
 #include "HitProxies.h"
 #include "RenderUtils.h"
 #include "LocalVertexFactory.h"
+#if UE_ENABLE_INCLUDE_ORDER_DEPRECATED_IN_5_2
 #include "SceneManagement.h"
+#endif
 
 class FMaterialRenderProxy;
+class FMeshBuilderOneFrameResources;
 class FMeshElementCollector;
 class FPrimitiveDrawInterface;
+struct FMeshBatch;
 
 /** The vertex type used for dynamic meshes. */
 struct FDynamicMeshVertex
@@ -101,21 +105,6 @@ struct FDynamicMeshVertex
 	FColor Color;
 };
 
-class FMeshBuilderOneFrameResources : public FOneFrameResource
-{
-public:
-	class FPooledDynamicMeshVertexBuffer* VertexBuffer = nullptr;
-	class FPooledDynamicMeshIndexBuffer* IndexBuffer = nullptr;
-	class FPooledDynamicMeshVertexFactory* VertexFactory = nullptr;
-	class FDynamicMeshPrimitiveUniformBuffer* PrimitiveUniformBuffer = nullptr;
-	virtual ENGINE_API ~FMeshBuilderOneFrameResources();
-
-	inline bool IsValidForRendering() 
-	{
-		return VertexBuffer && IndexBuffer && PrimitiveUniformBuffer && VertexFactory;
-	}
-};
-
 struct FDynamicMeshDrawOffset
 {
 	uint32 FirstIndex = 0;
diff --git a/Engine/Source/Runtime/Engine/Public/EngineUtils.h b/Engine/Source/Runtime/Engine/Public/EngineUtils.h
--- a/Engine/Source/Runtime/Engine/Public/EngineUtils.h
+++ b/Engine/Source/Runtime/Engine/Public/EngineUtils.h
@@ -12,11 +12,14 @@
 #include "Engine/EngineTypes.h"
 #include "GameFramework/Actor.h"
 #include "HitProxies.h"
+#include "Engine/Level.h"
 #include "Engine/World.h"
 #include "UObject/UObjectHash.h"
 #include "ProfilingDebugging/ProfilingHelpers.h"
 #include "GameFramework/WorldSettings.h"
+#if UE_ENABLE_INCLUDE_ORDER_DEPRECATED_IN_5_2
 #include "RendererInterface.h"
+#endif
 
 class FCanvas;
 class FViewport;
