Removed lots of includes in UE based on compile profiling

diff --git a/Engine/Source/Runtime/Engine/Public/GameTime.h b/Engine/Source/Runtime/Engine/Public/GameTime.h
--- a/Engine/Source/Runtime/Engine/Public/GameTime.h
+++ b/Engine/Source/Runtime/Engine/Public/GameTime.h
@@ -2,4 +2,83 @@
 
 #pragma once
 
-// Placeholder for upcoming change
\ No newline at end of file
+#include "Misc/CoreMiscDefines.h"
+
+/** Contains all the timings of a gaming frame, to handle pause and time dilation (for instance bullet time) of the world. */
+struct ENGINE_API FGameTime
+{
+	FORCEINLINE_DEBUGGABLE FGameTime()
+		: RealTimeSeconds(0.0)
+		, WorldTimeSeconds(0.0)
+		, DeltaRealTimeSeconds(0.0f)
+		, DeltaWorldTimeSeconds(0.0f)
+	{ }
+
+	FGameTime(const FGameTime&) = default;
+	FGameTime& operator = (const FGameTime&) = default;
+
+	// Returns the game time since GStartTime.
+	static FGameTime GetTimeSinceAppStart();
+
+	static FORCEINLINE_DEBUGGABLE FGameTime CreateUndilated(double InRealTimeSeconds, float InDeltaRealTimeSeconds)
+	{
+		return FGameTime::CreateDilated(InRealTimeSeconds, InDeltaRealTimeSeconds, InRealTimeSeconds, InDeltaRealTimeSeconds);
+	}
+
+	static FORCEINLINE_DEBUGGABLE FGameTime CreateDilated(double InRealTimeSeconds, float InDeltaRealTimeSeconds, double InWorldTimeSeconds, float InDeltaWorldTimeSeconds)
+	{
+		return FGameTime(InRealTimeSeconds, InDeltaRealTimeSeconds, InWorldTimeSeconds, InDeltaWorldTimeSeconds);
+	}
+
+	/** Returns time in seconds since level began play, but IS NOT paused when the game is paused, and IS NOT dilated/clamped. */
+	FORCEINLINE_DEBUGGABLE double GetRealTimeSeconds() const
+	{
+		return RealTimeSeconds;
+	}
+
+	/** Returns frame delta time in seconds with no adjustment for time dilation and pause. */
+	FORCEINLINE_DEBUGGABLE float GetDeltaRealTimeSeconds() const
+	{
+		return DeltaRealTimeSeconds;
+	}
+
+	/** Returns time in seconds since level began play, but IS paused when the game is paused, and IS dilated/clamped. */
+	FORCEINLINE_DEBUGGABLE double GetWorldTimeSeconds() const
+	{
+		return WorldTimeSeconds;
+	}
+
+	/** Returns frame delta time in seconds adjusted by e.g. time dilation. */
+	FORCEINLINE_DEBUGGABLE float GetDeltaWorldTimeSeconds() const
+	{
+		return DeltaWorldTimeSeconds;
+	}
+
+	/** Returns how much world time is slowed compared to real time. */
+	FORCEINLINE_DEBUGGABLE float GetTimeDilation() const
+	{
+		ensure(DeltaRealTimeSeconds > 0.0f);
+		return DeltaWorldTimeSeconds / DeltaRealTimeSeconds;
+	}
+
+	/** Returns whether the world time is paused. */
+	FORCEINLINE_DEBUGGABLE bool IsPaused() const
+	{
+		return DeltaWorldTimeSeconds == 0.0f;
+	}
+
+private:
+	double RealTimeSeconds;
+	double WorldTimeSeconds;
+
+	float DeltaRealTimeSeconds;
+	float DeltaWorldTimeSeconds;
+
+	FORCEINLINE_DEBUGGABLE FGameTime(double InRealTimeSeconds, float InDeltaRealTimeSeconds, double InWorldTimeSeconds, float InDeltaWorldTimeSeconds)
+		: RealTimeSeconds(InRealTimeSeconds)
+		, WorldTimeSeconds(InWorldTimeSeconds)
+		, DeltaRealTimeSeconds(InDeltaRealTimeSeconds)
+		, DeltaWorldTimeSeconds(InDeltaWorldTimeSeconds)
+	{ }
+
+};
diff --git a/Engine/Source/Runtime/Engine/Public/InputKeyEventArgs.h b/Engine/Source/Runtime/Engine/Public/InputKeyEventArgs.h
new file mode 100644
--- /dev/null
+++ b/Engine/Source/Runtime/Engine/Public/InputKeyEventArgs.h
@@ -0,0 +1,79 @@
+// Copyright Epic Games, Inc. All Rights Reserved.
+
+#pragma once
+
+#include "Engine/EngineBaseTypes.h"
+#include "GenericPlatform/GenericPlatformInputDeviceMapper.h"
+#include "GenericPlatform/GenericWindow.h"
+#include "InputCoreTypes.h"
+
+class FViewport;
+
+struct FInputKeyEventArgs
+{
+public:
+	FInputKeyEventArgs(FViewport* InViewport, int32 InControllerId, FKey InKey, EInputEvent InEvent)
+		: Viewport(InViewport)
+		, ControllerId(InControllerId)
+		, InputDevice(FInputDeviceId::CreateFromInternalId(InControllerId))
+		, Key(InKey)
+		, Event(InEvent)
+		, AmountDepressed(1.0f)
+		, bIsTouchEvent(false)
+	{
+	}
+
+	FInputKeyEventArgs(FViewport* InViewport, FInputDeviceId InInputDevice, FKey InKey, EInputEvent InEvent)
+		: Viewport(InViewport)
+		, InputDevice(InInputDevice)
+		, Key(InKey)
+		, Event(InEvent)
+		, AmountDepressed(1.0f)
+		, bIsTouchEvent(false)
+	{
+		FPlatformUserId UserID = IPlatformInputDeviceMapper::Get().GetUserForInputDevice(InInputDevice);
+		ControllerId = FPlatformMisc::GetUserIndexForPlatformUser(UserID);
+	}
+
+	FInputKeyEventArgs(FViewport* InViewport, int32 InControllerId, FKey InKey, EInputEvent InEvent, float InAmountDepressed, bool bInIsTouchEvent)
+		: Viewport(InViewport)
+		, ControllerId(InControllerId)
+		, InputDevice(FInputDeviceId::CreateFromInternalId(InControllerId))
+		, Key(InKey)
+		, Event(InEvent)
+		, AmountDepressed(InAmountDepressed)
+		, bIsTouchEvent(bInIsTouchEvent)
+	{
+	}
+	
+	FInputKeyEventArgs(FViewport* InViewport, FInputDeviceId InInputDevice, FKey InKey, EInputEvent InEvent, float InAmountDepressed, bool bInIsTouchEvent)
+		: Viewport(InViewport)
+		, InputDevice(InInputDevice)
+		, Key(InKey)
+		, Event(InEvent)
+		, AmountDepressed(InAmountDepressed)
+		, bIsTouchEvent(bInIsTouchEvent)
+	{
+		FPlatformUserId UserID = IPlatformInputDeviceMapper::Get().GetUserForInputDevice(InInputDevice);
+		ControllerId = FPlatformMisc::GetUserIndexForPlatformUser(UserID);
+	}
+
+	bool IsGamepad() const { return Key.IsGamepadKey(); }
+
+public:
+
+	// The viewport which the key event is from.
+	FViewport* Viewport;
+	// The controller which the key event is from.
+	int32 ControllerId;
+	// The input device which this key event is from
+	FInputDeviceId InputDevice;
+	// The type of event which occurred.
+	FKey Key;
+	// The type of event which occurred.
+	EInputEvent Event;
+	// For analog keys, the depression percent.
+	float AmountDepressed;
+	// input came from a touch surface.This may be a faked mouse button from touch.
+	bool bIsTouchEvent;
+};
diff --git a/Engine/Source/Runtime/Engine/Public/MaterialShared.h b/Engine/Source/Runtime/Engine/Public/MaterialShared.h
--- a/Engine/Source/Runtime/Engine/Public/MaterialShared.h
+++ b/Engine/Source/Runtime/Engine/Public/MaterialShared.h
@@ -14,7 +14,7 @@
 #include "Templates/RefCounting.h"
 #include "Templates/UniquePtr.h"
 #include "Misc/SecureHash.h"
-#include "RHI.h"
+#include "RHIDefinitions.h"
 #include "RenderResource.h"
 #include "RenderingThread.h"
 #include "UniformBuffer.h"
@@ -25,12 +25,19 @@
 #include "Misc/Optional.h"
 #include "Serialization/MemoryWriter.h"
 #include "Serialization/ArchiveProxy.h"
+#include "MaterialRecursionGuard.h"
 #include "MaterialSceneTextureId.h"
+#include "MaterialShaderPrecompileMode.h"
+#include "MaterialValueType.h"
 #include "VirtualTexturing.h"
 #include "Templates/UnrealTemplate.h"
 #include "ShaderCompilerCore.h"
 #include "PSOPrecache.h"
 
+#if UE_ENABLE_INCLUDE_ORDER_DEPRECATED_IN_5_2
+#include "RHI.h"
+#endif
+
 #include "Shader/Preshader.h"
 
 #include <atomic>
@@ -143,69 +150,6 @@ int32 GetCmdLineMaterialQualityToCache();
 
 #endif
 
-/**
- * The types which can be used by materials.
- */
-enum EMaterialValueType
-{
-	/** 
-	 * A scalar float type.  
-	 * Note that MCT_Float1 will not auto promote to any other float types, 
-	 * So use MCT_Float instead for scalar expression return types.
-	 */
-	MCT_Float1		= 1,
-	MCT_Float2		= 2,
-	MCT_Float3		= 4,
-	MCT_Float4		= 8,
-
-	/** 
-	 * Any size float type by definition, but this is treated as a scalar which can auto convert (by replication) to any other size float vector.
-	 * Use this as the type for any scalar expressions.
-	 */
-	MCT_Float                 = 8|4|2|1,
-	MCT_Texture2D	          = 1 << 4,
-	MCT_TextureCube	          = 1 << 5,
-	MCT_Texture2DArray		  = 1 << 6,
-	MCT_TextureCubeArray      = 1 << 7,
-	MCT_VolumeTexture         = 1 << 8,
-	MCT_StaticBool            = 1 << 9,
-	MCT_Unknown               = 1 << 10,
-	MCT_MaterialAttributes	  = 1 << 11,
-	MCT_TextureExternal       = 1 << 12,
-	MCT_TextureVirtual        = 1 << 13,
-	MCT_Texture               = MCT_Texture2D | MCT_TextureCube | MCT_Texture2DArray | MCT_TextureCubeArray | MCT_VolumeTexture | MCT_TextureExternal | MCT_TextureVirtual,
-
-	/** Used internally when sampling from virtual textures */
-	MCT_VTPageTableResult     = 1 << 14,
-	
-	MCT_ShadingModel		  = 1 << 15,
-
-	MCT_Strata				  = 1 << 16,
-
-	MCT_LWCScalar			  = 1 << 17,
-	MCT_LWCVector2			  = 1 << 18,
-	MCT_LWCVector3			  = 1 << 19,
-	MCT_LWCVector4			  = 1 << 20,
-	MCT_LWCType				  = MCT_LWCScalar | MCT_LWCVector2 | MCT_LWCVector3 | MCT_LWCVector4,
-
-	MCT_Execution             = 1 << 21,
-
-	/** Used for code chunks that are statements with no value, rather than expressions */
-	MCT_VoidStatement         = 1 << 22,
-
-	/** Non-static bool, only used in new HLSL translator */
-	MCT_Bool				 = 1 << 23,
-
-	/** Unsigned int types */
-	MCT_UInt1				 = 1 << 24,
-	MCT_UInt2				 = 1 << 25,
-	MCT_UInt3				 = 1 << 26,
-	MCT_UInt4				 = 1 << 27,
-	MCT_UInt				 = MCT_UInt1 | MCT_UInt2 | MCT_UInt3 | MCT_UInt4,
-
-	MCT_Numeric = MCT_Float | MCT_LWCType | MCT_Bool,
-};
-
 /** @return the number of components in a vector type. */
 inline uint32 GetNumComponents(EMaterialValueType Type)
 {
@@ -1163,15 +1107,6 @@ private:
 	LAYOUT_FIELD_EDITORONLY(FMemoryImageString, MaterialPath);
 };
 
-enum class EMaterialShaderPrecompileMode
-{
-	None,
-	Background,
-	Synchronous,
-
-	Default = Background,
-};
-
 /**
  * The set of material shaders for a single material.
  */
@@ -3600,38 +3535,3 @@ inline bool ShouldIncludeMaterialInDefaultOpaquePass(const FMaterial& Material)
 	return !Material.IsSky()
 		&& !Material.GetShadingModels().HasShadingModel(MSM_SingleLayerWater);
 }
-
-template<typename T>
-struct TMaterialRecursionGuard
-{
-	inline TMaterialRecursionGuard() = default;
-	inline TMaterialRecursionGuard(const TMaterialRecursionGuard& Parent)
-		: Value(nullptr)
-		, PreviousLink(&Parent)
-	{
-	}
-
-	inline void Set(const T* InValue)
-	{
-		check(Value == nullptr);
-		Value = InValue;
-	}
-
-	inline bool Contains(const T* InValue)
-	{
-		TMaterialRecursionGuard const* Link = this;
-		do
-		{
-			if (Link->Value == InValue)
-			{
-				return true;
-			}
-			Link = Link->PreviousLink;
-		} while (Link);
-		return false;
-	}
-
-private:
-	const T* Value = nullptr;
-	TMaterialRecursionGuard const* PreviousLink = nullptr;
-};
diff --git a/Engine/Source/Runtime/Engine/Public/MaterialValueType.h b/Engine/Source/Runtime/Engine/Public/MaterialValueType.h
new file mode 100644
--- /dev/null
+++ b/Engine/Source/Runtime/Engine/Public/MaterialValueType.h
@@ -0,0 +1,70 @@
+// Copyright Epic Games, Inc. All Rights Reserved.
+
+/*=============================================================================
+	MaterialValueType.h
+=============================================================================*/
+
+#pragma once
+
+/**
+ * The types which can be used by materials.
+ */
+enum EMaterialValueType
+{
+	/** 
+	 * A scalar float type.  
+	 * Note that MCT_Float1 will not auto promote to any other float types, 
+	 * So use MCT_Float instead for scalar expression return types.
+	 */
+	MCT_Float1		= 1,
+	MCT_Float2		= 2,
+	MCT_Float3		= 4,
+	MCT_Float4		= 8,
+
+	/** 
+	 * Any size float type by definition, but this is treated as a scalar which can auto convert (by replication) to any other size float vector.
+	 * Use this as the type for any scalar expressions.
+	 */
+	MCT_Float                 = 8|4|2|1,
+	MCT_Texture2D	          = 1 << 4,
+	MCT_TextureCube	          = 1 << 5,
+	MCT_Texture2DArray		  = 1 << 6,
+	MCT_TextureCubeArray      = 1 << 7,
+	MCT_VolumeTexture         = 1 << 8,
+	MCT_StaticBool            = 1 << 9,
+	MCT_Unknown               = 1 << 10,
+	MCT_MaterialAttributes	  = 1 << 11,
+	MCT_TextureExternal       = 1 << 12,
+	MCT_TextureVirtual        = 1 << 13,
+	MCT_Texture               = MCT_Texture2D | MCT_TextureCube | MCT_Texture2DArray | MCT_TextureCubeArray | MCT_VolumeTexture | MCT_TextureExternal | MCT_TextureVirtual,
+
+	/** Used internally when sampling from virtual textures */
+	MCT_VTPageTableResult     = 1 << 14,
+	
+	MCT_ShadingModel		  = 1 << 15,
+
+	MCT_Strata				  = 1 << 16,
+
+	MCT_LWCScalar			  = 1 << 17,
+	MCT_LWCVector2			  = 1 << 18,
+	MCT_LWCVector3			  = 1 << 19,
+	MCT_LWCVector4			  = 1 << 20,
+	MCT_LWCType				  = MCT_LWCScalar | MCT_LWCVector2 | MCT_LWCVector3 | MCT_LWCVector4,
+
+	MCT_Execution             = 1 << 21,
+
+	/** Used for code chunks that are statements with no value, rather than expressions */
+	MCT_VoidStatement         = 1 << 22,
+
+	/** Non-static bool, only used in new HLSL translator */
+	MCT_Bool				 = 1 << 23,
+
+	/** Unsigned int types */
+	MCT_UInt1				 = 1 << 24,
+	MCT_UInt2				 = 1 << 25,
+	MCT_UInt3				 = 1 << 26,
+	MCT_UInt4				 = 1 << 27,
+	MCT_UInt				 = MCT_UInt1 | MCT_UInt2 | MCT_UInt3 | MCT_UInt4,
+
+	MCT_Numeric = MCT_Float | MCT_LWCType | MCT_Bool,
+};
diff --git a/Engine/Source/Runtime/Engine/Public/MeshBuilderOneFrameResources.h b/Engine/Source/Runtime/Engine/Public/MeshBuilderOneFrameResources.h
--- a/Engine/Source/Runtime/Engine/Public/MeshBuilderOneFrameResources.h
+++ b/Engine/Source/Runtime/Engine/Public/MeshBuilderOneFrameResources.h
@@ -2,4 +2,19 @@
 
 #pragma once
 
-// Placeholder for upcoming change
\ No newline at end of file
+#include "SceneManagement.h"
+
+class FMeshBuilderOneFrameResources : public FOneFrameResource
+{
+public:
+	class FPooledDynamicMeshVertexBuffer* VertexBuffer = nullptr;
+	class FPooledDynamicMeshIndexBuffer* IndexBuffer = nullptr;
+	class FPooledDynamicMeshVertexFactory* VertexFactory = nullptr;
+	class FDynamicMeshPrimitiveUniformBuffer* PrimitiveUniformBuffer = nullptr;
+	virtual ENGINE_API ~FMeshBuilderOneFrameResources();
+
+	inline bool IsValidForRendering() 
+	{
+		return VertexBuffer && IndexBuffer && PrimitiveUniformBuffer && VertexFactory;
+	}
+};
diff --git a/Engine/Source/Runtime/Engine/Public/MeshUVChannelInfo.h b/Engine/Source/Runtime/Engine/Public/MeshUVChannelInfo.h
--- a/Engine/Source/Runtime/Engine/Public/MeshUVChannelInfo.h
+++ b/Engine/Source/Runtime/Engine/Public/MeshUVChannelInfo.h
@@ -2,4 +2,60 @@
 
 #pragma once
 
-// Placeholder for upcoming change
\ No newline at end of file
+#include "MeshUVChannelInfo.generated.h"
+
+// Constants.
+enum { MAX_TEXCOORDS = 4 };
+
+/** The world size for each texcoord mapping. Used by the texture streaming. */
+USTRUCT(BlueprintType)
+struct FMeshUVChannelInfo
+{
+	GENERATED_USTRUCT_BODY()
+
+	/** Default constructor (no initialization). */
+	FORCEINLINE FMeshUVChannelInfo() { FMemory::Memzero(*this); }
+
+	/** Constructor which initializes all components to zero. */
+	FMeshUVChannelInfo(ENoInit) { }
+
+	FMeshUVChannelInfo(float DefaultDensity) : bInitialized(true), bOverrideDensities(false)
+	{
+		for (float& Density : LocalUVDensities)
+		{
+			Density = DefaultDensity;
+		}
+	}
+
+	/** Returns whether the structure contains any valid LocalUVDensities. */
+	bool IsInitialized() const
+	{
+		if (bInitialized)
+		{
+			for (float Density : LocalUVDensities)
+			{
+				if (Density != 0)
+				{
+					return true;
+				}
+			}
+		}
+		return false;
+	}
+
+	UPROPERTY()
+	bool bInitialized;
+
+	/** Whether this values was set manually or is auto generated. */
+	UPROPERTY(EditAnywhere, AdvancedDisplay, Category = Basic)
+	bool bOverrideDensities;
+
+	/**
+	 * The UV density in the mesh, before any transform scaling, in world unit per UV.
+	 * This value represents the length taken to cover a full UV unit.
+	 */
+	UPROPERTY(EditAnywhere, AdvancedDisplay, Category = Basic, meta = (EditCondition = "bOverrideDensities"))
+	float LocalUVDensities[MAX_TEXCOORDS];
+
+	friend FArchive& operator<<(FArchive& Ar, FMeshUVChannelInfo& Info);
+};
diff --git a/Engine/Source/Runtime/Engine/Public/PerPlatformProperties.h b/Engine/Source/Runtime/Engine/Public/PerPlatformProperties.h
--- a/Engine/Source/Runtime/Engine/Public/PerPlatformProperties.h
+++ b/Engine/Source/Runtime/Engine/Public/PerPlatformProperties.h
@@ -6,7 +6,9 @@ PerPlatformProperties.h: Property types that can be overridden on a per-platform
 
 #pragma once
 
+#if UE_ENABLE_INCLUDE_ORDER_DEPRECATED_IN_5_2
 #include "Engine/Engine.h"
+#endif
 #include "Serialization/Archive.h"
 #include "RHIDefinitions.h"
 #include "Containers/Map.h"
@@ -14,12 +16,13 @@ PerPlatformProperties.h: Property types that can be overridden on a per-platform
 #include "Serialization/MemoryLayout.h"
 #include "Misc/DataDrivenPlatformInfoRegistry.h"
 
+#include "PerPlatformProperties.generated.h"
+
+/** Helper function to avoid including Engine.h in this header */
 #if WITH_EDITORONLY_DATA && WITH_EDITOR
-#include "Engine/Engine.h"
+ENGINE_API bool GEngine_GetPreviewPlatformName(FName& PlatformName);
 #endif
 
-#include "PerPlatformProperties.generated.h"
-
 namespace PerPlatformProperty::Private
 {
 	template<typename MapType>
@@ -193,7 +196,7 @@ struct ENGINE_API TPerPlatformProperty
 		FName PlatformName;
 		// Lookup the override preview platform info, if any
 		// @todo this doesn't set PlatformName, just a group, but GetValueForPlatform() will technically work being given a Group name instead of a platform name, so we just use it
-		if (GEngine && GEngine->GetPreviewPlatformName(PlatformName))
+		if (GEngine_GetPreviewPlatformName(PlatformName))
 		{
 			return GetValueForPlatform(PlatformName);
 		}
@@ -474,4 +477,4 @@ struct TStructOpsTypeTraits<FPerPlatformBool>
 		WithSerializeFromMismatchedTag = true,
 		WithSerializer = true
 	};
-};
\ No newline at end of file
+};
diff --git a/Engine/Source/Runtime/Engine/Public/PerQualityLevelProperties.h b/Engine/Source/Runtime/Engine/Public/PerQualityLevelProperties.h
--- a/Engine/Source/Runtime/Engine/Public/PerQualityLevelProperties.h
+++ b/Engine/Source/Runtime/Engine/Public/PerQualityLevelProperties.h
@@ -12,10 +12,10 @@ FPerQualityLevelProperties.h: Property types that can be overridden on a quality
 #include "Serialization/MemoryLayout.h"
 #include "Scalability.h"
 #include "CoreMinimal.h"
-//#if WITH_EDITORONLY_DATA && WITH_EDITOR
+#if UE_ENABLE_INCLUDE_ORDER_DEPRECATED_IN_5_2
 #include "Engine/Engine.h"
 #include "HAL/ConsoleManager.h"
-//#endif
+#endif
 
 #include "PerQualityLevelProperties.generated.h"
 
diff --git a/Engine/Source/Runtime/Engine/Public/PrimitiveUniformShaderParameters.h b/Engine/Source/Runtime/Engine/Public/PrimitiveUniformShaderParameters.h
--- a/Engine/Source/Runtime/Engine/Public/PrimitiveUniformShaderParameters.h
+++ b/Engine/Source/Runtime/Engine/Public/PrimitiveUniformShaderParameters.h
@@ -14,7 +14,9 @@
 #include "UnifiedBuffer.h"
 #include "Containers/StaticArray.h"
 #include "NaniteDefinitions.h"
+#if UE_ENABLE_INCLUDE_ORDER_DEPRECATED_IN_5_2
 #include "UnrealEngine.h"
+#endif
 
 /** 
  * The uniform shader parameters associated with a primitive. 
@@ -322,23 +324,9 @@ public:
 		return *this;
 	}
 
-	inline FPrimitiveUniformShaderParametersBuilder& InstanceDrawDistance(FVector2f DistanceMinMax)
-	{
-		// Only scale the far distance by scalability parameters
-		DistanceMinMax.Y *= GetCachedScalabilityCVars().ViewDistanceScale;
-		Parameters.InstanceDrawDistanceMinMaxSquared = FMath::Square(DistanceMinMax);
-		bHasInstanceDrawDistanceCull = true;
-		return *this;
-	}
+	ENGINE_API FPrimitiveUniformShaderParametersBuilder& InstanceDrawDistance(FVector2f DistanceMinMax);
 
-	inline FPrimitiveUniformShaderParametersBuilder& InstanceWorldPositionOffsetDisableDistance(float WPODisableDistance)
-	{
-		WPODisableDistance *= GetCachedScalabilityCVars().ViewDistanceScale;
-		bHasWPODisableDistance = true;
-		Parameters.InstanceWPODisableDistanceSquared = WPODisableDistance * WPODisableDistance;
-
-		return *this;
-	}
+	ENGINE_API FPrimitiveUniformShaderParametersBuilder& InstanceWorldPositionOffsetDisableDistance(float WPODisableDistance);
 
 	inline const FPrimitiveUniformShaderParameters& Build()
 	{
@@ -540,7 +528,7 @@ struct FPrimitiveSceneShaderData
 	FPrimitiveSceneShaderData()
 		: Data(InPlace, NoInit)
 	{
-		static_assert(FPrimitiveSceneShaderData::DataStrideInFloat4s == FScatterUploadBuffer::PrimitiveDataStrideInFloat4s,"");
+		static_assert(FPrimitiveSceneShaderData::DataStrideInFloat4s == (int)FScatterUploadBuffer::PrimitiveDataStrideInFloat4s,"");
 		Setup(GetIdentityPrimitiveParameters());
 	}
 
diff --git a/Engine/Source/Runtime/Engine/Public/SceneManagement.h b/Engine/Source/Runtime/Engine/Public/SceneManagement.h
--- a/Engine/Source/Runtime/Engine/Public/SceneManagement.h
+++ b/Engine/Source/Runtime/Engine/Public/SceneManagement.h
@@ -17,11 +17,14 @@
 #include "RHI.h"
 #include "RenderResource.h"
 #include "EngineDefines.h"
+#include "GameFramework/Actor.h"
 #include "HitProxies.h"
 #include "SceneTypes.h"
 #include "ConvexVolume.h"
 #include "RendererInterface.h"
+#if UE_ENABLE_INCLUDE_ORDER_DEPRECATED_IN_5_2
 #include "Engine/TextureLightProfile.h"
+#endif
 #include "BatchedElements.h"
 #include "MeshBatch.h"
 #include "SceneUtils.h"
@@ -39,6 +42,7 @@ class FPrimitiveSceneProxy;
 class FSceneViewState;
 class FShadowMap;
 class FStaticMeshRenderData;
+class FTexture;
 class UDecalComponent;
 class ULightComponent;
 class ULightMapTexture2D;
@@ -48,6 +52,7 @@ class UShadowMapTexture2D;
 class USkyAtmosphereComponent;
 class FSkyAtmosphereRenderSceneInfo;
 class USkyLightComponent;
+class UTextureLightProfile;
 struct FDynamicMeshVertex;
 class ULightMapVirtualTexture2D;
 class FGPUScenePrimitiveCollector;
@@ -1559,7 +1564,7 @@ public:
 	inline float GetLightFunctionFadeDistance() const { return LightFunctionFadeDistance; }
 	inline float GetLightFunctionDisabledBrightness() const { return LightFunctionDisabledBrightness; }
 	inline UTextureLightProfile* GetIESTexture() const { return IESTexture; }
-	inline FTexture* GetIESTextureResource() const { return IESTexture ? IESTexture->GetResource() : nullptr; }
+	FTexture* GetIESTextureResource() const;
 	inline const FMaterialRenderProxy* GetLightFunctionMaterial() const { return LightFunctionMaterial; }
 	inline bool IsMovable() const { return bMovable; }
 	inline bool HasStaticLighting() const { return bStaticLighting; }
diff --git a/Engine/Source/Runtime/Engine/Public/SceneUtils.h b/Engine/Source/Runtime/Engine/Public/SceneUtils.h
--- a/Engine/Source/Runtime/Engine/Public/SceneUtils.h
+++ b/Engine/Source/Runtime/Engine/Public/SceneUtils.h
@@ -11,11 +11,15 @@
 #pragma once
 
 #include "UObject/ObjectMacros.h"
+#if UE_ENABLE_INCLUDE_ORDER_DEPRECATED_IN_5_2
 #include "ProfilingDebugging/RealtimeGPUProfiler.h"
+#endif
 #include "RHIDefinitions.h"
 
 ENGINE_API DECLARE_LOG_CATEGORY_EXTERN(LogSceneUtils, Log, All);
 
+class FRHICommandListImmediate;
+
 enum class EShadingPath
 {
 	Mobile,
diff --git a/Engine/Source/Runtime/Engine/Public/SceneView.h b/Engine/Source/Runtime/Engine/Public/SceneView.h
--- a/Engine/Source/Runtime/Engine/Public/SceneView.h
+++ b/Engine/Source/Runtime/Engine/Public/SceneView.h
@@ -5,22 +5,29 @@
 #include "CoreMinimal.h"
 #include "Engine/EngineTypes.h"
 #include "Engine/EngineBaseTypes.h"
+#include "EngineDefines.h"
 #include "RenderResource.h"
 #include "UniformBuffer.h"
-#include "Engine/World.h"
+#include "GameTime.h"
 #include "SceneTypes.h"
 #include "ShowFlags.h"
 #include "ConvexVolume.h"
 #include "Engine/GameViewportClient.h"
 #include "SceneInterface.h"
 #include "FinalPostProcessSettings.h"
-#include "GlobalDistanceFieldParameters.h"
+#include "GlobalDistanceFieldConstants.h"
 #include "DebugViewModeHelpers.h"
+#include "PhysicsInterfaceDeclaresCore.h"
 #include "RendererInterface.h"
 #include "Interfaces/Interface_PostProcessVolume.h"
 #include "DynamicRenderScaling.h"
 #include "Math/MirrorMatrix.h"
 
+#if UE_ENABLE_INCLUDE_ORDER_DEPRECATED_IN_5_2
+#include "Engine/World.h"
+#include "GlobalDistanceFieldParameters.h"
+#endif
+
 class FSceneView;
 class FSceneViewFamily;
 class FSceneViewStateInterface;
@@ -1638,32 +1645,12 @@ public:
 	*/
 	struct ConstructionValues
 	{
-		ConstructionValues(
+		ENGINE_API ConstructionValues(
 			const FRenderTarget* InRenderTarget,
 			FSceneInterface* InScene,
 			const FEngineShowFlags& InEngineShowFlags
-			)
-		:	RenderTarget(InRenderTarget)
-		,	Scene(InScene)
-		,	EngineShowFlags(InEngineShowFlags)
-		,	ViewModeParam(-1)
-		,	GammaCorrection(1.0f)
-		,	bAdditionalViewFamily(false)
-		,	bRealtimeUpdate(false)
-		,	bDeferClear(false)
-		,	bResolveScene(true)			
-		,	bTimesSet(false)
-		{
-			if( InScene != NULL )			
-			{
-				UWorld* World = InScene->GetWorld();
-				// Ensure the world is valid and that we are being called from a game thread (GetRealTimeSeconds requires this)
-				if( World && IsInGameThread() )
-				{	
-					SetTime(World->GetTime());
-				}
-			}
-		}
+			);
+
 		/** The views which make up the family. */
 		const FRenderTarget* RenderTarget;
 
diff --git a/Engine/Source/Runtime/Engine/Public/SkeletalMeshSceneProxy.h b/Engine/Source/Runtime/Engine/Public/SkeletalMeshSceneProxy.h
--- a/Engine/Source/Runtime/Engine/Public/SkeletalMeshSceneProxy.h
+++ b/Engine/Source/Runtime/Engine/Public/SkeletalMeshSceneProxy.h
@@ -2,4 +2,219 @@
 
 #pragma once
 
-// Placeholder for upcoming change
\ No newline at end of file
+#include "PrimitiveSceneProxy.h"
+#include "SkeletalMeshTypes.h"
+
+/*-----------------------------------------------------------------------------
+FSkeletalMeshSceneProxy
+-----------------------------------------------------------------------------*/
+
+/**
+ * A skeletal mesh component scene proxy.
+ */
+class ENGINE_API FSkeletalMeshSceneProxy : public FPrimitiveSceneProxy
+{
+public:
+	SIZE_T GetTypeHash() const override;
+
+	/** 
+	 * Constructor. 
+	 * @param	Component - skeletal mesh primitive being added
+	 */
+	FSkeletalMeshSceneProxy(const USkinnedMeshComponent* Component, FSkeletalMeshRenderData* InSkelMeshRenderData);
+
+#if WITH_EDITOR
+	virtual HHitProxy* CreateHitProxies(UPrimitiveComponent* Component, TArray<TRefCountPtr<HHitProxy> >& OutHitProxies) override;
+#endif
+	virtual void DrawStaticElements(FStaticPrimitiveDrawInterface* PDI) override;
+	virtual void GetDynamicMeshElements(const TArray<const FSceneView*>& Views, const FSceneViewFamily& ViewFamily, uint32 VisibilityMap, FMeshElementCollector& Collector) const override;
+
+#if RHI_RAYTRACING
+	virtual bool HasRayTracingRepresentation() const override;
+
+	virtual bool IsRayTracingRelevant() const override { return true; }
+
+	virtual bool IsRayTracingStaticRelevant() const override
+	{
+		return bRenderStatic;
+	}
+
+	virtual void GetDynamicRayTracingInstances(struct FRayTracingMaterialGatheringContext& Context, TArray<struct FRayTracingInstance>& OutRayTracingInstances) override;
+#endif // RHI_RAYTRACING
+
+	virtual FPrimitiveViewRelevance GetViewRelevance(const FSceneView* View) const override;
+	virtual bool CanBeOccluded() const override;
+	virtual bool IsUsingDistanceCullFade() const override;
+	
+	virtual bool HasDynamicIndirectShadowCasterRepresentation() const override;
+	virtual void GetShadowShapes(FVector PreViewTranslation, TArray<FCapsuleShape3f>& OutCapsuleShapes) const override;
+
+	/** Return the bounds for the pre-skinned primitive in local space */
+	virtual void GetPreSkinnedLocalBounds(FBoxSphereBounds& OutBounds) const override { OutBounds = PreSkinnedLocalBounds; }
+
+	/** Returns a pre-sorted list of shadow capsules's bone indicies */
+	const TArray<uint16>& GetSortedShadowBoneIndices() const
+	{
+		return ShadowCapsuleBoneIndices;
+	}
+
+	/**
+	 * Returns the world transform to use for drawing.
+	 * @param OutLocalToWorld - Will contain the local-to-world transform when the function returns.
+	 * @param OutWorldToLocal - Will contain the world-to-local transform when the function returns.
+	 * 
+	 * @return true if out matrices are valid 
+	 */
+	bool GetWorldMatrices( FMatrix& OutLocalToWorld, FMatrix& OutWorldToLocal ) const;
+
+	/** Util for getting LOD index currently used by this SceneProxy. */
+	int32 GetCurrentLODIndex();
+
+	/** 
+	 * Render physics asset for debug display
+	 */
+	virtual void DebugDrawPhysicsAsset(int32 ViewIndex, FMeshElementCollector& Collector, const FEngineShowFlags& EngineShowFlags) const;
+
+	/** Render the bones of the skeleton for debug display */ 
+	void DebugDrawSkeleton(int32 ViewIndex, FMeshElementCollector& Collector, const FEngineShowFlags& EngineShowFlags) const;
+
+#if WITH_EDITOR
+	void DebugDrawPoseWatchSkeletons(int32 ViewIndex, FMeshElementCollector& Collector, const FEngineShowFlags& EngineShowFlags) const;
+#endif
+
+	virtual uint32 GetMemoryFootprint( void ) const override { return( sizeof( *this ) + GetAllocatedSize() ); }
+	SIZE_T GetAllocatedSize( void ) const { return( FPrimitiveSceneProxy::GetAllocatedSize() + LODSections.GetAllocatedSize() ); }
+
+	/**
+	* Updates morph material usage for materials referenced by each LOD entry
+	*
+	* @param bNeedsMorphUsage - true if the materials used by this skeletal mesh need morph target usage
+	*/
+	void UpdateMorphMaterialUsage_GameThread(TArray<UMaterialInterface*>& MaterialUsingMorphTarget);
+
+
+#if WITH_EDITORONLY_DATA
+	virtual bool GetPrimitiveDistance(int32 LODIndex, int32 SectionIndex, const FVector& ViewOrigin, float& PrimitiveDistance) const override;
+	virtual bool GetMeshUVDensities(int32 LODIndex, int32 SectionIndex, FVector4& WorldUVDensities) const override;
+	virtual bool GetMaterialTextureScales(int32 LODIndex, int32 SectionIndex, const FMaterialRenderProxy* MaterialRenderProxy, FVector4f* OneOverScales, FIntVector4* UVChannelIndices) const override;
+#endif
+
+	friend class FSkeletalMeshSectionIter;
+
+	virtual void OnTransformChanged() override;
+
+	virtual uint8 GetCurrentFirstLODIdx_RenderThread() const final override
+	{
+		return GetCurrentFirstLODIdx_Internal();
+	}
+
+	bool GetCachedGeometry(struct FCachedGeometry& OutCachedGeometry) const;
+
+protected:
+	AActor* Owner;
+	class FSkeletalMeshObject* MeshObject;
+	FSkeletalMeshRenderData* SkeletalMeshRenderData;
+
+	/** The points to the skeletal mesh and physics assets are purely for debug purposes. Access is NOT thread safe! */
+	const class USkinnedAsset* SkeletalMeshForDebug;
+	class UPhysicsAsset* PhysicsAssetForDebug;
+	
+	UMaterialInterface* OverlayMaterial;
+	float OverlayMaterialMaxDrawDistance;
+public:
+#if RHI_RAYTRACING
+	bool bAnySegmentUsesWorldPositionOffset : 1;
+#endif
+
+protected:
+	/** data copied for rendering */
+	uint8 bForceWireframe : 1;
+	uint8 bIsCPUSkinned : 1;
+	uint8 bCanHighlightSelectedSections : 1;
+	uint8 bRenderStatic:1;
+#if !(UE_BUILD_SHIPPING || UE_BUILD_TEST)
+	uint8 bDrawDebugSkeleton:1;
+#endif
+
+	TEnumAsByte<ERHIFeatureLevel::Type> FeatureLevel;
+
+	bool bMaterialsNeedMorphUsage_GameThread;
+
+	FMaterialRelevance MaterialRelevance;
+
+
+	/** info for section element in an LOD */
+	struct FSectionElementInfo
+	{
+		FSectionElementInfo(UMaterialInterface* InMaterial, bool bInEnableShadowCasting, int32 InUseMaterialIndex)
+		: Material( InMaterial )
+		, bEnableShadowCasting( bInEnableShadowCasting )
+		, UseMaterialIndex( InUseMaterialIndex )
+#if WITH_EDITOR
+		, HitProxy(NULL)
+#endif
+		{}
+		
+		UMaterialInterface* Material;
+		
+		/** Whether shadow casting is enabled for this section. */
+		bool bEnableShadowCasting;
+		
+		/** Index into the materials array of the skel mesh or the component after LOD mapping */
+		int32 UseMaterialIndex;
+
+#if WITH_EDITOR
+		/** The editor needs to be able to individual sub-mesh hit detection, so we store a hit proxy on each mesh. */
+		HHitProxy* HitProxy;
+#endif
+	};
+
+	/** Section elements for a particular LOD */
+	struct FLODSectionElements
+	{
+		TArray<FSectionElementInfo> SectionElements;
+	};
+	
+	/** Array of section elements for each LOD */
+	TArray<FLODSectionElements> LODSections;
+
+	/** 
+	 * BoneIndex->capsule pairs used for rendering sphere/capsule shadows 
+	 * Note that these are in refpose component space, NOT bone space.
+	 */
+	TArray<TPair<int32, FCapsuleShape>> ShadowCapsuleData;
+	TArray<uint16> ShadowCapsuleBoneIndices;
+
+	/** Set of materials used by this scene proxy, safe to access from the game thread. */
+	TSet<UMaterialInterface*> MaterialsInUse_GameThread;
+	
+	/** The primitive's pre-skinned local space bounds. */
+	FBoxSphereBounds PreSkinnedLocalBounds;
+
+#if !(UE_BUILD_SHIPPING || UE_BUILD_TEST)
+	/** The color we draw this component in if drawing debug bones */
+	TOptional<FLinearColor> DebugDrawColor;
+#endif
+
+#if WITH_EDITORONLY_DATA
+	/** The component streaming distance multiplier */
+	float StreamingDistanceMultiplier;
+#endif
+
+	void GetDynamicElementsSection(const TArray<const FSceneView*>& Views, const FSceneViewFamily& ViewFamily, uint32 VisibilityMap,
+									const FSkeletalMeshLODRenderData& LODData, const int32 LODIndex, const int32 SectionIndex, bool bSectionSelected,
+								   const FSectionElementInfo& SectionElementInfo, bool bInSelectable, FMeshElementCollector& Collector) const;
+
+	void GetMeshElementsConditionallySelectable(const TArray<const FSceneView*>& Views, const FSceneViewFamily& ViewFamily, bool bInSelectable, uint32 VisibilityMap, FMeshElementCollector& Collector) const;
+
+	/** Only call on render thread timeline */
+	uint8 GetCurrentFirstLODIdx_Internal() const;
+private:
+	void CreateBaseMeshBatch(const FSceneView* View, const FSkeletalMeshLODRenderData& LODData, const int32 LODIndex, const int32 SectionIndex, const FSectionElementInfo& SectionElementInfo, FMeshBatch& Mesh, ESkinVertexFactoryMode VFMode = ESkinVertexFactoryMode::Default) const;
+	void UpdateLooseParametersUniformBuffer(const FSceneView* View, const int32 SectionIndex, const FMeshBatch& Mesh, const struct FGPUSkinBatchElementUserData* BatchUserData) const;
+
+public:
+#if WITH_EDITORONLY_DATA
+	struct FPoseWatchDynamicData* PoseWatchDynamicData;
+#endif
+};
\ No newline at end of file
diff --git a/Engine/Source/Runtime/Engine/Public/SkeletalMeshTypes.h b/Engine/Source/Runtime/Engine/Public/SkeletalMeshTypes.h
--- a/Engine/Source/Runtime/Engine/Public/SkeletalMeshTypes.h
+++ b/Engine/Source/Runtime/Engine/Public/SkeletalMeshTypes.h
@@ -3,7 +3,9 @@
 #pragma once
 
 #include "CoreMinimal.h"
+#if UE_ENABLE_INCLUDE_ORDER_DEPRECATED_IN_5_2
 #include "PrimitiveSceneProxy.h"
+#endif
 #include "Materials/MaterialInterface.h"
 #include "ComponentReregisterContext.h"
 #include "SkeletalMeshLegacyCustomVersions.h"
@@ -131,220 +133,6 @@ struct FClothingSectionData
 };
 
 
-/*-----------------------------------------------------------------------------
-FSkeletalMeshSceneProxy
------------------------------------------------------------------------------*/
-
-/**
- * A skeletal mesh component scene proxy.
- */
-class ENGINE_API FSkeletalMeshSceneProxy : public FPrimitiveSceneProxy
-{
-public:
-	SIZE_T GetTypeHash() const override;
-
-	/** 
-	 * Constructor. 
-	 * @param	Component - skeletal mesh primitive being added
-	 */
-	FSkeletalMeshSceneProxy(const USkinnedMeshComponent* Component, FSkeletalMeshRenderData* InSkelMeshRenderData);
-
-#if WITH_EDITOR
-	virtual HHitProxy* CreateHitProxies(UPrimitiveComponent* Component, TArray<TRefCountPtr<HHitProxy> >& OutHitProxies) override;
-#endif
-	virtual void DrawStaticElements(FStaticPrimitiveDrawInterface* PDI) override;
-	virtual void GetDynamicMeshElements(const TArray<const FSceneView*>& Views, const FSceneViewFamily& ViewFamily, uint32 VisibilityMap, FMeshElementCollector& Collector) const override;
-
-#if RHI_RAYTRACING
-	virtual bool HasRayTracingRepresentation() const override;
-
-	virtual bool IsRayTracingRelevant() const override { return true; }
-
-	virtual bool IsRayTracingStaticRelevant() const override
-	{
-		return bRenderStatic;
-	}
-
-	virtual void GetDynamicRayTracingInstances(struct FRayTracingMaterialGatheringContext& Context, TArray<struct FRayTracingInstance>& OutRayTracingInstances) override;
-#endif // RHI_RAYTRACING
-
-	virtual FPrimitiveViewRelevance GetViewRelevance(const FSceneView* View) const override;
-	virtual bool CanBeOccluded() const override;
-	virtual bool IsUsingDistanceCullFade() const override;
-	
-	virtual bool HasDynamicIndirectShadowCasterRepresentation() const override;
-	virtual void GetShadowShapes(FVector PreViewTranslation, TArray<FCapsuleShape3f>& OutCapsuleShapes) const override;
-
-	/** Return the bounds for the pre-skinned primitive in local space */
-	virtual void GetPreSkinnedLocalBounds(FBoxSphereBounds& OutBounds) const override { OutBounds = PreSkinnedLocalBounds; }
-
-	/** Returns a pre-sorted list of shadow capsules's bone indicies */
-	const TArray<uint16>& GetSortedShadowBoneIndices() const
-	{
-		return ShadowCapsuleBoneIndices;
-	}
-
-	/**
-	 * Returns the world transform to use for drawing.
-	 * @param OutLocalToWorld - Will contain the local-to-world transform when the function returns.
-	 * @param OutWorldToLocal - Will contain the world-to-local transform when the function returns.
-	 * 
-	 * @return true if out matrices are valid 
-	 */
-	bool GetWorldMatrices( FMatrix& OutLocalToWorld, FMatrix& OutWorldToLocal ) const;
-
-	/** Util for getting LOD index currently used by this SceneProxy. */
-	int32 GetCurrentLODIndex();
-
-	/** 
-	 * Render physics asset for debug display
-	 */
-	virtual void DebugDrawPhysicsAsset(int32 ViewIndex, FMeshElementCollector& Collector, const FEngineShowFlags& EngineShowFlags) const;
-
-	/** Render the bones of the skeleton for debug display */ 
-	void DebugDrawSkeleton(int32 ViewIndex, FMeshElementCollector& Collector, const FEngineShowFlags& EngineShowFlags) const;
-
-#if WITH_EDITOR
-	void DebugDrawPoseWatchSkeletons(int32 ViewIndex, FMeshElementCollector& Collector, const FEngineShowFlags& EngineShowFlags) const;
-#endif
-
-	virtual uint32 GetMemoryFootprint( void ) const override { return( sizeof( *this ) + GetAllocatedSize() ); }
-	SIZE_T GetAllocatedSize( void ) const { return( FPrimitiveSceneProxy::GetAllocatedSize() + LODSections.GetAllocatedSize() ); }
-
-	/**
-	* Updates morph material usage for materials referenced by each LOD entry
-	*
-	* @param bNeedsMorphUsage - true if the materials used by this skeletal mesh need morph target usage
-	*/
-	void UpdateMorphMaterialUsage_GameThread(TArray<UMaterialInterface*>& MaterialUsingMorphTarget);
-
-
-#if WITH_EDITORONLY_DATA
-	virtual bool GetPrimitiveDistance(int32 LODIndex, int32 SectionIndex, const FVector& ViewOrigin, float& PrimitiveDistance) const override;
-	virtual bool GetMeshUVDensities(int32 LODIndex, int32 SectionIndex, FVector4& WorldUVDensities) const override;
-	virtual bool GetMaterialTextureScales(int32 LODIndex, int32 SectionIndex, const FMaterialRenderProxy* MaterialRenderProxy, FVector4f* OneOverScales, FIntVector4* UVChannelIndices) const override;
-#endif
-
-	friend class FSkeletalMeshSectionIter;
-
-	virtual void OnTransformChanged() override;
-
-	virtual uint8 GetCurrentFirstLODIdx_RenderThread() const final override
-	{
-		return GetCurrentFirstLODIdx_Internal();
-	}
-
-	bool GetCachedGeometry(struct FCachedGeometry& OutCachedGeometry) const;
-
-protected:
-	AActor* Owner;
-	class FSkeletalMeshObject* MeshObject;
-	FSkeletalMeshRenderData* SkeletalMeshRenderData;
-
-	/** The points to the skeletal mesh and physics assets are purely for debug purposes. Access is NOT thread safe! */
-	const class USkinnedAsset* SkeletalMeshForDebug;
-	class UPhysicsAsset* PhysicsAssetForDebug;
-	
-	UMaterialInterface* OverlayMaterial;
-	float OverlayMaterialMaxDrawDistance;
-public:
-#if RHI_RAYTRACING
-	bool bAnySegmentUsesWorldPositionOffset : 1;
-#endif
-
-protected:
-	/** data copied for rendering */
-	uint8 bForceWireframe : 1;
-	uint8 bIsCPUSkinned : 1;
-	uint8 bCanHighlightSelectedSections : 1;
-	uint8 bRenderStatic:1;
-#if !(UE_BUILD_SHIPPING || UE_BUILD_TEST)
-	uint8 bDrawDebugSkeleton:1;
-#endif
-
-	TEnumAsByte<ERHIFeatureLevel::Type> FeatureLevel;
-
-	bool bMaterialsNeedMorphUsage_GameThread;
-
-	FMaterialRelevance MaterialRelevance;
-
-
-	/** info for section element in an LOD */
-	struct FSectionElementInfo
-	{
-		FSectionElementInfo(UMaterialInterface* InMaterial, bool bInEnableShadowCasting, int32 InUseMaterialIndex)
-		: Material( InMaterial )
-		, bEnableShadowCasting( bInEnableShadowCasting )
-		, UseMaterialIndex( InUseMaterialIndex )
-#if WITH_EDITOR
-		, HitProxy(NULL)
-#endif
-		{}
-		
-		UMaterialInterface* Material;
-		
-		/** Whether shadow casting is enabled for this section. */
-		bool bEnableShadowCasting;
-		
-		/** Index into the materials array of the skel mesh or the component after LOD mapping */
-		int32 UseMaterialIndex;
-
-#if WITH_EDITOR
-		/** The editor needs to be able to individual sub-mesh hit detection, so we store a hit proxy on each mesh. */
-		HHitProxy* HitProxy;
-#endif
-	};
-
-	/** Section elements for a particular LOD */
-	struct FLODSectionElements
-	{
-		TArray<FSectionElementInfo> SectionElements;
-	};
-	
-	/** Array of section elements for each LOD */
-	TArray<FLODSectionElements> LODSections;
-
-	/** 
-	 * BoneIndex->capsule pairs used for rendering sphere/capsule shadows 
-	 * Note that these are in refpose component space, NOT bone space.
-	 */
-	TArray<TPair<int32, FCapsuleShape>> ShadowCapsuleData;
-	TArray<uint16> ShadowCapsuleBoneIndices;
-
-	/** Set of materials used by this scene proxy, safe to access from the game thread. */
-	TSet<UMaterialInterface*> MaterialsInUse_GameThread;
-	
-	/** The primitive's pre-skinned local space bounds. */
-	FBoxSphereBounds PreSkinnedLocalBounds;
-
-#if !(UE_BUILD_SHIPPING || UE_BUILD_TEST)
-	/** The color we draw this component in if drawing debug bones */
-	TOptional<FLinearColor> DebugDrawColor;
-#endif
-
-#if WITH_EDITORONLY_DATA
-	/** The component streaming distance multiplier */
-	float StreamingDistanceMultiplier;
-#endif
-
-	void GetDynamicElementsSection(const TArray<const FSceneView*>& Views, const FSceneViewFamily& ViewFamily, uint32 VisibilityMap,
-									const FSkeletalMeshLODRenderData& LODData, const int32 LODIndex, const int32 SectionIndex, bool bSectionSelected,
-								   const FSectionElementInfo& SectionElementInfo, bool bInSelectable, FMeshElementCollector& Collector) const;
-
-	void GetMeshElementsConditionallySelectable(const TArray<const FSceneView*>& Views, const FSceneViewFamily& ViewFamily, bool bInSelectable, uint32 VisibilityMap, FMeshElementCollector& Collector) const;
-
-	/** Only call on render thread timeline */
-	uint8 GetCurrentFirstLODIdx_Internal() const;
-private:
-	void CreateBaseMeshBatch(const FSceneView* View, const FSkeletalMeshLODRenderData& LODData, const int32 LODIndex, const int32 SectionIndex, const FSectionElementInfo& SectionElementInfo, FMeshBatch& Mesh, ESkinVertexFactoryMode VFMode = ESkinVertexFactoryMode::Default) const;
-	void UpdateLooseParametersUniformBuffer(const FSceneView* View, const int32 SectionIndex, const FMeshBatch& Mesh, const struct FGPUSkinBatchElementUserData* BatchUserData) const;
-
-public:
-#if WITH_EDITORONLY_DATA
-	struct FPoseWatchDynamicData* PoseWatchDynamicData;
-#endif
-};
-
 /** Used to recreate all skinned mesh components for a given skeletal mesh */
 class ENGINE_API FSkinnedMeshComponentRecreateRenderStateContext
 {
diff --git a/Engine/Source/Runtime/Engine/Public/StaticMeshComponentLODInfo.h b/Engine/Source/Runtime/Engine/Public/StaticMeshComponentLODInfo.h
--- a/Engine/Source/Runtime/Engine/Public/StaticMeshComponentLODInfo.h
+++ b/Engine/Source/Runtime/Engine/Public/StaticMeshComponentLODInfo.h
@@ -2,4 +2,98 @@
 
 #pragma once
 
-// placeholder
\ No newline at end of file
+#include "RawIndexBuffer.h"
+
+#include "StaticMeshComponentLODInfo.generated.h"
+
+class FColorVertexBuffer;
+class FMeshMapBuildData;
+
+struct FPreCulledStaticMeshSection
+{
+	/** Range of vertices and indices used when rendering this section. */
+	uint32 FirstIndex;
+	uint32 NumTriangles;
+};
+
+USTRUCT()
+struct FStaticMeshComponentLODInfo
+{
+	GENERATED_USTRUCT_BODY()
+
+	/** Uniquely identifies this LOD's built map data. */
+	FGuid MapBuildDataId;
+
+	/** Used during deserialization to temporarily store legacy lightmap data. */
+	FMeshMapBuildData* LegacyMapBuildData;
+
+	/** 
+	 * Transient override lightmap data, used by landscape grass.
+	 * Be sure to add your component to UMapBuildDataRegistry::CleanupTransientOverrideMapBuildData() for proper cleanup
+	 * so that you don't get stale rendering resource references if the underlying MapBuildData is gone (lighting scenario changes, new static lighting build, etc.)
+	 */
+	TUniquePtr<FMeshMapBuildData> OverrideMapBuildData;
+
+	/** Vertex data cached at the time this LOD was painted, if any */
+	TArray<struct FPaintedVertex> PaintedVertices;
+
+	/** Vertex colors to use for this mesh LOD */
+	FColorVertexBuffer* OverrideVertexColors;
+
+	/** Information for each section about what range of PreCulledIndexBuffer to use.  If no preculled index data is available, PreCulledSections will be empty. */
+	TArray<FPreCulledStaticMeshSection> PreCulledSections;
+
+	FRawStaticIndexBuffer PreCulledIndexBuffer;
+
+	/** 
+	 * Owner of this FStaticMeshComponentLODInfo 
+	 * Warning, can be NULL for a component created via SpawnActor off of a blueprint default (LODData will be created without a call to SetLODDataCount).
+	 */
+	class UStaticMeshComponent* OwningComponent;
+
+	/** Default constructor */
+	FStaticMeshComponentLODInfo();
+	FStaticMeshComponentLODInfo(UStaticMeshComponent* InOwningComponent);
+	/** Destructor */
+	~FStaticMeshComponentLODInfo();
+
+	/** Delete existing resources */
+	void CleanUp();
+
+	/** 
+	 * Ensure this LODInfo has a valid MapBuildDataId GUID.
+	 * @param LodIndex Index of the LOD this LODInfo represents.
+	 * @return true if a new GUID was created, false otherwise.
+	 */
+	ENGINE_API bool CreateMapBuildDataId(int32 LodIndex);
+
+	/**
+	* Enqueues a rendering command to release the vertex colors.
+	* The game thread must block until the rendering thread has processed the command before deleting OverrideVertexColors.
+	*/
+	ENGINE_API void BeginReleaseOverrideVertexColors();
+
+	ENGINE_API void ReleaseOverrideVertexColorsAndBlock();
+
+	void ReleaseResources();
+
+	/** Methods for importing and exporting the painted vertex array to text */
+	void ExportText(FString& ValueStr);
+	void ImportText(const TCHAR** SourceText);
+
+	/** Serializer. */
+	friend FArchive& operator<<(FArchive& Ar,FStaticMeshComponentLODInfo& I);
+
+private:
+	/** Purposely hidden */
+	FStaticMeshComponentLODInfo &operator=( const FStaticMeshComponentLODInfo &rhs ) { check(0); return *this; }
+};
+
+template<>
+struct TStructOpsTypeTraits<FStaticMeshComponentLODInfo> : public TStructOpsTypeTraitsBase2<FStaticMeshComponentLODInfo>
+{
+	enum
+	{
+		WithCopy = false
+	};
+};
diff --git a/Engine/Source/Runtime/Engine/Public/StaticMeshResources.h b/Engine/Source/Runtime/Engine/Public/StaticMeshResources.h
--- a/Engine/Source/Runtime/Engine/Public/StaticMeshResources.h
+++ b/Engine/Source/Runtime/Engine/Public/StaticMeshResources.h
@@ -19,7 +19,9 @@
 #include "Components.h"
 #include "LocalVertexFactory.h"
 #include "PrimitiveViewRelevance.h"
+#if UE_ENABLE_INCLUDE_ORDER_DEPRECATED_IN_5_2
 #include "PrimitiveSceneProxy.h"
+#endif
 #include "Engine/MeshMerging.h"
 #include "UObject/UObjectHash.h"
 #include "MeshBatch.h"
@@ -790,298 +792,6 @@ private:
 	bool bRefreshBounds;
 };
 
-/**
- * A static mesh component scene proxy.
- */
-class ENGINE_API FStaticMeshSceneProxy : public FPrimitiveSceneProxy
-{
-public:
-	SIZE_T GetTypeHash() const override;
-
-	/** Initialization constructor. */
-	FStaticMeshSceneProxy(UStaticMeshComponent* Component, bool bForceLODsShareStaticLighting);
-
-	virtual ~FStaticMeshSceneProxy();
-
-	/** Gets the number of mesh batches required to represent the proxy, aside from section needs. */
-	virtual int32 GetNumMeshBatches() const
-	{
-		return 1;
-	}
-
-	/** Sets up a shadow FMeshBatch for a specific LOD. */
-	virtual bool GetShadowMeshElement(int32 LODIndex, int32 BatchIndex, uint8 InDepthPriorityGroup, FMeshBatch& OutMeshBatch, bool bDitheredLODTransition) const;
-
-	/** Sets up a FMeshBatch for a specific LOD and element. */
-	virtual bool GetMeshElement(
-		int32 LODIndex, 
-		int32 BatchIndex, 
-		int32 ElementIndex, 
-		uint8 InDepthPriorityGroup, 
-		bool bUseSelectionOutline,
-		bool bAllowPreCulledIndices,
-		FMeshBatch& OutMeshBatch) const;
-
-	virtual void CreateRenderThreadResources() override;
-
-	virtual void DestroyRenderThreadResources() override;
-
-	/** Sets up a wireframe FMeshBatch for a specific LOD. */
-	virtual bool GetWireframeMeshElement(int32 LODIndex, int32 BatchIndex, const FMaterialRenderProxy* WireframeRenderProxy, uint8 InDepthPriorityGroup, bool bAllowPreCulledIndices, FMeshBatch& OutMeshBatch) const;
-
-	/** Sets up a collision FMeshBatch for a specific LOD and element. */
-	virtual bool GetCollisionMeshElement(
-		int32 LODIndex,
-		int32 BatchIndex,
-		int32 ElementIndex,
-		uint8 InDepthPriorityGroup,
-		const FMaterialRenderProxy* RenderProxy,
-		FMeshBatch& OutMeshBatch) const;
-
-	virtual void SetEvaluateWorldPositionOffsetInRayTracing(bool NewValue);
-
-	virtual uint8 GetCurrentFirstLODIdx_RenderThread() const final override
-	{
-		return GetCurrentFirstLODIdx_Internal();
-	}
-
-	virtual int32 GetLightMapCoordinateIndex() const override;
-
-protected:
-	/** Configures mesh batch vertex / index state. Returns the number of primitives used in the element. */
-	uint32 SetMeshElementGeometrySource(
-		int32 LODIndex,
-		int32 ElementIndex,
-		bool bWireframe,
-		bool bUseInversedIndices,
-		bool bAllowPreCulledIndices,
-		const FVertexFactory* VertexFactory,
-		FMeshBatch& OutMeshElement) const;
-
-	/** Sets the screen size on a mesh element. */
-	void SetMeshElementScreenSize(int32 LODIndex, bool bDitheredLODTransition, FMeshBatch& OutMeshBatch) const;
-
-	/** Returns whether this mesh needs reverse culling when using reversed indices. */
-	bool IsReversedCullingNeeded(bool bUseReversedIndices) const;
-
-	bool IsCollisionView(const FEngineShowFlags& EngineShowFlags, bool& bDrawSimpleCollision, bool& bDrawComplexCollision) const;
-
-	/** Only call on render thread timeline */
-	uint8 GetCurrentFirstLODIdx_Internal() const
-	{
-		return RenderData->CurrentFirstLODIdx;
-	}
-
-public:
-	// FPrimitiveSceneProxy interface.
-#if WITH_EDITOR
-	virtual HHitProxy* CreateHitProxies(UPrimitiveComponent* Component, TArray<TRefCountPtr<HHitProxy> >& OutHitProxies) override;
-#endif
-	virtual void DrawStaticElements(FStaticPrimitiveDrawInterface* PDI) override;
-	virtual int32 GetLOD(const FSceneView* View) const override;
-	virtual FPrimitiveViewRelevance GetViewRelevance(const FSceneView* View) const override;
-	virtual bool CanBeOccluded() const override;
-	virtual bool IsUsingDistanceCullFade() const override;
-	virtual void GetLightRelevance(const FLightSceneProxy* LightSceneProxy, bool& bDynamic, bool& bRelevant, bool& bLightMapped, bool& bShadowMapped) const override;
-	virtual void GetDistanceFieldAtlasData(const FDistanceFieldVolumeData*& OutDistanceFieldData, float& SelfShadowBias) const override;
-	virtual void GetDistanceFieldInstanceData(TArray<FRenderTransform>& InstanceLocalToPrimitiveTransforms) const override;
-	virtual bool HasDistanceFieldRepresentation() const override;
-	virtual bool HasDynamicIndirectShadowCasterRepresentation() const override;
-	virtual uint32 GetMemoryFootprint( void ) const override { return( sizeof( *this ) + GetAllocatedSize() ); }
-	SIZE_T GetAllocatedSize( void ) const { return( FPrimitiveSceneProxy::GetAllocatedSize() + LODs.GetAllocatedSize() ); }
-
-	virtual void GetMeshDescription(int32 LODIndex, TArray<FMeshBatch>& OutMeshElements) const override;
-
-	virtual void GetDynamicMeshElements(const TArray<const FSceneView*>& Views, const FSceneViewFamily& ViewFamily, uint32 VisibilityMap, FMeshElementCollector& Collector) const override;
-
-	virtual const FCardRepresentationData* GetMeshCardRepresentation() const override;
-
-#if RHI_RAYTRACING
-	virtual void GetDynamicRayTracingInstances(FRayTracingMaterialGatheringContext& Context, TArray<FRayTracingInstance>& OutRayTracingInstances) override;
-	virtual bool HasRayTracingRepresentation() const override;
-	virtual bool IsRayTracingRelevant() const override { return true; }
-	virtual bool IsRayTracingStaticRelevant() const override 
-	{ 
-		const bool bAllowStaticLighting = FReadOnlyCVARCache::Get().bAllowStaticLighting;
-		const bool bIsStaticInstance = !bDynamicRayTracingGeometry;
-		return bIsStaticInstance && !HasViewDependentDPG() && !(bAllowStaticLighting && HasStaticLighting() && !HasValidSettingsForStaticLighting());
-	}
-#endif // RHI_RAYTRACING
-
-	virtual void GetLCIs(FLCIArray& LCIs) override;
-
-#if WITH_EDITORONLY_DATA
-	virtual bool GetPrimitiveDistance(int32 LODIndex, int32 SectionIndex, const FVector& ViewOrigin, float& PrimitiveDistance) const override;
-	virtual bool GetMeshUVDensities(int32 LODIndex, int32 SectionIndex, FVector4& WorldUVDensities) const override;
-	virtual bool GetMaterialTextureScales(int32 LODIndex, int32 SectionIndex, const FMaterialRenderProxy* MaterialRenderProxy, FVector4f* OneOverScales, FIntVector4* UVChannelIndices) const override;
-#endif
-
-#if STATICMESH_ENABLE_DEBUG_RENDERING
-	virtual int32 GetLightMapResolution() const override { return LightMapResolution; }
-#endif
-
-protected:
-	/** Information used by the proxy about a single LOD of the mesh. */
-	class FLODInfo : public FLightCacheInterface
-	{
-	public:
-
-		/** Information about an element of a LOD. */
-		struct FSectionInfo
-		{
-			/** Default constructor. */
-			FSectionInfo()
-				: Material(nullptr)
-			#if WITH_EDITOR
-				, bSelected(false)
-				, HitProxy(nullptr)
-			#endif
-				, FirstPreCulledIndex(0)
-				, NumPreCulledTriangles(-1)
-			{}
-
-			/** The material with which to render this section. */
-			UMaterialInterface* Material;
-
-		#if WITH_EDITOR
-			/** True if this section should be rendered as selected (editor only). */
-			bool bSelected;
-
-			/** The editor needs to be able to individual sub-mesh hit detection, so we store a hit proxy on each mesh. */
-			HHitProxy* HitProxy;
-		#endif
-
-		#if WITH_EDITORONLY_DATA
-			// The material index from the component. Used by the texture streaming accuracy viewmodes.
-			int32 MaterialIndex;
-		#endif
-
-			int32 FirstPreCulledIndex;
-			int32 NumPreCulledTriangles;
-		};
-
-		/** Per-section information. */
-		TArray<FSectionInfo, TInlineAllocator<1>> Sections;
-
-		/** Vertex color data for this LOD (or NULL when not overridden), FStaticMeshComponentLODInfo handle the release of the memory */
-		FColorVertexBuffer* OverrideColorVertexBuffer;
-
-		TUniformBufferRef<FLocalVertexFactoryUniformShaderParameters> OverrideColorVFUniformBuffer;
-
-		const FRawStaticIndexBuffer* PreCulledIndexBuffer;
-
-		/** Initialization constructor. */
-		FLODInfo(const UStaticMeshComponent* InComponent, const FStaticMeshVertexFactoriesArray& InLODVertexFactories, int32 InLODIndex, int32 InClampedMinLOD, bool bLODsShareStaticLighting);
-
-		bool UsesMeshModifyingMaterials() const { return bUsesMeshModifyingMaterials; }
-
-		// FLightCacheInterface.
-		virtual FLightInteraction GetInteraction(const FLightSceneProxy* LightSceneProxy) const override;
-
-	private:
-		TArray<FGuid> IrrelevantLights;
-
-		/** True if any elements in this LOD use mesh-modifying materials **/
-		bool bUsesMeshModifyingMaterials;
-	};
-
-	FStaticMeshRenderData* RenderData;
-
-	TArray<FLODInfo> LODs;
-
-	const FDistanceFieldVolumeData* DistanceFieldData;	
-	const FCardRepresentationData* CardRepresentationData;	
-
-	UMaterialInterface* OverlayMaterial;
-	float OverlayMaterialMaxDrawDistance;
-
-#if RHI_RAYTRACING
-	bool bSupportRayTracing;
-	bool bDynamicRayTracingGeometry;
-	TArray<FRayTracingGeometry, TInlineAllocator<MAX_MESH_LOD_COUNT>> DynamicRayTracingGeometries;
-	TArray<FRWBuffer, TInlineAllocator<MAX_MESH_LOD_COUNT>> DynamicRayTracingGeometryVertexBuffers;
-	TArray<FMeshBatch> CachedRayTracingMaterials;
-	int16 CachedRayTracingMaterialsLODIndex = INDEX_NONE;
-	FRayTracingMaskAndFlags CachedRayTracingInstanceMaskAndFlags;
-#endif
-	/**
-	 * The forcedLOD set in the static mesh editor, copied from the mesh component
-	 */
-	int32 ForcedLodModel;
-
-	/** Minimum LOD index to use.  Clamped to valid range [0, NumLODs - 1]. */
-	int32 ClampedMinLOD;
-
-	uint32 bCastShadow : 1;
-
-	/** This primitive has culling reversed */
-	uint32 bReverseCulling : 1;
-
-	/** The view relevance for all the static mesh's materials. */
-	FMaterialRelevance MaterialRelevance;
-
-#if WITH_EDITORONLY_DATA
-	/** The component streaming distance multiplier */
-	float StreamingDistanceMultiplier;
-	/** The cached GetTextureStreamingTransformScale */
-	float StreamingTransformScale;
-	/** Material bounds used for texture streaming. */
-	TArray<uint32> MaterialStreamingRelativeBoxes;
-
-	/** Index of the section to preview. If set to INDEX_NONE, all section will be rendered */
-	int32 SectionIndexPreview;
-	/** Index of the material to preview. If set to INDEX_NONE, all section will be rendered */
-	int32 MaterialIndexPreview;
-
-	/** Whether selection should be per section or per entire proxy. */
-	bool bPerSectionSelection;
-#endif
-
-private:
-
-	const UStaticMesh* StaticMesh;
-
-#if STATICMESH_ENABLE_DEBUG_RENDERING
-	AActor* Owner;
-	/** LightMap resolution used for VMI_LightmapDensity */
-	int32 LightMapResolution;
-	/** Body setup for collision debug rendering */
-	UBodySetup* BodySetup;
-	/** Collision trace flags */
-	ECollisionTraceFlag		CollisionTraceFlag;
-	/** Collision Response of this component */
-	FCollisionResponseContainer CollisionResponse;
-	/** LOD used for collision */
-	int32 LODForCollision;
-	/** Draw mesh collision if used for complex collision */
-	uint32 bDrawMeshCollisionIfComplex : 1;
-	/** Draw mesh collision if used for simple collision */
-	uint32 bDrawMeshCollisionIfSimple : 1;
-
-protected:
-	/** Hierarchical LOD Index used for rendering */
-	uint8 HierarchicalLODIndex;
-#endif
-
-public:
-
-	/**
-	 * Returns the display factor for the given LOD level
-	 *
-	 * @Param LODIndex - The LOD to get the display factor for
-	 */
-	float GetScreenSize(int32 LODIndex) const;
-
-	/**
-	 * Returns the LOD mask for a view, this is like the ordinary LOD but can return two values for dither fading
-	 */
-	FLODMask GetLODMask(const FSceneView* View) const;
-
-private:
-	void AddSpeedTreeWind();
-	void RemoveSpeedTreeWind();
-};
-
 /*-----------------------------------------------------------------------------
 	FStaticMeshInstanceData
 -----------------------------------------------------------------------------*/
diff --git a/Engine/Source/Runtime/Engine/Public/StaticMeshSceneProxy.h b/Engine/Source/Runtime/Engine/Public/StaticMeshSceneProxy.h
--- a/Engine/Source/Runtime/Engine/Public/StaticMeshSceneProxy.h
+++ b/Engine/Source/Runtime/Engine/Public/StaticMeshSceneProxy.h
@@ -1,5 +1,305 @@
 // Copyright Epic Games, Inc. All Rights Reserved.
 
+/*=============================================================================
+	StaticMesh.h: Static mesh class definition.
+=============================================================================*/
+
 #pragma once
 
-// Placeholder for upcoming change
\ No newline at end of file
+#include "Components/StaticMeshComponent.h"
+#include "PrimitiveSceneProxy.h"
+#include "RayTracingInstance.h"
+
+class FLocalVertexFactoryUniformShaderParameters;
+class FRawStaticIndexBuffer;
+struct FStaticMeshVertexFactories;
+using FStaticMeshVertexFactoriesArray = TArray<FStaticMeshVertexFactories>;
+
+/**
+ * A static mesh component scene proxy.
+ */
+class ENGINE_API FStaticMeshSceneProxy : public FPrimitiveSceneProxy
+{
+public:
+	SIZE_T GetTypeHash() const override;
+
+	/** Initialization constructor. */
+	FStaticMeshSceneProxy(UStaticMeshComponent* Component, bool bForceLODsShareStaticLighting);
+
+	virtual ~FStaticMeshSceneProxy();
+
+	/** Gets the number of mesh batches required to represent the proxy, aside from section needs. */
+	virtual int32 GetNumMeshBatches() const
+	{
+		return 1;
+	}
+
+	/** Sets up a shadow FMeshBatch for a specific LOD. */
+	virtual bool GetShadowMeshElement(int32 LODIndex, int32 BatchIndex, uint8 InDepthPriorityGroup, FMeshBatch& OutMeshBatch, bool bDitheredLODTransition) const;
+
+	/** Sets up a FMeshBatch for a specific LOD and element. */
+	virtual bool GetMeshElement(
+		int32 LODIndex, 
+		int32 BatchIndex, 
+		int32 ElementIndex, 
+		uint8 InDepthPriorityGroup, 
+		bool bUseSelectionOutline,
+		bool bAllowPreCulledIndices,
+		FMeshBatch& OutMeshBatch) const;
+
+	virtual void CreateRenderThreadResources() override;
+
+	virtual void DestroyRenderThreadResources() override;
+
+	/** Sets up a wireframe FMeshBatch for a specific LOD. */
+	virtual bool GetWireframeMeshElement(int32 LODIndex, int32 BatchIndex, const FMaterialRenderProxy* WireframeRenderProxy, uint8 InDepthPriorityGroup, bool bAllowPreCulledIndices, FMeshBatch& OutMeshBatch) const;
+
+	/** Sets up a collision FMeshBatch for a specific LOD and element. */
+	virtual bool GetCollisionMeshElement(
+		int32 LODIndex,
+		int32 BatchIndex,
+		int32 ElementIndex,
+		uint8 InDepthPriorityGroup,
+		const FMaterialRenderProxy* RenderProxy,
+		FMeshBatch& OutMeshBatch) const;
+
+	virtual void SetEvaluateWorldPositionOffsetInRayTracing(bool NewValue);
+
+	virtual uint8 GetCurrentFirstLODIdx_RenderThread() const final override
+	{
+		return GetCurrentFirstLODIdx_Internal();
+	}
+
+	virtual int32 GetLightMapCoordinateIndex() const override;
+
+protected:
+	/** Configures mesh batch vertex / index state. Returns the number of primitives used in the element. */
+	uint32 SetMeshElementGeometrySource(
+		int32 LODIndex,
+		int32 ElementIndex,
+		bool bWireframe,
+		bool bUseInversedIndices,
+		bool bAllowPreCulledIndices,
+		const FVertexFactory* VertexFactory,
+		FMeshBatch& OutMeshElement) const;
+
+	/** Sets the screen size on a mesh element. */
+	void SetMeshElementScreenSize(int32 LODIndex, bool bDitheredLODTransition, FMeshBatch& OutMeshBatch) const;
+
+	/** Returns whether this mesh needs reverse culling when using reversed indices. */
+	bool IsReversedCullingNeeded(bool bUseReversedIndices) const;
+
+	bool IsCollisionView(const FEngineShowFlags& EngineShowFlags, bool& bDrawSimpleCollision, bool& bDrawComplexCollision) const;
+
+	/** Only call on render thread timeline */
+	uint8 GetCurrentFirstLODIdx_Internal() const;
+
+public:
+	// FPrimitiveSceneProxy interface.
+#if WITH_EDITOR
+	virtual HHitProxy* CreateHitProxies(UPrimitiveComponent* Component, TArray<TRefCountPtr<HHitProxy> >& OutHitProxies) override;
+#endif
+	virtual void DrawStaticElements(FStaticPrimitiveDrawInterface* PDI) override;
+	virtual int32 GetLOD(const FSceneView* View) const override;
+	virtual FPrimitiveViewRelevance GetViewRelevance(const FSceneView* View) const override;
+	virtual bool CanBeOccluded() const override;
+	virtual bool IsUsingDistanceCullFade() const override;
+	virtual void GetLightRelevance(const FLightSceneProxy* LightSceneProxy, bool& bDynamic, bool& bRelevant, bool& bLightMapped, bool& bShadowMapped) const override;
+	virtual void GetDistanceFieldAtlasData(const FDistanceFieldVolumeData*& OutDistanceFieldData, float& SelfShadowBias) const override;
+	virtual void GetDistanceFieldInstanceData(TArray<FRenderTransform>& InstanceLocalToPrimitiveTransforms) const override;
+	virtual bool HasDistanceFieldRepresentation() const override;
+	virtual bool HasDynamicIndirectShadowCasterRepresentation() const override;
+	virtual uint32 GetMemoryFootprint( void ) const override { return( sizeof( *this ) + GetAllocatedSize() ); }
+	SIZE_T GetAllocatedSize( void ) const { return( FPrimitiveSceneProxy::GetAllocatedSize() + LODs.GetAllocatedSize() ); }
+
+	virtual void GetMeshDescription(int32 LODIndex, TArray<FMeshBatch>& OutMeshElements) const override;
+
+	virtual void GetDynamicMeshElements(const TArray<const FSceneView*>& Views, const FSceneViewFamily& ViewFamily, uint32 VisibilityMap, FMeshElementCollector& Collector) const override;
+
+	virtual const FCardRepresentationData* GetMeshCardRepresentation() const override;
+
+#if RHI_RAYTRACING
+	virtual void GetDynamicRayTracingInstances(FRayTracingMaterialGatheringContext& Context, TArray<FRayTracingInstance>& OutRayTracingInstances) override;
+	virtual bool HasRayTracingRepresentation() const override;
+	virtual bool IsRayTracingRelevant() const override { return true; }
+	virtual bool IsRayTracingStaticRelevant() const override 
+	{ 
+		const bool bAllowStaticLighting = FReadOnlyCVARCache::Get().bAllowStaticLighting;
+		const bool bIsStaticInstance = !bDynamicRayTracingGeometry;
+		return bIsStaticInstance && !HasViewDependentDPG() && !(bAllowStaticLighting && HasStaticLighting() && !HasValidSettingsForStaticLighting());
+	}
+#endif // RHI_RAYTRACING
+
+	virtual void GetLCIs(FLCIArray& LCIs) override;
+
+#if WITH_EDITORONLY_DATA
+	virtual bool GetPrimitiveDistance(int32 LODIndex, int32 SectionIndex, const FVector& ViewOrigin, float& PrimitiveDistance) const override;
+	virtual bool GetMeshUVDensities(int32 LODIndex, int32 SectionIndex, FVector4& WorldUVDensities) const override;
+	virtual bool GetMaterialTextureScales(int32 LODIndex, int32 SectionIndex, const FMaterialRenderProxy* MaterialRenderProxy, FVector4f* OneOverScales, FIntVector4* UVChannelIndices) const override;
+#endif
+
+#if STATICMESH_ENABLE_DEBUG_RENDERING
+	virtual int32 GetLightMapResolution() const override { return LightMapResolution; }
+#endif
+
+protected:
+	/** Information used by the proxy about a single LOD of the mesh. */
+	class FLODInfo : public FLightCacheInterface
+	{
+	public:
+
+		/** Information about an element of a LOD. */
+		struct FSectionInfo
+		{
+			/** Default constructor. */
+			FSectionInfo()
+				: Material(nullptr)
+			#if WITH_EDITOR
+				, bSelected(false)
+				, HitProxy(nullptr)
+			#endif
+				, FirstPreCulledIndex(0)
+				, NumPreCulledTriangles(-1)
+			{}
+
+			/** The material with which to render this section. */
+			UMaterialInterface* Material;
+
+		#if WITH_EDITOR
+			/** True if this section should be rendered as selected (editor only). */
+			bool bSelected;
+
+			/** The editor needs to be able to individual sub-mesh hit detection, so we store a hit proxy on each mesh. */
+			HHitProxy* HitProxy;
+		#endif
+
+		#if WITH_EDITORONLY_DATA
+			// The material index from the component. Used by the texture streaming accuracy viewmodes.
+			int32 MaterialIndex;
+		#endif
+
+			int32 FirstPreCulledIndex;
+			int32 NumPreCulledTriangles;
+		};
+
+		/** Per-section information. */
+		TArray<FSectionInfo, TInlineAllocator<1>> Sections;
+
+		/** Vertex color data for this LOD (or NULL when not overridden), FStaticMeshComponentLODInfo handle the release of the memory */
+		FColorVertexBuffer* OverrideColorVertexBuffer;
+
+		TUniformBufferRef<FLocalVertexFactoryUniformShaderParameters> OverrideColorVFUniformBuffer;
+
+		const FRawStaticIndexBuffer* PreCulledIndexBuffer;
+
+		/** Initialization constructor. */
+		FLODInfo(const UStaticMeshComponent* InComponent, const FStaticMeshVertexFactoriesArray& InLODVertexFactories, int32 InLODIndex, int32 InClampedMinLOD, bool bLODsShareStaticLighting);
+
+		bool UsesMeshModifyingMaterials() const { return bUsesMeshModifyingMaterials; }
+
+		// FLightCacheInterface.
+		virtual FLightInteraction GetInteraction(const FLightSceneProxy* LightSceneProxy) const override;
+
+	private:
+		TArray<FGuid> IrrelevantLights;
+
+		/** True if any elements in this LOD use mesh-modifying materials **/
+		bool bUsesMeshModifyingMaterials;
+	};
+
+	FStaticMeshRenderData* RenderData;
+
+	TArray<FLODInfo> LODs;
+
+	const FDistanceFieldVolumeData* DistanceFieldData;	
+	const FCardRepresentationData* CardRepresentationData;	
+
+	UMaterialInterface* OverlayMaterial;
+	float OverlayMaterialMaxDrawDistance;
+
+#if RHI_RAYTRACING
+	bool bSupportRayTracing;
+	bool bDynamicRayTracingGeometry;
+	TArray<FRayTracingGeometry, TInlineAllocator<MAX_MESH_LOD_COUNT>> DynamicRayTracingGeometries;
+	TArray<FRWBuffer, TInlineAllocator<MAX_MESH_LOD_COUNT>> DynamicRayTracingGeometryVertexBuffers;
+	TArray<FMeshBatch> CachedRayTracingMaterials;
+	int16 CachedRayTracingMaterialsLODIndex = INDEX_NONE;
+	FRayTracingMaskAndFlags CachedRayTracingInstanceMaskAndFlags;
+#endif
+	/**
+	 * The forcedLOD set in the static mesh editor, copied from the mesh component
+	 */
+	int32 ForcedLodModel;
+
+	/** Minimum LOD index to use.  Clamped to valid range [0, NumLODs - 1]. */
+	int32 ClampedMinLOD;
+
+	uint32 bCastShadow : 1;
+
+	/** This primitive has culling reversed */
+	uint32 bReverseCulling : 1;
+
+	/** The view relevance for all the static mesh's materials. */
+	FMaterialRelevance MaterialRelevance;
+
+#if WITH_EDITORONLY_DATA
+	/** The component streaming distance multiplier */
+	float StreamingDistanceMultiplier;
+	/** The cached GetTextureStreamingTransformScale */
+	float StreamingTransformScale;
+	/** Material bounds used for texture streaming. */
+	TArray<uint32> MaterialStreamingRelativeBoxes;
+
+	/** Index of the section to preview. If set to INDEX_NONE, all section will be rendered */
+	int32 SectionIndexPreview;
+	/** Index of the material to preview. If set to INDEX_NONE, all section will be rendered */
+	int32 MaterialIndexPreview;
+
+	/** Whether selection should be per section or per entire proxy. */
+	bool bPerSectionSelection;
+#endif
+
+private:
+
+	const UStaticMesh* StaticMesh;
+
+#if STATICMESH_ENABLE_DEBUG_RENDERING
+	AActor* Owner;
+	/** LightMap resolution used for VMI_LightmapDensity */
+	int32 LightMapResolution;
+	/** Body setup for collision debug rendering */
+	UBodySetup* BodySetup;
+	/** Collision trace flags */
+	ECollisionTraceFlag		CollisionTraceFlag;
+	/** Collision Response of this component */
+	FCollisionResponseContainer CollisionResponse;
+	/** LOD used for collision */
+	int32 LODForCollision;
+	/** Draw mesh collision if used for complex collision */
+	uint32 bDrawMeshCollisionIfComplex : 1;
+	/** Draw mesh collision if used for simple collision */
+	uint32 bDrawMeshCollisionIfSimple : 1;
+
+protected:
+	/** Hierarchical LOD Index used for rendering */
+	uint8 HierarchicalLODIndex;
+#endif
+
+public:
+
+	/**
+	 * Returns the display factor for the given LOD level
+	 *
+	 * @Param LODIndex - The LOD to get the display factor for
+	 */
+	float GetScreenSize(int32 LODIndex) const;
+
+	/**
+	 * Returns the LOD mask for a view, this is like the ordinary LOD but can return two values for dither fading
+	 */
+	FLODMask GetLODMask(const FSceneView* View) const;
+
+private:
+	void AddSpeedTreeWind();
+	void RemoveSpeedTreeWind();
+};
diff --git a/Engine/Source/Runtime/Engine/Public/TextureReferenceResolver.h b/Engine/Source/Runtime/Engine/Public/TextureReferenceResolver.h
--- a/Engine/Source/Runtime/Engine/Public/TextureReferenceResolver.h
+++ b/Engine/Source/Runtime/Engine/Public/TextureReferenceResolver.h
@@ -2,4 +2,35 @@
 
 #pragma once
 
-// Placeholder for upcoming change
\ No newline at end of file
+#include "Engine/Texture.h"
+
+/** 
+* Replaces the RHI reference of one texture with another.
+* Allows one texture to be replaced with another at runtime and have all existing references to it remain valid.
+*/
+struct FTextureReferenceReplacer
+{
+	FTextureReferenceRHIRef OriginalRef;
+
+	FTextureReferenceReplacer(UTexture* OriginalTexture)
+	{
+		if (OriginalTexture)
+		{
+			OriginalTexture->ReleaseResource();
+			OriginalRef = OriginalTexture->TextureReference.TextureReferenceRHI;
+		}
+		else
+		{
+			OriginalRef = nullptr;
+		}
+	}
+
+	void Replace(UTexture* NewTexture)
+	{
+		if (OriginalRef)
+		{
+			NewTexture->TextureReference.TextureReferenceRHI = OriginalRef;
+		}
+	}
+};
+ 
\ No newline at end of file
diff --git a/Engine/Source/Runtime/Engine/Public/UnrealClient.h b/Engine/Source/Runtime/Engine/Public/UnrealClient.h
--- a/Engine/Source/Runtime/Engine/Public/UnrealClient.h
+++ b/Engine/Source/Runtime/Engine/Public/UnrealClient.h
@@ -7,20 +7,25 @@
 #pragma once
 
 #include "CoreMinimal.h"
-#include "InputCoreTypes.h"
+#include "InputKeyEventArgs.h"
 #include "Engine/EngineBaseTypes.h"
-#include "Input/PopupMethodReply.h"
-#include "Widgets/SWidget.h"
 #include "UObject/GCObject.h"
-#include "RHI.h"
 #include "RenderResource.h"
 #include "HitProxies.h"
 #include "RenderGraphDefinitions.h"
 #include "Elements/Framework/TypedElementListFwd.h"
 #include "DynamicRenderScaling.h"
+
+#if UE_ENABLE_INCLUDE_ORDER_DEPRECATED_IN_5_2
+#include "InputCoreTypes.h"
+#include "Input/PopupMethodReply.h"
+#include "Widgets/SWidget.h"
+#include "RHI.h"
 #include "GenericPlatform/GenericPlatformInputDeviceMapper.h"
+#endif
 
 class FCanvas;
+class FRDGBuilder;
 class FViewport;
 class FViewportClient;
 class UModel;
@@ -790,508 +795,5 @@ extern ENGINE_API bool GetViewportScreenShot(FViewport* Viewport, TArray<FColor>
 extern ENGINE_API bool GetViewportScreenShotHDR(FViewport* Viewport, TArray<FLinearColor>& Bitmap, const FIntRect& ViewRect = FIntRect());
 extern ENGINE_API bool GetHighResScreenShotInput(const TCHAR* Cmd, FOutputDevice& Ar, uint32& OutXRes, uint32& OutYRes, float& OutResMult, FIntRect& OutCaptureRegion, bool& OutShouldEnableMask, bool& OutDumpBufferVisualizationTargets, bool& OutCaptureHDR, FString& OutFilenameOverride, bool& OutUseDateTimeAsFileName);
 
-struct FInputKeyEventArgs
-{
-public:
-	FInputKeyEventArgs(FViewport* InViewport, int32 InControllerId, FKey InKey, EInputEvent InEvent)
-		: Viewport(InViewport)
-		, ControllerId(InControllerId)
-		, InputDevice(FInputDeviceId::CreateFromInternalId(InControllerId))
-		, Key(InKey)
-		, Event(InEvent)
-		, AmountDepressed(1.0f)
-		, bIsTouchEvent(false)
-	{
-	}
-
-	FInputKeyEventArgs(FViewport* InViewport, FInputDeviceId InInputDevice, FKey InKey, EInputEvent InEvent)
-		: Viewport(InViewport)
-		, InputDevice(InInputDevice)
-		, Key(InKey)
-		, Event(InEvent)
-		, AmountDepressed(1.0f)
-		, bIsTouchEvent(false)
-	{
-		FPlatformUserId UserID = IPlatformInputDeviceMapper::Get().GetUserForInputDevice(InInputDevice);
-		ControllerId = FPlatformMisc::GetUserIndexForPlatformUser(UserID);
-	}
-
-	FInputKeyEventArgs(FViewport* InViewport, int32 InControllerId, FKey InKey, EInputEvent InEvent, float InAmountDepressed, bool bInIsTouchEvent)
-		: Viewport(InViewport)
-		, ControllerId(InControllerId)
-		, InputDevice(FInputDeviceId::CreateFromInternalId(InControllerId))
-		, Key(InKey)
-		, Event(InEvent)
-		, AmountDepressed(InAmountDepressed)
-		, bIsTouchEvent(bInIsTouchEvent)
-	{
-	}
-	
-	FInputKeyEventArgs(FViewport* InViewport, FInputDeviceId InInputDevice, FKey InKey, EInputEvent InEvent, float InAmountDepressed, bool bInIsTouchEvent)
-		: Viewport(InViewport)
-		, InputDevice(InInputDevice)
-		, Key(InKey)
-		, Event(InEvent)
-		, AmountDepressed(InAmountDepressed)
-		, bIsTouchEvent(bInIsTouchEvent)
-	{
-		FPlatformUserId UserID = IPlatformInputDeviceMapper::Get().GetUserForInputDevice(InInputDevice);
-		ControllerId = FPlatformMisc::GetUserIndexForPlatformUser(UserID);
-	}
-
-	bool IsGamepad() const { return Key.IsGamepadKey(); }
-
-public:
-
-	// The viewport which the key event is from.
-	FViewport* Viewport;
-	// The controller which the key event is from.
-	int32 ControllerId;
-	// The input device which this key event is from
-	FInputDeviceId InputDevice;
-	// The type of event which occurred.
-	FKey Key;
-	// The type of event which occurred.
-	EInputEvent Event;
-	// For analog keys, the depression percent.
-	float AmountDepressed;
-	// input came from a touch surface.This may be a faked mouse button from touch.
-	bool bIsTouchEvent;
-};
-
-/**
- * An abstract interface to a viewport's client.
- * The viewport's client processes input received by the viewport, and draws the viewport.
- */
-class FViewportClient
-{
-public:
-	virtual ~FViewportClient(){}
-	virtual void Precache() {}
-	virtual void RedrawRequested(FViewport* Viewport) { Viewport->Draw(); }
-	virtual void RequestInvalidateHitProxy(FViewport* Viewport) { Viewport->InvalidateHitProxy(); }
-	virtual void Draw(FViewport* Viewport,FCanvas* Canvas) {}
-	virtual bool ProcessScreenShots(FViewport* Viewport) { return false; }
-	virtual UWorld* GetWorld() const { return NULL; }
-	virtual struct FEngineShowFlags* GetEngineShowFlags() { return NULL; }
-
-	/**
-	 * Check a key event received by the viewport.
-	 * If the viewport client uses the event, it should return true to consume it.
-	 * @param	EventArgs - The Input event args.
-	 * @return	True to consume the key event, false to pass it on.
-	 */
-	virtual bool InputKey(const FInputKeyEventArgs& EventArgs)
-	{
-		PRAGMA_DISABLE_DEPRECATION_WARNINGS
-		return InputKey(EventArgs.Viewport, EventArgs.ControllerId, EventArgs.Key, EventArgs.Event, EventArgs.AmountDepressed, EventArgs.Key.IsGamepadKey());
-		PRAGMA_ENABLE_DEPRECATION_WARNINGS
-	}
-
-	UE_DEPRECATED(5.1, "Use the new InputKey(const FInputKeyEventArgs& EventArgs) function.")
-	virtual bool InputKey(FViewport* Viewport,int32 ControllerId,FKey Key,EInputEvent Event,float AmountDepressed = 1.f,bool bGamepad=false) { return false; }
-
-	/**
-	 * Check an axis movement received by the viewport.
-	 * If the viewport client uses the movement, it should return true to consume it.
-	 * @param	Viewport - The viewport which the axis movement is from.
-	 * @param	ControllerId - The controller which the axis movement is from.
-	 * @param	Key - The name of the axis which moved.
-	 * @param	Delta - The axis movement delta.
-	 * @param	DeltaTime - The time since the last axis update.
-	 * @param	NumSamples - The number of device samples that contributed to this Delta, useful for things like smoothing
-	 * @param	bGamepad - input came from gamepad (ie xbox controller)
-	 * @return	True to consume the axis movement, false to pass it on.
-	 */
-	UE_DEPRECATED(5.1, "Use the new InputAxis function that takes a FInputDeviceId.")
-	virtual bool InputAxis(FViewport* Viewport,int32 ControllerId,FKey Key,float Delta,float DeltaTime,int32 NumSamples=1,bool bGamepad=false) { return false; }
-
-	/**
-	 * Check an axis movement received by the viewport.
-	 * If the viewport client uses the movement, it should return true to consume it.
-	 * @param	Viewport - The viewport which the axis movement is from.
-	 * @param	InputDevice - The input device that triggered this axis movement
-	 * @param	Key - The name of the axis which moved.
-	 * @param	Delta - The axis movement delta.
-	 * @param	DeltaTime - The time since the last axis update.
-	 * @param	NumSamples - The number of device samples that contributed to this Delta, useful for things like smoothing
-	 * @param	bGamepad - input came from gamepad (ie xbox controller)
-	 * @return	True to consume the axis movement, false to pass it on.
-	 */
-	virtual bool InputAxis(FViewport* Viewport, FInputDeviceId InputDevice, FKey Key, float Delta, float DeltaTime, int32 NumSamples = 1, bool bGamepad = false) { return false; }
-
-	/**
-	 * Check a character input received by the viewport.
-	 * If the viewport client uses the character, it should return true to consume it.
-	 * @param	Viewport - The viewport which the axis movement is from.
-	 * @param	ControllerId - The controller which the axis movement is from.
-	 * @param	Character - The character.
-	 * @return	True to consume the character, false to pass it on.
-	 */
-	virtual bool InputChar(FViewport* Viewport,int32 ControllerId,TCHAR Character) { return false; }
-
-	/**
-	 * Check a key event received by the viewport.
-	 * If the viewport client uses the event, it should return true to consume it.
-	 * @param	Viewport - The viewport which the event is from.
-	 * @param	ControllerId - The controller which the key event is from.
-	 * @param	Handle - Identifier unique to this touch event
-	 * @param	Type - What kind of touch event this is (see ETouchType)
-	 * @param	TouchLocation - Screen position of the touch
-	 * @param	Force - How hard the touch is
-	 * @param	DeviceTimestamp - Timestamp of the event
-	 * @param	TouchpadIndex - For devices with multiple touchpads, this is the index of which one
-	 * @return	True to consume the key event, false to pass it on.
-	 */
-	virtual bool InputTouch(FViewport* Viewport, int32 ControllerId, uint32 Handle, ETouchType::Type Type, const FVector2D& TouchLocation, float Force, FDateTime DeviceTimestamp, uint32 TouchpadIndex) { return false; }
-
-	UE_DEPRECATED(4.20, "InputTouch now takes a Force")
-	bool InputTouch(FViewport* Viewport, int32 ControllerId, uint32 Handle, ETouchType::Type Type, const FVector2D& TouchLocation, FDateTime DeviceTimestamp, uint32 TouchpadIndex) { return InputTouch(Viewport, ControllerId, Handle, Type, TouchLocation, 1.0f, DeviceTimestamp, TouchpadIndex); }
-
-	/**
-	 * Check a gesture event received by the viewport.
-	 * If the viewport client uses the event, it should return true to consume it.
-	 * @param	Viewport - The viewport which the event is from.
-	 * @param	GestureType - @todo desc
-	 * @param	GestureDelta - @todo desc
-	 * @return	True to consume the gesture event, false to pass it on.
-	 */
-	virtual bool InputGesture(FViewport* Viewport, EGestureEvent GestureType, const FVector2D& GestureDelta, bool bIsDirectionInvertedFromDevice) { return false; }
-
-	/**
-	 * Each frame, the input system will update the motion data.
-	 *
-	 * @param Viewport - The viewport which the key event is from.
-	 * @param ControllerId - The controller which the key event is from.
-	 * @param Tilt			The current orientation of the device
-	 * @param RotationRate	How fast the tilt is changing
-	 * @param Gravity		Describes the current gravity of the device
-	 * @param Acceleration  Describes the acceleration of the device
-	 * @return	True to consume the motion event, false to pass it on.
-	 */
-	virtual bool InputMotion(FViewport* Viewport, int32 ControllerId, const FVector& Tilt, const FVector& RotationRate, const FVector& Gravity, const FVector& Acceleration) { return false; }
-
-	virtual void SetIsSimulateInEditorViewport(bool bInIsSimulateInEditorViewport) { };
-
-	virtual bool WantsPollingMouseMovement(void) const { return true; }
-
-	virtual void MouseEnter( FViewport* Viewport,int32 x, int32 y ) {}
-
-	virtual void MouseLeave( FViewport* Viewport ) {}
-
-	virtual void MouseMove(FViewport* Viewport,int32 X,int32 Y) {}
-
-	/**
-	 * Called when the mouse is moved while a window input capture is in effect
-	 *
-	 * @param	InViewport	Viewport that captured the mouse input
-	 * @param	InMouseX	New mouse cursor X coordinate
-	 * @param	InMouseY	New mouse cursor Y coordinate
-	 */
-	virtual void CapturedMouseMove( FViewport* InViewport, int32 InMouseX, int32 InMouseY ) { }
-
-	/**
-	 * Called from slate when input is finished for this frame, and we should process any queued mouse moves.
-	 */
-	virtual void ProcessAccumulatedPointerInput(FViewport* InViewport) {};
-
-	/**
-	 * Retrieves the cursor that should be displayed by the OS
-	 *
-	 * @param	Viewport	the viewport that contains the cursor
-	 * @param	X			the x position of the cursor
-	 * @param	Y			the Y position of the cursor
-	 *
-	 * @return	the cursor that the OS should display
-	 */
-	virtual EMouseCursor::Type GetCursor(FViewport* Viewport, int32 X,int32 Y) { return EMouseCursor::Default; }
-
-	/**
-	 * Called to map a cursor reply to an actual widget to render.
-	 *
-	 * @return	the widget that should be rendered for this cursor, return TOptional<TSharedRef<SWidget>>() if no mapping.
-	 */
-	virtual TOptional<TSharedRef<SWidget>> MapCursor(FViewport* Viewport, const FCursorReply& CursorReply) { return TOptional<TSharedRef<SWidget>>(); }
-
-	/**
-	 * Called to determine if we should render the focus brush.
-	 *
-	 * @param InFocusCause	The cause of focus
-	 */
-	virtual TOptional<bool> QueryShowFocus(const EFocusCause InFocusCause) const { return TOptional<bool>(); }
-
-	virtual void LostFocus(FViewport* Viewport) {}
-	virtual void ReceivedFocus(FViewport* Viewport) {}
-	virtual bool IsFocused(FViewport* Viewport) { return true; }
-
-	virtual void Activated(FViewport* Viewport, const FWindowActivateEvent& InActivateEvent) {}
-	virtual void Deactivated(FViewport* Viewport, const FWindowActivateEvent& InActivateEvent) {}
-
-	virtual bool IsInPermanentCapture()
-	{ 
-		return  !GIsEditor && ((GetMouseCaptureMode() == EMouseCaptureMode::CapturePermanently) ||
-			(GetMouseCaptureMode() == EMouseCaptureMode::CapturePermanently_IncludingInitialMouseDown));
-	}
-
-	/**
-	 * Called when the top level window associated with the viewport has been requested to close.
-	 * At this point, the viewport has not been closed and the operation may be canceled.
-	 * This may not called from PIE, Editor Windows, on consoles, or before the game ends
-	 * from other methods.
-	 * This is only when the platform specific window is closed.
-	 *
-	 * @return True if the viewport may be closed, false otherwise.
-	 */
-	virtual bool WindowCloseRequested() { return true; }
-
-	virtual void CloseRequested(FViewport* Viewport) {}
-
-	virtual bool RequiresHitProxyStorage() { return true; }
-
-	/**
-	 * Determines whether this viewport client should receive calls to InputAxis() if the game's window is not currently capturing the mouse.
-	 * Used by the UI system to easily receive calls to InputAxis while the viewport's mouse capture is disabled.
-	 */
-	virtual bool RequiresUncapturedAxisInput() const { return false; }
-
-	/**
-	* Determine if the viewport client is going to need any keyboard input
-	* @return true if keyboard input is needed
-	*/
-	virtual bool RequiresKeyboardInput() const { return true; }
-
-	/**
-	 * Returns true if this viewport is orthogonal.
-	 * If hit proxies are ever used in-game, this will need to be
-	 * overridden correctly in GameViewportClient.
-	 */
-	virtual bool IsOrtho() const { return false; }
-
-	/**
-	 * Returns true if this viewport is excluding non-game elements from its display
-	 */
-	virtual bool IsInGameView() const { return false; }
-
-	/**
-	 * Sets GWorld to the appropriate world for this client
-	 *
-	 * @return the previous GWorld
-	 */
-	virtual class UWorld* ConditionalSetWorld() { return NULL; }
-
-	/**
-	 * Restores GWorld to InWorld
-	 *
-	 * @param InWorld	The world to restore
-	 */
-	virtual void ConditionalRestoreWorld( class UWorld* InWorld ) {}
-
-	/**
-	 * Allow viewport client to override the current capture region
-	 *
-	 * @param OutCaptureRegion    Ref to rectangle where we will write the overridden region
-	 * @return true if capture region has been overridden, false otherwise
-	 */
-	virtual bool OverrideHighResScreenshotCaptureRegion(FIntRect& OutCaptureRegion) { return false; }
-
-	/**
-	 * Get a ptr to the stat unit data for this viewport
-	 */
-	virtual FStatUnitData* GetStatUnitData() const { return NULL; }
-
-	/**
-	* Get a ptr to the stat unit data for this viewport
-	*/
-	virtual FStatHitchesData* GetStatHitchesData() const { return NULL; }
-
-	/**
-	 * Get a ptr to the enabled stats list
-	 */
-	virtual const TArray<FString>* GetEnabledStats() const { return NULL; }
-
-	/**
-	 * Sets all the stats that should be enabled for the viewport
-	 */
-	virtual void SetEnabledStats(const TArray<FString>& InEnabledStats) {}
-
-	/**
-	 * Check whether a specific stat is enabled for this viewport
-	 */
-	virtual bool IsStatEnabled(const FString& InName) const { return false; }
-
-	/**
-	* Sets whether stats should be visible for the viewport
-	*/
-	virtual void SetShowStats(bool bWantStats) { }
-
-	/**
-	 * Check whether we should ignore input.
-	 */
-	virtual bool IgnoreInput() { return false; }
-
-	/**
-	 * Gets the mouse capture behavior when the viewport is clicked
-	 */
-	virtual EMouseCaptureMode GetMouseCaptureMode() const { return EMouseCaptureMode::CapturePermanently; }
-
-	UE_DEPRECATED(4.26, "Please call GetMouseCaptureMode() instead.")
-	void CaptureMouseOnClick() { GetMouseCaptureMode(); }
-
-	/**
-	 * Gets whether or not the viewport captures the Mouse on launch of the application
-	 * Technically this controls capture on the first window activate, so in situations
-	 * where the application is launched but isn't activated the effect is delayed until
-	 * activation.
-	 */
-	virtual bool CaptureMouseOnLaunch() { return true; }
-
-	/**
-	 * Gets whether or not the cursor is locked to the viewport when the viewport captures the mouse
-	 */
-	virtual bool LockDuringCapture() { return true; }
-
-	/**
-	 * Gets whether or not the cursor should always be locked to the viewport.
-	 */
-	virtual bool ShouldAlwaysLockMouse() { return false; }
-
-	/**
-	 * Gets whether or not the cursor is hidden when the viewport captures the mouse
-	 */
-	virtual bool HideCursorDuringCapture() const { return false; }
-
-	/** 
-	 * Should we make new windows for popups or create an overlay in the current window.
-	 */
-	virtual FPopupMethodReply OnQueryPopupMethod() const { return FPopupMethodReply::Unhandled(); }
-
-	/**
-	 * Optionally do custom handling of a navigation. 
-	 */
-	virtual bool HandleNavigation(const uint32 InUserIndex, TSharedPtr<SWidget> InDestination) { return false; }
-
-	/**
-	 * @return Whether or not the scene canvas should be scaled.  Note: The debug canvas is always scaled 
-	 */
-	virtual bool ShouldDPIScaleSceneCanvas() const { return true; }
-
-	/**
-	 * @return The DPI Scale of this viewport
-	 */
-	virtual float GetDPIScale() const { return 1.0f; }
-};
-
 /** Tracks the viewport client that should process the stat command, can be NULL */
 extern ENGINE_API class FCommonViewportClient* GStatProcessingViewportClient;
-
-/**
- * Common functionality for game and editor viewport clients
- */
-
-class FCommonViewportClient : public FViewportClient
-{
-public:
-	FCommonViewportClient()
-		: CachedDPIScale(1.0f)
-		, bShouldUpdateDPIScale(true)
-	{}
-
-	virtual ~FCommonViewportClient()
-	{
-		//make to clean up the global "stat" client when we delete the active one.
-		if (GStatProcessingViewportClient == this)
-		{
-			GStatProcessingViewportClient = NULL;
-		}
-	}
-
-	/** Tells this viewport to update editor dpi scale when needed */
-	ENGINE_API void RequestUpdateDPIScale();
-
-	/** @return the current resolution fraction to be used for scene rendering in this client. */
-	ENGINE_API float GetDPIDerivedResolutionFraction() const;
-
-	/**
-	 * @return The DPI Scale of this viewport
-	 */
-	ENGINE_API virtual float GetDPIScale() const override;
-
-	ENGINE_API void DrawHighResScreenshotCaptureRegion(FCanvas& Canvas);
-
-protected:
-	/** @return the DPI scale of the window that the viewport client is in */
-	virtual float UpdateViewportClientWindowDPIScale() const { return 1.0; }
-
-private:
-	mutable float CachedDPIScale;
-	mutable bool bShouldUpdateDPIScale;
-};
-
-
-/**
- * Minimal viewport for assisting with taking screenshots (also used within a plugin)
- * @todo: This should be refactored
- */
-class ENGINE_API FDummyViewport : public FViewport
-{
-public:
-	FDummyViewport(FViewportClient* InViewportClient);
-
-	virtual ~FDummyViewport();
-
-	//~ Begin FViewport Interface
-	virtual void BeginRenderFrame(FRHICommandListImmediate& RHICmdList) override
-	{
-		check( IsInRenderingThread() );
-		//SetRenderTarget(RHICmdList,  RenderTargetTextureRHI,  FTexture2DRHIRef() );
-	};
-
-	virtual void EndRenderFrame(FRHICommandListImmediate& RHICmdList, bool bPresent, bool bLockToVsync) override
-	{
-		check( IsInRenderingThread() );
-	}
-
-	void SetupHDR(EDisplayColorGamut InDisplayColorGamut, EDisplayOutputFormat InDisplayOutputFormat, bool bInSceneHDREnabled)
-	{
-		DisplayColorGamut = InDisplayColorGamut;
-		DisplayOutputFormat = InDisplayOutputFormat;
-		bSceneHDREnabled = bInSceneHDREnabled;
-	}
-
-	virtual void*	GetWindow() override { return 0; }
-	virtual void	MoveWindow(int32 NewPosX, int32 NewPosY, int32 NewSizeX, int32 NewSizeY) override {}
-	virtual void	Destroy() override {}
-	virtual bool SetUserFocus(bool bFocus) override { return false; }
-	virtual bool	KeyState(FKey Key) const override { return false; }
-	virtual int32	GetMouseX() const override { return 0; }
-	virtual int32	GetMouseY() const override { return 0; }
-	virtual void	GetMousePos( FIntPoint& MousePosition, const bool bLocalPosition = true) override { MousePosition = FIntPoint(0, 0); }
-	virtual void	SetMouse(int32 x, int32 y) override { }
-	virtual void	ProcessInput( float DeltaTime ) override { }
-	virtual FVector2D VirtualDesktopPixelToViewport(FIntPoint VirtualDesktopPointPx) const override { return FVector2D::ZeroVector; }
-	virtual FIntPoint ViewportToVirtualDesktopPixel(FVector2D ViewportCoordinate) const override { return FIntPoint::ZeroValue; }
-	virtual void InvalidateDisplay() override { }
-	virtual void DeferInvalidateHitProxy() override { }
-	virtual FViewportFrame* GetViewportFrame() override { return 0; }
-	virtual FCanvas* GetDebugCanvas() override { return DebugCanvas; }
-	virtual EDisplayColorGamut GetDisplayColorGamut() const;
-	virtual EDisplayOutputFormat GetDisplayOutputFormat() const;
-	virtual bool GetSceneHDREnabled() const;
-	//~ End FViewport Interface
-
-	//~ Begin FRenderResource Interface
-	virtual void InitDynamicRHI() override;
-
-	// @todo DLL: Without these functions we get unresolved linker errors with FRenderResource
-	virtual void InitRHI() override{}
-	virtual void ReleaseRHI() override{}
-	virtual void InitResource() override{ FViewport::InitResource(); }
-	virtual void ReleaseResource() override { FViewport::ReleaseResource(); }
-	virtual FString GetFriendlyName() const override { return FString(TEXT("FDummyViewport"));}
-	//~ End FRenderResource Interface
-private:
-	FCanvas* DebugCanvas;
-	EDisplayColorGamut DisplayColorGamut = EDisplayColorGamut::sRGB_D65;
-	EDisplayOutputFormat DisplayOutputFormat = EDisplayOutputFormat::SDR_sRGB;
-	bool bSceneHDREnabled = false;
-};
