Removed lots of includes in UE based on compile profiling

diff --git a/Engine/Source/Runtime/Engine/Public/ViewportClient.h b/Engine/Source/Runtime/Engine/Public/ViewportClient.h
--- a/Engine/Source/Runtime/Engine/Public/ViewportClient.h
+++ b/Engine/Source/Runtime/Engine/Public/ViewportClient.h
@@ -2,4 +2,373 @@
 
 #pragma once
 
-// Placeholder for upcoming change
\ No newline at end of file
+#include "Engine/EngineBaseTypes.h"
+#include "GenericPlatform/ICursor.h"
+#include "InputCoreTypes.h"
+
+class FCanvas;
+class FCursorReply;
+class FViewport;
+class FPopupMethodReply;
+class FWindowActivateEvent;
+class SWidget;
+enum class EFocusCause : uint8;
+enum class EGestureEvent : uint8;
+struct FInputKeyEventArgs;
+struct FKey;
+struct FStatHitchesData;
+struct FStatUnitData;
+
+/**
+ * An abstract interface to a viewport's client.
+ * The viewport's client processes input received by the viewport, and draws the viewport.
+ */
+class FViewportClient
+{
+public:
+	virtual ~FViewportClient(){}
+	virtual void Precache() {}
+	ENGINE_API virtual void RedrawRequested(FViewport* Viewport);
+	ENGINE_API virtual void RequestInvalidateHitProxy(FViewport* Viewport);
+	virtual void Draw(FViewport* Viewport,FCanvas* Canvas) {}
+	virtual bool ProcessScreenShots(FViewport* Viewport) { return false; }
+	virtual UWorld* GetWorld() const { return NULL; }
+	virtual struct FEngineShowFlags* GetEngineShowFlags() { return NULL; }
+
+	/**
+	 * Check a key event received by the viewport.
+	 * If the viewport client uses the event, it should return true to consume it.
+	 * @param	EventArgs - The Input event args.
+	 * @return	True to consume the key event, false to pass it on.
+	 */
+	ENGINE_API virtual bool InputKey(const FInputKeyEventArgs& EventArgs);
+
+	UE_DEPRECATED(5.1, "Use the new InputKey(const FInputKeyEventArgs& EventArgs) function.")
+	virtual bool InputKey(FViewport* Viewport,int32 ControllerId,FKey Key,EInputEvent Event,float AmountDepressed = 1.f,bool bGamepad=false) { return false; }
+
+	/**
+	 * Check an axis movement received by the viewport.
+	 * If the viewport client uses the movement, it should return true to consume it.
+	 * @param	Viewport - The viewport which the axis movement is from.
+	 * @param	ControllerId - The controller which the axis movement is from.
+	 * @param	Key - The name of the axis which moved.
+	 * @param	Delta - The axis movement delta.
+	 * @param	DeltaTime - The time since the last axis update.
+	 * @param	NumSamples - The number of device samples that contributed to this Delta, useful for things like smoothing
+	 * @param	bGamepad - input came from gamepad (ie xbox controller)
+	 * @return	True to consume the axis movement, false to pass it on.
+	 */
+	UE_DEPRECATED(5.1, "Use the new InputAxis function that takes a FInputDeviceId.")
+	virtual bool InputAxis(FViewport* Viewport,int32 ControllerId,FKey Key,float Delta,float DeltaTime,int32 NumSamples=1,bool bGamepad=false) { return false; }
+
+	/**
+	 * Check an axis movement received by the viewport.
+	 * If the viewport client uses the movement, it should return true to consume it.
+	 * @param	Viewport - The viewport which the axis movement is from.
+	 * @param	InputDevice - The input device that triggered this axis movement
+	 * @param	Key - The name of the axis which moved.
+	 * @param	Delta - The axis movement delta.
+	 * @param	DeltaTime - The time since the last axis update.
+	 * @param	NumSamples - The number of device samples that contributed to this Delta, useful for things like smoothing
+	 * @param	bGamepad - input came from gamepad (ie xbox controller)
+	 * @return	True to consume the axis movement, false to pass it on.
+	 */
+	virtual bool InputAxis(FViewport* Viewport, FInputDeviceId InputDevice, FKey Key, float Delta, float DeltaTime, int32 NumSamples = 1, bool bGamepad = false) { return false; }
+
+	/**
+	 * Check a character input received by the viewport.
+	 * If the viewport client uses the character, it should return true to consume it.
+	 * @param	Viewport - The viewport which the axis movement is from.
+	 * @param	ControllerId - The controller which the axis movement is from.
+	 * @param	Character - The character.
+	 * @return	True to consume the character, false to pass it on.
+	 */
+	virtual bool InputChar(FViewport* Viewport,int32 ControllerId,TCHAR Character) { return false; }
+
+	/**
+	 * Check a key event received by the viewport.
+	 * If the viewport client uses the event, it should return true to consume it.
+	 * @param	Viewport - The viewport which the event is from.
+	 * @param	ControllerId - The controller which the key event is from.
+	 * @param	Handle - Identifier unique to this touch event
+	 * @param	Type - What kind of touch event this is (see ETouchType)
+	 * @param	TouchLocation - Screen position of the touch
+	 * @param	Force - How hard the touch is
+	 * @param	DeviceTimestamp - Timestamp of the event
+	 * @param	TouchpadIndex - For devices with multiple touchpads, this is the index of which one
+	 * @return	True to consume the key event, false to pass it on.
+	 */
+	virtual bool InputTouch(FViewport* Viewport, int32 ControllerId, uint32 Handle, ETouchType::Type Type, const FVector2D& TouchLocation, float Force, FDateTime DeviceTimestamp, uint32 TouchpadIndex) { return false; }
+
+	UE_DEPRECATED(4.20, "InputTouch now takes a Force")
+	bool InputTouch(FViewport* Viewport, int32 ControllerId, uint32 Handle, ETouchType::Type Type, const FVector2D& TouchLocation, FDateTime DeviceTimestamp, uint32 TouchpadIndex) { return InputTouch(Viewport, ControllerId, Handle, Type, TouchLocation, 1.0f, DeviceTimestamp, TouchpadIndex); }
+
+	/**
+	 * Check a gesture event received by the viewport.
+	 * If the viewport client uses the event, it should return true to consume it.
+	 * @param	Viewport - The viewport which the event is from.
+	 * @param	GestureType - @todo desc
+	 * @param	GestureDelta - @todo desc
+	 * @return	True to consume the gesture event, false to pass it on.
+	 */
+	virtual bool InputGesture(FViewport* Viewport, EGestureEvent GestureType, const FVector2D& GestureDelta, bool bIsDirectionInvertedFromDevice) { return false; }
+
+	/**
+	 * Each frame, the input system will update the motion data.
+	 *
+	 * @param Viewport - The viewport which the key event is from.
+	 * @param ControllerId - The controller which the key event is from.
+	 * @param Tilt			The current orientation of the device
+	 * @param RotationRate	How fast the tilt is changing
+	 * @param Gravity		Describes the current gravity of the device
+	 * @param Acceleration  Describes the acceleration of the device
+	 * @return	True to consume the motion event, false to pass it on.
+	 */
+	virtual bool InputMotion(FViewport* Viewport, int32 ControllerId, const FVector& Tilt, const FVector& RotationRate, const FVector& Gravity, const FVector& Acceleration) { return false; }
+
+	virtual void SetIsSimulateInEditorViewport(bool bInIsSimulateInEditorViewport) { };
+
+	virtual bool WantsPollingMouseMovement(void) const { return true; }
+
+	virtual void MouseEnter( FViewport* Viewport,int32 x, int32 y ) {}
+
+	virtual void MouseLeave( FViewport* Viewport ) {}
+
+	virtual void MouseMove(FViewport* Viewport,int32 X,int32 Y) {}
+
+	/**
+	 * Called when the mouse is moved while a window input capture is in effect
+	 *
+	 * @param	InViewport	Viewport that captured the mouse input
+	 * @param	InMouseX	New mouse cursor X coordinate
+	 * @param	InMouseY	New mouse cursor Y coordinate
+	 */
+	virtual void CapturedMouseMove( FViewport* InViewport, int32 InMouseX, int32 InMouseY ) { }
+
+	/**
+	 * Called from slate when input is finished for this frame, and we should process any queued mouse moves.
+	 */
+	virtual void ProcessAccumulatedPointerInput(FViewport* InViewport) {};
+
+	/**
+	 * Retrieves the cursor that should be displayed by the OS
+	 *
+	 * @param	Viewport	the viewport that contains the cursor
+	 * @param	X			the x position of the cursor
+	 * @param	Y			the Y position of the cursor
+	 *
+	 * @return	the cursor that the OS should display
+	 */
+	virtual EMouseCursor::Type GetCursor(FViewport* Viewport, int32 X,int32 Y) { return EMouseCursor::Default; }
+
+	/**
+	 * Called to map a cursor reply to an actual widget to render.
+	 *
+	 * @return	the widget that should be rendered for this cursor, return TOptional<TSharedRef<SWidget>>() if no mapping.
+	 */
+	ENGINE_API virtual TOptional<TSharedRef<SWidget>> MapCursor(FViewport* Viewport, const FCursorReply& CursorReply);
+
+	/**
+	 * Called to determine if we should render the focus brush.
+	 *
+	 * @param InFocusCause	The cause of focus
+	 */
+	virtual TOptional<bool> QueryShowFocus(const EFocusCause InFocusCause) const { return TOptional<bool>(); }
+
+	virtual void LostFocus(FViewport* Viewport) {}
+	virtual void ReceivedFocus(FViewport* Viewport) {}
+	virtual bool IsFocused(FViewport* Viewport) { return true; }
+
+	virtual void Activated(FViewport* Viewport, const FWindowActivateEvent& InActivateEvent) {}
+	virtual void Deactivated(FViewport* Viewport, const FWindowActivateEvent& InActivateEvent) {}
+
+	virtual bool IsInPermanentCapture()
+	{ 
+		return  !GIsEditor && ((GetMouseCaptureMode() == EMouseCaptureMode::CapturePermanently) ||
+			(GetMouseCaptureMode() == EMouseCaptureMode::CapturePermanently_IncludingInitialMouseDown));
+	}
+
+	/**
+	 * Called when the top level window associated with the viewport has been requested to close.
+	 * At this point, the viewport has not been closed and the operation may be canceled.
+	 * This may not called from PIE, Editor Windows, on consoles, or before the game ends
+	 * from other methods.
+	 * This is only when the platform specific window is closed.
+	 *
+	 * @return True if the viewport may be closed, false otherwise.
+	 */
+	virtual bool WindowCloseRequested() { return true; }
+
+	virtual void CloseRequested(FViewport* Viewport) {}
+
+	virtual bool RequiresHitProxyStorage() { return true; }
+
+	/**
+	 * Determines whether this viewport client should receive calls to InputAxis() if the game's window is not currently capturing the mouse.
+	 * Used by the UI system to easily receive calls to InputAxis while the viewport's mouse capture is disabled.
+	 */
+	virtual bool RequiresUncapturedAxisInput() const { return false; }
+
+	/**
+	* Determine if the viewport client is going to need any keyboard input
+	* @return true if keyboard input is needed
+	*/
+	virtual bool RequiresKeyboardInput() const { return true; }
+
+	/**
+	 * Returns true if this viewport is orthogonal.
+	 * If hit proxies are ever used in-game, this will need to be
+	 * overridden correctly in GameViewportClient.
+	 */
+	virtual bool IsOrtho() const { return false; }
+
+	/**
+	 * Returns true if this viewport is excluding non-game elements from its display
+	 */
+	virtual bool IsInGameView() const { return false; }
+
+	/**
+	 * Sets GWorld to the appropriate world for this client
+	 *
+	 * @return the previous GWorld
+	 */
+	virtual class UWorld* ConditionalSetWorld() { return NULL; }
+
+	/**
+	 * Restores GWorld to InWorld
+	 *
+	 * @param InWorld	The world to restore
+	 */
+	virtual void ConditionalRestoreWorld( class UWorld* InWorld ) {}
+
+	/**
+	 * Allow viewport client to override the current capture region
+	 *
+	 * @param OutCaptureRegion    Ref to rectangle where we will write the overridden region
+	 * @return true if capture region has been overridden, false otherwise
+	 */
+	virtual bool OverrideHighResScreenshotCaptureRegion(FIntRect& OutCaptureRegion) { return false; }
+
+	/**
+	 * Get a ptr to the stat unit data for this viewport
+	 */
+	virtual FStatUnitData* GetStatUnitData() const { return NULL; }
+
+	/**
+	* Get a ptr to the stat unit data for this viewport
+	*/
+	virtual FStatHitchesData* GetStatHitchesData() const { return NULL; }
+
+	/**
+	 * Get a ptr to the enabled stats list
+	 */
+	virtual const TArray<FString>* GetEnabledStats() const { return NULL; }
+
+	/**
+	 * Sets all the stats that should be enabled for the viewport
+	 */
+	virtual void SetEnabledStats(const TArray<FString>& InEnabledStats) {}
+
+	/**
+	 * Check whether a specific stat is enabled for this viewport
+	 */
+	virtual bool IsStatEnabled(const FString& InName) const { return false; }
+
+	/**
+	* Sets whether stats should be visible for the viewport
+	*/
+	virtual void SetShowStats(bool bWantStats) { }
+
+	/**
+	 * Check whether we should ignore input.
+	 */
+	virtual bool IgnoreInput() { return false; }
+
+	/**
+	 * Gets the mouse capture behavior when the viewport is clicked
+	 */
+	virtual EMouseCaptureMode GetMouseCaptureMode() const { return EMouseCaptureMode::CapturePermanently; }
+
+	UE_DEPRECATED(4.26, "Please call GetMouseCaptureMode() instead.")
+	void CaptureMouseOnClick() { GetMouseCaptureMode(); }
+
+	/**
+	 * Gets whether or not the viewport captures the Mouse on launch of the application
+	 * Technically this controls capture on the first window activate, so in situations
+	 * where the application is launched but isn't activated the effect is delayed until
+	 * activation.
+	 */
+	virtual bool CaptureMouseOnLaunch() { return true; }
+
+	/**
+	 * Gets whether or not the cursor is locked to the viewport when the viewport captures the mouse
+	 */
+	virtual bool LockDuringCapture() { return true; }
+
+	/**
+	 * Gets whether or not the cursor should always be locked to the viewport.
+	 */
+	virtual bool ShouldAlwaysLockMouse() { return false; }
+
+	/**
+	 * Gets whether or not the cursor is hidden when the viewport captures the mouse
+	 */
+	virtual bool HideCursorDuringCapture() const { return false; }
+
+	/** 
+	 * Should we make new windows for popups or create an overlay in the current window.
+	 */
+	ENGINE_API virtual FPopupMethodReply OnQueryPopupMethod() const;
+
+	/**
+	 * Optionally do custom handling of a navigation. 
+	 */
+	virtual bool HandleNavigation(const uint32 InUserIndex, TSharedPtr<SWidget> InDestination) { return false; }
+
+	/**
+	 * @return Whether or not the scene canvas should be scaled.  Note: The debug canvas is always scaled 
+	 */
+	virtual bool ShouldDPIScaleSceneCanvas() const { return true; }
+
+	/**
+	 * @return The DPI Scale of this viewport
+	 */
+	virtual float GetDPIScale() const { return 1.0f; }
+};
+
+/**
+ * Common functionality for game and editor viewport clients
+ */
+
+class FCommonViewportClient : public FViewportClient
+{
+public:
+	FCommonViewportClient()
+		: CachedDPIScale(1.0f)
+		, bShouldUpdateDPIScale(true)
+	{}
+
+	ENGINE_API virtual ~FCommonViewportClient();
+
+	/** Tells this viewport to update editor dpi scale when needed */
+	ENGINE_API void RequestUpdateDPIScale();
+
+	/** @return the current resolution fraction to be used for scene rendering in this client. */
+	ENGINE_API float GetDPIDerivedResolutionFraction() const;
+
+	/**
+	 * @return The DPI Scale of this viewport
+	 */
+	ENGINE_API virtual float GetDPIScale() const override;
+
+	ENGINE_API void DrawHighResScreenshotCaptureRegion(FCanvas& Canvas);
+
+protected:
+	/** @return the DPI scale of the window that the viewport client is in */
+	virtual float UpdateViewportClientWindowDPIScale() const { return 1.0; }
+
+private:
+	mutable float CachedDPIScale;
+	mutable bool bShouldUpdateDPIScale;
+};
+
diff --git a/Engine/Source/Runtime/Engine/Public/VisualLogger/VisualLogger.h b/Engine/Source/Runtime/Engine/Public/VisualLogger/VisualLogger.h
--- a/Engine/Source/Runtime/Engine/Public/VisualLogger/VisualLogger.h
+++ b/Engine/Source/Runtime/Engine/Public/VisualLogger/VisualLogger.h
@@ -7,7 +7,9 @@
 #include "Engine/EngineTypes.h"
 #include "EngineDefines.h"
 #include "VisualLogger/VisualLoggerTypes.h"
+#if UE_ENABLE_INCLUDE_ORDER_DEPRECATED_IN_5_2
 #include "EngineStats.h"
+#endif
 #include "Templates/IsValidVariadicFunctionArg.h"
 #include "Templates/AndOrNot.h"
 #include "Templates/IsArrayOrRefOfTypeByPredicate.h"
@@ -676,224 +678,4 @@ protected:
 	static int32 bIsRecording;
 };
 
-// Unfortunately needs to be a #define since it uses GET_VARARGS_RESULT which uses the va_list stuff which operates on the
-// current function, so we can't easily call a function
-#define COLLAPSED_LOGF(SerializeFunc) \
-	SCOPE_CYCLE_COUNTER(STAT_VisualLog); \
-	UWorld *World = nullptr; \
-	FVisualLogEntry *CurrentEntry = nullptr; \
-	if (CheckVisualLogInputInternal(Object, CategoryName, Verbosity, &World, &CurrentEntry) == false) \
-	{ \
-		return; \
-	} \
-	int32	BufferSize	= 1024; \
-	TCHAR*	Buffer		= nullptr; \
-	int32	Result		= -1; \
-	/* allocate some stack space to use on the first pass, which matches most strings */ \
-	TCHAR	StackBuffer[512]; \
-	TCHAR*	AllocatedBuffer = nullptr; \
-	\
-	/* first, try using the stack buffer */ \
-	Buffer = StackBuffer; \
-	GET_VARARGS_RESULT( Buffer, UE_ARRAY_COUNT(StackBuffer), UE_ARRAY_COUNT(StackBuffer) - 1, Fmt, Fmt, Result ); \
-	\
-	/* if that fails, then use heap allocation to make enough space */ \
-			while(Result == -1) \
-						{ \
-		FMemory::SystemFree(AllocatedBuffer); \
-		/* We need to use malloc here directly as GMalloc might not be safe. */ \
-		Buffer = AllocatedBuffer = (TCHAR*) FMemory::SystemMalloc( BufferSize * sizeof(TCHAR) ); \
-		GET_VARARGS_RESULT( Buffer, BufferSize, BufferSize-1, Fmt, Fmt, Result ); \
-		BufferSize *= 2; \
-						}; \
-	Buffer[Result] = 0; \
-	; \
-	\
-	SerializeFunc; \
-	FMemory::SystemFree(AllocatedBuffer);
-
-inline void FVisualLogger::CategorizedLogfImpl(const UObject* Object, const FLogCategoryBase& Category, ELogVerbosity::Type Verbosity, const TCHAR* Fmt, ...)
-{
-	const FName CategoryName = Category.GetCategoryName(); 
-	COLLAPSED_LOGF(
-		CurrentEntry->AddText(Buffer, CategoryName, Verbosity);
-	);
-}
-inline void FVisualLogger::CategorizedLogfImpl(const UObject* Object, const FName& CategoryName, ELogVerbosity::Type Verbosity, const TCHAR* Fmt, ...)
-{
-	COLLAPSED_LOGF(
-		CurrentEntry->AddText(Buffer, CategoryName, Verbosity);
-	);
-}
-
-inline void FVisualLogger::GeometryShapeLogfImpl(const UObject* Object, const FLogCategoryBase& Category, ELogVerbosity::Type Verbosity, const FVector& Start, const FVector& End, const FColor& Color, const uint16 Thickness, const TCHAR* Fmt, ...)
-{
-	const FName CategoryName = Category.GetCategoryName(); 
-	COLLAPSED_LOGF(
-		CurrentEntry->AddElement(Start, End, CategoryName, Verbosity, Color, Buffer, Thickness);
-	);
-}
-inline void FVisualLogger::GeometryShapeLogfImpl(const UObject* Object, const FName& CategoryName, ELogVerbosity::Type Verbosity, const FVector& Start, const FVector& End, const FColor& Color, const uint16 Thickness, const TCHAR* Fmt, ...)
-{
-	COLLAPSED_LOGF(
-		CurrentEntry->AddElement(Start, End, CategoryName, Verbosity, Color, Buffer, Thickness);
-	);
-}
-
-inline void FVisualLogger::ArrowLogfImpl(const UObject* Object, const FLogCategoryBase& Category, ELogVerbosity::Type Verbosity, const FVector& Start, const FVector& End, const FColor& Color, const TCHAR* Fmt, ...)
-{
-	const FName CategoryName = Category.GetCategoryName();
-	COLLAPSED_LOGF(
-		CurrentEntry->AddArrow(Start, End, CategoryName, Verbosity, Color, Buffer);
-	);
-}
-inline void FVisualLogger::ArrowLogfImpl(const UObject* Object, const FName& CategoryName, ELogVerbosity::Type Verbosity, const FVector& Start, const FVector& End, const FColor& Color, const TCHAR* Fmt, ...)
-{
-	COLLAPSED_LOGF(
-		CurrentEntry->AddArrow(Start, End, CategoryName, Verbosity, Color, Buffer);
-	);
-}
-
-inline void FVisualLogger::CircleLogfImpl(const UObject* Object, const FLogCategoryBase& Category, ELogVerbosity::Type Verbosity, const FVector& Center, const FVector& UpAxis, const float Radius, const FColor& Color, const uint16 Thickness, const TCHAR* Fmt, ...)
-{
-	const FName CategoryName = Category.GetCategoryName();
-	COLLAPSED_LOGF(
-		CurrentEntry->AddCircle(Center, UpAxis, Radius, CategoryName, Verbosity, Color, Buffer, Thickness);
-	);
-}
-inline void FVisualLogger::CircleLogfImpl(const UObject* Object, const FName& CategoryName, ELogVerbosity::Type Verbosity, const FVector& Center, const FVector& UpAxis, const float Radius, const FColor& Color, const uint16 Thickness, const TCHAR* Fmt, ...)
-{
-	COLLAPSED_LOGF(
-		CurrentEntry->AddCircle(Center, UpAxis, Radius, CategoryName, Verbosity, Color, Buffer, Thickness);
-	);
-}
-
-inline void FVisualLogger::GeometryShapeLogfImpl(const UObject* Object, const FLogCategoryBase& Category, ELogVerbosity::Type Verbosity, const FVector& Location, float Radius, const FColor& Color, const TCHAR* Fmt, ...)
-{
-	const FName CategoryName = Category.GetCategoryName(); 
-	COLLAPSED_LOGF(
-		CurrentEntry->AddElement(Location, CategoryName, Verbosity, Color, Buffer, (uint16)Radius);
-	);
-}
-inline void FVisualLogger::GeometryShapeLogfImpl(const UObject* Object, const FName& CategoryName, ELogVerbosity::Type Verbosity, const FVector& Location, float Radius, const FColor& Color, const TCHAR* Fmt, ...)
-{
-	COLLAPSED_LOGF(
-		CurrentEntry->AddElement(Location, CategoryName, Verbosity, Color, Buffer, (uint16)Radius);
-	);
-}
-
-inline void FVisualLogger::GeometryBoxLogfImpl(const UObject* Object, const FLogCategoryBase& Category, ELogVerbosity::Type Verbosity, const FBox& Box, const FMatrix& Matrix, const FColor& Color, const TCHAR* Fmt, ...)
-{
-	const FName CategoryName = Category.GetCategoryName(); 
-	COLLAPSED_LOGF(
-		CurrentEntry->AddElement(Box, Matrix, CategoryName, Verbosity, Color, Buffer);
-	);
-}
-inline void FVisualLogger::GeometryBoxLogfImpl(const UObject* Object, const FName& CategoryName, ELogVerbosity::Type Verbosity, const FBox& Box, const FMatrix& Matrix, const FColor& Color, const TCHAR* Fmt, ...)
-{
-	COLLAPSED_LOGF(
-		CurrentEntry->AddElement(Box, Matrix, CategoryName, Verbosity, Color, Buffer);
-	);
-}
-
-inline void FVisualLogger::GeometryShapeLogfImpl(const UObject* Object, const FLogCategoryBase& Category, ELogVerbosity::Type Verbosity, const FVector& Origin, const FVector& Direction, const float Length, const float Angle, const FColor& Color, const TCHAR* Fmt, ...)
-{
-	const FName CategoryName = Category.GetCategoryName(); 
-	COLLAPSED_LOGF(
-		CurrentEntry->AddElement(Origin, Direction, Length, Angle, Angle, CategoryName, Verbosity, Color, Buffer);
-	);
-}
-inline void FVisualLogger::GeometryShapeLogfImpl(const UObject* Object, const FName& CategoryName, ELogVerbosity::Type Verbosity, const FVector& Origin, const FVector& Direction, const float Length, const float Angle, const FColor& Color, const TCHAR* Fmt, ...)
-{
-	COLLAPSED_LOGF(
-		CurrentEntry->AddElement(Origin, Direction, Length, Angle, Angle, CategoryName, Verbosity, Color, Buffer);
-	);
-}
-
-inline void FVisualLogger::GeometryShapeLogfImpl(const UObject* Object, const FLogCategoryBase& Category, ELogVerbosity::Type Verbosity, const FVector& Start, const FVector& End, const float Radius, const FColor& Color, const TCHAR* Fmt, ...)
-{
-	const FName CategoryName = Category.GetCategoryName(); 
-	COLLAPSED_LOGF(
-		CurrentEntry->AddElement(Start, End, Radius, CategoryName, Verbosity, Color, Buffer);
-	);
-}
-inline void FVisualLogger::GeometryShapeLogfImpl(const UObject* Object, const FName& CategoryName, ELogVerbosity::Type Verbosity, const FVector& Start, const FVector& End, const float Radius, const FColor& Color, const TCHAR* Fmt, ...)
-{
-	COLLAPSED_LOGF(
-		CurrentEntry->AddElement(Start, End, Radius, CategoryName, Verbosity, Color, Buffer);
-	);
-}
-
-inline void FVisualLogger::GeometryShapeLogfImpl(const UObject* Object, const FLogCategoryBase& Category, ELogVerbosity::Type Verbosity, const FVector& Base, float HalfHeight, float Radius, const FQuat& Rotation, const FColor& Color, const TCHAR* Fmt, ...)
-{
-	const FName CategoryName = Category.GetCategoryName(); 
-	COLLAPSED_LOGF(
-		CurrentEntry->AddElement(Base, HalfHeight, Radius, Rotation, CategoryName, Verbosity, Color, Buffer);
-	);
-}
-inline void FVisualLogger::GeometryShapeLogfImpl(const UObject* Object, const FName& CategoryName, ELogVerbosity::Type Verbosity, const FVector& Base, float HalfHeight, float Radius, const FQuat& Rotation, const FColor& Color, const TCHAR* Fmt, ...)
-{
-	COLLAPSED_LOGF(
-		CurrentEntry->AddElement(Base, HalfHeight, Radius, Rotation, CategoryName, Verbosity, Color, Buffer);
-	);
-}
-
-inline void FVisualLogger::NavAreaShapeLogfImpl(const UObject* Object, const FLogCategoryBase& Category, ELogVerbosity::Type Verbosity, const TArray<FVector>& ConvexPoints, float MinZ, float MaxZ, const FColor& Color, const TCHAR* Fmt, ...)
-{
-	const FName CategoryName = Category.GetCategoryName(); 
-	COLLAPSED_LOGF(
-		CurrentEntry->AddElement(ConvexPoints, MinZ, MaxZ, CategoryName, Verbosity, Color, Buffer);
-	);
-}
-inline void FVisualLogger::NavAreaShapeLogfImpl(const UObject* Object, const FName& CategoryName, ELogVerbosity::Type Verbosity, const TArray<FVector>& ConvexPoints, float MinZ, float MaxZ, const FColor& Color, const TCHAR* Fmt, ...)
-{
-	COLLAPSED_LOGF(
-		CurrentEntry->AddElement(ConvexPoints, MinZ, MaxZ, CategoryName, Verbosity, Color, Buffer);
-	);
-}
-
-inline void FVisualLogger::GeometryShapeLogfImpl(const UObject* Object, const FLogCategoryBase& Category, ELogVerbosity::Type Verbosity, const TArray<FVector>& Vertices, const TArray<int32>& Indices, const FColor& Color, const TCHAR* Fmt, ...)
-{
-	const FName CategoryName = Category.GetCategoryName(); 
-	COLLAPSED_LOGF(
-		CurrentEntry->AddElement(Vertices, Indices, CategoryName, Verbosity, Color, Buffer);
-	);
-}
-inline void FVisualLogger::GeometryShapeLogfImpl(const UObject* Object, const FName& CategoryName, ELogVerbosity::Type Verbosity, const TArray<FVector>& Vertices, const TArray<int32>& Indices, const FColor& Color, const TCHAR* Fmt, ...)
-{
-	COLLAPSED_LOGF(
-		CurrentEntry->AddElement(Vertices, Indices, CategoryName, Verbosity, Color, Buffer);
-	);
-}
-
-inline void FVisualLogger::GeometryConvexLogfImpl(const UObject* Object, const FLogCategoryBase& Category, ELogVerbosity::Type Verbosity, const TArray<FVector>& Points, const FColor& Color, const TCHAR* Fmt, ...)
-{
-	const FName CategoryName = Category.GetCategoryName(); 
-	COLLAPSED_LOGF(
-		CurrentEntry->AddConvexElement(Points, CategoryName, Verbosity, Color, Buffer);
-	);
-}
-inline void FVisualLogger::GeometryConvexLogfImpl(const UObject* Object, const FName& CategoryName, ELogVerbosity::Type Verbosity, const TArray<FVector>& Points, const FColor& Color, const TCHAR* Fmt, ...)
-{
-	COLLAPSED_LOGF(
-		CurrentEntry->AddConvexElement(Points, CategoryName, Verbosity, Color, Buffer);
-	);
-}
-
-inline void FVisualLogger::HistogramDataLogfImpl(const UObject* Object, const FLogCategoryBase& Category, ELogVerbosity::Type Verbosity, FName GraphName, FName DataName, const FVector2D& Data, const FColor& Color, const TCHAR* Fmt, ...)
-{
-	const FName CategoryName = Category.GetCategoryName(); 
-	COLLAPSED_LOGF(
-		CurrentEntry->AddHistogramData(Data, CategoryName, Verbosity, GraphName, DataName);
-	);
-}
-inline void FVisualLogger::HistogramDataLogfImpl(const UObject* Object, const FName& CategoryName, ELogVerbosity::Type Verbosity, FName GraphName, FName DataName, const FVector2D& Data, const FColor& Color, const TCHAR* Fmt, ...)
-{
-	COLLAPSED_LOGF(
-		CurrentEntry->AddHistogramData(Data, CategoryName, Verbosity, GraphName, DataName);
-	);
-}
-
-#undef COLLAPSED_LOGF
-
 #endif //ENABLE_VISUAL_LOG
diff --git a/Engine/Source/Runtime/Experimental/Chaos/Private/Chaos/Levelset.cpp b/Engine/Source/Runtime/Experimental/Chaos/Private/Chaos/Levelset.cpp
--- a/Engine/Source/Runtime/Experimental/Chaos/Private/Chaos/Levelset.cpp
+++ b/Engine/Source/Runtime/Experimental/Chaos/Private/Chaos/Levelset.cpp
@@ -156,6 +156,7 @@ FLevelSet::FLevelSet(FErrorReporter& ErrorReporter, const TUniformGrid<FReal, 3>
 	ComputeConvexity(InterfaceIndices);
 }
 
+#if COMPILE_WITHOUT_UNREAL_SUPPORT
 FLevelSet::FLevelSet(std::istream& Stream)
     : FImplicitObject(EImplicitObject::HasBoundingBox, ImplicitObjectType::LevelSet)
     , MGrid(Stream)
@@ -165,7 +166,7 @@ FLevelSet::FLevelSet(std::istream& Stream)
 	Stream.read(reinterpret_cast<char*>(&MBandWidth), sizeof(MBandWidth));
 	ComputeNormals();
 }
-
+#endif
 FLevelSet::FLevelSet(TUniformGrid<FReal, 3>&& Grid, TArrayND<FReal, 3>&& Phi, int32 BandWidth)
 	: FImplicitObject(EImplicitObject::HasBoundingBox, ImplicitObjectType::LevelSet)
 	, MGrid(MoveTemp(Grid))
@@ -1398,12 +1399,14 @@ void FLevelSet::ComputeNormals(const FParticles& InParticles, const FTriangleMes
 	}
 }
 
+#if COMPILE_WITHOUT_UNREAL_SUPPORT
 void FLevelSet::Write(std::ostream& Stream) const
 {
 	MGrid.Write(Stream);
 	MPhi.Write(Stream);
 	Stream.write(reinterpret_cast<const char*>(&MBandWidth), sizeof(MBandWidth));
 }
+#endif
 
 FReal FLevelSet::SignedDistance(const FVec3& x) const
 {
diff --git a/Engine/Source/Runtime/Experimental/Chaos/Public/Chaos/AABBTreeDirtyGridUtils.h b/Engine/Source/Runtime/Experimental/Chaos/Public/Chaos/AABBTreeDirtyGridUtils.h
--- a/Engine/Source/Runtime/Experimental/Chaos/Public/Chaos/AABBTreeDirtyGridUtils.h
+++ b/Engine/Source/Runtime/Experimental/Chaos/Public/Chaos/AABBTreeDirtyGridUtils.h
@@ -4,6 +4,7 @@
 #include "Chaos/AABB.h"
 #include "Chaos/Defines.h"
 #include "ChaosLog.h"
+#include <limits>
 
 namespace Chaos
 {
diff --git a/Engine/Source/Runtime/Experimental/Chaos/Public/Chaos/Levelset.h b/Engine/Source/Runtime/Experimental/Chaos/Public/Chaos/Levelset.h
--- a/Engine/Source/Runtime/Experimental/Chaos/Public/Chaos/Levelset.h
+++ b/Engine/Source/Runtime/Experimental/Chaos/Public/Chaos/Levelset.h
@@ -33,7 +33,9 @@ class CHAOS_API FLevelSet final : public FImplicitObject
 
 	FLevelSet(FErrorReporter& ErrorReporter, const TUniformGrid<FReal, 3>& InGrid, const FParticles& InParticles, const FTriangleMesh& Mesh, const int32 BandWidth = 0);
 	FLevelSet(FErrorReporter& ErrorReporter, const TUniformGrid<FReal, 3>& InGrid, const FImplicitObject& InObject, const int32 BandWidth = 0, const bool bUseObjectPhi = false);
+#if COMPILE_WITHOUT_UNREAL_SUPPORT
 	FLevelSet(std::istream& Stream);
+#endif
 	FLevelSet(TUniformGrid<FReal, 3>&& Grid, TArrayND<FReal, 3>&& Phi, int32 BandWidth);
 	FLevelSet(const FLevelSet& Other) = delete;
 	FLevelSet(FLevelSet&& Other);
@@ -44,7 +46,9 @@ class CHAOS_API FLevelSet final : public FImplicitObject
 	virtual TUniquePtr<FImplicitObject> DeepCopy() const;
 	virtual TUniquePtr<FImplicitObject> DeepCopyWithScale(const FVec3& Scale) const;
 
+#if COMPILE_WITHOUT_UNREAL_SUPPORT
 	void Write(std::ostream& Stream) const;
+#endif
 	virtual FReal PhiWithNormal(const FVec3& x, FVec3& Normal) const override;
 	FReal SignedDistance(const FVec3& x) const;
 
diff --git a/Engine/Source/Runtime/Experimental/Chaos/Public/Chaos/UniformGrid.h b/Engine/Source/Runtime/Experimental/Chaos/Public/Chaos/UniformGrid.h
--- a/Engine/Source/Runtime/Experimental/Chaos/Public/Chaos/UniformGrid.h
+++ b/Engine/Source/Runtime/Experimental/Chaos/Public/Chaos/UniformGrid.h
@@ -84,21 +84,24 @@ class CHAOS_API TUniformGridBase
 			}
 		}
 	}
+#if COMPILE_WITHOUT_UNREAL_SUPPORT
 	TUniformGridBase(std::istream& Stream)
 	    : MMinCorner(Stream), MMaxCorner(Stream), MCells(Stream)
 	{
 		MDx = TVector<T, d>(MMaxCorner - MMinCorner) / MCells;
 	}
-
+#endif
 	~TUniformGridBase() {}
 
   public:
+#if COMPILE_WITHOUT_UNREAL_SUPPORT
 	void Write(std::ostream& Stream) const
 	{
 		MMinCorner.Write(Stream);
 		MMaxCorner.Write(Stream);
 		MCells.Write(Stream);
 	}
+#endif
 	void Serialize(FArchive& Ar)
 	{
 		Ar << MMinCorner;
@@ -273,8 +276,10 @@ class CHAOS_API TUniformGrid : public TUniformGridBase<T, d>
 	TUniformGrid() {}
 	TUniformGrid(const TVector<T, d>& MinCorner, const TVector<T, d>& MaxCorner, const TVector<int32, d>& Cells, const uint32 GhostCells = 0)
 	    : TUniformGridBase<T, d>(MinCorner, MaxCorner, Cells, GhostCells) {}
+#if COMPILE_WITHOUT_UNREAL_SUPPORT
 	TUniformGrid(std::istream& Stream)
 	    : TUniformGridBase<T, d>(Stream) {}
+#endif
 	~TUniformGrid() {}
 	TVector<int32, d> GetIndex(const int32 Index) const;
 	TVector<T, d> Center(const int32 Index) const
@@ -322,8 +327,10 @@ public:
 		: TUniformGridBase<T, 3>(MinCorner, MaxCorner, Cells, GhostCells) {}
 	TMPMGrid(const int32 GridN) { for (int32 i = 0; i < 3; i++) { MDx[i] = (T)1. / (T)GridN; } }
 	TMPMGrid(const T GridDx) { for (int32 i = 0; i < 3; i++) { MDx[i] = GridDx; } }
+#if COMPILE_WITHOUT_UNREAL_SUPPORT
 	TMPMGrid(std::istream& Stream)
 		: TUniformGridBase<T, 3>(Stream) {}
+#endif
 	~TMPMGrid() {}
 
 	void BaseNodeIndex(const TVector<T, 3>& X, TVector<int32, 3>& Index, TVector<T, 3>& weights) const;
@@ -423,8 +430,10 @@ class TUniformGrid<T, 3> : public TUniformGridBase<T, 3>
 	TUniformGrid() {}
 	TUniformGrid(const TVector<T, 3>& MinCorner, const TVector<T, 3>& MaxCorner, const TVector<int32, 3>& Cells, const uint32 GhostCells = 0)
 	    : TUniformGridBase<T, 3>(MinCorner, MaxCorner, Cells, GhostCells) {}
+#if COMPILE_WITHOUT_UNREAL_SUPPORT
 	TUniformGrid(std::istream& Stream)
 	    : TUniformGridBase<T, 3>(Stream) {}
+#endif
 	~TUniformGrid() {}
 	TVector<int32, 3> GetIndex(const int32 Index) const;
 	Pair<int32, TVector<int32, 3>> GetFaceIndex(int32 Index) const;
diff --git a/Engine/Source/Runtime/Experimental/ChaosCore/Public/Chaos/ArrayND.h b/Engine/Source/Runtime/Experimental/ChaosCore/Public/Chaos/ArrayND.h
--- a/Engine/Source/Runtime/Experimental/ChaosCore/Public/Chaos/ArrayND.h
+++ b/Engine/Source/Runtime/Experimental/ChaosCore/Public/Chaos/ArrayND.h
@@ -97,6 +97,7 @@ class TArrayNDBase
 	FORCEINLINE TArrayNDBase(const TArrayNDBase<T_DERIVED, T, d>& Other) = delete;
 	FORCEINLINE TArrayNDBase(TArrayNDBase<T_DERIVED, T, d>&& Other)
 	    : MCounts(Other.MCounts), MArray(MoveTemp(Other.MArray)) {}
+#if COMPILE_WITHOUT_UNREAL_SUPPORT
 	FORCEINLINE TArrayNDBase(std::istream& Stream)
 	    : MCounts(Stream)
 	{
@@ -108,7 +109,7 @@ class TArrayNDBase
 		MCounts.Write(Stream);
 		Stream.write(reinterpret_cast<const char*>(MArray.GetData()), sizeof(T) * MArray.Num());
 	}
-
+#endif
 	void Serialize(FArchive& Ar)
 	{
 		Ar.UsingCustomVersion(FDestructionObjectVersion::GUID);
@@ -197,8 +198,10 @@ class TArrayND : public TArrayNDBase<TArrayND<T, d>, T, d>
 	FORCEINLINE TArrayND(const TArrayND<T, d>& Other) = delete;
 	FORCEINLINE TArrayND(TArrayND<T, d>&& Other)
 	    : Base(MoveTemp(Other)) {}
+#if COMPILE_WITHOUT_UNREAL_SUPPORT
 	FORCEINLINE TArrayND(std::istream& Stream)
 	    : Base(Stream) {}
+#endif
 	FORCEINLINE TArrayND<T, d>& operator=(const TArrayND<T, d>& Other)
 	{
 		Base::operator=(Other);
@@ -250,8 +253,10 @@ class TArrayND<T, 3> : public TArrayNDBase<TArrayND<T, 3>, T, 3>
 	FORCEINLINE TArrayND(const TArrayND<T, 3>& Other) = delete;
 	FORCEINLINE TArrayND(TArrayND<T, 3>&& Other)
 	    : Base(MoveTemp(Other)) {}
+#if COMPILE_WITHOUT_UNREAL_SUPPORT
 	FORCEINLINE TArrayND(std::istream& Stream)
 	    : Base(Stream) {}
+#endif
 	FORCEINLINE TArrayND<T, 3>& operator=(const TArrayND<T, 3>& Other)
 	{
 		Base::operator=(Other);
diff --git a/Engine/Source/Runtime/Experimental/ChaosCore/Public/Chaos/Vector.h b/Engine/Source/Runtime/Experimental/ChaosCore/Public/Chaos/Vector.h
--- a/Engine/Source/Runtime/Experimental/ChaosCore/Public/Chaos/Vector.h
+++ b/Engine/Source/Runtime/Experimental/ChaosCore/Public/Chaos/Vector.h
@@ -5,6 +5,11 @@
 #include "Math/Vector.h"
 #include "Math/Vector2D.h"
 #include "Math/Vector4.h"
+#include <cmath>
+#else
+#include <iostream>
+#include <utility>
+#include <limits>
 #endif
 
 #include "Chaos/Real.h"
@@ -12,10 +17,8 @@
 #include "Chaos/Pair.h"
 
 #include "Containers/StaticArray.h"
-#include <iostream>
-#include <utility>
+
 #include <initializer_list>
-#include <limits>
 
 namespace Chaos
 {
@@ -137,6 +140,7 @@ namespace Chaos
 			return NumElements;
 		}
 
+#if COMPILE_WITHOUT_UNREAL_SUPPORT
 		TVector(std::istream& Stream)
 		{
 			for (int32 i = 0; i < NumElements; ++i)
@@ -151,7 +155,7 @@ namespace Chaos
 				Stream.write(reinterpret_cast<const char*>(&V[i]), sizeof(FElement));
 			}
 		}
-
+#endif
 		friend bool operator==(const TVector<T, d>& L, const TVector<T, d>& R)
 		{
 			for (int32 i = 0; i < NumElements; ++i)
@@ -409,6 +413,8 @@ namespace Chaos
 		    : UE::Math::TVector<FRealSingle>(vec.X, vec.Y, vec.Z) {}
 		TVector(const UE::Math::TVector4<FRealDouble>& vec)					// LWC_TODO: Precision loss. Make explicit for FRealSingle = FRealSingle?
 			: UE::Math::TVector<FRealSingle>((UE::Math::TVector4<FRealSingle>)vec) {}
+
+#if COMPILE_WITHOUT_UNREAL_SUPPORT
 		TVector(std::istream& Stream)
 		{
 			Stream.read(reinterpret_cast<char*>(&X), sizeof(X));
@@ -422,6 +428,7 @@ namespace Chaos
 			Stream.write(reinterpret_cast<const char*>(&Y), sizeof(Y));
 			Stream.write(reinterpret_cast<const char*>(&Z), sizeof(Z));
 		}
+#endif
 		static inline TVector<FRealSingle, 3> Lerp(const TVector<FRealSingle, 3>& V1, const TVector<FRealSingle, 3>& V2, const FRealSingle F) { return FMath::Lerp<UE::Math::TVector<FRealSingle>, FRealSingle>(V1, V2, F); }
 		static inline TVector<FRealSingle, 3> CrossProduct(const TVector<FRealSingle, 3>& V1, const TVector<FRealSingle, 3>& V2) { return UE::Math::TVector<FRealSingle>::CrossProduct(V1, V2); }
 		static inline FRealSingle DotProduct(const TVector<FRealSingle, 3>& V1, const TVector<FRealSingle, 3>& V2) { return UE::Math::TVector<FRealSingle>::DotProduct(V1, V2); }
@@ -610,6 +617,7 @@ namespace Chaos
 			: UE::Math::TVector<FRealDouble>((UE::Math::TVector<FRealDouble>)vec) {}
 		TVector(const FVector4& vec)
 			: UE::Math::TVector<FRealDouble>(vec.X, vec.Y, vec.Z) {}
+#if COMPILE_WITHOUT_UNREAL_SUPPORT
 		TVector(std::istream& Stream)
 		{
 			Stream.read(reinterpret_cast<char*>(&X), sizeof(X));
@@ -623,6 +631,7 @@ namespace Chaos
 			Stream.write(reinterpret_cast<const char*>(&Y), sizeof(Y));
 			Stream.write(reinterpret_cast<const char*>(&Z), sizeof(Z));
 		}
+#endif
 		static inline TVector<FRealDouble, 3> Lerp(const TVector<FRealDouble, 3>& V1, const TVector<FRealDouble, 3>& V2, const FRealDouble F) { return FMath::Lerp<UE::Math::TVector<FRealDouble>, FRealDouble>(V1, V2, F); }
 		static inline TVector<FRealDouble, 3> CrossProduct(const TVector<FRealDouble, 3>& V1, const TVector<FRealDouble, 3>& V2) { return UE::Math::TVector<FRealDouble>::CrossProduct(V1, V2); }
 		static inline FRealDouble DotProduct(const TVector<FRealDouble, 3>& V1, const TVector<FRealDouble, 3>& V2) { return UE::Math::TVector<FRealDouble>::DotProduct(V1, V2); }
@@ -807,23 +816,26 @@ namespace Chaos
 		    : FVector2f((decltype(FVector2f::X))x, (decltype(FVector2f::X))y) {}
 		TVector(const FVector2f& vec)
 		    : FVector2f(vec) {}
+#if COMPILE_WITHOUT_UNREAL_SUPPORT
 		TVector(std::istream& Stream)
 		{
 			Stream.read(reinterpret_cast<char*>(&X), sizeof(X));
 			Stream.read(reinterpret_cast<char*>(&Y), sizeof(Y));
 		}
+#endif
 		template <typename OtherT>
 		TVector(const TVector<OtherT, 2>& InVector)
 		{
 			X = ((decltype(X))InVector[0]);	// LWC_TODO: Remove casts once FVector2f supports variants
 			Y = ((decltype(Y))InVector[1]);
 		}
+#if COMPILE_WITHOUT_UNREAL_SUPPORT
 		void Write(std::ostream& Stream) const
 		{
 			Stream.write(reinterpret_cast<const char*>(&X), sizeof(X));
 			Stream.write(reinterpret_cast<const char*>(&Y), sizeof(Y));
 		}
-
+#endif
 		static TVector<FRealSingle, 2> AxisVector(const int32 Axis)
 		{
 			check(Axis >= 0 && Axis <= 1);
@@ -892,23 +904,26 @@ namespace Chaos
 		    : FVector2d((decltype(FVector2d::X))x, (decltype(FVector2d::X))y) {}
 		TVector(const FVector2d& vec)
 		    : FVector2d(vec) {}
+#if COMPILE_WITHOUT_UNREAL_SUPPORT
 		TVector(std::istream& Stream)
 		{
 			Stream.read(reinterpret_cast<char*>(&X), sizeof(X));
 			Stream.read(reinterpret_cast<char*>(&Y), sizeof(Y));
 		}
+#endif
 		template <typename OtherT>
 		TVector(const TVector<OtherT, 2>& InVector)
 		{
 			X = ((decltype(X))InVector[0]);	// LWC_TODO: Remove casts once FVector2d supports variants
 			Y = ((decltype(Y))InVector[1]);
 		}
+#if COMPILE_WITHOUT_UNREAL_SUPPORT
 		void Write(std::ostream& Stream) const
 		{
 			Stream.write(reinterpret_cast<const char*>(&X), sizeof(X));
 			Stream.write(reinterpret_cast<const char*>(&Y), sizeof(Y));
 		}
-
+#endif
 		static TVector<FRealDouble, 2> AxisVector(const int32 Axis)
 		{
 			check(Axis >= 0 && Axis <= 1);
@@ -991,6 +1006,7 @@ namespace Chaos
 		    , Z(static_cast<T>(Other.Z))
 		{}
 
+#if COMPILE_WITHOUT_UNREAL_SUPPORT
 		FORCEINLINE TVector(std::istream& Stream)
 		{
 			Stream.read(reinterpret_cast<char*>(&X), sizeof(T));
@@ -1003,6 +1019,7 @@ namespace Chaos
 			Stream.write(reinterpret_cast<const char*>(&Y), sizeof(T));
 			Stream.write(reinterpret_cast<const char*>(&Z), sizeof(T));
 		}
+#endif
 		FORCEINLINE TVector<T, 3>& operator=(const TVector<T, 3>& Other)
 		{
 			X = Other.X;
@@ -1165,12 +1182,13 @@ PRAGMA_ENABLE_DEPRECATION_WARNINGS
 			: X((int32)InVector.X)
 			, Y((int32)InVector.Y)
 		{}
+#if COMPILE_WITHOUT_UNREAL_SUPPORT
 		FORCEINLINE TVector(std::istream& Stream)
 		{
 			Stream.read(reinterpret_cast<char*>(&X), sizeof(int32));
 			Stream.read(reinterpret_cast<char*>(&Y), sizeof(int32));
 		}
-
+#endif
 		FORCEINLINE int32 Num() const { return 2; }
 
 		FORCEINLINE FElement Product() const
@@ -1185,11 +1203,13 @@ PRAGMA_ENABLE_DEPRECATION_WARNINGS
 			return Result;
 		}
 
+#if COMPILE_WITHOUT_UNREAL_SUPPORT
 		FORCEINLINE void Write(std::ostream& Stream) const
 		{
 			Stream.write(reinterpret_cast<const char*>(&X), sizeof(FElement));
 			Stream.write(reinterpret_cast<const char*>(&Y), sizeof(FElement));
 		}
+#endif
 
 		FORCEINLINE TVector<FElement, 2>& operator=(const TVector<FElement, 2>& Other)
 		{
diff --git a/Engine/Source/Runtime/RHI/Public/RHICommandList.h b/Engine/Source/Runtime/RHI/Public/RHICommandList.h
--- a/Engine/Source/Runtime/RHI/Public/RHICommandList.h
+++ b/Engine/Source/Runtime/RHI/Public/RHICommandList.h
@@ -25,6 +25,7 @@
 #include "HAL/IConsoleManager.h"
 #include "Async/TaskGraphInterfaces.h"
 #include "HAL/LowLevelMemTracker.h"
+#include "RHI.h"
 #include "ProfilingDebugging/CsvProfiler.h"
 #include "ProfilingDebugging/CpuProfilerTrace.h"
 #include "Trace/Trace.h"
@@ -1271,72 +1272,6 @@ FRHICOMMAND_MACRO(FRHICommandNextSubpass)
 	RHI_API void Execute(FRHICommandListBase& CmdList);
 };
 
-struct FLocalCmdListParallelRenderPass
-{
-	TRefCountPtr<struct FRHIParallelRenderPass> RenderPass;
-};
-
-FRHICOMMAND_MACRO(FRHICommandBeginParallelRenderPass)
-{
-	FRHIRenderPassInfo Info;
-	FLocalCmdListParallelRenderPass* LocalRenderPass;
-	const TCHAR* Name;
-
-	FRHICommandBeginParallelRenderPass(const FRHIRenderPassInfo& InInfo, FLocalCmdListParallelRenderPass* InLocalRenderPass, const TCHAR* InName)
-		: Info(InInfo)
-		, LocalRenderPass(InLocalRenderPass)
-		, Name(InName)
-	{
-	}
-
-	RHI_API void Execute(FRHICommandListBase& CmdList);
-};
-
-FRHICOMMAND_MACRO(FRHICommandEndParallelRenderPass)
-{
-	FLocalCmdListParallelRenderPass* LocalRenderPass;
-
-	FRHICommandEndParallelRenderPass(FLocalCmdListParallelRenderPass* InLocalRenderPass)
-		: LocalRenderPass(InLocalRenderPass)
-	{
-	}
-
-	RHI_API void Execute(FRHICommandListBase& CmdList);
-};
-
-struct FLocalCmdListRenderSubPass
-{
-	TRefCountPtr<struct FRHIRenderSubPass> RenderSubPass;
-};
-
-FRHICOMMAND_MACRO(FRHICommandBeginRenderSubPass)
-{
-	FLocalCmdListParallelRenderPass* LocalRenderPass;
-	FLocalCmdListRenderSubPass* LocalRenderSubPass;
-
-	FRHICommandBeginRenderSubPass(FLocalCmdListParallelRenderPass* InLocalRenderPass, FLocalCmdListRenderSubPass* InLocalRenderSubPass)
-		: LocalRenderPass(InLocalRenderPass)
-		, LocalRenderSubPass(InLocalRenderSubPass)
-	{
-	}
-
-	RHI_API void Execute(FRHICommandListBase& CmdList);
-};
-
-FRHICOMMAND_MACRO(FRHICommandEndRenderSubPass)
-{
-	FLocalCmdListParallelRenderPass* LocalRenderPass;
-	FLocalCmdListRenderSubPass* LocalRenderSubPass;
-
-	FRHICommandEndRenderSubPass(FLocalCmdListParallelRenderPass* InLocalRenderPass, FLocalCmdListRenderSubPass* InLocalRenderSubPass)
-		: LocalRenderPass(InLocalRenderPass)
-		, LocalRenderSubPass(InLocalRenderSubPass)
-	{
-	}
-
-	RHI_API void Execute(FRHICommandListBase& CmdList);
-};
-
 struct FRHICommandSetComputeShaderString
 {
 	static const TCHAR* TStr() { return TEXT("FRHICommandSetComputeShader"); }
@@ -1727,67 +1662,6 @@ FRHICOMMAND_MACRO(FRHICommandWriteGPUFence)
 	RHI_API void Execute(FRHICommandListBase& CmdList);
 };
 
-FRHICOMMAND_MACRO(FRHICommandClearColorTexture)
-{
-	FRHITexture* Texture;
-	FLinearColor Color;
-
-	FORCEINLINE_DEBUGGABLE FRHICommandClearColorTexture(
-		FRHITexture* InTexture,
-		const FLinearColor& InColor
-		)
-		: Texture(InTexture)
-		, Color(InColor)
-	{
-	}
-	RHI_API void Execute(FRHICommandListBase& CmdList);
-};
-
-FRHICOMMAND_MACRO(FRHICommandClearDepthStencilTexture)
-{
-	FRHITexture* Texture;
-	float Depth;
-	uint32 Stencil;
-	EClearDepthStencil ClearDepthStencil;
-
-	FORCEINLINE_DEBUGGABLE FRHICommandClearDepthStencilTexture(
-		FRHITexture* InTexture,
-		EClearDepthStencil InClearDepthStencil,
-		float InDepth,
-		uint32 InStencil
-	)
-		: Texture(InTexture)
-		, Depth(InDepth)
-		, Stencil(InStencil)
-		, ClearDepthStencil(InClearDepthStencil)
-	{
-	}
-	RHI_API void Execute(FRHICommandListBase& CmdList);
-};
-
-FRHICOMMAND_MACRO(FRHICommandClearColorTextures)
-{
-	FLinearColor ColorArray[MaxSimultaneousRenderTargets];
-	FRHITexture* Textures[MaxSimultaneousRenderTargets];
-	int32 NumClearColors;
-
-	FORCEINLINE_DEBUGGABLE FRHICommandClearColorTextures(
-		int32 InNumClearColors,
-		FRHITexture** InTextures,
-		const FLinearColor* InColorArray
-		)
-		: NumClearColors(InNumClearColors)
-	{
-		check(InNumClearColors <= MaxSimultaneousRenderTargets);
-		for (int32 Index = 0; Index < InNumClearColors; Index++)
-		{
-			ColorArray[Index] = InColorArray[Index];
-			Textures[Index] = InTextures[Index];
-		}
-	}
-	RHI_API void Execute(FRHICommandListBase& CmdList);
-};
-
 FRHICOMMAND_MACRO(FRHICommandSetStaticUniformBuffers)
 {
 	FUniformBufferStaticBindings UniformBuffers;
diff --git a/Engine/Source/Runtime/Renderer/Public/GlobalDistanceFieldParameters.h b/Engine/Source/Runtime/Renderer/Public/GlobalDistanceFieldParameters.h
--- a/Engine/Source/Runtime/Renderer/Public/GlobalDistanceFieldParameters.h
+++ b/Engine/Source/Runtime/Renderer/Public/GlobalDistanceFieldParameters.h
@@ -3,6 +3,7 @@
 #pragma once
 
 #include "CoreMinimal.h"
+#include "GlobalDistanceFieldConstants.h"
 #include "ShaderParameterMacros.h"
 #include "ShaderParameters.h"
 #include "RenderUtils.h"
@@ -11,12 +12,6 @@
 
 class FShaderParameterMap;
 
-namespace GlobalDistanceField
-{
-	/** Must match global distance field shaders. */
-	const int32 MaxClipmaps = 6;
-}
-
 class FGlobalDistanceFieldParameterData
 {
 public:
diff --git a/Engine/Source/Runtime/SlateCore/Private/Rendering/ShaderResourceManager.cpp b/Engine/Source/Runtime/SlateCore/Private/Rendering/ShaderResourceManager.cpp
--- a/Engine/Source/Runtime/SlateCore/Private/Rendering/ShaderResourceManager.cpp
+++ b/Engine/Source/Runtime/SlateCore/Private/Rendering/ShaderResourceManager.cpp
@@ -4,6 +4,30 @@
 
 DECLARE_CYCLE_STAT(TEXT("GetResourceHandle Time"), STAT_SlateGetResourceHandle, STATGROUP_Slate);
 
+FSlateResourceHandle::FSlateResourceHandle()
+{
+}
+
+FSlateResourceHandle::~FSlateResourceHandle()
+{
+}
+
+bool FSlateResourceHandle::IsValid() const
+{
+	return Data.IsValid() && Data->Proxy;
+}
+
+const FSlateShaderResourceProxy* FSlateResourceHandle::GetResourceProxy() const
+{
+	return Data.IsValid() ? Data->Proxy : nullptr;
+}
+
+FSlateResourceHandle::FSlateResourceHandle(const TSharedPtr<FSlateSharedHandleData>& InData)
+	: Data(InData)
+{
+}
+
+
 FSlateResourceHandle FSlateShaderResourceManager::GetResourceHandle(const FSlateBrush& Brush, FVector2f LocalSize, float DrawScale)
 {
 	SCOPE_CYCLE_COUNTER(STAT_SlateGetResourceHandle);
diff --git a/Engine/Source/Runtime/SlateCore/Private/Styling/SlateBrush.cpp b/Engine/Source/Runtime/SlateCore/Private/Styling/SlateBrush.cpp
--- a/Engine/Source/Runtime/SlateCore/Private/Styling/SlateBrush.cpp
+++ b/Engine/Source/Runtime/SlateCore/Private/Styling/SlateBrush.cpp
@@ -139,6 +139,18 @@ const FString FSlateBrush::UTextureIdentifier()
 	return FString(TEXT("texture:/"));
 }
 
+const FSlateResourceHandle& FSlateBrush::GetRenderingResource() const
+{
+	if (ImageType == ESlateBrushImageType::Vector)
+	{
+		UE_LOG(LogSlate, Warning, TEXT("FSlateBrush::GetRenderingResource should be called with a size and scale for vector brushes"));
+	}
+	
+	UpdateRenderingResource(GetImageSize(), 1.0f);
+
+	return ResourceHandle;
+}
+
 void FSlateBrush::UpdateRenderingResource(FVector2D LocalSize, float DrawScale) const
 {
 	if (DrawAs != ESlateBrushDrawType::NoDrawType && (ResourceName != NAME_None || ResourceObject != nullptr))
diff --git a/Engine/Source/Runtime/SlateCore/Public/Rendering/DrawElementPayloads.h b/Engine/Source/Runtime/SlateCore/Public/Rendering/DrawElementPayloads.h
--- a/Engine/Source/Runtime/SlateCore/Public/Rendering/DrawElementPayloads.h
+++ b/Engine/Source/Runtime/SlateCore/Public/Rendering/DrawElementPayloads.h
@@ -2,6 +2,8 @@
 
 #pragma once
 
+#include "DrawElementTextOverflowArgs.h"
+#include "Fonts/FontCache.h"
 #include "SlateRenderBatch.h"
 #include "Styling/SlateBrush.h"
 #include "Types/SlateVector2.h"
@@ -33,16 +35,6 @@ struct FSlateGradientStop
 };
 template <> struct TIsPODType<FSlateGradientStop> { enum { Value = true }; };
 
-enum class ETextOverflowDirection : uint8
-{
-	// No overflow
-	NoOverflow,
-	// Left justification overflow
-	LeftToRight,
-	// Right justification overflow
-	RightToLeft
-};
-
 struct FSlateDataPayload
 {
 	virtual ~FSlateDataPayload() {}
@@ -156,28 +148,6 @@ struct FSlateTextPayload : public FSlateDataPayload, public FSlateTintableElemen
 };
 
 
-
-struct FTextOverflowArgs
-{
-	FTextOverflowArgs(FShapedGlyphSequencePtr& InOverflowText, ETextOverflowDirection InOverflowDirection)
-		: OverflowTextPtr(InOverflowText)
-		, OverflowDirection(InOverflowDirection)
-		, bForceEllipsisDueToClippedLine(false)
-		
-	{}
-
-	FTextOverflowArgs()
-		: OverflowDirection(ETextOverflowDirection::NoOverflow)
-		, bForceEllipsisDueToClippedLine(false)
-	{}
-
-	/** Sequence that represents the ellipsis glyph */
-	FShapedGlyphSequencePtr OverflowTextPtr;
-	ETextOverflowDirection OverflowDirection;
-	bool bForceEllipsisDueToClippedLine;
-	
-};
-
 struct FSlateShapedTextPayload : public FSlateDataPayload, public FSlateTintableElement
 { 
 	// Shaped text data
diff --git a/Engine/Source/Runtime/SlateCore/Public/Rendering/DrawElements.h b/Engine/Source/Runtime/SlateCore/Public/Rendering/DrawElements.h
--- a/Engine/Source/Runtime/SlateCore/Public/Rendering/DrawElements.h
+++ b/Engine/Source/Runtime/SlateCore/Public/Rendering/DrawElements.h
@@ -9,25 +9,32 @@
 #include "Misc/MemStack.h"
 #include "Styling/WidgetStyle.h"
 #include "Fonts/SlateFontInfo.h"
-#include "Fonts/FontCache.h"
 #include "Layout/SlateRect.h"
 #include "Layout/Clipping.h"
 #include "Types/PaintArgs.h"
 #include "Types/SlateVector2.h"
 #include "Layout/Geometry.h"
-#include "Rendering/ShaderResourceManager.h"
 #include "Rendering/RenderingCommon.h"
 #include "Debugging/SlateDebugging.h"
-#include "DrawElementPayloads.h"
+#include "Rendering/SlateRenderBatch.h"
+#include "DrawElementTextOverflowArgs.h"
 #include "ElementBatcher.h"
 
+#if UE_ENABLE_INCLUDE_ORDER_DEPRECATED_IN_5_2
+#include "Fonts/FontCache.h"
+#include "Rendering/ShaderResourceManager.h"
+#include "DrawElementPayloads.h"
+#endif
 
 class FSlateRenderBatch;
+class FSlateDrawLayerHandle;
+class FSlateResourceHandle;
 class FSlateWindowElementList;
 class SWidget;
 class SWindow;
-
-class FSlateDrawLayerHandle;
+struct FSlateBrush;
+struct FSlateDataPayload;
+struct FSlateGradientStop;
 
 enum class EElementType : uint8
 {
diff --git a/Engine/Source/Runtime/SlateCore/Public/Styling/SlateBrush.h b/Engine/Source/Runtime/SlateCore/Public/Styling/SlateBrush.h
--- a/Engine/Source/Runtime/SlateCore/Public/Styling/SlateBrush.h
+++ b/Engine/Source/Runtime/SlateCore/Public/Styling/SlateBrush.h
@@ -6,7 +6,10 @@
 #include "UObject/ObjectMacros.h"
 #include "Styling/SlateColor.h"
 #include "Layout/Margin.h"
+#include "Rendering/SlateResourceHandle.h"
+#if UE_ENABLE_INCLUDE_ORDER_DEPRECATED_IN_5_2
 #include "Textures/SlateShaderResource.h"
+#endif
 #include "Types/SlateBox2.h"
 #include "Types/SlateVector2.h"
 #include "SlateBrush.generated.h"
@@ -441,17 +444,7 @@ public:
 		return ResourceHandle;
 	}
 
-	const FSlateResourceHandle& GetRenderingResource() const
-	{
-		if (ImageType == ESlateBrushImageType::Vector)
-		{
-			UE_LOG(LogSlate, Warning, TEXT("FSlateBrush::GetRenderingResource should be called with a size and scale for vector brushes"));
-		}
-	
-		UpdateRenderingResource(GetImageSize(), 1.0f);
-
-		return ResourceHandle;
-	}
+	const FSlateResourceHandle& GetRenderingResource() const;
 
 	bool IsSet() const { return bIsSet; }
 
diff --git a/Engine/Source/Runtime/SlateCore/Public/Textures/SlateShaderResource.h b/Engine/Source/Runtime/SlateCore/Public/Textures/SlateShaderResource.h
--- a/Engine/Source/Runtime/SlateCore/Public/Textures/SlateShaderResource.h
+++ b/Engine/Source/Runtime/SlateCore/Public/Textures/SlateShaderResource.h
@@ -5,6 +5,7 @@
 #include "CoreMinimal.h"
 #include "Stats/Stats.h"
 #include "SlateGlobals.h"
+#include "Rendering/SlateResourceHandle.h"
 
 class FSlateShaderResourceProxy;
 
@@ -201,36 +202,3 @@ class IViewportRenderTargetProvider
 public:
 	virtual FSlateShaderResource* GetViewportRenderTargetTexture() = 0;
 };
-
-/**
- * A SlateResourceHandle is used as fast path for looking up a rendering resource for a given brush when adding Slate draw elements
- * This can be cached and stored safely in code.  It will become invalid when a resource is destroyed
-*/
-class FSlateResourceHandle
-{
-	friend class FSlateShaderResourceManager;
-	friend class FSlateNullShaderResourceManager;
-public:
-	FSlateResourceHandle() {}
-
-	/**
-	 * @return true if the handle still points to a valid rendering resource
-	 */
-	bool IsValid() const { return Data.IsValid() && Data->Proxy; }
-
-	/**
-	 * @return the resource proxy used to render.
-	 */
-	const FSlateShaderResourceProxy* GetResourceProxy() const
-	{
-		return Data.IsValid() ? Data->Proxy : nullptr;
-	}
-
-private:
-	FSlateResourceHandle(const TSharedPtr<FSlateSharedHandleData>& InData)
-		: Data(InData)
-	{}
-
-	/** Internal data to pair the handle to the resource */
-	TSharedPtr<FSlateSharedHandleData> Data;
-};
diff --git a/Engine/Source/Runtime/SlateCore/Public/Types/SlateEnums.h b/Engine/Source/Runtime/SlateCore/Public/Types/SlateEnums.h
--- a/Engine/Source/Runtime/SlateCore/Public/Types/SlateEnums.h
+++ b/Engine/Source/Runtime/SlateCore/Public/Types/SlateEnums.h
@@ -14,7 +14,7 @@ namespace EButtonClickMethod
 	/**
 	 * Enumerates different methods that a button click can be triggered. Normally, DownAndUp is appropriate.
 	 */
-	enum Type
+	enum Type : int
 	{
 		/**
 		 * User must press the button, then release while over the button to trigger the click.
@@ -46,7 +46,7 @@ UENUM(BlueprintType)
 namespace EButtonTouchMethod
 {
 	/** Ways in which touch interactions trigger a "Clicked" event. */
-	enum Type
+	enum Type : int
 	{
 		/** Most buttons behave this way. */
 		DownAndUp,
@@ -70,7 +70,7 @@ namespace EButtonPressMethod
 	/**
 	* Enumerates different methods that a button can be triggered with keyboard/controller. Normally, DownAndUp is appropriate.
 	*/
-	enum Type
+	enum Type : int
 	{
 		/**
 		* User must press the button, then release while the button has focus to trigger the click.
@@ -170,7 +170,7 @@ enum class ENavigationGenesis : uint8
  * Enumerates horizontal alignment options, i.e. for widget slots.
  */
 UENUM(BlueprintType)
-enum EHorizontalAlignment
+enum EHorizontalAlignment : int
 {
 	/** Fill the entire width. */
 	HAlign_Fill UMETA(DisplayName="Fill"),
@@ -190,7 +190,7 @@ enum EHorizontalAlignment
  * Enumerates vertical alignment options, i.e. for widget slots.
  */
 UENUM(BlueprintType)
-enum EVerticalAlignment
+enum EVerticalAlignment : int
 {
 	/** Fill the entire height. */
 	VAlign_Fill UMETA(DisplayName="Fill"),
@@ -210,7 +210,7 @@ enum EVerticalAlignment
  * Enumerates possible placements for pop-up menus.
  */
 UENUM(BlueprintType)
-enum EMenuPlacement
+enum EMenuPlacement : int
 {
 	/** Place the menu immediately below the anchor */
 	MenuPlacement_BelowAnchor UMETA(DisplayName="Below"),
@@ -257,7 +257,7 @@ enum EMenuPlacement
  * Enumerates widget orientations.
  */
 UENUM(BlueprintType)
-enum EOrientation
+enum EOrientation : int
 {
 	/** Orient horizontally, i.e. left to right. */
 	Orient_Horizontal UMETA(DisplayName="Horizontal"),
@@ -271,7 +271,7 @@ enum EOrientation
  * Enumerates scroll directions.
  */
 UENUM(BlueprintType)
-enum EScrollDirection
+enum EScrollDirection : int
 {
 	/** Scroll down. */
 	Scroll_Down UMETA(DisplayName="Down"),
@@ -287,7 +287,7 @@ enum EScrollDirection
 UENUM(BlueprintType)
 namespace ETextCommit
 {
-	enum Type
+	enum Type : int
 	{
 		/** Losing focus or similar event caused implicit commit */
 		Default,
@@ -307,7 +307,7 @@ namespace ETextCommit
 UENUM(BlueprintType)
 namespace ESelectInfo
 {
-	enum Type
+	enum Type : int
 	{
 		/** User selected via a key press */
 		OnKeyPress,
diff --git a/Engine/Source/Runtime/UMG/Public/Blueprint/UserWidget.h b/Engine/Source/Runtime/UMG/Public/Blueprint/UserWidget.h
--- a/Engine/Source/Runtime/UMG/Public/Blueprint/UserWidget.h
+++ b/Engine/Source/Runtime/UMG/Public/Blueprint/UserWidget.h
@@ -24,7 +24,9 @@
 #include "EngineStats.h"
 #include "SlateGlobals.h"
 #include "Animation/WidgetAnimationEvents.h"
+#if UE_ENABLE_INCLUDE_ORDER_DEPRECATED_IN_5_2
 #include "Blueprint/WidgetBlueprintGeneratedClass.h"
+#endif
 
 #if UE_ENABLE_INCLUDE_ORDER_DEPRECATED_IN_5_1
 	#include "Animation/WidgetAnimation.h"
@@ -39,6 +41,7 @@ class UTexture2D;
 class UUMGSequencePlayer;
 class UUMGSequenceTickManager;
 class UWidgetAnimation;
+class UWidgetBlueprintGeneratedClass;
 class UWidgetTree;
 class UNamedSlot;
 class UUserWidgetExtension;
