Undoing CL 22691328 to prevent crashes when "Show Navigation" is enabled.

diff --git a/Engine/Source/Runtime/Engine/Private/Materials/MaterialShared.cpp b/Engine/Source/Runtime/Engine/Private/Materials/MaterialShared.cpp
--- a/Engine/Source/Runtime/Engine/Private/Materials/MaterialShared.cpp
+++ b/Engine/Source/Runtime/Engine/Private/Materials/MaterialShared.cpp
@@ -4016,9 +4016,7 @@ void FMaterialRenderProxy::EvaluateUniformExpressions(FUniformExpressionCache& O
 	OutUniformExpressionCache.ParameterCollections = UniformExpressionSet.ParameterCollections;
 
 	++UniformExpressionCacheSerialNumber;
-	// The cache is invalidated whenever new shaders are added, so it's fine to always set this to true here
-	// instead of checking if the shader map is complete.
-	OutUniformExpressionCache.bUpToDate = true;
+	OutUniformExpressionCache.bUpToDate = Context.Material.IsRenderingThreadShaderMapComplete();
 }
 
 void FMaterialRenderProxy::CacheUniformExpressions(bool bRecreateUniformBuffer)
