World Partition:
- Unregister event delegates when destroying an actor desc container.
- Added missing delegate unregistration.

#rnx

diff --git a/Engine/Source/Runtime/Engine/Private/WorldPartition/ActorDescContainerCollection.cpp b/Engine/Source/Runtime/Engine/Private/WorldPartition/ActorDescContainerCollection.cpp
--- a/Engine/Source/Runtime/Engine/Private/WorldPartition/ActorDescContainerCollection.cpp
+++ b/Engine/Source/Runtime/Engine/Private/WorldPartition/ActorDescContainerCollection.cpp
@@ -6,18 +6,25 @@
 #include "WorldPartition/WorldPartition.h"
 
 #if WITH_EDITOR
+FActorDescContainerCollection::~FActorDescContainerCollection()
+{
+	Empty();
+}
 
 void FActorDescContainerCollection::AddContainer(UActorDescContainer* Container)
 {
 	ActorDescContainerCollection.Add(Container);
-	Container->OnActorDescAddedEvent.AddRaw(this, &FActorDescContainerCollection::OnActorDescAdded);
-	Container->OnActorDescRemovedEvent.AddRaw(this, &FActorDescContainerCollection::OnActorDescRemoved);
+	RegisterDelegates(Container);
 }
 
 bool FActorDescContainerCollection::RemoveContainer(UActorDescContainer* Container)
 {
-	Container->OnActorDescRemovedEvent.RemoveAll(this);
-	return ActorDescContainerCollection.RemoveSwap(Container) > 0;
+	if (ActorDescContainerCollection.RemoveSwap(Container) > 0)
+	{
+		UnregisterDelegates(Container);
+		return true;
+	}
+	return false;
 }
 
 bool FActorDescContainerCollection::Contains(const FName& ContainerPackageName) const
@@ -31,6 +38,28 @@ UActorDescContainer* FActorDescContainerCollection::FindContainer(const FName& C
 	return ContainerPtr != nullptr ? *ContainerPtr : nullptr;
 }
 
+void FActorDescContainerCollection::Empty()
+{
+	ForEachActorDescContainer([this](UActorDescContainer* ActorDescContainer)
+	{
+		UnregisterDelegates(ActorDescContainer);
+	});
+
+	ActorDescContainerCollection.Empty();
+}
+
+void FActorDescContainerCollection::RegisterDelegates(UActorDescContainer* Container)
+{
+	Container->OnActorDescAddedEvent.AddRaw(this, &FActorDescContainerCollection::OnActorDescAdded);
+	Container->OnActorDescRemovedEvent.AddRaw(this, &FActorDescContainerCollection::OnActorDescRemoved);
+}
+
+void FActorDescContainerCollection::UnregisterDelegates(UActorDescContainer* Container)
+{
+	Container->OnActorDescAddedEvent.RemoveAll(this);
+	Container->OnActorDescRemovedEvent.RemoveAll(this);
+}
+
 const FWorldPartitionActorDesc* FActorDescContainerCollection::GetActorDesc(const FGuid& Guid) const
 {
 	const FWorldPartitionActorDesc* ActorDesc = nullptr;
diff --git a/Engine/Source/Runtime/Engine/Public/WorldPartition/ActorDescContainerCollection.h b/Engine/Source/Runtime/Engine/Public/WorldPartition/ActorDescContainerCollection.h
--- a/Engine/Source/Runtime/Engine/Public/WorldPartition/ActorDescContainerCollection.h
+++ b/Engine/Source/Runtime/Engine/Public/WorldPartition/ActorDescContainerCollection.h
@@ -12,12 +12,14 @@ class ENGINE_API FActorDescContainerCollection
 {
 #if WITH_EDITOR
 public:
+	~FActorDescContainerCollection();
+
 	void AddContainer(UActorDescContainer* Container);
 	bool RemoveContainer(UActorDescContainer* Container);
 	bool Contains(const FName& ContainerPackageName) const;
 	UActorDescContainer* FindContainer(const FName& ContainerPackageName) const;
 
-	void Empty() { ActorDescContainerCollection.Empty(); }
+	void Empty();
 	bool IsEmpty() const { return ActorDescContainerCollection.IsEmpty(); }
 	uint32 GetActorDescContainerCount() const { return ActorDescContainerCollection.Num(); }
 
@@ -54,6 +56,9 @@ public:
 	void ForEachActorDescContainer(TFunctionRef<void(UActorDescContainer*)> Func) const;
 
 protected:
+	void RegisterDelegates(UActorDescContainer* Container);
+	void UnregisterDelegates(UActorDescContainer* Container);
+
 	TArray<UActorDescContainer*> ActorDescContainerCollection;
 
 public:
@@ -169,7 +174,5 @@ public:
 			: BaseType(Collection, InActorClass ? InActorClass : ActorType::StaticClass())
 		{}
 	};
-
-
 #endif
 };
