[Engine]
* GameViewportClient - Added include to fix non-unity

diff --git a/Engine/Source/Runtime/Engine/Private/GameViewportClient.cpp b/Engine/Source/Runtime/Engine/Private/GameViewportClient.cpp
--- a/Engine/Source/Runtime/Engine/Private/GameViewportClient.cpp
+++ b/Engine/Source/Runtime/Engine/Private/GameViewportClient.cpp
@@ -27,6 +27,7 @@
 #include "UnrealEngine.h"
 #include "EngineUtils.h"
 #include "Framework/Application/SlateApplication.h"
+#include "Widgets/SOverlay.h"
 #include "Widgets/SViewport.h"
 #include "Engine/Console.h"
 #include "GameFramework/HUD.h"
