[Engine]
* NetDriver - Removed include of WorldSetttings

Hopefully this will fix the odd Battleye compile error on linux

diff --git a/Engine/Source/Editor/SequenceRecorder/Private/SequenceRecorder.cpp b/Engine/Source/Editor/SequenceRecorder/Private/SequenceRecorder.cpp
--- a/Engine/Source/Editor/SequenceRecorder/Private/SequenceRecorder.cpp
+++ b/Engine/Source/Editor/SequenceRecorder/Private/SequenceRecorder.cpp
@@ -1,6 +1,7 @@
 // Copyright Epic Games, Inc. All Rights Reserved.
 
 #include "SequenceRecorder.h"
+#include "GameFramework/WorldSettings.h"
 #include "ISequenceAudioRecorder.h"
 #include "Misc/ScopedSlowTask.h"
 #include "Modules/ModuleManager.h"
diff --git a/Engine/Source/Runtime/Engine/Classes/Engine/NetDriver.h b/Engine/Source/Runtime/Engine/Classes/Engine/NetDriver.h
--- a/Engine/Source/Runtime/Engine/Classes/Engine/NetDriver.h
+++ b/Engine/Source/Runtime/Engine/Classes/Engine/NetDriver.h
@@ -5,13 +5,17 @@
 #pragma once
 
 #include "CoreMinimal.h"
+#include "Engine/EngineBaseTypes.h"
+#include "Engine/EngineTypes.h"
 #include "Math/RandomStream.h"
 #include "UObject/ObjectMacros.h"
 #include "UObject/UObjectGlobals.h"
 #include "UObject/Object.h"
 #include "Misc/NetworkGuid.h"
 #include "UObject/CoreNet.h"
+#if UE_ENABLE_INCLUDE_ORDER_DEPRECATED_IN_5_2
 #include "GameFramework/WorldSettings.h"
+#endif
 #include "PacketHandler.h"
 #include "Channel.h"
 #include "Net/Core/Misc/DDoSDetection.h"
