On some environments, cooking happens without network, so we can't retrieve the MAC adress yo generate a globally unique name for spawned actors, assign a default one in this case.

#rnx

diff --git a/Engine/Source/Runtime/Engine/Private/LevelActor.cpp b/Engine/Source/Runtime/Engine/Private/LevelActor.cpp
--- a/Engine/Source/Runtime/Engine/Private/LevelActor.cpp
+++ b/Engine/Source/Runtime/Engine/Private/LevelActor.cpp
@@ -280,8 +280,10 @@ void LineCheckTracker::CaptureLineCheck(int32 LineCheckFlags, const FVector* Ext
  *	  change every ~17 minutes for a specific actor class).
  *  - The name number bit 30 is reserved to know if this is a globally unique name. This
  *	  is not 100% safe, but should cover most cases.
- *  - The name number bit 31 is reserved to preserve the  fast path name generation (see
+ *  - The name number bit 31 is reserved to preserve the fast path name generation (see
  *	  GFastPathUniqueNameGeneration).
+ *	- On some environments, cooking happens without network, so we can't retrieve the 
+ *	  MAC adress for spawned actors, assign a default one in this case.
  **/
 class FActorGUIDGenerator
 {
@@ -291,7 +293,13 @@ public:
 		, MacAddress(FPlatformMisc::GetMacAddress())
 		, Counter(0)
 	{
-		check(MacAddress.Num() == 6);
+		if (MacAddress.Num() != 6)
+		{
+			// During cooking, some platform can't retrieve the MAC adress, so force a default one.
+			check(IsRunningCookCommandlet());
+			MacAddress.Empty(6);
+			MacAddress.AddZeroed(6);
+		}
 	}
 
 	FName NewActorGUID(FName BaseName)
@@ -326,8 +334,8 @@ public:
 	}
 	
 private:
-	const FDateTime Origin;
-	const TArray<uint8> MacAddress;
+	FDateTime Origin;
+	TArray<uint8> MacAddress;
 	uint32 Counter;
 };
 #endif
