Blind fix to fix compile errors by adding forward declaration of ESkeletalMeshVertexFlags. Also had to add a fixed backing type for the enum to make it forward declarable

diff --git a/Engine/Source/Runtime/Engine/Classes/Engine/SkinnedAsset.h b/Engine/Source/Runtime/Engine/Classes/Engine/SkinnedAsset.h
--- a/Engine/Source/Runtime/Engine/Classes/Engine/SkinnedAsset.h
+++ b/Engine/Source/Runtime/Engine/Classes/Engine/SkinnedAsset.h
@@ -31,7 +31,7 @@ class UMorphTarget;
 class UPhysicsAsset;
 class USkeleton;
 struct FSkinnedAssetAsyncBuildTask;
-
+enum class ESkeletalMeshVertexFlags : uint8;
 enum class ESkinnedAssetAsyncPropertyLockType
 {
 	None = 0,
diff --git a/Engine/Source/Runtime/Engine/Public/SkeletalMeshTypes.h b/Engine/Source/Runtime/Engine/Public/SkeletalMeshTypes.h
--- a/Engine/Source/Runtime/Engine/Public/SkeletalMeshTypes.h
+++ b/Engine/Source/Runtime/Engine/Public/SkeletalMeshTypes.h
@@ -24,7 +24,7 @@ class FSkeletalMeshRenderData;
 class FSkeletalMeshLODRenderData;
 
 /** Flags used when building vertex buffers. */
-enum class ESkeletalMeshVertexFlags
+enum class ESkeletalMeshVertexFlags : uint8
 {
 	None								= 0x0,
 	UseFullPrecisionUVs					= 0x1,
