Turn off StreamableRenderAsset FResourceSizeNeedsUpdating as it is causing multi-second hitches constantly in editor when camera moves

diff --git a/Engine/Source/Runtime/Engine/Private/StreamableRenderAsset.cpp b/Engine/Source/Runtime/Engine/Private/StreamableRenderAsset.cpp
--- a/Engine/Source/Runtime/Engine/Private/StreamableRenderAsset.cpp
+++ b/Engine/Source/Runtime/Engine/Private/StreamableRenderAsset.cpp
@@ -18,6 +18,14 @@ static FAutoConsoleVariableRef CVarNoRefBiasQualityLevel(
 	TEXT("The quality level for the no-ref mesh streaming LOD bias"),
 	ECVF_Scalability);
 
+namespace StreamableRenderAsset
+{
+#if WITH_EDITOR
+	bool bAllowUpdateResourceSize = false;
+	FAutoConsoleVariableRef CVarAllowUpdateResourceSize(TEXT("r.Streaming.AllowUpdateResourceSize"), bAllowUpdateResourceSize, TEXT("AllowUpdateResourceSize"));
+#endif
+}
+
 extern bool TrackRenderAssetEvent(struct FStreamingRenderAsset* StreamingRenderAsset, UStreamableRenderAsset* RenderAsset, bool bForceMipLevelsToBeResident, const FRenderAssetStreamingManager* Manager);
 
 UStreamableRenderAsset::UStreamableRenderAsset(const FObjectInitializer& ObjectInitializer)
@@ -146,11 +154,13 @@ private:
 			FScopeLock ScopeLock(&Lock);
 			if (Pending.Num() > 0)
 			{
+				TRACE_CPUPROFILER_EVENT_SCOPE(FResourceSizeNeedsUpdating::BroadcastOnObjectPropertyChanged);
 				FPropertyChangedEvent EmptyPropertyChangedEvent(nullptr);
 				for (const TWeakObjectPtr<UObject>& WeakObjectPtr : Pending)
 				{
 					if (UObject* Obj = WeakObjectPtr.Get())
 					{
+						// Note: This is too expensive 0.05 to 3 seconds per call, needs to be re-written in a more performance friendly manner
 						FCoreUObjectDelegates::OnObjectPropertyChanged.Broadcast(Obj, EmptyPropertyChangedEvent);
 					}
 				}
@@ -190,7 +200,7 @@ void UStreamableRenderAsset::TickStreaming(bool bSendCompletionEvents, TArray<US
 			PendingUpdate.SafeRelease();
 
 #if WITH_EDITOR
-			if (GIsEditor && bSendCompletionEvents)
+			if (StreamableRenderAsset::bAllowUpdateResourceSize && GIsEditor && bSendCompletionEvents)
 			{
 				// When all the requested mips are streamed in, generate an empty property changed event, to force the
 				// ResourceSize asset registry tag to be recalculated.
