[Engine]
* Removed includes from SoundBase.h since it is frequently included outside engine
* enum EAudioSpeakers  -> enum EAudioSpeakers : int (This just makes it possible to forward declare enum.. nothing else changes. Not defining is the same as being int)

diff --git a/Engine/Plugins/Experimental/MotoSynth/Source/MotoSynth/Public/MotoSynthEngine.h b/Engine/Plugins/Experimental/MotoSynth/Source/MotoSynth/Public/MotoSynthEngine.h
--- a/Engine/Plugins/Experimental/MotoSynth/Source/MotoSynth/Public/MotoSynthEngine.h
+++ b/Engine/Plugins/Experimental/MotoSynth/Source/MotoSynth/Public/MotoSynthEngine.h
@@ -3,6 +3,7 @@
 #pragma once
 
 #include "CoreMinimal.h"
+#include "Sound/SoundGenerator.h"
 #include "UObject/ObjectMacros.h"
 #include "UObject/Object.h"
 #include "UObject/Class.h"
diff --git a/Engine/Source/Runtime/AudioMixerCore/Private/AudioMixer.cpp b/Engine/Source/Runtime/AudioMixerCore/Private/AudioMixer.cpp
--- a/Engine/Source/Runtime/AudioMixerCore/Private/AudioMixer.cpp
+++ b/Engine/Source/Runtime/AudioMixerCore/Private/AudioMixer.cpp
@@ -907,3 +907,37 @@ namespace Audio
 
 
 }
+
+FAudioPlatformSettings FAudioPlatformSettings::GetPlatformSettings(const TCHAR* PlatformSettingsConfigFile)
+{
+	FAudioPlatformSettings Settings;
+
+	FString TempString;
+
+	if (GConfig->GetString(PlatformSettingsConfigFile, TEXT("AudioSampleRate"), TempString, GEngineIni))
+	{
+		Settings.SampleRate = FMath::Max(FCString::Atoi(*TempString), 8000);
+	}
+
+	if (GConfig->GetString(PlatformSettingsConfigFile, TEXT("AudioCallbackBufferFrameSize"), TempString, GEngineIni))
+	{
+		Settings.CallbackBufferFrameSize = FMath::Max(FCString::Atoi(*TempString), 256);
+	}
+
+	if (GConfig->GetString(PlatformSettingsConfigFile, TEXT("AudioNumBuffersToEnqueue"), TempString, GEngineIni))
+	{
+		Settings.NumBuffers = FMath::Max(FCString::Atoi(*TempString), 1);
+	}
+
+	if (GConfig->GetString(PlatformSettingsConfigFile, TEXT("AudioMaxChannels"), TempString, GEngineIni))
+	{
+		Settings.MaxChannels = FMath::Max(FCString::Atoi(*TempString), 0);
+	}
+
+	if (GConfig->GetString(PlatformSettingsConfigFile, TEXT("AudioNumSourceWorkers"), TempString, GEngineIni))
+	{
+		Settings.NumSourceWorkers = FMath::Max(FCString::Atoi(*TempString), 0);
+	}
+
+	return Settings;
+}
diff --git a/Engine/Source/Runtime/AudioMixerCore/Public/AudioMixerTypes.h b/Engine/Source/Runtime/AudioMixerCore/Public/AudioMixerTypes.h
--- a/Engine/Source/Runtime/AudioMixerCore/Public/AudioMixerTypes.h
+++ b/Engine/Source/Runtime/AudioMixerCore/Public/AudioMixerTypes.h
@@ -2,7 +2,9 @@
 
 #pragma once
 #include "CoreMinimal.h"
+#if UE_ENABLE_INCLUDE_ORDER_DEPRECATED_IN_5_2
 #include "Misc/ConfigCacheIni.h"
+#endif
 
 namespace Audio {
 
@@ -73,39 +75,7 @@ struct AUDIOMIXERCORE_API FAudioPlatformSettings
 	/** The number of workers to use to compute source audio. Will only use up to the max number of sources. Will evenly divide sources to each source worker. */
 	int32 NumSourceWorkers;
 
-	static FAudioPlatformSettings GetPlatformSettings(const TCHAR* PlatformSettingsConfigFile)
-	{
-		FAudioPlatformSettings Settings;
-
-		FString TempString;
-
-		if (GConfig->GetString(PlatformSettingsConfigFile, TEXT("AudioSampleRate"), TempString, GEngineIni))
-		{
-			Settings.SampleRate = FMath::Max(FCString::Atoi(*TempString), 8000);
-		}
-
-		if (GConfig->GetString(PlatformSettingsConfigFile, TEXT("AudioCallbackBufferFrameSize"), TempString, GEngineIni))
-		{
-			Settings.CallbackBufferFrameSize = FMath::Max(FCString::Atoi(*TempString), 256);
-		}
-
-		if (GConfig->GetString(PlatformSettingsConfigFile, TEXT("AudioNumBuffersToEnqueue"), TempString, GEngineIni))
-		{
-			Settings.NumBuffers = FMath::Max(FCString::Atoi(*TempString), 1);
-		}
-
-		if (GConfig->GetString(PlatformSettingsConfigFile, TEXT("AudioMaxChannels"), TempString, GEngineIni))
-		{
-			Settings.MaxChannels = FMath::Max(FCString::Atoi(*TempString), 0);
-		}
-
-		if (GConfig->GetString(PlatformSettingsConfigFile, TEXT("AudioNumSourceWorkers"), TempString, GEngineIni))
-		{
-			Settings.NumSourceWorkers = FMath::Max(FCString::Atoi(*TempString), 0);
-		}
-
-		return Settings;
-	}
+	static FAudioPlatformSettings GetPlatformSettings(const TCHAR* PlatformSettingsConfigFile);
 
 	FAudioPlatformSettings()
 		: SampleRate(48000)
diff --git a/Engine/Source/Runtime/Engine/Classes/Sound/SoundAttenuation.h b/Engine/Source/Runtime/Engine/Classes/Sound/SoundAttenuation.h
--- a/Engine/Source/Runtime/Engine/Classes/Sound/SoundAttenuation.h
+++ b/Engine/Source/Runtime/Engine/Classes/Sound/SoundAttenuation.h
@@ -3,13 +3,19 @@
 #pragma once
 
 #include "Engine/Attenuation.h"
+#if UE_ENABLE_INCLUDE_ORDER_DEPRECATED_IN_5_2
 #include "IAudioExtensionPlugin.h"
+#endif
 #include "IAudioParameterInterfaceRegistry.h"
 #include "AudioLinkSettingsAbstract.h"
 
 #include "SoundAttenuation.generated.h"
 
+class UOcclusionPluginSourceSettingsBase;
+class UReverbPluginSourceSettingsBase;
+class USourceDataOverridePluginSourceSettingsBase;
 class USoundSubmixBase;
+class USpatializationPluginSourceSettingsBase;
 
 // This enumeration is deprecated
 UENUM()
diff --git a/Engine/Source/Runtime/Engine/Classes/Sound/SoundBase.h b/Engine/Source/Runtime/Engine/Classes/Sound/SoundBase.h
--- a/Engine/Source/Runtime/Engine/Classes/Sound/SoundBase.h
+++ b/Engine/Source/Runtime/Engine/Classes/Sound/SoundBase.h
@@ -2,21 +2,25 @@
 #pragma once
 
 
-#include "Audio.h"
+#include "AudioDefines.h"
 #include "CoreMinimal.h"
+#include "Sound/SoundTimecodeOffset.h"
+#include "SoundConcurrency.h"
+#include "UObject/Object.h"
+#include "UObject/ObjectMacros.h"
+#include "Interfaces/Interface_AssetUserData.h"
+
+#if UE_ENABLE_INCLUDE_ORDER_DEPRECATED_IN_5_2
+#include "Audio.h"
 #include "IAudioExtensionPlugin.h"
 #include "Sound/AudioSettings.h"
 #include "Sound/SoundClass.h"
-#include "Sound/SoundTimecodeOffset.h"
-#include "SoundConcurrency.h"
 #include "SoundModulationDestination.h"
 #include "SoundSourceBusSend.h"
 #include "SoundSubmixSend.h"
 #include "SoundGenerator.h"
-#include "UObject/Object.h"
-#include "UObject/ObjectMacros.h"
 #include "AudioDeviceManager.h"
-#include "Interfaces/Interface_AssetUserData.h"
+#endif
 
 #include "SoundBase.generated.h"
 
@@ -25,18 +29,33 @@
 namespace Audio
 {
 	class IParameterTransmitter;
+	struct FParameterInterface;
 	struct FParameterTransmitterInitParams;
+	using FDeviceId = uint32;;
+	using FParameterInterfacePtr = TSharedPtr<FParameterInterface, ESPMode::ThreadSafe>;
 } // namespace Audio
 
+class ISoundGenerator;
+class USoundAttenuation;
+class USoundClass;
 class USoundEffectSourcePreset;
 class USoundSourceBus;
 class USoundSubmix;
+class USoundSubmixBase;
+class USoundWave;
 class USoundEffectSourcePresetChain;
 class UAssetUserData;
-
+enum class EBusSendType : uint8;
+namespace EMaxConcurrentResolutionRule { enum Type : int; }
 struct FActiveSound;
+struct FAudioParameter;
+struct FSoundAttenuationSettings;
+struct FSoundGeneratorInitParams;
 struct FSoundParseParameters;
-
+struct FSoundSourceBusSendInfo;
+struct FSoundSubmixSendInfo;
+struct FWaveInstance;
+typedef TSharedPtr<ISoundGenerator, ESPMode::ThreadSafe> ISoundGeneratorPtr;
 
 /**
  * Method of virtualization when a sound is stopped due to playback constraints
diff --git a/Engine/Source/Runtime/Engine/Classes/Sound/SoundConcurrency.h b/Engine/Source/Runtime/Engine/Classes/Sound/SoundConcurrency.h
--- a/Engine/Source/Runtime/Engine/Classes/Sound/SoundConcurrency.h
+++ b/Engine/Source/Runtime/Engine/Classes/Sound/SoundConcurrency.h
@@ -29,7 +29,7 @@ ENGINE_API DECLARE_LOG_CATEGORY_EXTERN(LogAudioConcurrency, Log, All);
 UENUM()
 namespace EMaxConcurrentResolutionRule
 {
-	enum Type
+	enum Type : int
 	{
 		/** When Max Concurrent sounds are active do not start a new sound. */
 		PreventNew,
diff --git a/Engine/Source/Runtime/Engine/Classes/Sound/SoundWave.h b/Engine/Source/Runtime/Engine/Classes/Sound/SoundWave.h
--- a/Engine/Source/Runtime/Engine/Classes/Sound/SoundWave.h
+++ b/Engine/Source/Runtime/Engine/Classes/Sound/SoundWave.h
@@ -8,6 +8,8 @@
 
 #include "CoreMinimal.h"
 #include "Containers/IndirectArray.h"
+#include "Sound/AudioSettings.h"
+#include "Sound/SoundModulationDestination.h"
 #include "UObject/ObjectMacros.h"
 #include "Misc/Guid.h"
 #include "Async/AsyncWork.h"
@@ -31,6 +33,7 @@
 
 class FSoundWaveData;
 class ITargetPlatform;
+enum EAudioSpeakers : int;
 struct FActiveSound;
 struct FSoundParseParameters;
 struct FPlatformAudioCookOverrides;
diff --git a/Engine/Source/Runtime/Engine/Public/Audio.h b/Engine/Source/Runtime/Engine/Public/Audio.h
--- a/Engine/Source/Runtime/Engine/Public/Audio.h
+++ b/Engine/Source/Runtime/Engine/Public/Audio.h
@@ -74,7 +74,7 @@ struct FSoundSourceBusSendInfo;
  * These are in the sample order OpenAL expects for a 7.1 sound
  * 
  */
-enum EAudioSpeakers
+enum EAudioSpeakers : int
 {							//	4.0	5.1	6.1	7.1
 	SPEAKER_FrontLeft,		//	*	*	*	*
 	SPEAKER_FrontRight,		//	*	*	*	*
