Reverting CL 22629055 since it appears to be causing a bunch of blurriness on canvas text items such as Debug Camera debug prints. https://epicgames.slack.com/archives/C044LSC7U9L/p1666621600270659

Also ensuring that the DebugCameraHUD rendering starts pixel aligned.

diff --git a/Engine/Source/Runtime/Engine/Private/DebugCameraHUD.cpp b/Engine/Source/Runtime/Engine/Private/DebugCameraHUD.cpp
--- a/Engine/Source/Runtime/Engine/Private/DebugCameraHUD.cpp
+++ b/Engine/Source/Runtime/Engine/Private/DebugCameraHUD.cpp
@@ -94,7 +94,7 @@ void ADebugCameraHUD::PostRender()
 			FString MyText = TEXT("Debug Camera");
 			float xl, yl;
 			Canvas->StrLen(RenderFont, MyText, xl, yl);
-			float X = Canvas->SizeX * 0.05f;
+			float X = FMath::FloorToFloat(Canvas->SizeX * 0.05f);
 			float Y = yl;//*1.67;
 			yl += 2*Y;
 			Canvas->DrawText(RenderFont, MyText, X, yl, 1.f, 1.f, FontRenderInfo);
diff --git a/Engine/Source/Runtime/Engine/Private/UserInterface/CanvasItem.cpp b/Engine/Source/Runtime/Engine/Private/UserInterface/CanvasItem.cpp
--- a/Engine/Source/Runtime/Engine/Private/UserInterface/CanvasItem.cpp
+++ b/Engine/Source/Runtime/Engine/Private/UserInterface/CanvasItem.cpp
@@ -1083,30 +1083,30 @@ void FCanvasTextItem::DrawStringInternal_OfflineCache(FCanvas* InCanvas, const F
 			const float SizeU	= Char.USize * InvTextureSize.X;
 			const float SizeV	= Char.VSize * InvTextureSize.Y;				
 
+			const float Left = X * Depth;
+			const float Top = Y * Depth;
+			const float Right = (X + SizeX) * Depth;
+			const float Bottom = (Y + SizeY) * Depth;
+
 			const auto AddTriangles = [&](const FVector2f& Offset, const FLinearColor& InColor)
 			{
-				const float Left = (X + Offset.X) * Depth;
-				const float Top = (Y + Offset.Y) * Depth;
-				const float Right = (X + SizeX) * Depth;
-				const float Bottom = (Y + SizeY) * Depth;
-
 				int32 V00 = BatchedElements->AddVertexf(
-					FVector4f(Left, Top, 0.f, Depth),
+					FVector4f(Left + Offset.X, Top + Offset.Y, 0.f, Depth),
 					FVector2f(U, V),
 					InColor,
 					HitProxyId);
 				int32 V10 = BatchedElements->AddVertexf(
-					FVector4f(Right, Top, 0.0f, Depth),
+					FVector4f(Right + Offset.X, Top + Offset.Y, 0.0f, Depth),
 					FVector2f(U + SizeU, V),
 					InColor,
 					HitProxyId);
 				int32 V01 = BatchedElements->AddVertexf(
-					FVector4f(Left, Bottom, 0.0f, Depth),
+					FVector4f(Left + Offset.X, Bottom + Offset.Y, 0.0f, Depth),
 					FVector2f(U, V + SizeV),
 					InColor,
 					HitProxyId);
 				int32 V11 = BatchedElements->AddVertexf(
-					FVector4f(Right, Bottom, 0.0f, Depth),
+					FVector4f(Right + Offset.X, Bottom + Offset.Y, 0.0f, Depth),
 					FVector2f(U + SizeU, V + SizeV),
 					InColor,
 					HitProxyId);
@@ -1278,30 +1278,30 @@ void FCanvasTextItem::DrawStringInternal_RuntimeCache(FCanvas* InCanvas, const F
 				const float SizeU = Entry.USize * InvTextureSizeX;
 				const float SizeV = Entry.VSize * InvTextureSizeY;
 
+				const float Left = X * Depth;
+				const float Top = Y * Depth;
+				const float Right = (X + SizeX) * Depth;
+				const float Bottom = (Y + SizeY) * Depth;
+
 				const auto AddTriangles = [&](const FVector2f& Offset, const FLinearColor& InColor)
 				{
-					const float Left = (X + Offset.X) * Depth;
-					const float Top = (Y + Offset.Y) * Depth;
-					const float Right = (X + SizeX) * Depth;
-					const float Bottom = (Y + SizeY) * Depth;
-
 					int32 V00 = BatchedElements->AddVertexf(
-						FVector4f(Left, Top, 0.f, Depth),
+						FVector4f(Left + Offset.X, Top + Offset.Y, 0.f, Depth),
 						FVector2f(U, V),
 						InColor,
 						HitProxyId);
 					int32 V10 = BatchedElements->AddVertexf(
-						FVector4f(Right, Top, 0.0f, Depth),
+						FVector4f(Right + Offset.X, Top + Offset.Y, 0.0f, Depth),
 						FVector2f(U + SizeU, V),
 						InColor,
 						HitProxyId);
 					int32 V01 = BatchedElements->AddVertexf(
-						FVector4f(Left, Bottom, 0.0f, Depth),
+						FVector4f(Left + Offset.X, Bottom + Offset.Y, 0.0f, Depth),
 						FVector2f(U, V + SizeV),
 						InColor,
 						HitProxyId);
 					int32 V11 = BatchedElements->AddVertexf(
-						FVector4f(Right, Bottom, 0.0f, Depth),
+						FVector4f(Right + Offset.X, Bottom + Offset.Y, 0.0f, Depth),
 						FVector2f(U + SizeU, V + SizeV),
 						InColor,
 						HitProxyId);
@@ -1441,30 +1441,30 @@ void FCanvasShapedTextItem::DrawStringInternal(FCanvas* InCanvas, const FVector2
 				const float SizeU = GlyphAtlasData.USize * InvTextureSizeX;
 				const float SizeV = GlyphAtlasData.VSize * InvTextureSizeY;
 
+				const float Left = X * Depth;
+				const float Top = Y * Depth;
+				const float Right = (X + SizeX) * Depth;
+				const float Bottom = (Y + SizeY) * Depth;
+
 				const auto AddTriangles = [&](const FVector2f& Offset, const FLinearColor& InColor)
 				{
-					const float Left = (X + Offset.X) * Depth;
-					const float Top = (Y + Offset.Y) * Depth;
-					const float Right = (X + SizeX) * Depth;
-					const float Bottom = (Y + SizeY) * Depth;
-
 					int32 V00 = BatchedElements->AddVertexf(
-						FVector4f(Left, Top, 0.f, Depth),
+						FVector4f(Left + Offset.X, Top + Offset.Y, 0.f, Depth),
 						FVector2f(U, V),
 						InColor,
 						HitProxyId);
 					int32 V10 = BatchedElements->AddVertexf(
-						FVector4f(Right, Top, 0.0f, Depth),
+						FVector4f(Right + Offset.X, Top + Offset.Y, 0.0f, Depth),
 						FVector2f(U + SizeU, V),
 						InColor,
 						HitProxyId);
 					int32 V01 = BatchedElements->AddVertexf(
-						FVector4f(Left, Bottom, 0.0f, Depth),
+						FVector4f(Left + Offset.X, Bottom + Offset.Y, 0.0f, Depth),
 						FVector2f(U, V + SizeV),
 						InColor,
 						HitProxyId);
 					int32 V11 = BatchedElements->AddVertexf(
-						FVector4f(Right, Bottom, 0.0f, Depth),
+						FVector4f(Right + Offset.X, Bottom + Offset.Y, 0.0f, Depth),
 						FVector2f(U + SizeU, V + SizeV),
 						InColor,
 						HitProxyId);
