#UE Make sure you always remove references from the PackageReferencer in UWorldPartitionLevelStreamingDynamic.

diff --git a/Engine/Source/Runtime/Engine/Private/WorldPartition/WorldPartitionLevelStreamingDynamic.cpp b/Engine/Source/Runtime/Engine/Private/WorldPartition/WorldPartitionLevelStreamingDynamic.cpp
--- a/Engine/Source/Runtime/Engine/Private/WorldPartition/WorldPartitionLevelStreamingDynamic.cpp
+++ b/Engine/Source/Runtime/Engine/Private/WorldPartition/WorldPartitionLevelStreamingDynamic.cpp
@@ -132,6 +132,7 @@ void UWorldPartitionLevelStreamingDynamic::BeginDestroy()
 	{
 		RuntimeLevel->OnCleanupLevel.Remove(OnCleanupLevelDelegateHandle);
 	}
+	PackageReferencer.RemoveReferences();
 	Super::BeginDestroy();
 }
 
@@ -434,6 +435,7 @@ void UWorldPartitionLevelStreamingDynamic::FinalizeRuntimeLevel()
 
 	if (IsEngineExitRequested())
 	{
+		PackageReferencer.RemoveReferences();
 		return;
 	}
 
@@ -500,10 +502,10 @@ void UWorldPartitionLevelStreamingDynamic::FinalizeRuntimeLevel()
  */
 void UWorldPartitionLevelStreamingDynamic::OnCleanupLevel()
 {
+	PackageReferencer.RemoveReferences();
+
 	if (RuntimeLevel)
 	{
-		PackageReferencer.RemoveReferences();
-
 		RuntimeLevel->OnCleanupLevel.Remove(OnCleanupLevelDelegateHandle);
 
 		// If reusing levels is enabled, trash world partition level/actor packages
