Allow SkinWeightProfiles to be added at runtime.

#rnx

diff --git a/Engine/Source/Runtime/Engine/Classes/Engine/SkeletalMesh.h b/Engine/Source/Runtime/Engine/Classes/Engine/SkeletalMesh.h
--- a/Engine/Source/Runtime/Engine/Classes/Engine/SkeletalMesh.h
+++ b/Engine/Source/Runtime/Engine/Classes/Engine/SkeletalMesh.h
@@ -2882,7 +2882,6 @@ public:
 		PRAGMA_ENABLE_DEPRECATION_WARNINGS
 	}
 
-#if WITH_EDITOR
 	TArray<FSkinWeightProfileInfo>& GetSkinWeightProfiles() 
 	{
 		WaitUntilAsyncPropertyReleased(ESkeletalMeshAsyncProperties::SkinWeightProfiles); 
@@ -2899,7 +2898,6 @@ public:
 		return SkinWeightProfiles.Num(); 
 		PRAGMA_ENABLE_DEPRECATION_WARNINGS
 	}
-#endif
 
 	/** Releases all allocated Skin Weight Profile resources, assumes none are currently in use */
 	void ReleaseSkinWeightProfileResources();
diff --git a/Engine/Source/Runtime/Engine/Private/Animation/SkinWeightProfile.cpp b/Engine/Source/Runtime/Engine/Private/Animation/SkinWeightProfile.cpp
--- a/Engine/Source/Runtime/Engine/Private/Animation/SkinWeightProfile.cpp
+++ b/Engine/Source/Runtime/Engine/Private/Animation/SkinWeightProfile.cpp
@@ -491,7 +491,7 @@ bool FSkinWeightProfilesData::ContainsOverrideBuffer(const FName& ProfileName) c
 	return ProfileNameToBuffer.Contains(ProfileName);
 }
 
-#if WITH_EDITOR
+
 const FRuntimeSkinWeightProfileData* FSkinWeightProfilesData::GetOverrideData(const FName& ProfileName) const
 {
 	return OverrideData.Find(ProfileName);
@@ -501,7 +501,7 @@ FRuntimeSkinWeightProfileData& FSkinWeightProfilesData::AddOverrideData(const FN
 {
 	return OverrideData.FindOrAdd(ProfileName);
 }
-#endif // WITH_EDITOR
+
 
 void FSkinWeightProfilesData::ReleaseBuffer(const FName& ProfileName, bool bForceRelease /*= false*/)
 {
diff --git a/Engine/Source/Runtime/Engine/Private/SkeletalMesh.cpp b/Engine/Source/Runtime/Engine/Private/SkeletalMesh.cpp
--- a/Engine/Source/Runtime/Engine/Private/SkeletalMesh.cpp
+++ b/Engine/Source/Runtime/Engine/Private/SkeletalMesh.cpp
@@ -5185,7 +5185,7 @@ bool USkeletalMesh::IsMaterialUsed(int32 MaterialIndex) const
 	return false;
 }
 
-#if WITH_EDITOR
+
 void USkeletalMesh::AddSkinWeightProfile(const FSkinWeightProfileInfo& Profile) 
 {
 	WaitUntilAsyncPropertyReleased(ESkeletalMeshAsyncProperties::SkinWeightProfiles); 
@@ -5193,7 +5193,7 @@ void USkeletalMesh::AddSkinWeightProfile(const FSkinWeightProfileInfo& Profile)
 	SkinWeightProfiles.Add(Profile); 
 	PRAGMA_ENABLE_DEPRECATION_WARNINGS
 }
-#endif
+
 
 void USkeletalMesh::ReleaseSkinWeightProfileResources()
 {
diff --git a/Engine/Source/Runtime/Engine/Public/Animation/SkinWeightProfile.h b/Engine/Source/Runtime/Engine/Public/Animation/SkinWeightProfile.h
--- a/Engine/Source/Runtime/Engine/Public/Animation/SkinWeightProfile.h
+++ b/Engine/Source/Runtime/Engine/Public/Animation/SkinWeightProfile.h
@@ -148,10 +148,8 @@ struct ENGINE_API FSkinWeightProfilesData
 	FSkinWeightVertexBuffer* GetOverrideBuffer(const FName& ProfileName) const;
 	bool ContainsOverrideBuffer(const FName& ProfileName) const;
 	
-#if WITH_EDITOR
 	const FRuntimeSkinWeightProfileData* GetOverrideData(const FName& ProfileName) const;
 	FRuntimeSkinWeightProfileData& AddOverrideData(const FName& ProfileName);
-#endif // WITH_EDITOR
 	
 	void ReleaseBuffer(const FName& ProfileName, bool bForceRelease = false);
 	void ReleaseResources();
