WorldPartition: Streaming Source Changes
-IStreamingSourceProvider can return more than one StreamingSource.
-Deprecate TargetGrid/TargetHLODLayer in favor of TargetGrids/TargetHLODLayers
-Add TargetBehavior so a StreamingSource can be applied on its Targets or on all Targets excluding its targets.

diff --git a/Engine/Source/Runtime/Core/Public/UObject/FortniteMainBranchObjectVersion.h b/Engine/Source/Runtime/Core/Public/UObject/FortniteMainBranchObjectVersion.h
--- a/Engine/Source/Runtime/Core/Public/UObject/FortniteMainBranchObjectVersion.h
+++ b/Engine/Source/Runtime/Core/Public/UObject/FortniteMainBranchObjectVersion.h
@@ -256,6 +256,9 @@ struct CORE_API FFortniteMainBranchObjectVersion
 		// Serialize HLOD stats in HLODActorDesc
 		WorldPartitionHLODActorDescSerializeStats,
 
+		// WorldPartitionStreamingSourceComponent property deprecation
+		WorldPartitionStreamingSourceComponentTargetDeprecation,
+
 		// -----<new versions can be added above this line>-------------------------------------------------
 		VersionPlusOne,
 		LatestVersion = VersionPlusOne - 1
diff --git a/Engine/Source/Runtime/Engine/Classes/Components/WorldPartitionStreamingSourceComponent.h b/Engine/Source/Runtime/Engine/Classes/Components/WorldPartitionStreamingSourceComponent.h
--- a/Engine/Source/Runtime/Engine/Classes/Components/WorldPartitionStreamingSourceComponent.h
+++ b/Engine/Source/Runtime/Engine/Classes/Components/WorldPartitionStreamingSourceComponent.h
@@ -20,6 +20,9 @@ class ENGINE_API UWorldPartitionStreamingSourceComponent : public UActorComponen
 	virtual void OnRegister() override;
 	virtual void OnUnregister() override;
 
+	virtual void Serialize(FArchive& Ar) override;
+	virtual void PostLoad() override;
+
 	/** Enable the component */
 	UFUNCTION(BlueprintCallable, Category = "Streaming")
 	void EnableStreamingSource() { bStreamingSourceEnabled = true; }
@@ -48,17 +51,27 @@ class ENGINE_API UWorldPartitionStreamingSourceComponent : public UActorComponen
 	float DefaultVisualizerLoadingRange;
 #endif
 
-	/** Optional target grid affected by streaming source. */
+	/** When TargetGrids or TargetHLODLayers are specified, this indicates the behavior. */
+	UPROPERTY(BlueprintReadWrite, EditAnywhere, Category = "Streaming")
+	EStreamingSourceTargetBehavior TargetBehavior;
+			
+	/** Optional target grids affected by streaming source. */
 	UPROPERTY(BlueprintReadWrite, EditAnywhere, Category = "Streaming")
-	FName TargetGrid;
+	TArray<FName> TargetGrids;
+
+	UPROPERTY(meta = (DeprecatedProperty, DeprecationMessage = "Use TargetGrids instead."))
+	FName TargetGrid_DEPRECATED;
 
 	/** Color used for debugging. */
 	UPROPERTY(EditAnywhere, Category = "Streaming")
 	FColor DebugColor;
 
-	/** Optional target HLODLayer affected by the streaming source. */
+	/** Optional target HLODLayers affected by the streaming source. */
 	UPROPERTY(BlueprintReadWrite, EditAnywhere, Category = "Streaming")
-	TObjectPtr<const UHLODLayer> TargetHLODLayer;
+	TArray<TObjectPtr<const UHLODLayer>> TargetHLODLayers;
+
+	UPROPERTY(meta = (DeprecatedProperty, DeprecationMessage = "Use TargetHLODLayers instead."))
+	TObjectPtr<const UHLODLayer> TargetHLODLayer_DEPRECATED;
 
 	/** Optional aggregated shape list used to build a custom shape for the streaming source. When empty, fallbacks sphere shape with a radius equal to grid's loading range. */
 	UPROPERTY(BlueprintReadWrite, EditAnywhere, Category = "Streaming")
@@ -67,10 +80,6 @@ class ENGINE_API UWorldPartitionStreamingSourceComponent : public UActorComponen
 	UPROPERTY(BlueprintReadWrite, EditAnywhere, Category = "Streaming")
 	EStreamingSourcePriority Priority;
 
-#if WITH_EDITOR
-	virtual bool CanEditChange(const FProperty* InProperty) const override;
-#endif
-
 private:
 	/** Whether this component is enabled or not */
 	UPROPERTY(EditAnywhere, Category = "Streaming")
diff --git a/Engine/Source/Runtime/Engine/Classes/GameFramework/PlayerController.h b/Engine/Source/Runtime/Engine/Classes/GameFramework/PlayerController.h
--- a/Engine/Source/Runtime/Engine/Classes/GameFramework/PlayerController.h
+++ b/Engine/Source/Runtime/Engine/Classes/GameFramework/PlayerController.h
@@ -822,12 +822,15 @@ public:
 	virtual void GetStreamingSourceShapes(TArray<FStreamingSourceShape>& OutShapes) const;
 
 	/**
-	 * Gets the PlayerController's streaming source
-	 * @return the streaming source.
+	 * Gets the PlayerController's streaming sources
+	 * @return the streaming sources.
 	 */
-	virtual bool GetStreamingSource(FWorldPartitionStreamingSource& OutStreamingSource) const final;
+	virtual bool GetStreamingSources(TArray<FWorldPartitionStreamingSource>& OutStreamingSources) const final;
 
 protected:
+	virtual bool GetStreamingSourcesInternal(TArray<FWorldPartitionStreamingSource>& OutStreamingSources) const;
+	virtual bool GetStreamingSource(FWorldPartitionStreamingSource& OutStreamingSource) const final;
+	
 	/** Pawn has been possessed, so changing state to NAME_Playing. Start it walking and begin playing with it. */
 	virtual void BeginPlayingState();
 
diff --git a/Engine/Source/Runtime/Engine/Private/Components/WorldPartitionStreamingSourceComponent.cpp b/Engine/Source/Runtime/Engine/Private/Components/WorldPartitionStreamingSourceComponent.cpp
--- a/Engine/Source/Runtime/Engine/Private/Components/WorldPartitionStreamingSourceComponent.cpp
+++ b/Engine/Source/Runtime/Engine/Private/Components/WorldPartitionStreamingSourceComponent.cpp
@@ -8,6 +8,7 @@
 #include "WorldPartition/WorldPartitionSubsystem.h"
 #include "WorldPartition/DataLayer/DataLayerSubsystem.h"
 #include "WorldPartition/WorldPartitionStreamingSource.h"
+#include "UObject/FortniteMainBranchObjectVersion.h"
 #include "SceneManagement.h"
 #include "SceneView.h"
 #include "Engine/World.h"
@@ -19,7 +20,7 @@ UWorldPartitionStreamingSourceComponent::UWorldPartitionStreamingSourceComponent
 #if WITH_EDITORONLY_DATA
 	, DefaultVisualizerLoadingRange(10000.f)
 #endif
-	, TargetGrid(NAME_None)
+	, TargetBehavior(EStreamingSourceTargetBehavior::Include)
 	, DebugColor(FColor::MakeRedToGreenColorFromScalar(FRandomStream(GetFName()).GetFraction()))
 	, Priority(EStreamingSourcePriority::Normal)
 	, bStreamingSourceEnabled(true)
@@ -64,6 +65,31 @@ void UWorldPartitionStreamingSourceComponent::OnUnregister()
 	verify(WorldPartitionSubsystem->UnregisterStreamingSourceProvider(this));
 }
 
+void UWorldPartitionStreamingSourceComponent::Serialize(FArchive& Ar)
+{
+	Super::Serialize(Ar);
+
+	Ar.UsingCustomVersion(FFortniteMainBranchObjectVersion::GUID);
+}
+
+void UWorldPartitionStreamingSourceComponent::PostLoad()
+{
+	Super::PostLoad();
+
+	if (GetLinkerCustomVersion(FFortniteMainBranchObjectVersion::GUID) < FFortniteMainBranchObjectVersion::WorldPartitionStreamingSourceComponentTargetDeprecation)
+	{
+		if (!TargetGrid_DEPRECATED.IsNone())
+		{
+			TargetGrids.Add(TargetGrid_DEPRECATED);
+		}
+
+		if (TargetHLODLayer_DEPRECATED)
+		{
+			TargetHLODLayers.Add(TargetHLODLayer_DEPRECATED);
+		}
+	}
+}
+
 bool UWorldPartitionStreamingSourceComponent::GetStreamingSource(FWorldPartitionStreamingSource& OutStreamingSource) const
 {
 	if (bStreamingSourceEnabled)
@@ -75,10 +101,27 @@ bool UWorldPartitionStreamingSourceComponent::GetStreamingSource(FWorldPartition
 		OutStreamingSource.Rotation = Actor->GetActorRotation();
 		OutStreamingSource.TargetState = TargetState;
 		OutStreamingSource.DebugColor = DebugColor;
-		OutStreamingSource.TargetGrid = TargetGrid;
-		OutStreamingSource.TargetHLODLayer = TargetHLODLayer;
-		OutStreamingSource.Shapes = Shapes;
+		OutStreamingSource.TargetBehavior = TargetBehavior;
 		OutStreamingSource.Priority = Priority;
+						
+		if (!Shapes.IsEmpty())
+		{
+			OutStreamingSource.Shapes.Append(Shapes);
+		}
+
+		if (!TargetGrids.IsEmpty())
+		{
+			OutStreamingSource.TargetGrids.Append(TargetGrids);
+		}
+
+		if (!TargetHLODLayers.IsEmpty())
+		{
+			OutStreamingSource.TargetHLODLayers.Reserve(TargetHLODLayers.Num());
+			for (const TObjectPtr<const UHLODLayer>& HLODLayer : TargetHLODLayers)
+			{
+				OutStreamingSource.TargetHLODLayers.Add(HLODLayer);
+			}
+		}
 		return true;
 	}
 	return false;
@@ -115,14 +158,4 @@ void UWorldPartitionStreamingSourceComponent::DrawVisualization(const FSceneView
 #endif
 }
 
-#if WITH_EDITOR
-bool UWorldPartitionStreamingSourceComponent::CanEditChange(const FProperty* InProperty) const
-{
-	if (InProperty && InProperty->GetName() == GET_MEMBER_NAME_STRING_CHECKED(UWorldPartitionStreamingSourceComponent, TargetGrid))
-	{
-		return TargetHLODLayer == nullptr;
-	}
-	return Super::CanEditChange(InProperty);
-}
-#endif
 
diff --git a/Engine/Source/Runtime/Engine/Private/PlayerController.cpp b/Engine/Source/Runtime/Engine/Private/PlayerController.cpp
--- a/Engine/Source/Runtime/Engine/Private/PlayerController.cpp
+++ b/Engine/Source/Runtime/Engine/Private/PlayerController.cpp
@@ -3635,23 +3635,35 @@ void APlayerController::GetStreamingSourceShapes(TArray<FStreamingSourceShape>&
 }
 
 bool APlayerController::GetStreamingSource(FWorldPartitionStreamingSource& OutStreamingSource) const
+{
+	checkNoEntry();
+	return false;
+}
+
+bool APlayerController::GetStreamingSources(TArray<FWorldPartitionStreamingSource>& OutStreamingSources) const
 {
 	const ENetMode NetMode = GetNetMode();
 	const bool bIsServer = (NetMode == NM_DedicatedServer || NetMode == NM_ListenServer);
 	if (IsStreamingSourceEnabled() && (IsLocalController() || bIsServer))
 	{
-		GetStreamingSourceLocationAndRotation(OutStreamingSource.Location, OutStreamingSource.Rotation);
-		OutStreamingSource.Name = GetFName();
-		OutStreamingSource.TargetState = StreamingSourceShouldActivate() ? EStreamingSourceTargetState::Activated : EStreamingSourceTargetState::Loaded;
-		OutStreamingSource.bBlockOnSlowLoading = StreamingSourceShouldBlockOnSlowStreaming();
-		OutStreamingSource.DebugColor = StreamingSourceDebugColor;
-		OutStreamingSource.Priority = GetStreamingSourcePriority();
-		GetStreamingSourceShapes(OutStreamingSource.Shapes);
-		return true;
+		return GetStreamingSourcesInternal(OutStreamingSources);
 	}
 	return false;
 }
 
+bool APlayerController::GetStreamingSourcesInternal(TArray<FWorldPartitionStreamingSource>& OutStreamingSources) const
+{
+	FWorldPartitionStreamingSource& StreamingSource = OutStreamingSources.AddDefaulted_GetRef();
+	GetStreamingSourceLocationAndRotation(StreamingSource.Location, StreamingSource.Rotation);
+	StreamingSource.Name = GetFName();
+	StreamingSource.TargetState = StreamingSourceShouldActivate() ? EStreamingSourceTargetState::Activated : EStreamingSourceTargetState::Loaded;
+	StreamingSource.bBlockOnSlowLoading = StreamingSourceShouldBlockOnSlowStreaming();
+	StreamingSource.DebugColor = StreamingSourceDebugColor;
+	StreamingSource.Priority = GetStreamingSourcePriority();
+	GetStreamingSourceShapes(StreamingSource.Shapes);
+	return true;
+}
+
 /// @cond DOXYGEN_WARNINGS
 
 void APlayerController::ClientEnableNetworkVoice_Implementation(bool bEnable)
diff --git a/Engine/Source/Runtime/Engine/Private/WorldPartition/WorldPartition.cpp b/Engine/Source/Runtime/Engine/Private/WorldPartition/WorldPartition.cpp
--- a/Engine/Source/Runtime/Engine/Private/WorldPartition/WorldPartition.cpp
+++ b/Engine/Source/Runtime/Engine/Private/WorldPartition/WorldPartition.cpp
@@ -1351,11 +1351,11 @@ bool UWorldPartition::CanAddLoadedLevelToWorld(class ULevel* InLevel) const
 	return true;
 }
 
-bool UWorldPartition::IsStreamingCompleted(const FWorldPartitionStreamingSource* InStreamingSource) const
+bool UWorldPartition::IsStreamingCompleted(const TArray<FWorldPartitionStreamingSource>* InStreamingSources) const
 {
 	if (GetWorld()->IsGameWorld())
 	{
-		return StreamingPolicy->IsStreamingCompleted(InStreamingSource);
+		return StreamingPolicy->IsStreamingCompleted(InStreamingSources);
 	}
 	return true;
 }
diff --git a/Engine/Source/Runtime/Engine/Private/WorldPartition/WorldPartitionLevelStreamingPolicy.cpp b/Engine/Source/Runtime/Engine/Private/WorldPartition/WorldPartitionLevelStreamingPolicy.cpp
--- a/Engine/Source/Runtime/Engine/Private/WorldPartition/WorldPartitionLevelStreamingPolicy.cpp
+++ b/Engine/Source/Runtime/Engine/Private/WorldPartition/WorldPartitionLevelStreamingPolicy.cpp
@@ -52,18 +52,18 @@ void UWorldPartitionLevelStreamingPolicy::ForEachActiveRuntimeCell(TFunctionRef<
 	}
 }
 
-bool UWorldPartitionLevelStreamingPolicy::IsStreamingCompleted(const FWorldPartitionStreamingSource* InStreamingSource) const
+bool UWorldPartitionLevelStreamingPolicy::IsStreamingCompleted(const TArray<FWorldPartitionStreamingSource>* InStreamingSources) const
 {
 	const UWorld* World = GetWorld();
 	check(World);
 	check(World->IsGameWorld());
 
-	if (!Super::IsStreamingCompleted(InStreamingSource))
+	if (!Super::IsStreamingCompleted(InStreamingSources))
 	{
 		return false;
 	}
 
-	if (!InStreamingSource)
+	if (!InStreamingSources)
 	{
 		// Also verify that there's no remaining activity (mainly for unloading) 
 		// waiting to be processed on level streaming of world partition runtime cells
diff --git a/Engine/Source/Runtime/Engine/Private/WorldPartition/WorldPartitionStreamingPolicy.cpp b/Engine/Source/Runtime/Engine/Private/WorldPartition/WorldPartitionStreamingPolicy.cpp
--- a/Engine/Source/Runtime/Engine/Private/WorldPartition/WorldPartitionStreamingPolicy.cpp
+++ b/Engine/Source/Runtime/Engine/Private/WorldPartition/WorldPartitionStreamingPolicy.cpp
@@ -120,17 +120,19 @@ void UWorldPartitionStreamingPolicy::UpdateStreamingSources()
 			UWorldPartitionSubsystem* WorldPartitionSubsystem = GetWorld()->GetSubsystem<UWorldPartitionSubsystem>();
 			check(WorldPartitionSubsystem);
 
+			TArray<FWorldPartitionStreamingSource> ProviderStreamingSources;
 			for (IWorldPartitionStreamingSourceProvider* StreamingSourceProvider : WorldPartitionSubsystem->GetStreamingSourceProviders())
 			{
-				FWorldPartitionStreamingSource StreamingSource;
-				// Default Streaming Source provider's priority to be less than those based on player controllers
-				StreamingSource.Priority = EStreamingSourcePriority::Low;
-				if (StreamingSourceProvider->GetStreamingSource(StreamingSource))
+				ProviderStreamingSources.Reset();
+				if (StreamingSourceProvider->GetStreamingSources(ProviderStreamingSources))
 				{
-					// Transform to Local
-					StreamingSource.Location = WorldToLocal.TransformPosition(StreamingSource.Location);
-					StreamingSource.Rotation = WorldToLocal.TransformRotation(StreamingSource.Rotation.Quaternion()).Rotator();
-					StreamingSources.Add(StreamingSource);
+					for (FWorldPartitionStreamingSource& ProviderStreamingSource : ProviderStreamingSources)
+					{
+						// Transform to Local
+						ProviderStreamingSource.Location = WorldToLocal.TransformPosition(ProviderStreamingSource.Location);
+						ProviderStreamingSource.Rotation = WorldToLocal.TransformRotation(ProviderStreamingSource.Rotation.Quaternion()).Rotator();
+						StreamingSources.Add(MoveTemp(ProviderStreamingSource));
+					}
 				}
 			}
 		}
@@ -725,13 +727,13 @@ bool UWorldPartitionStreamingPolicy::CanAddLoadedLevelToWorld(ULevel* InLevel) c
 	return true;
 }
 
-bool UWorldPartitionStreamingPolicy::IsStreamingCompleted(const FWorldPartitionStreamingSource* InStreamingSource) const
+bool UWorldPartitionStreamingPolicy::IsStreamingCompleted(const TArray<FWorldPartitionStreamingSource>* InStreamingSources) const
 {
 	const UWorld* World = GetWorld();
 	check(World);
 	check(World->IsGameWorld());
 	const UDataLayerSubsystem* DataLayerSubsystem = World->GetSubsystem<UDataLayerSubsystem>();
-	const bool bTestProvidedStreamingSource = !!InStreamingSource;
+	const bool bTestProvidedStreamingSource = !!InStreamingSources;
 
 	// Always test non-spatial cells
 	{
@@ -759,7 +761,7 @@ bool UWorldPartitionStreamingPolicy::IsStreamingCompleted(const FWorldPartitionS
 	}
 
 	// Test spatially loaded cells using streaming sources (or provided streaming source)
-	TArrayView<const FWorldPartitionStreamingSource> QueriedStreamingSources = bTestProvidedStreamingSource ? MakeArrayView({ *InStreamingSource }) : StreamingSources;
+	TArrayView<const FWorldPartitionStreamingSource> QueriedStreamingSources = bTestProvidedStreamingSource ? *InStreamingSources : StreamingSources;
 	for (const FWorldPartitionStreamingSource& StreamingSource : QueriedStreamingSources)
 	{
 		// Build a query source from a Streaming Source
@@ -768,7 +770,9 @@ bool UWorldPartitionStreamingPolicy::IsStreamingCompleted(const FWorldPartitionS
 		QuerySource.bSpatialQuery = true;
 		QuerySource.Location = StreamingSource.Location;
 		QuerySource.Rotation = StreamingSource.Rotation;
-		QuerySource.TargetGrid = StreamingSource.TargetGrid;
+		QuerySource.TargetBehavior = StreamingSource.TargetBehavior;
+		QuerySource.TargetGrids = StreamingSource.TargetGrids;
+		QuerySource.TargetHLODLayers = StreamingSource.TargetHLODLayers;
 		QuerySource.Shapes = StreamingSource.Shapes;
 		QuerySource.bUseGridLoadingRange = true;
 		QuerySource.Radius = 0.f;
diff --git a/Engine/Source/Runtime/Engine/Private/WorldPartition/WorldPartitionSubsystem.cpp b/Engine/Source/Runtime/Engine/Private/WorldPartition/WorldPartitionSubsystem.cpp
--- a/Engine/Source/Runtime/Engine/Private/WorldPartition/WorldPartitionSubsystem.cpp
+++ b/Engine/Source/Runtime/Engine/Private/WorldPartition/WorldPartitionSubsystem.cpp
@@ -270,12 +270,12 @@ bool UWorldPartitionSubsystem::IsStreamingCompleted(const IWorldPartitionStreami
 {
 	// Convert specified/optional streaming source provider to a world partition 
 	// streaming source and pass it along to each registered world partition
-	FWorldPartitionStreamingSource StreamingSource;
-	FWorldPartitionStreamingSource* StreamingSourcePtr = nullptr;
+	TArray<FWorldPartitionStreamingSource> StreamingSources;
+	TArray<FWorldPartitionStreamingSource>* StreamingSourcesPtr = nullptr;
 	if (InStreamingSourceProvider)
 	{
-		StreamingSourcePtr = &StreamingSource;
-		if (!InStreamingSourceProvider->GetStreamingSource(StreamingSource))
+		StreamingSourcesPtr = &StreamingSources;
+		if (!InStreamingSourceProvider->GetStreamingSources(StreamingSources))
 		{
 			return true;
 		}
@@ -283,7 +283,7 @@ bool UWorldPartitionSubsystem::IsStreamingCompleted(const IWorldPartitionStreami
 
 	for (UWorldPartition* RegisteredWorldPartition : RegisteredWorldPartitions)
 	{
-		if (!RegisteredWorldPartition->IsStreamingCompleted(StreamingSourcePtr))
+		if (!RegisteredWorldPartition->IsStreamingCompleted(StreamingSourcesPtr))
 		{
 			return false;
 		}
diff --git a/Engine/Source/Runtime/Engine/Public/WorldPartition/WorldPartition.h b/Engine/Source/Runtime/Engine/Public/WorldPartition/WorldPartition.h
--- a/Engine/Source/Runtime/Engine/Public/WorldPartition/WorldPartition.h
+++ b/Engine/Source/Runtime/Engine/Public/WorldPartition/WorldPartition.h
@@ -237,7 +237,7 @@ public:
 	void Tick(float DeltaSeconds);
 	void UpdateStreamingState();
 	bool CanAddLoadedLevelToWorld(ULevel* InLevel) const;
-	bool IsStreamingCompleted(const FWorldPartitionStreamingSource* InStreamingSource) const;
+	bool IsStreamingCompleted(const TArray<FWorldPartitionStreamingSource>* InStreamingSources) const;
 	bool IsStreamingCompleted(EWorldPartitionRuntimeCellState QueryState, const TArray<FWorldPartitionStreamingQuerySource>& QuerySources, bool bExactState) const;
 
 	const TArray<FWorldPartitionStreamingSource>& GetStreamingSources() const;
diff --git a/Engine/Source/Runtime/Engine/Public/WorldPartition/WorldPartitionLevelStreamingPolicy.h b/Engine/Source/Runtime/Engine/Public/WorldPartition/WorldPartitionLevelStreamingPolicy.h
--- a/Engine/Source/Runtime/Engine/Public/WorldPartition/WorldPartitionLevelStreamingPolicy.h
+++ b/Engine/Source/Runtime/Engine/Public/WorldPartition/WorldPartitionLevelStreamingPolicy.h
@@ -24,7 +24,7 @@ class UWorldPartitionLevelStreamingPolicy : public UWorldPartitionStreamingPolic
 public:
 	virtual void DrawRuntimeCellsDetails(class UCanvas* Canvas, FVector2D& Offset) override;
 	virtual void DrawStreamingStatusLegend(UCanvas* Canvas, FVector2D& Offset) override;
-	virtual bool IsStreamingCompleted(const FWorldPartitionStreamingSource* InStreamingSource) const override;
+	virtual bool IsStreamingCompleted(const TArray<FWorldPartitionStreamingSource>* InStreamingSources) const override;
 
 #if WITH_EDITOR
 	virtual TSubclassOf<class UWorldPartitionRuntimeCell> GetRuntimeCellClass() const override;
diff --git a/Engine/Source/Runtime/Engine/Public/WorldPartition/WorldPartitionStreamingPolicy.h b/Engine/Source/Runtime/Engine/Public/WorldPartition/WorldPartitionStreamingPolicy.h
--- a/Engine/Source/Runtime/Engine/Public/WorldPartition/WorldPartitionStreamingPolicy.h
+++ b/Engine/Source/Runtime/Engine/Public/WorldPartition/WorldPartitionStreamingPolicy.h
@@ -49,7 +49,7 @@ public:
 	virtual void DrawRuntimeCellsDetails(class UCanvas* Canvas, FVector2D& Offset) {}
 	virtual void DrawStreamingStatusLegend(class UCanvas* Canvas, FVector2D& Offset) {}
 
-	virtual bool IsStreamingCompleted(const FWorldPartitionStreamingSource* InStreamingSource) const;
+	virtual bool IsStreamingCompleted(const TArray<FWorldPartitionStreamingSource>* InStreamingSources) const;
 	virtual bool IsStreamingCompleted(EWorldPartitionRuntimeCellState QueryState, const TArray<FWorldPartitionStreamingQuerySource>& QuerySources, bool bExactState = true) const;
 
 #if WITH_EDITOR
diff --git a/Engine/Source/Runtime/Engine/Public/WorldPartition/WorldPartitionStreamingSource.h b/Engine/Source/Runtime/Engine/Public/WorldPartition/WorldPartitionStreamingSource.h
--- a/Engine/Source/Runtime/Engine/Public/WorldPartition/WorldPartitionStreamingSource.h
+++ b/Engine/Source/Runtime/Engine/Public/WorldPartition/WorldPartitionStreamingSource.h
@@ -129,6 +129,7 @@ struct FStreamingSourceShape
 
 	FStreamingSourceShape()
 	: bUseGridLoadingRange(true)
+	, LoadingRangeScale(1.f)
 	, Radius(10000.0f)
 	, bIsSector(false)
 	, SectorAngle(360.0f)
@@ -140,6 +141,10 @@ struct FStreamingSourceShape
 	UPROPERTY(BlueprintReadWrite, EditAnywhere, Category = Streaming)
 	bool bUseGridLoadingRange;
 
+	/* Applies a scale to the grid's loading range (used if bUseGridLoadingRange is True). */
+	UPROPERTY(BlueprintReadWrite, EditAnywhere, Category = Streaming, meta = (EditCondition = "bUseGridLoadingRange"))
+	float LoadingRangeScale;
+
 	/* Custom streaming source shape radius (not used if bUseGridLoadingRange is True). */
 	UPROPERTY(BlueprintReadWrite, EditAnywhere, Category = Streaming, meta = (EditCondition = "!bUseGridLoadingRange"))
 	float Radius;
@@ -161,11 +166,34 @@ struct FStreamingSourceShape
 	FRotator Rotation;
 };
 
+UENUM(BlueprintType)
+enum class EStreamingSourceTargetBehavior : uint8
+{
+	Include,
+	Exclude
+};
+
 /** Helper class used to iterate over streaming source shapes. */
 class FStreamingSourceShapeHelper
 {
 public:
 
+	FORCEINLINE static bool IsSourceAffectingGrid(const TSet<FName>& InSourceTargetGrids, const TSet<FSoftObjectPath>& InSourceTargetHLODLayers, EStreamingSourceTargetBehavior InSourceTargetBehavior, FName InGridName, const FSoftObjectPath& InGridHLODLayer)
+	{
+		switch (InSourceTargetBehavior)
+		{
+		case EStreamingSourceTargetBehavior::Include:
+			return (InSourceTargetGrids.IsEmpty() && InSourceTargetHLODLayers.IsEmpty()) || (InSourceTargetGrids.Contains(InGridName) || InSourceTargetHLODLayers.Contains(InGridHLODLayer));
+		case EStreamingSourceTargetBehavior::Exclude:
+			return !(InSourceTargetGrids.Contains(InGridName) || InSourceTargetHLODLayers.Contains(InGridHLODLayer));
+		default:
+			checkNoEntry();
+		}
+
+		return false;
+	}
+
+	UE_DEPRECATED(5.1, "Use IsSourceAffectingGrid that takes Target Sets and Behavior instead.")
 	FORCEINLINE static bool IsSourceAffectingGrid(FName InSourceTargetGrid, const FSoftObjectPath& InSourceTargetHLODLayer, FName InGridName, const FSoftObjectPath& InGridHLODLayer)
 	{
 		if ((InSourceTargetGrid.IsNone() && InSourceTargetHLODLayer.IsNull()) ||
@@ -191,7 +219,7 @@ public:
 		{
 			for (const FStreamingSourceShape& Shape : InShapes)
 			{
-				const FVector::FReal ShapeRadius = Shape.bUseGridLoadingRange ? InGridLoadingRange : Shape.Radius;
+				const FVector::FReal ShapeRadius = Shape.bUseGridLoadingRange ? (InGridLoadingRange * Shape.LoadingRangeScale) : Shape.Radius;
 				const FVector::FReal ShapeAngle = Shape.bIsSector ? Shape.SectorAngle : 360.0f;
 				const FVector ShapeAxis = bInProjectIn2D ? FRotator(0, Shape.Rotation.Yaw, 0).Vector() : Shape.Rotation.Vector();
 				FSphericalSector LocalShape(bInProjectIn2D ? FVector(Shape.Location.X, Shape.Location.Y, 0) : Shape.Location, ShapeRadius, ShapeAxis, ShapeAngle);
@@ -241,7 +269,7 @@ struct FWorldPartitionStreamingQuerySource
 		, bDataLayersOnly(false)
 		, bSpatialQuery(true)
 		, Rotation(ForceInitToZero)
-		, TargetGrid(NAME_None)
+		, TargetBehavior(EStreamingSourceTargetBehavior::Include)
 	{}
 
 	FWorldPartitionStreamingQuerySource(const FVector& InLocation)
@@ -251,9 +279,37 @@ struct FWorldPartitionStreamingQuerySource
 		, bDataLayersOnly(false)
 		, bSpatialQuery(true)
 		, Rotation(ForceInitToZero)
-		, TargetGrid(NAME_None)
+		, TargetBehavior(EStreamingSourceTargetBehavior::Include)
 	{}
 
+	// Define Copy Constructor to avoid deprecation warnings
+	FWorldPartitionStreamingQuerySource(const FWorldPartitionStreamingQuerySource& Other)
+	{
+		*this = Other;
+	}
+
+	FWorldPartitionStreamingQuerySource& operator=(const FWorldPartitionStreamingQuerySource& Other)
+	{
+		Location = Other.Location;
+		Radius = Other.Radius;
+		bUseGridLoadingRange = Other.bUseGridLoadingRange;
+		DataLayers = Other.DataLayers;
+		bDataLayersOnly = Other.bDataLayersOnly;
+		bSpatialQuery = Other.bSpatialQuery;
+		Rotation = Other.Rotation;
+		TargetBehavior = Other.TargetBehavior;
+		TargetGrids = Other.TargetGrids;
+		TargetHLODLayers = Other.TargetHLODLayers;
+		Shapes = Other.Shapes;
+
+		PRAGMA_DISABLE_DEPRECATION_WARNINGS
+		TargetGrid = Other.TargetGrid;
+		TargetHLODLayer = Other.TargetHLODLayer;
+		PRAGMA_ENABLE_DEPRECATION_WARNINGS
+
+		return *this;
+	}
+
 	/* Location to query. (not used if bSpatialQuery is false) */
 	UPROPERTY(BlueprintReadWrite, EditAnywhere, Category = "Query")
 	FVector Location;
@@ -280,8 +336,20 @@ struct FWorldPartitionStreamingQuerySource
 
 	/* Reserved settings used by UWorldPartitionStreamingSourceComponent::IsStreamingCompleted. */
 	FRotator Rotation;
+	
+	/** Defines how TargetGrids/TargetHLODLayers will be applied to this streaming source. */
+	EStreamingSourceTargetBehavior TargetBehavior;
+
+	UE_DEPRECATED(5.1, "Use TargetGrids instead.")
 	FName TargetGrid;
+
+	TSet<FName> TargetGrids;
+
+	UE_DEPRECATED(5.1, "Use TargetHLODLayers")
 	FSoftObjectPath TargetHLODLayer;
+
+	TSet<FSoftObjectPath> TargetHLODLayers;
+	
 	TArray<FStreamingSourceShape> Shapes;
 
 	/** Helper method that iterates over all shapes. If none is provided, it will still pass a sphere shape using Radius or grid's loading range (see bUseGridLoadingRange). */
@@ -292,10 +360,21 @@ struct FWorldPartitionStreamingQuerySource
 			return;
 		}
 
-		if (FStreamingSourceShapeHelper::IsSourceAffectingGrid(TargetGrid, TargetHLODLayer, InGridName, InGridHLODLayer))
+		if (FStreamingSourceShapeHelper::IsSourceAffectingGrid(TargetGrids, TargetHLODLayers, TargetBehavior, InGridName, InGridHLODLayer))
 		{
 			FStreamingSourceShapeHelper::ForEachShape(InGridLoadingRange, bUseGridLoadingRange ? InGridLoadingRange : Radius, bInProjectIn2D, Location, Rotation, Shapes, InOperation);
 		}
+
+		PRAGMA_DISABLE_DEPRECATION_WARNINGS
+		if(!TargetGrid.IsNone() || !TargetHLODLayer.IsNull())
+		{
+
+			if (FStreamingSourceShapeHelper::IsSourceAffectingGrid({ TargetGrid }, { TargetHLODLayers }, TargetBehavior, InGridName, InGridHLODLayer))
+			{
+				FStreamingSourceShapeHelper::ForEachShape(InGridLoadingRange, bUseGridLoadingRange ? InGridLoadingRange : Radius, bInProjectIn2D, Location, Rotation, Shapes, InOperation);
+			}
+		}
+		PRAGMA_ENABLE_DEPRECATION_WARNINGS
 	}
 };
 
@@ -323,7 +402,7 @@ struct ENGINE_API FWorldPartitionStreamingSource
 		, Priority(EStreamingSourcePriority::Default)
 		, Velocity(0.f)
 		, DebugColor(ForceInit)
-		, TargetGrid(NAME_None)
+		, TargetBehavior(EStreamingSourceTargetBehavior::Include)
 		, bReplay(false)
 		, bRemote(false)
 	{}
@@ -337,11 +416,42 @@ struct ENGINE_API FWorldPartitionStreamingSource
 		, Priority(InPriority)
 		, Velocity(InVelocity)
 		, DebugColor(ForceInit)
-		, TargetGrid(NAME_None)
+		, TargetBehavior(EStreamingSourceTargetBehavior::Include)
 		, bReplay(false)
 		, bRemote(bRemote)
 	{}
 
+	// Define Copy Constructor to avoid deprecation warnings
+	FWorldPartitionStreamingSource(const FWorldPartitionStreamingSource& Other)
+	{
+		*this = Other;
+	}
+
+	FWorldPartitionStreamingSource& operator=(const FWorldPartitionStreamingSource& Other)
+	{
+		Name = Other.Name;
+		Location = Other.Location;
+		Rotation = Other.Rotation;
+		TargetState = Other.TargetState;
+		bBlockOnSlowLoading = Other.bBlockOnSlowLoading;
+		Priority = Other.Priority;
+		Velocity = Other.Velocity;
+		DebugColor = Other.DebugColor;
+		TargetBehavior = Other.TargetBehavior;
+		TargetGrids = Other.TargetGrids;
+		TargetHLODLayers = Other.TargetHLODLayers;
+		Shapes = Other.Shapes;
+		bReplay = Other.bReplay;
+		bRemote = Other.bRemote;
+
+		PRAGMA_DISABLE_DEPRECATION_WARNINGS
+		TargetGrid = Other.TargetGrid;
+		TargetHLODLayer = Other.TargetHLODLayer;
+		PRAGMA_ENABLE_DEPRECATION_WARNINGS
+
+		return *this;
+	}
+	
 	FColor GetDebugColor() const
 	{
 		if (!DebugColor.ToPackedBGRA())
@@ -376,12 +486,21 @@ struct ENGINE_API FWorldPartitionStreamingSource
 	/** Color used for debugging. */
 	FColor DebugColor;
 
-	/** When set, will only affect streaming on the provided target runtime streaming grid. */
+	/** Defines how TargetGrids/TargetHLODLayers will be applied to this streaming source. */
+	EStreamingSourceTargetBehavior TargetBehavior;
+
+	UE_DEPRECATED(5.1, "Use TargetGrids instead.")
 	FName TargetGrid;
 
-	/** When set, will only affect streaming on HLODs associated to the provided target HLODLayer. */
+	/** When set, this will change how this streaming source is applied to the provided runtime streaming grids based on the TargetBehavior. */
+	TSet<FName> TargetGrids;
+	
+	UE_DEPRECATED(5.1, "Use TargetHLODLayers")
 	FSoftObjectPath TargetHLODLayer;
 
+	/** When set, this will change how this streaming source is applied to the provided HLODLayers based on the TargetBehavior. */
+	TSet<FSoftObjectPath> TargetHLODLayers;
+
 	/** Source internal shapes. When none are provided, a sphere is automatically used. It's radius is equal to grid's loading range and center equals source's location. */
 	TArray<FStreamingSourceShape> Shapes;
 
@@ -392,7 +511,7 @@ struct ENGINE_API FWorldPartitionStreamingSource
 	bool bRemote;
 
 	/** Returns a box encapsulating all shapes. */
-	FORCEINLINE FBox CalcBounds(float InGridLoadingRange, FName InGridName, const FSoftObjectPath& InGridHLODLayer, bool bCalcIn2D = false) const
+	FBox CalcBounds(float InGridLoadingRange, FName InGridName, const FSoftObjectPath& InGridHLODLayer, bool bCalcIn2D = false) const
 	{
 		FBox OutBounds(ForceInit);
 		ForEachShape(InGridLoadingRange, InGridName, InGridHLODLayer, bCalcIn2D, [&OutBounds](const FSphericalSector& Sector)
@@ -403,12 +522,22 @@ struct ENGINE_API FWorldPartitionStreamingSource
 	}
 
 	/** Helper method that iterates over all shapes. If none is provided, it will still pass a sphere shape using grid's loading range. */
-	FORCEINLINE void ForEachShape(float InGridLoadingRange, FName InGridName, const FSoftObjectPath& InGridHLODLayer, bool bInProjectIn2D, TFunctionRef<void(const FSphericalSector&)> InOperation) const
+	void ForEachShape(float InGridLoadingRange, FName InGridName, const FSoftObjectPath& InGridHLODLayer, bool bInProjectIn2D, TFunctionRef<void(const FSphericalSector&)> InOperation) const
 	{
-		if (FStreamingSourceShapeHelper::IsSourceAffectingGrid(TargetGrid, TargetHLODLayer, InGridName, InGridHLODLayer))
+		if (FStreamingSourceShapeHelper::IsSourceAffectingGrid(TargetGrids, TargetHLODLayers, TargetBehavior, InGridName, InGridHLODLayer))
 		{
 			FStreamingSourceShapeHelper::ForEachShape(InGridLoadingRange, InGridLoadingRange, bInProjectIn2D, Location, Rotation, Shapes, InOperation);
 		}
+
+		PRAGMA_DISABLE_DEPRECATION_WARNINGS
+		if (!TargetGrid.IsNone() || !TargetHLODLayer.IsNull())
+		{
+			if (FStreamingSourceShapeHelper::IsSourceAffectingGrid({ TargetGrid }, { TargetHLODLayers }, TargetBehavior, InGridName, InGridHLODLayer))
+			{
+				FStreamingSourceShapeHelper::ForEachShape(InGridLoadingRange, InGridLoadingRange, bInProjectIn2D, Location, Rotation, Shapes, InOperation);
+			}
+		}
+		PRAGMA_ENABLE_DEPRECATION_WARNINGS
 	}
 
 	FString ToString() const
@@ -433,5 +562,19 @@ struct ENGINE_API FWorldPartitionStreamingSource
  */
 struct ENGINE_API IWorldPartitionStreamingSourceProvider
 {
-	virtual bool GetStreamingSource(FWorldPartitionStreamingSource& StreamingSource) const = 0;
+	virtual bool GetStreamingSource(FWorldPartitionStreamingSource& StreamingSource) const
+	{
+		return false;
+	}
+
+	virtual bool GetStreamingSources(TArray<FWorldPartitionStreamingSource>& StreamingSources) const
+	{
+		FWorldPartitionStreamingSource StreamingSource;
+		if (GetStreamingSource(StreamingSource))
+		{
+			StreamingSources.Add(StreamingSource);
+			return true;
+		}
+		return false;
+	}
 };
\ No newline at end of file
