Do not forward events on Level instances actors.
OnObjectReplace now filters actors as well

#test LI edit in place, template map, PIE

diff --git a/Engine/Source/Runtime/Engine/Private/WorldPartition/ActorDescContainer.cpp b/Engine/Source/Runtime/Engine/Private/WorldPartition/ActorDescContainer.cpp
--- a/Engine/Source/Runtime/Engine/Private/WorldPartition/ActorDescContainer.cpp
+++ b/Engine/Source/Runtime/Engine/Private/WorldPartition/ActorDescContainer.cpp
@@ -154,9 +154,7 @@ bool UActorDescContainer::ShouldHandleActorEvent(const AActor* Actor)
 {
 	if (Actor != nullptr)
 	{
-		bool bHandlesActor = GetActorDesc(Actor->GetActorGuid()) != nullptr || IsActorDescHandled(Actor);
-		return bHandlesActor && Actor->IsMainPackageActor() && (Actor->GetLevel() != nullptr);
-		
+		return IsActorDescHandled(Actor) && Actor->IsMainPackageActor() && (Actor->GetLevel() != nullptr);
 	}
 	
 	return false;
@@ -206,10 +204,13 @@ void UActorDescContainer::OnObjectsReplaced(const TMap<UObject*, UObject*>& OldT
 	{
 		if (AActor* OldActor = Cast<AActor>(OldObject))
 		{
-			AActor* NewActor = Cast<AActor>(NewObject);
-			if (FWorldPartitionActorDesc* ActorDesc = GetActorDesc(OldActor->GetActorGuid()))
-			{				
-				FWorldPartitionActorDescUtils::ReplaceActorDescriptorPointerFromActor(OldActor, NewActor, ActorDesc);
+			if (ShouldHandleActorEvent(OldActor))
+			{
+				AActor* NewActor = Cast<AActor>(NewObject);
+				if (FWorldPartitionActorDesc* ActorDesc = GetActorDesc(OldActor->GetActorGuid()))
+				{
+					FWorldPartitionActorDescUtils::ReplaceActorDescriptorPointerFromActor(OldActor, NewActor, ActorDesc);
+				}
 			}
 		}
 	}
