Add GetLevelStreamingClass to ILevelInstanceInterface so level instance can use a derived streaming class
#rnx

diff --git a/Engine/Source/Runtime/Engine/Private/LevelInstance/LevelInstanceInterface.cpp b/Engine/Source/Runtime/Engine/Private/LevelInstance/LevelInstanceInterface.cpp
--- a/Engine/Source/Runtime/Engine/Private/LevelInstance/LevelInstanceInterface.cpp
+++ b/Engine/Source/Runtime/Engine/Private/LevelInstance/LevelInstanceInterface.cpp
@@ -2,6 +2,7 @@
 
 #include "LevelInstance/LevelInstanceInterface.h"
 #include "LevelInstance/LevelInstanceSubsystem.h"
+#include "LevelInstance/LevelInstanceLevelStreaming.h"
 #include "Engine/Level.h"
 #include "Engine/World.h"
 #include "GameFramework/Actor.h"
@@ -73,6 +74,11 @@ ULevelStreamingLevelInstance* ILevelInstanceInterface::GetLevelStreaming() const
 	return nullptr;
 }
 
+TSubclassOf<ULevelStreamingLevelInstance> ILevelInstanceInterface::GetLevelStreamingClass() const
+{
+	return ULevelStreamingLevelInstance::StaticClass();
+}
+
 #if WITH_EDITOR
 ULevel* ILevelInstanceInterface::GetLoadedLevel() const
 {
diff --git a/Engine/Source/Runtime/Engine/Private/LevelInstance/LevelInstanceLevelStreaming.cpp b/Engine/Source/Runtime/Engine/Private/LevelInstance/LevelInstanceLevelStreaming.cpp
--- a/Engine/Source/Runtime/Engine/Private/LevelInstance/LevelInstanceLevelStreaming.cpp
+++ b/Engine/Source/Runtime/Engine/Private/LevelInstance/LevelInstanceLevelStreaming.cpp
@@ -205,7 +205,7 @@ ULevelStreamingLevelInstance* ULevelStreamingLevelInstance::LoadInstance(ILevelI
 
 	FLoadLevelInstanceParams Params(World, LevelInstance->GetWorldAssetPackage(), LevelInstanceActor->GetActorTransform());
 	Params.OptionalLevelNameOverride = &Suffix;
-	Params.OptionalLevelStreamingClass = ULevelStreamingLevelInstance::StaticClass();
+	Params.OptionalLevelStreamingClass = LevelInstance->GetLevelStreamingClass();
 	Params.bLoadAsTempPackage = true;
 	
 	if (World->IsGameWorld())
diff --git a/Engine/Source/Runtime/Engine/Public/LevelInstance/LevelInstanceInterface.h b/Engine/Source/Runtime/Engine/Public/LevelInstance/LevelInstanceInterface.h
--- a/Engine/Source/Runtime/Engine/Public/LevelInstance/LevelInstanceInterface.h
+++ b/Engine/Source/Runtime/Engine/Public/LevelInstance/LevelInstanceInterface.h
@@ -70,6 +70,8 @@ class ENGINE_API ILevelInstanceInterface
 
 	virtual ULevelStreamingLevelInstance* GetLevelStreaming() const;
 
+	virtual TSubclassOf<ULevelStreamingLevelInstance> GetLevelStreamingClass() const;
+
 #if WITH_EDITOR
 	virtual ULevel* GetLoadedLevel() const;
 
