Lay the foundation for advanced controller vibrations.

Expose Input Device properties as UInputDeviceProperty. This utilizes the existing IInputInterface API to set device specific properties like Light Color, trigger effects, and advanced rumble.

These UInputDeviceProperty's can be added to any Force Feedback Effect by default, allowing you to create nicely contained timeline driven properties. Alternatively, you can make a blueprint of UInputDeviceProperty, add it as a variable, and manually tick and apply it to the input interface. This means that the new custom device properties are not limited to force feedback effects, and you can create your own effects in a blueprint! This will make doing complex vibration patterns much easier for designers and gameplay programmers to implement.

diff --git a/Engine/Source/Runtime/Engine/Classes/GameFramework/ForceFeedbackEffect.h b/Engine/Source/Runtime/Engine/Classes/GameFramework/ForceFeedbackEffect.h
--- a/Engine/Source/Runtime/Engine/Classes/GameFramework/ForceFeedbackEffect.h
+++ b/Engine/Source/Runtime/Engine/Classes/GameFramework/ForceFeedbackEffect.h
@@ -1,4 +1,4 @@
-﻿// Copyright Epic Games, Inc. All Rights Reserved.
+// Copyright Epic Games, Inc. All Rights Reserved.
 #pragma once
 
 #include "CoreMinimal.h"
@@ -9,6 +9,7 @@
 
 class UForceFeedbackEffect;
 struct FForceFeedbackValues;
+class UInputDeviceProperty;
 
 USTRUCT()
 struct FForceFeedbackChannelDetails
@@ -75,16 +76,21 @@ struct ENGINE_API FActiveForceFeedbackEffect
 	FForceFeedbackParameters Parameters;
 	float PlayTime;
 
+	/** The platform user that should receive this effect */
+	FPlatformUserId PlatformUser = PLATFORMUSERID_NONE;
+
 	FActiveForceFeedbackEffect()
 		: ForceFeedbackEffect(nullptr)
 		, PlayTime(0.f)
+		, PlatformUser(PLATFORMUSERID_NONE)
 	{
 	}
 
-	FActiveForceFeedbackEffect(UForceFeedbackEffect* InEffect, FForceFeedbackParameters InParameters)
+	FActiveForceFeedbackEffect(UForceFeedbackEffect* InEffect, FForceFeedbackParameters InParameters, FPlatformUserId InPlatformUser)
 		: ForceFeedbackEffect(InEffect)
 		, Parameters(InParameters)
 		, PlayTime(0.f)
+		, PlatformUser(InPlatformUser)
 	{
 	}
 
@@ -106,6 +112,10 @@ class UForceFeedbackEffect : public UObject
 	UPROPERTY(EditAnywhere, Category="ForceFeedbackEffect")
 	TArray<FForceFeedbackChannelDetails> ChannelDetails;
 
+	/** A map of input device properties that we want to set while this effect is playing */
+	UPROPERTY(EditAnywhere, Instanced, Category = "ForceFeedbackEffect")
+	TArray<TObjectPtr<UInputDeviceProperty>> DeviceProperties;
+
 	/** Duration of force feedback pattern in seconds. */
 	UPROPERTY(Category=Info, AssetRegistrySearchable, VisibleAnywhere, BlueprintReadOnly)
 	float Duration;
@@ -117,4 +127,6 @@ class UForceFeedbackEffect : public UObject
 	float GetDuration();
 
 	void GetValues(const float EvalTime, FForceFeedbackValues& Values, float ValueMultiplier = 1.f) const;
+
+	void SetDeviceProperties(const FPlatformUserId PlatformUser, const float DeltaTime, const float EvalTime);
 };
diff --git a/Engine/Source/Runtime/Engine/Classes/GameFramework/InputDeviceProperties.h b/Engine/Source/Runtime/Engine/Classes/GameFramework/InputDeviceProperties.h
new file mode 100644
--- /dev/null
+++ b/Engine/Source/Runtime/Engine/Classes/GameFramework/InputDeviceProperties.h
@@ -0,0 +1,112 @@
+// Copyright Epic Games, Inc. All Rights Reserved.
+
+#pragma once
+
+#include "CoreMinimal.h"
+#include "UObject/Object.h"
+#include "Curves/CurveFloat.h"
+#include "GenericPlatform/IInputInterface.h"
+
+#include "InputDeviceProperties.generated.h"
+
+/**
+* Base class that represents a single Input Device Property. An Input Device Property
+* represents a feature that can be set on an input device. Things like what color a
+* light is, advanced rumble patterns, or trigger haptics.
+* 
+* This top level object can then be evaluated at a specific time to create a lower level
+* FInputDeviceProperty, which the IInputInterface implementation can interpret however it desires.
+* 
+* The behavior of device properties can vary depending on the current platform. Some platforms may not
+* support certain device properties. An older gamepad may not have any advanced trigger haptics for 
+* example. 
+*/
+UCLASS(Abstract, Blueprintable, EditInlineNew)
+class ENGINE_API UInputDeviceProperty : public UObject
+{
+	GENERATED_BODY()
+public:
+
+	/**
+	* Evaluate this device property for a given duration. 
+	* 
+	* @param DeltaTime		Delta time
+	* @param Duration		The number of seconds that this property has been active. Use this to get things like curve data over time.
+	* @return				A pointer to the evaluated input device property.
+	*/
+	UFUNCTION(BlueprintCallable, BlueprintNativeEvent, Category = "InputDevice")
+	void EvaluateDeviceProperty(const float DeltaTime, const float Duration);
+
+	/** 
+	* Native C++ implementation of EvaluateDeviceProperty.
+	* 
+	* Override this to alter your device property in native code.
+	* @see UInputDeviceProperty::EvaluateDeviceProperty
+	*/
+	virtual void EvaluateDeviceProperty_Implementation(const float DeltaTime, const float Duration) { }
+
+	/**
+	* Apply the given device property to the Input Interface 
+	* 
+	* @param UserId		The owning Platform User whose input device this property should be applied to.
+	*/
+	UFUNCTION(BlueprintCallable, Category = "InputDevice")
+	virtual void ApplyDeviceProperty(const FPlatformUserId UserId);
+	
+	/** Gets a pointer to the current input device property that the IInputInterface can use. */
+	virtual FInputDeviceProperty* GetInternalDeviceProperty() { return nullptr; };
+
+	/**
+	* The duration that this device property should last. Override this if your property has any dynamic curves 
+	* to be the max time range.
+	*/
+	float GetDuration() const { return PropertyDuration; }
+
+protected:
+
+	/** 
+	* The name of this unique device property. 
+	* This is used internally by the platform to determine which device property this represents. 
+	* This needs to correspond to a FInputDeviceProperty name!
+	*/
+	UPROPERTY(VisibleAnywhere, Category = "Info")
+	FName PropertyName;
+
+	/**
+	* The duration that this device property should last. Override this if your property has any dynamic curves 
+	* to be the max time range.
+	*/
+	UPROPERTY(EditAnywhere, BlueprintReadWrite, Category = "Info")
+	float PropertyDuration = 0.1f;
+};
+
+class UCurveLinearColor;
+
+/** A property that can be used to change the color of an input device's light */
+UCLASS(Blueprintable, EditInlineNew)
+class UColorInputDeviceProperty : public UInputDeviceProperty
+{
+	GENERATED_BODY()
+
+public:
+
+	UColorInputDeviceProperty();
+
+	virtual void EvaluateDeviceProperty_Implementation(const float DeltaTime, const float Duration) override;
+	virtual FInputDeviceProperty* GetInternalDeviceProperty() override;
+
+protected:
+
+	/** True if the light should be enabled at all */
+	UPROPERTY(EditAnywhere, BlueprintReadWrite, Category = "DeviceProperty")
+	bool bEnable = true;
+
+	/** The color the device light should be */
+	UPROPERTY(EditAnywhere, BlueprintReadWrite, Category = "DeviceProperty")
+	TObjectPtr<UCurveLinearColor> DeviceColorCurve;
+
+private:
+
+	/** The internal light color property that this represents; */
+	FInputDeviceLightColorProperty InternalProperty;
+};
\ No newline at end of file
diff --git a/Engine/Source/Runtime/Engine/Private/GameFramework/ForceFeedbackEffect.cpp b/Engine/Source/Runtime/Engine/Private/GameFramework/ForceFeedbackEffect.cpp
--- a/Engine/Source/Runtime/Engine/Private/GameFramework/ForceFeedbackEffect.cpp
+++ b/Engine/Source/Runtime/Engine/Private/GameFramework/ForceFeedbackEffect.cpp
@@ -3,6 +3,7 @@
 #include "GameFramework/ForceFeedbackEffect.h"
 #include "GenericPlatform/IInputInterface.h"
 #include "Misc/App.h"
+#include "GameFramework/InputDeviceProperties.h"
 
 #include UE_INLINE_GENERATED_CPP_BY_NAME(ForceFeedbackEffect)
 
@@ -15,7 +16,7 @@ UForceFeedbackEffect::UForceFeedbackEffect(const FObjectInitializer& ObjectIniti
 }
 
 #if WITH_EDITOR
-void UForceFeedbackEffect::PostEditChangeChainProperty( struct FPropertyChangedChainEvent& PropertyChangedEvent)
+void UForceFeedbackEffect::PostEditChangeChainProperty(struct FPropertyChangedChainEvent& PropertyChangedEvent)
 {
 	// After any edit (really we only care about the curve, but easier this way) update the cached duration value
 	GetDuration();
@@ -40,6 +41,19 @@ float UForceFeedbackEffect::GetDuration()
 				Duration = MaxTime;
 			}
 		}
+
+		// Check the device properties for any longer durations
+		for (TObjectPtr<UInputDeviceProperty> DeviceProperty : DeviceProperties)
+		{
+			if (DeviceProperty)
+			{
+				const float PropertyDuration = DeviceProperty->GetDuration();
+				if (PropertyDuration > Duration)
+				{
+					Duration = PropertyDuration;
+				}	
+			}
+		}
 	}
 
 	return Duration;
@@ -71,6 +85,18 @@ void UForceFeedbackEffect::GetValues(const float EvalTime, FForceFeedbackValues&
 	}
 }
 
+void UForceFeedbackEffect::SetDeviceProperties(const FPlatformUserId PlatformUser, const float DeltaTime, const float EvalTime)
+{
+	for (TObjectPtr<UInputDeviceProperty> DeviceProp : DeviceProperties)
+	{
+		if (DeviceProp)
+		{
+			DeviceProp->EvaluateDeviceProperty(DeltaTime, EvalTime);
+			DeviceProp->ApplyDeviceProperty(PlatformUser);
+		}
+	}
+}
+
 void FActiveForceFeedbackEffect::GetValues(FForceFeedbackValues& Values) const
 {
 	if (ForceFeedbackEffect)
@@ -102,6 +128,10 @@ bool FActiveForceFeedbackEffect::Update(const float DeltaTime, FForceFeedbackVal
 	}
 	
 	GetValues(Values);
-	return true;
-}
 
+	// Set any input device properties associated with this effect
+	const float EvalTime = PlayTime - Duration * FMath::FloorToFloat(PlayTime / Duration);
+	ForceFeedbackEffect->SetDeviceProperties(PlatformUser, DeltaTime, EvalTime);
+
+	return true;
+}
\ No newline at end of file
diff --git a/Engine/Source/Runtime/Engine/Private/GameFramework/InputDeviceProperties.cpp b/Engine/Source/Runtime/Engine/Private/GameFramework/InputDeviceProperties.cpp
new file mode 100644
--- /dev/null
+++ b/Engine/Source/Runtime/Engine/Private/GameFramework/InputDeviceProperties.cpp
@@ -0,0 +1,56 @@
+// Copyright Epic Games, Inc. All Rights Reserved.
+
+#include "GameFramework/InputDeviceProperties.h"
+#include "GenericPlatform/GenericPlatformInputDeviceMapper.h"
+#include "Framework/Application/SlateApplication.h"
+#include "Curves/CurveLinearColor.h"
+
+///////////////////////////////////////////////////////////////////////
+// UInputDeviceProperty
+
+void UInputDeviceProperty::ApplyDeviceProperty(const FPlatformUserId UserId)
+{
+	if (FInputDeviceProperty* RawProp = GetInternalDeviceProperty())
+	{
+		IInputInterface* InputInterface = FSlateApplication::Get().IsInitialized() ? FSlateApplication::Get().GetInputInterface() : nullptr;
+		if (InputInterface)
+		{
+			int32 ControllerId = INDEX_NONE;
+			IPlatformInputDeviceMapper::Get().RemapUserAndDeviceToControllerId(UserId, ControllerId);
+
+			// TODO_BH: Refactor input interface to take an FPlatformUserId directly (UE-158881)
+			InputInterface->SetDeviceProperty(ControllerId, RawProp);
+		}
+	}
+}
+
+///////////////////////////////////////////////////////////////////////
+// UColorInputDeviceProperty
+
+UColorInputDeviceProperty::UColorInputDeviceProperty()
+: UInputDeviceProperty()
+{
+	PropertyName = FInputDeviceLightColorProperty::PropertyName();
+	float MinTime, MaxTime;
+	if (DeviceColorCurve)
+	{
+		DeviceColorCurve->GetTimeRange(MinTime, MaxTime);	
+	}
+	else
+	{
+		PropertyDuration = 1.0f;	
+	}
+}
+
+void UColorInputDeviceProperty::EvaluateDeviceProperty_Implementation(const float DeltaTime, const float Duration)
+{
+	InternalProperty.bEnable = bEnable;
+	
+	FLinearColor CurveColor = DeviceColorCurve->GetLinearColorValue(Duration);
+	InternalProperty.Color = CurveColor.ToFColorSRGB();
+}
+
+FInputDeviceProperty* UColorInputDeviceProperty::GetInternalDeviceProperty()
+{
+	return &InternalProperty;
+}
\ No newline at end of file
diff --git a/Engine/Source/Runtime/Engine/Private/PlayerController.cpp b/Engine/Source/Runtime/Engine/Private/PlayerController.cpp
--- a/Engine/Source/Runtime/Engine/Private/PlayerController.cpp
+++ b/Engine/Source/Runtime/Engine/Private/PlayerController.cpp
@@ -4104,8 +4104,8 @@ void APlayerController::ClientPlayForceFeedback_Internal_Implementation( UForceF
 				}
 			}
 		}
-
-		ActiveForceFeedbackEffects.Emplace(ForceFeedbackEffect, Params);
+		
+		ActiveForceFeedbackEffects.Emplace(ForceFeedbackEffect, Params, GetPlatformUserId());
 
 #if !(UE_BUILD_SHIPPING || UE_BUILD_TEST)
 		ForceFeedbackEffectHistoryEntries.Emplace(ActiveForceFeedbackEffects.Last(), GetWorld()->GetTimeSeconds());
