Fix CanvasText offset causing text to be blurry.

Apply the offset to the draw position before the other calculations to match the original behavior.

diff --git a/Engine/Source/Runtime/Engine/Private/UserInterface/CanvasItem.cpp b/Engine/Source/Runtime/Engine/Private/UserInterface/CanvasItem.cpp
--- a/Engine/Source/Runtime/Engine/Private/UserInterface/CanvasItem.cpp
+++ b/Engine/Source/Runtime/Engine/Private/UserInterface/CanvasItem.cpp
@@ -1074,8 +1074,6 @@ void FCanvasTextItem::DrawStringInternal_OfflineCache(FCanvas* InCanvas, const F
 			}
 			LastTexture = Tex->GetResource();
 
-			const float X		= CurrentPos.X + DrawPos.X;
-			const float Y		= CurrentPos.Y + DrawPos.Y + Char.VerticalOffset * Scale.Y;
 			float SizeX			= Char.USize * Scale.X;
 			const float SizeY	= Char.VSize * Scale.Y;
 			const float U		= Char.StartU * InvTextureSize.X;
@@ -1083,30 +1081,33 @@ void FCanvasTextItem::DrawStringInternal_OfflineCache(FCanvas* InCanvas, const F
 			const float SizeU	= Char.USize * InvTextureSize.X;
 			const float SizeV	= Char.VSize * InvTextureSize.Y;				
 
-			const float Left = X * Depth;
-			const float Top = Y * Depth;
-			const float Right = (X + SizeX) * Depth;
-			const float Bottom = (Y + SizeY) * Depth;
-
 			const auto AddTriangles = [&](const FVector2f& Offset, const FLinearColor& InColor)
 			{
+				const float X = CurrentPos.X + (DrawPos.X + Offset.X);
+				const float Y = CurrentPos.Y + (DrawPos.Y + Offset.Y) + Char.VerticalOffset * Scale.Y;
+
+				const float Left = X * Depth;
+				const float Top = Y * Depth;
+				const float Right = (X + SizeX) * Depth;
+				const float Bottom = (Y + SizeY) * Depth;
+
 				int32 V00 = BatchedElements->AddVertexf(
-					FVector4f(Left + Offset.X, Top + Offset.Y, 0.f, Depth),
+					FVector4f(Left, Top, 0.f, Depth),
 					FVector2f(U, V),
 					InColor,
 					HitProxyId);
 				int32 V10 = BatchedElements->AddVertexf(
-					FVector4f(Right + Offset.X, Top + Offset.Y, 0.0f, Depth),
+					FVector4f(Right, Top, 0.0f, Depth),
 					FVector2f(U + SizeU, V),
 					InColor,
 					HitProxyId);
 				int32 V01 = BatchedElements->AddVertexf(
-					FVector4f(Left + Offset.X, Bottom + Offset.Y, 0.0f, Depth),
+					FVector4f(Left, Bottom, 0.0f, Depth),
 					FVector2f(U, V + SizeV),
 					InColor,
 					HitProxyId);
 				int32 V11 = BatchedElements->AddVertexf(
-					FVector4f(Right + Offset.X, Bottom + Offset.Y, 0.0f, Depth),
+					FVector4f(Right, Bottom, 0.0f, Depth),
 					FVector2f(U + SizeU, V + SizeV),
 					InColor,
 					HitProxyId);
@@ -1266,11 +1267,7 @@ void FCanvasTextItem::DrawStringInternal_RuntimeCache(FCanvas* InCanvas, const F
 			if (!bIsWhitespace)
 			{
 				const float InvBitmapRenderScale = 1.0f / Entry.BitmapRenderScale;
-
-				const float X = DrawPos.X + LineX + (Entry.HorizontalOffset * Scale.X);
-				// Note PosX,PosY is the upper left corner of the bounding box representing the string.  This computes the Y position of the baseline where text will sit
-
-				const float Y = DrawPos.Y + PosY - (Entry.VerticalOffset * Scale.Y) + (((Entry.GlobalDescender * Scale.Y) + ScaledMaxHeight) * InvBitmapRenderScale);
+				
 				const float U = Entry.StartU * InvTextureSizeX;
 				const float V = Entry.StartV * InvTextureSizeY;
 				const float SizeX = Entry.USize * Scale.X * Entry.BitmapRenderScale;
@@ -1278,30 +1275,34 @@ void FCanvasTextItem::DrawStringInternal_RuntimeCache(FCanvas* InCanvas, const F
 				const float SizeU = Entry.USize * InvTextureSizeX;
 				const float SizeV = Entry.VSize * InvTextureSizeY;
 
-				const float Left = X * Depth;
-				const float Top = Y * Depth;
-				const float Right = (X + SizeX) * Depth;
-				const float Bottom = (Y + SizeY) * Depth;
-
 				const auto AddTriangles = [&](const FVector2f& Offset, const FLinearColor& InColor)
 				{
+					// Note PosX,PosY is the upper left corner of the bounding box representing the string.  This computes the Y position of the baseline where text will sit
+					const float X = (DrawPos.X + Offset.X) + LineX + (Entry.HorizontalOffset * Scale.X);
+					const float Y = (DrawPos.Y + Offset.Y) + PosY - (Entry.VerticalOffset * Scale.Y) + (((Entry.GlobalDescender * Scale.Y) + ScaledMaxHeight) * InvBitmapRenderScale);
+
+					const float Left = X * Depth;
+					const float Top = Y * Depth;
+					const float Right = (X + SizeX) * Depth;
+					const float Bottom = (Y + SizeY) * Depth;
+
 					int32 V00 = BatchedElements->AddVertexf(
-						FVector4f(Left + Offset.X, Top + Offset.Y, 0.f, Depth),
+						FVector4f(Left, Top, 0.f, Depth),
 						FVector2f(U, V),
 						InColor,
 						HitProxyId);
 					int32 V10 = BatchedElements->AddVertexf(
-						FVector4f(Right + Offset.X, Top + Offset.Y, 0.0f, Depth),
+						FVector4f(Right, Top, 0.0f, Depth),
 						FVector2f(U + SizeU, V),
 						InColor,
 						HitProxyId);
 					int32 V01 = BatchedElements->AddVertexf(
-						FVector4f(Left + Offset.X, Bottom + Offset.Y, 0.0f, Depth),
+						FVector4f(Left, Bottom, 0.0f, Depth),
 						FVector2f(U, V + SizeV),
 						InColor,
 						HitProxyId);
 					int32 V11 = BatchedElements->AddVertexf(
-						FVector4f(Right + Offset.X, Bottom + Offset.Y, 0.0f, Depth),
+						FVector4f(Right, Bottom, 0.0f, Depth),
 						FVector2f(U + SizeU, V + SizeV),
 						InColor,
 						HitProxyId);
@@ -1430,10 +1431,6 @@ void FCanvasShapedTextItem::DrawStringInternal(FCanvas* InCanvas, const FVector2
 				const float BitmapRenderScale = GlyphToRender.GetBitmapRenderScale();
 				const float InvBitmapRenderScale = 1.0f / BitmapRenderScale;
 
-				const float X = DrawPos.X + LineX + (GlyphAtlasData.HorizontalOffset * Scale.X) + (GlyphToRender.XOffset * Scale.X);
-				// Note PosX,PosY is the upper left corner of the bounding box representing the string.  This computes the Y position of the baseline where text will sit
-
-				const float Y = DrawPos.Y + PosY - (GlyphAtlasData.VerticalOffset * Scale.Y) + (GlyphToRender.YOffset * Scale.Y) + ((ScaledBaseline + ScaledMaxHeight) * InvBitmapRenderScale);
 				const float U = GlyphAtlasData.StartU * InvTextureSizeX;
 				const float V = GlyphAtlasData.StartV * InvTextureSizeY;
 				const float SizeX = GlyphAtlasData.USize * Scale.X * BitmapRenderScale;
@@ -1441,30 +1438,34 @@ void FCanvasShapedTextItem::DrawStringInternal(FCanvas* InCanvas, const FVector2
 				const float SizeU = GlyphAtlasData.USize * InvTextureSizeX;
 				const float SizeV = GlyphAtlasData.VSize * InvTextureSizeY;
 
-				const float Left = X * Depth;
-				const float Top = Y * Depth;
-				const float Right = (X + SizeX) * Depth;
-				const float Bottom = (Y + SizeY) * Depth;
-
 				const auto AddTriangles = [&](const FVector2f& Offset, const FLinearColor& InColor)
 				{
+					// Note PosX,PosY is the upper left corner of the bounding box representing the string.  This computes the Y position of the baseline where text will sit
+					const float X = (DrawPos.X + Offset.X) + LineX + (GlyphAtlasData.HorizontalOffset * Scale.X) + (GlyphToRender.XOffset * Scale.X);
+					const float Y = (DrawPos.Y + Offset.Y) + PosY - (GlyphAtlasData.VerticalOffset * Scale.Y) + (GlyphToRender.YOffset * Scale.Y) + ((ScaledBaseline + ScaledMaxHeight) * InvBitmapRenderScale);
+
+					const float Left = X * Depth;
+					const float Top = Y * Depth;
+					const float Right = (X + SizeX) * Depth;
+					const float Bottom = (Y + SizeY) * Depth;
+
 					int32 V00 = BatchedElements->AddVertexf(
-						FVector4f(Left + Offset.X, Top + Offset.Y, 0.f, Depth),
+						FVector4f(Left, Top, 0.f, Depth),
 						FVector2f(U, V),
 						InColor,
 						HitProxyId);
 					int32 V10 = BatchedElements->AddVertexf(
-						FVector4f(Right + Offset.X, Top + Offset.Y, 0.0f, Depth),
+						FVector4f(Right, Top, 0.0f, Depth),
 						FVector2f(U + SizeU, V),
 						InColor,
 						HitProxyId);
 					int32 V01 = BatchedElements->AddVertexf(
-						FVector4f(Left + Offset.X, Bottom + Offset.Y, 0.0f, Depth),
+						FVector4f(Left, Bottom, 0.0f, Depth),
 						FVector2f(U, V + SizeV),
 						InColor,
 						HitProxyId);
 					int32 V11 = BatchedElements->AddVertexf(
-						FVector4f(Right + Offset.X, Bottom + Offset.Y, 0.0f, Depth),
+						FVector4f(Right, Bottom, 0.0f, Depth),
 						FVector2f(U + SizeU, V + SizeV),
 						InColor,
 						HitProxyId);
