Multiple improvements to replay scrub respawning of placed actors

diff --git a/Engine/Source/Runtime/Engine/Classes/Engine/DemoNetDriver.h b/Engine/Source/Runtime/Engine/Classes/Engine/DemoNetDriver.h
--- a/Engine/Source/Runtime/Engine/Classes/Engine/DemoNetDriver.h
+++ b/Engine/Source/Runtime/Engine/Classes/Engine/DemoNetDriver.h
@@ -213,6 +213,7 @@ public:
 	void RespawnNecessaryNetStartupActors(TArray<AActor*>& SpawnedActors, ULevel* Level = nullptr);
 
 private:
+	void RestoreComponentState(UActorComponent* ActorComp, FRollbackNetStartupActorInfo& RollbackActor);
 
 	void InitDefaults();
 
diff --git a/Engine/Source/Runtime/Engine/Private/DemoNetDriver.cpp b/Engine/Source/Runtime/Engine/Private/DemoNetDriver.cpp
--- a/Engine/Source/Runtime/Engine/Private/DemoNetDriver.cpp
+++ b/Engine/Source/Runtime/Engine/Private/DemoNetDriver.cpp
@@ -3272,6 +3272,24 @@ bool UDemoNetDriver::SetExternalDataForObject(UObject* OwningObject, const uint8
 	return false;
 }
 
+void UDemoNetDriver::RestoreComponentState(UActorComponent* ActorComp, FRollbackNetStartupActorInfo& RollbackActor)
+{
+	TSharedPtr<FRepLayout> SubObjLayout = GetObjectClassRepLayout(ActorComp->GetClass());
+	if (SubObjLayout.IsValid())
+	{
+		TSharedPtr<FRepState> RepState = RollbackActor.SubObjRepState.FindRef(ActorComp->GetFullName());
+		FReceivingRepState* SubObjReceivingRepState = RepState.IsValid() ? RepState->GetReceivingRepState() : nullptr;
+
+		if (SubObjReceivingRepState)
+		{
+			FRepObjectDataBuffer ActorCompData(ActorComp);
+			FConstRepShadowDataBuffer ShadowData(SubObjReceivingRepState->StaticBuffer.GetData());
+
+			SubObjLayout->DiffStableProperties(&SubObjReceivingRepState->RepNotifies, nullptr, ActorCompData, ShadowData);
+		}
+	}
+}
+
 void UDemoNetDriver::RespawnNecessaryNetStartupActors(TArray<AActor*>& SpawnedActors, ULevel* Level /* = nullptr */)
 {
 	QUICK_SCOPE_CYCLE_COUNTER(STAT_RespawnNecessaryNetStartupActors);
@@ -3338,7 +3356,7 @@ void UDemoNetDriver::RespawnNecessaryNetStartupActors(TArray<AActor*>& SpawnedAc
 
 		const FTransform SpawnTransform = FTransform(RollbackActor.Rotation, RollbackActor.Location, RollbackActor.Scale3D);
 
-		AActor* Actor = World->SpawnActorAbsolute(RollbackActor.Archetype->GetClass(), SpawnTransform, SpawnInfo);
+		AActor* Actor = World->SpawnActor(RollbackActor.Archetype->GetClass(), &SpawnTransform, SpawnInfo);
 		if (Actor)
 		{
 			if (!ensure( Actor->GetFullName() == RollbackIt.Key()))
@@ -3346,13 +3364,13 @@ void UDemoNetDriver::RespawnNecessaryNetStartupActors(TArray<AActor*>& SpawnedAc
 				UE_LOG(LogDemo, Log, TEXT("RespawnNecessaryNetStartupActors: NetStartupRollbackActor name doesn't match original: %s, %s"), *Actor->GetFullName(), *RollbackIt.Key());
 			}
 
-			bool bSanityCheckReferences = true;
+			bool bValidObjReferences = true;
 
 			for (UObject* ObjRef : RollbackActor.ObjReferences)
 			{
-				if (ObjRef == nullptr)
+				if (!IsValid(ObjRef))
 				{
-					bSanityCheckReferences = false;
+					bValidObjReferences = false;
 					UE_LOG(LogDemo, Warning, TEXT("RespawnNecessaryNetStartupActors: Rollback actor reference was gc'd, skipping state restore: %s"), *GetFullNameSafe(Actor));
 					break;
 				}
@@ -3361,7 +3379,8 @@ void UDemoNetDriver::RespawnNecessaryNetStartupActors(TArray<AActor*>& SpawnedAc
 			TSharedPtr<FRepLayout> RepLayout = GetObjectClassRepLayout(Actor->GetClass());
 			FReceivingRepState* ReceivingRepState = RollbackActor.RepState.IsValid() ? RollbackActor.RepState->GetReceivingRepState() : nullptr;
 
-			if (RepLayout.IsValid() && ReceivingRepState && bSanityCheckReferences)
+			// Restore saved actor state
+			if (RepLayout.IsValid() && ReceivingRepState && bValidObjReferences)
 			{
 				const ENetRole SavedRole = Actor->GetLocalRole();
 
@@ -3382,10 +3401,49 @@ void UDemoNetDriver::RespawnNecessaryNetStartupActors(TArray<AActor*>& SpawnedAc
 				Actor->SwapRoles();
 			}
 
-			UGameplayStatics::FinishSpawningActor(Actor, SpawnTransform);
+			TSet<UActorComponent*> DiffedComponents;
 
+			if (bValidObjReferences)
+			{
+				// Restore replicated component state for any objects that exist prior to construction (default subobjects)
+				for (UActorComponent* ActorComp : Actor->GetComponents())
+				{
+					if (ActorComp)
+					{
+						RestoreComponentState(ActorComp, RollbackActor);
+
+						DiffedComponents.Add(ActorComp);
+					}
+				}
+			}
+			
+			// Update transforms based on restored state
+			Actor->UpdateComponentTransforms();
+
+			// Finish spawning/construction
+			Actor->FinishSpawning(SpawnTransform, true);
+
+			if (bValidObjReferences)
+			{
+				// Restore replicated component state of anything new (created during construction)
+				for (UActorComponent* ActorComp : Actor->GetComponents())
+				{
+					// Could have been created by FinishSpawning (construction script)
+					if (ActorComp && !DiffedComponents.Contains(ActorComp))
+					{
+						RestoreComponentState(ActorComp, RollbackActor);
+					}
+				}
+			}
+
+			// Update transforms based on restored state, and dirty render state
+			Actor->UpdateComponentTransforms();
+			Actor->MarkComponentsRenderStateDirty();
+
+			// BeginPlay
 			Actor->PostNetInit();
 
+			// Call actor rep notifies
 			if (RepLayout.IsValid() && ReceivingRepState)
 			{
 				if (ReceivingRepState->RepNotifies.Num() > 0)
@@ -3396,24 +3454,20 @@ void UDemoNetDriver::RespawnNecessaryNetStartupActors(TArray<AActor*>& SpawnedAc
 				}
 			}
 
-			for (UActorComponent* ActorComp : Actor->GetComponents())
+			if (bValidObjReferences)
 			{
-				if (ActorComp)
+				// Call component rep notifies
+				for (UActorComponent* ActorComp : Actor->GetComponents())
 				{
-					TSharedPtr<FRepLayout> SubObjLayout = GetObjectClassRepLayout(ActorComp->GetClass());
-					if (SubObjLayout.IsValid() && bSanityCheckReferences)
+					if (ActorComp)
 					{
-						TSharedPtr<FRepState> RepState = RollbackActor.SubObjRepState.FindRef(ActorComp->GetFullName());
-						FReceivingRepState* SubObjReceivingRepState = RepState.IsValid() ? RepState->GetReceivingRepState() : nullptr;
-
-						if (SubObjReceivingRepState)
+						TSharedPtr<FRepLayout> SubObjLayout = GetObjectClassRepLayout(ActorComp->GetClass());
+						if (SubObjLayout.IsValid())
 						{
-							FRepObjectDataBuffer ActorCompData(ActorComp);
-							FConstRepShadowDataBuffer ShadowData(SubObjReceivingRepState->StaticBuffer.GetData());
-
-							SubObjLayout->DiffStableProperties(&SubObjReceivingRepState->RepNotifies, nullptr, ActorCompData, ShadowData);
+							TSharedPtr<FRepState> RepState = RollbackActor.SubObjRepState.FindRef(ActorComp->GetFullName());
+							FReceivingRepState* SubObjReceivingRepState = RepState.IsValid() ? RepState->GetReceivingRepState() : nullptr;
 
-							if (SubObjReceivingRepState->RepNotifies.Num() > 0)
+							if (SubObjReceivingRepState && SubObjReceivingRepState->RepNotifies.Num() > 0)
 							{
 								SubObjLayout->CallRepNotifies(SubObjReceivingRepState, ActorComp);
 
@@ -5212,6 +5266,7 @@ void UDemoNetDriver::QueueNetStartupActorForRollbackViaDeletion(AActor* Actor)
 
 	if (GDemoSaveRollbackActorState != 0)
 	{
+		// Save actor state
 		{
 			TSharedPtr<FObjectReplicator> NewReplicator = MakeShared<FObjectReplicator>();
 			NewReplicator->InitWithObject(Actor->GetArchetype(), ServerConnection, false);
@@ -5229,6 +5284,7 @@ void UDemoNetDriver::QueueNetStartupActorForRollbackViaDeletion(AActor* Actor)
 			}
 		}
 
+		// Save component state
 		for (UActorComponent* ActorComp : Actor->GetComponents())
 		{
 			if (ActorComp)
