Added an option in animation blueprints to share linked anim instance across group and subgraphs instead of having one layer instance per group / linked anim instance

#tests Test with multiple players and AI in PIE

diff --git a/Engine/Source/Editor/AnimGraph/Private/AnimBlueprintExtension.cpp b/Engine/Source/Editor/AnimGraph/Private/AnimBlueprintExtension.cpp
--- a/Engine/Source/Editor/AnimGraph/Private/AnimBlueprintExtension.cpp
+++ b/Engine/Source/Editor/AnimGraph/Private/AnimBlueprintExtension.cpp
@@ -11,6 +11,7 @@
 #include "AnimBlueprintExtension_PropertyAccess.h"
 #include "Kismet2/BlueprintEditorUtils.h"
 #include "AnimBlueprintExtension_NodeRelevancy.h"
+#include "AnimBlueprintExtension_SharedLinkedAnimLayers.h"
 #include "AnimBlueprintExtension_Tag.h"
 
 // Set used to refresh extensions. Checks that an extension has a reference from an anim node for each refresh.
@@ -79,6 +80,11 @@ void UAnimBlueprintExtension::RequestExtensionsForNode(UAnimGraphNode_Base* InAn
             UAnimBlueprintExtension_PropertyAccess::StaticClass()
         };
 
+		if (AnimBlueprint->bEnableLinkedAnimLayerInstanceSharing)
+		{
+			ExtensionClasses.Add(UAnimBlueprintExtension_SharedLinkedAnimLayers::StaticClass());
+		}
+
 		// As this can be called when we have not regenerated the skeleton class, we need to be less conservative
 		// and request extensions if it *looks* like we have a valid reference, rather than a verified concrete ref.
 		// Later validation code will do the actual validation for us anyways and fail compilation if the reference is
diff --git a/Engine/Source/Editor/AnimGraph/Private/AnimBlueprintExtension_SharedLinkedAnimLayers.h b/Engine/Source/Editor/AnimGraph/Private/AnimBlueprintExtension_SharedLinkedAnimLayers.h
new file mode 100644
--- /dev/null
+++ b/Engine/Source/Editor/AnimGraph/Private/AnimBlueprintExtension_SharedLinkedAnimLayers.h
@@ -0,0 +1,18 @@
+// Copyright Epic Games, Inc. All Rights Reserved.
+
+#pragma once
+
+#include "CoreMinimal.h"
+#include "AnimBlueprintExtension.h"
+#include "Animation/AnimSubSystem_SharedLinkedAnimLayers.h"
+#include "AnimBlueprintExtension_SharedLinkedAnimLayers.generated.h"
+
+UCLASS(MinimalAPI)
+class UAnimBlueprintExtension_SharedLinkedAnimLayers : public UAnimBlueprintExtension
+{
+	GENERATED_BODY()
+
+private:
+	UPROPERTY()
+	FAnimSubsystem_SharedLinkedAnimLayers Subsystem;
+};
\ No newline at end of file
diff --git a/Engine/Source/Runtime/Engine/Classes/Animation/AnimBlueprint.h b/Engine/Source/Runtime/Engine/Classes/Animation/AnimBlueprint.h
--- a/Engine/Source/Runtime/Engine/Classes/Animation/AnimBlueprint.h
+++ b/Engine/Source/Runtime/Engine/Classes/Animation/AnimBlueprint.h
@@ -113,6 +113,11 @@ class ENGINE_API UAnimBlueprint : public UBlueprint, public IInterface_PreviewMe
 	UPROPERTY(EditAnywhere, Category = Optimization)
 	bool bWarnAboutBlueprintUsage;
 
+	/** If true, linked animation layers will be instantiated only once per AnimClass instead of once per AnimInstance, AnimClass and AnimGroup.
+	Extra instances will be created if two or more active anim graph override the same layer Function */
+	UPROPERTY(EditDefaultsOnly, Category = Optimization)
+	uint8 bEnableLinkedAnimLayerInstanceSharing : 1;
+
 	// @todo document
 	class UAnimBlueprintGeneratedClass* GetAnimBlueprintGeneratedClass() const;
 
@@ -181,7 +186,8 @@ class ENGINE_API UAnimBlueprint : public UBlueprint, public IInterface_PreviewMe
 	virtual bool SupportsDelegates() const override;
 	virtual bool SupportsMacros() const override;
 	virtual bool AllowFunctionOverride(const UFunction* const InFunction) const override;
-
+	virtual void PostEditChangeProperty(FPropertyChangedEvent& PropertyChangedEvent) override;
+	
 protected:
 	// Broadcast when an override is changed, allowing derived blueprints to be updated
 	FOnOverrideChangedMulticaster OnOverrideChanged;
diff --git a/Engine/Source/Runtime/Engine/Classes/Animation/AnimNode_LinkedAnimGraph.h b/Engine/Source/Runtime/Engine/Classes/Animation/AnimNode_LinkedAnimGraph.h
--- a/Engine/Source/Runtime/Engine/Classes/Animation/AnimNode_LinkedAnimGraph.h
+++ b/Engine/Source/Runtime/Engine/Classes/Animation/AnimNode_LinkedAnimGraph.h
@@ -113,6 +113,9 @@ protected:
 	// Shutdown the currently running instance
 	void TeardownInstance(const UAnimInstance* InOwningAnimInstance);
 
+	// Check if the currently linked instance can be teared down
+	virtual bool CanTeardownLinkedInstance(const UAnimInstance* LinkedInstance) const {return true;}
+
 	// FAnimNode_CustomProperty interface
 	virtual UClass* GetTargetClass() const override 
 	{
diff --git a/Engine/Source/Runtime/Engine/Classes/Animation/AnimNode_LinkedAnimLayer.h b/Engine/Source/Runtime/Engine/Classes/Animation/AnimNode_LinkedAnimLayer.h
--- a/Engine/Source/Runtime/Engine/Classes/Animation/AnimNode_LinkedAnimLayer.h
+++ b/Engine/Source/Runtime/Engine/Classes/Animation/AnimNode_LinkedAnimLayer.h
@@ -40,6 +40,7 @@ public:
 	/** FAnimNode_Base interface */
 	virtual void OnInitializeAnimInstance(const FAnimInstanceProxy* InProxy, const UAnimInstance* InAnimInstance) override;
 	virtual bool NeedsOnInitializeAnimInstance() const override { return true; }
+	void OnUninitializeAnimInstance(UAnimInstance* InAnimInstance);
 
 	/** FAnimNode_CustomProperty interface */
 	virtual void InitializeProperties(const UObject* InSourceInstance, UClass* InTargetClass) override;
@@ -64,6 +65,11 @@ protected:
 	// Initialize the source properties to copy from
 	void InitializeSourceProperties(const UAnimInstance* InAnimInstance);
 
+	virtual bool CanTeardownLinkedInstance(const UAnimInstance* LinkedInstance) const override;
+
+	// Cleanup Shared LinkedLayers Data associated with InPreviousTargetInstance
+	void CleanupSharedLinkedLayersData(const UAnimInstance* InOwningAnimInstance, UAnimInstance* InPreviousTargetInstance);
+	
 #if WITH_EDITOR
 	// Event fired when the instance we are running has changed
 	FSimpleMulticastDelegate OnInstanceChangedEvent;
diff --git a/Engine/Source/Runtime/Engine/Private/AnimBlueprint.cpp b/Engine/Source/Runtime/Engine/Private/AnimBlueprint.cpp
--- a/Engine/Source/Runtime/Engine/Private/AnimBlueprint.cpp
+++ b/Engine/Source/Runtime/Engine/Private/AnimBlueprint.cpp
@@ -261,6 +261,19 @@ bool UAnimBlueprint::AllowFunctionOverride(const UFunction* const InFunction) co
 	}
 }
 
+void UAnimBlueprint::PostEditChangeProperty(struct FPropertyChangedEvent& PropertyChangedEvent)
+{
+	Super::PostEditChangeProperty(PropertyChangedEvent);
+
+	const FName PropertyName = (PropertyChangedEvent.Property != NULL) ? PropertyChangedEvent.Property->GetFName() : NAME_None;
+
+	// Make sure we add / remove linked anim layer sharing blueprint extension
+	if (PropertyName == GET_MEMBER_NAME_CHECKED(UAnimBlueprint, bEnableLinkedAnimLayerInstanceSharing))
+	{
+		bRefreshExtensions = true;
+	}
+}
+
 #endif
 
 USkeletalMesh* UAnimBlueprint::GetPreviewMesh(bool bFindIfNotSet/*=false*/)
diff --git a/Engine/Source/Runtime/Engine/Private/Animation/AnimInstance.cpp b/Engine/Source/Runtime/Engine/Private/Animation/AnimInstance.cpp
--- a/Engine/Source/Runtime/Engine/Private/Animation/AnimInstance.cpp
+++ b/Engine/Source/Runtime/Engine/Private/Animation/AnimInstance.cpp
@@ -35,6 +35,7 @@
 #include "Animation/AnimSubsystem_Tag.h"
 #include "Animation/AnimStateMachineTypes.h"
 #include "Animation/ActiveMontageInstanceScope.h"
+#include "Animation/AnimSubsystem_SharedLinkedAnimLayers.h"
 
 #include UE_INLINE_GENERATED_CPP_BY_NAME(AnimInstance)
 
@@ -340,6 +341,26 @@ void UAnimInstance::UninitializeAnimation()
 
 	SlotGroupInertializationRequestMap.Reset();
 	
+	// Cleanup layer nodes
+	if (IAnimClassInterface* AnimBlueprintClass = IAnimClassInterface::GetFromClass(GetClass()))
+	{
+		for (FStructProperty* LayerNodeProperty : AnimBlueprintClass->GetLinkedAnimLayerNodeProperties())
+		{
+			if (FAnimNode_LinkedAnimLayer* Layer = LayerNodeProperty->ContainerPtrToValuePtr<FAnimNode_LinkedAnimLayer>(this))
+			{
+				Layer->OnUninitializeAnimInstance(this);
+			}
+		}
+	}
+	
+	// Cleanup shared layers data
+	if (FAnimSubsystem_SharedLinkedAnimLayers* SharedLinkedAnimLayers = FindSubsystem<FAnimSubsystem_SharedLinkedAnimLayers>())
+	{
+		// Reset shared linked instances when the main instance in uninitialized. 
+		// This is required in part for SkeletalMeshComponent::OnUnregister so that the linked instances array isn't modified as we iterate on it to unitialize them. (see FAnimNode_LinkedLayer::CleanupSharedLinkedLayersData)
+		SharedLinkedAnimLayers->Reset();
+	}
+
 	TRACE_OBJECT_LIFETIME_END(this);
 }
 
@@ -2990,7 +3011,82 @@ void UAnimInstance::PerformLinkedLayerOverlayOperation(TSubclassOf<UAnimInstance
 		{
 			for (TPair<FName, TArray<FAnimNode_LinkedAnimLayer*, TInlineAllocator<4>>> LayerPair : ClassLayerNodesToSet.Value)
 			{
-				if (LayerPair.Key == NAME_None)
+				if (FAnimSubsystem_SharedLinkedAnimLayers* SharedLinkedAnimLayers = MeshComp->GetAnimInstance()->FindSubsystem<FAnimSubsystem_SharedLinkedAnimLayers>())
+				{
+					// Shared instances path
+					UClass* ClassToSet = ClassLayerNodesToSet.Key;
+					for (FAnimNode_LinkedAnimLayer* LayerNode : LayerPair.Value)
+					{
+						// Disallow setting the same class as this instance, which would create infinite recursion
+						if (ClassToSet != nullptr && ClassToSet != GetClass())
+						{
+							UAnimInstance* TargetInstance = LayerNode->GetTargetInstance<UAnimInstance>();
+
+							// Skip setting if the class is the same
+							if (TargetInstance == nullptr || ClassToSet != TargetInstance->GetClass())
+							{
+								FLinkedAnimLayerClassData& ClassesData = SharedLinkedAnimLayers->FindOrAddClassData(ClassToSet);
+
+								const FName FunctionToLink = LayerNode->GetDynamicLinkFunctionName();
+
+								// Try to find an existing instance whose function is available for linking
+								FLinkedAnimLayerInstanceData* LayerAnimInstanceData = ClassesData.FindInstanceForLinking(FunctionToLink);
+
+								bool bIsNewInstance;
+								UAnimInstance* LinkedInstance = nullptr;
+								// Use existing instance if possible
+								if (LayerAnimInstanceData)
+								{
+									bIsNewInstance = false;
+									LinkedInstance = LayerAnimInstanceData->Instance;
+								}
+								// Create the new linked instance if needed
+								else
+								{
+									bIsNewInstance = true;
+									
+									// Init new instance
+									LinkedInstance = NewObject<UAnimInstance>(MeshComp, ClassToSet);
+									LinkedInstance->bCreatedByLinkedAnimGraph = true;
+									LinkedInstance->InitializeAnimation();
+									
+									// Insert new instance data
+									LayerAnimInstanceData = &ClassesData.AddInstance(LinkedInstance);
+									
+									// Unlink any layer nodes in the new linked instance, as they may have been hooked up to self in InitializeAnimation above.
+									UnlinkLayerNodesInInstance(LinkedInstance, LayerPair.Value);
+								}
+
+								// Mark function as linked
+								LayerAnimInstanceData->AddLinkedFunction(FunctionToLink, this);
+								LayerNode->SetLinkedLayerInstance(this, LinkedInstance);
+
+								// Propagate notify flags. If any nodes have this set then we need to propagate to the group.
+								LinkedInstance->bPropagateNotifiesToLinkedInstances |= LayerNode->bPropagateNotifiesToLinkedInstances;
+								LinkedInstance->bReceiveNotifiesFromLinkedInstances |= LayerNode->bReceiveNotifiesFromLinkedInstances;
+
+								if (!bInDeferSubGraphInitialization)
+								{
+									FAnimInstanceProxy& ThisProxy = GetProxyOnAnyThread<FAnimInstanceProxy>();
+									FAnimInstanceProxy& LinkedProxy = LinkedInstance->GetProxyOnAnyThread<FAnimInstanceProxy>();
+
+									// Initialize the correct parts of the linked instance
+									InitializeAndCacheBonesForLinkedRoot(LayerNode, ThisProxy, LinkedInstance, LinkedProxy);
+								}
+
+								if (bIsNewInstance)
+								{
+									MeshComp->GetLinkedAnimInstances().Add(LinkedInstance);
+								}
+							}
+						}
+						else
+						{
+							LayerNode->SetLinkedLayerInstance(this, nullptr);
+						}
+					}
+				}
+				else if (LayerPair.Key == NAME_None)
 				{
 					// Ungrouped path - each layer gets a separate instance
 					for (FAnimNode_LinkedAnimLayer* LayerNode : LayerPair.Value)
diff --git a/Engine/Source/Runtime/Engine/Private/Animation/AnimNode_LinkedAnimGraph.cpp b/Engine/Source/Runtime/Engine/Private/Animation/AnimNode_LinkedAnimGraph.cpp
--- a/Engine/Source/Runtime/Engine/Private/Animation/AnimNode_LinkedAnimGraph.cpp
+++ b/Engine/Source/Runtime/Engine/Private/Animation/AnimNode_LinkedAnimGraph.cpp
@@ -231,12 +231,15 @@ void FAnimNode_LinkedAnimGraph::TeardownInstance(const UAnimInstance* InOwningAn
 		// Never delete the owning animation instance
 		if (InstanceToRun != InOwningAnimInstance)
 		{
-			USkeletalMeshComponent* MeshComp = InOwningAnimInstance->GetSkelMeshComponent();
-			check(MeshComp);
-			MeshComp->GetLinkedAnimInstances().Remove(InstanceToRun);
-			// Only call UninitializeAnimation if we are not the owning anim instance
-			InstanceToRun->UninitializeAnimation();
-			InstanceToRun->MarkAsGarbage();
+			if (CanTeardownLinkedInstance(InstanceToRun))
+			{
+				USkeletalMeshComponent* MeshComp = InOwningAnimInstance->GetSkelMeshComponent();
+				check(MeshComp);
+				MeshComp->GetLinkedAnimInstances().Remove(InstanceToRun);
+				// Only call UninitializeAnimation if we are not the owning anim instance
+				InstanceToRun->UninitializeAnimation();
+				InstanceToRun->MarkAsGarbage();
+			}
 		}
 
 		InstanceToRun = nullptr;
diff --git a/Engine/Source/Runtime/Engine/Private/Animation/AnimNode_LinkedAnimLayer.cpp b/Engine/Source/Runtime/Engine/Private/Animation/AnimNode_LinkedAnimLayer.cpp
--- a/Engine/Source/Runtime/Engine/Private/Animation/AnimNode_LinkedAnimLayer.cpp
+++ b/Engine/Source/Runtime/Engine/Private/Animation/AnimNode_LinkedAnimLayer.cpp
@@ -2,6 +2,7 @@
 
 #include "Animation/AnimNode_LinkedAnimLayer.h"
 #include "Animation/AnimClassInterface.h"
+#include "Animation/AnimSubsystem_SharedLinkedAnimLayers.h"
 #include "Animation/AnimInstanceProxy.h"
 #include "Components/SkeletalMeshComponent.h"
 
@@ -55,6 +56,26 @@ void FAnimNode_LinkedAnimLayer::OnInitializeAnimInstance(const FAnimInstanceProx
 	}
 }
 
+void FAnimNode_LinkedAnimLayer::OnUninitializeAnimInstance(UAnimInstance* InOwningAnimInstance)
+{
+	// Owning anim instance is being destroyed, clean dynamic layer data 
+	if (UAnimInstance* CurrentTarget = GetTargetInstance<UAnimInstance>())
+	{
+		USkeletalMeshComponent* MeshComp = InOwningAnimInstance->GetSkelMeshComponent();
+		check(MeshComp && MeshComp->GetAnimInstance());
+		if (FAnimSubsystem_SharedLinkedAnimLayers* SharedLinkedAnimLayers = MeshComp->GetAnimInstance()->FindSubsystem<FAnimSubsystem_SharedLinkedAnimLayers>())
+		{
+			// If target instance is a shared instance, unlink it when owner uninitialize
+			if (FLinkedAnimLayerInstanceData* PreviousLinkedLayerInstanceData = SharedLinkedAnimLayers->FindInstanceData(CurrentTarget))
+			{
+				DynamicUnlink(InOwningAnimInstance);
+				SetTargetInstance(nullptr);
+				CleanupSharedLinkedLayersData(InOwningAnimInstance, CurrentTarget);
+			}
+		}
+	}
+}
+
 void FAnimNode_LinkedAnimLayer::InitializeSelfLayer(const UAnimInstance* SelfAnimInstance)
 {
 	UAnimInstance* CurrentTarget = GetTargetInstance<UAnimInstance>();
@@ -72,10 +93,13 @@ void FAnimNode_LinkedAnimLayer::InitializeSelfLayer(const UAnimInstance* SelfAni
 	// Switch from dynamic external to internal, kill old instance
 	if (CurrentTarget && CurrentTarget != SelfAnimInstance)
 	{
-		CurrentTarget->UninitializeAnimation();
-		MeshComp->GetLinkedAnimInstances().Remove(CurrentTarget);
-		CurrentTarget->MarkAsGarbage();
-		CurrentTarget = nullptr;
+		if (CanTeardownLinkedInstance(CurrentTarget))
+		{
+			CurrentTarget->UninitializeAnimation();
+			MeshComp->GetLinkedAnimInstances().Remove(CurrentTarget);
+			CurrentTarget->MarkAsGarbage();
+			CurrentTarget = nullptr;
+		}
 	}
 
 	SetTargetInstance(const_cast<UAnimInstance*>(SelfAnimInstance));
@@ -93,6 +117,8 @@ void FAnimNode_LinkedAnimLayer::InitializeSelfLayer(const UAnimInstance* SelfAni
 
 void FAnimNode_LinkedAnimLayer::SetLinkedLayerInstance(const UAnimInstance* InOwningAnimInstance, UAnimInstance* InNewLinkedInstance)
 {
+	UAnimInstance* PreviousTargetInstance = GetTargetInstance<UAnimInstance>();
+
 	// Reseting to running as a self-layer, in case it is applicable
 	if ((Interface.Get() == nullptr || InstanceClass.Get() == nullptr) && (InNewLinkedInstance == nullptr))
 	{
@@ -103,6 +129,11 @@ void FAnimNode_LinkedAnimLayer::SetLinkedLayerInstance(const UAnimInstance* InOw
 		ReinitializeLinkedAnimInstance(InOwningAnimInstance, InNewLinkedInstance);
 	}
 
+	if (PreviousTargetInstance != GetTargetInstance<UAnimInstance>())
+	{
+		CleanupSharedLinkedLayersData(InOwningAnimInstance, PreviousTargetInstance);
+	}
+
 #if WITH_EDITOR
 	OnInstanceChangedEvent.Broadcast();
 #endif
@@ -148,6 +179,47 @@ void FAnimNode_LinkedAnimLayer::InitializeProperties(const UObject* InSourceInst
 	}
 }
 
+bool FAnimNode_LinkedAnimLayer::CanTeardownLinkedInstance(const UAnimInstance* LinkedInstance) const
+{
+	// Don't teardown instance that still have function linked to active shared instances
+	USkeletalMeshComponent* MeshComp = LinkedInstance->GetSkelMeshComponent();
+	check(MeshComp && MeshComp->GetAnimInstance());
+	if (FAnimSubsystem_SharedLinkedAnimLayers* SharedLinkedAnimLayers = MeshComp->GetAnimInstance()->FindSubsystem<FAnimSubsystem_SharedLinkedAnimLayers>())
+	{
+		return SharedLinkedAnimLayers->FindInstanceData(LinkedInstance) == nullptr;
+	}
+	return true; 
+}
+
+void FAnimNode_LinkedAnimLayer::CleanupSharedLinkedLayersData(const UAnimInstance* InOwningAnimInstance, UAnimInstance* InPreviousTargetInstance)
+{
+	if (InPreviousTargetInstance)
+	{
+		USkeletalMeshComponent* MeshComp = InPreviousTargetInstance->GetSkelMeshComponent();
+		check(MeshComp && MeshComp->GetAnimInstance());
+		if (FAnimSubsystem_SharedLinkedAnimLayers* SharedLinkedAnimLayers = MeshComp->GetAnimInstance()->FindSubsystem<FAnimSubsystem_SharedLinkedAnimLayers>() )
+		{
+			if (FLinkedAnimLayerInstanceData* PreviousLinkedLayerInstanceData = SharedLinkedAnimLayers->FindInstanceData(InPreviousTargetInstance))
+			{
+				const FName FunctionToLink = GetDynamicLinkFunctionName();
+
+				// The owning anim instance should never be part of the dynamic linked instance data
+				check(InOwningAnimInstance != InPreviousTargetInstance);
+
+				PreviousLinkedLayerInstanceData->RemoveLinkedFunction(FunctionToLink);
+				// Instance isn't used anymore, delete it
+				if (PreviousLinkedLayerInstanceData->GetLinkedFunctions().Num() == 0)
+				{
+					SharedLinkedAnimLayers->RemoveInstance(InPreviousTargetInstance);
+					MeshComp->GetLinkedAnimInstances().Remove(InPreviousTargetInstance);
+					// Only call UninitializeAnimation if we are not the owning anim instance
+					InPreviousTargetInstance->UninitializeAnimation();
+					InPreviousTargetInstance->MarkAsGarbage();
+				}
+			}
+		}
+	}
+}
 #if WITH_EDITOR
 
 void FAnimNode_LinkedAnimLayer::HandleObjectsReinstanced_Impl(UObject* InSourceObject, UObject* InTargetObject, const TMap<UObject*, UObject*>& OldToNewInstanceMap)
diff --git a/Engine/Source/Runtime/Engine/Private/Animation/AnimSubsystem_SharedLinkedAnimLayers.cpp b/Engine/Source/Runtime/Engine/Private/Animation/AnimSubsystem_SharedLinkedAnimLayers.cpp
new file mode 100644
--- /dev/null
+++ b/Engine/Source/Runtime/Engine/Private/Animation/AnimSubsystem_SharedLinkedAnimLayers.cpp
@@ -0,0 +1,157 @@
+// Copyright Epic Games, Inc. All Rights Reserved.
+
+#include "Animation/AnimSubsystem_SharedLinkedAnimLayers.h"
+#include "Animation/AnimInstance.h"
+
+#if LINKEDANIMLAYERSDATA_INTEGRITYCHECKS
+#include "Animation/AnimClassInterface.h"
+#include "Animation/AnimNode_LinkedAnimLayer.h"
+
+namespace
+{
+	// Check that function is either linked (!bIsFree) or unlinked (bIsFree)
+	void CheckLayerDataIntegrity(const UAnimInstance* AnimInstance, FName FunctionName, bool bIsFree)
+	{
+		bool bFound = false;
+
+		TArray<const UAnimInstance*> AllAnimInstances;
+		const USkeletalMeshComponent* SkelMesh = AnimInstance->GetSkelMeshComponent();
+		AllAnimInstances.Add(SkelMesh->GetAnimInstance());
+		AllAnimInstances.Append(SkelMesh->GetLinkedAnimInstances());
+		for (const UAnimInstance* LinkedAnimInstance : AllAnimInstances)
+		{
+			const IAnimClassInterface* const NewLinkedInstanceClass = IAnimClassInterface::GetFromClass(LinkedAnimInstance->GetClass());
+			for (const FStructProperty* const LayerNodeProperty : NewLinkedInstanceClass->GetLinkedAnimLayerNodeProperties())
+			{
+				const FAnimNode_LinkedAnimLayer* const LinkedAnimLayerNode = LayerNodeProperty->ContainerPtrToValuePtr<const FAnimNode_LinkedAnimLayer>(LinkedAnimInstance);
+				if (LinkedAnimLayerNode->GetDynamicLinkFunctionName() == FunctionName)
+				{
+					if (LinkedAnimLayerNode->GetTargetInstance<const UAnimInstance>() == AnimInstance)
+					{
+						// Function should be free to link but isn't
+						if (bIsFree)
+						{
+							check(0);
+						}
+						else
+						{
+							// Function shouldn't be linked more than once
+							check(!bFound);
+							bFound = true;
+						}
+					}
+				}
+			}
+		}
+		// Function is either found or free
+		check(bIsFree || bFound);
+	}
+}
+
+#endif // LINKEDANIMLAYERSDATA_INTEGRITYCHECKS
+
+void FLinkedAnimLayerInstanceData::AddLinkedFunction(FName Function, UAnimInstance* AnimInstance)
+{
+#if LINKEDANIMLAYERSDATA_INTEGRITYCHECKS
+	CheckLayerDataIntegrity(Instance, Function, true);
+#endif
+	check(!LinkedFunctions.Contains(Function));
+	LinkedFunctions.Add(Function, AnimInstance);
+}
+
+void FLinkedAnimLayerInstanceData::RemoveLinkedFunction(FName Function)
+{
+#if LINKEDANIMLAYERSDATA_INTEGRITYCHECKS
+	CheckLayerDataIntegrity(Instance, Function, true);
+#endif
+	check(LinkedFunctions.Contains(Function));
+	LinkedFunctions.Remove(Function);
+}
+FLinkedAnimLayerInstanceData* FLinkedAnimLayerClassData::FindInstanceData(const UAnimInstance* AnimInstance)
+{
+	return InstancesData.FindByPredicate([AnimInstance](FLinkedAnimLayerInstanceData& InstanceData) {return InstanceData.Instance == AnimInstance; });
+}
+
+FLinkedAnimLayerInstanceData* FLinkedAnimLayerClassData::FindInstanceForLinking(FName Function)
+{
+	for (FLinkedAnimLayerInstanceData& LayerInstanceData : InstancesData)
+	{
+		// Check if function is already linked
+		if (!LayerInstanceData.GetLinkedFunctions().Contains(Function))
+		{
+			// not linked, use this instance
+			return &LayerInstanceData;
+		}
+	}
+	return nullptr;
+}
+
+FLinkedAnimLayerInstanceData& FLinkedAnimLayerClassData::AddInstance(UAnimInstance* AnimInstance)
+{
+	InstancesData.Push(FLinkedAnimLayerInstanceData(AnimInstance));
+	return InstancesData.Last();
+}
+
+void FLinkedAnimLayerClassData::RemoveInstance(const UAnimInstance* AnimInstance)
+{
+	for (int i = 0; i < InstancesData.Num(); ++i)
+	{
+		if (InstancesData[i].Instance == AnimInstance)
+		{
+			// Make sure no functions are still linked
+			check(InstancesData[i].GetLinkedFunctions().Num() == 0);
+			InstancesData.RemoveAt(i);
+			return;
+		}
+	}
+	check(0);
+}
+
+void FAnimSubsystem_SharedLinkedAnimLayers::Reset()
+{
+	ClassesData.Empty(ClassesData.Num());
+}
+
+FLinkedAnimLayerInstanceData* FAnimSubsystem_SharedLinkedAnimLayers::FindInstanceData(const UAnimInstance* AnimInstance)
+{
+	if (FLinkedAnimLayerClassData* ClassData = FindClassData(AnimInstance->GetClass()))
+	{
+		return ClassData->FindInstanceData(AnimInstance);
+	}
+	return nullptr;
+}
+
+FLinkedAnimLayerClassData* FAnimSubsystem_SharedLinkedAnimLayers::FindClassData(TSubclassOf<UAnimInstance> AnimClass)
+{
+	return ClassesData.FindByPredicate([AnimClass](const FLinkedAnimLayerClassData& ClassData) {return ClassData.GetClass() == AnimClass; });
+}
+
+FLinkedAnimLayerClassData& FAnimSubsystem_SharedLinkedAnimLayers::FindOrAddClassData(TSubclassOf<UAnimInstance> AnimClass)
+{
+	if (FLinkedAnimLayerClassData* Result = FindClassData(AnimClass))
+	{
+		return *Result;
+	}
+	ClassesData.Add(FLinkedAnimLayerClassData(AnimClass));
+	return ClassesData.Last();
+};
+
+void FAnimSubsystem_SharedLinkedAnimLayers::RemoveInstance(const UAnimInstance* AnimInstance)
+{
+	int32 ClassIndex = ClassesData.IndexOfByPredicate([AnimInstance](const FLinkedAnimLayerClassData& ClassData) {return ClassData.GetClass() == AnimInstance->GetClass(); });
+	if (ClassIndex != INDEX_NONE)
+	{
+		FLinkedAnimLayerClassData& ClassData = ClassesData[ClassIndex];
+		ClassData.RemoveInstance(AnimInstance);
+		// Remove class data when all instances are removed
+		if (ClassData.GetInstancesData().Num() == 0)
+		{
+			// Remove class
+			ClassesData.RemoveAt(ClassIndex);
+		}
+	}
+	else
+	{
+		check(0);
+	}
+}
diff --git a/Engine/Source/Runtime/Engine/Private/Components/SkeletalMeshComponent.cpp b/Engine/Source/Runtime/Engine/Private/Components/SkeletalMeshComponent.cpp
--- a/Engine/Source/Runtime/Engine/Private/Components/SkeletalMeshComponent.cpp
+++ b/Engine/Source/Runtime/Engine/Private/Components/SkeletalMeshComponent.cpp
@@ -45,6 +45,7 @@
 #include "ContentStreaming.h"
 #include "Animation/AnimTrace.h"
 #include "Animation/BuiltInAttributeTypes.h"
+#include "Animation/AnimSubsystem_SharedLinkedAnimLayers.h"
 #include "HAL/LowLevelMemTracker.h"
 #include "HAL/LowLevelMemStats.h"
 #include "HAL/IConsoleManager.h"
@@ -3109,6 +3110,15 @@ UAnimInstance* USkeletalMeshComponent::GetPostProcessInstance() const
 
 void USkeletalMeshComponent::ResetLinkedAnimInstances()
 {
+	// Reset linked anim layers shared data 
+	if (AnimScriptInstance)
+	{
+		if (FAnimSubsystem_SharedLinkedAnimLayers* SharedLinkedAnimLayers = AnimScriptInstance->FindSubsystem<FAnimSubsystem_SharedLinkedAnimLayers>())
+		{
+			SharedLinkedAnimLayers->Reset();
+		}
+	}
+
 	for(UAnimInstance* LinkedInstance : LinkedInstances)
 	{
 		if(LinkedInstance && LinkedInstance->bCreatedByLinkedAnimGraph)
diff --git a/Engine/Source/Runtime/Engine/Public/Animation/AnimSubsystem_SharedLinkedAnimLayers.h b/Engine/Source/Runtime/Engine/Public/Animation/AnimSubsystem_SharedLinkedAnimLayers.h
new file mode 100644
--- /dev/null
+++ b/Engine/Source/Runtime/Engine/Public/Animation/AnimSubsystem_SharedLinkedAnimLayers.h
@@ -0,0 +1,94 @@
+// Copyright Epic Games, Inc. All Rights Reserved.
+#pragma once
+
+#include "CoreMinimal.h"
+#include "Templates/SubclassOf.h"
+#include "Animation\AnimSubsystemInstance.h"
+#include "AnimSubsystem_SharedLinkedAnimLayers.generated.h"
+
+class UAnimInstance;
+
+// Define to 1 for more extensive runtime validation of linked anim layers data
+#define LINKEDANIMLAYERSDATA_INTEGRITYCHECKS 0
+
+// Linked layer instance info
+USTRUCT()
+struct FLinkedAnimLayerInstanceData
+{
+	GENERATED_USTRUCT_BODY()
+
+	FLinkedAnimLayerInstanceData() {}
+	FLinkedAnimLayerInstanceData(UAnimInstance* AnimInstance) { Instance = AnimInstance; }
+
+	// Mark a function as linked
+	void AddLinkedFunction(FName Function, UAnimInstance* AnimInstance);
+	// Unmark a function as linked
+	void RemoveLinkedFunction(FName Function);
+
+	// Return map of currently linked functions
+	const TMap<FName, TWeakObjectPtr<UAnimInstance>>& GetLinkedFunctions() const {return LinkedFunctions;}
+	
+	UPROPERTY(Transient)
+	TObjectPtr<UAnimInstance> Instance;
+
+private:
+	UPROPERTY(Transient)
+	TMap<FName, TWeakObjectPtr<UAnimInstance>> LinkedFunctions;
+};
+
+// Linked layer class info
+USTRUCT()
+struct FLinkedAnimLayerClassData
+{
+	GENERATED_BODY()
+
+	FLinkedAnimLayerClassData() {}
+	FLinkedAnimLayerClassData(TSubclassOf<UAnimInstance> AnimClass) { Class = AnimClass; }
+
+	// Find instance data for given instance if it exists, returns nullptr otherwise
+	FLinkedAnimLayerInstanceData* FindInstanceData(const UAnimInstance* AnimInstance);
+
+	// Find an existing instance to link given function if one is available, returns nullptr otherwise
+	FLinkedAnimLayerInstanceData* FindInstanceForLinking(FName Function);
+	
+	// Add given instance
+	FLinkedAnimLayerInstanceData& AddInstance(UAnimInstance* AnimInstance);
+	// Remove given instance
+	void RemoveInstance(const UAnimInstance* AnimInstance);
+
+	TSubclassOf<UAnimInstance> GetClass() const {return Class;}
+	const TArray<FLinkedAnimLayerInstanceData>& GetInstancesData() const {return InstancesData;}
+
+private:
+	TSubclassOf<UAnimInstance> Class;
+
+	UPROPERTY(Transient)
+	TArray<FLinkedAnimLayerInstanceData> InstancesData;
+};
+
+// Data for shared linked anim instances module
+USTRUCT()
+struct ENGINE_API FAnimSubsystem_SharedLinkedAnimLayers : public FAnimSubsystemInstance
+{
+	GENERATED_BODY()
+
+	// Clear all linked layers data
+	void Reset();
+
+	// Find instance data for given instance if it exists, returns nullptr otherwise
+	FLinkedAnimLayerInstanceData* FindInstanceData(const UAnimInstance* AnimInstance);
+
+	// Find class data for given anim class if it exists, returns nullptr otherwise
+	FLinkedAnimLayerClassData* FindClassData(TSubclassOf<UAnimInstance> AnimClass);
+
+	// Find or add ClassData for given AnimClass
+	FLinkedAnimLayerClassData& FindOrAddClassData(TSubclassOf<UAnimInstance> AnimClass);
+
+	// Remove given instance
+	void RemoveInstance(const UAnimInstance* AnimInstance);
+
+	const TArray<FLinkedAnimLayerClassData>& GetClassesData() const { return ClassesData; }
+private:
+	UPROPERTY(Transient)
+	TArray<FLinkedAnimLayerClassData> ClassesData;
+};
\ No newline at end of file
