Fixed an issue where the instance bound extents weren't getting updated properly when using the mobile renderer.

diff --git a/Engine/Source/Runtime/Engine/Private/PrimitiveSceneProxy.cpp b/Engine/Source/Runtime/Engine/Private/PrimitiveSceneProxy.cpp
--- a/Engine/Source/Runtime/Engine/Private/PrimitiveSceneProxy.cpp
+++ b/Engine/Source/Runtime/Engine/Private/PrimitiveSceneProxy.cpp
@@ -514,7 +514,6 @@ void FPrimitiveSceneProxy::UpdateUniformBuffer()
 				.ActorWorldPosition(ActorPosition)
 				.WorldBounds(Bounds)
 				.LocalBounds(LocalBounds)
-				.InstanceLocalBounds(GetInstanceLocalBounds(0))
 				.PreSkinnedLocalBounds(PreSkinnedLocalBounds)
 				.ReceivesDecals(bReceivesDecals)
 				.CacheShadowAsStatic(PrimitiveSceneInfo ? PrimitiveSceneInfo->ShouldCacheShadowAsStatic() : false)
@@ -548,6 +547,12 @@ void FPrimitiveSceneProxy::UpdateUniformBuffer()
 			Builder.InstanceWorldPositionOffsetDisableDistance(WPODisableDistance);
 		}
 
+		const TConstArrayView<FRenderBounds> InstanceBounds = GetInstanceLocalBounds();
+		if (InstanceBounds.Num() > 0)
+		{
+			Builder.InstanceLocalBounds(InstanceBounds[0]);
+		}
+
 		FPrimitiveUniformShaderParameters PrimitiveParams = Builder.Build();
 
 		if (UniformBuffer.GetReference())
@@ -627,11 +632,13 @@ void FPrimitiveSceneProxy::SetTransform(const FMatrix& InLocalToWorld, const FBo
 	{
 		PrimitiveSceneInfo->bNeedsCachedReflectionCaptureUpdate = true;
 	}
-	
-	UpdateUniformBuffer();
-	
+
 	// Notify the proxy's implementation of the change.
 	OnTransformChanged();
+	
+	// Need to update the uniform buffer after calling OnTransformChanged because
+	// OnTransformChanged could change values used in the uniform buffer.
+	UpdateUniformBuffer();
 }
 
 void FPrimitiveSceneProxy::UpdateInstances_RenderThread(const FInstanceUpdateCmdBuffer& CmdBuffer, const FBoxSphereBounds& InBounds, const FBoxSphereBounds& InLocalBounds, const FBoxSphereBounds& InStaticMeshBounds)
@@ -1533,4 +1540,4 @@ ERayTracingPrimitiveFlags FPrimitiveSceneProxy::GetCachedRayTracingInstance(FRay
 
 	return ResultFlags;
 }
-#endif
\ No newline at end of file
+#endif
