[Engine]
* Level - Removed template specialization and put the ctor and the dtor of the struct needing the specialization in the cpp file instead

diff --git a/Engine/Source/Runtime/Engine/Classes/Engine/Level.h b/Engine/Source/Runtime/Engine/Classes/Engine/Level.h
--- a/Engine/Source/Runtime/Engine/Classes/Engine/Level.h
+++ b/Engine/Source/Runtime/Engine/Classes/Engine/Level.h
@@ -191,23 +191,14 @@ struct ENGINE_API FDynamicTextureInstance : public FStreamableTextureInstance
 	friend FArchive& operator<<( FArchive& Ar, FDynamicTextureInstance& TextureInstance );
 };
 
-/** Manually implement TPointerIsConvertibleFromTo for AActor so that TWeakObjectPtr can use it below when AActor is forward declared. */
-template<> struct TPointerIsConvertibleFromTo<AActor, const volatile UObject>
-{
-	enum { Value = 1 };
-};
-
 /** Struct that holds on to information about Actors that wish to be auto enabled for input before the player controller has been created */
 struct FPendingAutoReceiveInputActor
 {
 	TWeakObjectPtr<AActor> Actor;
 	int32 PlayerIndex;
 
-	FPendingAutoReceiveInputActor(AActor* InActor, const int32 InPlayerIndex)
-		: Actor(InActor)
-		, PlayerIndex(InPlayerIndex)
-	{
-	}
+	FPendingAutoReceiveInputActor(AActor* InActor, const int32 InPlayerIndex);
+	~FPendingAutoReceiveInputActor();
 };
 
 /** A precomputed visibility cell, whose data is stored in FCompressedVisibilityChunk. */
diff --git a/Engine/Source/Runtime/Engine/Private/Level.cpp b/Engine/Source/Runtime/Engine/Private/Level.cpp
--- a/Engine/Source/Runtime/Engine/Private/Level.cpp
+++ b/Engine/Source/Runtime/Engine/Private/Level.cpp
@@ -270,6 +270,18 @@ void FActorFolderSet::Add(UActorFolder* InActorFolder)
 	ActorFolders.Add(InActorFolder);
 }
 
+/*-----------------------------------------------------------------------------
+FPendingAutoReceiveInputActor implementation.
+-----------------------------------------------------------------------------*/
+
+FPendingAutoReceiveInputActor::FPendingAutoReceiveInputActor(AActor* InActor, const int32 InPlayerIndex)
+	: Actor(InActor)
+	, PlayerIndex(InPlayerIndex)
+{
+}
+
+FPendingAutoReceiveInputActor::~FPendingAutoReceiveInputActor() = default;
+
 /*-----------------------------------------------------------------------------
 ULevel implementation.
 -----------------------------------------------------------------------------*/
