temporary fix for cook failures when using PerQualityLevelProperties on Android and Matched Fragments

diff --git a/Engine/Source/Runtime/Engine/Private/DeviceProfiles/DeviceProfileMatching.cpp b/Engine/Source/Runtime/Engine/Private/DeviceProfiles/DeviceProfileMatching.cpp
--- a/Engine/Source/Runtime/Engine/Private/DeviceProfiles/DeviceProfileMatching.cpp
+++ b/Engine/Source/Runtime/Engine/Private/DeviceProfiles/DeviceProfileMatching.cpp
@@ -586,7 +586,14 @@ public:
 
 	virtual void Serialize(const TCHAR* V, ELogVerbosity::Type Verbosity, const class FName& Category) override
 	{
-		UE_LOG(LogInit, Error, TEXT("DeviceProfileMatching: Error while parsing Matching Rules (%s) : %s"), *Stage, V);
+		if (IsRunningCookCommandlet())
+		{
+			UE_LOG(LogInit, Warning, TEXT("DeviceProfileMatching: Error while parsing Matching Rules (%s) : %s"), *Stage, V);
+		}
+		else
+		{
+			UE_LOG(LogInit, Error, TEXT("DeviceProfileMatching: Error while parsing Matching Rules (%s) : %s"), *Stage, V);
+		}
 		NumErrors++;
 	}
 };
@@ -622,11 +629,14 @@ static FString LoadAndProcessMatchingRulesFromConfig(const FString& ParentDP, ID
 		const FString& RuleName = RuleStruct.RuleName.IsEmpty() ? NoName : RuleStruct.RuleName;
 		EvaluateMatch(DPSelector, RuleName, RuleStruct, RuleProperties, SelectedFragments, &DPMatchingErrorOutput);
 	}
-#if 1//UE_BUILD_SHIPPING
-	UE_CLOG(DPMatchingErrorOutput.NumErrors > 0, LogInit, Error, TEXT("DeviceProfileMatching: %d Error(s) encountered while processing MatchedRules for %s"), DPMatchingErrorOutput.NumErrors, *ParentDP);
-#else
-	UE_CLOG(DPMatchingErrorOutput.NumErrors > 0, LogInit, Fatal, TEXT("DeviceProfileMatching: %d Error(s) encountered while processing MatchedRules for %s"), DPMatchingErrorOutput.NumErrors, *ParentDP);
-#endif
+	if (IsRunningCookCommandlet())
+	{
+		UE_CLOG(DPMatchingErrorOutput.NumErrors > 0, LogInit, Warning, TEXT("DeviceProfileMatching: %d Warnings(s) encountered while processing MatchedRules for %s"), DPMatchingErrorOutput.NumErrors, *ParentDP);
+	}
+	else
+	{
+		UE_CLOG(DPMatchingErrorOutput.NumErrors > 0, LogInit, Error, TEXT("DeviceProfileMatching: %d Error(s) encountered while processing MatchedRules for %s"), DPMatchingErrorOutput.NumErrors, *ParentDP);
+	}
 
 	FGenericCrashContext::SetEngineData(TEXT("MatchingDPStatus"), ParentDP + (DPMatchingErrorOutput.NumErrors > 0 ? TEXT("Error") : TEXT("No errors")));
 
