Added include of SkeletalMeshComponent.h to TraceFilter.h for all clang platforms instead of only linux

diff --git a/Engine/Source/Runtime/Engine/Public/TraceFilter.h b/Engine/Source/Runtime/Engine/Public/TraceFilter.h
--- a/Engine/Source/Runtime/Engine/Public/TraceFilter.h
+++ b/Engine/Source/Runtime/Engine/Public/TraceFilter.h
@@ -12,8 +12,8 @@
 #include "Templates/UnrealTemplate.h"
 #include "UObject/Object.h"
 
-#if PLATFORM_LINUX
-#include "Components/SkeletalMeshComponent.h" // Needed because AnimInstance::GetOwningComponent. (this only fails on linux)
+#if defined(__clang__)
+#include "Components/SkeletalMeshComponent.h" // Needed because AnimInstance::GetOwningComponent. (this only fails on clang compilers)
 #endif
 
 
