Content Bundles no longer inject runtime grids but instead inject their content into the corresponding source runtime grid.
- Optimizes the time spent computing cells intersecting streaming sources (WorldPartition::UpdateStreamingState())
- Added a debug color per content bundle
- Adapted debug draw of 2D/3D runtime grids to draw content bundle cells with their associated debug color
- Changed wp.Runtime.DebugFilterOutContentBundles to wp.Runtime.DrawContentBundles and adapted all debug display to respect the option
- Added debug display in Content Bundle Manager to list all injected content bundles with their corresponding debug color (similar to runtime Data Layers)

diff --git a/Engine/Plugins/Experimental/GameFeatures/Source/GameFeatures/Private/GameFeatureAction_AddWPContent.cpp b/Engine/Plugins/Experimental/GameFeatures/Source/GameFeatures/Private/GameFeatureAction_AddWPContent.cpp
--- a/Engine/Plugins/Experimental/GameFeatures/Source/GameFeatures/Private/GameFeatureAction_AddWPContent.cpp
+++ b/Engine/Plugins/Experimental/GameFeatures/Source/GameFeatures/Private/GameFeatureAction_AddWPContent.cpp
@@ -17,11 +17,7 @@ UGameFeatureAction_AddWPContent::UGameFeatureAction_AddWPContent(const FObjectIn
 #if WITH_EDITOR
 	if (UGameFeatureData* GameFeatureData = GetTypedOuter<UGameFeatureData>())
 	{
-		ContentBundleDescriptor->SetDisplayName(GameFeatureData->GetName());
-
-		FString PackagePath = GetPackage()->GetName();
-		FName MountPoint = FPackageName::GetPackageMountPoint(PackagePath);
-		ContentBundleDescriptor->SetPackageRoot(MountPoint.ToString());
+		ContentBundleDescriptor->InitializeObject(GameFeatureData->GetName(), FPackageName::GetPackageMountPoint(GetPackage()->GetName()).ToString());
 	}
 #endif
 }
diff --git a/Engine/Source/Runtime/Engine/Private/WorldPartition/ContentBundle/ContentBundleBase.cpp b/Engine/Source/Runtime/Engine/Private/WorldPartition/ContentBundle/ContentBundleBase.cpp
--- a/Engine/Source/Runtime/Engine/Private/WorldPartition/ContentBundle/ContentBundleBase.cpp
+++ b/Engine/Source/Runtime/Engine/Private/WorldPartition/ContentBundle/ContentBundleBase.cpp
@@ -104,6 +104,11 @@ const FString& FContentBundleBase::GetDisplayName() const
 	return GetDescriptor()->GetDisplayName();
 }
 
+const FColor& FContentBundleBase::GetDebugColor() const
+{
+	return GetDescriptor()->GetDebugColor();
+}
+
 void FContentBundleBase::AddReferencedObjects(FReferenceCollector& Collector)
 {
 	Collector.AddReferencedObject(Descriptor);
diff --git a/Engine/Source/Runtime/Engine/Private/WorldPartition/ContentBundle/ContentBundleContainer.cpp b/Engine/Source/Runtime/Engine/Private/WorldPartition/ContentBundle/ContentBundleContainer.cpp
--- a/Engine/Source/Runtime/Engine/Private/WorldPartition/ContentBundle/ContentBundleContainer.cpp
+++ b/Engine/Source/Runtime/Engine/Private/WorldPartition/ContentBundle/ContentBundleContainer.cpp
@@ -83,7 +83,7 @@ uint32 FContentBundleContainer::GetNumContentBundles() const
 #if WITH_EDITOR
 	if(UseEditorContentBundle())
 	{
-		GetEditorContentBundles().Num();
+		return GetEditorContentBundles().Num();
 	}
 #endif
 
diff --git a/Engine/Source/Runtime/Engine/Private/WorldPartition/ContentBundle/ContentBundleDescriptor.cpp b/Engine/Source/Runtime/Engine/Private/WorldPartition/ContentBundle/ContentBundleDescriptor.cpp
--- a/Engine/Source/Runtime/Engine/Private/WorldPartition/ContentBundle/ContentBundleDescriptor.cpp
+++ b/Engine/Source/Runtime/Engine/Private/WorldPartition/ContentBundle/ContentBundleDescriptor.cpp
@@ -1,17 +1,12 @@
 // Copyright Epic Games, Inc. All Rights Reserved.
 #include "WorldPartition/ContentBundle/ContentBundleDescriptor.h"
+#include "Math/Color.h"
 
 #include UE_INLINE_GENERATED_CPP_BY_NAME(ContentBundleDescriptor)
 
 UContentBundleDescriptor::UContentBundleDescriptor(const FObjectInitializer& ObjectInitializer)
-	:Guid(FGuid::NewGuid())
+	: DebugColor(FColor::Black)
 {
-
-}
-
-const FString& UContentBundleDescriptor::GetDisplayName() const
-{
-	return DisplayName;
 }
 
 bool UContentBundleDescriptor::IsValid() const
@@ -20,3 +15,29 @@ bool UContentBundleDescriptor::IsValid() const
 		&& !DisplayName.IsEmpty()
 		&& !PackageRoot.IsEmpty();
 }
+
+#if WITH_EDITOR
+void UContentBundleDescriptor::InitializeObject(const FString& InContentBundleName, const FString& InPackageRoot)
+{
+	Guid = FGuid::NewGuid();
+	DisplayName = InContentBundleName;
+	PackageRoot = InPackageRoot;
+	InitDebugColor();
+}
+
+void UContentBundleDescriptor::PostLoad()
+{
+	InitDebugColor();
+
+	Super::PostLoad();
+}
+
+void UContentBundleDescriptor::InitDebugColor()
+{
+	// If not set, generate a color based on guid
+	if (DebugColor == FColor::Black)
+	{
+		DebugColor = FColor::MakeRandomSeededColor(GetTypeHash(Guid));
+	}
+}
+#endif
\ No newline at end of file
diff --git a/Engine/Source/Runtime/Engine/Private/WorldPartition/ContentBundle/ContentBundleWorldSubsystem.cpp b/Engine/Source/Runtime/Engine/Private/WorldPartition/ContentBundle/ContentBundleWorldSubsystem.cpp
--- a/Engine/Source/Runtime/Engine/Private/WorldPartition/ContentBundle/ContentBundleWorldSubsystem.cpp
+++ b/Engine/Source/Runtime/Engine/Private/WorldPartition/ContentBundle/ContentBundleWorldSubsystem.cpp
@@ -8,6 +8,8 @@
 #include "WorldPartition/ContentBundle/ContentBundleClient.h"
 #include "WorldPartition/WorldPartition.h"
 #include "WorldPartition/WorldPartitionSubsystem.h"
+#include "WorldPartition/WorldPartitionDebugHelper.h"
+#include "Engine/Canvas.h"
 #include "Engine/World.h"
 
 #include UE_INLINE_GENERATED_CPP_BY_NAME(ContentBundleWorldSubsystem)
@@ -66,6 +68,58 @@ void UContentBundleManager::AddReferencedObjects(UObject* InThis, FReferenceColl
 	}
 }
 
+void UContentBundleManager::DrawContentBundlesStatus(const UWorld* InWorld, UCanvas* Canvas, FVector2D& Offset) const
+{
+	if (const TUniquePtr<FContentBundleContainer>* ContentBundleContainer = GetContentBundleContainer(InWorld))
+	{
+		if ((*ContentBundleContainer)->GetNumContentBundles())
+		{
+			FVector2D Pos = Offset;
+			float MaxTextWidth = 0.f;
+
+			bool bHasInjectedContentBundle = false;
+			(*ContentBundleContainer)->ForEachContentBundleBreakable([&bHasInjectedContentBundle](FContentBundleBase* ContentBundle)
+			{
+				if (ContentBundle->GetStatus() == EContentBundleStatus::ContentInjected)
+				{
+					bHasInjectedContentBundle = true;
+				}
+				return !bHasInjectedContentBundle;
+			});
+
+			if (bHasInjectedContentBundle)
+			{
+				FWorldPartitionDebugHelper::DrawText(Canvas, TEXT("Injected Content Bundles"), GEngine->GetSmallFont(), FColor::Green, Pos, &MaxTextWidth);
+
+				(*ContentBundleContainer)->ForEachContentBundle([&Offset, &Canvas, &Pos, &MaxTextWidth](FContentBundleBase* ContentBundle)
+				{
+					FWorldPartitionDebugHelper::DrawLegendItem(Canvas, *ContentBundle->GetDisplayName(), GEngine->GetSmallFont(), ContentBundle->GetDebugColor(), FColor::White, Pos, &MaxTextWidth);
+				});
+			}
+
+			Offset.X += MaxTextWidth + 10;
+		}
+	}
+}
+
+const FContentBundleBase* UContentBundleManager::GetContentBundle(const UWorld* InWorld, const FGuid& Guid) const
+{
+	const FContentBundleBase* FoundContentBundle = nullptr;
+	if (const TUniquePtr<FContentBundleContainer>* ContentBundleContainer = GetContentBundleContainer(InWorld))
+	{
+		(*ContentBundleContainer)->ForEachContentBundleBreakable([&FoundContentBundle, &Guid](FContentBundleBase* ContentBundle)
+		{
+			if (ContentBundle->GetDescriptor()->GetGuid() == Guid)
+			{
+				FoundContentBundle = ContentBundle;
+				return false;
+			}
+			return true;
+		});
+	}
+	return FoundContentBundle;
+}
+
 #if WITH_EDITOR
 
 bool UContentBundleManager::GetEditorContentBundle(TArray<TSharedPtr<FContentBundleEditor>>& OutContentBundles)
@@ -93,7 +147,7 @@ TSharedPtr<FContentBundleEditor> UContentBundleManager::GetEditorContentBundle(c
 
 #endif
 
-uint32 UContentBundleManager::GetContentBundleContainerIndex(const UWorld* InjectedWorld)
+uint32 UContentBundleManager::GetContentBundleContainerIndex(const UWorld* InjectedWorld) const
 {
 	return ContentBundleContainers.IndexOfByPredicate([InjectedWorld](const TUniquePtr<FContentBundleContainer>& ContentBundleContainer)
 		{
@@ -101,7 +155,7 @@ uint32 UContentBundleManager::GetContentBundleContainerIndex(const UWorld* Injec
 		});
 }
 
-TUniquePtr<FContentBundleContainer>* UContentBundleManager::GetContentBundleContainer(const UWorld* InjectedWorld)
+const TUniquePtr<FContentBundleContainer>* UContentBundleManager::GetContentBundleContainer(const UWorld* InjectedWorld) const
 {
 	uint32 ContainerIndex = GetContentBundleContainerIndex(InjectedWorld);
 	if (ContainerIndex != INDEX_NONE)
diff --git a/Engine/Source/Runtime/Engine/Private/WorldPartition/WorldPartitionDebugHelper.cpp b/Engine/Source/Runtime/Engine/Private/WorldPartition/WorldPartitionDebugHelper.cpp
--- a/Engine/Source/Runtime/Engine/Private/WorldPartition/WorldPartitionDebugHelper.cpp
+++ b/Engine/Source/Runtime/Engine/Private/WorldPartition/WorldPartitionDebugHelper.cpp
@@ -12,11 +12,11 @@
 #include "Algo/AnyOf.h"
 #include "Algo/AllOf.h"
 
-bool FWorldPartitionDebugHelper::bDebugFilterOutContentBundles = true;
-FAutoConsoleVariableRef FWorldPartitionDebugHelper::DebugFilterOutContentBundlesCommand(
-	TEXT("wp.Runtime.DebugFilterOutContentBundles"),
-	FWorldPartitionDebugHelper::bDebugFilterOutContentBundles,
-	TEXT("Filter content bundle from world partition debug display."));
+bool FWorldPartitionDebugHelper::bCanDrawContentBundles = true;
+FAutoConsoleVariableRef FWorldPartitionDebugHelper::DrawContentBundlesCommand(
+	TEXT("wp.Runtime.DrawContentBundles"),
+	FWorldPartitionDebugHelper::bCanDrawContentBundles,
+	TEXT("Enable to draw debug display of content bundle."));
 
 TSet<FName> FWorldPartitionDebugHelper::DebugRuntimeHashFilter;
 FAutoConsoleCommand FWorldPartitionDebugHelper::DebugFilterByRuntimeHashGridNameCommand(
diff --git a/Engine/Source/Runtime/Engine/Private/WorldPartition/WorldPartitionRuntimeCell.cpp b/Engine/Source/Runtime/Engine/Private/WorldPartition/WorldPartitionRuntimeCell.cpp
--- a/Engine/Source/Runtime/Engine/Private/WorldPartition/WorldPartitionRuntimeCell.cpp
+++ b/Engine/Source/Runtime/Engine/Private/WorldPartition/WorldPartitionRuntimeCell.cpp
@@ -57,12 +57,6 @@ void UWorldPartitionRuntimeCell::SetDebugInfo(int64 InCoordX, int64 InCoordY, in
 	UpdateDebugName();
 }
 
-void UWorldPartitionRuntimeCell::SetGridName(FName InGridName)
-{
-	DebugInfo.GridName = InGridName;
-	UpdateDebugName();
-}
-
 void UWorldPartitionRuntimeCell::UpdateDebugName()
 {
 	TStringBuilder<512> Builder;
@@ -129,10 +123,11 @@ int32 UWorldPartitionRuntimeCell::SortCompare(const UWorldPartitionRuntimeCell*
 
 bool UWorldPartitionRuntimeCell::IsDebugShown() const
 {
-	return FWorldPartitionDebugHelper::IsDebugRuntimeHashGridShown(DebugInfo.GridName) &&
+	return FWorldPartitionDebugHelper::IsDebugRuntimeHashGridShown(GetGridName()) &&
 		   FWorldPartitionDebugHelper::IsDebugStreamingStatusShown(GetStreamingStatus()) &&
-		   FWorldPartitionDebugHelper::AreDebugDataLayersShown(DataLayers) &&
-		   FWorldPartitionDebugHelper::IsDebugCellNameShown(DebugInfo.Name);
+	       FWorldPartitionDebugHelper::AreDebugDataLayersShown(DataLayers) &&
+		   FWorldPartitionDebugHelper::IsDebugCellNameShown(DebugInfo.Name) &&
+		   (FWorldPartitionDebugHelper::CanDrawContentBundles() || !ContentBundleID.IsValid());
 }
 
 FLinearColor UWorldPartitionRuntimeCell::GetDebugStreamingPriorityColor() const
diff --git a/Engine/Source/Runtime/Engine/Private/WorldPartition/WorldPartitionRuntimeSpatialHash.cpp b/Engine/Source/Runtime/Engine/Private/WorldPartition/WorldPartitionRuntimeSpatialHash.cpp
--- a/Engine/Source/Runtime/Engine/Private/WorldPartition/WorldPartitionRuntimeSpatialHash.cpp
+++ b/Engine/Source/Runtime/Engine/Private/WorldPartition/WorldPartitionRuntimeSpatialHash.cpp
@@ -14,16 +14,15 @@
 #include "WorldPartition/WorldPartitionHandle.h"
 #include "WorldPartition/WorldPartitionDebugHelper.h"
 #include "WorldPartition/WorldPartitionLevelStreamingDynamic.h"
-
 #include "WorldPartition/RuntimeSpatialHash/RuntimeSpatialHashGridHelper.h"
-
 #include "WorldPartition/HLOD/HLODActor.h"
 #include "WorldPartition/HLOD/HLODActorDesc.h"
-
 #include "WorldPartition/DataLayer/DataLayerSubsystem.h"
 #include "WorldPartition/DataLayer/WorldDataLayers.h"
 #include "WorldPartition/DataLayer/DataLayerInstance.h"
-
+#include "WorldPartition/ContentBundle/ContentBundleWorldSubsystem.h"
+#include "WorldPartition/ContentBundle/ContentBundleDescriptor.h"
+#include "WorldPartition/ContentBundle/ContentBundleBase.h"
 #include "GameFramework/Actor.h"
 #include "GameFramework/WorldSettings.h"
 #include "ProfilingDebugging/ScopedTimers.h"
@@ -188,19 +187,16 @@ void FSpatialHashStreamingGrid::GetCells(const FWorldPartitionStreamingQuerySour
 		{
 			Helper.ForEachIntersectingCells(Shape, [&](const FGridCellCoord& Coords)
 			{
-				if (const FSpatialHashStreamingGridLayerCell* LayerCell = GetLayerCell(Coords))
+				ForEachRuntimeCell(Coords, [&](const UWorldPartitionRuntimeCell* Cell)
 				{
-					for (const UWorldPartitionRuntimeCell* Cell : LayerCell->GridCells)
+					if (!bEnableZCulling || TRange<double>(Cell->GetMinMaxZ().X, Cell->GetMinMaxZ().Y).Overlaps(TRange<double>(Shape.GetCenter().Z - Shape.GetRadius(), Shape.GetCenter().Z + Shape.GetRadius())))
 					{
-						if (!bEnableZCulling || TRange<double>(Cell->GetMinMaxZ().X, Cell->GetMinMaxZ().Y).Overlaps(TRange<double>(Shape.GetCenter().Z - Shape.GetRadius(), Shape.GetCenter().Z + Shape.GetRadius())))
+						if (ShouldAddCell(Cell, QuerySource))
 						{
-							if (ShouldAddCell(Cell, QuerySource))
-							{
-								OutCells.Add(Cell);
-							}
+							OutCells.Add(Cell);
 						}
 					}
-				}
+				});
 			});
 		});
 	}
@@ -252,32 +248,29 @@ void FSpatialHashStreamingGrid::GetCells(const TArray<FWorldPartitionStreamingSo
 				if ((GFilterRuntimeSpatialHashGridLevel == INDEX_NONE) || (GFilterRuntimeSpatialHashGridLevel == Coords.Z))
 #endif
 				{
-					if (const FSpatialHashStreamingGridLayerCell* LayerCell = GetLayerCell(Coords))
+					ForEachRuntimeCell(Coords, [&](const UWorldPartitionRuntimeCell* Cell)
 					{
-						for (const UWorldPartitionRuntimeCell* Cell : LayerCell->GridCells)
+						if (!bEnableZCulling || TRange<double>(Cell->GetMinMaxZ().X, Cell->GetMinMaxZ().Y).Overlaps(TRange<double>(Shape.GetCenter().Z - Shape.GetRadius(), Shape.GetCenter().Z + Shape.GetRadius())))
 						{
-							if (!bEnableZCulling || TRange<double>(Cell->GetMinMaxZ().X, Cell->GetMinMaxZ().Y).Overlaps(TRange<double>(Shape.GetCenter().Z - Shape.GetRadius(), Shape.GetCenter().Z + Shape.GetRadius())))
+							if (!Cell->HasDataLayers() || (DataLayerSubsystem && DataLayerSubsystem->IsAnyDataLayerInEffectiveRuntimeState(Cell->GetDataLayers(), EDataLayerRuntimeState::Activated)))
 							{
-								if (!Cell->HasDataLayers() || (DataLayerSubsystem && DataLayerSubsystem->IsAnyDataLayerInEffectiveRuntimeState(Cell->GetDataLayers(), EDataLayerRuntimeState::Activated)))
+								if (Source.TargetState == EStreamingSourceTargetState::Loaded)
 								{
-									if (Source.TargetState == EStreamingSourceTargetState::Loaded)
-									{
-										OutLoadCells.AddCell(Cell, Source, Shape);
-									}
-									else
-									{
-										check(Source.TargetState == EStreamingSourceTargetState::Activated);
-										OutActivateCells.AddCell(Cell, Source, Shape);
-										bAddedActivatedCell = !GRuntimeSpatialHashUseAlignedGridLevels && GRuntimeSpatialHashSnapNonAlignedGridLevelsToLowerLevels;
-									}
+									OutLoadCells.AddCell(Cell, Source, Shape);
 								}
-								else if (DataLayerSubsystem && DataLayerSubsystem->IsAnyDataLayerInEffectiveRuntimeState(Cell->GetDataLayers(), EDataLayerRuntimeState::Loaded))
+								else
 								{
-									OutLoadCells.AddCell(Cell, Source, Shape);
+									check(Source.TargetState == EStreamingSourceTargetState::Activated);
+									OutActivateCells.AddCell(Cell, Source, Shape);
+									bAddedActivatedCell = !GRuntimeSpatialHashUseAlignedGridLevels && GRuntimeSpatialHashSnapNonAlignedGridLevelsToLowerLevels;
 								}
 							}
+							else if (DataLayerSubsystem && DataLayerSubsystem->IsAnyDataLayerInEffectiveRuntimeState(Cell->GetDataLayers(), EDataLayerRuntimeState::Loaded))
+							{
+								OutLoadCells.AddCell(Cell, Source, Shape);
+							}
 						}
-					}
+					});
 				}
 				if (bAddedActivatedCell)
 				{
@@ -342,31 +335,164 @@ void FSpatialHashStreamingGrid::GetCells(const TArray<FWorldPartitionStreamingSo
 
 		for (const auto& ParentCell : AllParentCells)
 		{
-			if (const FSpatialHashStreamingGridLayerCell* LayerCell = GetLayerCell(ParentCell.Key))
+			ForEachRuntimeCell(ParentCell.Key, [&](const UWorldPartitionRuntimeCell* Cell)
+			{
+				if (!Cell->HasDataLayers() || (DataLayerSubsystem && DataLayerSubsystem->IsAnyDataLayerInEffectiveRuntimeState(Cell->GetDataLayers(), EDataLayerRuntimeState::Activated)))
+				{
+					for (const auto& Info : ParentCell.Value)
+					{
+						OutActivateCells.AddCell(Cell, Info.Source, Info.SourceShape);
+					}
+				}
+			});
+		}
+	}
+}
+
+void FSpatialHashStreamingGrid::InjectExternalStreamingObjectGrid(const FSpatialHashStreamingGrid& InExternalObjectStreamingGrid) const
+{
+	if (InjectedGridLevels.IsEmpty())
+	{
+		InjectedGridLevels.SetNumZeroed(GridLevels.Num());
+	}
+
+	check(GetGridHelper().Levels.Num() == InjectedGridLevels.Num());
+
+	for (int32 SourceGridLevel = 0; SourceGridLevel < InExternalObjectStreamingGrid.GridLevels.Num(); ++SourceGridLevel)
+	{
+		const FSpatialHashStreamingGridLevel& ExternalObjectGridLevel = InExternalObjectStreamingGrid.GridLevels[SourceGridLevel];
+		int32 TargetGridLevel = FMath::Min<int32>(SourceGridLevel, GridLevels.Num()-1);
+		if (ensure(InjectedGridLevels.IsValidIndex(TargetGridLevel)))
+		{
+			check(InExternalObjectStreamingGrid.GetGridHelper().Levels[TargetGridLevel].CellSize == GetGridHelper().Levels[TargetGridLevel].CellSize);
+
+			for (const FSpatialHashStreamingGridLayerCell& ExternalObjectLayerCell : ExternalObjectGridLevel.LayerCells)
 			{
-				for (const UWorldPartitionRuntimeCell* Cell : LayerCell->GridCells)
+				if (ensure(!ExternalObjectLayerCell.GridCells.IsEmpty()))
 				{
-					if (!Cell->HasDataLayers() || (DataLayerSubsystem && DataLayerSubsystem->IsAnyDataLayerInEffectiveRuntimeState(Cell->GetDataLayers(), EDataLayerRuntimeState::Activated)))
+					// Any cell position should be good to find the proper cell coord
+					const FVector& CellPosition = ExternalObjectLayerCell.GridCells[0]->Position;
+					FGridCellCoord2 CellCoords;
+					const FSquare2DGridHelper::FGridLevel& InjectedGridHelperLevel = GetGridHelper().Levels[TargetGridLevel];
+					if (!InjectedGridHelperLevel.GetCellCoords(FVector2D(CellPosition.X, CellPosition.Y), CellCoords))
 					{
-						for (const auto& Info : ParentCell.Value)
+						// Because we support ExternalObjectStreamingGrid covering more than its source grid, we must handle this case.
+						// Clamp cell position on level grid borders
+						CellCoords.X = FMath::Clamp(CellCoords.X, 0, InjectedGridHelperLevel.GridSize);
+						CellCoords.Y = FMath::Clamp(CellCoords.Y, 0, InjectedGridHelperLevel.GridSize);
+					}
+					const int64 CoordKey = CellCoords.Y * InjectedGridHelperLevel.GridSize + CellCoords.X;
+
+					FSpatialHashStreamingGridLevel& TargetInjectedGridLevel = InjectedGridLevels[TargetGridLevel];
+					int32 LayerCellIndex;
+					int32* LayerCellIndexPtr = TargetInjectedGridLevel.LayerCellsMapping.Find(CoordKey);
+					if (LayerCellIndexPtr)
+					{
+						LayerCellIndex = *LayerCellIndexPtr;
+					}
+					else
+					{
+						LayerCellIndex = TargetInjectedGridLevel.LayerCells.AddDefaulted();
+						TargetInjectedGridLevel.LayerCellsMapping.Add(CoordKey, LayerCellIndex);
+					}
+
+					for (UWorldPartitionRuntimeSpatialHashCell* StreamingCell : ExternalObjectLayerCell.GridCells)
+					{
+						TargetInjectedGridLevel.LayerCells[LayerCellIndex].GridCells.Add(StreamingCell);
+					}
+				}
+			}
+		}
+	}
+}
+
+void FSpatialHashStreamingGrid::RemoveExternalStreamingObjectGrid(const FSpatialHashStreamingGrid& InExternalObjectStreamingGrid) const
+{
+	for (int SourceGridLevel = 0; SourceGridLevel < InExternalObjectStreamingGrid.GridLevels.Num(); ++SourceGridLevel)
+	{
+		const FSpatialHashStreamingGridLevel& ExternalObjectGridLevel = InExternalObjectStreamingGrid.GridLevels[SourceGridLevel];
+		int32 TargetGridLevel = FMath::Min<int32>(SourceGridLevel, GridLevels.Num() - 1);
+		if (ensure(InjectedGridLevels.IsValidIndex(TargetGridLevel)))
+		{
+			TArray<FSpatialHashStreamingGridLayerCell>& LayerCells = InjectedGridLevels[TargetGridLevel].LayerCells;
+			TMap<int64, int32>& LayerCellsMapping = InjectedGridLevels[TargetGridLevel].LayerCellsMapping;
+			// Prepare a reverse lookup LayerCellIndex -> CellKey
+			TMap<int32, int64> InverseLayerCellsMapping;
+			for (const auto& Pair : LayerCellsMapping)
+			{
+				InverseLayerCellsMapping.Add(Pair.Value, Pair.Key);
+			}
+
+			for (const FSpatialHashStreamingGridLayerCell& ExternalObjectLayerCell : ExternalObjectGridLevel.LayerCells)
+			{
+				for (UWorldPartitionRuntimeSpatialHashCell* StreamingCell : ExternalObjectLayerCell.GridCells)
+				{
+					bool bFound = false;
+					for (int LayerCellIndex = 0; LayerCellIndex < LayerCells.Num(); ++ LayerCellIndex)
+					{
+						FSpatialHashStreamingGridLayerCell& LayerCell = LayerCells[LayerCellIndex];
+						if (LayerCell.GridCells.RemoveSingleSwap(StreamingCell, false))
 						{
-							OutActivateCells.AddCell(Cell, Info.Source, Info.SourceShape);
+							bFound = true;
+							// Cleanup LayerCell if empty
+							if (LayerCell.GridCells.IsEmpty())
+							{
+								int32 RemovedIndex = LayerCellIndex;
+								LayerCells.RemoveAtSwap(RemovedIndex, 1, false);
+								int64 RemovedCellKey = InverseLayerCellsMapping[RemovedIndex];
+								LayerCellsMapping.Remove(RemovedCellKey);
+								InverseLayerCellsMapping.Remove(RemovedIndex);
+
+								if (LayerCells.IsValidIndex(RemovedIndex))
+								{
+									int32 SwapedLayerCellIndex = LayerCells.Num();
+									int64 SwapedCellKey = InverseLayerCellsMapping[SwapedLayerCellIndex];
+									LayerCellsMapping.FindChecked(SwapedCellKey) = RemovedIndex;
+									InverseLayerCellsMapping.Remove(SwapedLayerCellIndex);
+									InverseLayerCellsMapping.Add(RemovedIndex, SwapedCellKey);
+								}
+							}
+							break;
 						}
 					}
+					LayerCells.Shrink();
+					check(bFound);
 				}
 			}
 		}
 	}
 }
 
-const FSpatialHashStreamingGridLayerCell* FSpatialHashStreamingGrid::GetLayerCell(const FGridCellCoord& Coords) const
+void FSpatialHashStreamingGrid::ForEachLayerCell(const FGridCellCoord& Coords, TFunctionRef<void(const FSpatialHashStreamingGridLayerCell*)> Func) const
 {
 	check(GridLevels.IsValidIndex(Coords.Z));
-	if (const int32* LayerCellIndexPtr = GridLevels[Coords.Z].LayerCellsMapping.Find(Coords.Y * GetGridHelper().Levels[Coords.Z].GridSize + Coords.X))
+
+	const int64 CoordKey = Coords.Y * GetGridHelper().Levels[Coords.Z].GridSize + Coords.X;
+
+	if (const int32* LayerCellIndexPtr = GridLevels[Coords.Z].LayerCellsMapping.Find(CoordKey))
 	{
-		return &GridLevels[Coords.Z].LayerCells[*LayerCellIndexPtr];
+		Func(&GridLevels[Coords.Z].LayerCells[*LayerCellIndexPtr]);
+	}
+
+	// Test in the injected grid levels (if any)
+	if (InjectedGridLevels.IsValidIndex(Coords.Z))
+	{
+		if (const int32* LayerCellIndexPtr = InjectedGridLevels[Coords.Z].LayerCellsMapping.Find(CoordKey))
+		{
+			Func(&InjectedGridLevels[Coords.Z].LayerCells[*LayerCellIndexPtr]);
+		}
 	}
-	return nullptr;
+}
+
+void FSpatialHashStreamingGrid::ForEachRuntimeCell(const FGridCellCoord& Coords, TFunctionRef<void(const UWorldPartitionRuntimeCell*)> Func) const
+{
+	ForEachLayerCell(Coords, [&](const FSpatialHashStreamingGridLayerCell* LayerCell)
+	{
+		for (const UWorldPartitionRuntimeCell* Cell : LayerCell->GridCells)
+		{
+			Func(Cell);
+		}
+	});
 }
 
 void FSpatialHashStreamingGrid::GetAlwaysLoadedCells(const UDataLayerSubsystem* DataLayerSubsystem, TSet<const UWorldPartitionRuntimeCell*>& OutActivateCells, TSet<const UWorldPartitionRuntimeCell*>& OutLoadCells) const
@@ -395,28 +521,22 @@ void FSpatialHashStreamingGrid::GetAlwaysLoadedCells(const UDataLayerSubsystem*
 
 void FSpatialHashStreamingGrid::GetFilteredCellsForDebugDraw(const FSpatialHashStreamingGridLayerCell* LayerCell, const UDataLayerSubsystem* DataLayerSubsystem, TArray<const UWorldPartitionRuntimeCell*>& FilteredCells) const
 {
-	FilteredCells.Reset();
 	if (LayerCell)
 	{
 		Algo::CopyIf(LayerCell->GridCells, FilteredCells, [DataLayerSubsystem](const UWorldPartitionRuntimeCell* GridCell)
-		{ 
+		{
 			if (GridCell->IsDebugShown())
 			{
 				EStreamingStatus StreamingStatus = GridCell->GetStreamingStatus();
 				const TArray<FName>& DataLayers = GridCell->GetDataLayers();
-				return (!DataLayers.Num() || 
-						DataLayerSubsystem->IsAnyDataLayerInEffectiveRuntimeState(DataLayers, EDataLayerRuntimeState::Loaded) ||
-						DataLayerSubsystem->IsAnyDataLayerInEffectiveRuntimeState(DataLayers, EDataLayerRuntimeState::Activated) ||
-						((StreamingStatus != LEVEL_Unloaded) && (StreamingStatus != LEVEL_UnloadedButStillAround)));
+				return (!DataLayers.Num() ||
+					DataLayerSubsystem->IsAnyDataLayerInEffectiveRuntimeState(DataLayers, EDataLayerRuntimeState::Loaded) ||
+					DataLayerSubsystem->IsAnyDataLayerInEffectiveRuntimeState(DataLayers, EDataLayerRuntimeState::Activated) ||
+					((StreamingStatus != LEVEL_Unloaded) && (StreamingStatus != LEVEL_UnloadedButStillAround)));
 			}
 			return false;
 		});
 	}
-	if (FilteredCells.IsEmpty())
-	{
-		const UWorldPartitionRuntimeCell* DefaultEmptyCell = UWorldPartitionRuntimeCell::StaticClass()->GetDefaultObject<UWorldPartitionRuntimeCell>();
-		FilteredCells.Add(DefaultEmptyCell);
-	}
 }
 
 EWorldPartitionRuntimeCellVisualizeMode FSpatialHashStreamingGrid::GetStreamingCellVisualizeMode() const
@@ -428,6 +548,7 @@ EWorldPartitionRuntimeCellVisualizeMode FSpatialHashStreamingGrid::GetStreamingC
 void FSpatialHashStreamingGrid::Draw3D(UWorld* World, const TArray<FWorldPartitionStreamingSource>& Sources, const FTransform& Transform) const
 {
 	const EWorldPartitionRuntimeCellVisualizeMode VisualizeMode = GetStreamingCellVisualizeMode();
+	const UContentBundleManager* ContentBundleManager = World->ContentBundleManager;
 	const UDataLayerSubsystem* DataLayerSubsystem = World->GetSubsystem<UDataLayerSubsystem>();
 	TMap<FName, FColor> DataLayerDebugColors;
 	if (DataLayerSubsystem)
@@ -452,13 +573,14 @@ void FSpatialHashStreamingGrid::Draw3D(UWorld* World, const TArray<FWorldPartiti
 
 		for (int32 GridLevel = MinGridLevel; GridLevel <= MaxGridLevel; ++GridLevel)
 		{
-			Helper.Levels[GridLevel].ForEachIntersectingCells(Region, [&](const FGridCellCoord2& Coords)
+			Helper.Levels[GridLevel].ForEachIntersectingCells(Region, [&](const FGridCellCoord2& Coord2D)
 			{
-				const int32* LayerCellIndexPtr = GridLevels[GridLevel].LayerCellsMapping.Find(Coords.Y * Helper.Levels[GridLevel].GridSize + Coords.X);
-				const FSpatialHashStreamingGridLayerCell* LayerCell = LayerCellIndexPtr ? &GridLevels[GridLevel].LayerCells[*LayerCellIndexPtr] : nullptr;
-				GetFilteredCellsForDebugDraw(LayerCell, DataLayerSubsystem, FilteredCells);
-				check(FilteredCells.Num());
-
+				FGridCellCoord Coords(Coord2D.X, Coord2D.Y, GridLevel);
+				FilteredCells.Reset();
+				ForEachLayerCell(Coords, [&](const FSpatialHashStreamingGridLayerCell* LayerCell)
+				{
+					GetFilteredCellsForDebugDraw(LayerCell, DataLayerSubsystem, FilteredCells);
+				});
 				for (const UWorldPartitionRuntimeCell* Cell : FilteredCells)
 				{
 					bool bIsAlreadyInSet = false;
@@ -482,7 +604,7 @@ void FSpatialHashStreamingGrid::Draw3D(UWorld* World, const TArray<FWorldPartiti
 					{
 						FBox DataLayerColoredBox = CellBox;
 						double DataLayerSizeX = DataLayerColoredBox.GetSize().X / (5 * Cell->GetDataLayers().Num()); // Use 20% of cell's width
-						DataLayerColoredBox.Max.X = DataLayerColoredBox.Min.X + DataLayerSizeX; 
+						DataLayerColoredBox.Max.X = DataLayerColoredBox.Min.X + DataLayerSizeX;
 						FTranslationMatrix DataLayerOffsetMatrix(FVector(DataLayerSizeX, 0, 0));
 						for (const FName& DataLayer : Cell->GetDataLayers())
 						{
@@ -491,6 +613,21 @@ void FSpatialHashStreamingGrid::Draw3D(UWorld* World, const TArray<FWorldPartiti
 							DataLayerColoredBox = DataLayerColoredBox.TransformBy(DataLayerOffsetMatrix);
 						}
 					}
+
+					// Draw Content Bundle colored boxes
+					if (ContentBundleManager && FWorldPartitionDebugHelper::CanDrawContentBundles() && Cell->GetContentBundleID().IsValid())
+					{
+						if (const FContentBundleBase* ContentBundle = ContentBundleManager->GetContentBundle(World, Cell->GetContentBundleID()))
+						{
+							check(ContentBundle->GetDescriptor());
+							FBox ContentBundleColoredBox = CellBox;
+							double ContentBundleSizeX = ContentBundleColoredBox.GetSize().X / 5; // Use 20% of cell's width
+							ContentBundleColoredBox.Max.X = ContentBundleColoredBox.Min.X + ContentBundleSizeX;
+							FTranslationMatrix ContentBundleOffsetMatrix(FVector(ContentBundleColoredBox.GetSize().X - ContentBundleSizeX, 0, 0));
+							ContentBundleColoredBox = ContentBundleColoredBox.TransformBy(ContentBundleOffsetMatrix);
+							DrawDebugSolidBox(World, ContentBundleColoredBox, ContentBundle->GetDescriptor()->GetDebugColor(), Transform, false, -1.f, 255);
+						}
+					}
 				}
 			});
 		}
@@ -552,6 +689,7 @@ void FSpatialHashStreamingGrid::Draw2D(UCanvas* Canvas, UWorld* World, const TAr
 	TRACE_CPUPROFILER_EVENT_SCOPE(FSpatialHashStreamingGrid::Draw2D);
 
 	const EWorldPartitionRuntimeCellVisualizeMode VisualizeMode = GetStreamingCellVisualizeMode();
+	const UContentBundleManager* ContentBundleManager = World->ContentBundleManager;
 	const UDataLayerSubsystem* DataLayerSubsystem = World->GetSubsystem<UDataLayerSubsystem>();
 	TMap<FName, FColor> DataLayerDebugColors;
 	if (DataLayerSubsystem)
@@ -568,16 +706,17 @@ void FSpatialHashStreamingGrid::Draw2D(UCanvas* Canvas, UWorld* World, const TAr
 	const FString CellCoordString = UWorldPartitionRuntimeSpatialHash::GetCellCoordString(FGridCellCoord(88,88,88));
 	float MaxCellCoordTextWidth, MaxCellCoordTextHeight;
 	Canvas->StrLen(GEngine->GetTinyFont(), CellCoordString, MaxCellCoordTextWidth, MaxCellCoordTextHeight);
-
+	
 	TArray<const UWorldPartitionRuntimeCell*> FilteredCells;
 	for (int32 GridLevel = MinGridLevel; GridLevel <= MaxGridLevel; ++GridLevel)
 	{
 		// Draw Grid cells at desired grid level
 		const FSquare2DGridHelper& Helper = GetGridHelper();
-		Helper.Levels[GridLevel].ForEachIntersectingCells(Region, [&](const FGridCellCoord2& Coords)
+		Helper.Levels[GridLevel].ForEachIntersectingCells(Region, [&](const FGridCellCoord2& Coords2D)
 		{
 			FBox2D CellWorldBounds;
-			Helper.Levels[GridLevel].GetCellBounds(FGridCellCoord2(Coords.X, Coords.Y), CellWorldBounds);
+			FGridCellCoord Coords(Coords2D.X, Coords2D.Y, GridLevel);
+			Helper.Levels[GridLevel].GetCellBounds(Coords2D, CellWorldBounds);
 			FBox2D CellScreenBounds = FBox2D(WorldToScreen(CellWorldBounds.Min), WorldToScreen(CellWorldBounds.Max));
 			// Clamp inside grid bounds
 			if (!GridScreenBounds.IsInside(CellScreenBounds))
@@ -590,7 +729,7 @@ void FSpatialHashStreamingGrid::Draw2D(UCanvas* Canvas, UWorld* World, const TAr
 			else
 			{
 				FGridCellCoord CellGlobalCoords;
-				verify(Helper.GetCellGlobalCoords(FGridCellCoord(Coords.X, Coords.Y, GridLevel), CellGlobalCoords));
+				verify(Helper.GetCellGlobalCoords(Coords, CellGlobalCoords));
 				const FString CellCoordString = UWorldPartitionRuntimeSpatialHash::GetCellCoordString(CellGlobalCoords);
 				const FVector2D CellBoundsSize = CellScreenBounds.GetSize();
 				if (MaxCellCoordTextWidth < CellBoundsSize.X && MaxCellCoordTextHeight < CellBoundsSize.Y)
@@ -603,35 +742,53 @@ void FSpatialHashStreamingGrid::Draw2D(UCanvas* Canvas, UWorld* World, const TAr
 				}
 			}
 
-			const int32* LayerCellIndexPtr = GridLevels[GridLevel].LayerCellsMapping.Find(Coords.Y * Helper.Levels[GridLevel].GridSize + Coords.X);
-			const FSpatialHashStreamingGridLayerCell* LayerCell = LayerCellIndexPtr ? &GridLevels[GridLevel].LayerCells[*LayerCellIndexPtr] : nullptr;
-			GetFilteredCellsForDebugDraw(LayerCell, DataLayerSubsystem, FilteredCells);
-			check(FilteredCells.Num());
+			FilteredCells.Reset();
+			ForEachLayerCell(Coords, [&](const FSpatialHashStreamingGridLayerCell* LayerCell)
+			{
+				GetFilteredCellsForDebugDraw(LayerCell, DataLayerSubsystem, FilteredCells);
+			});
 
-			FVector2D CellBoundsSize = CellScreenBounds.GetSize();
-			CellBoundsSize.Y /= FilteredCells.Num();
-			FVector2D CellOffset(0.f, 0.f);
-			for (const UWorldPartitionRuntimeCell* Cell : FilteredCells)
+			if (FilteredCells.Num())
 			{
-				// Draw Cell using its debug color
-				FVector2D StartPos = CellScreenBounds.Min + CellOffset;
-				FCanvasTileItem Item(StartPos, GWhiteTexture, CellBoundsSize, Cell->GetDebugColor(VisualizeMode));
-				Item.BlendMode = SE_BLEND_Translucent;
-				Canvas->DrawItem(Item);
-				CellOffset.Y += CellBoundsSize.Y;
-
-				// Draw Cell's DataLayers colored boxes
-				if (DataLayerDebugColors.Num() && Cell->GetDataLayers().Num() > 0)
+				FVector2D CellBoundsSize = CellScreenBounds.GetSize();
+				CellBoundsSize.Y /= FilteredCells.Num();
+				FVector2D CellOffset(0.f, 0.f);
+				for (const UWorldPartitionRuntimeCell* Cell : FilteredCells)
 				{
-					FVector2D DataLayerOffset(0, 0);
-					FVector2D DataLayerColoredBoxSize = CellBoundsSize;
-					DataLayerColoredBoxSize.X /= (5 * Cell->GetDataLayers().Num()); // Use 20% of cell's width
-					for (const FName& DataLayer : Cell->GetDataLayers())
+					// Draw Cell using its debug color
+					FVector2D StartPos = CellScreenBounds.Min + CellOffset;
+					FCanvasTileItem Item(StartPos, GWhiteTexture, CellBoundsSize, Cell->GetDebugColor(VisualizeMode));
+					Item.BlendMode = SE_BLEND_Translucent;
+					Canvas->DrawItem(Item);
+					CellOffset.Y += CellBoundsSize.Y;
+
+					// Draw Cell's DataLayers colored boxes
+					if (DataLayerDebugColors.Num() && Cell->GetDataLayers().Num() > 0)
 					{
-						const FColor& DataLayerColor = DataLayerDebugColors[DataLayer];
-						FCanvasTileItem DataLayerItem(StartPos + DataLayerOffset, GWhiteTexture, DataLayerColoredBoxSize, DataLayerColor);
-						Canvas->DrawItem(DataLayerItem);
-						DataLayerOffset.X += DataLayerColoredBoxSize.X;
+						FVector2D DataLayerOffset(0, 0);
+						FVector2D DataLayerColoredBoxSize = CellBoundsSize;
+						DataLayerColoredBoxSize.X /= (5 * Cell->GetDataLayers().Num()); // Use 20% of cell's width
+						for (const FName& DataLayer : Cell->GetDataLayers())
+						{
+							const FColor& DataLayerColor = DataLayerDebugColors[DataLayer];
+							FCanvasTileItem DataLayerItem(StartPos + DataLayerOffset, GWhiteTexture, DataLayerColoredBoxSize, DataLayerColor);
+							Canvas->DrawItem(DataLayerItem);
+							DataLayerOffset.X += DataLayerColoredBoxSize.X;
+						}
+					}
+
+					// Draw Content Bundle colored boxes
+					if (ContentBundleManager && FWorldPartitionDebugHelper::CanDrawContentBundles() && Cell->GetContentBundleID().IsValid())
+					{
+						if (const FContentBundleBase* ContentBundle = ContentBundleManager->GetContentBundle(World, Cell->GetContentBundleID()))
+						{
+							check(ContentBundle->GetDescriptor());
+							FVector2D BoxSize = CellBoundsSize;
+							BoxSize.X /= 5; // Use 20% of cell's width
+							FVector2D ContentBundleOffset(CellBoundsSize.X - BoxSize.X, 0);
+							FCanvasTileItem ContentBundleItem(StartPos + ContentBundleOffset, GWhiteTexture, BoxSize, ContentBundle->GetDescriptor()->GetDebugColor());
+							Canvas->DrawItem(ContentBundleItem);
+						}
 					}
 				}
 			}
@@ -824,6 +981,7 @@ void UWorldPartitionRuntimeSpatialHash::PreSave(FObjectPreSaveContext ObjectSave
 	{
 		// We don't want this to be persisted but we can't set the property Transient as it is NonPIEDuplicateTransient and those flags aren't compatible
 		// If at some point GenerateStreaming is done after duplication we can remove this code.
+		bIsNameToGridMappingDirty = true;
 		StreamingGrids.Empty();
 	}
 }
@@ -846,20 +1004,17 @@ URuntimeHashExternalStreamingObjectBase* UWorldPartitionRuntimeSpatialHash::Stor
 	StreamingObject->Initialize(GetWorld(), GetTypedOuter<UWorld>());
 	StreamingObject->StreamingGrids = MoveTemp(StreamingGrids);
 
-	// Rename grids so the name of external streaming objects do not clashes with the main streaming grids.
 	for (FSpatialHashStreamingGrid& StreamingGrid : StreamingObject->StreamingGrids)
 	{
-		FString ExternalGridName = StreamingGrid.GridName.ToString() + TEXT("(") + StreamingObjectName.ToString() + TEXT(")");
-		StreamingGrid.GridName = FName(*ExternalGridName);
-
 		for (FSpatialHashStreamingGridLevel& GridLevel : StreamingGrid.GridLevels)
 		{
+			// External streaming grids do not need to keep LayerCellsMapping as they will inject their cells into their associated source grid (see InjectExternalStreamingObject)
+			GridLevel.LayerCellsMapping.Empty();
+
 			for (FSpatialHashStreamingGridLayerCell& GridLayerCell : GridLevel.LayerCells)
 			{
 				for (UWorldPartitionRuntimeSpatialHashCell* Cell : GridLayerCell.GridCells)
 				{
-					Cell->SetGridName(StreamingGrid.GridName);
-
 					// Do not dirty, otherwise it dirties the level previous outer (WorldPartition) which dirties the map. Occurs when entering PIE. 
 					// Do not reset loaders, the cell was just created it did not exists when loading initially.
 					Cell->Rename(nullptr, StreamingObject,  REN_DoNotDirty | REN_ForceNoResetLoaders);
@@ -927,6 +1082,7 @@ bool UWorldPartitionRuntimeSpatialHash::GenerateStreaming(UWorldPartitionStreami
 	}
 
 	// Fix case where StreamingGrids might have been persisted.
+	bIsNameToGridMappingDirty = true;
 	StreamingGrids.Empty();
 
 	// Append grids from ASpatialHashRuntimeGridInfo actors to runtime spatial hash grids
@@ -1322,6 +1478,7 @@ bool UWorldPartitionRuntimeSpatialHash::PopulateGeneratorPackageForCook(const TA
 
 void UWorldPartitionRuntimeSpatialHash::FlushStreaming()
 {
+	bIsNameToGridMappingDirty = true;
 	StreamingGrids.Empty();
 }
 
@@ -1419,7 +1576,7 @@ void UWorldPartitionRuntimeSpatialHash::ForEachStreamingCellsSources(const TArra
 				StreamingGrid.GetAlwaysLoadedCells(DataLayerSubsystem, ActivateStreamingSourceCells.GetCells(), LoadStreamingSourceCells.GetCells());
 			}
 		});
-}
+	}
 	else
 	{
 		// Get cells based on streaming sources
@@ -1455,10 +1612,26 @@ bool UWorldPartitionRuntimeSpatialHash::InjectExternalStreamingObject(URuntimeHa
 
 	if (!ExternalStreamingObjects.Contains(ExternalStreamingObject))
 	{
-		if (URuntimeSpatialHashExternalStreamingObject* SpatialHashStreamingObject = Cast<URuntimeSpatialHashExternalStreamingObject>(ExternalStreamingObject))
+		if (URuntimeSpatialHashExternalStreamingObject* SpatialHashExternalStreamingObject = Cast<URuntimeSpatialHashExternalStreamingObject>(ExternalStreamingObject))
 		{
-			ExternalStreamingObjects.Add(SpatialHashStreamingObject);
-			bIsNameToGridMappingDirty = true;
+			// Validate that there's a corresponding streaming grid for each streaming grid of this external streaming object
+			for (const FSpatialHashStreamingGrid& ExternalStreamingGrid : SpatialHashExternalStreamingObject->StreamingGrids)
+			{
+				if (!GetStreamingGridByName(ExternalStreamingGrid.GridName))
+				{
+					UE_LOG(LogWorldPartition, Error, TEXT("Failed to inject external streaming object %s, can't find matching streaming grid %s."),
+						*ExternalStreamingObject->GetName(), *ExternalStreamingGrid.GridName.ToString());
+					return false;
+				}
+			}
+
+			ExternalStreamingObjects.Add(SpatialHashExternalStreamingObject);
+
+			for (const FSpatialHashStreamingGrid& ExternalStreamingGrid : SpatialHashExternalStreamingObject->StreamingGrids)
+			{
+				const FSpatialHashStreamingGrid* SourceGrid = GetStreamingGridByName(ExternalStreamingGrid.GridName);
+				SourceGrid->InjectExternalStreamingObjectGrid(ExternalStreamingGrid);
+			}
 			return true;
 		}
 		else
@@ -1474,11 +1647,20 @@ bool UWorldPartitionRuntimeSpatialHash::RemoveExternalStreamingObject(URuntimeHa
 {
 	check(GetWorld()->IsGameWorld());
 
-	if (URuntimeSpatialHashExternalStreamingObject* SpatialHashStreamingObject = Cast<URuntimeSpatialHashExternalStreamingObject>(ExternalStreamingObject))
+	if (URuntimeSpatialHashExternalStreamingObject* SpatialHashExternalStreamingObject = Cast<URuntimeSpatialHashExternalStreamingObject>(ExternalStreamingObject))
 	{
-		if (ExternalStreamingObjects.Remove(SpatialHashStreamingObject) > 0)
+		if (ExternalStreamingObjects.Remove(SpatialHashExternalStreamingObject) > 0)
 		{
-			bIsNameToGridMappingDirty = true;
+			for (const FSpatialHashStreamingGrid& ExternalStreamingGrid : SpatialHashExternalStreamingObject->StreamingGrids)
+			{
+				// In PIE, UWorldPartition::OnEndPlay() will call UWorldPartitionRuntimeSpatialHash::FlushStreaming() 
+				// which will empty StreamingGrids before RemoveExternalStreamingObject is called.
+				// In this case, we don't need to remove injected cells as the source grid was removed.
+				if (const FSpatialHashStreamingGrid* SourceGrid = GetStreamingGridByName(ExternalStreamingGrid.GridName))
+				{
+					SourceGrid->RemoveExternalStreamingObjectGrid(ExternalStreamingGrid);
+				}
+			}
 			return true;
 		}
 	}
@@ -1493,19 +1675,7 @@ bool UWorldPartitionRuntimeSpatialHash::RemoveExternalStreamingObject(URuntimeHa
 
 uint32 UWorldPartitionRuntimeSpatialHash::GetNumGrids() const
 {
-	uint32 NumGrids = Algo::Accumulate(ExternalStreamingObjects, StreamingGrids.Num(),
-		[](uint32 AccumulatedNumGrids, const TWeakObjectPtr<URuntimeSpatialHashExternalStreamingObject>& StreamingObject)
-	{
-		if (StreamingObject.IsValid())
-		{
-			AccumulatedNumGrids += StreamingObject->StreamingGrids.Num();
-		}
-
-		return AccumulatedNumGrids;
-
-	});
-
-	return NumGrids;
+	return StreamingGrids.Num();
 }
 
 void UWorldPartitionRuntimeSpatialHash::ForEachStreamingGrid(TFunctionRef<void(FSpatialHashStreamingGrid&)> Func)
@@ -1514,17 +1684,6 @@ void UWorldPartitionRuntimeSpatialHash::ForEachStreamingGrid(TFunctionRef<void(F
 	{
 		Func(StreamingGrid);
 	}
-
-	for (TWeakObjectPtr<URuntimeSpatialHashExternalStreamingObject>& StreamingObject : ExternalStreamingObjects)
-	{
-		if (StreamingObject.IsValid())
-		{
-			for (FSpatialHashStreamingGrid& StreamingGrid : StreamingObject->StreamingGrids)
-			{
-				Func(StreamingGrid);
-			}
-		}
-	}
 }
 
 void UWorldPartitionRuntimeSpatialHash::ForEachStreamingGridBreakable(TFunctionRef<bool(const FSpatialHashStreamingGrid&)> Func) const
@@ -1536,68 +1695,41 @@ void UWorldPartitionRuntimeSpatialHash::ForEachStreamingGridBreakable(TFunctionR
 			return;
 		}
 	}
-
-	for (const TWeakObjectPtr<URuntimeSpatialHashExternalStreamingObject>& StreamingObject : ExternalStreamingObjects)
-	{
-		if (StreamingObject.IsValid())
-		{
-			for (FSpatialHashStreamingGrid& StreamingGrid : StreamingObject->StreamingGrids)
-			{
-				if (!Func(StreamingGrid))
-				{
-					return;
-				}
-			}
-		}
-	}
 }
 
-void UWorldPartitionRuntimeSpatialHash::ForEachStreamingGrid(TFunctionRef<void(const FSpatialHashStreamingGrid&)> Func, bool bIncludeExternalStreamingObjects) const
+void UWorldPartitionRuntimeSpatialHash::ForEachStreamingGrid(TFunctionRef<void(const FSpatialHashStreamingGrid&)> Func) const
 {
 	for (const FSpatialHashStreamingGrid& StreamingGrid : StreamingGrids)
 	{
 		Func(StreamingGrid);
 	}
-
-	if (bIncludeExternalStreamingObjects)
-	{
-		for (const TWeakObjectPtr<URuntimeSpatialHashExternalStreamingObject>& StreamingObject : ExternalStreamingObjects)
-		{
-			if (StreamingObject.IsValid())
-			{
-				for (const FSpatialHashStreamingGrid& StreamingGrid : StreamingObject->StreamingGrids)
-				{
-					Func(StreamingGrid);
-				}
-			}
-		}
-	}
 }
 
 const TMap<FName, const FSpatialHashStreamingGrid*>& UWorldPartitionRuntimeSpatialHash::GetNameToGridMapping() const
 {
 	if (bIsNameToGridMappingDirty)
 	{
+		NameToGridMapping.Reset();
 		ForEachStreamingGrid([&](const FSpatialHashStreamingGrid& StreamingGrid)
 		{
-			// Only keep the first one if names are clashing. ForEachStreamingGrids start iterating with the main grids then iterates on external Objects.
-			// Prioritize the main grid in the mapping.
-			if (!NameToGridMapping.Contains(StreamingGrid.GridName))
-			{
-				NameToGridMapping.Add(StreamingGrid.GridName, &StreamingGrid);
-			}
-			else
-			{
-				UE_LOG(LogWorldPartition, Error, TEXT("2 Streaming grids share the name: \"%s\". Some features will not work properly."), *StreamingGrid.GridName.ToString());
-			}
+			ensureMsgf(!NameToGridMapping.Contains(StreamingGrid.GridName), TEXT("2 Streaming grids share the name: \"%s\". Some features will not work properly."), *StreamingGrid.GridName.ToString());
+			NameToGridMapping.Add(StreamingGrid.GridName, &StreamingGrid);
 		});
-
 		bIsNameToGridMappingDirty = false;
 	}
 
 	return NameToGridMapping;
 }
 
+const FSpatialHashStreamingGrid* UWorldPartitionRuntimeSpatialHash::GetStreamingGridByName(FName InGridName) const
+{
+	if (const FSpatialHashStreamingGrid* const* StreamingGrid = GetNameToGridMapping().Find(InGridName))
+	{
+		return *StreamingGrid;
+	}
+	return nullptr;
+}
+
 EWorldPartitionStreamingPerformance UWorldPartitionRuntimeSpatialHash::GetStreamingPerformanceForCell(const UWorldPartitionRuntimeCell* Cell) const
 {
 	// If base class already returning critical. Early out.
@@ -1612,24 +1744,25 @@ EWorldPartitionStreamingPerformance UWorldPartitionRuntimeSpatialHash::GetStream
 	const UWorldPartitionRuntimeSpatialHashCell* StreamingCell = Cast<const UWorldPartitionRuntimeSpatialHashCell>(Cell);
 	check(StreamingCell);
 
-	const FSpatialHashStreamingGrid* const* StreamingGrid = GetNameToGridMapping().Find(StreamingCell->GetGridName());
-	check(StreamingGrid && *StreamingGrid);
-	
-	const float LoadingRange = (*StreamingGrid)->LoadingRange;
-
-	if (StreamingCell->CachedIsBlockingSource)
+	const FSpatialHashStreamingGrid* StreamingGrid = GetStreamingGridByName(StreamingCell->GetGridName());
+	if (ensure(StreamingGrid))
 	{
-		const double Distance = FMath::Sqrt(StreamingCell->CachedMinSquareDistanceToBlockingSource) - ((double)(*StreamingGrid)->GetCellSize(StreamingCell->Level) / 2);
-
-		const double Ratio = Distance / LoadingRange;
+		const float LoadingRange = StreamingGrid->LoadingRange;
 
-		if (Ratio < GBlockOnSlowStreamingRatio)
-		{
-			return EWorldPartitionStreamingPerformance::Critical;
-		}
-		else if (Ratio < BlockOnSlowStreamingWarningRatio)
+		if (StreamingCell->CachedIsBlockingSource)
 		{
-			return EWorldPartitionStreamingPerformance::Slow;
+			const double Distance = FMath::Sqrt(StreamingCell->CachedMinSquareDistanceToBlockingSource) - ((double)StreamingGrid->GetCellSize(StreamingCell->Level) / 2);
+
+			const double Ratio = Distance / LoadingRange;
+
+			if (Ratio < GBlockOnSlowStreamingRatio)
+			{
+				return EWorldPartitionStreamingPerformance::Critical;
+			}
+			else if (Ratio < BlockOnSlowStreamingWarningRatio)
+			{
+				return EWorldPartitionStreamingPerformance::Slow;
+			}
 		}
 	}
 
@@ -1767,7 +1900,7 @@ TArray<const FSpatialHashStreamingGrid*> UWorldPartitionRuntimeSpatialHash::GetF
 				FilteredStreamingGrids.Add(&StreamingGrid);
 			}
 		}
-	}, !FWorldPartitionDebugHelper::GetDebugFilterOutContentBundles());
+	});
 
 	return FilteredStreamingGrids;
 }
diff --git a/Engine/Source/Runtime/Engine/Private/WorldPartition/WorldPartitionSubsystem.cpp b/Engine/Source/Runtime/Engine/Private/WorldPartition/WorldPartitionSubsystem.cpp
--- a/Engine/Source/Runtime/Engine/Private/WorldPartition/WorldPartitionSubsystem.cpp
+++ b/Engine/Source/Runtime/Engine/Private/WorldPartition/WorldPartitionSubsystem.cpp
@@ -5,6 +5,7 @@
 #include "WorldPartition/WorldPartitionStreamingPolicy.h"
 #include "WorldPartition/WorldPartitionDebugHelper.h"
 #include "WorldPartition/DataLayer/DataLayerSubsystem.h"
+#include "WorldPartition/ContentBundle/ContentBundleWorldSubsystem.h"
 #include "HAL/IConsoleManager.h"
 #include "Engine/Canvas.h"
 #include "Engine/Console.h"
@@ -441,19 +442,24 @@ void UWorldPartitionSubsystem::Draw(UCanvas* Canvas, class APlayerController* PC
 		}
 	}
 
-	UDataLayerSubsystem* DataLayerSubsystem = WorldPartition->GetWorld()->GetSubsystem<UDataLayerSubsystem>();
-
 	if (GDrawLegends || GDrawRuntimeHash2D)
 	{
 		// Streaming Status Legend
 		WorldPartition->DrawStreamingStatusLegend(Canvas, CurrentOffset);
 	}
 
+	UDataLayerSubsystem* DataLayerSubsystem = WorldPartition->GetWorld()->GetSubsystem<UDataLayerSubsystem>();
 	if (DataLayerSubsystem && (GDrawDataLayers || GDrawDataLayersLoadTime || GDrawRuntimeHash2D))
 	{
 		DataLayerSubsystem->DrawDataLayersStatus(Canvas, CurrentOffset);
 	}
 
+	UContentBundleManager* ContentBundleManager = GetWorld()->ContentBundleManager;
+	if (ContentBundleManager && (FWorldPartitionDebugHelper::CanDrawContentBundles() && GDrawRuntimeHash2D))
+	{
+		ContentBundleManager->DrawContentBundlesStatus(GetWorld(), Canvas, CurrentOffset);
+	}
+
 	if (GDrawRuntimeCellsDetails)
 	{
 		WorldPartition->DrawRuntimeCellsDetails(Canvas, CurrentOffset);
diff --git a/Engine/Source/Runtime/Engine/Public/WorldPartition/ContentBundle/ContentBundleBase.h b/Engine/Source/Runtime/Engine/Public/WorldPartition/ContentBundle/ContentBundleBase.h
--- a/Engine/Source/Runtime/Engine/Public/WorldPartition/ContentBundle/ContentBundleBase.h
+++ b/Engine/Source/Runtime/Engine/Public/WorldPartition/ContentBundle/ContentBundleBase.h
@@ -26,6 +26,7 @@ public:
 
 	UWorld* GetInjectedWorld() const;
 	const FString& GetDisplayName() const;
+	const FColor& GetDebugColor() const;
 	const TWeakPtr<FContentBundleClient>& GetClient() const;
 	const UContentBundleDescriptor* GetDescriptor() const;
 	EContentBundleStatus GetStatus() const { return Status; }
diff --git a/Engine/Source/Runtime/Engine/Public/WorldPartition/ContentBundle/ContentBundleDescriptor.h b/Engine/Source/Runtime/Engine/Public/WorldPartition/ContentBundle/ContentBundleDescriptor.h
--- a/Engine/Source/Runtime/Engine/Public/WorldPartition/ContentBundle/ContentBundleDescriptor.h
+++ b/Engine/Source/Runtime/Engine/Public/WorldPartition/ContentBundle/ContentBundleDescriptor.h
@@ -9,6 +9,7 @@
 
 class UActorDescContainer;
 class AWorldDataLayers;
+struct FColor;
 
 UCLASS()
 class ENGINE_API UContentBundleDescriptor : public UObject
@@ -16,21 +17,33 @@ class ENGINE_API UContentBundleDescriptor : public UObject
 	GENERATED_UCLASS_BODY()
 
 public:
-	const FString& GetDisplayName() const;
+	const FString& GetDisplayName() const { return DisplayName; }
+	const FColor& GetDebugColor() const { return DebugColor; }
 	const FString& GetPackageRoot() const { return PackageRoot; }
 	const FGuid& GetGuid() const { return Guid; }
 
 #if WITH_EDITOR
+	void InitializeObject(const FString& InContentBundleName, const FString& InPackageRoot);
 	void SetDisplayName(const FString& InName) { DisplayName = InName; }
 	void SetPackageRoot(const FString& InPackageRoot) { PackageRoot = InPackageRoot; }
+	void SetDebugColor(const FColor& InDebugColor) { DebugColor = InDebugColor; }
+
+	//~ Begin UObject Interface
+	virtual void PostLoad() override;
+	//~ End UObject Interface
 #endif
 
 	bool IsValid() const;
 
 private:
+	void InitDebugColor();
+
 	UPROPERTY(EditAnywhere, Category = BaseInformation)
 	FString DisplayName;
 
+	UPROPERTY(VisibleAnywhere, Category = BaseInformation)
+	FColor DebugColor;
+
 	UPROPERTY(VisibleAnywhere, Category = BaseInformation, AdvancedDisplay)
 	FGuid Guid;
 
diff --git a/Engine/Source/Runtime/Engine/Public/WorldPartition/ContentBundle/ContentBundleWorldSubsystem.h b/Engine/Source/Runtime/Engine/Public/WorldPartition/ContentBundle/ContentBundleWorldSubsystem.h
--- a/Engine/Source/Runtime/Engine/Public/WorldPartition/ContentBundle/ContentBundleWorldSubsystem.h
+++ b/Engine/Source/Runtime/Engine/Public/WorldPartition/ContentBundle/ContentBundleWorldSubsystem.h
@@ -14,6 +14,7 @@ class UContentBundleDescriptor;
 class FContentBundleClient;
 class FContentBundleEditor;
 class UWorldPartition;
+class UCanvas;
 
 #if WITH_EDITOR
 class URuntimeHashExternalStreamingObjectBase;
@@ -40,9 +41,12 @@ public:
 	UContentBundleDuplicateForPIEHelper* GetPIEDuplicateHelper() const { return PIEDuplicateHelper; }
 #endif
 
+	const FContentBundleBase* GetContentBundle(const UWorld* InWorld, const FGuid& Guid) const;
+	void DrawContentBundlesStatus(const UWorld* InWorld, UCanvas* Canvas, FVector2D& Offset) const;
+
 private:
-	uint32 GetContentBundleContainerIndex(const UWorld* InjectedWorld);
-	TUniquePtr<FContentBundleContainer>* GetContentBundleContainer(const UWorld* InjectedWorld);
+	uint32 GetContentBundleContainerIndex(const UWorld* InjectedWorld) const;
+	const TUniquePtr<FContentBundleContainer>* GetContentBundleContainer(const UWorld* InjectedWorld) const;
 
 	void OnWorldPartitionInitialized(UWorldPartition* WorldPartition);
 	void OnWorldPartitionUninitialized(UWorldPartition* WorldPartition);
diff --git a/Engine/Source/Runtime/Engine/Public/WorldPartition/WorldPartitionDebugHelper.h b/Engine/Source/Runtime/Engine/Public/WorldPartition/WorldPartitionDebugHelper.h
--- a/Engine/Source/Runtime/Engine/Public/WorldPartition/WorldPartitionDebugHelper.h
+++ b/Engine/Source/Runtime/Engine/Public/WorldPartition/WorldPartitionDebugHelper.h
@@ -24,7 +24,7 @@ public:
 	static void DrawLegendItem(UCanvas* Canvas, const FString& Text, const UFont* Font, const FColor& Color, const FColor& TextColor, FVector2D& Pos, float* MaxItemWidth = nullptr);
 	static FLinearColor GetHeatMapColor(float Value);
 	FORCEINLINE static bool IsRuntimeSpatialHashCellStreamingPriorityShown() { return bShowRuntimeSpatialHashCellStreamingPriority; }
-	FORCEINLINE static bool GetDebugFilterOutContentBundles() { return bDebugFilterOutContentBundles; }
+	FORCEINLINE static bool CanDrawContentBundles() { return bCanDrawContentBundles; }
 
 private:
 	static FAutoConsoleCommand DebugFilterByStreamingStatusCommand;
@@ -42,6 +42,6 @@ private:
 	static FAutoConsoleVariableRef ShowRuntimeSpatialHashCellStreamingPriorityCommand;
 	static bool bShowRuntimeSpatialHashCellStreamingPriority;
 
-	static FAutoConsoleVariableRef DebugFilterOutContentBundlesCommand;
-	static bool bDebugFilterOutContentBundles;
+	static FAutoConsoleVariableRef DrawContentBundlesCommand;
+	static bool bCanDrawContentBundles;
 };
\ No newline at end of file
diff --git a/Engine/Source/Runtime/Engine/Public/WorldPartition/WorldPartitionRuntimeCell.h b/Engine/Source/Runtime/Engine/Public/WorldPartition/WorldPartitionRuntimeCell.h
--- a/Engine/Source/Runtime/Engine/Public/WorldPartition/WorldPartitionRuntimeCell.h
+++ b/Engine/Source/Runtime/Engine/Public/WorldPartition/WorldPartitionRuntimeCell.h
@@ -212,7 +212,6 @@ class ENGINE_API UWorldPartitionRuntimeCell : public UObject, public IWorldParti
 	void SetContentBundleUID(const FGuid& InContentBundleID) { ContentBundleID = InContentBundleID; }
 	void SetContentBounds(const FBox& InBounds) { ContentBounds = InBounds; }
 	void SetDebugInfo(int64 InCoordX, int64 InCoordY, int64 InCoordZ, FName InGridName);
-	void SetGridName(FName InGridName);
 	virtual void AddActorToCell(const FWorldPartitionActorDescView& ActorDescView, const FActorContainerID& InContainerID, const FTransform& InContainerTransform, const UActorDescContainer* InContainer) PURE_VIRTUAL(UWorldPartitionRuntimeCell::AddActorToCell,);
 	virtual int32 GetActorCount() const PURE_VIRTUAL(UWorldPartitionRuntimeCell::GetActorCount, return 0;);
 
diff --git a/Engine/Source/Runtime/Engine/Public/WorldPartition/WorldPartitionRuntimeSpatialHash.h b/Engine/Source/Runtime/Engine/Public/WorldPartition/WorldPartitionRuntimeSpatialHash.h
--- a/Engine/Source/Runtime/Engine/Public/WorldPartition/WorldPartitionRuntimeSpatialHash.h
+++ b/Engine/Source/Runtime/Engine/Public/WorldPartition/WorldPartitionRuntimeSpatialHash.h
@@ -136,6 +136,9 @@ struct FSpatialHashStreamingGrid
 		return LoadingRange;
 	}
 
+	void InjectExternalStreamingObjectGrid(const FSpatialHashStreamingGrid& InExternalObjectStreamingGrid) const;
+	void RemoveExternalStreamingObjectGrid(const FSpatialHashStreamingGrid& InExternalObjectStreamingGrid) const;
+
 	// Used by PIE/Game
 	int64 GetCellSize(int32 Level) const;
 	void GetCells(const FWorldPartitionStreamingQuerySource& QuerySource, TSet<const UWorldPartitionRuntimeCell*>& OutCells, bool bEnableZCulling) const;
@@ -145,13 +148,17 @@ struct FSpatialHashStreamingGrid
 	void Draw3D(UWorld* World, const TArray<FWorldPartitionStreamingSource>& Sources, const FTransform& Transform) const;
 
 private:
-	const FSpatialHashStreamingGridLayerCell* GetLayerCell(const FGridCellCoord& Coords) const;
+	void ForEachRuntimeCell(const FGridCellCoord& Coords, TFunctionRef<void(const UWorldPartitionRuntimeCell*)> Func) const;
+	void ForEachLayerCell(const FGridCellCoord& Coords, TFunctionRef<void(const FSpatialHashStreamingGridLayerCell*)> Func) const;
 	void DrawStreamingSource2D(UCanvas* Canvas, const FSphericalSector& Shape, TFunctionRef<FVector2D(const FVector2D&)> WorldToScreen, const FColor& Color) const;
 	void DrawStreamingSource3D(UWorld* World, const FSphericalSector& Shape, const FTransform& Transform, const FColor& Color) const;
 	void GetFilteredCellsForDebugDraw(const FSpatialHashStreamingGridLayerCell* LayerCell, const UDataLayerSubsystem* DataLayerSubsystem, TArray<const UWorldPartitionRuntimeCell*>& FilteredCells) const;
 	const FSquare2DGridHelper& GetGridHelper() const;
 	EWorldPartitionRuntimeCellVisualizeMode GetStreamingCellVisualizeMode() const;
 	mutable FSquare2DGridHelper* GridHelper;
+
+	// Contains cells injected at runtime from content bundles
+	mutable TArray<FSpatialHashStreamingGridLevel> InjectedGridLevels;
 };
 
 // Disable copy of this USTRUCT by the autogenerated code.
@@ -287,9 +294,11 @@ public:
 	static FString GetCellCoordString(const FGridCellCoord& InCellGlobalCoord);
 
 protected:
+
+	const FSpatialHashStreamingGrid* GetStreamingGridByName(FName InGridName) const;
 	void ForEachStreamingGrid(TFunctionRef<void(FSpatialHashStreamingGrid&)> Func);
 	void ForEachStreamingGridBreakable(TFunctionRef<bool(const FSpatialHashStreamingGrid&)> Func) const;
-	void ForEachStreamingGrid(TFunctionRef<void(const FSpatialHashStreamingGrid&)> Func, bool bIncludeExternalStreamingObjects = true) const;
+	void ForEachStreamingGrid(TFunctionRef<void(const FSpatialHashStreamingGrid&)> Func) const;
 
 	virtual EWorldPartitionStreamingPerformance GetStreamingPerformanceForCell(const UWorldPartitionRuntimeCell* Cell) const override;
 
