Skip renaming cooked meta-data that has been loaded from a PKG_Cooked package

IO store currently doesn't support renaming these objects and it can cause issues if something tries to reload the package again later (eg, to load the BP rather than its BPGC)

#rnx

diff --git a/Engine/Source/Runtime/Engine/Private/CookedMetaData.cpp b/Engine/Source/Runtime/Engine/Private/CookedMetaData.cpp
--- a/Engine/Source/Runtime/Engine/Private/CookedMetaData.cpp
+++ b/Engine/Source/Runtime/Engine/Private/CookedMetaData.cpp
@@ -171,3 +171,18 @@ void UClassCookedMetaData::ApplyMetaData(UClass* TargetClass) const
 	}
 }
 
+namespace CookedMetaDataUtil::Internal
+{
+
+void RenameCookedMetaDataForPurge(UObject* CookedMetaDataPtr)
+{
+	// Skip the rename for cooked packages, as IO store cannot currently handle renames
+	if (!CookedMetaDataPtr->GetPackage()->HasAnyPackageFlags(PKG_Cooked))
+	{
+		FNameBuilder BaseMetaDataName(CookedMetaDataPtr->GetFName());
+		BaseMetaDataName << TEXT("_PURGED");
+		CookedMetaDataPtr->Rename(FNameBuilder(MakeUniqueObjectName(CookedMetaDataPtr->GetOuter(), CookedMetaDataPtr->GetClass(), FName(BaseMetaDataName))).ToString(), nullptr, REN_DoNotDirty | REN_DontCreateRedirectors | REN_ForceNoResetLoaders | REN_NonTransactional);
+	}
+}
+
+} // namespace CookedMetaDataUtil::Internal
diff --git a/Engine/Source/Runtime/Engine/Public/CookedMetaData.h b/Engine/Source/Runtime/Engine/Public/CookedMetaData.h
--- a/Engine/Source/Runtime/Engine/Public/CookedMetaData.h
+++ b/Engine/Source/Runtime/Engine/Public/CookedMetaData.h
@@ -123,6 +123,11 @@ protected:
 namespace CookedMetaDataUtil
 {
 
+namespace Internal
+{
+ENGINE_API void RenameCookedMetaDataForPurge(UObject* CookedMetaDataPtr);
+}
+
 template <typename CookedMetaDataType>
 CookedMetaDataType* NewCookedMetaData(UObject* Outer, FName Name, TSubclassOf<CookedMetaDataType> Class = CookedMetaDataType::StaticClass())
 {
@@ -138,11 +143,8 @@ CookedMetaDataType* FindCookedMetaData(UObject* Outer, const TCHAR* Name)
 template <typename CookedMetaDataType, typename CookedMetaDataPtrType>
 void PurgeCookedMetaData(CookedMetaDataPtrType& CookedMetaDataPtr)
 {
-	FNameBuilder BaseMetaDataName(CookedMetaDataPtr->GetFName());
-	BaseMetaDataName << TEXT("_PURGED");
-
+	Internal::RenameCookedMetaDataForPurge(CookedMetaDataPtr);
 	CookedMetaDataPtr->ClearFlags(RF_Standalone | RF_Public);
-	CookedMetaDataPtr->Rename(FNameBuilder(MakeUniqueObjectName(CookedMetaDataPtr->GetOuter(), CookedMetaDataPtr->GetClass(), FName(BaseMetaDataName))).ToString(), nullptr, REN_DoNotDirty | REN_DontCreateRedirectors | REN_ForceNoResetLoaders | REN_NonTransactional);
 	CookedMetaDataPtr = nullptr;
 }
 
