Fix for FMaterialRenderProxy::EvaluateUniformExpressions() being called every frame (when in editor) and not caching the results. Fixed by removing FUniformExpressionCache::bUpToDate and always checking CachedUniformExpressionShaderMap instead.

#rnx

diff --git a/Engine/Source/Runtime/Engine/Private/Materials/MaterialShared.cpp b/Engine/Source/Runtime/Engine/Private/Materials/MaterialShared.cpp
--- a/Engine/Source/Runtime/Engine/Private/Materials/MaterialShared.cpp
+++ b/Engine/Source/Runtime/Engine/Private/Materials/MaterialShared.cpp
@@ -4021,7 +4021,6 @@ void FMaterialRenderProxy::EvaluateUniformExpressions(FUniformExpressionCache& O
 	OutUniformExpressionCache.ParameterCollections = UniformExpressionSet.ParameterCollections;
 
 	++UniformExpressionCacheSerialNumber;
-	OutUniformExpressionCache.bUpToDate = Context.Material.IsRenderingThreadShaderMapComplete();
 }
 
 void FMaterialRenderProxy::CacheUniformExpressions(bool bRecreateUniformBuffer)
@@ -4083,7 +4082,6 @@ void FMaterialRenderProxy::InvalidateUniformExpressionCache(bool bRecreateUnifor
 	++UniformExpressionCacheSerialNumber;
 	for (int32 i = 0; i < ERHIFeatureLevel::Num; ++i)
 	{
-		UniformExpressionCache[i].bUpToDate = false;
 		UniformExpressionCache[i].CachedUniformExpressionShaderMap = nullptr;
 		UniformExpressionCache[i].ResetAllocatedVTs();
 
@@ -4099,16 +4097,15 @@ void FMaterialRenderProxy::InvalidateUniformExpressionCache(bool bRecreateUnifor
 
 void FMaterialRenderProxy::UpdateUniformExpressionCacheIfNeeded(ERHIFeatureLevel::Type InFeatureLevel) const
 {
-	if (!UniformExpressionCache[InFeatureLevel].bUpToDate)
+	// Don't cache uniform expressions if an entirely different FMaterialRenderProxy is going to be used for rendering
+	const FMaterial* Material = GetMaterialNoFallback(InFeatureLevel);
+	// Note: We would actually need to compare the FMaterialShaderMapId of both shader maps but a simple pointer compare also works
+	// because shader maps are currently swapped out whenever they are modified.
+	if (Material && Material->GetRenderingThreadShaderMap() != UniformExpressionCache[InFeatureLevel].CachedUniformExpressionShaderMap)
 	{
-		// Don't cache uniform expressions if an entirely different FMaterialRenderProxy is going to be used for rendering
-		const FMaterial* Material = GetMaterialNoFallback(InFeatureLevel);
-		if (Material)
-		{
-			FMaterialRenderContext MaterialRenderContext(this, *Material, nullptr);
-			MaterialRenderContext.bShowSelection = GIsEditor;
-			EvaluateUniformExpressions(UniformExpressionCache[InFeatureLevel], MaterialRenderContext);
-		}
+		FMaterialRenderContext MaterialRenderContext(this, *Material, nullptr);
+		MaterialRenderContext.bShowSelection = GIsEditor;
+		EvaluateUniformExpressions(UniformExpressionCache[InFeatureLevel], MaterialRenderContext);
 	}
 }
 
diff --git a/Engine/Source/Runtime/Engine/Public/MaterialShared.h b/Engine/Source/Runtime/Engine/Public/MaterialShared.h
--- a/Engine/Source/Runtime/Engine/Public/MaterialShared.h
+++ b/Engine/Source/Runtime/Engine/Public/MaterialShared.h
@@ -2453,10 +2453,8 @@ struct FUniformExpressionCache
 	TArray<IAllocatedVirtualTexture*> OwnedAllocatedVTs;
 	/** Ids of parameter collections needed for rendering. */
 	TArray<FGuid> ParameterCollections;
-	/** Shader map that was used to cache uniform expressions on this material.  This is used for debugging and verifying correct behavior. */
+	/** Shader map that was used to cache uniform expressions on this material. This is used for debugging, verifying correct behavior and checking if the cache is up to date. */
 	const FMaterialShaderMap* CachedUniformExpressionShaderMap = nullptr;
-	/** True if the cache is up to date. */
-	bool bUpToDate = false;
 
 	/** Destructor. */
 	ENGINE_API ~FUniformExpressionCache();
