World Partition - HLOD: Properly destroy temp world created UWorldPartitionLevelStreamingDynamic::LoadInEditor()

diff --git a/Engine/Source/Runtime/Engine/Private/WorldPartition/WorldPartitionLevelStreamingDynamic.cpp b/Engine/Source/Runtime/Engine/Private/WorldPartition/WorldPartitionLevelStreamingDynamic.cpp
--- a/Engine/Source/Runtime/Engine/Private/WorldPartition/WorldPartitionLevelStreamingDynamic.cpp
+++ b/Engine/Source/Runtime/Engine/Private/WorldPartition/WorldPartitionLevelStreamingDynamic.cpp
@@ -111,7 +111,14 @@ void UWorldPartitionLevelStreamingDynamic::UnloadFromEditor(UWorldPartitionLevel
 	InLevelStreaming->SetIsRequestingUnloadAndRemoval(true);
 	World->RemoveLevel(Level);
 	World->FlushLevelStreaming();
-	
+
+	// Destroy the package world and remove it from root
+	UWorld* PackageWorld = UWorld::FindWorldInPackage(Level->GetPackage());
+	if (PackageWorld)
+	{
+		PackageWorld->DestroyWorld(false);
+		PackageWorld->Rename(nullptr, GetTransientPackage());
+	}
 }
 
 void UWorldPartitionLevelStreamingDynamic::Initialize(UWorld* OuterWorld, const TArray<FWorldPartitionRuntimeCellObjectMapping>& InPackages)
