Implemented 'stat renderscaling' to show current status of the dynamic resolution heuristics

diff --git a/Engine/Source/Runtime/Core/Public/Stats/Stats2.h b/Engine/Source/Runtime/Core/Public/Stats/Stats2.h
--- a/Engine/Source/Runtime/Core/Public/Stats/Stats2.h
+++ b/Engine/Source/Runtime/Core/Public/Stats/Stats2.h
@@ -2404,6 +2404,7 @@ DECLARE_STATS_GROUP(TEXT("RHI"), STATGROUP_RHI, STATCAT_Advanced);
 DECLARE_STATS_GROUP(TEXT("RDG"), STATGROUP_RDG, STATCAT_Advanced);
 DECLARE_STATS_GROUP(TEXT("Render Thread"),STATGROUP_RenderThreadProcessing, STATCAT_Advanced);
 DECLARE_STATS_GROUP(TEXT("Render Target Pool"), STATGROUP_RenderTargetPool, STATCAT_Advanced);
+DECLARE_STATS_GROUP(TEXT("Render Scaling"), STATGROUP_RenderScaling, STATCAT_Advanced);
 DECLARE_STATS_GROUP(TEXT("Scene Memory"),STATGROUP_SceneMemory, STATCAT_Advanced);
 DECLARE_STATS_GROUP(TEXT("Scene Rendering"),STATGROUP_SceneRendering, STATCAT_Advanced);
 DECLARE_STATS_GROUP(TEXT("Scene Update"),STATGROUP_SceneUpdate, STATCAT_Advanced);
diff --git a/Engine/Source/Runtime/Engine/Private/DynamicResolution.cpp b/Engine/Source/Runtime/Engine/Private/DynamicResolution.cpp
--- a/Engine/Source/Runtime/Engine/Private/DynamicResolution.cpp
+++ b/Engine/Source/Runtime/Engine/Private/DynamicResolution.cpp
@@ -71,10 +71,10 @@ static TAutoConsoleVariable<float> CVarChangeThreshold(
 	TEXT("Minimal increase percentage threshold to alow when changing resolution."),
 	ECVF_RenderThreadSafe | ECVF_Default);
 
-static TAutoConsoleVariable<int32> CVarMaxConsecutiveOverbudgetGPUFrameCount(
-	TEXT("r.DynamicRes.MaxConsecutiveOverbudgetGPUFrameCount"),
+static TAutoConsoleVariable<int32> CVarMaxConsecutiveOverBudgetGPUFrameCount(
+	TEXT("r.DynamicRes.MaxConsecutiveOverBudgetGPUFrameCount"),
 	2,
-	TEXT("Maximum number of consecutive frame tolerated over GPU budget."),
+	TEXT("Maximum number of consecutive frames tolerated over GPU budget."),
 	ECVF_RenderThreadSafe | ECVF_Default);
 
 static TAutoConsoleVariable<int32> CVarTimingMeasureModel(
@@ -213,11 +213,11 @@ void FDynamicResolutionHeuristicProxy::CommitPreviousFrameGPUTimings_RenderThrea
 	History[EntryId].BudgetTimingMs = BudgetTimingMs;
 }
 
-void FDynamicResolutionHeuristicProxy::RefreshCurentFrameResolutionFraction_RenderThread()
+void FDynamicResolutionHeuristicProxy::RefreshCurrentFrameResolutionFraction_RenderThread()
 {
 	// Global constants.
 	const float FrameWeightExponent = CVarFrameWeightExponent.GetValueOnRenderThread();
-	const int32 MaxConsecutiveOverbudgetGPUFrameCount = FMath::Max(CVarMaxConsecutiveOverbudgetGPUFrameCount.GetValueOnRenderThread(), 2);
+	const int32 MaxConsecutiveOverBudgetGPUFrameCount = FMath::Max(CVarMaxConsecutiveOverBudgetGPUFrameCount.GetValueOnRenderThread(), 2);
 
 	const bool bCanChangeResolution = NumberOfFramesSinceScreenPercentageChange >= CVarFrameChangePeriod.GetValueOnRenderThread();
 
@@ -225,7 +225,7 @@ void FDynamicResolutionHeuristicProxy::RefreshCurentFrameResolutionFraction_Rend
 	DynamicRenderScaling::TMap<float> NewFrameResolutionFractions = CurrentFrameResolutionFractions;
 
 	// Whether there is a GPU over budget panic.
-	bool bGlobalGPUOverbugetPanic = false;
+	bool bGlobalGPUOverBudgetPanic = false;
 
 	for (TLinkedList<DynamicRenderScaling::FBudget*>::TIterator BudgetIt(DynamicRenderScaling::FBudget::GetGlobalList()); BudgetIt; BudgetIt.Next())
 	{
@@ -260,15 +260,15 @@ void FDynamicResolutionHeuristicProxy::RefreshCurentFrameResolutionFraction_Rend
 		float Weight = 1.0f;
 
 		// Number of consecutive frames that have over budget GPU.
-		int32 ConsecutiveOverbudgetGPUFramCount = 0;
+		int32 ConsecutiveOverBudgetGPUFrameCount = 0;
 
 		// Whether should continue browsing the frame history.
-		bool bCaryOnBrowsingFrameHistory = true;
+		bool bCarryOnBrowsingFrameHistory = true;
 
 		// Number of frame browsed.
 		int32 FrameCount = 0;
 		int32 BrowsingFrameId = 0;
-		for (; BrowsingFrameId < BudgetHistorySize && bCaryOnBrowsingFrameHistory; BrowsingFrameId++)
+		for (; BrowsingFrameId < BudgetHistorySize && bCarryOnBrowsingFrameHistory; BrowsingFrameId++)
 		{
 			const FrameHistoryEntry& NextFrameEntry = GetPreviousFrameEntry(BrowsingFrameId - 1);
 			const FrameHistoryEntry& FrameEntry = GetPreviousFrameEntry(BrowsingFrameId + 0);
@@ -282,6 +282,10 @@ void FDynamicResolutionHeuristicProxy::RefreshCurentFrameResolutionFraction_Rend
 				continue;
 			}
 
+		#if STATS
+			SET_FLOAT_STAT_FName(Budget.GetStatId_MeasuredMs().GetName(), TotalFrameGPUBusyTimeMs);
+		#endif
+
 			// Whether bound by game thread.
 			const bool bIsGameThreadBound = FrameEntry.GameThreadTimeMs > GDynamicPrimaryResolutionFraction.GetSettings().BudgetMs;
 
@@ -292,33 +296,33 @@ void FDynamicResolutionHeuristicProxy::RefreshCurentFrameResolutionFraction_Rend
 			const bool bIsCPUBound = bIsGameThreadBound || bIsRenderThreadBound;
 
 			// Whether GPU is over budget, when not CPU bound.
-			const bool bHasOverbudgetGPU = !bIsCPUBound && TotalFrameGPUBusyTimeMs > BudgetSettings.BudgetMs;
+			const bool bHasOverBudgetGPU = !bIsCPUBound && TotalFrameGPUBusyTimeMs > BudgetSettings.BudgetMs;
 
 			// Look if this is multiple consecutive GPU over budget frames.
-			if (bHasOverbudgetGPU)
+			if (bHasOverBudgetGPU)
 			{
-				ConsecutiveOverbudgetGPUFramCount++;
-				check(ConsecutiveOverbudgetGPUFramCount <= MaxConsecutiveOverbudgetGPUFrameCount);
+				ConsecutiveOverBudgetGPUFrameCount++;
+				check(ConsecutiveOverBudgetGPUFrameCount <= MaxConsecutiveOverBudgetGPUFrameCount);
 
 				// Max number of over budget frames where reached.
-				if (ConsecutiveOverbudgetGPUFramCount == MaxConsecutiveOverbudgetGPUFrameCount)
+				if (ConsecutiveOverBudgetGPUFrameCount == MaxConsecutiveOverBudgetGPUFrameCount)
 				{
-					// We ignore frames in history that happen before consecutive GPU overbudget frames.
-					bCaryOnBrowsingFrameHistory = false;
+					// We ignore frames in history that happen before consecutive GPU over budget frames.
+					bCarryOnBrowsingFrameHistory = false;
 				}
 			}
 			else
 			{
-				ConsecutiveOverbudgetGPUFramCount = 0;
+				ConsecutiveOverBudgetGPUFrameCount = 0;
 			}
 
 			float SuggestedResolutionFraction = 1.0f;
 
-			// If have reliable GPU times, or guess there is not GPU bubbles -> estimate the suggested resolution fraction that could have been used.
+			// If we have reliable GPU times, or guess there are no GPU bubbles -> estimate the suggested resolution fraction that could have been used.
 			{
 				// This assumes GPU busy time is directly proportional to ResolutionFraction^2, but in practice
 				// this is more A * ResolutionFraction^2 + B with B >= 0 non constant unknown cost such as unscaled
-				// post processing, vertex fetching & processing, occlusion queries, shadow maps rendering...
+				// post processing, vertex fetching & processing, occlusion queries, shadow map rendering...
 				//
 				// This assumption means we may drop ResolutionFraction lower than needed, or be slower to increase
 				// resolution.
@@ -341,17 +345,17 @@ void FDynamicResolutionHeuristicProxy::RefreshCurentFrameResolutionFraction_Rend
 		NewFrameResolutionFraction /= TotalWeight;
 
 		// If immediate previous frames where over budget, react immediately.
-		bool bGPUOverbugetPanic = FrameCount > 0 && ConsecutiveOverbudgetGPUFramCount == FrameCount;
+		bool bGPUOverBudgetPanic = FrameCount > 0 && ConsecutiveOverBudgetGPUFrameCount == FrameCount;
 
 		// If over budget, reset history size to 0 so that this frame really behave as a first frame after an history reset.
-		if (bGPUOverbugetPanic)
+		if (bGPUOverBudgetPanic)
 		{
 			BudgetHistorySizes[Budget] = 0;
 
 			if (Budget == GDynamicPrimaryResolutionFraction)
 			{
 				HistorySize = 0;
-				bGlobalGPUOverbugetPanic = true;
+				bGlobalGPUOverBudgetPanic = true;
 			}
 		}
 		// If not immediately over budget, refine the new resolution fraction.
@@ -394,7 +398,7 @@ void FDynamicResolutionHeuristicProxy::RefreshCurentFrameResolutionFraction_Rend
 		}
 		
 		// We do not change resolution too often to avoid interferences with temporal sub pixel in TAA upsample.
-		if ((bWouldBeWorthChangingRes && bCanChangeResolution) || bGlobalGPUOverbugetPanic)
+		if ((bWouldBeWorthChangingRes && bCanChangeResolution) || bGlobalGPUOverBudgetPanic)
 		{
 			NumberOfFramesSinceScreenPercentageChange = 0;
 			CurrentFrameResolutionFractions = NewFrameResolutionFractions;
@@ -405,10 +409,11 @@ void FDynamicResolutionHeuristicProxy::RefreshCurentFrameResolutionFraction_Rend
 		}
 	}
 	
-	return RefreshCurentFrameResolutionFractionUpperBound_RenderThread();
+	RefreshCurrentFrameResolutionFractionUpperBound_RenderThread();
+	RefreshHeuristicStats_RenderThread();
 }
 
-void FDynamicResolutionHeuristicProxy::RefreshCurentFrameResolutionFractionUpperBound_RenderThread()
+void FDynamicResolutionHeuristicProxy::RefreshCurrentFrameResolutionFractionUpperBound_RenderThread()
 {
 	// Compute max resolution for each budget by quantizing the new resolution fraction (falls back to the MaxResolution setting if BudgetSetting.UpperBoundQuantization==0)
 	DynamicRenderScaling::TMap<float> NewMaxResolutionFractions;
@@ -439,6 +444,35 @@ void FDynamicResolutionHeuristicProxy::RefreshCurentFrameResolutionFractionUpper
 	CurrentFrameMaxResolutionFractions = NewMaxResolutionFractions;
 }
 
+void FDynamicResolutionHeuristicProxy::RefreshHeuristicStats_RenderThread()
+{
+#if STATS
+	check(IsInRenderingThread());
+	for (TLinkedList<DynamicRenderScaling::FBudget*>::TIterator BudgetIt(DynamicRenderScaling::FBudget::GetGlobalList()); BudgetIt; BudgetIt.Next())
+	{
+		const DynamicRenderScaling::FBudget& Budget = **BudgetIt;
+		const DynamicRenderScaling::FHeuristicSettings& HeuristicSettings = Budget.GetSettings();
+
+		if (HeuristicSettings.IsEnabled())
+		{
+			SET_FLOAT_STAT_FName(Budget.GetStatId_TargetMs().GetName(), HeuristicSettings.GetTargetedMs());
+			// MeasuredMs is set in RefreshCurrentFrameResolutionFraction_RenderThread()
+			SET_FLOAT_STAT_FName(Budget.GetStatId_MinScaling().GetName(), HeuristicSettings.MinResolutionFraction);
+			SET_FLOAT_STAT_FName(Budget.GetStatId_MaxScaling().GetName(), HeuristicSettings.MaxResolutionFraction);
+			SET_FLOAT_STAT_FName(Budget.GetStatId_CurrentScaling().GetName(), CurrentFrameMaxResolutionFractions[Budget]);
+		}
+		else
+		{
+			SET_FLOAT_STAT_FName(Budget.GetStatId_TargetMs().GetName(), 0.0f);
+			SET_FLOAT_STAT_FName(Budget.GetStatId_MeasuredMs().GetName(), 0.0f);
+			SET_FLOAT_STAT_FName(Budget.GetStatId_MinScaling().GetName(), 0.0f);
+			SET_FLOAT_STAT_FName(Budget.GetStatId_MaxScaling().GetName(), 0.0f);
+			SET_FLOAT_STAT_FName(Budget.GetStatId_CurrentScaling().GetName(), 0.0f);
+		}
+	}
+#endif
+}
+
 // static
 DynamicRenderScaling::TMap<float> FDynamicResolutionHeuristicProxy::GetResolutionFractionUpperBounds() const
 {
@@ -448,7 +482,7 @@ DynamicRenderScaling::TMap<float> FDynamicResolutionHeuristicProxy::GetResolutio
 
 
 /** Returns the view fraction that should be used for current frame. */
-DynamicRenderScaling::TMap<float> FDynamicResolutionHeuristicProxy::QueryCurentFrameResolutionFractions_Internal() const
+DynamicRenderScaling::TMap<float> FDynamicResolutionHeuristicProxy::QueryCurrentFrameResolutionFractions_Internal() const
 {
 	DynamicRenderScaling::TMap<float> MaxResolutionFractions = GetResolutionFractionUpperBounds();
 	DynamicRenderScaling::TMap<float> ResolutionFractions = CurrentFrameResolutionFractions;
@@ -591,7 +625,7 @@ public:
 				/* DynamicResolutionGPUBusyTimeMs = */ PrevFrameGPUTimeMs,
 				/* BudgetTimingMs = */ BudgetTimingMs);
 
-			Heuristic.RefreshCurentFrameResolutionFraction_RenderThread();
+			Heuristic.RefreshCurrentFrameResolutionFraction_RenderThread();
 
 			// Set a non insane value for internal checks to pass as if GRHISupportsGPUBusyTimeQueries == true.
 			CurrentFrameInFlightIndex = 0;
@@ -767,7 +801,7 @@ private:
 		// Refresh the heuristic.
 		if (ShouldRefreshHeuristic)
 		{
-			Heuristic.RefreshCurentFrameResolutionFraction_RenderThread();
+			Heuristic.RefreshCurrentFrameResolutionFraction_RenderThread();
 		}
 	}
 
@@ -831,7 +865,7 @@ public:
 	{
 		check(IsInRenderingThread());
 
-		DynamicRenderScaling::TMap<float> ResolutionFractions = Proxy->Heuristic.QueryCurentFrameResolutionFractions();
+		DynamicRenderScaling::TMap<float> ResolutionFractions = Proxy->Heuristic.QueryCurrentFrameResolutionFractions();
 		if (!ViewFamily.EngineShowFlags.ScreenPercentage)
 		{
 			ResolutionFractions[GDynamicPrimaryResolutionFraction] = 1.0f;
diff --git a/Engine/Source/Runtime/Engine/Public/DynamicResolutionProxy.h b/Engine/Source/Runtime/Engine/Public/DynamicResolutionProxy.h
--- a/Engine/Source/Runtime/Engine/Public/DynamicResolutionProxy.h
+++ b/Engine/Source/Runtime/Engine/Public/DynamicResolutionProxy.h
@@ -6,18 +6,15 @@
 #include "SceneView.h"
 #include "Engine/EngineTypes.h"
 
-
-/** Render thread proxy that hold's the heuristic for dynamic resolution. */
+/** Render thread proxy that holds the heuristic for dynamic resolution. */
 class ENGINE_API FDynamicResolutionHeuristicProxy
 {
 public:
 	static constexpr uint64 kInvalidEntryId = ~uint64(0);
 
-
 	FDynamicResolutionHeuristicProxy();
 	~FDynamicResolutionHeuristicProxy();
 
-
 	/** Resets the proxy. */
 	void Reset_RenderThread();
 
@@ -32,29 +29,27 @@ public:
 		const DynamicRenderScaling::TMap<float>& BudgetTimingMs);
 
 	/** Refresh resolution fraction's from history. */
-	void RefreshCurentFrameResolutionFraction_RenderThread();
+	void RefreshCurrentFrameResolutionFraction_RenderThread();
 
 	/** Returns the view fraction that should be used for current frame. */
-	FORCEINLINE DynamicRenderScaling::TMap<float> QueryCurentFrameResolutionFractions() const
+	FORCEINLINE DynamicRenderScaling::TMap<float> QueryCurrentFrameResolutionFractions() const
 	{
 		check(IsInRenderingThread());
-		return QueryCurentFrameResolutionFractions_Internal();
+		return QueryCurrentFrameResolutionFractions_Internal();
 	}
 
 	/** Returns a non thread safe approximation of the current resolution fraction applied on render thread. */
 	FORCEINLINE DynamicRenderScaling::TMap<float> GetResolutionFractionsApproximation_GameThread() const
 	{
 		check(IsInGameThread());
-		return QueryCurentFrameResolutionFractions_Internal();
+		return QueryCurrentFrameResolutionFractions_Internal();
 	}
 
-
 	/** Returns the view fraction upper bound. */
 	DynamicRenderScaling::TMap<float> GetResolutionFractionUpperBounds() const;
 
 	/** Creates a default dynamic resolution state using this proxy that queries GPU timing from the RHI. */
-	static TSharedPtr< class IDynamicResolutionState > CreateDefaultState();
-
+	static TSharedPtr<class IDynamicResolutionState> CreateDefaultState();
 
 private:
 	struct FrameHistoryEntry
@@ -66,7 +61,7 @@ private:
 		// Total GPU busy time for the entire frame in milliseconds.
 		float TotalFrameGPUBusyTimeMs;
 
-		// Total GPU busy time for the render thread commands that does dynamic resolutions.
+		// Total GPU busy time for the render thread commands that perform dynamic resolutions.
 		float GlobalDynamicResolutionTimeMs;
 
 		// Time for each individual timings
@@ -85,7 +80,6 @@ private:
 			BudgetTimingMs.SetAll(-1.0f);
 		}
 
-
 		// Returns whether GPU timings have landed.
 		inline bool HasGPUTimings() const
 		{
@@ -94,7 +88,7 @@ private:
 	};
 
 	// Circular buffer of the history.
-	// We don't use TCircularBuffer because does not support resizes.
+	// We don't use TCircularBuffer because it does not support resizes.
 	TArray<FrameHistoryEntry> History;
 	int32 PreviousFrameIndex;
 	int32 HistorySize;
@@ -102,7 +96,7 @@ private:
 	// Counts the number of frame since the last screen percentage change.
 	int32 NumberOfFramesSinceScreenPercentageChange;
 
-	// Number of frame remaining to ignore.
+	// Number of frames remaining to ignore.
 	int32 IgnoreFrameRemainingCount;
 
 	// Current frame's view fraction.
@@ -113,7 +107,6 @@ private:
 	// Frame counter to allocate unique ID for CommitPreviousFrameGPUTimings_RenderThread().
 	uint64 FrameCounter;
 
-
 	inline const FrameHistoryEntry& GetPreviousFrameEntry(int32 BrowsingFrameId) const
 	{
 		if (BrowsingFrameId < 0 || BrowsingFrameId >= HistorySize)
@@ -124,9 +117,11 @@ private:
 		return History[(History.Num() + PreviousFrameIndex - BrowsingFrameId) % History.Num()];
 	}
 
-	DynamicRenderScaling::TMap<float> QueryCurentFrameResolutionFractions_Internal() const;
+	DynamicRenderScaling::TMap<float> QueryCurrentFrameResolutionFractions_Internal() const;
+
+	void RefreshCurrentFrameResolutionFractionUpperBound_RenderThread();
 
-	void RefreshCurentFrameResolutionFractionUpperBound_RenderThread();
+	void RefreshHeuristicStats_RenderThread();
 
 	void ResetInternal();
 
diff --git a/Engine/Source/Runtime/RenderCore/Private/DynamicRenderScaling.cpp b/Engine/Source/Runtime/RenderCore/Private/DynamicRenderScaling.cpp
--- a/Engine/Source/Runtime/RenderCore/Private/DynamicRenderScaling.cpp
+++ b/Engine/Source/Runtime/RenderCore/Private/DynamicRenderScaling.cpp
@@ -143,6 +143,14 @@ FBudget::FBudget(const TCHAR* InName, FHeuristicSettings (*InHeuristicSettingsGe
 	}
 #endif
 
+#if STATS
+	StatId_TargetMs			= FDynamicStats::CreateStatIdDouble<STAT_GROUP_TO_FStatGroup(STATGROUP_RenderScaling)>(FString::Printf(TEXT("%s_TargetMs"), InName));
+	StatId_MeasuredMs		= FDynamicStats::CreateStatIdDouble<STAT_GROUP_TO_FStatGroup(STATGROUP_RenderScaling)>(FString::Printf(TEXT("%s_MeasuredMs"), InName));
+	StatId_MinScaling		= FDynamicStats::CreateStatIdDouble<STAT_GROUP_TO_FStatGroup(STATGROUP_RenderScaling)>(FString::Printf(TEXT("%s_MinFraction"), InName));
+	StatId_MaxScaling		= FDynamicStats::CreateStatIdDouble<STAT_GROUP_TO_FStatGroup(STATGROUP_RenderScaling)>(FString::Printf(TEXT("%s_MaxFraction"), InName));
+	StatId_CurrentScaling	= FDynamicStats::CreateStatIdDouble<STAT_GROUP_TO_FStatGroup(STATGROUP_RenderScaling)>(FString::Printf(TEXT("%s_Fraction"), InName));
+#endif
+
 	CachedSettings = (*HeuristicSettingsGetter)();
 }
 
diff --git a/Engine/Source/Runtime/RenderCore/Public/DynamicRenderScaling.h b/Engine/Source/Runtime/RenderCore/Public/DynamicRenderScaling.h
--- a/Engine/Source/Runtime/RenderCore/Public/DynamicRenderScaling.h
+++ b/Engine/Source/Runtime/RenderCore/Public/DynamicRenderScaling.h
@@ -40,7 +40,6 @@ struct RENDERCORE_API FHeuristicSettings final
 	static constexpr int32 kDefaultFractionQuantization = 0;
 	static constexpr int32 kDefaultUpperBoundQuantization = 0;
 
-
 	EHeuristicModel Model = EHeuristicModel::Unknown;
 	bool bModelScalesWithPrimaryScreenPercentage = false;
 
@@ -53,7 +52,6 @@ struct RENDERCORE_API FHeuristicSettings final
 	int32 FractionQuantization       = kDefaultFractionQuantization;
 	int32 UpperBoundQuantization     = kDefaultUpperBoundQuantization;
 
-
 	/** Returns whether the heuristic is enabled or not. */
 	bool IsEnabled() const;
 
@@ -114,6 +112,33 @@ public:
 		return BudgetId == Other.BudgetId;
 	}
 
+#if STATS
+	inline const TStatId& GetStatId_TargetMs() const
+	{
+		return StatId_TargetMs;
+	}
+
+	inline const TStatId& GetStatId_MeasuredMs() const
+	{
+		return StatId_MeasuredMs;
+	}
+
+	inline const TStatId& GetStatId_MinScaling() const
+	{
+		return StatId_MinScaling;
+	}
+
+	inline const TStatId& GetStatId_MaxScaling() const
+	{
+		return StatId_MaxScaling;
+	}
+
+	inline const TStatId& GetStatId_CurrentScaling() const
+	{
+		return StatId_CurrentScaling;
+	}
+#endif
+
 private:
 	const TCHAR* Name;
 	TArray<char> AnsiName;
@@ -122,6 +147,14 @@ private:
 	FHeuristicSettings CachedSettings;
 	int32 BudgetId = 0;
 
+#if STATS
+	TStatId StatId_TargetMs;
+	TStatId StatId_MeasuredMs;
+	TStatId StatId_MinScaling;
+	TStatId StatId_MaxScaling;
+	TStatId StatId_CurrentScaling;
+#endif
+
 	UE_NONCOPYABLE(FBudget);
 
 	template<typename Type>
