Move level streaming state enums out of class so they can be forward declared.
ULevelStreaming::ECurrentState -> ELevelStreamingState
ULevelStreaming::ETargetState -> ELevelStreamingTargetState

diff --git a/Engine/Plugins/Runtime/nDisplay/Source/DisplayCluster/Private/Game/DisplayClusterGameManager.cpp b/Engine/Plugins/Runtime/nDisplay/Source/DisplayCluster/Private/Game/DisplayClusterGameManager.cpp
--- a/Engine/Plugins/Runtime/nDisplay/Source/DisplayCluster/Private/Game/DisplayClusterGameManager.cpp
+++ b/Engine/Plugins/Runtime/nDisplay/Source/DisplayCluster/Private/Game/DisplayClusterGameManager.cpp
@@ -197,7 +197,7 @@ ADisplayClusterRootActor* FDisplayClusterGameManager::FindRootActor(UWorld* InWo
 	const TArray<ULevelStreaming*>& StreamingLevels = InWorld->GetStreamingLevels();
 	for (const ULevelStreaming* const StreamingLevel : StreamingLevels)
 	{
-		if (StreamingLevel && StreamingLevel->GetCurrentState() == ULevelStreaming::ECurrentState::LoadedVisible)
+		if (StreamingLevel && StreamingLevel->GetCurrentState() == ELevelStreamingState::LoadedVisible)
 		{
 			// Look for the actor in those sub-levels that have been loaded already
 			const TSoftObjectPtr<UWorld>& SubWorldAsset = StreamingLevel->GetWorldAsset();
diff --git a/Engine/Plugins/Runtime/nDisplay/Source/DisplayClusterEditor/Private/DisplayClusterEditorEngine.cpp b/Engine/Plugins/Runtime/nDisplay/Source/DisplayClusterEditor/Private/DisplayClusterEditorEngine.cpp
--- a/Engine/Plugins/Runtime/nDisplay/Source/DisplayClusterEditor/Private/DisplayClusterEditorEngine.cpp
+++ b/Engine/Plugins/Runtime/nDisplay/Source/DisplayClusterEditor/Private/DisplayClusterEditorEngine.cpp
@@ -92,7 +92,7 @@ void UDisplayClusterEditorEngine::StartPlayInEditorSession(FRequestPlaySessionPa
 			const TArray<ULevelStreaming*>& StreamingLevels = EditorWorldPreDup->GetStreamingLevels();
 			for (const ULevelStreaming* const StreamingLevel : StreamingLevels)
 			{
-				if (StreamingLevel && StreamingLevel->GetCurrentState() == ULevelStreaming::ECurrentState::LoadedVisible)
+				if (StreamingLevel && StreamingLevel->GetCurrentState() == ELevelStreamingState::LoadedVisible)
 				{
 					// Look for the actor in those sub-levels that have been loaded already
 					const TSoftObjectPtr<UWorld>& SubWorldAsset = StreamingLevel->GetWorldAsset();
diff --git a/Engine/Source/Runtime/Engine/Classes/Engine/LevelStreaming.h b/Engine/Source/Runtime/Engine/Classes/Engine/LevelStreaming.h
--- a/Engine/Source/Runtime/Engine/Classes/Engine/LevelStreaming.h
+++ b/Engine/Source/Runtime/Engine/Classes/Engine/LevelStreaming.h
@@ -105,6 +105,28 @@ public:
 DECLARE_DYNAMIC_MULTICAST_DELEGATE( FLevelStreamingLoadedStatus );
 DECLARE_DYNAMIC_MULTICAST_DELEGATE( FLevelStreamingVisibilityStatus );
 
+enum class ELevelStreamingState : uint8
+{
+	Removed,
+	Unloaded,
+	FailedToLoad,
+	Loading,
+	LoadedNotVisible,
+	MakingVisible,
+	LoadedVisible,
+	MakingInvisible
+};
+ENGINE_API const TCHAR* EnumToString(ELevelStreamingState InState);
+
+enum class ELevelStreamingTargetState : uint8
+{
+	Unloaded,
+	UnloadedAndRemoved,
+	LoadedNotVisible,
+	LoadedVisible,
+};
+ENGINE_API const TCHAR* EnumToString(ELevelStreamingTargetState InTargetState);
+
 /**
  * Abstract base class of container object encapsulating data required for streaming and providing 
  * interface for when a level should be streamed in and out of memory.
@@ -115,32 +137,38 @@ class ENGINE_API ULevelStreaming : public UObject
 {
 	GENERATED_UCLASS_BODY()
 
-	enum class ECurrentState : uint8
+	enum class UE_DEPRECATED(5.2, "ULevelStreaming::ECurrentState has been replaced by ELevelStreamingState")
+	ECurrentState : uint8
 	{
-		Removed,
-		Unloaded,
-		FailedToLoad,
-		Loading,
-		LoadedNotVisible,
-		MakingVisible,
-		LoadedVisible,
-		MakingInvisible
+	 	Removed = (uint8)ELevelStreamingState::Removed,
+		Unloaded = (uint8)ELevelStreamingState::Unloaded,
+		FailedToLoad = (uint8)ELevelStreamingState::FailedToLoad,
+		Loading = (uint8)ELevelStreamingState::Loading,
+		LoadedNotVisible = (uint8)ELevelStreamingState::LoadedNotVisible,
+		MakingVisible = (uint8)ELevelStreamingState::MakingVisible,
+		LoadedVisible = (uint8)ELevelStreamingState::LoadedVisible,
+		MakingInvisible = (uint8)ELevelStreamingState::MakingInvisible
 	};
-
+PRAGMA_DISABLE_DEPRECATION_WARNINGS
 	static const TCHAR* EnumToString(ECurrentState InCurrentState);
-
-private:
-	enum class ETargetState : uint8
+	friend bool operator==(ELevelStreamingState A, ECurrentState B)
 	{
-		Unloaded,
-		UnloadedAndRemoved,
-		LoadedNotVisible,
-		LoadedVisible,
-	};
-
-	static const TCHAR* EnumToString(ETargetState InTargetState);
+		return A == (ELevelStreamingState)B;
+	}
+	friend bool operator==(ECurrentState A, ELevelStreamingState B)
+	{
+		return (ELevelStreamingState)A == (ELevelStreamingState)B;
+	}
+	friend bool operator!=(ELevelStreamingState A, ECurrentState B)
+	{
+		return A != (ELevelStreamingState)B;
+	}
+	friend bool operator!=(ECurrentState A, ELevelStreamingState B)
+	{
+		return (ELevelStreamingState)A != (ELevelStreamingState)B;
+	}
+PRAGMA_ENABLE_DEPRECATION_WARNINGS
 
-public:
 	static ULevelStreaming* FindStreamingLevel(const ULevel* Level);
 	static void RemoveLevelAnnotation(const ULevel* Level);
 
@@ -234,10 +262,12 @@ protected:
 	uint8 bSkipClientUseMakingVisibleTransactionRequest:1;
 
 	/** What the current streamed state of the streaming level is */
-	ECurrentState CurrentState;
+	ELevelStreamingState CurrentState;
 
 	/** What streamed state the streaming level is transitioning towards */
-	ETargetState TargetState;
+	ELevelStreamingTargetState TargetState;
+
+	void SetCurrentState(ELevelStreamingState NewState);
 
 public:
 
@@ -310,7 +340,12 @@ public:
 	//~ End UObject Interface
 
 	/** Returns the current loaded/visible state of the streaming level. */
-	ECurrentState GetCurrentState() const { return CurrentState; }
+PRAGMA_DISABLE_DEPRECATION_WARNINGS
+	UE_DEPRECATED(5.2, "ULevelStreaming::ECurrentState has been replaced by ELevelStreamingState. Use GetLevelStreamingState instead.")
+	ECurrentState GetCurrentState() const { return (ECurrentState)CurrentState; }
+PRAGMA_ENABLE_DEPRECATION_WARNINGS
+
+	ELevelStreamingState GetLevelStreamingState() const { return CurrentState; }
 
 private:
 
@@ -385,7 +420,7 @@ public:
 	void SetPriority(int32 NewPriority);
 
 	/** Returns whether the streaming level is in the loading state. */
-	bool HasLoadRequestPending() const { return GetCurrentState() == ECurrentState::Loading; }
+	bool HasLoadRequestPending() const { return CurrentState == ELevelStreamingState::Loading; }
 
 	/** Returns whether the streaming level has loaded a level. */
 	bool HasLoadedLevel() const
diff --git a/Engine/Source/Runtime/Engine/Private/LevelInstance/LevelInstanceLevelStreaming.cpp b/Engine/Source/Runtime/Engine/Private/LevelInstance/LevelInstanceLevelStreaming.cpp
--- a/Engine/Source/Runtime/Engine/Private/LevelInstance/LevelInstanceLevelStreaming.cpp
+++ b/Engine/Source/Runtime/Engine/Private/LevelInstance/LevelInstanceLevelStreaming.cpp
@@ -229,7 +229,7 @@ ULevelStreamingLevelInstance* ULevelStreamingLevelInstance::LoadInstance(ILevelI
 			// It might not be the final solution to support Undo/Redo in LevelInstances but it handles most of the non-editing part
 			if (ULevel* Level = LevelStreaming->GetLoadedLevel())
 			{
-				check(LevelStreaming->GetCurrentState() == ULevelStreaming::ECurrentState::LoadedVisible);
+				check(LevelStreaming->GetLevelStreamingState() == ELevelStreamingState::LoadedVisible);
 
 				Level->OnLoadedActorAddedToLevelEvent.AddUObject(LevelStreaming, &ULevelStreamingLevelInstance::OnLoadedActorAddedToLevel);
 				Level->OnLoadedActorRemovedFromLevelEvent.AddUObject(LevelStreaming, &ULevelStreamingLevelInstance::OnLoadedActorRemovedFromLevel);
diff --git a/Engine/Source/Runtime/Engine/Private/LevelStreaming.cpp b/Engine/Source/Runtime/Engine/Private/LevelStreaming.cpp
--- a/Engine/Source/Runtime/Engine/Private/LevelStreaming.cpp
+++ b/Engine/Source/Runtime/Engine/Private/LevelStreaming.cpp
@@ -386,7 +386,7 @@ ULevelStreaming::ULevelStreaming(const FObjectInitializer& ObjectInitializer)
 	MinTimeBetweenVolumeUnloadRequests = 2.0f;
 	bDrawOnLevelStatusMap = true;
 	LevelLODIndex = INDEX_NONE;
-	CurrentState = ECurrentState::Removed;
+	CurrentState = ELevelStreamingState::Removed;
 	bSkipClientUseMakingInvisibleTransactionRequest = false;
 	bSkipClientUseMakingVisibleTransactionRequest = false;
 }
@@ -480,27 +480,32 @@ void ULevelStreaming::OnLevelAdded()
 	{
 		if (LoadedLevel->bIsVisible)
 		{
-			CurrentState = ECurrentState::LoadedVisible;
+			SetCurrentState(ELevelStreamingState::LoadedVisible);
 		}
 		else
 		{
-			CurrentState = ECurrentState::LoadedNotVisible;
+			SetCurrentState(ELevelStreamingState::LoadedNotVisible);
 		}
 	}
 	else
 	{
-		CurrentState = ECurrentState::Unloaded;
+		SetCurrentState(ELevelStreamingState::Unloaded);
 	}
 }
 
 void ULevelStreaming::OnLevelRemoved()
 {
 	// If in one of the transitional states removing the level will be highly problematic
-	ensure(CurrentState != ECurrentState::Loading);
-	ensure(CurrentState != ECurrentState::MakingInvisible);
-	ensure(CurrentState != ECurrentState::MakingVisible);
+	ensure(CurrentState != ELevelStreamingState::Loading);
+	ensure(CurrentState != ELevelStreamingState::MakingInvisible);
+	ensure(CurrentState != ELevelStreamingState::MakingVisible);
 
-	CurrentState = ECurrentState::Removed;
+	SetCurrentState(ELevelStreamingState::Removed);
+}
+
+void ULevelStreaming::SetCurrentState(ELevelStreamingState NewState)
+{
+	CurrentState = NewState;
 }
 
 bool ULevelStreaming::DetermineTargetState()
@@ -524,34 +529,34 @@ bool ULevelStreaming::DetermineTargetState()
 
 	switch(CurrentState)
 	{
-	case ECurrentState::MakingVisible:
+	case ELevelStreamingState::MakingVisible:
 		ensure(LoadedLevel);
 		if (!ShouldBeVisible() && GetWorld()->GetCurrentLevelPendingVisibility() != LoadedLevel)
 		{
-			// Since level doesn't need to be visible anymore, change TargetState to ETargetState::LoadedNotVisible.
-			// Next UpdateStreamingState will handle switching CurrentState to ECurrentState::LoadedNotVisible.
+			// Since level doesn't need to be visible anymore, change TargetState to ELevelStreamingTargetState::LoadedNotVisible.
+			// Next UpdateStreamingState will handle switching CurrentState to ELevelStreamingState::LoadedNotVisible.
 			// From there, regular flow will properly handle TargetState.
-			TargetState = ETargetState::LoadedNotVisible;
+			TargetState = ELevelStreamingTargetState::LoadedNotVisible;
 		}
 		else
 		{
-			TargetState = ETargetState::LoadedVisible;
+			TargetState = ELevelStreamingTargetState::LoadedVisible;
 		}
 		break;
 
-	case ECurrentState::MakingInvisible:
+	case ELevelStreamingState::MakingInvisible:
 		ensure(LoadedLevel);
-		TargetState = ETargetState::LoadedNotVisible;
+		TargetState = ELevelStreamingTargetState::LoadedNotVisible;
 		break;
 
-	case ECurrentState::Loading:
-		TargetState = ETargetState::LoadedNotVisible;
+	case ELevelStreamingState::Loading:
+		TargetState = ELevelStreamingTargetState::LoadedNotVisible;
 		break;
 
-	case ECurrentState::Unloaded:
+	case ELevelStreamingState::Unloaded:
 		if (bIsRequestingUnloadAndRemoval)
 		{
-			TargetState = ETargetState::UnloadedAndRemoved;
+			TargetState = ELevelStreamingTargetState::UnloadedAndRemoved;
 		}
 		else if (World->GetShouldForceUnloadStreamingLevels())
 		{
@@ -559,11 +564,11 @@ bool ULevelStreaming::DetermineTargetState()
 		}
 		else if (!World->IsGameWorld())
 		{
-			TargetState = ETargetState::LoadedNotVisible;
+			TargetState = ELevelStreamingTargetState::LoadedNotVisible;
 		}
 		else if (ShouldBeLoaded())
 		{
-			TargetState = ETargetState::LoadedNotVisible;
+			TargetState = ELevelStreamingTargetState::LoadedNotVisible;
 		}
 		else
 		{
@@ -571,22 +576,22 @@ bool ULevelStreaming::DetermineTargetState()
 		}
 		break;
 
-	case ECurrentState::LoadedNotVisible:
+	case ELevelStreamingState::LoadedNotVisible:
 		if (bIsRequestingUnloadAndRemoval || World->GetShouldForceUnloadStreamingLevels())
 		{
-			TargetState = ETargetState::Unloaded;
+			TargetState = ELevelStreamingTargetState::Unloaded;
 		}
 		else if (World->IsGameWorld() && !ShouldBeLoaded())
 		{
-			TargetState = ETargetState::Unloaded;
+			TargetState = ELevelStreamingTargetState::Unloaded;
 		}
 		else if (!IsDesiredLevelLoaded())
 		{
-			TargetState = ETargetState::LoadedNotVisible;
+			TargetState = ELevelStreamingTargetState::LoadedNotVisible;
 		}
 		else if (ShouldBeVisible())
 		{
-			TargetState = ETargetState::LoadedVisible;
+			TargetState = ELevelStreamingTargetState::LoadedVisible;
 		}
 		else
 		{
@@ -594,22 +599,22 @@ bool ULevelStreaming::DetermineTargetState()
 		}
 		break;
 
-	case ECurrentState::LoadedVisible:
+	case ELevelStreamingState::LoadedVisible:
 		if (bIsRequestingUnloadAndRemoval || World->GetShouldForceUnloadStreamingLevels())
 		{
-			TargetState = ETargetState::LoadedNotVisible;
+			TargetState = ELevelStreamingTargetState::LoadedNotVisible;
 		}
 		else if (World->IsGameWorld() && !ShouldBeLoaded())
 		{
-			TargetState = ETargetState::LoadedNotVisible;
+			TargetState = ELevelStreamingTargetState::LoadedNotVisible;
 		}
 		else if (!ShouldBeVisible())
 		{
-			TargetState = ETargetState::LoadedNotVisible;
+			TargetState = ELevelStreamingTargetState::LoadedNotVisible;
 		}
 		else if (!IsDesiredLevelLoaded())
 		{
-			TargetState = ETargetState::LoadedVisible;
+			TargetState = ELevelStreamingTargetState::LoadedVisible;
 		}
 		else
 		{
@@ -617,13 +622,13 @@ bool ULevelStreaming::DetermineTargetState()
 		}
 		break;
 
-	case ECurrentState::FailedToLoad:
+	case ELevelStreamingState::FailedToLoad:
 
 		// Anything that affects whether we might try to reload changes current state itself
 		bContinueToConsider = false;
 		break;
 
-	case ECurrentState::Removed:
+	case ELevelStreamingState::Removed:
 
 		// Never continue to consider a removed streaming level
 		bContinueToConsider = false;
@@ -822,11 +827,11 @@ void ULevelStreaming::UpdateStreamingState(bool& bOutUpdateAgain, bool& bOutRede
 		const bool bAllowLevelLoadRequests = (bBlockOnLoad || World->AllowLevelLoadRequests());
 		bBlockOnLoad |= (!GUseBackgroundLevelStreaming || !World->IsGameWorld());
 
-		const ECurrentState PreviousState = CurrentState;
+		const ELevelStreamingState PreviousState = CurrentState;
 
 		RequestLevel(World, bAllowLevelLoadRequests, (bBlockOnLoad ? ULevelStreaming::AlwaysBlock : ULevelStreaming::BlockAlwaysLoadedLevelsOnly));
 
-		if (CurrentState != ECurrentState::Loading)
+		if (CurrentState != ELevelStreamingState::Loading)
 		{
 			bOutRedetermineTarget = true;
 
@@ -844,13 +849,13 @@ void ULevelStreaming::UpdateStreamingState(bool& bOutUpdateAgain, bool& bOutRede
 
 	switch(CurrentState)
 	{
-	case ECurrentState::MakingVisible:
+	case ELevelStreamingState::MakingVisible:
 		if (ensure(LoadedLevel))
 		{
 			// Handle case where MakingVisible is not needed anymore
-			if (TargetState == ETargetState::LoadedNotVisible)
+			if (TargetState == ELevelStreamingTargetState::LoadedNotVisible)
 			{
-				CurrentState = ECurrentState::LoadedNotVisible;
+				SetCurrentState(ELevelStreamingState::LoadedNotVisible);
 				bOutUpdateAgain = true;
 				bOutRedetermineTarget = true;
 				// Make sure to update level visibility state (in case the server already acknowledged a Making Visible request)
@@ -889,7 +894,7 @@ void ULevelStreaming::UpdateStreamingState(bool& bOutUpdateAgain, bool& bOutRede
 						World->Scene->OnLevelAddedToWorld(LoadedLevel->GetOutermost()->GetFName(), World, LoadedLevel->bIsLightingScenario);
 					}
 
-					CurrentState = ECurrentState::LoadedVisible;
+					SetCurrentState(ELevelStreamingState::LoadedVisible);
 					bOutUpdateAgain = true;
 					bOutRedetermineTarget = true;
 				}
@@ -897,7 +902,7 @@ void ULevelStreaming::UpdateStreamingState(bool& bOutUpdateAgain, bool& bOutRede
 		}
 		break;
 
-	case ECurrentState::MakingInvisible:
+	case ELevelStreamingState::MakingInvisible:
 		if (ensure(LoadedLevel))
 		{
 			auto RemoveLevelFromScene = [World, this]()
@@ -944,28 +949,28 @@ void ULevelStreaming::UpdateStreamingState(bool& bOutUpdateAgain, bool& bOutRede
 					RemoveLevelFromScene();
 				}
 
-				CurrentState = ECurrentState::LoadedNotVisible;
+				SetCurrentState(ELevelStreamingState::LoadedNotVisible);
 				bOutUpdateAgain = true;
 				bOutRedetermineTarget = true;
 			}
 		}
 		break;
 
-	case ECurrentState::Loading:
+	case ELevelStreamingState::Loading:
 		// Just waiting
 		break;
 
-	case ECurrentState::Unloaded:
+	case ELevelStreamingState::Unloaded:
 		
 		switch (TargetState)
 		{
-			case ETargetState::LoadedNotVisible:
+			case ELevelStreamingTargetState::LoadedNotVisible:
 			{
 				UpdateStreamingState_RequestLevel();
 			}
 			break;
 
-			case ETargetState::UnloadedAndRemoved:
+			case ELevelStreamingTargetState::UnloadedAndRemoved:
 				World->RemoveStreamingLevel(this);
 				bOutRedetermineTarget = true;
 				break;
@@ -975,18 +980,18 @@ void ULevelStreaming::UpdateStreamingState(bool& bOutUpdateAgain, bool& bOutRede
 		}
 		break;
 
-	case ECurrentState::LoadedNotVisible:
+	case ELevelStreamingState::LoadedNotVisible:
 		switch (TargetState)
 		{
-		case ETargetState::LoadedVisible:
-			CurrentState = ECurrentState::MakingVisible;
+		case ELevelStreamingTargetState::LoadedVisible:
+			SetCurrentState(ELevelStreamingState::MakingVisible);
 			// Make sure client pending visibility request (if any) matches current state
 			NetVisibilityState.InvalidateClientPendingRequest();
 			BeginClientNetVisibilityRequest(ShouldBeVisible());
 			bOutUpdateAgain = true;
 			break;
 
-		case ETargetState::Unloaded:
+		case ELevelStreamingTargetState::Unloaded:
 			DiscardPendingUnloadLevel(World);
 			ClearLoadedLevel();
 			DiscardPendingUnloadLevel(World);
@@ -995,7 +1000,7 @@ void ULevelStreaming::UpdateStreamingState(bool& bOutUpdateAgain, bool& bOutRede
 			bOutRedetermineTarget = true;
 			break;
 
-		case ETargetState::LoadedNotVisible:
+		case ELevelStreamingTargetState::LoadedNotVisible:
 			if (LoadedLevel && !IsDesiredLevelLoaded())
 			{
 				// Process PendingUnloadLevel to unblock level streaming state machine (no new request will start while there's a pending level to unload) 
@@ -1011,7 +1016,7 @@ void ULevelStreaming::UpdateStreamingState(bool& bOutUpdateAgain, bool& bOutRede
 			// In blocking load, loading will be done after the UpdateStreamingState_RequestLevel call and leave us in the LoadedNotVisible (with a loadedlevel) state which would prevent the bOutUpdateAgain from being set to true.
 			// So this condition here makes sure that we aren't loading (async) and that we should be visible (LoadedNotVisible isn't our final target).
 			// If that is the case we request another update.
-			if (CurrentState != ECurrentState::Loading)
+			if (CurrentState != ELevelStreamingState::Loading)
 			{
 				bOutUpdateAgain |= ShouldBeVisible();
 			}
@@ -1023,18 +1028,18 @@ void ULevelStreaming::UpdateStreamingState(bool& bOutUpdateAgain, bool& bOutRede
 
 		break;
 
-	case ECurrentState::LoadedVisible:
+	case ELevelStreamingState::LoadedVisible:
 		switch (TargetState)
 		{
-		case ETargetState::LoadedNotVisible:
-			CurrentState = ECurrentState::MakingInvisible;
+		case ELevelStreamingTargetState::LoadedNotVisible:
+			SetCurrentState(ELevelStreamingState::MakingInvisible);
 			// Make sure client pending visibility request (if any) matches current state
 			NetVisibilityState.InvalidateClientPendingRequest();
 			BeginClientNetVisibilityRequest(ShouldBeVisible());
 			bOutUpdateAgain = true;
 			break;
 
-		case ETargetState::LoadedVisible:
+		case ELevelStreamingTargetState::LoadedVisible:
 			UpdateStreamingState_RequestLevel();
 			break;
 
@@ -1044,51 +1049,58 @@ void ULevelStreaming::UpdateStreamingState(bool& bOutUpdateAgain, bool& bOutRede
 
 		break;
 
-	case ECurrentState::FailedToLoad:
+	case ELevelStreamingState::FailedToLoad:
 		bOutRedetermineTarget = true;
 		break;
 
 	default:
-		ensureMsgf(false, TEXT("Unexpected state in ULevelStreaming::UpdateStreamingState for '%s'. CurrentState='%s' TargetState='%s'"), *GetPathName(), EnumToString(CurrentState), EnumToString(TargetState));
+		ensureMsgf(false, TEXT("Unexpected state in ULevelStreaming::UpdateStreamingState for '%s'. CurrentState='%s' TargetState='%s'"), *GetPathName(), ::EnumToString(CurrentState), ::EnumToString(TargetState));
 	}
 }
 
-const TCHAR* ULevelStreaming::EnumToString(ECurrentState InCurrentState)
+PRAGMA_DISABLE_DEPRECATION_WARNINGS
+const TCHAR* ULevelStreaming::EnumToString(ULevelStreaming::ECurrentState InCurrentState)
+{
+	return ::EnumToString((ELevelStreamingState)InCurrentState);
+}
+PRAGMA_ENABLE_DEPRECATION_WARNINGS
+
+const TCHAR* EnumToString(ELevelStreamingState InCurrentState)
 {
 	switch (InCurrentState)
 	{
-	case ECurrentState::Removed:
+	case ELevelStreamingState::Removed:
 		return TEXT("Removed");
-	case ECurrentState::Unloaded:
+	case ELevelStreamingState::Unloaded:
 		return TEXT("Unloaded");
-	case ECurrentState::FailedToLoad:
+	case ELevelStreamingState::FailedToLoad:
 		return TEXT("FailedToLoad");
-	case ECurrentState::Loading:
+	case ELevelStreamingState::Loading:
 		return TEXT("Loading");
-	case ECurrentState::LoadedNotVisible:
+	case ELevelStreamingState::LoadedNotVisible:
 		return TEXT("LoadedNotVisible");
-	case ECurrentState::MakingVisible:
+	case ELevelStreamingState::MakingVisible:
 		return TEXT("MakingVisible");
-	case ECurrentState::LoadedVisible:
+	case ELevelStreamingState::LoadedVisible:
 		return TEXT("LoadedVisible");
-	case ECurrentState::MakingInvisible:
+	case ELevelStreamingState::MakingInvisible:
 		return TEXT("MakingInvisible");
 	}
 	ensure(false);
 	return TEXT("Unknown");
 }
 
-const TCHAR* ULevelStreaming::EnumToString(ETargetState InTargetState)
+const TCHAR* EnumToString(ELevelStreamingTargetState InTargetState)
 {
 	switch (InTargetState)
 	{
-	case ETargetState::Unloaded:
+	case ELevelStreamingTargetState::Unloaded:
 		return TEXT("Unloaded");
-	case ETargetState::UnloadedAndRemoved:
+	case ELevelStreamingTargetState::UnloadedAndRemoved:
 		return TEXT("UnloadedAndRemoved");
-	case ETargetState::LoadedNotVisible:
+	case ELevelStreamingTargetState::LoadedNotVisible:
 		return TEXT("LoadedNotVisible");
-	case ETargetState::LoadedVisible:
+	case ELevelStreamingTargetState::LoadedVisible:
 		return TEXT("LoadedVisible");
 	}
 	ensure(false);
@@ -1223,11 +1235,11 @@ void ULevelStreaming::SetLoadedLevel(ULevel* Level)
 		}
 		LC.AddLevel(LoadedLevel);
 
-		CurrentState = (LoadedLevel->bIsVisible ? ECurrentState::LoadedVisible : ECurrentState::LoadedNotVisible);
+		SetCurrentState((LoadedLevel->bIsVisible ? ELevelStreamingState::LoadedVisible : ELevelStreamingState::LoadedNotVisible));
 	}
 	else
 	{
-		CurrentState = ECurrentState::Unloaded;
+		SetCurrentState(ELevelStreamingState::Unloaded);
 	}
 
 	World->UpdateStreamingLevelShouldBeConsidered(this);
@@ -1339,13 +1351,13 @@ void ULevelStreaming::PrepareLoadedLevel(ULevel* InLevel, UPackage* InLevelPacka
 bool ULevelStreaming::RequestLevel(UWorld* PersistentWorld, bool bAllowLevelLoadRequests, EReqLevelBlock BlockPolicy)
 {
 	// Quit early in case load request already issued
-	if (CurrentState == ECurrentState::Loading)
+	if (CurrentState == ELevelStreamingState::Loading)
 	{
 		return true;
 	}
 
 	// Previous attempts have failed, no reason to try again
-	if (CurrentState == ECurrentState::FailedToLoad)
+	if (CurrentState == ELevelStreamingState::FailedToLoad)
 	{
 		return false;
 	}
@@ -1386,8 +1398,8 @@ bool ULevelStreaming::RequestLevel(UWorld* PersistentWorld, bool bAllowLevelLoad
 			continue;
 		}
 
-		const ECurrentState OtherState = OtherLevel->GetCurrentState();
-		if (OtherState == ECurrentState::FailedToLoad || OtherState == ECurrentState::Removed || (OtherState == ECurrentState::Unloaded && (OtherLevel->TargetState == ETargetState::Unloaded || OtherLevel->TargetState == ETargetState::UnloadedAndRemoved)))
+		const ELevelStreamingState OtherState = OtherLevel->CurrentState;
+		if (OtherState == ELevelStreamingState::FailedToLoad || OtherState == ELevelStreamingState::Removed || (OtherState == ELevelStreamingState::Unloaded && (OtherLevel->TargetState == ELevelStreamingTargetState::Unloaded || OtherLevel->TargetState == ELevelStreamingTargetState::UnloadedAndRemoved)))
 		{
 			// If the other level isn't loaded or in the process of being loaded we don't need to consider it
 			continue;
@@ -1402,7 +1414,7 @@ bool ULevelStreaming::RequestLevel(UWorld* PersistentWorld, bool bAllowLevelLoad
 			else
 			{
 				UE_LOG(LogLevelStreaming, Warning, TEXT("Streaming Level '%s' uses same destination for level ('%s') as '%s'. Level cannot be loaded again and this StreamingLevel will be flagged as failed to load."), *GetPathName(), *WorldAsset.GetLongPackageName(), *OtherLevel->GetPathName());
-				CurrentState = ECurrentState::FailedToLoad;
+				SetCurrentState(ELevelStreamingState::FailedToLoad);
 				return false;
 			}
 		}
@@ -1524,7 +1536,7 @@ bool ULevelStreaming::RequestLevel(UWorld* PersistentWorld, bool bAllowLevelLoad
 
 		if (FPackageName::DoesPackageExist(PackagePath, &PackagePath))
 		{
-			CurrentState = ECurrentState::Loading;
+			SetCurrentState(ELevelStreamingState::Loading);
 			FWorldNotifyStreamingLevelLoading::Started(PersistentWorld);
 			
 			ULevel::StreamedLevelsOwningWorld.Add(DesiredPackageName, PersistentWorld);
@@ -1580,7 +1592,7 @@ bool ULevelStreaming::RequestLevel(UWorld* PersistentWorld, bool bAllowLevelLoad
 		else
 		{
 			UE_LOG(LogStreaming, Error,TEXT("Couldn't find file for package %s."), *PackagePath.GetDebugName());
-			CurrentState = ECurrentState::FailedToLoad;
+			SetCurrentState(ELevelStreamingState::FailedToLoad);
 			return false;
 		}
 	}
@@ -1590,7 +1602,7 @@ bool ULevelStreaming::RequestLevel(UWorld* PersistentWorld, bool bAllowLevelLoad
 
 void ULevelStreaming::AsyncLevelLoadComplete(const FName& InPackageName, UPackage* InLoadedPackage, EAsyncLoadingResult::Type Result)
 {
-	CurrentState = ECurrentState::LoadedNotVisible;
+	SetCurrentState(ELevelStreamingState::LoadedNotVisible);
 	if (UWorld* World = GetWorld())
 	{
 		if (World->GetStreamingLevels().Contains(this))
@@ -1701,14 +1713,14 @@ void ULevelStreaming::AsyncLevelLoadComplete(const FName& InPackageName, UPackag
 	else if (Result == EAsyncLoadingResult::Canceled)
 	{
 		// Cancel level streaming
-		CurrentState = ECurrentState::Unloaded;
+		SetCurrentState(ELevelStreamingState::Unloaded);
 		SetShouldBeLoaded(false);
 	}
 	else
 	{
 		UE_LOG(LogLevelStreaming, Warning, TEXT("Failed to load package '%s'"), *InPackageName.ToString() );
 		
-		CurrentState = ECurrentState::FailedToLoad;
+		SetCurrentState(ELevelStreamingState::FailedToLoad);
  		SetShouldBeLoaded(false);
 	}
 
@@ -1863,9 +1875,9 @@ void ULevelStreaming::SetWorldAsset(const TSoftObjectPtr<UWorld>& NewWorldAsset)
 		bHasCachedWorldAssetPackageFName = false;
 		bHasCachedLoadedLevelPackageName = false;
 
-		if (CurrentState == ECurrentState::FailedToLoad)
+		if (CurrentState == ELevelStreamingState::FailedToLoad)
 		{
-			CurrentState = ECurrentState::Unloaded;
+			SetCurrentState(ELevelStreamingState::Unloaded);
 		}
 
 		if (UWorld* World = GetWorld())
@@ -1964,7 +1976,7 @@ void ULevelStreaming::SetPriority(const int32 NewPriority)
 	{
 		StreamingPriority = NewPriority;
 
-		if (CurrentState != ECurrentState::Removed && CurrentState != ECurrentState::FailedToLoad)
+		if (CurrentState != ELevelStreamingState::Removed && CurrentState != ELevelStreamingState::FailedToLoad)
 		{
 			if (UWorld* World = GetWorld())
 			{
@@ -1980,9 +1992,9 @@ void ULevelStreaming::SetLevelLODIndex(const int32 LODIndex)
 	{
 		LevelLODIndex = LODIndex;
 
-		if (CurrentState == ECurrentState::FailedToLoad)
+		if (CurrentState == ELevelStreamingState::FailedToLoad)
 		{
-			CurrentState = ECurrentState::Unloaded;
+			SetCurrentState(ELevelStreamingState::Unloaded);
 		}
 
 		if (UWorld* World = GetWorld())
@@ -2154,21 +2166,21 @@ bool ULevelStreaming::IsValidStreamingLevel() const
 
 EStreamingStatus ULevelStreaming::GetLevelStreamingStatus() const
 {
-	if (CurrentState == ECurrentState::FailedToLoad)
+	if (CurrentState == ELevelStreamingState::FailedToLoad)
 	{
 		return LEVEL_FailedToLoad;
 	}
-	else if (CurrentState == ECurrentState::MakingInvisible)
+	else if (CurrentState == ELevelStreamingState::MakingInvisible)
 	{
 		return LEVEL_MakingInvisible;
 	}
 	else if (LoadedLevel)
 	{
-		if (CurrentState == ECurrentState::LoadedVisible)
+		if (CurrentState == ELevelStreamingState::LoadedVisible)
 		{
 			return LEVEL_Visible;
 		}
-		else if ((CurrentState == ECurrentState::MakingVisible) && (GetWorld()->GetCurrentLevelPendingVisibility() == LoadedLevel))
+		else if ((CurrentState == ELevelStreamingState::MakingVisible) && (GetWorld()->GetCurrentLevelPendingVisibility() == LoadedLevel))
 		{
 			return LEVEL_MakingVisible;
 		}
@@ -2176,12 +2188,12 @@ EStreamingStatus ULevelStreaming::GetLevelStreamingStatus() const
 	}
 	else
 	{
-		if (CurrentState == ECurrentState::Loading)
+		if (CurrentState == ELevelStreamingState::Loading)
 		{
 			return LEVEL_Loading;
 		}
 
-		check(CurrentState == ECurrentState::Removed || CurrentState == ECurrentState::Unloaded);
+		check(CurrentState == ELevelStreamingState::Removed || CurrentState == ELevelStreamingState::Unloaded);
 		// See whether the level's world object is still around.
 		UPackage* LevelPackage = FindObjectFast<UPackage>(nullptr, GetWorldAssetPackageFName());
 		UWorld* LevelWorld = LevelPackage ? UWorld::FindWorldInPackage(LevelPackage) : nullptr;
diff --git a/Engine/Source/Runtime/Engine/Private/World.cpp b/Engine/Source/Runtime/Engine/Private/World.cpp
--- a/Engine/Source/Runtime/Engine/Private/World.cpp
+++ b/Engine/Source/Runtime/Engine/Private/World.cpp
@@ -3982,7 +3982,7 @@ void UWorld::AddStreamingLevel(ULevelStreaming* StreamingLevelToAdd)
 	{
 		if (ensure(StreamingLevelToAdd->GetWorld() == this))
 		{
-			if (ensure(StreamingLevelToAdd->GetCurrentState() == ULevelStreaming::ECurrentState::Removed))
+			if (ensure(StreamingLevelToAdd->GetLevelStreamingState() == ELevelStreamingState::Removed))
 			{
 				StreamingLevels.Add(StreamingLevelToAdd);
 				FStreamingLevelPrivateAccessor::OnLevelAdded(StreamingLevelToAdd);
@@ -4039,7 +4039,7 @@ bool UWorld::RemoveStreamingLevelAt(const int32 IndexToRemove)
 	if (IndexToRemove >= 0 && IndexToRemove < StreamingLevels.Num())
 	{
 		ULevelStreaming* StreamingLevel = StreamingLevels[IndexToRemove];
-		if (StreamingLevel && StreamingLevel->GetCurrentState() == ULevelStreaming::ECurrentState::Loading)
+		if (StreamingLevel && StreamingLevel->GetLevelStreamingState() == ELevelStreamingState::Loading)
 		{
 			--NumStreamingLevelsBeingLoaded;
 		}
@@ -4086,11 +4086,11 @@ void UWorld::PopulateStreamingLevelsToConsider()
 	for (ULevelStreaming* StreamingLevel : StreamingLevels)
 	{
 		// The streaming level may have just gotten added to the list so have it update its state now
-		if (StreamingLevel->GetCurrentState() == ULevelStreaming::ECurrentState::Removed)
+		if (StreamingLevel->GetLevelStreamingState() == ELevelStreamingState::Removed)
 		{
 			FStreamingLevelPrivateAccessor::OnLevelAdded(StreamingLevel);
 		}
-		else if (StreamingLevel->GetCurrentState() == ULevelStreaming::ECurrentState::Loading)
+		else if (StreamingLevel->GetLevelStreamingState() == ELevelStreamingState::Loading)
 		{
 			++NumStreamingLevelsBeingLoaded;
 		}
@@ -4104,7 +4104,7 @@ void UWorld::PopulateStreamingLevelsToConsider()
 
 void UWorld::UpdateStreamingLevelShouldBeConsidered(ULevelStreaming* StreamingLevelToConsider)
 {
-	if (StreamingLevelToConsider && ensure(StreamingLevelToConsider->GetWorld() == this) && StreamingLevelToConsider->GetCurrentState() != ULevelStreaming::ECurrentState::Removed)
+	if (StreamingLevelToConsider && ensure(StreamingLevelToConsider->GetWorld() == this) && StreamingLevelToConsider->GetLevelStreamingState() != ELevelStreamingState::Removed)
 	{
 		if (FStreamingLevelPrivateAccessor::DetermineTargetState(StreamingLevelToConsider))
 		{
@@ -4240,7 +4240,7 @@ bool UWorld::AreAlwaysLoadedLevelsLoaded() const
 	for (ULevelStreaming* LevelStreaming : StreamingLevels)
 	{
 		// See whether there's a level with a pending request.
-		if (LevelStreaming && LevelStreaming->ShouldBeAlwaysLoaded() && LevelStreaming->GetCurrentState() != ULevelStreaming::ECurrentState::FailedToLoad)
+		if (LevelStreaming && LevelStreaming->ShouldBeAlwaysLoaded() && LevelStreaming->GetLevelStreamingState() != ELevelStreamingState::FailedToLoad)
 		{	
 			const ULevel* LoadedLevel = LevelStreaming->GetLoadedLevel();
 
diff --git a/Engine/Source/Runtime/Engine/Private/WorldPartition/WorldPartitionLevelStreamingDynamic.cpp b/Engine/Source/Runtime/Engine/Private/WorldPartition/WorldPartitionLevelStreamingDynamic.cpp
--- a/Engine/Source/Runtime/Engine/Private/WorldPartition/WorldPartitionLevelStreamingDynamic.cpp
+++ b/Engine/Source/Runtime/Engine/Private/WorldPartition/WorldPartitionLevelStreamingDynamic.cpp
@@ -184,13 +184,13 @@ void UWorldPartitionLevelStreamingDynamic::CreateRuntimeLevel()
 bool UWorldPartitionLevelStreamingDynamic::RequestLevel(UWorld* InPersistentWorld, bool bInAllowLevelLoadRequests, EReqLevelBlock InBlockPolicy)
 {
 	// Quit early in case load request already issued
-	if (GetCurrentState() == ECurrentState::Loading)
+	if (GetLevelStreamingState() == ELevelStreamingState::Loading)
 	{
 		return true;
 	}
 
 	// Previous attempts have failed, no reason to try again
-	if (GetCurrentState() == ECurrentState::FailedToLoad)
+	if (GetLevelStreamingState() == ELevelStreamingState::FailedToLoad)
 	{
 		return false;
 	}
@@ -262,7 +262,7 @@ bool UWorldPartitionLevelStreamingDynamic::RequestLevel(UWorld* InPersistentWorl
 		check(LODPackageNames.Num() == 0);
 		if (RuntimeLevel == nullptr)
 		{
-			check(GetCurrentState() == ECurrentState::Unloaded);
+			check(GetLevelStreamingState() == ELevelStreamingState::Unloaded);
 
 			check(!RuntimeLevel);
 			CreateRuntimeLevel();
@@ -291,7 +291,7 @@ bool UWorldPartitionLevelStreamingDynamic::RequestLevel(UWorld* InPersistentWorl
 				}
 				else
 				{
-					CurrentState = ECurrentState::Loading;
+					CurrentState = ELevelStreamingState::Loading;
 				}
 			}
 		}
diff --git a/Engine/Source/Runtime/Engine/Private/WorldPartition/WorldPartitionRuntimeLevelStreamingCell.cpp b/Engine/Source/Runtime/Engine/Private/WorldPartition/WorldPartitionRuntimeLevelStreamingCell.cpp
--- a/Engine/Source/Runtime/Engine/Private/WorldPartition/WorldPartitionRuntimeLevelStreamingCell.cpp
+++ b/Engine/Source/Runtime/Engine/Private/WorldPartition/WorldPartitionRuntimeLevelStreamingCell.cpp
@@ -30,12 +30,12 @@ EWorldPartitionRuntimeCellState UWorldPartitionRuntimeLevelStreamingCell::GetCur
 {
 	if (LevelStreaming)
 	{
-		ULevelStreaming::ECurrentState CurrentStreamingState = LevelStreaming->GetCurrentState();
-		if (CurrentStreamingState == ULevelStreaming::ECurrentState::LoadedVisible)
+		ELevelStreamingState CurrentStreamingState = LevelStreaming->GetLevelStreamingState();
+		if (CurrentStreamingState == ELevelStreamingState::LoadedVisible)
 		{
 			return EWorldPartitionRuntimeCellState::Activated;
 		}
-		else if (CurrentStreamingState >= ULevelStreaming::ECurrentState::LoadedNotVisible)
+		else if (CurrentStreamingState >= ELevelStreamingState::LoadedNotVisible)
 		{
 			return EWorldPartitionRuntimeCellState::Loaded;
 		}
@@ -125,8 +125,8 @@ bool UWorldPartitionRuntimeLevelStreamingCell::IsLoading() const
 {
 	if (LevelStreaming)
 	{
-		ULevelStreaming::ECurrentState CurrentState = LevelStreaming->GetCurrentState();
-		return (CurrentState == ULevelStreaming::ECurrentState::Removed || CurrentState == ULevelStreaming::ECurrentState::Unloaded || CurrentState == ULevelStreaming::ECurrentState::Loading);
+		ELevelStreamingState CurrentState = LevelStreaming->GetLevelStreamingState();
+		return (CurrentState == ELevelStreamingState::Removed || CurrentState == ELevelStreamingState::Unloaded || CurrentState == ELevelStreamingState::Loading);
 	}
 	return Super::IsLoading();
 }
@@ -354,7 +354,7 @@ bool UWorldPartitionRuntimeLevelStreamingCell::CanAddToWorld() const
 {
 	return LevelStreaming &&
 		   LevelStreaming->GetLoadedLevel() &&
-		   (LevelStreaming->GetCurrentState() == ULevelStreaming::ECurrentState::MakingVisible);
+		   (LevelStreaming->GetLevelStreamingState() == ELevelStreamingState::MakingVisible);
 }
 
 void UWorldPartitionRuntimeLevelStreamingCell::SetStreamingPriority(int32 InStreamingPriority) const
