Delete deprecated automation classes

These were using Matinee before, which has been removed. These automation classes would have to be rewritten using sequencer, but level sequence actors are in a different module.

diff --git a/Engine/Source/Runtime/AutomationWorker/Private/AutomationAnalyticParams.h b/Engine/Source/Runtime/AutomationWorker/Private/AutomationAnalyticParams.h
--- a/Engine/Source/Runtime/AutomationWorker/Private/AutomationAnalyticParams.h
+++ b/Engine/Source/Runtime/AutomationWorker/Private/AutomationAnalyticParams.h
@@ -20,7 +20,6 @@ namespace EAutomationAnalyticParam
 	enum Type
 	{
 		MapName,
-		SequenceName,
 		TimeStamp,
 		Platform,
 		Spec,
diff --git a/Engine/Source/Runtime/AutomationWorker/Private/AutomationAnalytics.cpp b/Engine/Source/Runtime/AutomationWorker/Private/AutomationAnalytics.cpp
--- a/Engine/Source/Runtime/AutomationWorker/Private/AutomationAnalytics.cpp
+++ b/Engine/Source/Runtime/AutomationWorker/Private/AutomationAnalytics.cpp
@@ -111,7 +111,6 @@ void FAutomationAnalytics::InititalizeAnalyticParameterNames()
 	AutomationParamNames.SetNum(EAutomationAnalyticParam::NUM_ENGINE_PARAMS);
 
 	AutomationParamNames[EAutomationAnalyticParam::MapName] = TEXT("MapName");
-	AutomationParamNames[EAutomationAnalyticParam::SequenceName] = TEXT("SequenceName");
 	AutomationParamNames[EAutomationAnalyticParam::TimeStamp] = TEXT("TimeStamp");
 	AutomationParamNames[EAutomationAnalyticParam::Platform] = TEXT("Platform");
 	AutomationParamNames[EAutomationAnalyticParam::Spec] = TEXT("Spec");
@@ -169,7 +168,6 @@ void FAutomationAnalytics::FireEvent_FPSCapture(const FAutomationPerformanceSnap
 	// @EventParam CL string The Changelist for the build
 	// @EventParam Spec string The spec of the machine running the test
 	// @EventParam MapName string The map that the test was run on
-	// @EventParam MatineeName string The name of the matinee that fired the event
 	// @EventParam FPS string 
 	// @EventParam BuildConfiguration string Debug/Development/Test/Shipping
 	// @EventParam AverageFrameTime string Time for a frame in ms
@@ -188,7 +186,6 @@ void FAutomationAnalytics::FireEvent_FPSCapture(const FAutomationPerformanceSnap
 		ParamArray.Add(FAnalyticsEventAttribute(GetAutomationParamName(EAutomationAnalyticParam::CL), PerfSnapshot.Changelist));
 		ParamArray.Add(FAnalyticsEventAttribute(GetAutomationParamName(EAutomationAnalyticParam::Spec), MachineSpec));
 		ParamArray.Add(FAnalyticsEventAttribute(GetAutomationParamName(EAutomationAnalyticParam::MapName), PerfSnapshot.MapName));
-		ParamArray.Add(FAnalyticsEventAttribute(GetAutomationParamName(EAutomationAnalyticParam::SequenceName), PerfSnapshot.SequenceName));
 		ParamArray.Add(FAnalyticsEventAttribute(GetAutomationParamName(EAutomationAnalyticParam::FPS), PerfSnapshot.AverageFPS));
 
 		ParamArray.Add(FAnalyticsEventAttribute(GetAutomationParamName(EAutomationAnalyticParam::BuildConfiguration), PerfSnapshot.BuildConfiguration));
diff --git a/Engine/Source/Runtime/Engine/Private/Tests/AutomationCommon.cpp b/Engine/Source/Runtime/Engine/Private/Tests/AutomationCommon.cpp
--- a/Engine/Source/Runtime/Engine/Private/Tests/AutomationCommon.cpp
+++ b/Engine/Source/Runtime/Engine/Private/Tests/AutomationCommon.cpp
@@ -685,17 +685,6 @@ bool FWaitForInteractiveFrameRate::Update()
 // Common Latent commands which are used across test type. I.e. Engine, Network, etc...
 
 
-bool FPlayMatineeLatentCommand::Update()
-{
-	return true;
-}
-
-
-bool FWaitForMatineeToCompleteLatentCommand::Update()
-{
-	return true;
-}
-
 
 bool FExecStringLatentCommand::Update()
 {
@@ -744,17 +733,6 @@ bool FStreamAllResourcesLatentCommand::Update()
 	return true;
 }
 
-bool FEnqueuePerformanceCaptureCommands::Update()
-{
-	return true;
-}
-
-
-bool FMatineePerformanceCaptureCommand::Update()
-{
-	return true;
-
-}
 
 
 bool FExecWorldStringLatentCommand::Update()
diff --git a/Engine/Source/Runtime/Engine/Private/Tests/EngineAutomationTests.cpp b/Engine/Source/Runtime/Engine/Private/Tests/EngineAutomationTests.cpp
--- a/Engine/Source/Runtime/Engine/Private/Tests/EngineAutomationTests.cpp
+++ b/Engine/Source/Runtime/Engine/Private/Tests/EngineAutomationTests.cpp
@@ -160,30 +160,6 @@ bool FStatsVerificationMapTest::RunTest(const FString& Parameters)
 	return true;
 }
 
-/**
- * LoadAutomationMap
- * Verification automation test to make sure features of map loading work (load, screen capture, performance capture)
- */
-IMPLEMENT_SIMPLE_AUTOMATION_TEST(FPerformanceCaptureTest, "System.Maps.Performance Capture", EAutomationTestFlags::ClientContext | EAutomationTestFlags::EngineFilter)
-
-/** 
- * Execute the loading of one map to verify screen captures and performance captures work
- *
- * @param Parameters - Unused for this test
- * @return	TRUE if the test was successful, FALSE otherwise
- */
-bool FPerformanceCaptureTest::RunTest(const FString& Parameters)
-{
-	UAutomationTestSettings const* AutomationTestSettings = GetDefault<UAutomationTestSettings>();
-	check(AutomationTestSettings);
-	FString MapName = AutomationTestSettings->AutomationTestmap.GetLongPackageName();
-
-	GEngine->Exec(GetSimpleEngineAutomationTestGameWorld(GetTestFlags()), *FString::Printf(TEXT("Open %s"), *MapName));
-	ADD_LATENT_AUTOMATION_COMMAND(FEnqueuePerformanceCaptureCommands());
-
-	return true;
-}
-
 /**
  * Latent command to take a screenshot of the viewport
  */
@@ -197,54 +173,6 @@ bool FTakeViewportScreenshotCommand::Update()
 	return true;
 }
 
-/**
- * LoadAllMapsInGame
- * Verification automation test to make sure loading all maps succeed without crashing AND does performance captures
- */
-IMPLEMENT_COMPLEX_AUTOMATION_TEST(FLoadAllMapsInGameTest, "Project.Maps.Load All In Game", EAutomationTestFlags::ClientContext | EAutomationTestFlags::StressFilter)
-
-/** 
- * Requests a enumeration of all maps to be loaded
- */
-void FLoadAllMapsInGameTest::GetTests(TArray<FString>& OutBeautifiedNames, TArray <FString>& OutTestCommands) const
-{
-	PopulateTestsForAllAvailableMaps(OutBeautifiedNames, OutTestCommands);
-}
-
-/** 
- * Execute the loading of each map and performance captures
- *
- * @param Parameters - Should specify which map name to load
- * @return	TRUE if the test was successful, FALSE otherwise
- */
-bool FLoadAllMapsInGameTest::RunTest(const FString& Parameters)
-{
-	FString MapName = Parameters;
-
-	//Open the map
-	GEngine->Exec(GetSimpleEngineAutomationTestGameWorld(GetTestFlags()), *FString::Printf(TEXT("Open %s"), *MapName));
-
-	//if( FAutomationTestFramework::Get().IsScreenshotAllowed() )
-	{
-		//Generate the screen shot name and path
-		FString ScreenshotFileName;
-		const FString LoadAllMapsTestName = FString::Printf(TEXT("LoadAllMaps_Game/%s"), *FPaths::GetBaseFilename(MapName));
-		ScreenshotFileName = AutomationCommon::GetLocalPathForScreenshot(AutomationCommon::GetScreenshotName(LoadAllMapsTestName));
-
-		//Give the map some time to load
-		ADD_LATENT_AUTOMATION_COMMAND(FEngineWaitLatentCommand(1.5f));
-		//Take the screen shot
-		ADD_LATENT_AUTOMATION_COMMAND(FTakeViewportScreenshotCommand(ScreenshotFileName));
-		//Give the screen shot a chance to capture the scene
-		ADD_LATENT_AUTOMATION_COMMAND(FEngineWaitLatentCommand(0.5f));
-	}
-
-	//Kick off any Automation matinees that are in this map
-	ADD_LATENT_AUTOMATION_COMMAND(FEnqueuePerformanceCaptureCommands());
-
-	return true;
-}
-
 /**
  * SaveGameTest
  * Test makes sure a save game (without UI) saves and loads correctly
@@ -308,98 +236,6 @@ bool FSaveGameTest::RunTest(const FString& Parameters)
 	return LoadedData == SavedData;
 }
 
-/**
- * Automation test to load a map and capture FPS performance charts
- */
-IMPLEMENT_COMPLEX_AUTOMATION_TEST(FCinematicFPSPerfTest, "Project.Maps.Cinematic FPS Perf Capture", (EAutomationTestFlags::ClientContext | EAutomationTestFlags::NonNullRHI | EAutomationTestFlags::StressFilter));
-
-void FCinematicFPSPerfTest::GetTests(TArray<FString>& OutBeautifiedNames, TArray <FString>& OutTestCommands) const
-{
-	PopulateTestsForAllAvailableMaps(OutBeautifiedNames, OutTestCommands);
-}
-
-bool FCinematicFPSPerfTest::RunTest(const FString& Parameters)
-{
-	//Map to use for this test.
-	const FString MapName = Parameters;
-
-	//The name of the matinee actor that will be triggered.
-	FString MatineeActorName;
-
-	//Check we are running from commandline
-	const FString CommandLine(FCommandLine::Get());
-	if (CommandLine.Contains(TEXT("Automation")))
-	{
-		//Get the name of the matinee to be used.
-		//If the game was not launched with the -MatineeName argument then this test will be ran based on time.
-		if (FParse::Value(*CommandLine, TEXT("MatineeName="), MatineeActorName))
-		{
-			//Load map
-			ADD_LATENT_AUTOMATION_COMMAND(FEngineWaitLatentCommand(1.0f));
-			ADD_LATENT_AUTOMATION_COMMAND(FLoadGameMapCommand(MapName));
-			ADD_LATENT_AUTOMATION_COMMAND(FEngineWaitLatentCommand(1.0f));
-
-			//Start the matinee and perform the FPS Chart
-			ADD_LATENT_AUTOMATION_COMMAND(FMatineePerformanceCaptureCommand(MatineeActorName));
-			ADD_LATENT_AUTOMATION_COMMAND(FEngineWaitLatentCommand(1.0f));
-
-			return true;
-		}
-		else
-		{
-			UE_LOG(LogEngineAutomationTests, Log, TEXT("The matinee name was not specified.  Run the game with -MatineeName=\"Name of the matinee actor\"."));
-
-			//Get the name of the console event to trigger the cinematic
-			FString CinematicEventCommand;
-			if (!FParse::Value(*CommandLine, TEXT("CE="), CinematicEventCommand))
-			{
-				UE_LOG(LogEngineAutomationTests, Log, TEXT("A console event command was not specified. Defaults to CE START.  Run the game with -CE=\"Command\"."));
-				CinematicEventCommand = TEXT("CE Start");
-			}
-
-			//Get the length of time the cinematic will run
-			float RunTime;
-			if (!FParse::Value(*CommandLine, TEXT("RunTime="), RunTime))
-			{
-				UE_LOG(LogEngineAutomationTests, Log, TEXT("A run time length in seconds was not specified. Defaults to 60 seconds. Run the game with -RunTime=###."));
-				RunTime = 60.f;
-			}
-
-			//Load map
-			ADD_LATENT_AUTOMATION_COMMAND(FEngineWaitLatentCommand(1.0f));
-			ADD_LATENT_AUTOMATION_COMMAND(FLoadGameMapCommand(MapName));
-			ADD_LATENT_AUTOMATION_COMMAND(FEngineWaitLatentCommand(1.0f));
-
-			//Start the matinee and perform the FPS Chart
-			ADD_LATENT_AUTOMATION_COMMAND(FExecWorldStringLatentCommand(CinematicEventCommand));
-			ADD_LATENT_AUTOMATION_COMMAND(FExecWorldStringLatentCommand(TEXT("StartFPSChart")));
-			ADD_LATENT_AUTOMATION_COMMAND(FEngineWaitLatentCommand(RunTime));
-			ADD_LATENT_AUTOMATION_COMMAND(FExecWorldStringLatentCommand(TEXT("StopFPSChart")));
-			ADD_LATENT_AUTOMATION_COMMAND(FEngineWaitLatentCommand(1.0f));
-
-			return true;
-
-		}
-	}
-	//If the user is running from the UFE then we'll use the default values.
-	//@todo Give the end user a way to specify the values for this test.
-
-	UE_LOG(LogEngineAutomationTests, Log, TEXT("Running the FPS chart performance capturing for 60 seconds while in '%s'.\nThe default CE command won't be used at this time."), *MapName);
-
-	//Load map
-	ADD_LATENT_AUTOMATION_COMMAND(FLoadGameMapCommand(MapName));
-	ADD_LATENT_AUTOMATION_COMMAND(FEngineWaitLatentCommand(1.0f));
-
-	//Start the matinee and perform the FPS Chart
-	ADD_LATENT_AUTOMATION_COMMAND(FExecWorldStringLatentCommand(TEXT("StartFPSChart")));
-	ADD_LATENT_AUTOMATION_COMMAND(FEngineWaitLatentCommand(60.0f));
-	ADD_LATENT_AUTOMATION_COMMAND(FExecWorldStringLatentCommand(TEXT("StopFPSChart")));
-	ADD_LATENT_AUTOMATION_COMMAND(FEngineWaitLatentCommand(1.0f));
-
-	return true;
-}
-
-
 IMPLEMENT_SIMPLE_AUTOMATION_TEST(FAutomationLogAddMessage, "System.Automation.Log.Add Log Message", EAutomationTestFlags::ApplicationContextMask | EAutomationTestFlags::EngineFilter)
 
 bool FAutomationLogAddMessage::RunTest(const FString& Parameters)
@@ -1076,111 +912,4 @@ bool FAutomationAttachment::RunTest(const FString& Parameters)
 	return true;
 }
 
-// ----------------------------------------------------------------------------------------------------------------
-
-DEFINE_LATENT_AUTOMATION_COMMAND_ONE_PARAMETER(FWaitForMatineeToCompleteAndDoScreenshotsLatentCommand, AActor*, MatineeActor);
-
-bool FWaitForMatineeToCompleteAndDoScreenshotsLatentCommand::Update()
-{
-	return true;
-}
-
-///**
-// * FRenderOutputValidation - render deterministic and compare results of multiple runs
-// */
-//IMPLEMENT_SIMPLE_AUTOMATION_TEST(FRenderOutputValidation, "Rendering.RenderOutputValidation", EAutomationTestFlags::ClientContext | EAutomationTestFlags::EngineFilter)
-//
-///** 
-// * render deterministic and compare results of multiple runs
-// *
-// * @param Parameters - Unused for this test
-// * @return	TRUE if the test was successful, FALSE otherwise
-// */
-//bool FRenderOutputValidation::RunTest(const FString& Parameters)
-//{
-//	const bool bChangeResolution = false;
-//	const bool bLoadMap = false;			// true doesn't work at the moment (iterate matinees needs to happen after loading the map)
-//
-//	//Gets the default map that the game uses.
-//	const UGameMapsSettings* GameMapsSettings = GetDefault<UGameMapsSettings>();
-////	const FString& MapName = GameMapsSettings->GetGameDefaultMap();
-//	const FString MapName = TEXT("ScreenshotsTest");
-//
-//	// request full res screenshots
-//	//FAutomationTestFramework::Get().SetScreenshotOptions(true);
-//
-//	if(bLoadMap)
-//	{
-//		// Opens the actual default map in game.
-//		GEngine->Exec(GetSimpleEngineAutomationTestGameWorld(GetTestFlags()), *FString::Printf(TEXT("Open %s"), *MapName));
-//		ADD_LATENT_AUTOMATION_COMMAND(FEngineWaitLatentCommand(2.0f));
-//	}
-//
-//	{
-//		UConsole* ViewportConsole = (GEngine->GameViewport != nullptr) ? GEngine->GameViewport->ViewportConsole : nullptr;
-//		FConsoleOutputDevice StrOut(ViewportConsole);
-//
-//		StrOut.Logf(TEXT("FRenderOutputValidation:Runtest()"));
-//	}
-//
-//	// Gets the current resolution.
-//	FString RestoreResolutionString = FString::Printf(TEXT("setres %dx%d"), GSystemResolution.ResX, GSystemResolution.ResY);
-//
-//	if(bChangeResolution)
-//	{
-//		// Change the resolution
-//		ADD_LATENT_AUTOMATION_COMMAND(FExecStringLatentCommand(TEXT("setres 640x480")));
-//		ADD_LATENT_AUTOMATION_COMMAND(FEngineWaitLatentCommand(2.0f));
-//	}
-//	// 5 sec should be enough, if more is needed - we should investigate
-//	ADD_LATENT_AUTOMATION_COMMAND(FStreamAllResourcesLatentCommand(5.0f));
-//
-////later	ADD_LATENT_AUTOMATION_COMMAND(FBuildLightingCommand);
-//
-//	// streaming should not be distributing it's work over multiple frames - to be more deterministic
-//	ADD_LATENT_AUTOMATION_COMMAND(FExecWorldStringLatentCommand(TEXT("r.Streaming.FramesForFullUpdate 0")));
-//
-//	// reset all data in the ViewState e.g. TemporalAA, SSR cyclic counter, TemporalAA images
-//	ADD_LATENT_AUTOMATION_COMMAND(FExecWorldStringLatentCommand(TEXT("r.ResetViewState")));
-//
-//	// this needs to be improved: better thank hanging would be to render frames (showing the user that it waits on something)
-//	ADD_LATENT_AUTOMATION_COMMAND(FWaitForShadersToFinishCompilingInGame);
-//
-//	//if( FAutomationTestFramework::Get().IsScreenshotAllowed() )
-//	{
-//		for (TObjectIterator<AMatineeActor> It; It; ++It)
-//		{
-//			AMatineeActor* MatineeActor = *It;
-//
-//			FString Name;
-//			MatineeActor->GetName(Name);
-//
-//			if(Name.Contains(TEXT("Automation")))
-//			{
-//				UE_LOG(LogEngineAutomationTests, Log, TEXT("Iterated matinee: %p '%s'"), MatineeActor, *MatineeActor->GetName())
-//
-//				//add latent action to execute this matinee
-//				ADD_LATENT_AUTOMATION_COMMAND(FPlayMatineeLatentCommand(MatineeActor));
-//
-//				//Run the Stat FPS Chart command
-//		//			ADD_LATENT_AUTOMATION_COMMAND(FExecWorldStringLatentCommand(TEXT("StartFPSChart")));
-//
-//				//add action to wait until matinee is complete
-//				ADD_LATENT_AUTOMATION_COMMAND(FWaitForMatineeToCompleteAndDoScreenshotsLatentCommand(MatineeActor));
-//
-//				//Stop the Stat FPS Chart command
-//		//			ADD_LATENT_AUTOMATION_COMMAND(FExecWorldStringLatentCommand(TEXT("StopFPSChart")));
-//			}
-//		}
-//	}
-//
-//	if(bChangeResolution)
-//	{
-//		// restore the resolution
-//		ADD_LATENT_AUTOMATION_COMMAND(FExecStringLatentCommand(RestoreResolutionString));
-//	}
-//
-//	return true;
-//}
-
 #endif //WITH_DEV_AUTOMATION_TESTS
diff --git a/Engine/Source/Runtime/Engine/Private/Tests/NetworkAutomationTests.cpp b/Engine/Source/Runtime/Engine/Private/Tests/NetworkAutomationTests.cpp
--- a/Engine/Source/Runtime/Engine/Private/Tests/NetworkAutomationTests.cpp
+++ b/Engine/Source/Runtime/Engine/Private/Tests/NetworkAutomationTests.cpp
@@ -68,20 +68,6 @@ bool FMultiplayer4PlayerTest::RunTest(const FString& Parameters)
 		END_NETWORK_AUTOMATION_COMMAND(InvitePlayers, i)
 	}
 
-
-	START_NETWORK_AUTOMATION_COMMAND(PerformanceHost)
-	{
-		ADD_LATENT_AUTOMATION_COMMAND(FEnqueuePerformanceCaptureCommands());
-	}
-	END_NETWORK_AUTOMATION_COMMAND(PerformanceHost, EMultiplayerAutomationRoles::Host)
-
-
-	START_NETWORK_AUTOMATION_COMMAND(PerformanceClient0)
-	{
-		ADD_LATENT_AUTOMATION_COMMAND(FEnqueuePerformanceCaptureCommands());
-	}
-	END_NETWORK_AUTOMATION_COMMAND(PerformanceClient0, EMultiplayerAutomationRoles::Client0)
-
 	return true;
 }
 
diff --git a/Engine/Source/Runtime/Engine/Public/Tests/AutomationCommon.h b/Engine/Source/Runtime/Engine/Public/Tests/AutomationCommon.h
--- a/Engine/Source/Runtime/Engine/Public/Tests/AutomationCommon.h
+++ b/Engine/Source/Runtime/Engine/Public/Tests/AutomationCommon.h
@@ -125,17 +125,6 @@ DEFINE_ENGINE_LATENT_AUTOMATION_COMMAND(FWaitForMapToLoadCommand);
 */
 DEFINE_ENGINE_LATENT_AUTOMATION_COMMAND_ONE_PARAMETER(FWaitForSpecifiedMapToLoadCommand, FString, MapName);
 
-/**
-* Force a matinee to not loop and request that it play
-*/
-DEFINE_ENGINE_LATENT_AUTOMATION_COMMAND_ONE_PARAMETER(FPlayMatineeLatentCommand, AActor*, MatineeActor);
-
-
-/**
-* Wait for a particular matinee actor to finish playing
-*/
-DEFINE_ENGINE_LATENT_AUTOMATION_COMMAND_ONE_PARAMETER(FWaitForMatineeToCompleteLatentCommand, AActor*, MatineeActor);
-
 
 /**
 * Execute command string
@@ -153,16 +142,6 @@ DEFINE_ENGINE_LATENT_AUTOMATION_COMMAND_ONE_PARAMETER(FEngineWaitLatentCommand,
 */
 DEFINE_ENGINE_LATENT_AUTOMATION_COMMAND_ONE_PARAMETER(FStreamAllResourcesLatentCommand, float, Duration);
 
-/**
-* Enqueue performance capture commands after a map has been loaded
-*/
-DEFINE_ENGINE_LATENT_AUTOMATION_COMMAND(FEnqueuePerformanceCaptureCommands);
-
-
-/**
-* Run FPS chart command using for the actual duration of the matinee.
-*/
-DEFINE_ENGINE_LATENT_AUTOMATION_COMMAND_ONE_PARAMETER(FMatineePerformanceCaptureCommand, FString, MatineeName);
 
 /**
 * Latent command to run an exec command that also requires a UWorld.
