Removing CachedCombinedLODBias from a ensureMsgf's format since it was removed in CL 22409186.

diff --git a/Engine/Source/Runtime/Engine/Private/Texture.cpp b/Engine/Source/Runtime/Engine/Private/Texture.cpp
--- a/Engine/Source/Runtime/Engine/Private/Texture.cpp
+++ b/Engine/Source/Runtime/Engine/Private/Texture.cpp
@@ -314,10 +314,10 @@ void UTexture::UpdateResource()
 					}
 
 					ensureMsgf( (SizeX%4)==0 && (SizeY%4) == 0, TEXT("Skipping init of %s texture %s with non 4x4-aligned size. Resource Size=%ix%ix%i. "
-						"Texture PD Size=%dx%d, mips=%d, nonstreaming=%d, nonopt=%d, LODBias=%d,cached=%d,cinematic=%d."), 
+						"Texture PD Size=%dx%d, mips=%d, nonstreaming=%d, nonopt=%d, LODBias=%d,cinematic=%d."), 
 						GPixelFormats[StreamableResource->GetPixelFormat()].Name, *GetName(), SizeX, SizeY, SizeZ,
 						PDSizeX,PDSizeY,MipsNum,NumNonStreamingMips,NumNonOptionalMips,
-						LODBias,CachedCombinedLODBias,NumCinematicMipLevels);
+						LODBias,NumCinematicMipLevels);
 
 					delete NewResource;
 					return;
