World Partition - HLOD: Fixed crash during HLOD, occurring when the engine performs a GC
#rnx

diff --git a/Engine/Source/Runtime/Engine/Private/WorldPartition/WorldPartitionLevelStreamingDynamic.cpp b/Engine/Source/Runtime/Engine/Private/WorldPartition/WorldPartitionLevelStreamingDynamic.cpp
--- a/Engine/Source/Runtime/Engine/Private/WorldPartition/WorldPartitionLevelStreamingDynamic.cpp
+++ b/Engine/Source/Runtime/Engine/Private/WorldPartition/WorldPartitionLevelStreamingDynamic.cpp
@@ -113,12 +113,9 @@ void UWorldPartitionLevelStreamingDynamic::UnloadFromEditor(UWorldPartitionLevel
 	World->FlushLevelStreaming();
 
 	// Destroy the package world and remove it from root
-	UWorld* PackageWorld = UWorld::FindWorldInPackage(Level->GetPackage());
-	if (PackageWorld)
-	{
-		PackageWorld->DestroyWorld(false);
-		PackageWorld->Rename(nullptr, GetTransientPackage());
-	}
+	UPackage* Package = Level->GetPackage();
+	UWorld* PackageWorld = UWorld::FindWorldInPackage(Package);
+	PackageWorld->DestroyWorld(false);
 }
 
 void UWorldPartitionLevelStreamingDynamic::Initialize(UWorld* OuterWorld, const TArray<FWorldPartitionRuntimeCellObjectMapping>& InPackages)
