Texture2D: fix ResourceMem pointer was not initialized

diff --git a/Engine/Source/Runtime/Engine/Private/Texture2D.cpp b/Engine/Source/Runtime/Engine/Private/Texture2D.cpp
--- a/Engine/Source/Runtime/Engine/Private/Texture2D.cpp
+++ b/Engine/Source/Runtime/Engine/Private/Texture2D.cpp
@@ -71,9 +71,13 @@ UTexture2D::UTexture2D(const FObjectInitializer& ObjectInitializer)
 		[this]()-> FTexturePlatformData* { return GetPlatformData(); },
 		[this](FTexturePlatformData* InPlatformData) { SetPlatformData(InPlatformData); })
 #endif
+	, ResourceMem(nullptr)
 {
 	PendingUpdate = nullptr;
 	SRGB = true;
+
+	// AddressX is default-constructed by Uproperty to enum value 0 which is TA_Wrap
+	check( AddressX == TA_Wrap );
 }
 
 /*-----------------------------------------------------------------------------
