Fix race conditions for registration and deregistration of delegates for FAudioDeviceManagerDelegates::OnAudioDeviceDestroyed.
Fix crash in TMulticastDelegateBase<FDefaultDelegateUserPolicy>::CompactInvocationList() when called from UWorld::BeginDestroy().
Move deregistration of delegates in FNiagaraDataInterfaceProxy derived classes from destructor on RT to OnDestroyed() on GT.
Move registration and deregistration of delegates in UWorld from constructor on ALT and destructor on GT to SetAudioDevice() on GT.
Only register the delegate for UWorlds with an audio device set. This reduces the number of registrations from thousands to a handful when using World Partition.

#rnx

diff --git a/Engine/Plugins/FX/Niagara/Source/Niagara/Private/NiagaraDataInterfaceAudio.cpp b/Engine/Plugins/FX/Niagara/Source/Niagara/Private/NiagaraDataInterfaceAudio.cpp
--- a/Engine/Plugins/FX/Niagara/Source/Niagara/Private/NiagaraDataInterfaceAudio.cpp
+++ b/Engine/Plugins/FX/Niagara/Source/Niagara/Private/NiagaraDataInterfaceAudio.cpp
@@ -124,8 +124,7 @@ FNiagaraDataInterfaceProxySubmix::FNiagaraDataInterfaceProxySubmix(int32 InNumSa
 
 FNiagaraDataInterfaceProxySubmix::~FNiagaraDataInterfaceProxySubmix()
 {
-	FAudioDeviceManagerDelegates::OnAudioDeviceCreated.Remove(DeviceCreatedHandle);
-	FAudioDeviceManagerDelegates::OnAudioDeviceDestroyed.Remove(DeviceDestroyedHandle);
+	check(IsInRenderingThread());
 }
 
 void FNiagaraDataInterfaceProxySubmix::RegisterToAllAudioDevices()
@@ -235,6 +234,10 @@ float FNiagaraDataInterfaceProxySubmix::GetSampleRate() const
 
 void FNiagaraDataInterfaceProxySubmix::OnBeginDestroy()
 {
+	check(IsInGameThread());
+	FAudioDeviceManagerDelegates::OnAudioDeviceCreated.Remove(DeviceCreatedHandle);
+	FAudioDeviceManagerDelegates::OnAudioDeviceDestroyed.Remove(DeviceDestroyedHandle);
+
 	if (bIsSubmixListenerRegistered)
 	{
 		UnregisterFromAllAudioDevices();
diff --git a/Engine/Plugins/FX/Niagara/Source/Niagara/Private/NiagaraDataInterfaceAudioOscilloscope.cpp b/Engine/Plugins/FX/Niagara/Source/Niagara/Private/NiagaraDataInterfaceAudioOscilloscope.cpp
--- a/Engine/Plugins/FX/Niagara/Source/Niagara/Private/NiagaraDataInterfaceAudioOscilloscope.cpp
+++ b/Engine/Plugins/FX/Niagara/Source/Niagara/Private/NiagaraDataInterfaceAudioOscilloscope.cpp
@@ -48,9 +48,6 @@ FNiagaraDataInterfaceProxyOscilloscope::~FNiagaraDataInterfaceProxyOscilloscope(
 {
 	check(IsInRenderingThread());
 	GPUDownsampledBuffer.Release();
-
-	FAudioDeviceManagerDelegates::OnAudioDeviceCreated.Remove(DeviceCreatedHandle);
-	FAudioDeviceManagerDelegates::OnAudioDeviceDestroyed.Remove(DeviceDestroyedHandle);
 }
 
 void FNiagaraDataInterfaceProxyOscilloscope::RegisterToAllAudioDevices()
@@ -410,6 +407,10 @@ void UNiagaraDataInterfaceAudioOscilloscope::PostInitProperties()
 
 void FNiagaraDataInterfaceProxyOscilloscope::OnBeginDestroy()
 {
+	check(IsInGameThread());
+	FAudioDeviceManagerDelegates::OnAudioDeviceCreated.Remove(DeviceCreatedHandle);
+	FAudioDeviceManagerDelegates::OnAudioDeviceDestroyed.Remove(DeviceDestroyedHandle);
+
 	if (bIsSubmixListenerRegistered)
 	{
 		UnregisterFromAllAudioDevices(SubmixRegisteredTo);
diff --git a/Engine/Source/Runtime/Engine/Private/World.cpp b/Engine/Source/Runtime/Engine/Private/World.cpp
--- a/Engine/Source/Runtime/Engine/Private/World.cpp
+++ b/Engine/Source/Runtime/Engine/Private/World.cpp
@@ -429,15 +429,6 @@ UWorld::UWorld( const FObjectInitializer& ObjectInitializer )
 
 	PerfTrackers = new FWorldInGamePerformanceTrackers();
 
-	AudioDeviceDestroyedHandle = FAudioDeviceManagerDelegates::OnAudioDeviceDestroyed.AddLambda([this](const Audio::FDeviceId InDeviceId)
-	{
-		if (InDeviceId == AudioDeviceHandle.GetDeviceID())
-		{
-			FAudioDeviceHandle EmptyHandle;
-			SetAudioDevice(EmptyHandle);
-		}
-	});
-
 	IsInBlockTillLevelStreamingCompleted = 0;
 	BlockTillLevelStreamingCompletedEpoch = 0;
 }
@@ -1051,14 +1042,9 @@ void UWorld::BeginDestroy()
 		Scene->UpdateParameterCollections(TArray<FMaterialParameterCollectionInstanceResource*>());
 	}
 
-	if (AudioDeviceDestroyedHandle.IsValid())
-	{
-		FAudioDeviceManagerDelegates::OnAudioDeviceDestroyed.Remove(AudioDeviceDestroyedHandle);
-		AudioDeviceDestroyedHandle.Reset();
-	}
-
 	FAudioDeviceHandle EmptyHandle;
 	SetAudioDevice(EmptyHandle);
+	check(!AudioDeviceDestroyedHandle.IsValid());
 
 #if UE_WITH_IRIS
 	if (IrisSystemHolder.IsHolding())
@@ -8714,6 +8700,8 @@ void UWorld::AddPostProcessingSettings(FVector ViewLocation, FSceneView* SceneVi
 
 void UWorld::SetAudioDevice(const FAudioDeviceHandle& InHandle)
 {
+	check(IsInGameThread());
+
 	if (InHandle.GetDeviceID() == AudioDeviceHandle.GetDeviceID())
 	{
 		return;
@@ -8742,6 +8730,23 @@ void UWorld::SetAudioDevice(const FAudioDeviceHandle& InHandle)
 	{
 		AudioDeviceHandle.Reset();
 	}
+
+	if (AudioDeviceHandle.IsValid() && !AudioDeviceDestroyedHandle.IsValid())
+	{
+		AudioDeviceDestroyedHandle = FAudioDeviceManagerDelegates::OnAudioDeviceDestroyed.AddLambda([this](const Audio::FDeviceId InDeviceId)
+		{
+			if (InDeviceId == AudioDeviceHandle.GetDeviceID())
+			{
+				FAudioDeviceHandle EmptyHandle;
+				SetAudioDevice(EmptyHandle);
+			}
+		});
+	}
+	else if (!AudioDeviceHandle.IsValid() && AudioDeviceDestroyedHandle.IsValid())
+	{
+		FAudioDeviceManagerDelegates::OnAudioDeviceDestroyed.Remove(AudioDeviceDestroyedHandle);
+		AudioDeviceDestroyedHandle.Reset();
+	}
 }
 
 FAudioDeviceHandle UWorld::GetAudioDevice() const
