Non-unity fix

#rnx

diff --git a/Engine/Source/Runtime/Engine/Private/WorldPartition/WorldPartitionStreamingPolicy.cpp b/Engine/Source/Runtime/Engine/Private/WorldPartition/WorldPartitionStreamingPolicy.cpp
--- a/Engine/Source/Runtime/Engine/Private/WorldPartition/WorldPartitionStreamingPolicy.cpp
+++ b/Engine/Source/Runtime/Engine/Private/WorldPartition/WorldPartitionStreamingPolicy.cpp
@@ -17,6 +17,7 @@
 #include "WorldPartition/ContentBundle/ContentBundle.h"
 #include "WorldPartition/WorldPartitionDebugHelper.h"
 #include "Engine/World.h"
+#include "Engine/Level.h"
 #include "Engine/Engine.h"
 #include "Engine/Canvas.h"
 #include "Engine/NetDriver.h"
