Get DDC Request Name from ShaderPlatformName and not ShaderFormatName

diff --git a/Engine/Source/Runtime/Engine/Private/Materials/MaterialShader.cpp b/Engine/Source/Runtime/Engine/Private/Materials/MaterialShader.cpp
--- a/Engine/Source/Runtime/Engine/Private/Materials/MaterialShader.cpp
+++ b/Engine/Source/Runtime/Engine/Private/Materials/MaterialShader.cpp
@@ -215,7 +215,7 @@ static UE::DerivedData::FSharedString GetMaterialShaderMapName(const FStringView
 	FName FeatureLevelName;
 	GetFeatureLevelName(ShaderMapId.FeatureLevel, FeatureLevelName);
 	return UE::DerivedData::FSharedString(WriteToString<256>(MaterialPath,
-		TEXTVIEW(" ["), LegacyShaderPlatformToShaderFormat(Platform),
+		TEXTVIEW(" ["), FDataDrivenShaderPlatformInfo::GetName(Platform),
 		TEXTVIEW(", "), FeatureLevelName,
 		TEXTVIEW(", "), LexToString(ShaderMapId.QualityLevel),
 		TEXTVIEW("]")));
