PR #9712: Correct spelling of "Collision" (Contributed by [REDACTED_PERSON])

diff --git a/Engine/Source/Runtime/Engine/Classes/Engine/EngineTypes.h b/Engine/Source/Runtime/Engine/Classes/Engine/EngineTypes.h
--- a/Engine/Source/Runtime/Engine/Classes/Engine/EngineTypes.h
+++ b/Engine/Source/Runtime/Engine/Classes/Engine/EngineTypes.h
@@ -866,7 +866,7 @@ enum { NumExtraFilterBits = 6 };
 // Otherwise it will mess up collision profile loading
 // If you change this, please also change FCollisionResponseContainers
 //
-// If you add any more TraceQuery="1", you also should change UCollsionProfile::LoadProfileConfig
+// If you add any more TraceQuery="1", you also should change UCollisionProfile::LoadProfileConfig
 // Metadata doesn't work outside of editor, so you'll need to add manually
 
 // @NOTE : when you add more here for predefined engine channel
