Preven't access to FAnimSubsystem_SharedLinkedAnimLayers while reinstancing

#tests PIE / BP compilation

diff --git a/Engine/Source/Runtime/Engine/Private/Animation/AnimInstance.cpp b/Engine/Source/Runtime/Engine/Private/Animation/AnimInstance.cpp
--- a/Engine/Source/Runtime/Engine/Private/Animation/AnimInstance.cpp
+++ b/Engine/Source/Runtime/Engine/Private/Animation/AnimInstance.cpp
@@ -353,12 +353,15 @@ void UAnimInstance::UninitializeAnimation()
 		}
 	}
 	
-	// Cleanup shared layers data
-	if (FAnimSubsystem_SharedLinkedAnimLayers* SharedLinkedAnimLayers = FindSubsystem<FAnimSubsystem_SharedLinkedAnimLayers>())
+	if (!GIsReinstancing)
 	{
-		// Reset shared linked instances when the main instance in uninitialized. 
-		// This is required in part for SkeletalMeshComponent::OnUnregister so that the linked instances array isn't modified as we iterate on it to unitialize them. (see FAnimNode_LinkedLayer::CleanupSharedLinkedLayersData)
-		SharedLinkedAnimLayers->Reset();
+		// Cleanup shared layers data (we don't use FAnimSubsystem_SharedLinkedAnimLayers::GetFromMesh here as we only want the main instance to clean the shared layers)
+		if (FAnimSubsystem_SharedLinkedAnimLayers* SharedLinkedAnimLayers = FindSubsystem<FAnimSubsystem_SharedLinkedAnimLayers>())
+		{
+			// Reset shared linked instances when the main instance in uninitialized. 
+			// This is required in part for SkeletalMeshComponent::OnUnregister so that the linked instances array isn't modified as we iterate on it to unitialize them. (see FAnimNode_LinkedLayer::CleanupSharedLinkedLayersData)
+			SharedLinkedAnimLayers->Reset();
+		}
 	}
 
 	TRACE_OBJECT_LIFETIME_END(this);
@@ -3011,7 +3014,7 @@ void UAnimInstance::PerformLinkedLayerOverlayOperation(TSubclassOf<UAnimInstance
 		{
 			for (TPair<FName, TArray<FAnimNode_LinkedAnimLayer*, TInlineAllocator<4>>> LayerPair : ClassLayerNodesToSet.Value)
 			{
-				if (FAnimSubsystem_SharedLinkedAnimLayers* SharedLinkedAnimLayers = MeshComp->GetAnimInstance()->FindSubsystem<FAnimSubsystem_SharedLinkedAnimLayers>())
+				if (FAnimSubsystem_SharedLinkedAnimLayers* SharedLinkedAnimLayers = FAnimSubsystem_SharedLinkedAnimLayers::GetFromMesh(MeshComp))
 				{
 					// Shared instances path
 					UClass* ClassToSet = ClassLayerNodesToSet.Key;
diff --git a/Engine/Source/Runtime/Engine/Private/Animation/AnimNode_LinkedAnimLayer.cpp b/Engine/Source/Runtime/Engine/Private/Animation/AnimNode_LinkedAnimLayer.cpp
--- a/Engine/Source/Runtime/Engine/Private/Animation/AnimNode_LinkedAnimLayer.cpp
+++ b/Engine/Source/Runtime/Engine/Private/Animation/AnimNode_LinkedAnimLayer.cpp
@@ -62,8 +62,7 @@ void FAnimNode_LinkedAnimLayer::OnUninitializeAnimInstance(UAnimInstance* InOwni
 	if (UAnimInstance* CurrentTarget = GetTargetInstance<UAnimInstance>())
 	{
 		USkeletalMeshComponent* MeshComp = InOwningAnimInstance->GetSkelMeshComponent();
-		check(MeshComp && MeshComp->GetAnimInstance());
-		if (FAnimSubsystem_SharedLinkedAnimLayers* SharedLinkedAnimLayers = MeshComp->GetAnimInstance()->FindSubsystem<FAnimSubsystem_SharedLinkedAnimLayers>())
+		if (FAnimSubsystem_SharedLinkedAnimLayers* SharedLinkedAnimLayers = FAnimSubsystem_SharedLinkedAnimLayers::GetFromMesh(MeshComp))
 		{
 			// If target instance is a shared instance, unlink it when owner uninitialize
 			if (FLinkedAnimLayerInstanceData* PreviousLinkedLayerInstanceData = SharedLinkedAnimLayers->FindInstanceData(CurrentTarget))
@@ -183,8 +182,7 @@ bool FAnimNode_LinkedAnimLayer::CanTeardownLinkedInstance(const UAnimInstance* L
 {
 	// Don't teardown instance that still have function linked to active shared instances
 	USkeletalMeshComponent* MeshComp = LinkedInstance->GetSkelMeshComponent();
-	check(MeshComp && MeshComp->GetAnimInstance());
-	if (FAnimSubsystem_SharedLinkedAnimLayers* SharedLinkedAnimLayers = MeshComp->GetAnimInstance()->FindSubsystem<FAnimSubsystem_SharedLinkedAnimLayers>())
+	if (FAnimSubsystem_SharedLinkedAnimLayers* SharedLinkedAnimLayers = FAnimSubsystem_SharedLinkedAnimLayers::GetFromMesh(MeshComp))
 	{
 		return SharedLinkedAnimLayers->FindInstanceData(LinkedInstance) == nullptr;
 	}
@@ -196,8 +194,7 @@ void FAnimNode_LinkedAnimLayer::CleanupSharedLinkedLayersData(const UAnimInstanc
 	if (InPreviousTargetInstance)
 	{
 		USkeletalMeshComponent* MeshComp = InPreviousTargetInstance->GetSkelMeshComponent();
-		check(MeshComp && MeshComp->GetAnimInstance());
-		if (FAnimSubsystem_SharedLinkedAnimLayers* SharedLinkedAnimLayers = MeshComp->GetAnimInstance()->FindSubsystem<FAnimSubsystem_SharedLinkedAnimLayers>() )
+		if (FAnimSubsystem_SharedLinkedAnimLayers* SharedLinkedAnimLayers = FAnimSubsystem_SharedLinkedAnimLayers::GetFromMesh(MeshComp))
 		{
 			if (FLinkedAnimLayerInstanceData* PreviousLinkedLayerInstanceData = SharedLinkedAnimLayers->FindInstanceData(InPreviousTargetInstance))
 			{
diff --git a/Engine/Source/Runtime/Engine/Private/Animation/AnimSubsystem_SharedLinkedAnimLayers.cpp b/Engine/Source/Runtime/Engine/Private/Animation/AnimSubsystem_SharedLinkedAnimLayers.cpp
--- a/Engine/Source/Runtime/Engine/Private/Animation/AnimSubsystem_SharedLinkedAnimLayers.cpp
+++ b/Engine/Source/Runtime/Engine/Private/Animation/AnimSubsystem_SharedLinkedAnimLayers.cpp
@@ -2,6 +2,7 @@
 
 #include "Animation/AnimSubsystem_SharedLinkedAnimLayers.h"
 #include "Animation/AnimInstance.h"
+#include "Components/SkeletalMeshComponent.h"
 
 #if LINKEDANIMLAYERSDATA_INTEGRITYCHECKS
 #include "Animation/AnimClassInterface.h"
@@ -107,6 +108,18 @@ void FLinkedAnimLayerClassData::RemoveInstance(const UAnimInstance* AnimInstance
 	check(0);
 }
 
+FAnimSubsystem_SharedLinkedAnimLayers* FAnimSubsystem_SharedLinkedAnimLayers::GetFromMesh(USkeletalMeshComponent* SkelMesh)
+{
+#if WITH_EDITOR
+	if (GIsReinstancing)
+	{
+		return nullptr;
+	}
+#endif
+	check(SkelMesh && SkelMesh->GetAnimInstance());
+	return SkelMesh->GetAnimInstance()->FindSubsystem<FAnimSubsystem_SharedLinkedAnimLayers>();
+}
+
 void FAnimSubsystem_SharedLinkedAnimLayers::Reset()
 {
 	ClassesData.Empty(ClassesData.Num());
@@ -154,4 +167,4 @@ void FAnimSubsystem_SharedLinkedAnimLayers::RemoveInstance(const UAnimInstance*
 	{
 		check(0);
 	}
-}
+}
\ No newline at end of file
diff --git a/Engine/Source/Runtime/Engine/Private/Components/SkeletalMeshComponent.cpp b/Engine/Source/Runtime/Engine/Private/Components/SkeletalMeshComponent.cpp
--- a/Engine/Source/Runtime/Engine/Private/Components/SkeletalMeshComponent.cpp
+++ b/Engine/Source/Runtime/Engine/Private/Components/SkeletalMeshComponent.cpp
@@ -3113,7 +3113,7 @@ void USkeletalMeshComponent::ResetLinkedAnimInstances()
 	// Reset linked anim layers shared data 
 	if (AnimScriptInstance)
 	{
-		if (FAnimSubsystem_SharedLinkedAnimLayers* SharedLinkedAnimLayers = AnimScriptInstance->FindSubsystem<FAnimSubsystem_SharedLinkedAnimLayers>())
+		if (FAnimSubsystem_SharedLinkedAnimLayers* SharedLinkedAnimLayers = FAnimSubsystem_SharedLinkedAnimLayers::GetFromMesh(this))
 		{
 			SharedLinkedAnimLayers->Reset();
 		}
diff --git a/Engine/Source/Runtime/Engine/Public/Animation/AnimSubsystem_SharedLinkedAnimLayers.h b/Engine/Source/Runtime/Engine/Public/Animation/AnimSubsystem_SharedLinkedAnimLayers.h
--- a/Engine/Source/Runtime/Engine/Public/Animation/AnimSubsystem_SharedLinkedAnimLayers.h
+++ b/Engine/Source/Runtime/Engine/Public/Animation/AnimSubsystem_SharedLinkedAnimLayers.h
@@ -7,6 +7,7 @@
 #include "AnimSubsystem_SharedLinkedAnimLayers.generated.h"
 
 class UAnimInstance;
+class USkeletalMeshComponent;
 
 // Define to 1 for more extensive runtime validation of linked anim layers data
 #define LINKEDANIMLAYERSDATA_INTEGRITYCHECKS 0
@@ -72,6 +73,9 @@ struct ENGINE_API FAnimSubsystem_SharedLinkedAnimLayers : public FAnimSubsystemI
 {
 	GENERATED_BODY()
 
+	// Retrieve subsystem from skeletal mesh component
+	static FAnimSubsystem_SharedLinkedAnimLayers* GetFromMesh(USkeletalMeshComponent* SkelMesh);
+	
 	// Clear all linked layers data
 	void Reset();
 
