Fix for alpha channel disappearing/resetting when graphics settings are toggled below Epic

diff --git a/Engine/Source/Runtime/Engine/Private/SceneTexturesConfig.cpp b/Engine/Source/Runtime/Engine/Private/SceneTexturesConfig.cpp
--- a/Engine/Source/Runtime/Engine/Private/SceneTexturesConfig.cpp
+++ b/Engine/Source/Runtime/Engine/Private/SceneTexturesConfig.cpp
@@ -222,6 +222,7 @@ void FSceneTexturesConfig::Init(const FSceneTexturesConfigInitSettings& InitSett
 	ColorFormat				= PF_Unknown;
 	ColorClearValue			= FClearValueBinding::Black;
 	DepthClearValue			= FClearValueBinding::DepthFar;
+	bRequiresAlphaChannel	= InitSettings.bRequiresAlphaChannel;
 	bRequireMultiView		= InitSettings.bRequireMultiView;
 	bIsUsingGBuffers		= IsUsingGBuffers(ShaderPlatform);
 	bSupportsXRTargetManagerDepthAlloc = InitSettings.bSupportsXRTargetManagerDepthAlloc;
