Bump number of concurrent anim compilation from 2 to 64 so that they progress more quickly after loading large amount of anims

#rnx

diff --git a/Engine/Source/Runtime/Engine/Private/Animation/AnimCompressionDerivedData.cpp b/Engine/Source/Runtime/Engine/Private/Animation/AnimCompressionDerivedData.cpp
--- a/Engine/Source/Runtime/Engine/Private/Animation/AnimCompressionDerivedData.cpp
+++ b/Engine/Source/Runtime/Engine/Private/Animation/AnimCompressionDerivedData.cpp
@@ -15,6 +15,7 @@
 #include "UObject/Package.h"
 #include "Misc/CoreDelegates.h"
 #include "ProfilingDebugging/CookStats.h"
+#include "ProfilingDebugging/CountersTrace.h"
 
 #if WITH_EDITOR
 
@@ -122,7 +123,7 @@ bool FDerivedDataAnimationCompression::Build( TArray<uint8>& OutDataArray )
 
 const uint64 GigaBytes = 1024 * 1024 * 1024;
 const uint64 MAX_ASYNC_COMPRESSION_MEM_USAGE = 3 * GigaBytes;
-const int32 MAX_ACTIVE_COMPRESSIONS = 2;
+const int32 MAX_ACTIVE_COMPRESSIONS = 64;
 
 FAsyncCompressedAnimationsManagement* GAsyncCompressedAnimationsTracker = nullptr;
 
@@ -180,6 +181,7 @@ void FAsyncCompressedAnimationsManagement::Shutdown()
 
 void FAsyncCompressedAnimationsManagement::OnActiveCompressionFinished(int32 ActiveAnimIndex)
 {
+	TRACE_CPUPROFILER_EVENT_SCOPE(FAsyncCompressedAnimationsManagement::OnActiveCompressionFinished);
 	COOK_STAT(auto Timer = AnimSequenceCookStats::UsageStats.TimeSyncWork());
 
 	FDerivedDataCacheInterface& DerivedDataCache = GetDerivedDataCacheRef();
@@ -240,12 +242,26 @@ float FormatBytes(uint64 Bytes, const TCHAR*& OutPostFix)
 
 #define ASYNC_MEM_LOG 0
 
+TRACE_DECLARE_INT_COUNTER(QueuedAnimationCompilation, TEXT("AsyncCompilation/QueuedAnimations"));
+TRACE_DECLARE_INT_COUNTER(ActiveAnimationCompilation, TEXT("AsyncCompilation/ActiveAnimations"));
+TRACE_DECLARE_INT_COUNTER(ProcessedAnimationCompilation, TEXT("AsyncCompilation/ProcessedAnimations"));
 void FAsyncCompressedAnimationsManagement::Tick(float DeltaTime)
 {
+	TRACE_CPUPROFILER_EVENT_SCOPE(FAsyncCompressedAnimationsManagement::Tick);
+
 	const double MaxProcessingTime = 0.1; // try not to hang the editor too much
 	const double EndTime = FPlatformTime::Seconds() + MaxProcessingTime;
 	const double StartTime = FPlatformTime::Seconds();
 
+	int32 ProcessedAnimation = 0;
+
+	ON_SCOPE_EXIT
+	{
+		TRACE_COUNTER_SET(ActiveAnimationCompilation, ActiveAsyncCompressionTasks.Num());
+		TRACE_COUNTER_SET(QueuedAnimationCompilation, QueuedAsyncCompressionWork.Num());
+		TRACE_COUNTER_SET(ProcessedAnimationCompilation, ProcessedAnimation);
+	};
+
 #if ASYNC_MEM_LOG
 	const TCHAR* BytesPostFix = TEXT("ERROR");
 	const float FormatSize = FormatBytes(ActiveMemoryUsage, BytesPostFix);
@@ -263,6 +279,7 @@ void FAsyncCompressedAnimationsManagement::Tick(float DeltaTime)
 		if (DerivedDataCache.PollAsynchronousCompletion(Task.AsyncHandle))
 		{
 			OnActiveCompressionFinished(ActiveAnim);
+			ProcessedAnimation++;
 		}
 
 		if (FPlatformTime::Seconds() > EndTime)
@@ -367,6 +384,8 @@ bool FAsyncCompressedAnimationsManagement::RequestAsyncCompression(FDerivedDataA
 
 void FAsyncCompressedAnimationsManagement::StartAsyncWork(FDerivedDataAnimationCompression& Compressor, UAnimSequence* Anim, const uint64 NewTaskSize, const bool bPerformFrameStripping)
 {
+	TRACE_CPUPROFILER_EVENT_SCOPE(FAsyncCompressedAnimationsManagement::StartAsyncWork);
+
 	const FString CacheKey = Compressor.GetPluginSpecificCacheKeySuffix();
 	FCompressibleAnimPtr SourceData = Compressor.GetCompressibleData();
 	uint32 AsyncHandle = GetDerivedDataCacheRef().GetAsynchronous(&Compressor);
@@ -375,22 +394,27 @@ void FAsyncCompressedAnimationsManagement::StartAsyncWork(FDerivedDataAnimationC
 
 bool FAsyncCompressedAnimationsManagement::WaitOnActiveCompression(UAnimSequence* Anim, bool bWantResults)
 {
-	for (int32 ActiveIndex = 0; ActiveIndex < ActiveAsyncCompressionTasks.Num(); ++ActiveIndex)
+	if (ActiveAsyncCompressionTasks.Num())
 	{
-		FActiveAsyncCompressionTask& Task = ActiveAsyncCompressionTasks[ActiveIndex];
-		if (Task.Sequence == Anim)
+		TRACE_CPUPROFILER_EVENT_SCOPE(FAsyncCompressedAnimationsManagement::WaitOnActiveCompression);
+
+		for (int32 ActiveIndex = 0; ActiveIndex < ActiveAsyncCompressionTasks.Num(); ++ActiveIndex)
 		{
-			if(bWantResults)
+			FActiveAsyncCompressionTask& Task = ActiveAsyncCompressionTasks[ActiveIndex];
+			if (Task.Sequence == Anim)
 			{
-				GetDerivedDataCacheRef().WaitAsynchronousCompletion(Task.AsyncHandle);
-				OnActiveCompressionFinished(ActiveIndex);
-			}
-			else
-			{
-				Task.DataToCompress->IsCancelledSignal.Cancel();
-				Task.Sequence = nullptr;
+				if(bWantResults)
+				{
+					GetDerivedDataCacheRef().WaitAsynchronousCompletion(Task.AsyncHandle);
+					OnActiveCompressionFinished(ActiveIndex);
+				}
+				else
+				{
+					Task.DataToCompress->IsCancelledSignal.Cancel();
+					Task.Sequence = nullptr;
+				}
+				return true; // Was active
 			}
-			return true; // Was active
 		}
 	}
 	return false;
