Stop a dedicated server instance from trying to create local players as it will crash later.
It probably isn't valid to be missing a GameViewportClient for any cases other than dedicated servers, but I left that enabled with an ensure

diff --git a/Engine/Source/Runtime/Engine/Private/GameInstance.cpp b/Engine/Source/Runtime/Engine/Private/GameInstance.cpp
--- a/Engine/Source/Runtime/Engine/Private/GameInstance.cpp
+++ b/Engine/Source/Runtime/Engine/Private/GameInstance.cpp
@@ -806,8 +806,18 @@ ULocalPlayer* UGameInstance::CreateLocalPlayer(FPlatformUserId UserId, FString&
 
 	ULocalPlayer* NewPlayer = NULL;
 	int32 InsertIndex = INDEX_NONE;
+	UGameViewportClient* GameViewport = GetGameViewportClient();
 
-	const int32 MaxSplitscreenPlayers = (GetGameViewportClient() != NULL) ? GetGameViewportClient()->MaxSplitscreenPlayers : 1;
+	if (GameViewport == nullptr)
+	{
+		if (ensure(IsDedicatedServerInstance()))
+		{
+			OutError = FString::Printf(TEXT("Dedicated servers cannot have local players"));
+			return nullptr;
+		}
+	}
+
+	const int32 MaxSplitscreenPlayers = GameViewport ? GameViewport->MaxSplitscreenPlayers : 1;
 
 	if (FindLocalPlayerFromPlatformUserId(UserId) != NULL)
 	{
