Replaced checks by test as packages can be cleaned up outside of FWorldPartitionLevelHelper

diff --git a/Engine/Source/Runtime/Engine/Private/WorldPartition/WorldPartitionLevelHelper.cpp b/Engine/Source/Runtime/Engine/Private/WorldPartition/WorldPartitionLevelHelper.cpp
--- a/Engine/Source/Runtime/Engine/Private/WorldPartition/WorldPartitionLevelHelper.cpp
+++ b/Engine/Source/Runtime/Engine/Private/WorldPartition/WorldPartitionLevelHelper.cpp
@@ -38,8 +38,11 @@ void FWorldPartitionLevelHelper::PreGarbageCollect()
 {
 	for (TWeakObjectPtr<UPackage>& PackageToUnload : PreGCPackagesToUnload)
 	{
-		check(PackageToUnload.IsValid());
-		FWorldPartitionPackageHelper::UnloadPackage(PackageToUnload.Get());
+		// Test if WeakObjectPtr is valid since clean up could have happened outside of this helper
+		if (PackageToUnload.IsValid())
+		{
+			FWorldPartitionPackageHelper::UnloadPackage(PackageToUnload.Get());
+		}
 	}
 	PreGCPackagesToUnload.Reset();
 }
@@ -62,8 +65,11 @@ void FWorldPartitionLevelHelper::RemoveReferences(FPackageReferencer* InReferenc
 		RefInfo.Referencers.Remove(InReferencer);
 		if (RefInfo.Referencers.Num() == 0)
 		{
-			check(RefInfo.Package.IsValid());
-			PreGCPackagesToUnload.Add(RefInfo.Package);
+			// Test if WeakObjectPtr is valid since clean up could have happened outside of this helper
+			if (RefInfo.Package.IsValid())
+			{
+				PreGCPackagesToUnload.Add(RefInfo.Package);
+			}
 			It.RemoveCurrent();
 		}
 	}
