[CIS] NonUnity fixes / ICWYU for AudioDeviceManager.cpp

diff --git a/Engine/Source/Runtime/Engine/Private/AudioDeviceManager.cpp b/Engine/Source/Runtime/Engine/Private/AudioDeviceManager.cpp
--- a/Engine/Source/Runtime/Engine/Private/AudioDeviceManager.cpp
+++ b/Engine/Source/Runtime/Engine/Private/AudioDeviceManager.cpp
@@ -11,6 +11,7 @@
 #include "Sound/AudioSettings.h"
 #include "Sound/SoundWave.h"
 #include "GameFramework/GameUserSettings.h"
+#include "Misc/App.h"
 #include "Misc/CommandLine.h"
 #include "Misc/ConfigCacheIni.h"
 #include "Misc/CoreDelegates.h"
