Add support for SeamlessTravel and NoSeamlessTravel URL options on game mode to override default
If seamless travel is disabled in PIE, fall back to hard travel instead of failing completely

diff --git a/Engine/Source/Runtime/Engine/Private/GameModeBase.cpp b/Engine/Source/Runtime/Engine/Private/GameModeBase.cpp
--- a/Engine/Source/Runtime/Engine/Private/GameModeBase.cpp
+++ b/Engine/Source/Runtime/Engine/Private/GameModeBase.cpp
@@ -415,17 +415,6 @@ bool AGameModeBase::CanServerTravel(const FString& FURL, bool bAbsolute)
 
 	check(World);
 
-	// NOTE - This is a temp check while we work on a long term fix
-	// There are a few issues with seamless travel using single process PIE, so we're disabling that for now while working on a fix
-	if (World->WorldType == EWorldType::PIE && bUseSeamlessTravel && !FParse::Param(FCommandLine::Get(), TEXT("MultiprocessOSS")))
-	{
-		if (!UE::GameModeBase::Private::bAllowPIESeamlessTravel)
-		{
-			UE_LOG(LogGameMode, Warning, TEXT("CanServerTravel: Seamless travel is not supported by default in PIE, set net.AllowPIESeamlessTravel=1 to enable experimental support."));
-			return false;
-		}
-	}
-
 	if (FURL.Contains(TEXT("%")))
 	{
 		UE_LOG(LogGameMode, Error, TEXT("CanServerTravel: FURL %s Contains illegal character '%%'."), *FURL);
@@ -470,12 +459,32 @@ void AGameModeBase::ProcessServerTravel(const FString& URL, bool bAbsolute)
 	check(World);
 	FWorldContext& WorldContext = GEngine->GetWorldContextFromWorldChecked(World);
 
-	// Force an old style load screen if the server has been up for a long time so that TimeSeconds doesn't overflow and break everything
+	// Use game mode setting but default to full load screen if the server has been up for a long time so that TimeSeconds doesn't overflow and break everything
 	bool bSeamless = (bUseSeamlessTravel && GetWorld()->TimeSeconds < 172800.0f); // 172800 seconds == 48 hours
 
 	// Compute the next URL, and pull the map out of it. This handles short->long package name conversion
 	FURL NextURL = FURL(&WorldContext.LastURL, *URL, bAbsolute ? TRAVEL_Absolute : TRAVEL_Relative);
 
+	// Override based on URL parameters
+	if (NextURL.HasOption(TEXT("SeamlessTravel")))
+	{
+		bSeamless = true;
+	}
+	else if (NextURL.HasOption(TEXT("NoSeamlessTravel")))
+	{
+		bSeamless = false;
+	}
+
+	// There are some issues with seamless travel in PIE, so fall back to hard travel unless it is supported
+	if (World->WorldType == EWorldType::PIE && bSeamless && !FParse::Param(FCommandLine::Get(), TEXT("MultiprocessOSS")))
+	{
+		if (!UE::GameModeBase::Private::bAllowPIESeamlessTravel)
+		{
+			UE_LOG(LogGameMode, Warning, TEXT("ProcessServerTravel: Seamless travel is disabled in PIE, set net.AllowPIESeamlessTravel=1 to enable."));
+			bSeamless = false;
+		}
+	}
+
 	PRAGMA_DISABLE_DEPRECATION_WARNINGS
 	FGuid NextMapGuid = UEngine::GetPackageGuid(FName(*NextURL.Map), GetWorld()->IsPlayInEditor());
 	PRAGMA_ENABLE_DEPRECATION_WARNINGS
