World Partition: fix crash when streaming a partitioned world from a non-partitioned one.

#rnx

diff --git a/Engine/Source/Runtime/Engine/Private/WorldPartition/WorldPartitionRuntimeLevelStreamingCell.cpp b/Engine/Source/Runtime/Engine/Private/WorldPartition/WorldPartitionRuntimeLevelStreamingCell.cpp
--- a/Engine/Source/Runtime/Engine/Private/WorldPartition/WorldPartitionRuntimeLevelStreamingCell.cpp
+++ b/Engine/Source/Runtime/Engine/Private/WorldPartition/WorldPartitionRuntimeLevelStreamingCell.cpp
@@ -376,10 +376,7 @@ bool UWorldPartitionRuntimeLevelStreamingCell::CanUnload() const
 	// Events to the HLodSubsystem which assumes knowledge of all cells (not true with plugins)
 	if (LevelStreaming && !GetContentBundleID().IsValid())
 	{
-		const UWorldPartition* WorldPartition = LevelStreaming->GetWorld()->GetWorldPartition();
-		check(WorldPartition);
-
-		if (WorldPartition->IsStreamingEnabled())
+		if (const UWorldPartition* WorldPartition = LevelStreaming->GetWorld()->GetWorldPartition(); WorldPartition && WorldPartition->IsStreamingEnabled())
 		{
 			if (UHLODSubsystem* HLODSubsystem = LevelStreaming->GetWorld()->GetSubsystem<UHLODSubsystem>())
 			{
@@ -432,11 +429,7 @@ void UWorldPartitionRuntimeLevelStreamingCell::Deactivate() const
 void UWorldPartitionRuntimeLevelStreamingCell::OnLevelShown()
 {
 	check(LevelStreaming);
-
-	const UWorldPartition* WorldPartition = LevelStreaming->GetWorld()->GetWorldPartition();
-	check(WorldPartition);
-
-	if (WorldPartition->IsStreamingEnabled())
+	if (const UWorldPartition* WorldPartition = LevelStreaming->GetWorld()->GetWorldPartition(); WorldPartition && WorldPartition->IsStreamingEnabled())
 	{
 		LevelStreaming->GetWorld()->GetSubsystem<UHLODSubsystem>()->OnCellShown(this);
 	}
@@ -445,11 +438,7 @@ void UWorldPartitionRuntimeLevelStreamingCell::OnLevelShown()
 void UWorldPartitionRuntimeLevelStreamingCell::OnLevelHidden()
 {
 	check(LevelStreaming);
-
-	const UWorldPartition* WorldPartition = LevelStreaming->GetWorld()->GetWorldPartition();
-	check(WorldPartition);
-
-	if (WorldPartition->IsStreamingEnabled())
+	if (const UWorldPartition* WorldPartition = LevelStreaming->GetWorld()->GetWorldPartition(); WorldPartition && WorldPartition->IsStreamingEnabled())
 	{
 		LevelStreaming->GetWorld()->GetSubsystem<UHLODSubsystem>()->OnCellHidden(this);
 	}
