Early out in FSkeletalMeshSceneProxy::DebugDrawSkeleton when MeshObject doesn't have valid ComponentSpaceTransforms. This happens with sk meshes rendered with FSkeletalMeshObjectStatic.

diff --git a/Engine/Source/Runtime/Engine/Private/SkeletalMesh.cpp b/Engine/Source/Runtime/Engine/Private/SkeletalMesh.cpp
--- a/Engine/Source/Runtime/Engine/Private/SkeletalMesh.cpp
+++ b/Engine/Source/Runtime/Engine/Private/SkeletalMesh.cpp
@@ -6909,6 +6909,12 @@ void FSkeletalMeshSceneProxy::DebugDrawSkeleton(int32 ViewIndex, FMeshElementCol
 		return; // Cannot draw this, world matrix not valid
 	}
 
+	// Can't draw this, don't have ComponentSpaceTransforms. This happens with sk meshes rendered with FSkeletalMeshObjectStatic.
+	if (!MeshObject->GetComponentSpaceTransforms())
+	{
+		return;
+	}
+
 	FTransform LocalToWorldTransform(ProxyLocalToWorld);
 
 	auto MakeRandomColorForSkeleton = [](uint32 InUID)
