Characters desync when standing on a moving vehicle.
Change how simulated actors on dynamic bases deal with encroachment
#rnx

diff --git a/Engine/Source/Runtime/Engine/Private/Components/CharacterMovementComponent.cpp b/Engine/Source/Runtime/Engine/Private/Components/CharacterMovementComponent.cpp
--- a/Engine/Source/Runtime/Engine/Private/Components/CharacterMovementComponent.cpp
+++ b/Engine/Source/Runtime/Engine/Private/Components/CharacterMovementComponent.cpp
@@ -2006,7 +2006,18 @@ void UCharacterMovementComponent::SimulateMovement(float DeltaSeconds)
 					CurrentFloor.Clear();
 				}
 
-				if (!CurrentFloor.IsWalkableFloor())
+				// Possible for dynamic movement bases, particularly those that align to slopes while the character does not, to encroach the character.
+				// Check to see if we can resolve the penetration in those cases, and if so find the floor.
+				if (CurrentFloor.HitResult.bStartPenetrating && MovementBaseUtility::IsDynamicBase(GetMovementBase()))
+				{
+					// Follows PhysWalking approach for encroachment on floor tests
+					FHitResult Hit(CurrentFloor.HitResult);
+					Hit.TraceEnd = Hit.TraceStart + FVector(0.f, 0.f, MAX_FLOOR_DIST);
+					const FVector RequestedAdjustment = GetPenetrationAdjustment(Hit);
+					const bool bResolved = ResolvePenetration(RequestedAdjustment, Hit, UpdatedComponent->GetComponentQuat());
+					bForceNextFloorCheck |= bResolved;
+				}
+				else if (!CurrentFloor.IsWalkableFloor())
 				{
 					if (!bSimGravityDisabled)
 					{
