fixed a deadlock (livelock?) caused by double busy waiting, one in audio pipe suspension and another in the scheduler: an RHI task is waiting for a task event, the scheduler "busy waits" for it by picking audio suspension task, which also busy waits until audio is resumed. this doesn't happen because of an assumed cyclic dependency: GT depends on RT, RT waits for the RHI task, the RHI task could complete but is stuck inside audio suspension task busy waiting for GT to resume it.
the fix is to replace busy waiting in audio suspension task by using a nested task. the suspension task execution is finished immediately but because of the nested task it's not completed and so keeping audio pipe blocked (suspended) until audio is resumed.

diff --git a/Engine/Source/Runtime/Engine/Private/AudioThread.cpp b/Engine/Source/Runtime/Engine/Private/AudioThread.cpp
--- a/Engine/Source/Runtime/Engine/Private/AudioThread.cpp
+++ b/Engine/Source/Runtime/Engine/Private/AudioThread.cpp
@@ -150,11 +150,8 @@ void FAudioThread::SuspendAudioThread()
 		[&SuspendEvent]
 		{
 			GIsAudioThreadSuspended.store(true, std::memory_order_release);
+			UE::Tasks::AddNested(*ResumeEvent);
 			SuspendEvent->Trigger();
-
-			ResumeEvent->BusyWait(); // don't block one of workers while audio pipe is suspended
-			ResumeEvent = MakeUnique<UE::Tasks::FTaskEvent>(UE_SOURCE_LOCATION); // thread-safe because happens "inside" a suspend call and 
-			// can't happen concurrently with a resume call
 		}
 	);
 
@@ -184,6 +181,7 @@ void FAudioThread::ResumeAudioThread()
 
 	check(!ResumeEvent->IsCompleted());
 	ResumeEvent->Trigger();
+	ResumeEvent = MakeUnique<UE::Tasks::FTaskEvent>(UE_SOURCE_LOCATION);
 }
 
 #else // UE_AUDIO_THREAD_AS_PIPE
