Fixed bad worldpartition cell bounds when containing level instance actors.

diff --git a/Engine/Source/Runtime/Engine/Private/WorldPartition/WorldPartitionRuntimeSpatialHash.cpp b/Engine/Source/Runtime/Engine/Private/WorldPartition/WorldPartitionRuntimeSpatialHash.cpp
--- a/Engine/Source/Runtime/Engine/Private/WorldPartition/WorldPartitionRuntimeSpatialHash.cpp
+++ b/Engine/Source/Runtime/Engine/Private/WorldPartition/WorldPartitionRuntimeSpatialHash.cpp
@@ -1370,7 +1370,7 @@ bool UWorldPartitionRuntimeSpatialHash::CreateStreamingGrid(const FSpatialHashRu
 					const FWorldPartitionActorDescView& ActorDescView = ActorInstance.GetActorDescView();
 					StreamingCell->AddActorToCell(ActorDescView, ActorInstance.GetContainerID(), ActorInstance.GetTransform(), ActorInstance.GetActorDescContainer());
 					
-					CellContentBounds += ActorDescView.GetBounds();
+					CellContentBounds += ActorDescView.GetBounds().TransformBy(ActorInstance.GetTransform());
 
 					if (ActorInstance.GetContainerID().IsMainContainer() && StreamingCell->UnsavedActorsContainer)
 					{
