Set bIsPreviewMaterial flag on MatFunction PreviewMaterials.
This prevents UpdateAllPrimitiveSceneInfos() being called unnecessarily when rendering thumbnails.
The calls to UpdateAllPrimitiveSceneInfos() cause RVTs to be flushed which looks glitchy.

diff --git a/Engine/Source/Runtime/Engine/Private/Materials/MaterialExpressions.cpp b/Engine/Source/Runtime/Engine/Private/Materials/MaterialExpressions.cpp
--- a/Engine/Source/Runtime/Engine/Private/Materials/MaterialExpressions.cpp
+++ b/Engine/Source/Runtime/Engine/Private/Materials/MaterialExpressions.cpp
@@ -13679,6 +13679,7 @@ UMaterialInterface* UMaterialFunction::GetPreviewMaterial()
 	if( nullptr == PreviewMaterial )
 	{
 		PreviewMaterial = NewObject<UMaterial>(this, NAME_None, RF_Transient | RF_Public);
+		PreviewMaterial->bIsPreviewMaterial = true;
 
 		PreviewMaterial->AssignExpressionCollection(GetExpressionCollection());
 
