Added stack crashdump helpers for async demo crashes.

diff --git a/Engine/Source/Runtime/CoreUObject/Public/UObject/CoreNetTypes.h b/Engine/Source/Runtime/CoreUObject/Public/UObject/CoreNetTypes.h
--- a/Engine/Source/Runtime/CoreUObject/Public/UObject/CoreNetTypes.h
+++ b/Engine/Source/Runtime/CoreUObject/Public/UObject/CoreNetTypes.h
@@ -5,6 +5,12 @@
 #include "CoreMinimal.h"
 #include "UObject/ObjectMacros.h"
 
+/** Whether or not caching of actor/subobject names to the stack should be enabled, for async demo recording crashdumps */
+#ifndef UE_NET_REPACTOR_NAME_DEBUG
+	#define UE_NET_REPACTOR_NAME_DEBUG 0
+#endif
+
+
 /** Secondary condition to check before considering the replication of a lifetime property. */
 UENUM(BlueprintType)
 enum ELifetimeCondition
diff --git a/Engine/Source/Runtime/Engine/Private/DataChannel.cpp b/Engine/Source/Runtime/Engine/Private/DataChannel.cpp
--- a/Engine/Source/Runtime/Engine/Private/DataChannel.cpp
+++ b/Engine/Source/Runtime/Engine/Private/DataChannel.cpp
@@ -33,6 +33,9 @@
 #include "Net/Core/PushModel/PushModel.h"
 #include "HAL/LowLevelMemStats.h"
 #include "Net/NetPing.h"
+#include "PacketHandler.h"
+#include "Net/NetNameDebug.h"
+
 
 DEFINE_LOG_CATEGORY(LogNet);
 DEFINE_LOG_CATEGORY(LogNetSubObject);
@@ -150,6 +153,11 @@ namespace UE::Net
 		TEXT("net.PushModelValidateSkipUpdate"),
 		bPushModelValidateSkipUpdate,
 		TEXT("If enabled, detect when we thought we could skip an object replication based on push model state, but we sent data anyway."), ECVF_Default);
+
+#if UE_NET_REPACTOR_NAME_DEBUG
+	/** Whether or not the currently executing ReplicateActor code has met the randomized chance for enabling name debugging */
+	static bool GMetAsyncDemoNameDebugChance = false;
+#endif
 }; // namespace UE::Net
 
 template<typename T>
@@ -3154,6 +3162,9 @@ int32 GNumReplicateActorCalls = 0;
 
 int64 UActorChannel::ReplicateActor()
 {
+	using namespace UE::Net;
+	using namespace UE::Net::Private;
+
 	LLM_SCOPE_BYTAG(NetChannel);
 	SCOPE_CYCLE_COUNTER(STAT_NetReplicateActorTime);
 
@@ -3223,7 +3234,7 @@ int64 UActorChannel::ReplicateActor()
 		UE_LOG(LogNet, Verbose, TEXT("ReplicateActor: bPausedUntilReliableACK is ending now that reliables have been ACK'd. %s"), *Describe());
 	}
 
-	if (UE::Net::bNetReplicateOnlyBeginPlay && !IsActorReadyForReplication() && !bIsForcedSerializeFromRPC)
+	if (bNetReplicateOnlyBeginPlay && !IsActorReadyForReplication() && !bIsForcedSerializeFromRPC)
 	{
 		UE_LOG(LogNet, Verbose, TEXT("ReplicateActor ignored since actor is not BeginPlay yet: %s"), *Describe());
 		return 0;
@@ -3302,6 +3313,23 @@ int64 UActorChannel::ReplicateActor()
 	}
 #endif
 
+#if UE_NET_REPACTOR_NAME_DEBUG
+	// View address in memory viewer
+	volatile PTRINT ActorName;
+
+	if (bReplay && (GCVarNetActorNameDebug || GCVarNetSubObjectNameDebug))
+	{
+		GMetAsyncDemoNameDebugChance = HasMetAsyncDemoNameDebugChance();
+
+		if (GCVarNetActorNameDebug && GMetAsyncDemoNameDebugChance)
+		{
+			ActorName = UE_NET_ALLOCA_NAME_DEBUG_BUFFER();
+
+			StoreFullName(ActorName, Actor);
+		}
+	}
+#endif
+
 	FGuardValue_Bitfield(bIsReplicatingActor, true);
 	FScopedRepContext RepContext(Connection, Actor);
 
@@ -3588,6 +3616,7 @@ bool UActorChannel::CanSubObjectReplicateToClient(ELifetimeCondition NetConditio
 bool UActorChannel::ReplicateRegisteredSubObjects(FOutBunch& Bunch, FReplicationFlags RepFlags)
 {
 	using namespace UE::Net;
+	using namespace UE::Net::Private;
 
 	check(Actor->IsUsingRegisteredSubObjectList());
 
@@ -3601,6 +3630,12 @@ bool UActorChannel::ReplicateRegisteredSubObjects(FOutBunch& Bunch, FReplication
 	}
 #endif
 
+#if UE_NET_REPACTOR_NAME_DEBUG
+	// View address in memory viewer
+	const bool bSubObjectNameDebug = GCVarNetSubObjectNameDebug && GMetAsyncDemoNameDebugChance && Connection->IsReplay();
+	volatile PTRINT SubObjName = bSubObjectNameDebug ? UE_NET_ALLOCA_NAME_DEBUG_BUFFER() : 0;
+#endif
+
 	const TStaticBitArray<COND_Max> ConditionMap = FSendingRepState::BuildConditionMapFromRepFlags(RepFlags);
 
 	bool bWroteSomethingImportant = false;
@@ -3614,6 +3649,13 @@ bool UActorChannel::ReplicateRegisteredSubObjects(FOutBunch& Bunch, FReplication
 			{
 				checkf(IsValid(SubObjectInfo.SubObject), TEXT("Found invalid subobject (%s) registered in %s"), *GetNameSafe(SubObjectInfo.SubObject), *Actor->GetName());
 
+#if UE_NET_REPACTOR_NAME_DEBUG
+				if (bSubObjectNameDebug)
+				{
+					StoreSubObjectName(SubObjName, SubObjectInfo);
+				}
+#endif
+
 				bWroteSomethingImportant |= WriteSubObjectInBunch(SubObjectInfo.SubObject, Bunch, RepFlags);
 			}
 		}
@@ -3628,6 +3670,13 @@ bool UActorChannel::ReplicateRegisteredSubObjects(FOutBunch& Bunch, FReplication
 
 			checkf(IsValid(ReplicatedComponent), TEXT("Found invalid replicated component (%s) registered in %s"), *GetNameSafe(ReplicatedComponent), *Actor->GetName());
 
+#if UE_NET_REPACTOR_NAME_DEBUG
+			if (bSubObjectNameDebug)
+			{
+				StoreSubObjectName(SubObjName, RepComponentInfo);
+			}
+#endif
+
 			SetCurrentSubObjectOwner(ReplicatedComponent);
 			if (ReplicatedComponent->IsUsingRegisteredSubObjectList())
 			{
@@ -3789,13 +3838,30 @@ bool UActorChannel::WriteSubObjects(UActorComponent* ReplicatedComponent, FOutBu
 
 bool UActorChannel::WriteSubObjects(UActorComponent* ReplicatedComponent, const UE::Net::FSubObjectRegistry& SubObjectList, FOutBunch& Bunch, FReplicationFlags RepFlags, const TStaticBitArray<COND_Max>& ConditionMap)
 {
+	using namespace UE::Net;
+	using namespace UE::Net::Private;
+
 	bool bWroteSomethingImportant = false;
-	for (const UE::Net::FSubObjectRegistry::FEntry& SubObjectInfo : SubObjectList.GetRegistryList())
+
+#if UE_NET_REPACTOR_NAME_DEBUG
+	// View address in memory viewer
+	const bool bSubObjectNameDebug = GCVarNetSubObjectNameDebug && GMetAsyncDemoNameDebugChance && Connection->IsReplay();
+	volatile PTRINT SubObjName = bSubObjectNameDebug ? UE_NET_ALLOCA_NAME_DEBUG_BUFFER() : 0;
+#endif
+
+	for (const FSubObjectRegistry::FEntry& SubObjectInfo : SubObjectList.GetRegistryList())
 	{
 		checkf(IsValid(SubObjectInfo.SubObject), TEXT("Found invalid subobject (%s) registered in %s::%s"), *GetNameSafe(SubObjectInfo.SubObject), *Actor->GetName(), *GetNameSafe(ReplicatedComponent));
 
 		if (CanSubObjectReplicateToClient(SubObjectInfo.NetCondition, SubObjectInfo.Key, ConditionMap))
 		{
+#if UE_NET_REPACTOR_NAME_DEBUG
+			if (bSubObjectNameDebug)
+			{
+				StoreSubObjectName(SubObjName, SubObjectInfo);
+			}
+#endif
+
 			bWroteSomethingImportant |= WriteSubObjectInBunch(SubObjectInfo.SubObject, Bunch, RepFlags);
 		}
 	}
diff --git a/Engine/Source/Runtime/Engine/Private/Net/NetNameDebug.cpp b/Engine/Source/Runtime/Engine/Private/Net/NetNameDebug.cpp
new file mode 100644
--- /dev/null
+++ b/Engine/Source/Runtime/Engine/Private/Net/NetNameDebug.cpp
@@ -0,0 +1,184 @@
+// Copyright Epic Games, Inc. All Rights Reserved.
+
+#include "Net/NetNameDebug.h"
+
+#if UE_NET_REPACTOR_NAME_DEBUG
+#include "UObject/Class.h"
+#include "HAL/IConsoleManager.h"
+#include "Stats/Stats2.h"
+
+
+#ifndef UE_NET_REPACTOR_NAME_DEBUG_PROFILING
+	#define UE_NET_REPACTOR_NAME_DEBUG_PROFILING 0
+#endif
+
+#if UE_NET_REPACTOR_NAME_DEBUG_PROFILING
+	DECLARE_CYCLE_STAT(TEXT("DemoAsync_NameDebug"), STAT_NetAsyncDemoNameDebug, STATGROUP_Net);
+#endif
+
+
+namespace UE::Net
+{
+	int32 GCVarNetActorNameDebug = 0;
+	int32 GCVarNetSubObjectNameDebug = 0;
+	static int32 GCVarNetSubObjectFullNameDebugUnsafe = 0;
+	static int32 GCVarNetAsyncDemoNameDebugChance = 10;
+	static int32 GAsyncDemoNameDebugCounter = 0;
+
+	static FAutoConsoleVariableRef CVarNetActorNameDebug(
+		TEXT("net.ActorNameDebug"),
+		GCVarNetActorNameDebug,
+		TEXT("When turned on the name of actors being replicated for replays, will be stored on the stack for crashdumps."));
+
+	static FAutoConsoleVariableRef CVarNetSubObjectNameDebug(
+		TEXT("net.SubObjectNameDebug"),
+		GCVarNetSubObjectNameDebug,
+		TEXT("When turned on the name of subobjects being replicated for replays, will be stored on the stack for crashdumps."));
+
+	static FAutoConsoleVariableRef CVarNetSubObjectFullNameDebugUnsafe(
+		TEXT("net.SubObjectFullNameDebugUnsafe"),
+		GCVarNetSubObjectFullNameDebugUnsafe,
+		TEXT("When turned on the full name of subobjects being replicated for replays, will be stored on the stack for crashdumps. ")
+		TEXT("This is unsafe, as when debugging for invalid objects, this will trigger the crash early."));
+
+	static FAutoConsoleVariableRef CVarNetAsyncDemoNameDebugChance(
+		TEXT("net.AsyncDemoNameDebugChance"),
+		GCVarNetAsyncDemoNameDebugChance,
+		TEXT("When actor/subobject demo name debugging is enabled, this is the percentage chance per replication that names will be cached. ")
+		TEXT("This is to minimize the performance impact - and also represents the percentage of crash dumps that will have data on the stack."));
+
+#if UE_NET_REPACTOR_NAME_DEBUG_PROFILING
+	static int32 GCVarNetAsyncDemoNameDebugProfiling = 0;
+
+	static FAutoConsoleVariableRef CVarNetAsyncDemoNameDebugProfiling(
+		TEXT("net.AsyncDemoNameDebugProfiling"),
+		GCVarNetAsyncDemoNameDebugProfiling,
+		TEXT("Enables profiling of 'net.ActorNameDebug' and 'net.SubObjectNameDebug' name caching."));
+#endif
+}
+
+namespace UE::Net::Private
+{
+	bool HasMetAsyncDemoNameDebugChance()
+	{
+		bool bReturnVal = false;
+
+		GAsyncDemoNameDebugCounter--;
+
+		if (GAsyncDemoNameDebugCounter <= 0)
+		{
+			const int32 Variance = static_cast<int32>(GCVarNetAsyncDemoNameDebugChance * 0.5);
+
+			bReturnVal = GAsyncDemoNameDebugCounter == 0;
+			GAsyncDemoNameDebugCounter = (GCVarNetAsyncDemoNameDebugChance - Variance) + FMath::RandRange(0, Variance) + 1;
+		}
+
+		return bReturnVal;
+	}
+
+	/** Stores a name in FNameDebugBuffer, replacing any existing value */
+	void StoreName_Internal(FNameDebugBuffer& Buffer, FName InName)
+	{
+		const FNameEntry* NameEntry = InName.GetDisplayNameEntry();
+
+		Buffer.bIsWide = NameEntry->IsWide();
+		Buffer.Length = NameEntry->GetNameLength();
+
+		if (Buffer.bIsWide)
+		{
+			NameEntry->GetWideName(*reinterpret_cast<WIDECHAR(*)[NAME_SIZE]>(const_cast<WIDECHAR*>(Buffer.Wide)));
+		}
+		else
+		{
+			NameEntry->GetAnsiName(*reinterpret_cast<ANSICHAR(*)[NAME_SIZE]>(const_cast<ANSICHAR*>(Buffer.Ansi)));
+		}
+	}
+
+	/** Appends a name to FNameDebugBuffer - not safe to use before a 'Store*' function initializes the buffer */
+	void AppendName_Internal(FNameDebugBuffer& Buffer, FName InName)
+	{
+		const FNameEntry* NameEntry = InName.GetDisplayNameEntry();
+		const bool bNewWide = NameEntry->IsWide();
+		const int32 AppendLength = NameEntry->GetNameLength();
+
+		if (Buffer.Length == 0)
+		{
+			Buffer.bIsWide = bNewWide;
+		}
+
+		if (Buffer.bIsWide == bNewWide && (Buffer.Length + AppendLength + 1) < FNameDebugBuffer::MaxNameAppend)
+		{
+			const int32 StartOffset = Buffer.Length;
+
+			Buffer.Length += AppendLength + 1;
+
+			if (Buffer.bIsWide)
+			{
+				Buffer.Wide[StartOffset] = '.';
+
+				NameEntry->GetWideName(*reinterpret_cast<WIDECHAR(*)[NAME_SIZE]>(const_cast<WIDECHAR*>(Buffer.Wide) + StartOffset + 1));
+			}
+			else
+			{
+				Buffer.Ansi[StartOffset] = '.';
+
+				NameEntry->GetAnsiName(*reinterpret_cast<ANSICHAR(*)[NAME_SIZE]>(const_cast<ANSICHAR*>(Buffer.Ansi) + StartOffset + 1));
+			}
+		}
+	}
+
+	/** Equivalent to UObjectBaseUtility::GetPathName - not safe to use before a 'Store*' function initializes the buffer */
+	void AppendPathName_Internal(FNameDebugBuffer& Buffer, const UObject* InObj)
+	{
+		const UObject* CurOuter = InObj->GetOuter();
+
+		if (CurOuter != nullptr)
+		{
+			AppendPathName_Internal(Buffer, CurOuter);
+		}
+
+		AppendName_Internal(Buffer, InObj->GetFName());
+	}
+
+	void AppendFullName_Internal(FNameDebugBuffer& Buffer, const UObject* InObj)
+	{
+		AppendName_Internal(Buffer, InObj->GetClass()->GetFName());
+		AppendPathName_Internal(Buffer, InObj);
+	}
+
+	void StoreFullName(PTRINT BufferPtr, const UObject* InObj)
+	{
+		if (BufferPtr != 0)
+		{
+			FNameDebugBuffer& Buffer = *reinterpret_cast<FNameDebugBuffer*>(BufferPtr);
+
+#if UE_NET_REPACTOR_NAME_DEBUG_PROFILING
+			CONDITIONAL_SCOPE_CYCLE_COUNTER(STAT_NetAsyncDemoNameDebug, GCVarNetAsyncDemoNameDebugProfiling);
+#endif
+
+			StoreName_Internal(Buffer, InObj->GetClass()->GetFName());
+			AppendPathName_Internal(Buffer, InObj);
+		}
+	}
+
+	void StoreSubObjectName(PTRINT BufferPtr, FName ClassName, FName ObjName, const UObject* InObj)
+	{
+		if (BufferPtr != 0)
+		{
+			FNameDebugBuffer& Buffer = *reinterpret_cast<FNameDebugBuffer*>(BufferPtr);
+
+#if UE_NET_REPACTOR_NAME_DEBUG_PROFILING
+			CONDITIONAL_SCOPE_CYCLE_COUNTER(STAT_NetAsyncDemoNameDebug, GCVarNetAsyncDemoNameDebugProfiling);
+#endif
+
+			StoreName_Internal(Buffer, ClassName);
+			AppendName_Internal(Buffer, ObjName);
+
+			if (GCVarNetSubObjectFullNameDebugUnsafe)
+			{
+				AppendFullName_Internal(Buffer, InObj);
+			}
+		}
+	}
+}
+#endif
diff --git a/Engine/Source/Runtime/Engine/Private/Net/NetNameDebug.h b/Engine/Source/Runtime/Engine/Private/Net/NetNameDebug.h
new file mode 100644
--- /dev/null
+++ b/Engine/Source/Runtime/Engine/Private/Net/NetNameDebug.h
@@ -0,0 +1,62 @@
+// Copyright Epic Games, Inc. All Rights Reserved.
+
+#pragma once
+
+#include "UObject/CoreNetTypes.h"
+
+#if UE_NET_REPACTOR_NAME_DEBUG
+#include "UObject/NameTypes.h"
+#include "Net/Core/Misc/NetSubObjectRegistry.h"
+
+/**
+ * Alloca's an FNameDebugBuffer to a PTRINT - to be used by one of the 'Store*' name debug buffer functions.
+ *
+ * Uses PTRINT, because actual pointers get optimized away, even when volatile. Copy the PTRINT value, and paste the address to a debug Memory window.
+ */
+#define UE_NET_ALLOCA_NAME_DEBUG_BUFFER() \
+	reinterpret_cast<PTRINT>(FMemory_Alloca_Aligned(sizeof(FNameDebugBuffer), alignof(FNameDebugBuffer)))
+
+namespace UE::Net
+{
+	extern ENGINE_API int32 GCVarNetActorNameDebug;
+	extern ENGINE_API int32 GCVarNetSubObjectNameDebug;
+};
+
+namespace UE::Net::Private
+{
+	/** Minimal FName data, kept on stack for crash dumps - volatile to prevent optimization - no construction */
+	struct FNameDebugBuffer
+	{
+		static constexpr int32 MaxNameAppend = 128;
+
+		union
+		{
+			volatile ANSICHAR Ansi[NAME_SIZE + MaxNameAppend];
+			volatile WIDECHAR Wide[NAME_SIZE + MaxNameAppend];
+		};
+
+		int32 Length;
+		bool bIsWide;
+	};
+
+
+	/** Returns true 'GCVarNetAsyncDemoNameDebugChance' percent of the time - but with some added determinism and variance */
+	bool HasMetAsyncDemoNameDebugChance();
+
+	/** Equivalent to UObjectBaseUtility::GetFullName */
+	void StoreFullName(PTRINT BufferPtr, const UObject* InObj);
+
+	/** Safely stores a subobject class and name in FNameDebugBuffer - optionally attempting an unsafe append of the full name after */
+	void StoreSubObjectName(PTRINT BufferPtr, FName ClassName, FName ObjName, const UObject* InObj);
+
+	inline void StoreSubObjectName(PTRINT BufferPtr, const FSubObjectRegistry::FEntry& SubObjectInfo)
+	{
+		StoreSubObjectName(BufferPtr, SubObjectInfo.SubObjectClassName, SubObjectInfo.SubObjectName, SubObjectInfo.SubObject);
+	}
+
+	inline void StoreSubObjectName(PTRINT BufferPtr, const FReplicatedComponentInfo& ComponentInfo)
+	{
+		StoreSubObjectName(BufferPtr, ComponentInfo.ComponentClassName, ComponentInfo.ComponentName, (UObject*)ComponentInfo.Component);
+	}
+};
+#endif
diff --git a/Engine/Source/Runtime/Net/Core/Public/Net/Core/Misc/NetSubObjectRegistry.h b/Engine/Source/Runtime/Net/Core/Public/Net/Core/Misc/NetSubObjectRegistry.h
--- a/Engine/Source/Runtime/Net/Core/Public/Net/Core/Misc/NetSubObjectRegistry.h
+++ b/Engine/Source/Runtime/Net/Core/Public/Net/Core/Misc/NetSubObjectRegistry.h
@@ -55,10 +55,22 @@ public:
 		/** The network condition that chooses which connection this subobject can be replicated to. Default is none which means all connections receive it. */
 		ELifetimeCondition NetCondition = COND_None;
 
+#if UE_NET_REPACTOR_NAME_DEBUG
+		/** Cached name of the subobject class, for minidump debugging */
+		FName SubObjectClassName;
+
+		/** Cached name of the subobject, for minidump debugging */
+		FName SubObjectName;
+#endif
+
 		FEntry(UObject* InSubObject, ELifetimeCondition InNetCondition = COND_None)
 			: SubObject(InSubObject)
 			, Key(InSubObject)
 			, NetCondition(InNetCondition)
+#if UE_NET_REPACTOR_NAME_DEBUG
+			, SubObjectClassName(InSubObject != nullptr ? ((UObjectBase*)((UObjectBase*)InSubObject)->GetClass())->GetFName() : NAME_None)
+			, SubObjectName(InSubObject != nullptr ? ((UObjectBase*)InSubObject)->GetFName() : NAME_None)
+#endif
 		{
 		}
 
@@ -92,10 +104,22 @@ struct FReplicatedComponentInfo
 	/** Collection of subobjects replicated with this component */
 	FSubObjectRegistry SubObjects;
 
+#if UE_NET_REPACTOR_NAME_DEBUG
+	/** Cached name of the component class, for minidump debugging */
+	FName ComponentClassName;
+
+	/** Cached name of the component, for minidump debugging */
+	FName ComponentName;
+#endif
+
 	FReplicatedComponentInfo(UActorComponent* InComponent, ELifetimeCondition InNetCondition = COND_None)
 		: Component(InComponent)
 		, Key((UObject*)InComponent)
 		, NetCondition(InNetCondition)
+#if UE_NET_REPACTOR_NAME_DEBUG
+		, ComponentClassName(InComponent != nullptr ? ((UObjectBase*)((UObjectBase*)InComponent)->GetClass())->GetFName() : NAME_None)
+		, ComponentName(InComponent != nullptr ? ((UObjectBase*)InComponent)->GetFName() : NAME_None)
+#endif
 	{
 	};
 
