Added optional validation of magic header, to reduce logspam/fallout from e.g. port scanning.

diff --git a/Engine/Source/Runtime/Engine/Private/PacketHandlers/StatelessConnectHandlerComponent.cpp b/Engine/Source/Runtime/Engine/Private/PacketHandlers/StatelessConnectHandlerComponent.cpp
--- a/Engine/Source/Runtime/Engine/Private/PacketHandlers/StatelessConnectHandlerComponent.cpp
+++ b/Engine/Source/Runtime/Engine/Private/PacketHandlers/StatelessConnectHandlerComponent.cpp
@@ -242,7 +242,7 @@ namespace UE::Net
 {
 	static float HandshakeResendInterval = 1.f;
 
-	FAutoConsoleVariableRef CVarNetHandshakeResendInterval(
+	static FAutoConsoleVariableRef CVarNetHandshakeResendInterval(
 		TEXT("net.HandshakeResendInterval"),
 		HandshakeResendInterval,
 		TEXT("The delay between resending handshake packets which we have not received a response for."));
@@ -261,24 +261,24 @@ namespace UE::Net
 	static int32 CurrentHandshakeVersion = static_cast<uint8>(EHandshakeVersion::Latest);
 #endif
 
-	FAutoConsoleVariableRef CVarNetMinHandshakeVersion(
+	static FAutoConsoleVariableRef CVarNetMinHandshakeVersion(
 		TEXT("net.MinHandshakeVersion"),
 		MinSupportedHandshakeVersion,
 		TEXT("The minimum supported stateless handshake protocol version (numeric)."));
 
-	FAutoConsoleVariableRef CVarNetCurrentHandshakeVersion(
+	static FAutoConsoleVariableRef CVarNetCurrentHandshakeVersion(
 		TEXT("net.CurrentHandshakeVersion"),
 		CurrentHandshakeVersion,
 		TEXT("The current supported stateless handshake protocol version (numeric)"));
 
-	TAutoConsoleVariable<int32> CVarNetDoHandshakeVersionFallback(
+	static TAutoConsoleVariable<int32> CVarNetDoHandshakeVersionFallback(
 		TEXT("net.DoHandshakeVersionFallback"),
 		0,
 		TEXT("Whether or not to (clientside) perform randomized falling-back to previous versions of the handshake protocol, upon failure."));
 
 	static int32 GHandshakeEnforceNetworkCLVersion = 0;
 
-	FAutoConsoleVariableRef CVarNetHandshakeEnforceNetworkCLVersion(
+	static FAutoConsoleVariableRef CVarNetHandshakeEnforceNetworkCLVersion(
 		TEXT("net.HandshakeEnforceNetworkCLVersion"),
 		GHandshakeEnforceNetworkCLVersion,
 		TEXT("Whether or not the stateless handshake should enforce the Network CL version, instead of the higher level netcode. ")
@@ -286,18 +286,26 @@ namespace UE::Net
 
 	static int32 GVerifyNetSessionID = 1;
 
-	FAutoConsoleVariableRef CVarNetVerifyNetSessionID(
+	static FAutoConsoleVariableRef CVarNetVerifyNetSessionID(
 		TEXT("net.VerifyNetSessionID"),
 		GVerifyNetSessionID,
 		TEXT("Whether or not verification of the packet SessionID value is performed."));
 
 	static int32 GVerifyNetClientID = 1;
 
-	FAutoConsoleVariableRef CVarNetVerifyNetClientID(
+	static FAutoConsoleVariableRef CVarNetVerifyNetClientID(
 		TEXT("net.VerifyNetClientID"),
 		GVerifyNetSessionID,
 		TEXT("Whether or not verification of the packet ClientID value is performed."));
 
+	static int32 GVerifyMagicHeader = 0;
+
+	static FAutoConsoleVariableRef CVarNetVerifyMagicHeader(
+		TEXT("net.VerifyMagicHeader"),
+		GVerifyMagicHeader,
+		TEXT("Whether or not verification of the magic header is performed, prior to processing a packet. ")
+		TEXT("Disable if transitioning to a new magic header, while wishing to continue supporting the old header for a time."));
+
 
 
 
@@ -378,7 +386,11 @@ StatelessConnectHandlerComponent::StatelessConnectHandlerComponent()
 				MagicHeader.Add(CurChar != '0');
 			}
 
-			if (!bValidBinaryStr)
+			if (bValidBinaryStr)
+			{
+				MagicHeaderUint = *static_cast<uint32*>(MagicHeader.GetData());
+			}
+			else
 			{
 				UE_LOG(LogHandshake, Error, TEXT("CVar net.MagicHeader must be a binary string, containing only 1's and 0's, e.g.: 00010101. Current string: %s"), *MagicHeaderStr);
 
@@ -1294,7 +1306,7 @@ void StatelessConnectHandlerComponent::Incoming(FBitReader& Packet)
 			else if (!bHasValidSessionID || !bHasValidClientID)
 			{
 #if !UE_BUILD_SHIPPING
-				UE_CLOG(TrackSessionConnIDLogs(), LogHandshake, Log,
+				UE_CLOG(TrackValidationLogs(), LogHandshake, Log,
 						TEXT("Incoming: Rejecting handshake packet with invalid session id (%i vs %i) or connection id (%i vs %i)."),
 						SessionID, CachedGlobalNetTravelCount, ClientID, CachedClientID);
 #endif
@@ -1321,7 +1333,7 @@ void StatelessConnectHandlerComponent::Incoming(FBitReader& Packet)
 	else if (!bHasValidSessionID || !bHasValidClientID)
 	{
 #if !UE_BUILD_SHIPPING
-		UE_CLOG(TrackSessionConnIDLogs(), LogHandshake, Log,
+		UE_CLOG(TrackValidationLogs(), LogHandshake, Log,
 				TEXT("Incoming: Rejecting game packet with invalid session id (%i vs %i) or connection id (%i vs %i)."),
 				SessionID, CachedGlobalNetTravelCount, ClientID, CachedClientID);
 #endif
@@ -1381,11 +1393,24 @@ void StatelessConnectHandlerComponent::IncomingConnectionless(FIncomingPacketRef
 
 	if (MagicHeader.Num() > 0)
 	{
-		// Don't bother with the expense of verifying the magic header here.
 		uint32 ReadMagic = 0;
+
 		Packet.SerializeBits(&ReadMagic, MagicHeader.Num());
+
+		if (GVerifyMagicHeader && ReadMagic != MagicHeaderUint)
+		{
+#if !UE_BUILD_SHIPPING
+			UE_CLOG(TrackValidationLogs(), LogNet, Log, TEXT("Rejecting packet with invalid magic header '%08X' vs '%08X' (%i bits)"),
+					ReadMagic, MagicHeaderUint, MagicHeader.Num());
+#endif
+
+			Packet.SetError();
+
+			return;
+		}
 	}
 
+
 	bool bHasValidSessionID = true;
 	uint8 SessionID = 0;
 	uint8 ClientID = 0;
@@ -1913,23 +1938,23 @@ void StatelessConnectHandlerComponent::Tick(float DeltaTime)
 }
 
 #if !UE_BUILD_SHIPPING
-bool StatelessConnectHandlerComponent::TrackSessionConnIDLogs()
+bool StatelessConnectHandlerComponent::TrackValidationLogs()
 {
-	const uint32 NumSessionConnIDLogsPerPeriod = 10;
-	const double SessionConnIDLogPeriod = 30.0;
+	const uint32 NumValidationLogsPerPeriod = 10;
+	const double ValidationLogPeriod = 30.0;
 	const double ElapsedTime = Driver->GetElapsedTime();
 
-	if (ElapsedTime - LastSessionConnIDLogPeriodStart >= SessionConnIDLogPeriod)
+	if (ElapsedTime - LastValidationLogPeriodStart >= ValidationLogPeriod)
 	{
-		LastSessionConnIDLogPeriodStart = ElapsedTime;
-		SessionConnIDLogCounter = 0;
+		LastValidationLogPeriodStart = ElapsedTime;
+		ValidationLogCounter = 0;
 	}
 	else
 	{
-		SessionConnIDLogCounter++;
+		ValidationLogCounter++;
 	}
 
-	return SessionConnIDLogCounter < NumSessionConnIDLogsPerPeriod;
+	return ValidationLogCounter < NumValidationLogsPerPeriod;
 }
 #endif
 
diff --git a/Engine/Source/Runtime/Engine/Public/PacketHandlers/StatelessConnectHandlerComponent.h b/Engine/Source/Runtime/Engine/Public/PacketHandlers/StatelessConnectHandlerComponent.h
--- a/Engine/Source/Runtime/Engine/Public/PacketHandlers/StatelessConnectHandlerComponent.h
+++ b/Engine/Source/Runtime/Engine/Public/PacketHandlers/StatelessConnectHandlerComponent.h
@@ -338,11 +338,11 @@ private:
 
 #if !UE_BUILD_SHIPPING
 	/**
-	 * Tracks logs for invalid Session ID's or Connection ID's, and limits the number of these logs
+	 * Tracks logs for invalid Session ID's, Connection ID's or Magic Header's, and limits the number of these logs
 	 *
 	 * @return		Whether or not logging is currently allowed
 	 */
-	bool TrackSessionConnIDLogs();
+	bool TrackValidationLogs();
 #endif
 
 
@@ -378,11 +378,11 @@ private:
 	EHandshakeVersion MinClientHandshakeVersion;
 
 #if !UE_BUILD_SHIPPING
-	/** The last time an 'invalid session/connection id' log period began */
-	double LastSessionConnIDLogPeriodStart = 0.0;
+	/** The last time an 'invalid SessionId/ClientId/MagicHeader' log period began */
+	double LastValidationLogPeriodStart = 0.0;
 
-	/** The number of 'invalid session/connection id' logs in the current log period (for logspam limiting) */
-	uint32 SessionConnIDLogCounter = 0;
+	/** The number of 'invalid SessionId/ClientId/MagicHeader' logs in the current log period (for logspam limiting) */
+	uint32 ValidationLogCounter = 0;
 #endif
 
 
@@ -419,6 +419,9 @@ private:
 	/** The magic header which is prepended to all packets */
 	TBitArray<> MagicHeader;
 
+	/** Integer representation of the magic header, for direct comparison */
+	uint32 MagicHeaderUint = 0;
+
 	/** The number of sent handshake packets, added to packets to aid packet dump debugging (may roll over) */
 	uint8 SentHandshakePacketCount = 0;
 
