Minor include compile time tweaks.

diff --git a/Engine/Plugins/Experimental/WebSocketNetworking/Source/WebSocketNetworking/Private/WebSocketNetDriver.cpp b/Engine/Plugins/Experimental/WebSocketNetworking/Source/WebSocketNetworking/Private/WebSocketNetDriver.cpp
--- a/Engine/Plugins/Experimental/WebSocketNetworking/Source/WebSocketNetworking/Private/WebSocketNetDriver.cpp
+++ b/Engine/Plugins/Experimental/WebSocketNetworking/Source/WebSocketNetworking/Private/WebSocketNetDriver.cpp
@@ -18,6 +18,7 @@ Unreal WebSocket network driver.
 
 #include "Engine/ChildConnection.h"
 #include "Misc/CommandLine.h"
+#include "PacketHandler.h"
 
 /*-----------------------------------------------------------------------------
 Declarations.
diff --git a/Engine/Plugins/NetcodeUnitTest/NetcodeUnitTest/Source/NetcodeUnitTest/Private/MinimalClient.cpp b/Engine/Plugins/NetcodeUnitTest/NetcodeUnitTest/Source/NetcodeUnitTest/Private/MinimalClient.cpp
--- a/Engine/Plugins/NetcodeUnitTest/NetcodeUnitTest/Source/NetcodeUnitTest/Private/MinimalClient.cpp
+++ b/Engine/Plugins/NetcodeUnitTest/NetcodeUnitTest/Source/NetcodeUnitTest/Private/MinimalClient.cpp
@@ -8,6 +8,7 @@
 #include "Engine/GameEngine.h"
 #include "Engine/ActorChannel.h"
 #include "Net/DataChannel.h"
+#include "PacketHandler.h"
 
 #include "UnitTest.h"
 #include "NetcodeUnitTest.h"
diff --git a/Engine/Plugins/Online/OnlineSubsystemOculus/Source/Private/OculusNetConnection.cpp b/Engine/Plugins/Online/OnlineSubsystemOculus/Source/Private/OculusNetConnection.cpp
--- a/Engine/Plugins/Online/OnlineSubsystemOculus/Source/Private/OculusNetConnection.cpp
+++ b/Engine/Plugins/Online/OnlineSubsystemOculus/Source/Private/OculusNetConnection.cpp
@@ -6,6 +6,7 @@
 #include "IPAddressOculus.h"
 
 #include "Net/DataChannel.h"
+#include "PacketHandler.h"
 
 void UOculusNetConnection::InitBase(UNetDriver* InDriver, class FSocket* InSocket, const FURL& InURL, EConnectionState InState, int32 InMaxPacket, int32 InPacketOverhead)
 {
diff --git a/Engine/Plugins/Online/OnlineSubsystemUtils/Source/OnlineSubsystemUtils/Private/IpConnection.cpp b/Engine/Plugins/Online/OnlineSubsystemUtils/Source/OnlineSubsystemUtils/Private/IpConnection.cpp
--- a/Engine/Plugins/Online/OnlineSubsystemUtils/Source/OnlineSubsystemUtils/Private/IpConnection.cpp
+++ b/Engine/Plugins/Online/OnlineSubsystemUtils/Source/OnlineSubsystemUtils/Private/IpConnection.cpp
@@ -21,6 +21,7 @@ Notes:
 #include "Net/Core/Misc/PacketAudit.h"
 #include "Misc/ScopeExit.h"
 #include "NetAddressResolution.h"
+#include "PacketHandler.h"
 
 #include UE_INLINE_GENERATED_CPP_BY_NAME(IpConnection)
 
diff --git a/Engine/Plugins/Runtime/Steam/SteamSockets/Source/SteamSockets/Private/SteamSocketsNetDriver.cpp b/Engine/Plugins/Runtime/Steam/SteamSockets/Source/SteamSockets/Private/SteamSocketsNetDriver.cpp
--- a/Engine/Plugins/Runtime/Steam/SteamSockets/Source/SteamSockets/Private/SteamSocketsNetDriver.cpp
+++ b/Engine/Plugins/Runtime/Steam/SteamSockets/Source/SteamSockets/Private/SteamSocketsNetDriver.cpp
@@ -12,6 +12,7 @@
 #include "Engine/NetworkDelegates.h"
 #include "Engine/World.h"
 #include "Misc/CommandLine.h"
+#include "PacketHandler.h"
 
 void USteamSocketsNetDriver::PostInitProperties()
 {
diff --git a/Engine/Source/Runtime/Engine/Classes/Engine/NetConnection.h b/Engine/Source/Runtime/Engine/Classes/Engine/NetConnection.h
--- a/Engine/Source/Runtime/Engine/Classes/Engine/NetConnection.h
+++ b/Engine/Source/Runtime/Engine/Classes/Engine/NetConnection.h
@@ -6,8 +6,12 @@
 
 #pragma once
 
+#if UE_ENABLE_INCLUDE_ORDER_DEPRECATED_IN_5_2
 #include "CoreMinimal.h"
-#include "GameFramework/Actor.h"
+#include "Net/RPCDoSDetection.h"
+#include "Net/NetConnectionFaultRecovery.h"
+#endif
+#include "UObject/ObjectKey.h"
 #include "UObject/ObjectMacros.h"
 #include "UObject/UObjectGlobals.h"
 #include "Serialization/BitWriter.h"
@@ -29,11 +33,10 @@
 #include "Net/Common/Packets/PacketTraits.h"
 #include "Net/Core/Misc/ResizableCircularQueue.h"
 #include "Net/NetAnalyticsTypes.h"
-#include "Net/RPCDoSDetection.h"
 #include "Net/Core/Connection/NetCloseResult.h"
-#include "Net/NetConnectionFaultRecovery.h"
 #include "Net/TrafficControl.h"
 #include "HAL/LowLevelMemTracker.h"
+#include "GameFramework/Actor.h"
 
 #include "NetConnection.generated.h"
 
@@ -48,10 +51,13 @@ class UActorChannel;
 class UChildConnection;
 class ULevelStreaming;
 struct FEncryptionKeyResponse;
+class PacketHandler;
+class FRPCDoSDetection;
 
 namespace UE::Net
 {
 	class FNetPing;
+	class FNetConnectionFaultRecovery;
 }
 
 typedef TMap<TWeakObjectPtr<AActor>, UActorChannel*, FDefaultSetAllocator, TWeakObjectPtrMapKeyFuncs<TWeakObjectPtr<AActor>, UActorChannel*>> FActorChannelMap;
@@ -334,7 +340,7 @@ public:
 	}
 
 	/** Destructor */
-	ENGINE_API virtual ~UNetConnection() {};
+	ENGINE_API virtual ~UNetConnection();
 
 private:
 	uint32 bInternalAck : 1;				// Internally ack all packets, for 100% reliable connections.
@@ -858,6 +864,8 @@ public:
 	// Constructors and destructors.
 	ENGINE_API UNetConnection(const FObjectInitializer& ObjectInitializer = FObjectInitializer::Get());
 
+	ENGINE_API UNetConnection(FVTableHelper& Helper);
+
 	//~ Begin UObject Interface.
 
 	ENGINE_API virtual void Serialize( FArchive& Ar ) override;
@@ -1654,10 +1662,10 @@ private:
 	bool bLoggedFlushNetQueuedBitsOverflow = false;
 
 	/** RPC/Replication code DoS detection */
-	FRPCDoSDetection RPCDoS;
+	TUniquePtr<FRPCDoSDetection> RPCDoS;
 
 	/** NetConnection specific Fault Recovery for attempting to recover from connection faults, before triggering Close */
-	UE::Net::FNetConnectionFaultRecovery FaultRecovery;
+	TUniquePtr<UE::Net::FNetConnectionFaultRecovery> FaultRecovery;
 
 	/** Whether or not this NetConnection has already received an NMT_CloseReason message */
 	bool bReceivedCloseReason = false;
@@ -1675,14 +1683,14 @@ public:
 	bool GetAutoFlush() const { return bAutoFlush; }
 	void SetAutoFlush(bool bValue) { bAutoFlush = bValue; }
 
-	FRPCDoSDetection& GetRPCDoS()
+	FRPCDoSDetection* GetRPCDoS()
 	{
-		return RPCDoS;
+		return RPCDoS.Get();
 	}
 
 	UE::Net::FNetConnectionFaultRecovery* GetFaultRecovery()
 	{
-		return &FaultRecovery;
+		return FaultRecovery.Get();
 	}
 
 	UE::Net::FNetPing* GetNetPing()
@@ -1733,21 +1741,7 @@ protected:
 struct FScopedRepContext
 {
 public:
-	FScopedRepContext(UNetConnection* InConnection, AActor* InActor)
-		: Connection(InConnection)
-	{
-		if (Connection)
-		{
-			check(!Connection->RepContextActor);
-			check(!Connection->RepContextLevel);
-
-			Connection->RepContextActor = InActor;
-			if (InActor)
-			{
-				Connection->RepContextLevel = InActor->GetLevel();
-			}
-		}
-	}
+	FScopedRepContext(UNetConnection* InConnection, AActor* InActor);
 
 	FScopedRepContext(UNetConnection* InConnection, ULevel* InLevel)
 		: Connection(InConnection)
diff --git a/Engine/Source/Runtime/Engine/Classes/Engine/NetDriver.h b/Engine/Source/Runtime/Engine/Classes/Engine/NetDriver.h
--- a/Engine/Source/Runtime/Engine/Classes/Engine/NetDriver.h
+++ b/Engine/Source/Runtime/Engine/Classes/Engine/NetDriver.h
@@ -15,8 +15,8 @@
 #include "UObject/CoreNet.h"
 #if UE_ENABLE_INCLUDE_ORDER_DEPRECATED_IN_5_2
 #include "GameFramework/WorldSettings.h"
-#endif
 #include "PacketHandler.h"
+#endif
 #include "Channel.h"
 #include "Net/Core/Misc/DDoSDetection.h"
 #include "IPAddress.h"
@@ -341,6 +341,8 @@ class UChannel;
 class IAnalyticsProvider;
 class FNetAnalyticsAggregator;
 class UNetDriver;
+class PacketHandler;
+struct FReplicatedStaticActorDestructionInfo;
 
 enum class ECreateReplicationChangelistMgrFlags;
 enum class EEngineNetworkRuntimeFeatures : uint16;
@@ -719,9 +721,7 @@ private:
 public:
 
 	/** Destructor */
-	PRAGMA_DISABLE_DEPRECATION_WARNINGS
-	ENGINE_API virtual ~UNetDriver() {};
-	PRAGMA_ENABLE_DEPRECATION_WARNINGS
+	ENGINE_API virtual ~UNetDriver();
 
 	/** Used to specify the class to use for connections */
 	UPROPERTY(Config)
@@ -1238,6 +1238,8 @@ public:
 	// Constructors.
 	ENGINE_API UNetDriver(const FObjectInitializer& ObjectInitializer = FObjectInitializer::Get());
 
+	ENGINE_API UNetDriver(FVTableHelper& Helper);
+
 
 	//~ Begin UObject Interface.
 	ENGINE_API virtual void PostInitProperties() override;
diff --git a/Engine/Source/Runtime/Engine/Private/DataReplication.cpp b/Engine/Source/Runtime/Engine/Private/DataReplication.cpp
--- a/Engine/Source/Runtime/Engine/Private/DataReplication.cpp
+++ b/Engine/Source/Runtime/Engine/Private/DataReplication.cpp
@@ -23,6 +23,7 @@
 #include "ProfilingDebugging/CsvProfiler.h"
 #include "HAL/LowLevelMemStats.h"
 #include "Net/Core/PushModel/Types/PushModelPerNetDriverState.h"
+#include "Net/RPCDoSDetection.h"
 
 DECLARE_LLM_MEMORY_STAT(TEXT("NetObjReplicator"), STAT_NetObjReplicatorLLM, STATGROUP_LLMFULL);
 LLM_DEFINE_TAG(NetObjReplicator, NAME_None, TEXT("Networking"), GET_STATFNAME(STAT_NetObjReplicatorLLM), GET_STATFNAME(STAT_NetworkingSummaryLLM));
@@ -1263,7 +1264,7 @@ bool FObjectReplicator::ReceivedRPC(FNetBitReader& Reader, const FReplicationFla
 	UObject* Object = GetObject();
 	FName FunctionName = FieldCache->Field.GetFName();
 	UFunction* Function = Object->FindFunction(FunctionName);
-	FRPCDoSDetection& RPCDoS = Connection->GetRPCDoS();
+	FRPCDoSDetection* RPCDoS = Connection->GetRPCDoS();
 	FScopedRPCTimingTracker ScopedTracker(Function, Connection);
 	SCOPE_CYCLE_COUNTER(STAT_NetReceiveRPC);
 	SCOPE_CYCLE_UOBJECT(Function, Function);
@@ -1287,13 +1288,13 @@ bool FObjectReplicator::ReceivedRPC(FNetBitReader& Reader, const FReplicationFla
 	}
 
 
-	if (bIsServer && RPCDoS.IsRPCDoSDetectionEnabled() && !Connection->IsReplay())
+	if (RPCDoS != nullptr && bIsServer && RPCDoS->IsRPCDoSDetectionEnabled() && !Connection->IsReplay())
 	{
-		ScopedTracker.RPCDoS = &RPCDoS;
+		ScopedTracker.RPCDoS = RPCDoS;
 
-		if (UNLIKELY(RPCDoS.ShouldMonitorReceivedRPC()))
+		if (UNLIKELY(RPCDoS->ShouldMonitorReceivedRPC()))
 		{
-			ERPCNotifyResult Result = RPCDoS.NotifyReceivedRPC(Reader, UnmappedGuids, Object, Function, FunctionName);
+			ERPCNotifyResult Result = RPCDoS->NotifyReceivedRPC(Reader, UnmappedGuids, Object, Function, FunctionName);
 
 			if (UNLIKELY(Result == ERPCNotifyResult::BlockRPC))
 			{
@@ -1304,7 +1305,7 @@ bool FObjectReplicator::ReceivedRPC(FNetBitReader& Reader, const FReplicationFla
 		}
 		else
 		{
-			RPCDoS.LightweightReceivedRPC(Function, FunctionName);
+			RPCDoS->LightweightReceivedRPC(Function, FunctionName);
 		}
 	}
 
diff --git a/Engine/Source/Runtime/Engine/Private/Net/DefaultFaultHandler.cpp b/Engine/Source/Runtime/Engine/Private/Net/DefaultFaultHandler.cpp
--- a/Engine/Source/Runtime/Engine/Private/Net/DefaultFaultHandler.cpp
+++ b/Engine/Source/Runtime/Engine/Private/Net/DefaultFaultHandler.cpp
@@ -3,6 +3,7 @@
 // Includes
 #include "Net/DefaultFaultHandler.h"
 #include "Engine/NetConnection.h"
+#include "Net/Core/Connection/NetConnectionFaultRecoveryBase.h"
 
 
 namespace UE
diff --git a/Engine/Source/Runtime/Engine/Private/NetConnection.cpp b/Engine/Source/Runtime/Engine/Private/NetConnection.cpp
--- a/Engine/Source/Runtime/Engine/Private/NetConnection.cpp
+++ b/Engine/Source/Runtime/Engine/Private/NetConnection.cpp
@@ -43,6 +43,8 @@
 #include "HAL/LowLevelMemStats.h"
 #include "Net/NetPing.h"
 #include "LevelUtils.h"
+#include "Net/RPCDoSDetection.h"
+#include "Net/NetConnectionFaultRecovery.h"
 #if UE_WITH_IRIS
 #include "Iris/IrisConfig.h"
 #include "Iris/ReplicationSystem/ReplicationSystem.h"
@@ -333,6 +335,13 @@ UNetConnection::UNetConnection(const FObjectInitializer& ObjectInitializer)
 {
 }
 
+UNetConnection::UNetConnection(FVTableHelper& Helper)
+	: Super(Helper)
+{
+}
+
+UNetConnection::~UNetConnection() = default;
+
 void UNetConnection::InitChannelData()
 {
 	LLM_SCOPE_BYTAG(NetConnection);
@@ -447,16 +456,17 @@ void UNetConnection::InitBase(UNetDriver* InDriver,class FSocket* InSocket, cons
 	bLoggedFlushNetQueuedBitsOverflow = false;
 
 	// Reset Handler
-	Handler.Reset(NULL);
+	Handler.Reset();
 
 	InitHandler();
 
+	FaultRecovery = MakeUnique<FNetConnectionFaultRecovery>();
 
-	FaultRecovery.InitDefaults((Driver != nullptr ? Driver->GetNetDriverDefinition().ToString() : TEXT("")), this);
+	FaultRecovery->InitDefaults((Driver != nullptr ? Driver->GetNetDriverDefinition().ToString() : TEXT("")), this);
 
 	if (CVarLogUnhandledFaults.GetValueOnAnyThread() != 0)
 	{
-		FaultRecovery.FaultManager.SetUnhandledResultCallback([](FNetResult&& InResult)
+		FaultRecovery->FaultManager.SetUnhandledResultCallback([](FNetResult&& InResult)
 			{
 				static TArray<uint32> LoggedResults;
 
@@ -508,7 +518,9 @@ void UNetConnection::InitBase(UNetDriver* InDriver,class FSocket* InSocket, cons
 
 	if (bIsServer && !bIsReplay)
 	{
-		RPCDoS.Init(Driver->GetNetDriverDefinition(), AnalyticsAggregator,
+		RPCDoS = MakeUnique<FRPCDoSDetection>();
+
+		RPCDoS->Init(Driver->GetNetDriverDefinition(), AnalyticsAggregator,
 			[WorldPtr = TWeakObjectPtr<UWorld>(InDriver->GetWorld())]() -> UWorld*
 			{
 				return WorldPtr.IsValid() ? WorldPtr.Get() : nullptr;
@@ -984,7 +996,11 @@ void UNetConnection::Close(FNetResult&& CloseReason)
 			NetAnalyticsData->CommitAnalytics(AnalyticsVars);
 		}
 
-		RPCDoS.NotifyClose();
+		if (RPCDoS.IsValid())
+		{
+			RPCDoS->NotifyClose();
+		}
+
 		NetPing.Reset();
 
 		if (const uint32 MyConnectionId = GetConnectionId())
@@ -1002,7 +1018,7 @@ void UNetConnection::HandleNetResultOrClose(ENetCloseResult InResult)
 {
 	using namespace UE::Net;
 
-	const EHandleNetResult RecoveryResult = FaultRecovery.HandleNetResult(InResult);
+	const EHandleNetResult RecoveryResult = (FaultRecovery.IsValid() ? FaultRecovery->HandleNetResult(InResult) : EHandleNetResult::NotHandled);
 
 	if (RecoveryResult == EHandleNetResult::NotHandled)
 	{
@@ -1201,7 +1217,7 @@ void UNetConnection::CleanUp()
 
 	CleanupDormantActorState();
 
-	Handler.Reset(NULL);
+	Handler.Reset();
 
 	SetClientLoginState(EClientLoginState::CleanedUp);
 
@@ -1739,7 +1755,9 @@ void UNetConnection::ReceivedRawPacket( void* InData, int32 Count )
 
 			if (!bErrorNotRecoverable)
 			{
-				const EHandleNetResult RecoveryResult = FaultRecovery.FaultManager.HandleNetResult(MoveTemp(*Traits.ExtendedError));
+				const EHandleNetResult RecoveryResult = (FaultRecovery.IsValid() ?
+					FaultRecovery->FaultManager.HandleNetResult(MoveTemp(*Traits.ExtendedError)) :
+					EHandleNetResult::NotHandled);
 
 				bCloseConnection = RecoveryResult == EHandleNetResult::NotHandled;
 			}
@@ -1836,14 +1854,14 @@ void UNetConnection::PreTickDispatch()
 	{
 		double LastTickDispatchRealtime = Driver->LastTickDispatchRealtime;
 
-		if (bIsServer)
+		if (bIsServer && RPCDoS.IsValid())
 		{
-			RPCDoS.PreTickDispatch(LastTickDispatchRealtime);
+			RPCDoS->PreTickDispatch(LastTickDispatchRealtime);
 		}
 
-		if (FaultRecovery.DoesRequireTick())
+		if (FaultRecovery.IsValid() && FaultRecovery->DoesRequireTick())
 		{
-			FaultRecovery.TickRealtime(LastTickDispatchRealtime);
+			FaultRecovery->TickRealtime(LastTickDispatchRealtime);
 		}
 	}
 }
@@ -1862,9 +1880,9 @@ void UNetConnection::PostTickDispatch()
 		PacketAnalytics.Tick();
 	}
 
-	if (bIsServer && !IsReplay())
+	if (RPCDoS.IsValid() && bIsServer && !IsReplay())
 	{
-		RPCDoS.PostTickDispatch();
+		RPCDoS->PostTickDispatch();
 	}
 }
 
@@ -2893,18 +2911,18 @@ void UNetConnection::ReceivedPacket( FBitReader& Reader, bool bIsReinjectedPacke
 	double PostReceiveTime = 0.0;
 
 	{
-		if (bIsServer && !IsReplay())
+		if (RPCDoS.IsValid() && bIsServer && !IsReplay())
 		{
-			RPCDoS.PreReceivedPacket(CurrentReceiveTimeInS);
+			RPCDoS->PreReceivedPacket(CurrentReceiveTimeInS);
 		}
 
 		ON_SCOPE_EXIT
 		{
 			PostReceiveTime = FPlatformTime::Seconds();
 
-			if (bIsServer && !IsReplay())
+			if (RPCDoS.IsValid() && bIsServer && !IsReplay())
 			{
-				RPCDoS.PostReceivedPacket(PostReceiveTime);
+				RPCDoS->PostReceivedPacket(PostReceiveTime);
 			}
 		};
 
@@ -5538,4 +5556,23 @@ void ConsumeAllChannelRecords(FWrittenChannelsRecord& WrittenChannelsRecord, Fun
 
 }
 
+/**
+ * FScopedRepContext
+ */
+
+FScopedRepContext::FScopedRepContext(UNetConnection* InConnection, AActor* InActor)
+	: Connection(InConnection)
+{
+	if (Connection)
+	{
+		check(!Connection->RepContextActor);
+		check(!Connection->RepContextLevel);
+
+		Connection->RepContextActor = InActor;
+		if (InActor)
+		{
+			Connection->RepContextLevel = InActor->GetLevel();
+		}
+	}
+}
 
diff --git a/Engine/Source/Runtime/Engine/Private/NetDriver.cpp b/Engine/Source/Runtime/Engine/Private/NetDriver.cpp
--- a/Engine/Source/Runtime/Engine/Private/NetDriver.cpp
+++ b/Engine/Source/Runtime/Engine/Private/NetDriver.cpp
@@ -5,7 +5,6 @@
 =============================================================================*/
 
 #include "Engine/NetDriver.h"
-#include "CoreMinimal.h"
 #include "Misc/CoreMisc.h"
 #include "Misc/CommandLine.h"
 #include "Misc/NetworkGuid.h"
@@ -530,6 +529,15 @@ UNetDriver::UNetDriver(const FObjectInitializer& ObjectInitializer)
 	UpdateDelayRandomStream.Initialize(FApp::bUseFixedSeed ? GetFName() : NAME_None);
 }
 
+UNetDriver::UNetDriver(FVTableHelper& Helper)
+	: Super(Helper)
+{
+}
+
+PRAGMA_DISABLE_DEPRECATION_WARNINGS
+UNetDriver::~UNetDriver() = default;
+PRAGMA_ENABLE_DEPRECATION_WARNINGS
+
 void UNetDriver::InitPacketSimulationSettings()
 {
 #if DO_ENABLE_NET_TEST
@@ -1385,7 +1393,7 @@ bool UNetDriver::InitBase(bool bInitAsClient, FNetworkNotify* InNotify, const FU
 
 	if (!bInitAsClient)
 	{
-		ConnectionlessHandler.Reset(nullptr);
+		ConnectionlessHandler.Reset();
 		
 		if (!IsUsingIrisReplication())
 		{
@@ -1646,7 +1654,7 @@ void UNetDriver::Shutdown()
 	}
 	ActorChannelPool.Empty();
 
-	ConnectionlessHandler.Reset(nullptr);
+	ConnectionlessHandler.Reset();
 
 	SetReplicationDriver(nullptr);
 #if UE_WITH_IRIS
@@ -1790,13 +1798,16 @@ void UNetDriver::TickDispatch( float DeltaTime )
 			{
 				UNetConnection* CurConn = ClientConnections[ConnIdx];
 
-				if (CurConn->GetConnectionState() == USOCK_Closed)
-				{
-					CurConn->CleanUp();
-				}
-				else if (IsValid(CurConn))
+				if (IsValid(CurConn))
 				{
-					CurConn->PreTickDispatch();
+					if (CurConn->GetConnectionState() == USOCK_Closed)
+					{
+						CurConn->CleanUp();
+					}
+					else
+					{
+						CurConn->PreTickDispatch();
+					}
 				}
 			}
 		}
@@ -5722,9 +5733,9 @@ void UNetDriver::InitDestroyedStartupActors()
 			if (Level)
 			{
 				const TArray<FReplicatedStaticActorDestructionInfo>& DestroyedReplicatedStaticActors = Level->GetDestroyedReplicatedStaticActors();
-				for(const FReplicatedStaticActorDestructionInfo& Info : DestroyedReplicatedStaticActors)
+				for(const FReplicatedStaticActorDestructionInfo& CurInfo : DestroyedReplicatedStaticActors)
 				{
-					CreateReplicatedStaticActorDestructionInfo(Level, Info);
+					CreateReplicatedStaticActorDestructionInfo(Level, CurInfo);
 				}
 			}
 		}
