Skeletal Mesh - Fix missing FSkeletalMaterial declaration caused by the removal of SkeletalMeshRenderData.h from SkeletalMesh.h in the recent global removal of includes.

#rnx

diff --git a/Engine/Source/Runtime/Engine/Classes/Engine/SkeletalMesh.h b/Engine/Source/Runtime/Engine/Classes/Engine/SkeletalMesh.h
--- a/Engine/Source/Runtime/Engine/Classes/Engine/SkeletalMesh.h
+++ b/Engine/Source/Runtime/Engine/Classes/Engine/SkeletalMesh.h
@@ -62,6 +62,7 @@ class FSkinnedAssetPostLoadContext;
 class USkeletalMeshLODSettings;
 class USkeleton;
 struct FMeshUVChannelInfo;
+struct FSkeletalMaterial;
 struct FSkinWeightProfileInfo;
 class FSkinWeightVertexBuffer;
 enum class ESkeletalMeshGeoImportVersions : uint8;
