Deprecate initialize function since it should have been private

diff --git a/Engine/Source/Runtime/Engine/Classes/Animation/AnimNodeBase.h b/Engine/Source/Runtime/Engine/Classes/Animation/AnimNodeBase.h
--- a/Engine/Source/Runtime/Engine/Classes/Animation/AnimNodeBase.h
+++ b/Engine/Source/Runtime/Engine/Classes/Animation/AnimNodeBase.h
@@ -453,7 +453,7 @@ public:
 		: FAnimationBaseContext(InAnimInstanceProxy)
 		, bExpectsAdditivePose(bInExpectsAdditivePose)
 	{
-		Initialize(InAnimInstanceProxy);
+		InitializeImpl(InAnimInstanceProxy);
 	}
 
 	// This constructor allocates a new uninitialized pose, copying non-pose state from the source context
@@ -461,13 +461,14 @@ public:
 		: FAnimationBaseContext(SourceContext.AnimInstanceProxy)
 		, bExpectsAdditivePose(SourceContext.bExpectsAdditivePose || bInOverrideExpectsAdditivePose)
 	{
-		Initialize(SourceContext.AnimInstanceProxy);
+		InitializeImpl(SourceContext.AnimInstanceProxy);
 
 		CurrentNodeId = SourceContext.CurrentNodeId;
 		PreviousNodeId = SourceContext.PreviousNodeId;
 	}
 
-	ENGINE_API void Initialize(FAnimInstanceProxy* InAnimInstanceProxy);
+	UE_DEPRECATED(5.2, "This function will be made private. It should never be called externally, use the constructor instead.")
+	void Initialize(FAnimInstanceProxy* InAnimInstanceProxy) { InitializeImpl(InAnimInstanceProxy); }
 
 	// Log evaluation message
 	void LogMessage(EMessageSeverity::Type InSeverity, FText InMessage) const { LogMessageInternal("Evaluate", InSeverity, InMessage); }
@@ -503,7 +504,7 @@ public:
 	{
 		if (AnimInstanceProxy != Other.AnimInstanceProxy)
 		{
-			Initialize(AnimInstanceProxy);
+			InitializeImpl(AnimInstanceProxy);
 		}
 
 		Pose = Other.Pose;
@@ -517,6 +518,7 @@ public:
 	bool ExpectsAdditivePose() const { return bExpectsAdditivePose; }
 
 private:
+	ENGINE_API void InitializeImpl(FAnimInstanceProxy* InAnimInstanceProxy);
 
 	// Is this pose expected to be an additive pose
 	bool bExpectsAdditivePose;
diff --git a/Engine/Source/Runtime/Engine/Private/Animation/AnimNodeBase.cpp b/Engine/Source/Runtime/Engine/Private/Animation/AnimNodeBase.cpp
--- a/Engine/Source/Runtime/Engine/Private/Animation/AnimNodeBase.cpp
+++ b/Engine/Source/Runtime/Engine/Private/Animation/AnimNodeBase.cpp
@@ -110,7 +110,7 @@ void FAnimationBaseContext::LogMessageInternal(FName InLogType, EMessageSeverity
 /////////////////////////////////////////////////////
 // FPoseContext
 
-void FPoseContext::Initialize(FAnimInstanceProxy* InAnimInstanceProxy)
+void FPoseContext::InitializeImpl(FAnimInstanceProxy* InAnimInstanceProxy)
 {
 	AnimInstanceProxy = InAnimInstanceProxy;
 
