Make sure unreachable assets are removed from the streamer before deleting it

#tests made a new level sequence and rendered it without crash

diff --git a/Engine/Source/Runtime/Engine/Private/ContentStreaming.cpp b/Engine/Source/Runtime/Engine/Private/ContentStreaming.cpp
--- a/Engine/Source/Runtime/Engine/Private/ContentStreaming.cpp
+++ b/Engine/Source/Runtime/Engine/Private/ContentStreaming.cpp
@@ -1365,6 +1365,9 @@ void FStreamingManagerCollection::AddOrRemoveTextureStreamingManagerIfNeeded(boo
 				It->UnlinkStreaming();
 			}
 
+			// Remove unreachable assets from the streamer before it goes away
+			UnhashUnreachableObjects(false);
+
 			RemoveStreamingManager(RenderAssetStreamingManager);
 			delete RenderAssetStreamingManager;
 			RenderAssetStreamingManager = nullptr;
