Fixed performance issues in camera rig rail actor and made spline mesh component skip mesh and physics updates when incoming data hasn't changed from the current state (on behalf of Jaco.VanDyk)

diff --git a/Engine/Source/Runtime/CinematicCamera/Private/CameraRig_Rail.cpp b/Engine/Source/Runtime/CinematicCamera/Private/CameraRig_Rail.cpp
--- a/Engine/Source/Runtime/CinematicCamera/Private/CameraRig_Rail.cpp
+++ b/Engine/Source/Runtime/CinematicCamera/Private/CameraRig_Rail.cpp
@@ -128,11 +128,12 @@ void ACameraRig_Rail::UpdatePreviewMeshes()
 				USplineMeshComponent* const SplineMeshComp = PreviewRailMeshSegments[PtIdx];
 				if (SplineMeshComp)
 				{
-					SplineMeshComp->SetVisibility(bShowRailVisualization);
-					SplineMeshComp->SetStartScale(FVector2D(PreviewMeshScale, PreviewMeshScale));
-					SplineMeshComp->SetEndScale(FVector2D(PreviewMeshScale, PreviewMeshScale));
-					SplineMeshComp->SetForwardAxis(ESplineMeshAxis::Z);
-					SplineMeshComp->SetStartAndEnd(StartLoc, StartTangent, EndLoc, EndTangent, true);
+					SplineMeshComp->SetVisibility(bShowRailVisualization,false);
+					SplineMeshComp->SetStartScale(FVector2D(PreviewMeshScale, PreviewMeshScale),false);
+					SplineMeshComp->SetEndScale(FVector2D(PreviewMeshScale, PreviewMeshScale),false);
+					SplineMeshComp->SetForwardAxis(ESplineMeshAxis::Z,false);
+					SplineMeshComp->SetStartAndEnd(StartLoc, StartTangent, EndLoc, EndTangent, false);
+					SplineMeshComp->UpdateMesh();
 				}
 			}
 
diff --git a/Engine/Source/Runtime/Engine/Private/Components/SplineMeshComponent.cpp b/Engine/Source/Runtime/Engine/Private/Components/SplineMeshComponent.cpp
--- a/Engine/Source/Runtime/Engine/Private/Components/SplineMeshComponent.cpp
+++ b/Engine/Source/Runtime/Engine/Private/Components/SplineMeshComponent.cpp
@@ -265,6 +265,11 @@ FVector USplineMeshComponent::GetEndTangent() const
 
 void USplineMeshComponent::SetEndTangent(FVector EndTangent, bool bUpdateMesh)
 {
+	if (SplineParams.EndTangent == EndTangent)
+	{
+		return;
+	}
+
 	SplineParams.EndTangent = EndTangent;
 	bMeshDirty = true;
 	if (bUpdateMesh)
@@ -275,6 +280,11 @@ void USplineMeshComponent::SetEndTangent(FVector EndTangent, bool bUpdateMesh)
 
 void USplineMeshComponent::SetStartAndEnd(FVector StartPos, FVector StartTangent, FVector EndPos, FVector EndTangent, bool bUpdateMesh)
 {
+	if (SplineParams.StartPos == StartPos && SplineParams.StartTangent == StartTangent && SplineParams.EndPos == EndPos && SplineParams.EndTangent == EndTangent)
+	{
+		return;
+	}
+
 	SplineParams.StartPos = StartPos;
 	SplineParams.StartTangent = StartTangent;
 	SplineParams.EndPos = EndPos;
@@ -293,6 +303,11 @@ FVector2D USplineMeshComponent::GetStartScale() const
 
 void USplineMeshComponent::SetStartScale(FVector2D StartScale, bool bUpdateMesh)
 {
+	if (SplineParams.StartScale == StartScale)
+	{
+		return;
+	}
+
 	SplineParams.StartScale = StartScale;
 	bMeshDirty = true;
 	if (bUpdateMesh)
@@ -338,6 +353,11 @@ FVector2D USplineMeshComponent::GetEndScale() const
 
 void USplineMeshComponent::SetEndScale(FVector2D EndScale, bool bUpdateMesh)
 {
+	if (SplineParams.EndScale == EndScale)
+	{
+		return;
+	}
+
 	SplineParams.EndScale = EndScale;
 	bMeshDirty = true;
 	if (bUpdateMesh)
@@ -383,6 +403,11 @@ ESplineMeshAxis::Type USplineMeshComponent::GetForwardAxis() const
 
 void USplineMeshComponent::SetForwardAxis(ESplineMeshAxis::Type InForwardAxis, bool bUpdateMesh)
 {
+	if (ForwardAxis == InForwardAxis)
+	{
+		return;
+	}
+
 	ForwardAxis = InForwardAxis;
 	bMeshDirty = true;
 	if (bUpdateMesh)
