Fixed Data Layers Loaded In Editor and Not Loaded in Editor flags not properly read from the EditorPerUserSettings config file.

diff --git a/Engine/Source/Runtime/Engine/Private/WorldPartition/DataLayer/WorldDataLayers.cpp b/Engine/Source/Runtime/Engine/Private/WorldPartition/DataLayer/WorldDataLayers.cpp
--- a/Engine/Source/Runtime/Engine/Private/WorldPartition/DataLayer/WorldDataLayers.cpp
+++ b/Engine/Source/Runtime/Engine/Private/WorldPartition/DataLayer/WorldDataLayers.cpp
@@ -790,18 +790,18 @@ void AWorldDataLayers::PostLoad()
 
 	// Initialize DataLayer's IsLoadedInEditor based on DataLayerEditorPerProjectUserSettings
 	TArray<FName> SettingsDataLayersNotLoadedInEditor = GetMutableDefault<UWorldPartitionEditorPerProjectUserSettings>()->GetWorldDataLayersNotLoadedInEditor(GetWorld());
-	for (const FName& DataLayerAssetName : SettingsDataLayersNotLoadedInEditor)
+	for (const FName& DataLayerName : SettingsDataLayersNotLoadedInEditor)
 	{
-		if (UDataLayerInstance* DataLayerInstance = const_cast<UDataLayerInstance*>(GetDataLayerInstanceFromAssetName(DataLayerAssetName)))
+		if (UDataLayerInstance* DataLayerInstance = const_cast<UDataLayerInstance*>(GetDataLayerInstance(DataLayerName)))
 		{
 			DataLayerInstance->SetIsLoadedInEditor(false, /*bFromUserChange*/false);
 		}
 	}
 
 	TArray<FName> SettingsDataLayersLoadedInEditor = GetMutableDefault<UWorldPartitionEditorPerProjectUserSettings>()->GetWorldDataLayersLoadedInEditor(GetWorld());
-	for (const FName& DataLayerAssetName : SettingsDataLayersLoadedInEditor)
+	for (const FName& DataLayerName : SettingsDataLayersLoadedInEditor)
 	{
-		if (UDataLayerInstance* DataLayerInstance = const_cast<UDataLayerInstance*>(GetDataLayerInstanceFromAssetName(DataLayerAssetName)))
+		if (UDataLayerInstance* DataLayerInstance = const_cast<UDataLayerInstance*>(GetDataLayerInstance(DataLayerName)))
 		{
 			DataLayerInstance->SetIsLoadedInEditor(true, /*bFromUserChange*/false);
 		}
