Texture: Added an option to serialize cooked textures using derived data references

This is opt-in and only uses derived data references for streaming mips of non-virtual textures in this revision. The state of derived data reference serialization means that this will only work correctly in a staged build.

#rnx

diff --git a/Engine/Source/Runtime/Engine/Private/Streaming/Texture2DStreamIn_DerivedData.cpp b/Engine/Source/Runtime/Engine/Private/Streaming/Texture2DStreamIn_DerivedData.cpp
new file mode 100644
--- /dev/null
+++ b/Engine/Source/Runtime/Engine/Private/Streaming/Texture2DStreamIn_DerivedData.cpp
@@ -0,0 +1,329 @@
+// Copyright Epic Games, Inc. All Rights Reserved.
+
+#include "Streaming/Texture2DStreamIn_DerivedData.h"
+
+#include "EngineLogs.h"
+#include "Streaming/TextureStreamingHelpers.h"
+
+namespace UE
+{
+
+///////////////////////////////////////////////////////////////////////////////////////////////////////////////////////
+
+FTexture2DStreamIn_DerivedData::FTexture2DStreamIn_DerivedData(UTexture2D* InTexture, bool bInHighPriority)
+	: FTexture2DStreamIn(InTexture)
+	, ResponseComplete(EEventMode::AutoReset)
+	, bHighPriority(bInHighPriority)
+{
+}
+
+void FTexture2DStreamIn_DerivedData::DoBeginIoRequests(const FContext& Context)
+{
+	TaskSynchronization.Increment();
+
+	FDerivedDataIoBatch Batch;
+
+	for (int32 MipIndex = PendingFirstLODIdx; MipIndex < CurrentFirstLODIdx && !IsCancelled(); ++MipIndex)
+	{
+		const FTexture2DMipMap& MipMap = *Context.MipsView[MipIndex];
+		check(MipData[MipIndex].Data);
+
+		if (MipMap.DerivedData.HasData())
+		{
+			uint32 ExpectedPitch = 0;
+			const uint32 MipSize = FTexture2DResource::CalculateTightPackedMipSize(MipMap.SizeX, MipMap.SizeY, Context.Resource->GetPixelFormat(), ExpectedPitch);
+
+			// @@!! MipData[] should have size
+			FDerivedDataIoOptions Options(MakeMemoryView(MipData[MipIndex].Data, MipSize));
+			Batch.Read(MipMap.DerivedData, Options);
+
+			// Pitch is ignored because this reads directly into MipData[].
+			// The Batch.Dispatch callback could be used to fix Pitch.
+			FTexture2DResource::WarnRequiresTightPackedMip(MipMap.SizeX, MipMap.SizeY, Context.Resource->GetPixelFormat(), MipData[MipIndex].Pitch);
+		}
+		else if (MipMap.BulkData.IsBulkDataLoaded())
+		{
+			UE_LOG(LogTexture, Error, TEXT("Attempting to stream data for non-streaming mip %d of %s."),
+				MipIndex, *Context.Texture->GetPathName());
+			MarkAsCancelled();
+		}
+		else
+		{
+			UE_LOG(LogTexture, Error, TEXT("Attempting to stream data that has not yet been generated for mip %d of %s."),
+				MipIndex, *Context.Texture->GetPathName());
+			MarkAsCancelled();
+		}
+	}
+
+	FDerivedDataIoPriority Priority = FDerivedDataIoPriority::Low().InterpolateTo(FDerivedDataIoPriority::Normal(), bHighPriority ? 0.5f : 0.0f);
+
+	Batch.Dispatch(Response, Priority, [this]
+	{
+		TaskSynchronization.Decrement();
+
+		if (Response.GetOverallStatus() != EDerivedDataIoStatus::Ok)
+		{
+			MarkAsCancelled();
+		}
+
+		if constexpr (!UE_BUILD_SHIPPING)
+		{
+			if (const int32 ExtraLatencyMs = FRenderAssetStreamingSettings::ExtraIOLatency; ExtraLatencyMs > 0)
+			{
+				// Simulate extra latency to allow enough time to cancel requests even when I/O is very fast.
+				FPlatformProcess::Sleep((CurrentFirstLODIdx - PendingFirstLODIdx) * ExtraLatencyMs * 0.001f);
+			}
+		}
+
+		ResponseComplete.Trigger();
+
+		// Tick to schedule the task or cancellation callback. Use TT_None to avoid a deadlock.
+		Tick(FTexture2DUpdate::TT_None);
+	});
+}
+
+bool FTexture2DStreamIn_DerivedData::DoPollIoRequests(const FContext& Context)
+{
+	return Response.Poll();
+}
+
+void FTexture2DStreamIn_DerivedData::DoEndIoRequests(const FContext& Context)
+{
+	ResponseComplete.Wait();
+}
+
+void FTexture2DStreamIn_DerivedData::DoCancelIoRequests()
+{
+	Response.Cancel();
+}
+
+void FTexture2DStreamIn_DerivedData::Abort()
+{
+	FTexture2DStreamIn::Abort();
+	DoCancelIoRequests();
+}
+
+///////////////////////////////////////////////////////////////////////////////////////////////////////////////////////
+
+FTexture2DStreamIn_DerivedData_AsyncCreate::FTexture2DStreamIn_DerivedData_AsyncCreate(UTexture2D* InTexture, bool bInHighPriority)
+	: FTexture2DStreamIn_DerivedData(InTexture, bInHighPriority)
+{
+	const FContext Context(InTexture, TT_None);
+	PushTask(Context, TT_Async, SRA_UPDATE_CALLBACK(AllocateMipsAndBeginIoRequests), TT_None, nullptr);
+}
+
+void FTexture2DStreamIn_DerivedData_AsyncCreate::AllocateMipsAndBeginIoRequests(const FContext& Context)
+{
+	DECLARE_SCOPE_CYCLE_COUNTER(TEXT("FTexture2DStreamIn_DerivedData_AsyncCreate::AllocateMipsAndBeginIoRequests"),
+		STAT_Texture2DStreamIn_DerivedData_AsyncCreate_AllocateMipsAndBeginIoRequests, STATGROUP_StreamingDetails);
+	check(Context.CurrentThread == TT_Async);
+
+	DoAllocateNewMips(Context);
+	DoBeginIoRequests(Context);
+
+	bDeferExecution = !IsCancelled(); // Set as pending.
+	PushTask(Context, TT_Async, SRA_UPDATE_CALLBACK(PollIoRequests), TT_Render, SRA_UPDATE_CALLBACK(Cancel));
+}
+
+void FTexture2DStreamIn_DerivedData_AsyncCreate::PollIoRequests(const FContext& Context)
+{
+	DECLARE_SCOPE_CYCLE_COUNTER(TEXT("FTexture2DStreamIn_DerivedData_AsyncCreate::PollIoRequests"),
+		STAT_Texture2DStreamIn_DerivedData_AsyncCreate_PollIoRequests, STATGROUP_StreamingDetails);
+	check(Context.CurrentThread == TT_Async);
+
+	if (!DoPollIoRequests(Context))
+	{
+		bDeferExecution = !IsCancelled(); // Set as pending.
+		PushTask(Context, TT_Async, SRA_UPDATE_CALLBACK(PollIoRequests), TT_Render, SRA_UPDATE_CALLBACK(Cancel));
+	}
+	else
+	{
+		PushTask(Context, TT_Async, SRA_UPDATE_CALLBACK(EndIoRequests), TT_Render, SRA_UPDATE_CALLBACK(Cancel));
+	}
+}
+
+void FTexture2DStreamIn_DerivedData_AsyncCreate::EndIoRequests(const FContext& Context)
+{
+	DECLARE_SCOPE_CYCLE_COUNTER(TEXT("FTexture2DStreamIn_DerivedData_AsyncCreate::EndIoRequests"),
+		STAT_Texture2DStreamIn_DerivedData_AsyncCreate_EndIoRequests, STATGROUP_StreamingDetails);
+	check(Context.CurrentThread == TT_Async);
+
+	DoEndIoRequests(Context);
+
+	PushTask(Context, TT_Async, SRA_UPDATE_CALLBACK(AsyncCreate), TT_Render, SRA_UPDATE_CALLBACK(Cancel));
+}
+
+void FTexture2DStreamIn_DerivedData_AsyncCreate::AsyncCreate(const FContext& Context)
+{
+	DECLARE_SCOPE_CYCLE_COUNTER(TEXT("FTexture2DStreamIn_DerivedData_AsyncCreate::AsyncCreate"),
+		STAT_Texture2DStreamIn_DerivedData_AsyncCreate_AsyncCreate, STATGROUP_StreamingDetails);
+	check(Context.CurrentThread == TT_Async);
+
+	DoAsyncCreateWithNewMips(Context);
+	DoFreeNewMips(Context);
+
+	PushTask(Context, TT_Render, SRA_UPDATE_CALLBACK(Finalize), TT_Render, SRA_UPDATE_CALLBACK(Cancel));
+}
+
+void FTexture2DStreamIn_DerivedData_AsyncCreate::Finalize(const FContext& Context)
+{
+	DECLARE_SCOPE_CYCLE_COUNTER(TEXT("FTexture2DStreamIn_DerivedData_AsyncCreate::Finalize"),
+		STAT_Texture2DStreamIn_DerivedData_AsyncCreate_Finalize, STATGROUP_StreamingDetails);
+	check(Context.CurrentThread == TT_Render);
+
+	DoCopySharedMips(Context);
+	DoFinishUpdate(Context);
+}
+
+void FTexture2DStreamIn_DerivedData_AsyncCreate::Cancel(const FContext& Context)
+{
+	DECLARE_SCOPE_CYCLE_COUNTER(TEXT("FTexture2DStreamIn_DerivedData_AsyncCreate::Cancel"),
+		STAT_Texture2DStreamIn_DerivedData_AsyncCreate_Cancel, STATGROUP_StreamingDetails);
+	check(Context.CurrentThread == TT_Render);
+
+	DoCancelIoRequests();
+	DoEndIoRequests(Context);
+
+	DoFreeNewMips(Context);
+	DoFinishUpdate(Context);
+}
+
+///////////////////////////////////////////////////////////////////////////////////////////////////////////////////////
+
+FTexture2DStreamIn_DerivedData_AsyncReallocate::FTexture2DStreamIn_DerivedData_AsyncReallocate(UTexture2D* InTexture, bool bInHighPriority)
+	: FTexture2DStreamIn_DerivedData(InTexture, bInHighPriority)
+{
+	const FContext Context(InTexture, TT_None);
+	PushTask(Context, TT_Render, SRA_UPDATE_CALLBACK(AsyncReallocate), TT_None, nullptr);
+}
+
+void FTexture2DStreamIn_DerivedData_AsyncReallocate::AsyncReallocate(const FContext& Context)
+{
+	DECLARE_SCOPE_CYCLE_COUNTER(TEXT("FTexture2DStreamIn_DerivedData_AsyncReallocate::AsyncReallocate"),
+		STAT_Texture2DStreamIn_DerivedData_AsyncReallocate_AsyncReallocate, STATGROUP_StreamingDetails);
+	check(Context.CurrentThread == TT_Render);
+
+	DoAsyncReallocate(Context);
+
+	PushTask(Context, TT_Render, SRA_UPDATE_CALLBACK(LockMips), TT_Render, SRA_UPDATE_CALLBACK(Cancel));
+}
+
+void FTexture2DStreamIn_DerivedData_AsyncReallocate::LockMips(const FContext& Context)
+{
+	DECLARE_SCOPE_CYCLE_COUNTER(TEXT("FTexture2DStreamIn_DerivedData_AsyncReallocate::LockMips"),
+		STAT_Texture2DStreamIn_DerivedData_AsyncReallocate_LockMips, STATGROUP_StreamingDetails);
+	check(Context.CurrentThread == TT_Render);
+
+	RHIFinalizeAsyncReallocateTexture2D(IntermediateTextureRHI, true);
+	DoLockNewMips(Context);
+
+	PushTask(Context, TT_Async, SRA_UPDATE_CALLBACK(LoadMips), TT_Render, SRA_UPDATE_CALLBACK(Cancel));
+}
+
+void FTexture2DStreamIn_DerivedData_AsyncReallocate::LoadMips(const FContext& Context)
+{
+	DECLARE_SCOPE_CYCLE_COUNTER(TEXT("FTexture2DStreamIn_DerivedData_AsyncReallocate::LoadMips"),
+		STAT_Texture2DStreamIn_DerivedData_AsyncReallocate_LoadMips, STATGROUP_StreamingDetails);
+	check(Context.CurrentThread == TT_Async);
+
+	DoBeginIoRequests(Context);
+	DoEndIoRequests(Context);
+
+	PushTask(Context, TT_Render, SRA_UPDATE_CALLBACK(Finalize), TT_Render, SRA_UPDATE_CALLBACK(Cancel));
+}
+
+void FTexture2DStreamIn_DerivedData_AsyncReallocate::Finalize(const FContext& Context)
+{
+	DECLARE_SCOPE_CYCLE_COUNTER(TEXT("FTexture2DStreamIn_DerivedData_AsyncReallocate::Finalize"),
+		STAT_Texture2DStreamIn_DerivedData_AsyncReallocate_Finalize, STATGROUP_StreamingDetails);
+	check(Context.CurrentThread == TT_Render);
+
+	DoUnlockNewMips(Context);
+	DoFinishUpdate(Context);
+}
+
+void FTexture2DStreamIn_DerivedData_AsyncReallocate::Cancel(const FContext& Context)
+{
+	DECLARE_SCOPE_CYCLE_COUNTER(TEXT("FTexture2DStreamIn_DerivedData_AsyncReallocate::Cancel"),
+		STAT_Texture2DStreamIn_DerivedData_AsyncReallocate_Cancel, STATGROUP_StreamingDetails);
+	check(Context.CurrentThread == TT_Render);
+
+	DoCancelIoRequests();
+	DoEndIoRequests(Context);
+
+	DoUnlockNewMips(Context);
+	DoFinishUpdate(Context);
+}
+
+///////////////////////////////////////////////////////////////////////////////////////////////////////////////////////
+
+FTexture2DStreamIn_DerivedData_Virtual::FTexture2DStreamIn_DerivedData_Virtual(UTexture2D* InTexture, bool bInHighPriority)
+	: FTexture2DStreamIn_DerivedData(InTexture, bInHighPriority)
+{
+	PushTask(FContext(InTexture, TT_None), TT_Render, SRA_UPDATE_CALLBACK(LockMips), TT_None, nullptr);
+}
+
+void FTexture2DStreamIn_DerivedData_Virtual::LockMips(const FContext& Context)
+{
+	DECLARE_SCOPE_CYCLE_COUNTER(TEXT("FTexture2DStreamIn_DerivedData_Virtual::LockMips"), STAT_Texture2DStreamIn_DerivedData_Virtual_LockMips, STATGROUP_StreamingDetails);
+	check(Context.CurrentThread == TT_Render);
+
+	DoConvertToVirtualWithNewMips(Context);
+	DoLockNewMips(Context);
+
+	PushTask(Context, TT_Async, SRA_UPDATE_CALLBACK(LoadMips), TT_Render, SRA_UPDATE_CALLBACK(Cancel));
+}
+
+void FTexture2DStreamIn_DerivedData_Virtual::LoadMips(const FContext& Context)
+{
+	DECLARE_SCOPE_CYCLE_COUNTER(TEXT("FTexture2DStreamIn_DerivedData_Virtual::LoadMips"), STAT_Texture2DStreamIn_DerivedData_Virtual_LoadMips, STATGROUP_StreamingDetails);
+	check(Context.CurrentThread == TT_Async);
+
+	DoBeginIoRequests(Context);
+
+	PushTask(Context, TT_Async, SRA_UPDATE_CALLBACK(PostLoadMips), TT_Render, SRA_UPDATE_CALLBACK(Cancel));
+}
+
+void FTexture2DStreamIn_DerivedData_Virtual::PostLoadMips(const FContext& Context)
+{
+	DECLARE_SCOPE_CYCLE_COUNTER(TEXT("FTexture2DStreamIn_DerivedData_Virtual::PostLoadMips"), STAT_Texture2DStreamIn_DerivedData_Virtual_LoadMips, STATGROUP_StreamingDetails);
+	check(Context.CurrentThread == TT_Async);
+
+	DoEndIoRequests(Context);
+
+	PushTask(Context, TT_Render, SRA_UPDATE_CALLBACK(Finalize), TT_Render, SRA_UPDATE_CALLBACK(Cancel));
+}
+
+
+void FTexture2DStreamIn_DerivedData_Virtual::Finalize(const FContext& Context)
+{
+	DECLARE_SCOPE_CYCLE_COUNTER(TEXT("FTexture2DStreamIn_DerivedData_Virtual::Finalize"), STAT_Texture2DStreamIn_DerivedData_Virtual_Finalize, STATGROUP_StreamingDetails);
+	check(Context.CurrentThread == TT_Render);
+
+	DoUnlockNewMips(Context);
+	if (IntermediateTextureRHI)
+	{
+		RHIVirtualTextureSetFirstMipVisible(IntermediateTextureRHI, PendingFirstLODIdx);
+	}
+	DoFinishUpdate(Context);
+}
+
+void FTexture2DStreamIn_DerivedData_Virtual::Cancel(const FContext& Context)
+{
+	DECLARE_SCOPE_CYCLE_COUNTER(TEXT("FTexture2DStreamIn_DerivedData_Virtual::Cancel"), STAT_Texture2DStreamIn_DerivedData_Virtual_Cancel, STATGROUP_StreamingDetails);
+	check(Context.CurrentThread == TT_Render);
+
+	DoCancelIoRequests();
+	DoEndIoRequests(Context);
+
+	DoUnlockNewMips(Context);
+	if (IntermediateTextureRHI)
+	{
+		RHIVirtualTextureSetFirstMipInMemory(IntermediateTextureRHI, CurrentFirstLODIdx);
+	}
+	DoFinishUpdate(Context);
+}
+
+///////////////////////////////////////////////////////////////////////////////////////////////////////////////////////
+
+} // UE
diff --git a/Engine/Source/Runtime/Engine/Private/Streaming/Texture2DStreamIn_DerivedData.h b/Engine/Source/Runtime/Engine/Private/Streaming/Texture2DStreamIn_DerivedData.h
new file mode 100644
--- /dev/null
+++ b/Engine/Source/Runtime/Engine/Private/Streaming/Texture2DStreamIn_DerivedData.h
@@ -0,0 +1,94 @@
+// Copyright Epic Games, Inc. All Rights Reserved.
+
+#pragma once
+
+#include "Experimental/Async/LazyEvent.h"
+#include "Serialization/DerivedData.h"
+#include "Texture2DStreamIn.h"
+
+namespace UE
+{
+
+class FTexture2DStreamIn_DerivedData : public FTexture2DStreamIn
+{
+protected:
+	FTexture2DStreamIn_DerivedData(UTexture2D* Texture, bool bHighPriority);
+
+	/** Begin the async requests to load the pending mips. */
+	void DoBeginIoRequests(const FContext& Context);
+
+	/** Poll the async requests and return true if they are complete. */
+	bool DoPollIoRequests(const FContext& Context);
+
+	/** Wait until the async requests are complete and clean them up. */
+	void DoEndIoRequests(const FContext& Context);
+
+	/** Cancel any active async requests. Does not wait! */
+	void DoCancelIoRequests();
+
+	/** Cancel any active async requests. Does not wait! */
+	void Abort() final;
+
+private:
+	FDerivedDataIoResponse Response;
+	FLazyEvent ResponseComplete;
+	bool bHighPriority;
+};
+
+class FTexture2DStreamIn_DerivedData_AsyncCreate final : public FTexture2DStreamIn_DerivedData
+{
+public:
+	FTexture2DStreamIn_DerivedData_AsyncCreate(UTexture2D* Texture, bool bHighPriority);
+
+protected:
+	// Allocate mips and begin the I/O request for each mip. (AsyncThread)
+	void AllocateMipsAndBeginIoRequests(const FContext& Context);
+	// Poll the I/O requests. (AsyncThread)
+	void PollIoRequests(const FContext& Context);
+	// Wait for the I/O requests to complete and clean them up. (AsyncThread)
+	void EndIoRequests(const FContext& Context);
+	// Async create an intermediate texture from the loaded mips. (AsyncThread)
+	void AsyncCreate(const FContext& Context);
+	// Apply the intermediate texture and clean up. (RenderThread)
+	void Finalize(const FContext& Context);
+	// Cancel the update. (RenderThread)
+	void Cancel(const FContext& Context);
+};
+
+class FTexture2DStreamIn_DerivedData_AsyncReallocate final : public FTexture2DStreamIn_DerivedData
+{
+public:
+	FTexture2DStreamIn_DerivedData_AsyncReallocate(UTexture2D* Texture, bool bHighPriority);
+
+protected:
+	// Create an intermediate bigger copy of the texture. (RenderThread)
+	void AsyncReallocate(const FContext& Context);
+	// Lock each new mips of the intermediate texture. (RenderThread)
+	void LockMips(const FContext& Context);
+	// Create load requests into each locked mips. (AsyncThread)
+	void LoadMips(const FContext& Context);
+	// Unlock the mips, apply the intermediate texture and clean up. (RenderThread)
+	void Finalize(const FContext& Context);
+	// Unlock the mips, cancel load requests and clean up. (RenderThread)
+	void Cancel(const FContext& Context);
+};
+
+class FTexture2DStreamIn_DerivedData_Virtual final : public FTexture2DStreamIn_DerivedData
+{
+public:
+	FTexture2DStreamIn_DerivedData_Virtual(UTexture2D* Texture, bool bHighPriority);
+
+protected:
+	// Validate the filename and bulk data, then lock new mips. (RenderThread)
+	void LockMips(const FContext& Context);
+	// Create load requests into each locked mips. (AsyncThread)
+	void LoadMips(const FContext& Context);
+	// Clear the IO requests and call Finalize. (AsyncThread)
+	void PostLoadMips(const FContext& Context);
+	// Apply the intermediate texture and cleanup. (RenderThread)
+	void Finalize(const FContext& Context);
+	// Unlock the mips, cancel load requests and clean up. (RenderThread)
+	void Cancel(const FContext& Context);
+};
+
+} // UE
diff --git a/Engine/Source/Runtime/Engine/Private/Texture2D.cpp b/Engine/Source/Runtime/Engine/Private/Texture2D.cpp
--- a/Engine/Source/Runtime/Engine/Private/Texture2D.cpp
+++ b/Engine/Source/Runtime/Engine/Private/Texture2D.cpp
@@ -6,6 +6,7 @@
 
 #include "Engine/Texture2D.h"
 
+#include "Algo/AnyOf.h"
 #include "Misc/App.h"
 #include "HAL/PlatformFileManager.h"
 #include "HAL/FileManager.h"
@@ -30,6 +31,7 @@
 #include "Streaming/Texture2DStreamOut_Virtual.h"
 #include "Streaming/Texture2DStreamIn_DDC_AsyncCreate.h"
 #include "Streaming/Texture2DStreamIn_DDC_AsyncReallocate.h"
+#include "Streaming/Texture2DStreamIn_DerivedData.h"
 #include "Streaming/Texture2DStreamIn_IO_AsyncCreate.h"
 #include "Streaming/Texture2DStreamIn_IO_AsyncReallocate.h"
 #include "Streaming/Texture2DStreamIn_IO_Virtual.h"
@@ -1629,6 +1631,16 @@ bool UTexture2D::StreamIn(int32 NewMipCount, bool bHighPrio)
 
 		if (!CustomMipDataProvider && GUseGenericStreamingPath != 1)
 		{
+			const auto HasDerivedData = [](UTexture2D* Texture) -> bool
+			{
+				if (FTexturePlatformData* LocalPlatformData = Texture->GetPlatformData())
+				{
+					return Algo::AnyOf(LocalPlatformData->Mips, &FTexture2DMipMap::DerivedData) ||
+						(LocalPlatformData->VTData && Algo::AnyOf(LocalPlatformData->VTData->Chunks, &FVirtualTextureDataChunk::DerivedData));
+				}
+				return false;
+			};
+
 	#if WITH_EDITORONLY_DATA
 			if (FPlatformProperties::HasEditorOnlyData() && !GetOutermost()->bIsCookedForEditor && !GetOutermost()->HasAnyPackageFlags(PKG_Cooked))
 			{
@@ -1643,6 +1655,25 @@ bool UTexture2D::StreamIn(int32 NewMipCount, bool bHighPrio)
 			}
 			else
 	#endif
+			if (HasDerivedData(this))
+			{
+				// If the future texture is to be a virtual texture, use the virtual stream in path.
+				if (Texture2DResource->bUsePartiallyResidentMips)
+				{
+					PendingUpdate = new UE::FTexture2DStreamIn_DerivedData_Virtual(this, bHighPrio);
+				}
+				// If the platform supports creating the new texture on an async thread, use that path.
+				else if (GRHISupportsAsyncTextureCreation)
+				{
+					PendingUpdate = new UE::FTexture2DStreamIn_DerivedData_AsyncCreate(this, bHighPrio);
+				}
+				// Otherwise use the default path.
+				else
+				{
+					PendingUpdate = new UE::FTexture2DStreamIn_DerivedData_AsyncReallocate(this, bHighPrio);
+				}
+			}
+			else
 			{
 				// If the future texture is to be a virtual texture, use the virtual stream in path.
 				if (Texture2DResource->bUsePartiallyResidentMips)
diff --git a/Engine/Source/Runtime/Engine/Private/TextureDerivedData.cpp b/Engine/Source/Runtime/Engine/Private/TextureDerivedData.cpp
--- a/Engine/Source/Runtime/Engine/Private/TextureDerivedData.cpp
+++ b/Engine/Source/Runtime/Engine/Private/TextureDerivedData.cpp
@@ -41,6 +41,7 @@
 #include "DerivedDataCache.h"
 #include "DerivedDataCacheInterface.h"
 #include "DerivedDataRequestOwner.h"
+#include "HAL/IConsoleManager.h"
 #include "Interfaces/ITargetPlatform.h"
 #include "Interfaces/ITargetPlatformManagerModule.h"
 #include "Interfaces/ITextureFormat.h"
@@ -53,6 +54,11 @@
 #include "TextureCompiler.h"
 #include "TextureEncodingSettings.h"
 
+static TAutoConsoleVariable<int32> CVarTexturesCookToDerivedDataReferences(
+	TEXT("r.TexturesCookToDerivedDataReferences"),
+	0,
+	TEXT("Whether cooked textures are serialized using Derived Data References."),
+	ECVF_ReadOnly);
 
 /*------------------------------------------------------------------------------
 	Versioning for texture derived data.
@@ -2095,7 +2101,7 @@ int32 FTexturePlatformData::GetNumNonStreamingMips(bool bIsStreamingPossible) co
 
 		for (const FTexture2DMipMap& Mip : Mips)
 		{
-			if ( Mip.BulkData.IsInSeparateFile() || !Mip.BulkData.IsInlined() )
+			if (Mip.DerivedData || Mip.BulkData.IsInSeparateFile() || !Mip.BulkData.IsInlined())
 			{
 				--NumNonStreamingMips;
 			}
@@ -2164,7 +2170,7 @@ int32 FTexturePlatformData::GetNumNonOptionalMips() const
 
 		for (const FTexture2DMipMap& Mip : Mips)
 		{
-			if (Mip.BulkData.IsOptional())
+			if ((Mip.DerivedData && !EnumHasAnyFlags(Mip.DerivedData.GetFlags(), UE::EDerivedDataFlags::Required)) || Mip.BulkData.IsOptional())
 			{
 				--NumNonOptionalMips;
 			}
@@ -2193,7 +2199,7 @@ bool FTexturePlatformData::CanBeLoaded() const
 {
 	for (const FTexture2DMipMap& Mip : Mips)
 	{
-		if (Mip.DerivedData.HasData())
+		if (Mip.DerivedData)
 		{
 			return true;
 		}
@@ -2256,7 +2262,7 @@ int64 FTexturePlatformData::GetPayloadSize(int32 MipBias) const
 bool FTexturePlatformData::CanUseCookedDataPath() const
 {
 #if WITH_IOSTORE_IN_EDITOR
-	return Mips.Num() > 0 && Mips[0].BulkData.IsUsingIODispatcher();
+	return Mips.Num() > 0 && (Mips[0].BulkData.IsUsingIODispatcher() || Mips[0].DerivedData.IsCooked());
 #else	
 	return FPlatformProperties::RequiresCookedData();
 #endif //WITH_IOSTORE_IN_EDITOR
@@ -2382,245 +2388,405 @@ static void SerializePlatformData(
 	EPlatformDataSerializationFlags Flags
 )
 {
-	DECLARE_SCOPE_CYCLE_COUNTER(TEXT("SerializePlatformData"), STAT_Texture_SerializePlatformData, STATGROUP_LoadTime);
-
 	// note: if BuildTexture failed, we still get called here,
 	//	just with a default-constructed PlatformData
 	//	(no mips, sizes=0, PF=Unknown)
 
-	if (Ar.IsFilterEditorOnly())
-	{
-		constexpr int64 PlaceholderDerivedDataSize = 16;
-		uint8 PlaceholderDerivedData[PlaceholderDerivedDataSize]{};
-		Ar.Serialize(PlaceholderDerivedData, PlaceholderDerivedDataSize);
-		check(Algo::AllOf(PlaceholderDerivedData, [](uint8 Value) { return Value == 0; }));
-	}
+	DECLARE_SCOPE_CYCLE_COUNTER(TEXT("SerializePlatformData"), STAT_Texture_SerializePlatformData, STATGROUP_LoadTime);
+
+	UEnum* PixelFormatEnum = UTexture::GetPixelFormatEnum();
 
 	const bool bCooked = (Flags & EPlatformDataSerializationFlags::Cooked) == EPlatformDataSerializationFlags::Cooked;
 	const bool bStreamable = (Flags & EPlatformDataSerializationFlags::Streamable) == EPlatformDataSerializationFlags::Streamable;
 
-	UEnum* PixelFormatEnum = UTexture::GetPixelFormatEnum();
+	bool bIsVirtual = Ar.IsSaving() && PlatformData->VTData;
+	int32 NumMips = PlatformData->Mips.Num();
+	int32 FirstMipToSerialize = 0;
+	int32 FirstInlineMip = 0;
+	// TODO: Do we need to consider platforms saving texture assets as cooked files?
+	//       The info used to calculate optional mips is part of the editor only data.
+	int32 OptionalMips = 0;
+	bool bDuplicateNonOptionalMips = false;
 
-	Ar << PlatformData->SizeX;
-	Ar << PlatformData->SizeY;
-	Ar << PlatformData->PackedData;
-	if (Ar.IsLoading())
+	if (bCooked && bIsVirtual)
 	{
-		FString PixelFormatString;
-		Ar << PixelFormatString;
-		const int64 PixelFormatValue = PixelFormatEnum->GetValueByName(*PixelFormatString);
-		if (PixelFormatValue != INDEX_NONE && PixelFormatValue < PF_MAX)
+		check(NumMips == 0);
+	}
+
+#if WITH_EDITORONLY_DATA
+	if (bCooked && Ar.IsSaving())
+	{
+		check(Texture);
+		check(Ar.CookingTarget());
+
+		const int32 Width = PlatformData->SizeX;
+		const int32 Height = PlatformData->SizeY;
+		const int32 LODGroup = Texture->LODGroup;
+		const int32 LODBias = Texture->LODBias;
+		const TextureMipGenSettings MipGenSetting = Texture->MipGenSettings;
+		const int32 LastMip = FMath::Max(NumMips - 1, 0);
+		const int32 FirstMipTailMip = NumMips - (int32)PlatformData->GetNumMipsInTail();
+		check(FirstMipTailMip >= 0);
+
+		FirstMipToSerialize = Ar.CookingTarget()->GetTextureLODSettings().CalculateLODBias(Width, Height, Texture->MaxTextureSize, LODGroup, LODBias, 0, MipGenSetting, bIsVirtual);
+		if (!bIsVirtual)
 		{
-			PlatformData->PixelFormat = (EPixelFormat)PixelFormatValue;
+			// Reassign NumMips as the number of mips starting from FirstMipToSerialize.
+			FirstMipToSerialize = FMath::Clamp(FirstMipToSerialize, 0, PlatformData->GetNumMipsInTail() > 0 ? FirstMipTailMip : LastMip);
+			NumMips = FMath::Max(0, NumMips - FirstMipToSerialize);
 		}
 		else
 		{
-			UE_LOG(LogTexture, Warning, TEXT("Invalid pixel format '%s' for texture '%s'."), *PixelFormatString, Texture ? *Texture->GetPathName() : TEXT(""));
-			PlatformData->PixelFormat = PF_Unknown;
+			FirstMipToSerialize = FMath::Clamp(FirstMipToSerialize, 0, FMath::Max((int32)PlatformData->VTData->GetNumMips() - 1, 0));
 		}
 	}
-	else if (Ar.IsSaving())
-	{
-		FString PixelFormatString = PixelFormatEnum->GetNameByValue(PlatformData->PixelFormat).GetPlainNameString();
-		Ar << PixelFormatString;
-	}
+#endif // #if WITH_EDITORONLY_DATA
 
-	if (PlatformData->GetHasOptData())
-	{
-		Ar << PlatformData->OptData;
-	}
+	// Force resident mips inline
+	if (bCooked && Ar.IsSaving() && !bIsVirtual)
+	{
+		// bStreamable comes from IsCandidateForTextureStreaming
+		//  if not bStreamable, all mips are written inline
+		//  so the runtime will see NumNonStreamingMips = all
+
+	#if WITH_EDITORONLY_DATA
+		check(Ar.CookingTarget());
+		// This also needs to check whether the project enables texture streaming.
+		// Currently, there is no reliable way to implement this because there is no difference
+		// between the project settings (CVar) and the command line setting (from -NoTextureStreaming)
+		if (bStreamable && Ar.CookingTarget()->SupportsFeature(ETargetPlatformFeatures::TextureStreaming))
+	#else
+		if (bStreamable)
+	#endif
+		{
+			check(Texture->IsPossibleToStream());
 
-	int32 NumMips = PlatformData->Mips.Num();
-	int32 FirstMipToSerialize = 0;
+			int32 NumNonStreamingMips = PlatformData->GetNumNonStreamingMips(/*bIsStreamingPossible*/ true);
+			// NumMips has been reduced by FirstMipToSerialize (LODBias)
+			NumNonStreamingMips = FMath::Min(NumNonStreamingMips, NumMips);
+			// NumNonStreamingMips is not serialized. The runtime will use NumNonStreamingMips = NumInlineMips.
+			FirstInlineMip = NumMips - NumNonStreamingMips;
 
-	bool bIsVirtual = false;
-	if (Ar.IsSaving())
-	{
-		bIsVirtual = PlatformData->VTData != nullptr;
-	}
+		#if WITH_EDITORONLY_DATA
+			const int32 LODGroup = Texture->LODGroup;
 
-	if (bCooked && bIsVirtual)
-	{
-		check(PlatformData->Mips.Num() == 0);
+			const int32 FirstMipWidth  = PlatformData->Mips[FirstMipToSerialize].SizeX;
+			const int32 FirstMipHeight = PlatformData->Mips[FirstMipToSerialize].SizeY;
+			OptionalMips = Ar.CookingTarget()->GetTextureLODSettings().CalculateNumOptionalMips(LODGroup, FirstMipWidth, FirstMipHeight, NumMips, FirstInlineMip, Texture->MipGenSettings);
+			bDuplicateNonOptionalMips = Ar.CookingTarget()->GetTextureLODSettings().TextureLODGroups[LODGroup].DuplicateNonOptionalMips;
+
+			// OptionalMips must be streaming mips.
+			check(OptionalMips <= FirstInlineMip);
+		#endif
+
+		#if WITH_EDITOR
+			// Record the use of streaming mips on the owner.
+			if (NumNonStreamingMips < NumMips)
+			{
+				// Use FindChecked because this was previously added and set to false.
+				const FString PlatformName = Ar.CookingTarget()->PlatformName();
+				Texture->DidSerializeStreamingMipsForPlatform.FindChecked(PlatformName) = true;
+			}
+		#endif
+		}
 	}
 
-	if (bCooked)
-	{
 #if WITH_EDITORONLY_DATA
-		if (Ar.IsSaving())
-		{
-			check(Ar.CookingTarget());
-			check(Texture);
-
-			const int32 Width = PlatformData->SizeX;
-			const int32 Height = PlatformData->SizeY;
-			const int32 LODGroup = Texture->LODGroup;
-			const int32 LODBias = Texture->LODBias;
-			const TextureMipGenSettings MipGenSetting = Texture->MipGenSettings;
-			const int32 LastMip = FMath::Max(NumMips - 1, 0);
-			check(NumMips >= (int32)PlatformData->GetNumMipsInTail());
-			const int32 FirstMipTailMip = NumMips - (int32)PlatformData->GetNumMipsInTail();
+	// Save cook tags
+	if (bCooked && Ar.IsSaving() && Ar.GetCookContext() && Ar.GetCookContext()->GetCookTagList())
+	{
+		FCookTagList* CookTags = Ar.GetCookContext()->GetCookTagList();
 
-			FirstMipToSerialize = Ar.CookingTarget()->GetTextureLODSettings().CalculateLODBias(Width, Height, Texture->MaxTextureSize, LODGroup, LODBias, 0, MipGenSetting, bIsVirtual);
-			if (!bIsVirtual)
+		if (bIsVirtual)
+		{
+			FVirtualTextureBuiltData* VTData = PlatformData->VTData;
+			CookTags->Add(Texture, "Size", FString::Printf(TEXT("%dx%d"), VTData->Width, VTData->Height));
+		}
+		else if (PlatformData->Mips.Num() > 0) // PlatformData->Mips is empty if BuildTexture failed
+		{
+			FString DimensionsStr;
+			FTexture2DMipMap& TopMip = PlatformData->Mips[FirstMipToSerialize];
+			if (TopMip.SizeZ != 1)
 			{
-				// NumMips is the number of mips starting from FirstMipToSerialize
-				FirstMipToSerialize = FMath::Clamp(FirstMipToSerialize, 0, PlatformData->GetNumMipsInTail() > 0 ? FirstMipTailMip : LastMip);
-				NumMips = FMath::Max(0, NumMips - FirstMipToSerialize);
+				DimensionsStr = FString::Printf(TEXT("%dx%dx%d"), TopMip.SizeX, TopMip.SizeY, TopMip.SizeZ);
 			}
 			else
 			{
-				FirstMipToSerialize = FMath::Clamp(FirstMipToSerialize, 0, FMath::Max((int32)PlatformData->VTData->GetNumMips() - 1, 0));
+				DimensionsStr = FString::Printf(TEXT("%dx%d"), TopMip.SizeX, TopMip.SizeY);
 			}
+			CookTags->Add(Texture, "Size", MoveTemp(DimensionsStr));
 		}
-#endif // #if WITH_EDITORONLY_DATA
-		Ar << FirstMipToSerialize;
-		if (Ar.IsLoading())
-		{
-			check(Texture);
-			FirstMipToSerialize = 0;
-		}
-	}
 
-	TArray<uint32> BulkDataMipFlags;
+		CookTags->Add(Texture, "Format", FString(GPixelFormats[PlatformData->PixelFormat].Name));
 
-	// Force resident mips inline
-	if (bCooked && Ar.IsSaving())
-	{
-		if (bIsVirtual == false)
+		// Add in diff keys for change detection/blame.
 		{
-			BulkDataMipFlags.AddZeroed(PlatformData->Mips.Num());
-			for (int32 MipIndex = 0; MipIndex < PlatformData->Mips.Num(); ++MipIndex)
+			// Did the source change?
+			CookTags->Add(Texture, "Diff_10_Tex2D_Source", Texture->Source.GetIdString());
+
+			// Did the settings change?
+			if (const UE::DerivedData::FCacheKeyProxy* CacheKey = PlatformData->DerivedDataKey.TryGet<UE::DerivedData::FCacheKeyProxy>())
+			{
+				CookTags->Add(Texture, "Diff_20_Tex2D_CacheKey", *WriteToString<64>(*CacheKey->AsCacheKey()));
+			}
+			else if (const FString* DDK = PlatformData->DerivedDataKey.TryGet<FString>())
 			{
-				BulkDataMipFlags[MipIndex] = PlatformData->Mips[MipIndex].BulkData.GetBulkDataFlags();
+				CookTags->Add(Texture, "Diff_20_Tex2D_DDK", FString(*DDK));
 			}
+		}
+	}
+#endif
 
-			int32 MinMipToInline = 0;
-			int32 OptionalMips = 0; // TODO: do we need to consider platforms saving texture assets as cooked files? all the info to calculate the optional is part of the editor only data
-			bool DuplicateNonOptionalMips = false;
-		
-			// bStreamable comes from IsCandidateForTextureStreaming
-			//  if not bStreamable, all mips are written inline
-			//  so the runtime will see NumNonStreamingMips = all
+	// DO NOT SERIALIZE ANYTHING BEFORE THIS POINT IN THE FUNCTION!
 
-#if WITH_EDITORONLY_DATA
-			check(Ar.CookingTarget());
-			// This also needs to check whether the project enables texture streaming.
-			// Currently, there is no reliable way to implement this because there is no difference
-			// between the project settings (CVar) and the command line setting (from -NoTextureStreaming)
-			if (bStreamable && Ar.CookingTarget()->SupportsFeature(ETargetPlatformFeatures::TextureStreaming))
-#else
-			if (bStreamable)
-#endif
-			{
-				check( Texture->IsPossibleToStream() );
+	// The DerivedData and BulkData serialization paths are expected to be distinct.
+	// Since 5.0, cooked textures using the BulkData serialization path have a 16-byte zero block
+	// that acts as a placeholder for the DerivedData serialization path to be optionally enabled
+	// without requiring unaffected textures to be patched.
 
-				int32 NumNonStreamingMips = PlatformData->GetNumNonStreamingMips(true);
-				// NumMips has been reduced by FirstMipToSerialize (LODBias)
-				NumNonStreamingMips = FMath::Min(NumNonStreamingMips,NumMips);
-				// NumNonStreamingMips is not serialized
-				// the runtime will just use NumNonStreamingMips == NumInlineMips
-				MinMipToInline = NumMips - NumNonStreamingMips;
-#if WITH_EDITORONLY_DATA
-				const int32 LODGroup = Texture->LODGroup;
+	bool bUsingDerivedData = false;
+	if (bCooked)
+	{
+		bUsingDerivedData = !bIsVirtual && Ar.IsSaving() && Ar.IsFilterEditorOnly();
+	#if WITH_EDITOR
+		bUsingDerivedData &= CVarTexturesCookToDerivedDataReferences.GetValueOnAnyThread() != 0;
+	#endif
+		Ar.Serialize(&bUsingDerivedData, sizeof(bool));
+		static_assert(sizeof(bool) == 1);
+	}
 
-				// was a bug ? CalculateNumOptional mips is using full Width/Height ?  should be from FirstMipToSerialize? -> moot, it's not actually used
-				const int32 FirstMipWidth  = PlatformData->Mips[FirstMipToSerialize].SizeX;
-				const int32 FirstMipHeight = PlatformData->Mips[FirstMipToSerialize].SizeY;
+	///////////////////////////////////////////////////////////////////////////////////////////////////////////////////
+	// DERIVED DATA REFERENCE FORMAT BEGINS HERE
 
-				OptionalMips = Ar.CookingTarget()->GetTextureLODSettings().CalculateNumOptionalMips(LODGroup, FirstMipWidth, FirstMipHeight, NumMips, MinMipToInline, Texture->MipGenSettings);
-				DuplicateNonOptionalMips = Ar.CookingTarget()->GetTextureLODSettings().TextureLODGroups[LODGroup].DuplicateNonOptionalMips;
+	if (bUsingDerivedData)
+	{
+		int32 MipSizeX = NumMips ? PlatformData->Mips[FirstMipToSerialize].SizeX : 0;
+		int32 MipSizeY = NumMips ? PlatformData->Mips[FirstMipToSerialize].SizeY : 0;
+		int32 MipSizeZ = NumMips ? PlatformData->Mips[FirstMipToSerialize].SizeZ : 0;
 
-				// Optional mips must not overlap the non-streaming mips : (MinMipToInline ensures this)
-				check( OptionalMips + NumNonStreamingMips <= NumMips );
-#endif
-			
-#if WITH_EDITOR
-				// if we make any streaming mips, record it on the owner :
-				if ( NumNonStreamingMips < NumMips )
-				{
-					// was added previously and set to false, so FindChecked here :
-					const FString PlatformName = Ar.CookingTarget()->PlatformName();
-					Texture->DidSerializeStreamingMipsForPlatform.FindChecked( PlatformName ) = true;
-				}
-#endif
-			}
+		Ar << bIsVirtual;
 
-			for (int32 MipIndex = 0; MipIndex < NumMips && MipIndex < OptionalMips; ++MipIndex) //-V654
+		// Serialize SizeX, SizeY
+		if (bIsVirtual)
+		{
+			checkNoEntry();
+		}
+		else
+		{
+			Ar << MipSizeX;
+			Ar << MipSizeY;
+			Ar << MipSizeZ;
+
+			if (Ar.IsLoading())
 			{
-				// optional (and streamed) mips
-				PlatformData->Mips[MipIndex + FirstMipToSerialize].BulkData.SetBulkDataFlags(BULKDATA_Force_NOT_InlinePayload | BULKDATA_OptionalPayload);
+				PlatformData->SizeX = MipSizeX;
+				PlatformData->SizeY = MipSizeY;
 			}
+		}
 
-			const uint32 AdditionalNonOptionalBulkDataFlags = DuplicateNonOptionalMips ? BULKDATA_DuplicateNonOptionalPayload : 0;
-			for (int32 MipIndex = OptionalMips; MipIndex < NumMips && MipIndex < MinMipToInline; ++MipIndex)
+		// Serialize PackedData, OptData
+		Ar << PlatformData->PackedData;
+		if (PlatformData->GetHasOptData())
+		{
+			Ar << PlatformData->OptData;
+		}
+
+		// Serialize PixelFormat
+		if (Ar.IsSaving())
+		{
+			FString PixelFormatString = PixelFormatEnum->GetNameByValue(PlatformData->PixelFormat).GetPlainNameString();
+			Ar << PixelFormatString;
+		}
+		else
+		{
+			FString PixelFormatString;
+			Ar << PixelFormatString;
+			const int64 PixelFormatValue = PixelFormatEnum->GetValueByName(*PixelFormatString);
+			if (PixelFormatValue != INDEX_NONE && PixelFormatValue < PF_MAX)
 			{
-				// non-optional but streamed mips
-				PlatformData->Mips[MipIndex + FirstMipToSerialize].BulkData.SetBulkDataFlags(BULKDATA_Force_NOT_InlinePayload | AdditionalNonOptionalBulkDataFlags);
+				PlatformData->PixelFormat = (EPixelFormat)PixelFormatValue;
 			}
-			for (int32 MipIndex = MinMipToInline; MipIndex < NumMips; ++MipIndex)
+			else
 			{
-				// non-streamed (inline) mips
-				PlatformData->Mips[MipIndex + FirstMipToSerialize].BulkData.SetBulkDataFlags(BULKDATA_ForceInlinePayload | BULKDATA_SingleUse);
+				UE_LOG(LogTexture, Warning, TEXT("Invalid pixel format '%s' for texture '%s'."), *PixelFormatString, Texture ? *Texture->GetPathName() : TEXT(""));
+				PlatformData->PixelFormat = PF_Unknown;
 			}
 		}
-		else // bVirtual == true
+
+		// Serialize DerivedData
+		if (bIsVirtual)
 		{
-			const int32 NumChunks = PlatformData->VTData->Chunks.Num();
-			BulkDataMipFlags.AddZeroed(NumChunks);
-			for (int32 ChunkIndex = 0; ChunkIndex < NumChunks; ++ChunkIndex)
-			{
-				BulkDataMipFlags[ChunkIndex] = PlatformData->VTData->Chunks[ChunkIndex].BulkData.GetBulkDataFlags();
-				PlatformData->VTData->Chunks[ChunkIndex].BulkData.SetBulkDataFlags(BULKDATA_Force_NOT_InlinePayload);
-			}	
+			checkNoEntry();
 		}
-
-		// Save cook tags
-#if WITH_EDITORONLY_DATA
-		if (Ar.GetCookContext() && Ar.GetCookContext()->GetCookTagList())
+		else
 		{
-			FCookTagList* CookTags = Ar.GetCookContext()->GetCookTagList();
+			Ar << NumMips;
+			check(NumMips >= (int32)PlatformData->GetNumMipsInTail());
+
+			Ar << FirstInlineMip;
+			check(FirstInlineMip >= 0 && FirstInlineMip <= NumMips);
 
-			if (bIsVirtual)
+			if (Ar.IsLoading())
 			{
-				FVirtualTextureBuiltData* VTData = PlatformData->VTData;
-				CookTags->Add(Texture, "Size", FString::Printf(TEXT("%dx%d"), VTData->Width, VTData->Height));
-	}
-			else if ( PlatformData->Mips.Num() > 0 ) // PlatformData->Mips is empty if BuildTexture failed
+				PlatformData->Mips.Empty(NumMips);
+				for (int32 MipIndex = 0; MipIndex < NumMips; ++MipIndex)
+				{
+					PlatformData->Mips.Add(new FTexture2DMipMap());
+					PlatformData->Mips.Last().BulkData.RemoveBulkData();
+				}
+			}
+
+			for (int32 MipIndex = 0; MipIndex < NumMips; ++MipIndex)
 			{
-				FString DimensionsStr;
-				FTexture2DMipMap& TopMip = PlatformData->Mips[FirstMipToSerialize];
-				if (TopMip.SizeZ != 1)
+				FTexture2DMipMap& Mip = PlatformData->Mips[FirstMipToSerialize + MipIndex];
+				if (Ar.IsSaving())
 				{
-					DimensionsStr = FString::Printf(TEXT("%dx%dx%d"), TopMip.SizeX, TopMip.SizeY, TopMip.SizeZ);
+					if (Mip.SizeZ > 1 || MipSizeZ > 1)
+					{
+						checkf(Mip.SizeX == MipSizeX && Mip.SizeY == MipSizeY && Mip.SizeZ == MipSizeZ,
+							TEXT("Expected %dx%dx%d mip and had %dx%dx%d mip for '%s'"),
+							Mip.SizeX, Mip.SizeY, Mip.SizeZ, MipSizeX, MipSizeY, MipSizeZ, *Texture->GetPathName());
+					}
+					else
+					{
+						checkf(Mip.SizeX == MipSizeX && Mip.SizeY == MipSizeY,
+							TEXT("Expected %dx%d mip and had %dx%d mip for '%s'"),
+							Mip.SizeX, Mip.SizeY, MipSizeX, MipSizeY, *Texture->GetPathName());
+					}
 				}
 				else
 				{
-					DimensionsStr = FString::Printf(TEXT("%dx%d"), TopMip.SizeX, TopMip.SizeY);
+					Mip.SizeX = MipSizeX;
+					Mip.SizeY = MipSizeY;
+					Mip.SizeZ = MipSizeZ;
 				}
-				CookTags->Add(Texture, "Size", MoveTemp(DimensionsStr));	
+				MipSizeX = FMath::Max(MipSizeX / 2, 1);
+				MipSizeY = FMath::Max(MipSizeY / 2, 1);
 			}
 
-			CookTags->Add(Texture, "Format", FString(GPixelFormats[PlatformData->PixelFormat].Name));
-
-			// Add in diff keys for change detection/blame.
+			for (int32 MipIndex = 0; MipIndex < FirstInlineMip; ++MipIndex)
 			{
-				// Did the source change?
-				CookTags->Add(Texture, "Diff_10_Tex2D_Source", Texture->Source.GetIdString());
+				FTexture2DMipMap& Mip = PlatformData->Mips[FirstMipToSerialize + MipIndex];
+				Mip.DerivedData.Serialize(Ar, Texture);
+				check(Mip.DerivedData);
+			}
 
-				// Did the settings change?
-				if (const UE::DerivedData::FCacheKeyProxy* CacheKey = PlatformData->DerivedDataKey.TryGet<UE::DerivedData::FCacheKeyProxy>())
+			for (int32 MipIndex = FirstInlineMip; MipIndex < NumMips; ++MipIndex)
+			{
+				FTexture2DMipMap& Mip = PlatformData->Mips[FirstMipToSerialize + MipIndex];
+				if (Ar.IsSaving())
 				{
-					CookTags->Add(Texture, "Diff_20_Tex2D_CacheKey", *WriteToString<64>(*CacheKey->AsCacheKey()));
+					const uint32 OriginalBulkDataFlags = Mip.BulkData.GetBulkDataFlags();
+					Mip.BulkData.SetBulkDataFlags(BULKDATA_ForceInlinePayload | BULKDATA_SingleUse);
+					Mip.BulkData.Serialize(Ar, Texture, MipIndex);
+					Mip.BulkData.ResetBulkDataFlags(OriginalBulkDataFlags);
 				}
-				else if (const FString* DDK = PlatformData->DerivedDataKey.TryGet<FString>())
+				else
 				{
-					CookTags->Add(Texture, "Diff_20_Tex2D_DDK", FString(*DDK));
+					Mip.BulkData.Serialize(Ar, Texture, MipIndex);
 				}
 			}
 		}
-#endif
+
+		return;
+	}
+
+	///////////////////////////////////////////////////////////////////////////////////////////////////////////////////
+	// BULK DATA FORMAT BEGINS HERE
+
+	if (bCooked)
+	{
+		constexpr int64 PlaceholderDerivedDataSize = 15;
+		uint8 PlaceholderDerivedData[PlaceholderDerivedDataSize]{};
+		Ar.Serialize(PlaceholderDerivedData, PlaceholderDerivedDataSize);
+		check(Algo::AllOf(PlaceholderDerivedData, [](uint8 Value) { return Value == 0; }));
+	}
+
+	Ar << PlatformData->SizeX;
+	Ar << PlatformData->SizeY;
+	Ar << PlatformData->PackedData;
+	if (Ar.IsLoading())
+	{
+		FString PixelFormatString;
+		Ar << PixelFormatString;
+		const int64 PixelFormatValue = PixelFormatEnum->GetValueByName(*PixelFormatString);
+		if (PixelFormatValue != INDEX_NONE && PixelFormatValue < PF_MAX)
+		{
+			PlatformData->PixelFormat = (EPixelFormat)PixelFormatValue;
+		}
+		else
+		{
+			UE_LOG(LogTexture, Warning, TEXT("Invalid pixel format '%s' for texture '%s'."), *PixelFormatString, Texture ? *Texture->GetPathName() : TEXT(""));
+			PlatformData->PixelFormat = PF_Unknown;
+		}
+	}
+	else if (Ar.IsSaving())
+	{
+		FString PixelFormatString = PixelFormatEnum->GetNameByValue(PlatformData->PixelFormat).GetPlainNameString();
+		Ar << PixelFormatString;
 	}
+
+	if (PlatformData->GetHasOptData())
+	{
+		Ar << PlatformData->OptData;
+	}
+
+	if (bCooked)
+	{
+		Ar << FirstMipToSerialize;
+		if (Ar.IsLoading())
+		{
+			check(Texture);
+			FirstMipToSerialize = 0;
+		}
+	}
+
+	TArray<uint32> BulkDataMipFlags;
+
+	// Update BulkDataFlags for cooked textures before saving.
+	if (bCooked && Ar.IsSaving())
+	{
+		if (bIsVirtual)
+		{
+			const int32 NumChunks = PlatformData->VTData->Chunks.Num();
+			BulkDataMipFlags.Reserve(NumChunks);
+			for (FVirtualTextureDataChunk& Chunk : PlatformData->VTData->Chunks)
+			{
+				BulkDataMipFlags.Add(Chunk.BulkData.GetBulkDataFlags());
+				Chunk.BulkData.SetBulkDataFlags(BULKDATA_Force_NOT_InlinePayload);
+			}
+		}
+		else
+		{
+			BulkDataMipFlags.Reserve(FirstMipToSerialize + NumMips);
+			for (const FTexture2DMipMap& Mip : PlatformData->Mips)
+			{
+				BulkDataMipFlags.Add(Mip.BulkData.GetBulkDataFlags());
+			}
+
+			// Optional Mips (Streaming)
+			const uint32 OptionalBulkDataFlags = BULKDATA_Force_NOT_InlinePayload | BULKDATA_OptionalPayload;
+			for (int32 MipIndex = 0; MipIndex < OptionalMips; ++MipIndex)
+			{
+				PlatformData->Mips[MipIndex + FirstMipToSerialize].BulkData.SetBulkDataFlags(OptionalBulkDataFlags);
+			}
+
+			// Streamed Mips (Non-Optional)
+			const uint32 StreamedBulkDataFlags = BULKDATA_Force_NOT_InlinePayload | (bDuplicateNonOptionalMips ? BULKDATA_DuplicateNonOptionalPayload : 0);
+			for (int32 MipIndex = OptionalMips; MipIndex < FirstInlineMip; ++MipIndex)
+			{
+				PlatformData->Mips[MipIndex + FirstMipToSerialize].BulkData.SetBulkDataFlags(StreamedBulkDataFlags);
+			}
+
+			// Inline Mips (Non-Optional)
+			const uint32 InlineBulkDataFlags = BULKDATA_ForceInlinePayload | BULKDATA_SingleUse;
+			for (int32 MipIndex = FirstInlineMip; MipIndex < NumMips; ++MipIndex)
+			{
+				PlatformData->Mips[MipIndex + FirstMipToSerialize].BulkData.SetBulkDataFlags(InlineBulkDataFlags);
+			}
+		}
+	}
+
 	Ar << NumMips;
 	check(NumMips >= (int32)PlatformData->GetNumMipsInTail());
 	if (Ar.IsLoading())
@@ -2641,32 +2807,26 @@ static void SerializePlatformData(
 	Ar << bIsVirtual;
 	if (bIsVirtual)
 	{
-		if (Ar.IsLoading() && PlatformData->VTData == nullptr)
+		if (Ar.IsLoading() && !PlatformData->VTData)
 		{
 			PlatformData->VTData = new FVirtualTextureBuiltData();
 		}
-		else
-		{
-			check(PlatformData->VTData);
-		}
 
+		check(PlatformData->VTData);
 		PlatformData->VTData->Serialize(Ar, Texture, FirstMipToSerialize);
-	}
 
-	if (bIsVirtual == false)
-	{
-		for (int32 MipIndex = 0; MipIndex < BulkDataMipFlags.Num(); ++MipIndex)
+		for (int32 ChunkIndex = 0; ChunkIndex < BulkDataMipFlags.Num(); ++ChunkIndex)
 		{
-			check(Ar.IsSaving());
-			PlatformData->Mips[MipIndex].BulkData.ResetBulkDataFlags(BulkDataMipFlags[MipIndex]);
+			check(Ar.IsSaving() && bCooked);
+			PlatformData->VTData->Chunks[ChunkIndex].BulkData.ResetBulkDataFlags(BulkDataMipFlags[ChunkIndex]);
 		}
 	}
 	else
 	{
-		for (int32 ChunkIndex = 0; ChunkIndex < BulkDataMipFlags.Num(); ++ChunkIndex)
+		for (int32 MipIndex = 0; MipIndex < BulkDataMipFlags.Num(); ++MipIndex)
 		{
-			check(Ar.IsSaving() && bCooked);
-			PlatformData->VTData->Chunks[ChunkIndex].BulkData.ResetBulkDataFlags(BulkDataMipFlags[ChunkIndex]);
+			check(Ar.IsSaving());
+			PlatformData->Mips[MipIndex].BulkData.ResetBulkDataFlags(BulkDataMipFlags[MipIndex]);
 		}
 	}
 }
