Fix ContentBundleGuid Crash: Add some fixup code on load to prevent actors from having invalid ContentBundleGuids while we find out why

#rnx

diff --git a/Engine/Source/Runtime/Engine/Private/Actor.cpp b/Engine/Source/Runtime/Engine/Private/Actor.cpp
--- a/Engine/Source/Runtime/Engine/Private/Actor.cpp
+++ b/Engine/Source/Runtime/Engine/Private/Actor.cpp
@@ -883,6 +883,16 @@ void AActor::Serialize(FArchive& Ar)
 		{
 			ContentBundleGuid = ContentBundlePaths::GetContentBundleGuidFromExternalActorPackagePath(GetPackage()->GetFName().ToString());
 		}
+		else // @todo_ow: remove once we find why some actors end up with invalid ContentBundleGuids
+		{
+			FGuid FixupContentBundleGuid = ContentBundlePaths::GetContentBundleGuidFromExternalActorPackagePath(GetPackage()->GetFName().ToString());
+			if (ContentBundleGuid != FixupContentBundleGuid)
+			{
+				UE_LOG(LogActor, Warning, TEXT("Actor ContentBundleGuid was fixed up: %s"), *GetName());
+				ContentBundleGuid = FixupContentBundleGuid;
+			}
+		}
+
 	}
 #endif
 }
diff --git a/Engine/Source/Runtime/Engine/Private/WorldPartition/WorldPartitionActorDesc.cpp b/Engine/Source/Runtime/Engine/Private/WorldPartition/WorldPartitionActorDesc.cpp
--- a/Engine/Source/Runtime/Engine/Private/WorldPartition/WorldPartitionActorDesc.cpp
+++ b/Engine/Source/Runtime/Engine/Private/WorldPartition/WorldPartitionActorDesc.cpp
@@ -483,6 +483,16 @@ PRAGMA_ENABLE_DEPRECATION_WARNINGS
 	else
 	{
 		Ar << ContentBundleGuid;
+		// @todo_ow: remove once we find why some actors end up with invalid ContentBundleGuids
+		if (Ar.IsLoading())
+		{
+			FGuid FixupContentBundleGuid = ContentBundlePaths::GetContentBundleGuidFromExternalActorPackagePath(ActorPackage.ToString());
+			if (ContentBundleGuid != FixupContentBundleGuid)
+			{
+				UE_LOG(LogWorldPartition, Warning, TEXT("ActorDesc ContentBundleGuid was fixed up: %s"), *GetActorName().ToString());
+				ContentBundleGuid = FixupContentBundleGuid;
+			}
+		}
 	}
 }
 
