Don't display progress dialog when re-caching materials as this introduces a crash at the moment. We can still display progress while re-caching instance materials. This should unblock devs needing to work in different preview modes.

diff --git a/Engine/Source/Runtime/Engine/Private/UnrealEngine.cpp b/Engine/Source/Runtime/Engine/Private/UnrealEngine.cpp
--- a/Engine/Source/Runtime/Engine/Private/UnrealEngine.cpp
+++ b/Engine/Source/Runtime/Engine/Private/UnrealEngine.cpp
@@ -836,7 +836,7 @@ void ScalabilityCVarsSinkCallback()
 			if (bCacheResourceShaders)
 			{
 				// For all materials, UMaterial::CacheResourceShadersForRendering
-				UMaterial::AllMaterialsCacheResourceShadersForRendering(true);
+				UMaterial::AllMaterialsCacheResourceShadersForRendering();
 				UMaterialInstance::AllMaterialsCacheResourceShadersForRendering(true);
 			}
 		}
