Prepare to move StaticSwitchData from EditorOnlyData into RuntimeData in preparation to make StaticSwitches dynamic

diff --git a/Engine/Plugins/Editor/SpeedTreeImporter/Source/SpeedTreeImporter/Private/SpeedTreeImportFactory.cpp b/Engine/Plugins/Editor/SpeedTreeImporter/Source/SpeedTreeImporter/Private/SpeedTreeImportFactory.cpp
--- a/Engine/Plugins/Editor/SpeedTreeImporter/Source/SpeedTreeImporter/Private/SpeedTreeImportFactory.cpp
+++ b/Engine/Plugins/Editor/SpeedTreeImporter/Source/SpeedTreeImporter/Private/SpeedTreeImportFactory.cpp
@@ -1387,7 +1387,7 @@ UMaterialInterface* CreateSpeedTreeMaterial9(UObject* Parent, FString MaterialFu
 	// material static setup
 	FStaticParameterSet StaticParameters;
 	UnrealMaterialInstance->GetStaticParameterValues(StaticParameters);
-	for (FStaticSwitchParameter& SwitchParameter : StaticParameters.EditorOnly.StaticSwitchParameters)
+	for (FStaticSwitchParameter& SwitchParameter : StaticParameters.StaticSwitchParameters)
 	{
 		if (SwitchParameter.ParameterInfo.Name == FName(TEXT("SharedEnable")))
 		{
diff --git a/Engine/Plugins/Editor/WorldPartitionHLODUtilities/Source/Private/HLODBuilderMeshApproximate.cpp b/Engine/Plugins/Editor/WorldPartitionHLODUtilities/Source/Private/HLODBuilderMeshApproximate.cpp
--- a/Engine/Plugins/Editor/WorldPartitionHLODUtilities/Source/Private/HLODBuilderMeshApproximate.cpp
+++ b/Engine/Plugins/Editor/WorldPartitionHLODUtilities/Source/Private/HLODBuilderMeshApproximate.cpp
@@ -157,7 +157,7 @@ TArray<UActorComponent*> UHLODBuilderMeshApproximate::Build(const FHLODBuildCont
 					SwitchParameter.ParameterInfo.Name = ParamName;
 					SwitchParameter.Value = true;
 					SwitchParameter.bOverride = true;
-					StaticParameterSet.EditorOnly.StaticSwitchParameters.Add(SwitchParameter);
+					StaticParameterSet.StaticSwitchParameters.Add(SwitchParameter);
 				}
 			};
 
diff --git a/Engine/Plugins/Enterprise/DataprepEditor/Source/DataprepCore/Private/DataprepCorePrivateUtils.cpp b/Engine/Plugins/Enterprise/DataprepEditor/Source/DataprepCore/Private/DataprepCorePrivateUtils.cpp
--- a/Engine/Plugins/Enterprise/DataprepEditor/Source/DataprepCore/Private/DataprepCorePrivateUtils.cpp
+++ b/Engine/Plugins/Enterprise/DataprepEditor/Source/DataprepCore/Private/DataprepCorePrivateUtils.cpp
@@ -211,7 +211,7 @@ void DataprepCorePrivateUtils::CompileMaterial(UMaterialInterface* MaterialInter
 			FStaticParameterSet StaticParameters;
 			ConstantMaterialInstance->GetStaticParameterValues( StaticParameters );
 
-			for ( FStaticSwitchParameter& Switch : StaticParameters.EditorOnly.StaticSwitchParameters )
+			for ( FStaticSwitchParameter& Switch : StaticParameters.StaticSwitchParameters )
 			{
 				if ( Switch.bOverride )
 				{
diff --git a/Engine/Plugins/Enterprise/DatasmithContent/Source/DatasmithContent/Private/ObjectTemplates/DatasmithMaterialInstanceTemplate.cpp b/Engine/Plugins/Enterprise/DatasmithContent/Source/DatasmithContent/Private/ObjectTemplates/DatasmithMaterialInstanceTemplate.cpp
--- a/Engine/Plugins/Enterprise/DatasmithContent/Source/DatasmithContent/Private/ObjectTemplates/DatasmithMaterialInstanceTemplate.cpp
+++ b/Engine/Plugins/Enterprise/DatasmithContent/Source/DatasmithContent/Private/ObjectTemplates/DatasmithMaterialInstanceTemplate.cpp
@@ -103,7 +103,7 @@ void FDatasmithStaticParameterSetTemplate::Apply( UMaterialInstanceConstant* Des
 			}
 		}
 
-		for ( FStaticSwitchParameter& DestinationSwitchParameter : DestinationStaticParameters.EditorOnly.StaticSwitchParameters )
+		for ( FStaticSwitchParameter& DestinationSwitchParameter : DestinationStaticParameters.StaticSwitchParameters )
 		{
 			if ( DestinationSwitchParameter.ParameterInfo.Name == It->Key )
 			{
@@ -132,9 +132,9 @@ void FDatasmithStaticParameterSetTemplate::Load( const UMaterialInstanceConstant
 	FStaticParameterSet SourceStaticParameters;
 	const_cast< UMaterialInstanceConstant& >( Source ).GetStaticParameterValues( SourceStaticParameters );
 
-	StaticSwitchParameters.Empty( SourceStaticParameters.EditorOnly.StaticSwitchParameters.Num() );
+	StaticSwitchParameters.Empty( SourceStaticParameters.StaticSwitchParameters.Num() );
 
-	for ( const FStaticSwitchParameter& SourceSwitch : SourceStaticParameters.EditorOnly.StaticSwitchParameters )
+	for ( const FStaticSwitchParameter& SourceSwitch : SourceStaticParameters.StaticSwitchParameters )
 	{
 		if ( !bOverridesOnly || SourceSwitch.bOverride )
 		{
diff --git a/Engine/Plugins/Enterprise/DatasmithImporter/Source/DatasmithImporter/Private/Utility/DatasmithImporterImpl.cpp b/Engine/Plugins/Enterprise/DatasmithImporter/Source/DatasmithImporter/Private/Utility/DatasmithImporterImpl.cpp
--- a/Engine/Plugins/Enterprise/DatasmithImporter/Source/DatasmithImporter/Private/Utility/DatasmithImporterImpl.cpp
+++ b/Engine/Plugins/Enterprise/DatasmithImporter/Source/DatasmithImporter/Private/Utility/DatasmithImporterImpl.cpp
@@ -370,7 +370,7 @@ void FDatasmithImporterImpl::CompileMaterial( UObject* Material, FMaterialUpdate
 			FStaticParameterSet StaticParameters;
 			ConstantMaterialInstance->GetStaticParameterValues( StaticParameters );
 
-			for ( FStaticSwitchParameter& Switch : StaticParameters.EditorOnly.StaticSwitchParameters )
+			for ( FStaticSwitchParameter& Switch : StaticParameters.StaticSwitchParameters )
 			{
 				if ( Switch.bOverride )
 				{
diff --git a/Engine/Plugins/Importers/USDImporter/Source/USDUtilities/Private/USDShadeConversion.cpp b/Engine/Plugins/Importers/USDImporter/Source/USDUtilities/Private/USDShadeConversion.cpp
--- a/Engine/Plugins/Importers/USDImporter/Source/USDUtilities/Private/USDShadeConversion.cpp
+++ b/Engine/Plugins/Importers/USDImporter/Source/USDUtilities/Private/USDShadeConversion.cpp
@@ -1071,7 +1071,7 @@ namespace UE
 					FStaticParameterSet StaticParameters;
 					Constant->GetStaticParameterValues( StaticParameters );
 
-					for ( FStaticSwitchParameter& StaticSwitchParameter : StaticParameters.EditorOnly.StaticSwitchParameters )
+					for ( FStaticSwitchParameter& StaticSwitchParameter : StaticParameters.StaticSwitchParameters )
 					{
 						if ( StaticSwitchParameter.ParameterInfo.Name == ParameterName )
 						{
diff --git a/Engine/Source/Developer/HierarchicalLODUtilities/Private/HierarchicalLODUtilities.cpp b/Engine/Source/Developer/HierarchicalLODUtilities/Private/HierarchicalLODUtilities.cpp
--- a/Engine/Source/Developer/HierarchicalLODUtilities/Private/HierarchicalLODUtilities.cpp
+++ b/Engine/Source/Developer/HierarchicalLODUtilities/Private/HierarchicalLODUtilities.cpp
@@ -632,7 +632,7 @@ FHLODBuildResults GenerateHLODMesh_Approximate(const FHLODBuildParams& InBuildPa
 					SwitchParameter.ParameterInfo.Name = ParamName;
 					SwitchParameter.Value = true;
 					SwitchParameter.bOverride = true;
-					StaticParameterSet.EditorOnly.StaticSwitchParameters.Add(SwitchParameter);
+					StaticParameterSet.StaticSwitchParameters.Add(SwitchParameter);
 				}
 			};
 
diff --git a/Engine/Source/Developer/MaterialUtilities/Private/MaterialUtilities.cpp b/Engine/Source/Developer/MaterialUtilities/Private/MaterialUtilities.cpp
--- a/Engine/Source/Developer/MaterialUtilities/Private/MaterialUtilities.cpp
+++ b/Engine/Source/Developer/MaterialUtilities/Private/MaterialUtilities.cpp
@@ -159,7 +159,7 @@ UMaterialInterface* FMaterialUtilities::CreateProxyMaterialAndTextures(UPackage*
 			SwitchParameter.ParameterInfo.Name = *(TEXT("Use") + TrimmedPropertyName);
 			SwitchParameter.Value = true;
 			SwitchParameter.bOverride = true;
-			NewStaticParameterSet.EditorOnly.StaticSwitchParameters.Add(SwitchParameter);
+			NewStaticParameterSet.StaticSwitchParameters.Add(SwitchParameter);
 		}
 		else
 		{
@@ -194,10 +194,10 @@ UMaterialInterface* FMaterialUtilities::CreateProxyMaterialAndTextures(UPackage*
 		SwitchParameter.ParameterInfo.Name = TEXT("UseCustomUV");
 		SwitchParameter.Value = true;
 		SwitchParameter.bOverride = true;
-		NewStaticParameterSet.EditorOnly.StaticSwitchParameters.Add(SwitchParameter);
+		NewStaticParameterSet.StaticSwitchParameters.Add(SwitchParameter);
 
 		SwitchParameter.ParameterInfo.Name = *(TEXT("UseUV") + FString::FromInt(MeshData.TextureCoordinateIndex));
-		NewStaticParameterSet.EditorOnly.StaticSwitchParameters.Add(SwitchParameter);
+		NewStaticParameterSet.StaticSwitchParameters.Add(SwitchParameter);
 	}
 
 	Material->UpdateStaticPermutation(NewStaticParameterSet);
diff --git a/Engine/Source/Developer/MaterialUtilities/Private/ProxyMaterialUtilities.cpp b/Engine/Source/Developer/MaterialUtilities/Private/ProxyMaterialUtilities.cpp
--- a/Engine/Source/Developer/MaterialUtilities/Private/ProxyMaterialUtilities.cpp
+++ b/Engine/Source/Developer/MaterialUtilities/Private/ProxyMaterialUtilities.cpp
@@ -307,10 +307,10 @@ UMaterialInstanceConstant* FMaterialUtilities::CreateFlattenMaterialInstance(UPa
 		SwitchParameter.ParameterInfo.Name = TEXT("UseCustomUV");
 		SwitchParameter.Value = true;
 		SwitchParameter.bOverride = true;
-		NewStaticParameterSet.EditorOnly.StaticSwitchParameters.Add(SwitchParameter);
+		NewStaticParameterSet.StaticSwitchParameters.Add(SwitchParameter);
 
 		SwitchParameter.ParameterInfo.Name = *(TEXT("UseUV") + FString::FromInt(FlattenMaterial.UVChannel));
-		NewStaticParameterSet.EditorOnly.StaticSwitchParameters.Add(SwitchParameter);
+		NewStaticParameterSet.StaticSwitchParameters.Add(SwitchParameter);
 	}
 
 	auto CreateTextureFromDefault = [&](const FName TextureName, const FString& AssetLongName, FIntPoint Size, const TArray<FColor>& Samples)
@@ -349,7 +349,7 @@ UMaterialInstanceConstant* FMaterialUtilities::CreateFlattenMaterialInstance(UPa
 			SwitchParameter.ParameterInfo.Name = UseTexture;
 			SwitchParameter.Value = true;
 			SwitchParameter.bOverride = true;
-			NewStaticParameterSet.EditorOnly.StaticSwitchParameters.Add(SwitchParameter);
+			NewStaticParameterSet.StaticSwitchParameters.Add(SwitchParameter);
 
 			OutMaterial->SetTextureParameterValueEditorOnly(TextureName, Texture);
 			
@@ -476,17 +476,17 @@ UMaterialInstanceConstant* FMaterialUtilities::CreateFlattenMaterialInstance(UPa
 		SwitchParameter.ParameterInfo.Name = TEXT("PackMetallic");
 		SwitchParameter.Value = bPackMetallic;
 		SwitchParameter.bOverride = true;
-		NewStaticParameterSet.EditorOnly.StaticSwitchParameters.Add(SwitchParameter);
+		NewStaticParameterSet.StaticSwitchParameters.Add(SwitchParameter);
 
 		SwitchParameter.ParameterInfo.Name = TEXT("PackSpecular");
 		SwitchParameter.Value = bPackSpecular;
 		SwitchParameter.bOverride = true;
-		NewStaticParameterSet.EditorOnly.StaticSwitchParameters.Add(SwitchParameter);
+		NewStaticParameterSet.StaticSwitchParameters.Add(SwitchParameter);
 
 		SwitchParameter.ParameterInfo.Name = TEXT("PackRoughness");
 		SwitchParameter.Value = bPackRoughness;
 		SwitchParameter.bOverride = true;
-		NewStaticParameterSet.EditorOnly.StaticSwitchParameters.Add(SwitchParameter);
+		NewStaticParameterSet.StaticSwitchParameters.Add(SwitchParameter);
 
 		// Set up switch and texture values
 		OutMaterial->SetTextureParameterValueEditorOnly(PackedTextureName, PackedTexture);
diff --git a/Engine/Source/Editor/MaterialEditor/Private/MaterialEditor.cpp b/Engine/Source/Editor/MaterialEditor/Private/MaterialEditor.cpp
--- a/Engine/Source/Editor/MaterialEditor/Private/MaterialEditor.cpp
+++ b/Engine/Source/Editor/MaterialEditor/Private/MaterialEditor.cpp
@@ -3177,7 +3177,7 @@ void FMaterialEditor::UpdateGraphNodeStates()
 				FStaticSwitchParameter SwitchParam;
 				SwitchParam.Value = StaticSwitch->DefaultValue;
 				SwitchParam.ExpressionGUID = StaticSwitch->ExpressionGUID;
-				StaticSwitchSet.EditorOnly.StaticSwitchParameters.Add(SwitchParam);
+				StaticSwitchSet.StaticSwitchParameters.Add(SwitchParam);
 			}
 		}
 	}
diff --git a/Engine/Source/Editor/MaterialEditor/Private/MaterialInstanceEditor.cpp b/Engine/Source/Editor/MaterialEditor/Private/MaterialInstanceEditor.cpp
--- a/Engine/Source/Editor/MaterialEditor/Private/MaterialInstanceEditor.cpp
+++ b/Engine/Source/Editor/MaterialEditor/Private/MaterialInstanceEditor.cpp
@@ -532,7 +532,7 @@ void FMaterialInstanceEditor::ReInitMaterialFunctionProxies()
 		TArray<FFontParameterValue> FontParameterValues = FunctionInstanceProxy->FontParameterValues;
 
 		const FStaticParameterSet& OldStaticParameters = FunctionInstanceProxy->GetStaticParameters();
-		TArray<FStaticSwitchParameter> StaticSwitchParameters = OldStaticParameters.EditorOnly.StaticSwitchParameters;
+		TArray<FStaticSwitchParameter> StaticSwitchParameters = OldStaticParameters.StaticSwitchParameters;
 		TArray<FStaticComponentMaskParameter> StaticComponentMaskParameters = OldStaticParameters.EditorOnly.StaticComponentMaskParameters;
 
 		// Regenerate proxies
@@ -632,14 +632,14 @@ void FMaterialInstanceEditor::ReInitMaterialFunctionProxies()
 		FStaticParameterSet StaticParametersOverride = FunctionInstanceProxy->GetStaticParameters();
 
 		FunctionInstanceProxy->GetAllStaticSwitchParameterInfo(OutParameterInfo, Guids);
-		StaticParametersOverride.EditorOnly.StaticSwitchParameters.Empty();
+		StaticParametersOverride.StaticSwitchParameters.Empty();
 		for (FStaticSwitchParameter& StaticSwitchParameter : StaticSwitchParameters)
 		{
 			int32 Index = Guids.Find(StaticSwitchParameter.ExpressionGUID);
 			if (Index != INDEX_NONE)
 			{
-				StaticParametersOverride.EditorOnly.StaticSwitchParameters.Add(StaticSwitchParameter);
-				StaticParametersOverride.EditorOnly.StaticSwitchParameters.Last().ParameterInfo = OutParameterInfo[Index];
+				StaticParametersOverride.StaticSwitchParameters.Add(StaticSwitchParameter);
+				StaticParametersOverride.StaticSwitchParameters.Last().ParameterInfo = OutParameterInfo[Index];
 			}
 		}
 
diff --git a/Engine/Source/Editor/MaterialEditor/Private/MaterialPropertyHelpers.cpp b/Engine/Source/Editor/MaterialEditor/Private/MaterialPropertyHelpers.cpp
--- a/Engine/Source/Editor/MaterialEditor/Private/MaterialPropertyHelpers.cpp
+++ b/Engine/Source/Editor/MaterialEditor/Private/MaterialPropertyHelpers.cpp
@@ -430,7 +430,7 @@ FReply FMaterialPropertyHelpers::OnClickedSaveNewFunctionInstance(class UMateria
 					ChildInstance->FontParameterValues = EditedInstance->FontParameterValues;
 
 					const FStaticParameterSet& StaticParameters = EditedInstance->GetStaticParameters();
-					ChildInstance->StaticSwitchParameterValues = StaticParameters.EditorOnly.StaticSwitchParameters;
+					ChildInstance->StaticSwitchParameterValues = StaticParameters.StaticSwitchParameters;
 					ChildInstance->StaticComponentMaskParameterValues = StaticParameters.EditorOnly.StaticComponentMaskParameters;
 				}
 			}
@@ -523,7 +523,7 @@ FReply FMaterialPropertyHelpers::OnClickedSaveNewLayerInstance(class UMaterialFu
 					ChildInstance->FontParameterValues = EditedInstance->FontParameterValues;
 
 					const FStaticParameterSet& StaticParameters = EditedInstance->GetStaticParameters();
-					ChildInstance->StaticSwitchParameterValues = StaticParameters.EditorOnly.StaticSwitchParameters;
+					ChildInstance->StaticSwitchParameterValues = StaticParameters.StaticSwitchParameters;
 					ChildInstance->StaticComponentMaskParameterValues = StaticParameters.EditorOnly.StaticComponentMaskParameters;
 				}
 			}
diff --git a/Engine/Source/Editor/UnrealEd/Private/Commandlets/DumpMaterialShaderTypes.cpp b/Engine/Source/Editor/UnrealEd/Private/Commandlets/DumpMaterialShaderTypes.cpp
--- a/Engine/Source/Editor/UnrealEd/Private/Commandlets/DumpMaterialShaderTypes.cpp
+++ b/Engine/Source/Editor/UnrealEd/Private/Commandlets/DumpMaterialShaderTypes.cpp
@@ -430,7 +430,7 @@ int ProcessMaterialInstances(const ITargetPlatform* TargetPlatform, const EShade
 			Output.Log(TEXT(""));
 			Output.Log(FString::Printf(TEXT("Material Instance: %s - %d shaders"), *AssetData.AssetName.ToString(), TotalShadersForMaterial));
 			Output.Log(FString::Printf(TEXT("Parent: %s"), Top ? *Top->GetFullName() : TEXT("NO PARENT")));
-			Output.Log(FString::Printf(TEXT("Static Parameters: %d"), MaterialInstance->bHasStaticPermutationResource ? MaterialInstance->GetStaticParameters().EditorOnly.StaticSwitchParameters.Num() : 0));
+			Output.Log(FString::Printf(TEXT("Static Parameters: %d"), MaterialInstance->bHasStaticPermutationResource ? MaterialInstance->GetStaticParameters().StaticSwitchParameters.Num() : 0));
 
 			FSHAHash OutHash;
 
@@ -438,9 +438,9 @@ int ProcessMaterialInstances(const ITargetPlatform* TargetPlatform, const EShade
 			{
 				FSHA1 Hasher;
 				const FStaticParameterSet& ParameterSet = MaterialInstance->GetStaticParameters();
-				for (int32 StaticSwitchIndex = 0; StaticSwitchIndex < ParameterSet.EditorOnly.StaticSwitchParameters.Num(); ++StaticSwitchIndex)
+				for (int32 StaticSwitchIndex = 0; StaticSwitchIndex < ParameterSet.StaticSwitchParameters.Num(); ++StaticSwitchIndex)
 				{
-					const FStaticSwitchParameter& StaticSwitchParameter = ParameterSet.EditorOnly.StaticSwitchParameters[StaticSwitchIndex];
+					const FStaticSwitchParameter& StaticSwitchParameter = ParameterSet.StaticSwitchParameters[StaticSwitchIndex];
 
 					StaticSwitchParameter.UpdateHash(Hasher);
 
diff --git a/Engine/Source/Editor/UnrealEd/Private/PreviewMaterial.cpp b/Engine/Source/Editor/UnrealEd/Private/PreviewMaterial.cpp
--- a/Engine/Source/Editor/UnrealEd/Private/PreviewMaterial.cpp
+++ b/Engine/Source/Editor/UnrealEd/Private/PreviewMaterial.cpp
@@ -836,8 +836,9 @@ void UMaterialEditorInstanceConstant::ApplySourceFunctionChanges()
 		SourceFunction->SparseVolumeTextureParameterValues = SourceInstance->SparseVolumeTextureParameterValues;
 		SourceFunction->FontParameterValues = SourceInstance->FontParameterValues;
 
+		
+		SourceFunction->StaticSwitchParameterValues = SourceInstance->GetStaticParameters().StaticSwitchParameters;
 		const FStaticParameterSetEditorOnlyData& StaticParameters = SourceInstance->GetEditorOnlyStaticParameters();
-		SourceFunction->StaticSwitchParameterValues = StaticParameters.StaticSwitchParameters;
 		SourceFunction->StaticComponentMaskParameterValues = StaticParameters.StaticComponentMaskParameters;
 
 		SourceFunction->MarkPackageDirty();
diff --git a/Engine/Source/Runtime/Engine/Private/Materials/HLSLMaterialTranslator.cpp b/Engine/Source/Runtime/Engine/Private/Materials/HLSLMaterialTranslator.cpp
--- a/Engine/Source/Runtime/Engine/Private/Materials/HLSLMaterialTranslator.cpp
+++ b/Engine/Source/Runtime/Engine/Private/Materials/HLSLMaterialTranslator.cpp
@@ -7029,7 +7029,7 @@ int32 FHLSLMaterialTranslator::StaticBoolParameter(FName ParameterName,bool bDef
 	FMaterialParameterInfo ParameterInfo = GetParameterAssociationInfo();
 	ParameterInfo.Name = ParameterName;
 
-	for (const FStaticSwitchParameter& Parameter : StaticParameters.EditorOnly.StaticSwitchParameters)
+	for (const FStaticSwitchParameter& Parameter : StaticParameters.StaticSwitchParameters)
 	{
 		if (Parameter.ParameterInfo == ParameterInfo)
 		{
diff --git a/Engine/Source/Runtime/Engine/Private/Materials/Material.cpp b/Engine/Source/Runtime/Engine/Private/Materials/Material.cpp
--- a/Engine/Source/Runtime/Engine/Private/Materials/Material.cpp
+++ b/Engine/Source/Runtime/Engine/Private/Materials/Material.cpp
@@ -5854,9 +5854,9 @@ bool UMaterial::RecursiveGetExpressionChain(
 						{
 							bool bUseInputA = StaticSwitchExp->DefaultValue;
 							FName StaticSwitchExpName = StaticSwitchExp->ParameterName;
-							for (int32 CheckIdx = 0; CheckIdx < InStaticParameterSet->EditorOnly.StaticSwitchParameters.Num(); CheckIdx++)
+							for (int32 CheckIdx = 0; CheckIdx < InStaticParameterSet->StaticSwitchParameters.Num(); CheckIdx++)
 							{
-								FStaticSwitchParameter& SwitchParam = InStaticParameterSet->EditorOnly.StaticSwitchParameters[CheckIdx];
+								FStaticSwitchParameter& SwitchParam = InStaticParameterSet->StaticSwitchParameters[CheckIdx];
 								if (SwitchParam.ParameterInfo.Name == StaticSwitchExpName)
 								{
 									// Found it...
diff --git a/Engine/Source/Runtime/Engine/Private/Materials/MaterialExpressions.cpp b/Engine/Source/Runtime/Engine/Private/Materials/MaterialExpressions.cpp
--- a/Engine/Source/Runtime/Engine/Private/Materials/MaterialExpressions.cpp
+++ b/Engine/Source/Runtime/Engine/Private/Materials/MaterialExpressions.cpp
@@ -14713,7 +14713,7 @@ void UMaterialFunctionInstance::OverrideMaterialInstanceParameterValues(UMateria
 
 	// Static parameters
 	FStaticParameterSet StaticParametersOverride = Instance->GetStaticParameters();
-	StaticParametersOverride.EditorOnly.StaticSwitchParameters = StaticSwitchParameterValues;
+	StaticParametersOverride.StaticSwitchParameters = StaticSwitchParameterValues;
 	StaticParametersOverride.EditorOnly.StaticComponentMaskParameters = StaticComponentMaskParameterValues;
 	Instance->UpdateStaticPermutation(StaticParametersOverride);
 }
diff --git a/Engine/Source/Runtime/Engine/Private/Materials/MaterialHLSLTree.cpp b/Engine/Source/Runtime/Engine/Private/Materials/MaterialHLSLTree.cpp
--- a/Engine/Source/Runtime/Engine/Private/Materials/MaterialHLSLTree.cpp
+++ b/Engine/Source/Runtime/Engine/Private/Materials/MaterialHLSLTree.cpp
@@ -515,7 +515,7 @@ void FExpressionParameter::EmitValuePreshader(FEmitContext& Context, FEmitScope&
 			switch (ParameterType)
 			{
 			case EMaterialParameterType::StaticSwitch:
-				for (const FStaticSwitchParameter& Parameter : EmitMaterialData.StaticParameters->EditorOnly.StaticSwitchParameters)
+				for (const FStaticSwitchParameter& Parameter : EmitMaterialData.StaticParameters->StaticSwitchParameters)
 				{
 					if (Parameter.ParameterInfo == ParameterInfo)
 					{
diff --git a/Engine/Source/Runtime/Engine/Private/Materials/MaterialInstance.cpp b/Engine/Source/Runtime/Engine/Private/Materials/MaterialInstance.cpp
--- a/Engine/Source/Runtime/Engine/Private/Materials/MaterialInstance.cpp
+++ b/Engine/Source/Runtime/Engine/Private/Materials/MaterialInstance.cpp
@@ -542,9 +542,9 @@ void UMaterialInstance::SwapLayerParameterIndices(int32 OriginalIndex, int32 New
 		SwapLayerParameterIndicesArray(RuntimeVirtualTextureParameterValues, OriginalIndex, NewIndex);
 		SwapLayerParameterIndicesArray(SparseVolumeTextureParameterValues, OriginalIndex, NewIndex);
 		SwapLayerParameterIndicesArray(FontParameterValues, OriginalIndex, NewIndex);
+		SwapLayerParameterIndicesArray(StaticParametersRuntime.StaticSwitchParameters, OriginalIndex, NewIndex);
 		if (EditorOnly)
 		{
-			SwapLayerParameterIndicesArray(EditorOnly->StaticParameters.StaticSwitchParameters, OriginalIndex, NewIndex);
 			SwapLayerParameterIndicesArray(EditorOnly->StaticParameters.StaticComponentMaskParameters, OriginalIndex, NewIndex);
 		}
 	}
@@ -560,9 +560,9 @@ void UMaterialInstance::RemoveLayerParameterIndex(int32 Index)
 	RemoveLayerParameterIndicesArray(RuntimeVirtualTextureParameterValues, Index);
 	RemoveLayerParameterIndicesArray(SparseVolumeTextureParameterValues, Index);
 	RemoveLayerParameterIndicesArray(FontParameterValues, Index);
+	RemoveLayerParameterIndicesArray(StaticParametersRuntime.StaticSwitchParameters, Index);
 	if (EditorOnly)
 	{
-		RemoveLayerParameterIndicesArray(EditorOnly->StaticParameters.StaticSwitchParameters, Index);
 		RemoveLayerParameterIndicesArray(EditorOnly->StaticParameters.StaticComponentMaskParameters, Index);
 	}
 }
@@ -616,7 +616,7 @@ bool UMaterialInstance::UpdateParameters()
 			bDirty = UpdateParameterSet<FFontParameterValue, UMaterialExpressionFontSampleParameter>(FontParameterValues, ParentMaterial) || bDirty;
 
 			// Static switch parameters
-			bDirty = UpdateParameterSet<FStaticSwitchParameter, UMaterialExpressionStaticBoolParameter>(EditorOnly->StaticParameters.StaticSwitchParameters, ParentMaterial) || bDirty;
+			bDirty = UpdateParameterSet<FStaticSwitchParameter, UMaterialExpressionStaticBoolParameter>(StaticParametersRuntime.StaticSwitchParameters, ParentMaterial) || bDirty;
 
 			// Static component mask parameters
 			bDirty = UpdateParameterSet<FStaticComponentMaskParameter, UMaterialExpressionStaticComponentMaskParameter>(EditorOnly->StaticParameters.StaticComponentMaskParameters, ParentMaterial) || bDirty;
@@ -655,7 +655,7 @@ bool UMaterialInstance::UpdateParameters()
 					RemapLayerParameterIndicesArray(RuntimeVirtualTextureParameterValues, RemapLayerIndices);
 					RemapLayerParameterIndicesArray(SparseVolumeTextureParameterValues, RemapLayerIndices);
 					RemapLayerParameterIndicesArray(FontParameterValues, RemapLayerIndices);
-					RemapLayerParameterIndicesArray(EditorOnly->StaticParameters.StaticSwitchParameters, RemapLayerIndices);
+					RemapLayerParameterIndicesArray(StaticParametersRuntime.StaticSwitchParameters, RemapLayerIndices);
 					RemapLayerParameterIndicesArray(EditorOnly->StaticParameters.StaticComponentMaskParameters, RemapLayerIndices);
 					bDirty = true;
 				}
@@ -946,10 +946,8 @@ bool UMaterialInstance::GetParameterOverrideValue(EMaterialParameterType Type, c
 	case EMaterialParameterType::RuntimeVirtualTexture: bResult = GameThread_GetParameterValue(RuntimeVirtualTextureParameterValues, ParameterInfo, OutResult); break;
 	case EMaterialParameterType::SparseVolumeTexture: bResult = GameThread_GetParameterValue(SparseVolumeTextureParameterValues, ParameterInfo, OutResult); break;
 	case EMaterialParameterType::Font: bResult = GameThread_GetParameterValue(FontParameterValues, ParameterInfo, OutResult); break;
+	case EMaterialParameterType::StaticSwitch: bResult = GameThread_GetParameterValue(StaticParametersRuntime.StaticSwitchParameters, ParameterInfo, OutResult); break;
 #if WITH_EDITORONLY_DATA
-	case EMaterialParameterType::StaticSwitch:
-		bResult = GameThread_GetParameterValue(GetEditorOnlyData()->StaticParameters.StaticSwitchParameters, ParameterInfo, OutResult);
-		break;
 	case EMaterialParameterType::StaticComponentMask:
 		bResult = GameThread_GetParameterValue(GetEditorOnlyData()->StaticParameters.StaticComponentMaskParameters, ParameterInfo, OutResult);
 		break;
@@ -1798,7 +1796,7 @@ void UMaterialInstance::SetStaticSwitchParameterValueEditorOnly(const FMaterialP
 	UMaterialInstanceEditorOnlyData* EditorOnly = GetEditorOnlyData();
 	check(EditorOnly);
 
-	for (FStaticSwitchParameter& StaticSwitches : EditorOnly->StaticParameters.StaticSwitchParameters)
+	for (FStaticSwitchParameter& StaticSwitches : StaticParametersRuntime.StaticSwitchParameters)
 	{
 		if (StaticSwitches.ParameterInfo == ParameterInfo)
 		{
@@ -1808,7 +1806,7 @@ void UMaterialInstance::SetStaticSwitchParameterValueEditorOnly(const FMaterialP
 		}
 	}
 
-	new(EditorOnly->StaticParameters.StaticSwitchParameters) FStaticSwitchParameter(ParameterInfo, Value, true, FGuid());
+	new(StaticParametersRuntime.StaticSwitchParameters) FStaticSwitchParameter(ParameterInfo, Value, true, FGuid());
 }
 
 void UMaterialInstance::GetStaticParameterValues(FStaticParameterSet& OutStaticParameters)
@@ -1914,10 +1912,8 @@ void UMaterialInstance::GetAllParametersOfType(EMaterialParameterType Type, TMap
 		case EMaterialParameterType::RuntimeVirtualTexture: GameThread_ApplyParameterOverrides(Instance->RuntimeVirtualTextureParameterValues, LayerIndexRemap, bSetOverride, OverridenParameters, OutParameters); break;
 		case EMaterialParameterType::SparseVolumeTexture: GameThread_ApplyParameterOverrides(Instance->SparseVolumeTextureParameterValues, LayerIndexRemap, bSetOverride, OverridenParameters, OutParameters); break;
 		case EMaterialParameterType::Font: GameThread_ApplyParameterOverrides(Instance->FontParameterValues, LayerIndexRemap, bSetOverride, OverridenParameters, OutParameters); break;
+		case EMaterialParameterType::StaticSwitch: GameThread_ApplyParameterOverrides(Instance->StaticParametersRuntime.StaticSwitchParameters, LayerIndexRemap, bSetOverride, OverridenParameters, OutParameters); break;
 #if WITH_EDITORONLY_DATA
-		case EMaterialParameterType::StaticSwitch:
-			GameThread_ApplyParameterOverrides(Instance->GetEditorOnlyData()->StaticParameters.StaticSwitchParameters, LayerIndexRemap, bSetOverride, OverridenParameters, OutParameters);
-			break;
 		case EMaterialParameterType::StaticComponentMask:
 			GameThread_ApplyParameterOverrides(Instance->GetEditorOnlyData()->StaticParameters.StaticComponentMaskParameters, LayerIndexRemap, bSetOverride, OverridenParameters, OutParameters);
 			break;
@@ -2632,11 +2628,19 @@ void UMaterialInstance::Serialize(FArchive& Ar)
 		StaticParameters_DEPRECATED.UpdateLegacyTerrainLayerWeightData();
 	}
 
-	if (Ar.IsLoading() && !StaticParameters_DEPRECATED.IsEmpty())
+	if (Ar.IsLoading())
 	{
-		StaticParametersRuntime = MoveTemp(StaticParameters_DEPRECATED.GetRuntime());
-		GetEditorOnlyData()->StaticParameters = MoveTemp(StaticParameters_DEPRECATED.EditorOnly);
-		StaticParameters_DEPRECATED.Empty();
+		if (!StaticParameters_DEPRECATED.IsEmpty())
+		{
+			StaticParametersRuntime = MoveTemp(StaticParameters_DEPRECATED.GetRuntime());
+			GetEditorOnlyData()->StaticParameters = MoveTemp(StaticParameters_DEPRECATED.EditorOnly);
+			StaticParameters_DEPRECATED.Empty();
+		}
+
+		if (!GetEditorOnlyData()->StaticParameters.StaticSwitchParameters_DEPRECATED.IsEmpty())
+		{
+			StaticParametersRuntime.StaticSwitchParameters = MoveTemp(GetEditorOnlyData()->StaticParameters.StaticSwitchParameters_DEPRECATED);
+		}
 	}
 
 #endif // WITH_EDITOR
@@ -2735,13 +2739,13 @@ void UMaterialInstance::Serialize(FArchive& Ar)
 			FMaterialShaderMapId LegacyId;
 			LegacyId.Serialize(Ar, bLoadedByCookedMaterial);
 
+			StaticParametersRuntime.StaticSwitchParameters = LegacyId.GetStaticSwitchParameters();
+			TrimToOverriddenOnly(StaticParametersRuntime.StaticSwitchParameters);
+
 			if (IsEditorOnlyDataValid())
 			{
-				GetEditorOnlyData()->StaticParameters.StaticSwitchParameters = LegacyId.GetStaticSwitchParameters();
 				GetEditorOnlyData()->StaticParameters.StaticComponentMaskParameters = LegacyId.GetStaticComponentMaskParameters();
 				GetEditorOnlyData()->StaticParameters.TerrainLayerWeightParameters = LegacyId.GetTerrainLayerWeightParameters();
-
-				TrimToOverriddenOnly(GetEditorOnlyData()->StaticParameters.StaticSwitchParameters);
 				TrimToOverriddenOnly(GetEditorOnlyData()->StaticParameters.StaticComponentMaskParameters);
 			}
 		}
@@ -3630,7 +3634,7 @@ void UMaterialInstance::UpdateStaticPermutation(const FStaticParameterSet& NewPa
 	UMaterialInstanceEditorOnlyData* EditorOnly = GetEditorOnlyData();
 	FStaticParameterSet CompareParameters = NewParameters;
 
-	TrimToOverriddenOnly(CompareParameters.EditorOnly.StaticSwitchParameters);
+	TrimToOverriddenOnly(CompareParameters.StaticSwitchParameters);
 	TrimToOverriddenOnly(CompareParameters.EditorOnly.StaticComponentMaskParameters);
 
 	// Check to see if the material layers being assigned match values from the parent
diff --git a/Engine/Source/Runtime/Engine/Private/Materials/MaterialInstanceConstant.cpp b/Engine/Source/Runtime/Engine/Private/Materials/MaterialInstanceConstant.cpp
--- a/Engine/Source/Runtime/Engine/Private/Materials/MaterialInstanceConstant.cpp
+++ b/Engine/Source/Runtime/Engine/Private/Materials/MaterialInstanceConstant.cpp
@@ -110,7 +110,7 @@ void UMaterialInstanceConstant::CopyMaterialUniformParametersEditorOnly(UMateria
 			FStaticParameterSet MyParamSet;
 			GetStaticParameterValues(MyParamSet);
 
-			MyParamSet.EditorOnly.StaticSwitchParameters = SourceParamSet.EditorOnly.StaticSwitchParameters;
+			MyParamSet.StaticSwitchParameters = SourceParamSet.StaticSwitchParameters;
 
 			UpdateStaticPermutation(MyParamSet);
 
diff --git a/Engine/Source/Runtime/Engine/Private/Materials/MaterialShader.cpp b/Engine/Source/Runtime/Engine/Private/Materials/MaterialShader.cpp
--- a/Engine/Source/Runtime/Engine/Private/Materials/MaterialShader.cpp
+++ b/Engine/Source/Runtime/Engine/Private/Materials/MaterialShader.cpp
@@ -265,8 +265,8 @@ FStaticParameterSet::FStaticParameterSet(const FStaticParameterSet& InValue) = d
 
 FStaticParameterSet& FStaticParameterSet::operator=(const FStaticParameterSet& InValue)
 {
+	StaticSwitchParameters = InValue.StaticSwitchParameters;
 #if WITH_EDITORONLY_DATA
-	EditorOnly.StaticSwitchParameters = InValue.EditorOnly.StaticSwitchParameters;
 	EditorOnly.StaticComponentMaskParameters = InValue.EditorOnly.StaticComponentMaskParameters;
 	EditorOnly.TerrainLayerWeightParameters = InValue.EditorOnly.TerrainLayerWeightParameters;
 #endif // WITH_EDITORONLY_DATA
@@ -287,7 +287,11 @@ void FStaticParameterSet::PostSerialize(const FArchive& Ar)
 #if WITH_EDITORONLY_DATA
 	if (StaticSwitchParameters_DEPRECATED.Num() > 0)
 	{
-		EditorOnly.StaticSwitchParameters = MoveTemp(StaticSwitchParameters_DEPRECATED);
+		EditorOnly.StaticSwitchParameters_DEPRECATED = MoveTemp(StaticSwitchParameters_DEPRECATED);
+	}
+	if (EditorOnly.StaticSwitchParameters_DEPRECATED.Num() > 0)
+	{
+		StaticSwitchParameters = MoveTemp(EditorOnly.StaticSwitchParameters_DEPRECATED);
 	}
 	if (StaticComponentMaskParameters_DEPRECATED.Num() > 0)
 	{
@@ -318,7 +322,7 @@ void FStaticParameterSet::SerializeLegacy(FArchive& Ar)
 	Ar.UsingCustomVersion(FReleaseObjectVersion::GUID);
 	Ar.UsingCustomVersion(FUE5ReleaseStreamObjectVersion::GUID);
 
-	Ar << EditorOnly.StaticSwitchParameters;
+	Ar << EditorOnly.StaticSwitchParameters_DEPRECATED;
 	Ar << EditorOnly.StaticComponentMaskParameters;
 	Ar << EditorOnly.TerrainLayerWeightParameters;
 
@@ -396,15 +400,15 @@ bool FStaticParameterSet::operator==(const FStaticParameterSet& ReferenceSet) co
 	{
 		return false;
 	}
-#if WITH_EDITORONLY_DATA
-	if (EditorOnly.StaticSwitchParameters.Num() != ReferenceSet.EditorOnly.StaticSwitchParameters.Num()
-		|| EditorOnly.StaticComponentMaskParameters.Num() != ReferenceSet.EditorOnly.StaticComponentMaskParameters.Num()
-		|| EditorOnly.TerrainLayerWeightParameters.Num() != ReferenceSet.EditorOnly.TerrainLayerWeightParameters.Num())
+
+	if(StaticSwitchParameters.Num() != ReferenceSet.StaticSwitchParameters.Num())
 	{
 		return false;
 	}
 
-	if (EditorOnly.StaticSwitchParameters != ReferenceSet.EditorOnly.StaticSwitchParameters)
+#if WITH_EDITORONLY_DATA
+	if (EditorOnly.StaticComponentMaskParameters.Num() != ReferenceSet.EditorOnly.StaticComponentMaskParameters.Num()
+		|| EditorOnly.TerrainLayerWeightParameters.Num() != ReferenceSet.EditorOnly.TerrainLayerWeightParameters.Num())
 	{
 		return false;
 	}
@@ -420,6 +424,11 @@ bool FStaticParameterSet::operator==(const FStaticParameterSet& ReferenceSet) co
 	}
 #endif // WITH_EDITORONLY_DATA
 
+	if (StaticSwitchParameters != ReferenceSet.StaticSwitchParameters)
+	{
+		return false;
+	}
+
 	if (bHasMaterialLayers)
 	{
 		if (MaterialLayers != ReferenceSet.MaterialLayers)
@@ -439,8 +448,8 @@ bool FStaticParameterSet::operator==(const FStaticParameterSet& ReferenceSet) co
 
 void FStaticParameterSet::SortForEquivalent()
 {
+	StaticSwitchParameters.Sort([](const FStaticSwitchParameter& A, const FStaticSwitchParameter& B) { return B.ExpressionGUID < A.ExpressionGUID; });
 #if WITH_EDITORONLY_DATA
-	EditorOnly.StaticSwitchParameters.Sort([](const FStaticSwitchParameter& A, const FStaticSwitchParameter& B) { return B.ExpressionGUID < A.ExpressionGUID; });
 	EditorOnly.StaticComponentMaskParameters.Sort([](const FStaticComponentMaskParameter& A, const FStaticComponentMaskParameter& B) { return B.ExpressionGUID < A.ExpressionGUID; });
 	EditorOnly.TerrainLayerWeightParameters.Sort([](const FStaticTerrainLayerWeightParameter& A, const FStaticTerrainLayerWeightParameter& B) { return B.LayerName.LexicalLess(A.LayerName); });
 #endif // WITH_EDITORONLY_DATA
@@ -453,9 +462,13 @@ bool FStaticParameterSet::Equivalent(const FStaticParameterSet& ReferenceSet) co
 		return false;
 	}
 
+	if(StaticSwitchParameters.Num() != ReferenceSet.StaticSwitchParameters.Num())
+	{
+		return false;
+	}
+
 #if WITH_EDITORONLY_DATA
-	if (EditorOnly.StaticSwitchParameters.Num() != ReferenceSet.EditorOnly.StaticSwitchParameters.Num()
-		|| EditorOnly.StaticComponentMaskParameters.Num() != ReferenceSet.EditorOnly.StaticComponentMaskParameters.Num()
+	if (EditorOnly.StaticComponentMaskParameters.Num() != ReferenceSet.EditorOnly.StaticComponentMaskParameters.Num()
 		|| EditorOnly.TerrainLayerWeightParameters.Num() != ReferenceSet.EditorOnly.TerrainLayerWeightParameters.Num())
 	{
 		return false;
@@ -489,12 +502,12 @@ void FStaticParameterSet::AddParametersOfType(EMaterialParameterType Type, const
 	switch (Type)
 	{
 	case EMaterialParameterType::StaticSwitch:
-		EditorOnly.StaticSwitchParameters.Empty(Values.Num());
+		StaticSwitchParameters.Empty(Values.Num());
 		for (const auto& It : Values)
 		{
 			const FMaterialParameterMetadata& Meta = It.Value;
 			check(Meta.Value.Type == Type);
-			EditorOnly.StaticSwitchParameters.Emplace(It.Key, Meta.Value.AsStaticSwitch(), Meta.bOverride, Meta.ExpressionGuid);
+			StaticSwitchParameters.Emplace(It.Key, Meta.Value.AsStaticSwitch(), Meta.bOverride, Meta.ExpressionGuid);
 		}
 		break;
 	case EMaterialParameterType::StaticComponentMask:
@@ -518,7 +531,7 @@ void FStaticParameterSet::AddParametersOfType(EMaterialParameterType Type, const
 
 void FStaticParameterSet::SetStaticSwitchParameterValue(const FMaterialParameterInfo& ParameterInfo, const FGuid& ExpressionGuid, bool Value)
 {
-	for (FStaticSwitchParameter& Parameter : EditorOnly.StaticSwitchParameters)
+	for (FStaticSwitchParameter& Parameter : StaticSwitchParameters)
 	{
 		if (Parameter.ParameterInfo == ParameterInfo)
 		{
@@ -528,7 +541,7 @@ void FStaticParameterSet::SetStaticSwitchParameterValue(const FMaterialParameter
 		}
 	}
 
-	new(EditorOnly.StaticSwitchParameters) FStaticSwitchParameter(ParameterInfo, Value, true, ExpressionGuid);
+	new(StaticSwitchParameters) FStaticSwitchParameter(ParameterInfo, Value, true, ExpressionGuid);
 }
 
 void FStaticParameterSet::SetStaticComponentMaskParameterValue(const FMaterialParameterInfo& ParameterInfo, const FGuid& ExpressionGuid, bool R, bool G, bool B, bool A)
@@ -968,7 +981,7 @@ void FMaterialShaderMapId::UpdateFromParameterSet(const FStaticParameterSet& Sta
 		}
 	};
 
-	StaticSwitchParameters = StaticParameters.EditorOnly.StaticSwitchParameters;
+	StaticSwitchParameters = StaticParameters.StaticSwitchParameters;
 	StaticComponentMaskParameters = StaticParameters.EditorOnly.StaticComponentMaskParameters;
 	TerrainLayerWeightParameters = StaticParameters.EditorOnly.TerrainLayerWeightParameters;
 	if (StaticParameters.bHasMaterialLayers)
diff --git a/Engine/Source/Runtime/Engine/Public/StaticParameterSet.h b/Engine/Source/Runtime/Engine/Public/StaticParameterSet.h
--- a/Engine/Source/Runtime/Engine/Public/StaticParameterSet.h
+++ b/Engine/Source/Runtime/Engine/Public/StaticParameterSet.h
@@ -385,6 +385,10 @@ struct FStaticParameterSetRuntimeData
 	FStaticParameterSetRuntimeData& operator=(const FStaticParameterSetRuntimeData& Rhs) = default;
 	FStaticParameterSetRuntimeData& operator=(const FStaticParameterSet& Rhs) = delete;
 
+	/** An array of static switch parameters in this set */
+	UPROPERTY()
+	TArray<FStaticSwitchParameter> StaticSwitchParameters;
+
 	/** Material layers for this set */
 	UPROPERTY()
 	FMaterialLayersFunctionsRuntimeData MaterialLayers;
@@ -394,13 +398,14 @@ struct FStaticParameterSetRuntimeData
 
 	void Empty()
 	{
+		StaticSwitchParameters.Empty();
 		MaterialLayers.Empty();
 		bHasMaterialLayers = false;
 	}
 
 	bool IsEmpty() const
 	{
-		return !bHasMaterialLayers;
+		return  StaticSwitchParameters.Num() == 0 && !bHasMaterialLayers;
 	}
 };
 
@@ -411,7 +416,7 @@ struct FStaticParameterSetEditorOnlyData
 
 	/** An array of static switch parameters in this set */
 	UPROPERTY()
-	TArray<FStaticSwitchParameter> StaticSwitchParameters;
+	TArray<FStaticSwitchParameter> StaticSwitchParameters_DEPRECATED;
 
 	/** An array of static component mask parameters in this set */
 	UPROPERTY()
@@ -428,7 +433,6 @@ struct FStaticParameterSetEditorOnlyData
 #if WITH_EDITOR
 	void Empty()
 	{
-		StaticSwitchParameters.Empty();
 		StaticComponentMaskParameters.Empty();
 		TerrainLayerWeightParameters.Empty();
 		MaterialLayers.Empty();
@@ -436,7 +440,7 @@ struct FStaticParameterSetEditorOnlyData
 
 	bool IsEmpty() const
 	{
-		return StaticSwitchParameters.Num() == 0 && StaticComponentMaskParameters.Num() == 0 && TerrainLayerWeightParameters.Num() == 0;
+		return StaticComponentMaskParameters.Num() == 0 && TerrainLayerWeightParameters.Num() == 0;
 	}
 #endif // WITH_EDITOR
 };
