UE-162245 - Level asset doesn't duplicate correctly with Event Dispatcher in Level BP

When duplicating a level with an event dispatcher, the signature function was actually pointing to the original UFunction that belonged to the UClass in the level that we were copying from. As a result, a duplicated level would have object references outside of its current package, which prevented the editor from saving it.

The fix updates ClassGeneratedBy to point to the duplicated ULevelScriptBlueprint. Otherwise, FMemberReference might resolve to the class/function that lives in the source level.  Also note that standard duplication will not perform a deep copy of this field since it's outered to a different object. Similarly, we need to regenerate the CDO for our new class to prevent any save errors.

diff --git a/Engine/Source/Runtime/Engine/Private/World.cpp b/Engine/Source/Runtime/Engine/Private/World.cpp
--- a/Engine/Source/Runtime/Engine/Private/World.cpp
+++ b/Engine/Source/Runtime/Engine/Private/World.cpp
@@ -908,20 +908,43 @@ void UWorld::PostDuplicate(bool bDuplicateForPIE)
 		UBlueprint* LevelScriptBlueprint = PersistentLevel->GetLevelScriptBlueprint(bDontCreate);
 		if (LevelScriptBlueprint)
 		{
+			// Duplicating a UClass shallow copies the ClassGeneratedBy field and the CDO.
+			// This is problematic because we might accidentally reference these objects,
+			// which prevents saving of the level since these objects are external
+			// (eg: event dispatchers in a level blueprint can trigger this scenario).
+			auto FixupGeneratedClass = [LevelScriptBlueprint](UObject* NewClass)
+			{
+				if (UClass* Class = Cast<UClass>(NewClass))
+				{
+					// We can likely work around a deep-copied CDO, but it's also unexpected in this scenario.
+					ensureMsgf(!Class->ClassDefaultObject || Class->ClassDefaultObject->GetOuter() != Class, TEXT("Expecting shallow copy of CDO"));
+
+					Class->ClassGeneratedBy = LevelScriptBlueprint;
+					Class->StaticLink(true);
+					Class->ClassDefaultObject = nullptr;
+					Class->GetDefaultObject();
+				}
+			};
+
 			UObject* OldGeneratedClass = LevelScriptBlueprint->GeneratedClass;
 			if (OldGeneratedClass)
 			{
 				UObject* NewGeneratedClass = StaticDuplicateObject(OldGeneratedClass, MyPackage, OldGeneratedClass->GetFName());
+				FixupGeneratedClass(NewGeneratedClass);
+
 				ReplacementMap.Add(OldGeneratedClass, NewGeneratedClass);
 
 				// The class may have referenced a lightmap or landscape resource that is also being duplicated. Add it to the list of objects that need references fixed up.
 				ObjectsToFixReferences.Add(NewGeneratedClass);
+
 			}
 
 			UObject* OldSkeletonClass = LevelScriptBlueprint->SkeletonGeneratedClass;
 			if (OldSkeletonClass)
 			{
 				UObject* NewSkeletonClass = StaticDuplicateObject(OldSkeletonClass, MyPackage, OldSkeletonClass->GetFName());
+				FixupGeneratedClass(NewSkeletonClass);
+
 				ReplacementMap.Add(OldSkeletonClass, NewSkeletonClass);
 
 				// The class may have referenced a lightmap or landscape resource that is also being duplicated. Add it to the list of objects that need references fixed up.
