Fixed RecursiveGetExpressionChain for Make and Set Material Attributes, in the case where we're looking for the MP_MaterialAttributes property : this property needs to be followed up since it's actually a group of several properties.
Also fixed potential issues with mobile material expressions not being properly filtered

diff --git a/Engine/Source/Runtime/Engine/Classes/Materials/Material.h b/Engine/Source/Runtime/Engine/Classes/Materials/Material.h
--- a/Engine/Source/Runtime/Engine/Classes/Materials/Material.h
+++ b/Engine/Source/Runtime/Engine/Classes/Materials/Material.h
@@ -1760,6 +1760,9 @@ public:
 	// Return true if the property is supported
 	ENGINE_API bool IsPropertySupported(EMaterialProperty InProperty) const;
 
+	/** Returns true if the material property is relevant for mobile platforms */
+	ENGINE_API static bool IsPropertyRelevantForMobile(EMaterialProperty InProperty);
+
 #if WITH_EDITOR
 	static void NotifyCompilationFinished(UMaterialInterface* Material);
 
diff --git a/Engine/Source/Runtime/Engine/Private/Materials/Material.cpp b/Engine/Source/Runtime/Engine/Private/Materials/Material.cpp
--- a/Engine/Source/Runtime/Engine/Private/Materials/Material.cpp
+++ b/Engine/Source/Runtime/Engine/Private/Materials/Material.cpp
@@ -93,6 +93,7 @@
 #include "Misc/OutputDeviceFile.h"
 #include "HAL/FileManager.h"
 #include "BuildSettings.h"
+#include "Algo/Transform.h"
 
 #if WITH_EDITOR
 #include "MaterialCachedHLSLTree.h"
@@ -141,6 +142,24 @@ bool Engine_IsStrataEnabled()
 	return CVar && CVar->GetValueOnAnyThread() > 0;
 }
 
+namespace MaterialImpl
+{
+	// Filtered list of properties that we follow on mobile platforms.
+	struct FMobileProperty
+	{
+		EMaterialProperty Property;
+		EShaderFrequency Frequency;
+	};
+	static constexpr FMobileProperty GMobileRelevantMaterialProperties[] =
+	{
+		{ MP_WorldPositionOffset, SF_Vertex },
+		{ MP_EmissiveColor, SF_Pixel },
+		{ MP_BaseColor, SF_Pixel },
+		{ MP_Normal, SF_Pixel },
+		{ MP_OpacityMask, SF_Pixel },
+	};
+}
+
 #if WITH_EDITOR
 const FMaterialsWithDirtyUsageFlags FMaterialsWithDirtyUsageFlags::DefaultAnnotation;
 
@@ -5501,6 +5520,8 @@ void UMaterial::GetAllExpressionsForCustomInterpolators(TArray<class UMaterialEx
 bool UMaterial::GetAllReferencedExpressions(TArray<UMaterialExpression*>& OutExpressions, struct FStaticParameterSet* InStaticParameterSet,
 	ERHIFeatureLevel::Type InFeatureLevel, EMaterialQualityLevel::Type InQuality, ERHIShadingPath::Type InShadingPath)
 {
+	using namespace MaterialImpl;
+
 	OutExpressions.Empty();
 
 	// For mobile only consider nodes connected to material properties that affect mobile, and exclude any custom outputs
@@ -5519,18 +5540,9 @@ bool UMaterial::GetAllReferencedExpressions(TArray<UMaterialExpression*>& OutExp
 		}
 		else
 		{
-			EMaterialProperty MobileRelevantMaterialProperties[] =
-			{
-				MP_EmissiveColor,
-				MP_OpacityMask,
-				MP_BaseColor,
-				MP_Normal,
-				MP_WorldPositionOffset,
-			};
-
-
-			for (EMaterialProperty MaterialProp : MobileRelevantMaterialProperties)
+			for (FMobileProperty MobileProperty : GMobileRelevantMaterialProperties)
 			{
+				EMaterialProperty MaterialProp = MobileProperty.Property;
 				TArray<UMaterialExpression*> MPRefdExpressions;
 				if (GetExpressionsInPropertyChain(MaterialProp, MPRefdExpressions, InStaticParameterSet, InFeatureLevel, InQuality, InShadingPath) == true)
 				{
@@ -5689,6 +5701,8 @@ bool UMaterial::RecursiveGetExpressionChain(
 	TArray<UMaterialExpression*>& OutExpressions, struct FStaticParameterSet* InStaticParameterSet,
 	ERHIFeatureLevel::Type InFeatureLevel, EMaterialQualityLevel::Type InQuality, ERHIShadingPath::Type InShadingPath, EShaderFrequency InShaderFrequency, EMaterialProperty InProperty)
 {
+	using namespace MaterialImpl;
+
 	OutExpressions.AddUnique(InExpression);
 	TArray<FExpressionInput*> Inputs;
 	TArray<EShaderFrequency> InputsFrequency;
@@ -5700,6 +5714,8 @@ bool UMaterial::RecursiveGetExpressionChain(
 	UMaterialExpressionSetMaterialAttributes* SetMaterialAttributesExp;
 	UMaterialExpressionShaderStageSwitch* ShaderStageSwitchExp;
 
+	const bool bMobileOnly = InFeatureLevel <= ERHIFeatureLevel::ES3_1;
+
 	if (InFeatureLevel != ERHIFeatureLevel::Num && (FeatureLevelSwitchExp = Cast<UMaterialExpressionFeatureLevelSwitch>(InExpression)) != nullptr)
 	{
 		if (FeatureLevelSwitchExp->Inputs[InFeatureLevel].IsConnected())
@@ -5754,81 +5770,98 @@ bool UMaterial::RecursiveGetExpressionChain(
 	}
 	else if ((MakeMaterialAttributesExp = Cast<UMaterialExpressionMakeMaterialAttributes>(InExpression)) != nullptr)
 	{
-		// Filtered list of properties that we follow on mobile platforms.
-		struct FMobileProperty 
-		{
-			EMaterialProperty Property;
-			EShaderFrequency Frequency;
-		};
-		static const FMobileProperty MobileProperties[] =
+		// MP_MAX means we want to follow all properties
+		// MP_MaterialAttributes means we want to follow properties that are part of the material attributes, which, in the case of UMaterialExpressionMakeMaterialAttributes, means: every one of them
+		if ((InProperty == MP_MAX) || (InProperty == MP_MaterialAttributes))
 		{
-			{ MP_WorldPositionOffset, SF_Vertex },
-			{ MP_EmissiveColor, SF_Pixel },
-			{ MP_BaseColor, SF_Pixel },
-			{ MP_Normal, SF_Pixel },
-			{ MP_OpacityMask, SF_Pixel },
-		};
-
-		const bool bMobileOnly = InFeatureLevel <= ERHIFeatureLevel::ES3_1;
-		
-		if (bMobileOnly)
-		{
-			// Follow mobile properties (and only InProperty if that is specified).
-			for (FMobileProperty MobileProperty : MobileProperties)
+			// If we're requested to return only the properties relevant for mobile, only follow up those: 
+			if (bMobileOnly)
 			{
-				if (InProperty == MP_MAX || InProperty == MobileProperty.Property)
+				for (FMobileProperty MobileProperty : GMobileRelevantMaterialProperties)
 				{
-					Inputs.Add(MakeMaterialAttributesExp->GetExpressionInput(MobileProperty.Property));
-					InputsFrequency.Add(MobileProperty.Frequency);
+					FExpressionInput* Input = MakeMaterialAttributesExp->GetExpressionInput(MobileProperty.Property);
+					if (Input != nullptr)
+					{
+						Inputs.Add(Input);
+						InputsFrequency.Add(MobileProperty.Frequency);
+					}
 				}
 			}
+			else
+			{
+				// Follow all properties.
+				Inputs = InExpression->GetInputs();
+				InputsFrequency.Init(InShaderFrequency, Inputs.Num());
+			}
 		}
-		else if (InProperty != MP_MAX)
+		else
 		{
 			// Only follow the specified InProperty.
 			FExpressionInput* Input = MakeMaterialAttributesExp->GetExpressionInput(InProperty);
-			if (Input != nullptr)
+			if ((Input != nullptr) 
+				&& (!bMobileOnly || IsPropertyRelevantForMobile(InProperty)))
 			{
 				Inputs.Add(Input);
 				InputsFrequency.Add(InShaderFrequency);
 			}
 		}
-		else
-		{
-			// Follow all properties.
-			Inputs = InExpression->GetInputs();
-			InputsFrequency.Init(InShaderFrequency, Inputs.Num());
-		}
 	}
 	else if ((SetMaterialAttributesExp = Cast<UMaterialExpressionSetMaterialAttributes>(InExpression)) != nullptr)
 	{
-		if (InProperty != MP_MAX)
+		checkf(!SetMaterialAttributesExp->GetInputs().IsEmpty() && (SetMaterialAttributesExp->GetInputType(0) == MCT_MaterialAttributes), TEXT("There must always be one input at least : the material attribute pin"));
+		// Always add the material attribute input, so that we keep on traversing up the property chain : 
+		Inputs.Add(SetMaterialAttributesExp->GetInput(0));
+		InputsFrequency.Add(InShaderFrequency);
+
+		TArray<EMaterialProperty> AttributeProperties;
+		AttributeProperties.Reserve(SetMaterialAttributesExp->AttributeSetTypes.Num());
+		Algo::Transform(SetMaterialAttributesExp->AttributeSetTypes, AttributeProperties, [](const FGuid& InAttributeID) { return FMaterialAttributeDefinitionMap::GetProperty(InAttributeID); });
+		
+		// MP_MAX means we want to follow all properties
+		// MP_MaterialAttributes means we want to follow properties that are part of the material attributes, which, in the case of UMaterialExpressionSetMaterialAttributes, means: every one of them
+		if ((InProperty == MP_MAX) || (InProperty == MP_MaterialAttributes))
 		{
-			for (int32 AttributeIndex = 0; AttributeIndex < SetMaterialAttributesExp->AttributeSetTypes.Num(); ++AttributeIndex)
+			// If we're requested to return only the properties relevant for mobile, only follow up those: 
+			if (bMobileOnly)
 			{
-				if (FMaterialAttributeDefinitionMap::GetProperty(SetMaterialAttributesExp->AttributeSetTypes[AttributeIndex]) == InProperty)
+				for (FMobileProperty MobileProperty : GMobileRelevantMaterialProperties)
 				{
-					Inputs.Add(SetMaterialAttributesExp->GetInput(AttributeIndex + 1)); // Need +1 for MaterialAttributes input pin.
-					InputsFrequency.Add(InShaderFrequency);
-					break;
+					int32 FoundAttributeIndex = INDEX_NONE;
+					if (AttributeProperties.Find(MobileProperty.Property, FoundAttributeIndex))
+					{
+						FExpressionInput* AttributeInput = SetMaterialAttributesExp->GetInput(FoundAttributeIndex + 1); // Need +1 for MaterialAttributes input pin.
+						Inputs.Add(AttributeInput);
+						InputsFrequency.Add(MobileProperty.Frequency);
+					}
 				}
 			}
+			else
+			{
+				// Follow all properties.
+				Inputs = InExpression->GetInputs();
+				InputsFrequency.Init(InShaderFrequency, Inputs.Num());
+			}
 		}
 		else
 		{
-			// Follow all properties.
-			Inputs = InExpression->GetInputs();
-			InputsFrequency.Init(InShaderFrequency, Inputs.Num());
+			// Only follow the specified InProperty.
+			int32 FoundAttributeIndex = INDEX_NONE;
+			if (AttributeProperties.Find(InProperty, FoundAttributeIndex))
+			{
+				FExpressionInput* AttributeInput = SetMaterialAttributesExp->GetInput(FoundAttributeIndex + 1); // Need +1 for MaterialAttributes input pin.
+				if ((AttributeInput != nullptr)
+					&& (!bMobileOnly || IsPropertyRelevantForMobile(InProperty)))
+				{ 
+					Inputs.Add(AttributeInput);
+					InputsFrequency.Add(InShaderFrequency);
+				}
+			}
 		}
 	}
 	else
 	{
 		Inputs = InExpression->GetInputs();
-		
-		for (FExpressionInput* Input : Inputs)
-		{
-			InputsFrequency.Add(InShaderFrequency);
-		}
+		InputsFrequency.Init(InShaderFrequency, Inputs.Num());
 	}
 
 	check(Inputs.Num() == InputsFrequency.Num());
@@ -6275,6 +6308,20 @@ bool UMaterial::IsPropertySupported(EMaterialProperty InProperty) const
 	return bSupported;
 }
 
+bool UMaterial::IsPropertyRelevantForMobile(EMaterialProperty InProperty)
+{
+	using namespace MaterialImpl;
+
+	for (FMobileProperty MobileProperty : GMobileRelevantMaterialProperties)
+	{
+		if (MobileProperty.Property == InProperty)
+		{
+			return true;
+		}
+	}
+	return false;
+}
+
 static bool IsPropertyActive_Internal(EMaterialProperty InProperty,
 	EMaterialDomain Domain,
 	EBlendMode BlendMode,
