Fix ReplaceActor not copying ContentBundleGuid to new actor causing invalid content bundle id

diff --git a/Engine/Source/Runtime/Engine/Private/ActorEditor.cpp b/Engine/Source/Runtime/Engine/Private/ActorEditor.cpp
--- a/Engine/Source/Runtime/Engine/Private/ActorEditor.cpp
+++ b/Engine/Source/Runtime/Engine/Private/ActorEditor.cpp
@@ -1409,6 +1409,8 @@ void AActor::SetFolderPath_Recursively(const FName& NewFolderPath)
 // Ideally, this should be revisited to implement something more generic.
 void AActor::EditorReplacedActor(AActor* OldActor)
 {
+	ContentBundleGuid = OldActor->ContentBundleGuid;
+
 	SetActorLabel(OldActor->GetActorLabel());
 	Tags = OldActor->Tags;
 	
