Level streaming delegates that can be bound globally rather than on a specific ULevelStreaming object.

diff --git a/Engine/Source/Runtime/Engine/Classes/Engine/LevelStreaming.h b/Engine/Source/Runtime/Engine/Classes/Engine/LevelStreaming.h
--- a/Engine/Source/Runtime/Engine/Classes/Engine/LevelStreaming.h
+++ b/Engine/Source/Runtime/Engine/Classes/Engine/LevelStreaming.h
@@ -261,14 +261,13 @@ protected:
 	/** Whether the streaming level can safely skip making visible transaction request from the client to the server */
 	uint8 bSkipClientUseMakingVisibleTransactionRequest:1;
 
+private:
 	/** What the current streamed state of the streaming level is */
 	ELevelStreamingState CurrentState;
 
 	/** What streamed state the streaming level is transitioning towards */
 	ELevelStreamingTargetState TargetState;
 
-	void SetCurrentState(ELevelStreamingState NewState);
-
 public:
 
 	/** Whether this level is locked; that is, its actors are read-only. */
@@ -347,6 +346,10 @@ PRAGMA_ENABLE_DEPRECATION_WARNINGS
 
 	ELevelStreamingState GetLevelStreamingState() const { return CurrentState; }
 
+protected:	
+	/** Updates the current state of the streaming level and notifies any observers. */
+	void SetCurrentState(ELevelStreamingState NewState);
+
 private:
 
 	/** If true client will wait for acknowledgment from server before making streaming levels invisible */
@@ -362,8 +365,11 @@ private:
 	/** Returns whether the streaming level can make invisible (can call RemoveFromWorld). */
 	bool CanMakeInvisible();
 
-	/** Determine what the streaming levels target state should be. Returns whether the streaming level should be in the consider list. */
-	bool DetermineTargetState();
+	/** Determine what the streaming level's target state should be. */
+	ELevelStreamingTargetState DetermineTargetState() const;
+
+	/** Determines a new target state, fires delegates, returns true if the level should continue to be considered for streaming. */
+	bool UpdateTargetState();
 
 	/** Update the load process of the streaming level. Out parameters instruct calling code how to proceed. */
 	void UpdateStreamingState(bool& bOutUpdateAgain, bool& bOutRedetermineTarget);
@@ -730,7 +736,7 @@ private:
 	/** Update internal variables when the level is removed from the streaming levels array */
 	static void OnLevelRemoved(ULevelStreaming* StreamingLevel) { StreamingLevel->OnLevelRemoved(); }
 	/** Determine what the streaming levels target state should be. Returns whether the streaming level should be in the consider list. */
-	static bool DetermineTargetState(ULevelStreaming* StreamingLevel) { return StreamingLevel->DetermineTargetState(); }
+	static bool UpdateTargetState(ULevelStreaming* StreamingLevel) { return StreamingLevel->UpdateTargetState(); }
 	/** Update the load process of the streaming level. Out parameters instruct calling code how to proceed. */
 	static void UpdateStreamingState(ULevelStreaming* StreamingLevel, bool& bOutUpdateAgain, bool& bOutRedetermineTarget) { StreamingLevel->UpdateStreamingState(bOutUpdateAgain, bOutRedetermineTarget); }
 
diff --git a/Engine/Source/Runtime/Engine/Classes/Engine/World.h b/Engine/Source/Runtime/Engine/Classes/Engine/World.h
--- a/Engine/Source/Runtime/Engine/Classes/Engine/World.h
+++ b/Engine/Source/Runtime/Engine/Classes/Engine/World.h
@@ -2868,16 +2868,18 @@ public:
 	 * @param LevelTransform	Transformation to apply to each actor in the level
 	 * @param bConsiderTimeLimie optional bool indicating if we should consider timelimit or not, default is true
 	 * @param TransactionId optional parameter that carries the current transaction id associated with calls updating LevelVisibility used when communicating level visibility with server
+	 * @param OwningLevelStreaming optional parameter, the ULevelStreaming object driving this level's presence in the world
 	 */
-	void AddToWorld(ULevel* Level, const FTransform& LevelTransform = FTransform::Identity, bool bConsiderTimeLimit = true, FNetLevelVisibilityTransactionId TransactionId = FNetLevelVisibilityTransactionId());
+	void AddToWorld(ULevel* Level, const FTransform& LevelTransform = FTransform::Identity, bool bConsiderTimeLimit = true, FNetLevelVisibilityTransactionId TransactionId = FNetLevelVisibilityTransactionId(), ULevelStreaming* OwningLevelStreaming = nullptr);
 
 	/** 
 	 * Dissociates the passed in level from the world. The removal is blocking.
 	 *
 	 * @param Level			Level object we should remove
 	 * @param TransactionId optional parameter that carries the current transaction id associated with calls updating LevelVisibility used when communicating level visibility with server
+	 * @param OwningLevelStreaming optional parameter, the ULevelStreaming object driving this level's presence in the world
 	 */
-	void RemoveFromWorld(ULevel* Level, bool bAllowIncrementalRemoval = false, FNetLevelVisibilityTransactionId TransactionId = FNetLevelVisibilityTransactionId());
+	void RemoveFromWorld(ULevel* Level, bool bAllowIncrementalRemoval = false, FNetLevelVisibilityTransactionId TransactionId = FNetLevelVisibilityTransactionId(), ULevelStreaming* OwningLevelStreaming = nullptr);
 
 	/**
 	 * Updates sub-levels (load/unload/show/hide) using streaming levels current state
diff --git a/Engine/Source/Runtime/Engine/Private/LevelStreaming.cpp b/Engine/Source/Runtime/Engine/Private/LevelStreaming.cpp
--- a/Engine/Source/Runtime/Engine/Private/LevelStreaming.cpp
+++ b/Engine/Source/Runtime/Engine/Private/LevelStreaming.cpp
@@ -33,6 +33,7 @@
 #include "Engine/NetConnection.h"
 #include "Engine/PackageMapClient.h"
 #include "Serialization/LoadTimeTrace.h"
+#include "Streaming/LevelStreamingDelegates.h"
 #include "AssetRegistry/AssetRegistryModule.h"
 
 #include UE_INLINE_GENERATED_CPP_BY_NAME(LevelStreaming)
@@ -505,12 +506,21 @@ void ULevelStreaming::OnLevelRemoved()
 
 void ULevelStreaming::SetCurrentState(ELevelStreamingState NewState)
 {
+	// TODO: We should only fire the delegate when the current state has changed, but first AsyncLevelLoadComplete needs to be fixed to 
+	// only set the new state once the LoadedLevel is assigned. Clients currently rely on getting a repeated notification after the loaded
+	// level is available.
+	ELevelStreamingState OldState = CurrentState;
 	CurrentState = NewState;
+	FLevelStreamingDelegates::OnLevelStreamingStateChanged.Broadcast(GetWorld(), this, GetLoadedLevel(), OldState, NewState);
 }
 
-bool ULevelStreaming::DetermineTargetState()
+bool ULevelStreaming::UpdateTargetState()
 {
 	FScopeCycleCounterUObject ContextScope(this);
+	if (CurrentState == ELevelStreamingState::FailedToLoad || CurrentState == ELevelStreamingState::Removed)
+	{
+		return false;
+	}
 
 	UWorld* World = GetWorld();
 
@@ -524,9 +534,47 @@ bool ULevelStreaming::DetermineTargetState()
 		}
 	}
 #endif
-	
-	bool bContinueToConsider = true;
 
+	ELevelStreamingTargetState NewTarget = DetermineTargetState();
+
+	if (TargetState != NewTarget)
+	{
+		ELevelStreamingTargetState OldTarget = TargetState;
+		TargetState = NewTarget;
+		FLevelStreamingDelegates::OnLevelStreamingTargetStateChanged.Broadcast(
+			GetWorld(), 
+			this, 
+			GetLoadedLevel(), 
+			CurrentState,
+			OldTarget,
+			NewTarget
+			);
+	}
+
+	// Return whether the level should continue to be considered
+	switch (CurrentState)
+	{
+	case ELevelStreamingState::MakingVisible:
+	case ELevelStreamingState::MakingInvisible:
+	case ELevelStreamingState::Loading:
+		// Always continue to consider if we are actively working on a state change 
+		return true;
+	case ELevelStreamingState::Unloaded:
+		return NewTarget != ELevelStreamingTargetState::Unloaded;
+	case ELevelStreamingState::LoadedNotVisible:
+		return NewTarget != ELevelStreamingTargetState::LoadedNotVisible || !IsDesiredLevelLoaded();
+	case ELevelStreamingState::LoadedVisible:
+		return NewTarget != ELevelStreamingTargetState::LoadedVisible || !IsDesiredLevelLoaded();
+	case ELevelStreamingState::FailedToLoad:
+	case ELevelStreamingState::Removed:
+	default:
+		return false;
+	}
+}
+
+ELevelStreamingTargetState ULevelStreaming::DetermineTargetState() const
+{
+	UWorld* World = GetWorld();
 	switch(CurrentState)
 	{
 	case ELevelStreamingState::MakingVisible:
@@ -536,109 +584,96 @@ bool ULevelStreaming::DetermineTargetState()
 			// Since level doesn't need to be visible anymore, change TargetState to ELevelStreamingTargetState::LoadedNotVisible.
 			// Next UpdateStreamingState will handle switching CurrentState to ELevelStreamingState::LoadedNotVisible.
 			// From there, regular flow will properly handle TargetState.
-			TargetState = ELevelStreamingTargetState::LoadedNotVisible;
+			return ELevelStreamingTargetState::LoadedNotVisible;
 		}
 		else
 		{
-			TargetState = ELevelStreamingTargetState::LoadedVisible;
+			return ELevelStreamingTargetState::LoadedVisible;
 		}
-		break;
 
 	case ELevelStreamingState::MakingInvisible:
 		ensure(LoadedLevel);
-		TargetState = ELevelStreamingTargetState::LoadedNotVisible;
-		break;
+		return ELevelStreamingTargetState::LoadedNotVisible;
 
 	case ELevelStreamingState::Loading:
-		TargetState = ELevelStreamingTargetState::LoadedNotVisible;
-		break;
+		return ELevelStreamingTargetState::LoadedNotVisible;
 
 	case ELevelStreamingState::Unloaded:
 		if (bIsRequestingUnloadAndRemoval)
 		{
-			TargetState = ELevelStreamingTargetState::UnloadedAndRemoved;
+			return ELevelStreamingTargetState::UnloadedAndRemoved;
 		}
 		else if (World->GetShouldForceUnloadStreamingLevels())
 		{
-			bContinueToConsider = false;
+			return ELevelStreamingTargetState::Unloaded;
 		}
 		else if (!World->IsGameWorld())
 		{
-			TargetState = ELevelStreamingTargetState::LoadedNotVisible;
+			return ELevelStreamingTargetState::LoadedNotVisible;
 		}
 		else if (ShouldBeLoaded())
 		{
-			TargetState = ELevelStreamingTargetState::LoadedNotVisible;
+			return ELevelStreamingTargetState::LoadedNotVisible;
 		}
 		else
 		{
-			bContinueToConsider = false;
+			return ELevelStreamingTargetState::Unloaded;
 		}
-		break;
 
 	case ELevelStreamingState::LoadedNotVisible:
 		if (bIsRequestingUnloadAndRemoval || World->GetShouldForceUnloadStreamingLevels())
 		{
-			TargetState = ELevelStreamingTargetState::Unloaded;
+			return ELevelStreamingTargetState::Unloaded;
 		}
 		else if (World->IsGameWorld() && !ShouldBeLoaded())
 		{
-			TargetState = ELevelStreamingTargetState::Unloaded;
+			return ELevelStreamingTargetState::Unloaded;
 		}
 		else if (!IsDesiredLevelLoaded())
 		{
-			TargetState = ELevelStreamingTargetState::LoadedNotVisible;
+			return ELevelStreamingTargetState::LoadedNotVisible;
 		}
 		else if (ShouldBeVisible())
 		{
-			TargetState = ELevelStreamingTargetState::LoadedVisible;
+			return ELevelStreamingTargetState::LoadedVisible;
 		}
 		else
 		{
-			bContinueToConsider = false;
+			return ELevelStreamingTargetState::LoadedNotVisible;
 		}
-		break;
 
 	case ELevelStreamingState::LoadedVisible:
 		if (bIsRequestingUnloadAndRemoval || World->GetShouldForceUnloadStreamingLevels())
 		{
-			TargetState = ELevelStreamingTargetState::LoadedNotVisible;
+			return ELevelStreamingTargetState::LoadedNotVisible;
 		}
 		else if (World->IsGameWorld() && !ShouldBeLoaded())
 		{
-			TargetState = ELevelStreamingTargetState::LoadedNotVisible;
+			return ELevelStreamingTargetState::LoadedNotVisible;
 		}
 		else if (!ShouldBeVisible())
 		{
-			TargetState = ELevelStreamingTargetState::LoadedNotVisible;
+			return ELevelStreamingTargetState::LoadedNotVisible;
 		}
 		else if (!IsDesiredLevelLoaded())
 		{
-			TargetState = ELevelStreamingTargetState::LoadedVisible;
+			return ELevelStreamingTargetState::LoadedVisible;
 		}
 		else
 		{
-			bContinueToConsider = false;
+			return ELevelStreamingTargetState::LoadedVisible;
 		}
-		break;
 
 	case ELevelStreamingState::FailedToLoad:
-
-		// Anything that affects whether we might try to reload changes current state itself
-		bContinueToConsider = false;
-		break;
+		return ELevelStreamingTargetState::Unloaded;
 
 	case ELevelStreamingState::Removed:
-
-		// Never continue to consider a removed streaming level
-		bContinueToConsider = false;
-		break;
+		return ELevelStreamingTargetState::Unloaded;
 
 	default:
 		ensure(false);
+		return ELevelStreamingTargetState::Unloaded;
 	}
-
-	return bContinueToConsider;
 }
 
 bool ULevelStreaming::IsConcernedByNetVisibilityTransactionAck() const
@@ -880,7 +915,7 @@ void ULevelStreaming::UpdateStreamingState(bool& bOutUpdateAgain, bool& bOutRede
 					break;
 				}
 
-				World->AddToWorld(LoadedLevel, LevelTransform, !bShouldBlockOnLoad, TransactionId);
+				World->AddToWorld(LoadedLevel, LevelTransform, !bShouldBlockOnLoad, TransactionId, this);
 
 				if (LoadedLevel->bIsVisible)
 				{
@@ -930,7 +965,7 @@ void ULevelStreaming::UpdateStreamingState(bool& bOutUpdateAgain, bool& bOutRede
 			}
 
 			// Hide loaded level, incrementally if necessary
-			World->RemoveFromWorld(LoadedLevel, !bShouldBlockOnUnload && World->IsGameWorld(), TransactionId);
+			World->RemoveFromWorld(LoadedLevel, !bShouldBlockOnUnload && World->IsGameWorld(), TransactionId, this);
 
 			// Hide loaded level immediately if bRequireFullVisibilityToRender is set
 			const bool LevelBecameInvisible = bWasVisible && !LoadedLevel->bIsVisible;
@@ -1602,6 +1637,7 @@ bool ULevelStreaming::RequestLevel(UWorld* PersistentWorld, bool bAllowLevelLoad
 
 void ULevelStreaming::AsyncLevelLoadComplete(const FName& InPackageName, UPackage* InLoadedPackage, EAsyncLoadingResult::Type Result)
 {
+	// TODO: Should not set state here so that observers will have access to LoadedLevel
 	SetCurrentState(ELevelStreamingState::LoadedNotVisible);
 	if (UWorld* World = GetWorld())
 	{
diff --git a/Engine/Source/Runtime/Engine/Private/Streaming/LevelStreamingDelegates.cpp b/Engine/Source/Runtime/Engine/Private/Streaming/LevelStreamingDelegates.cpp
new file mode 100644
--- /dev/null
+++ b/Engine/Source/Runtime/Engine/Private/Streaming/LevelStreamingDelegates.cpp
@@ -0,0 +1,23 @@
+// Copyright Epic Games, Inc. All Rights Reserved.
+
+#include "Streaming/LevelStreamingDelegates.h"
+
+TMulticastDelegate<void(
+	UWorld*, 
+	const ULevelStreaming*, 
+	ULevel* LevelIfLoaded,
+	ELevelStreamingState Current,
+	ELevelStreamingTargetState PrevTarget, 
+	ELevelStreamingTargetState NewTarget)> FLevelStreamingDelegates::OnLevelStreamingTargetStateChanged;
+
+TMulticastDelegate<void(
+	UWorld*, 
+	const ULevelStreaming*,
+	ULevel* LevelIfLoaded,
+	ELevelStreamingState PreviousState,
+	ELevelStreamingState NewState
+	)> FLevelStreamingDelegates::OnLevelStreamingStateChanged;
+
+TMulticastDelegate<void(UWorld*, const ULevelStreaming*, ULevel* LoadedLevel)> FLevelStreamingDelegates::OnLevelBeginMakingVisible;
+TMulticastDelegate<void(UWorld*, const ULevelStreaming*, ULevel* LoadedLevel)> FLevelStreamingDelegates::OnLevelBeginMakingInvisible;
+
diff --git a/Engine/Source/Runtime/Engine/Private/World.cpp b/Engine/Source/Runtime/Engine/Private/World.cpp
--- a/Engine/Source/Runtime/Engine/Private/World.cpp
+++ b/Engine/Source/Runtime/Engine/Private/World.cpp
@@ -97,6 +97,7 @@
 #include "ProfilingDebugging/LoadTimeTracker.h"
 #include "ProfilingDebugging/CsvProfiler.h"
 #include "Streaming/ServerStreamingLevelsVisibility.h"
+#include "Streaming/LevelStreamingDelegates.h"
 
 #if WITH_EDITOR
 	#include "DerivedDataCacheInterface.h"
@@ -1233,7 +1234,7 @@ void UWorld::PostLoad()
 			else
 			{
 				FStreamingLevelPrivateAccessor::OnLevelAdded(StreamingLevel);
-				if (FStreamingLevelPrivateAccessor::DetermineTargetState(StreamingLevel))
+				if (FStreamingLevelPrivateAccessor::UpdateTargetState(StreamingLevel))
 				{
 					StreamingLevelsToConsider.Add(StreamingLevel);
 				}
@@ -2828,7 +2829,7 @@ const AServerStreamingLevelsVisibility* UWorld::GetServerStreamingLevelsVisibili
 	return ServerStreamingLevelsVisibility;
 }
 
-void UWorld::AddToWorld( ULevel* Level, const FTransform& LevelTransform, bool bConsiderTimeLimit, FNetLevelVisibilityTransactionId TransactionId)
+void UWorld::AddToWorld( ULevel* Level, const FTransform& LevelTransform, bool bConsiderTimeLimit, FNetLevelVisibilityTransactionId TransactionId, ULevelStreaming* OwningLevelStreaming)
 {
 	TRACE_CPUPROFILER_EVENT_SCOPE(AddToWorld);
 	SCOPE_CYCLE_COUNTER(STAT_AddToWorldTime);
@@ -2862,6 +2863,11 @@ void UWorld::AddToWorld( ULevel* Level, const FTransform& LevelTransform, bool b
 	if( bExecuteNextStep && CurrentLevelPendingVisibility == NULL && CurrentLevelPendingInvisibility != Level )
 	{
 		Level->OwningWorld = this;
+
+		if (OwningLevelStreaming)
+		{
+			FLevelStreamingDelegates::OnLevelBeginMakingVisible.Broadcast(this, OwningLevelStreaming, Level);
+		}
 		
 		// Mark level as being the one in process of being made visible.
 		CurrentLevelPendingVisibility = Level;
@@ -3174,7 +3180,7 @@ void UWorld::BeginTearingDown()
 // Cumulated time doing IncrementalUnregisterComponents in UWorld::RemoveFromWorld since last call to UWorld::UpdateLevelStreaming.
 static double GRemoveFromWorldUnregisterComponentTimeCumul = 0.0;
 
-void UWorld::RemoveFromWorld( ULevel* Level, bool bAllowIncrementalRemoval, FNetLevelVisibilityTransactionId TransactionId)
+void UWorld::RemoveFromWorld( ULevel* Level, bool bAllowIncrementalRemoval, FNetLevelVisibilityTransactionId TransactionId, ULevelStreaming* OwningLevelStreaming)
 {
 	TRACE_CPUPROFILER_EVENT_SCOPE(RemoveFromWorld);
 	SCOPE_CYCLE_COUNTER(STAT_RemoveFromWorldTime);
@@ -3186,11 +3192,16 @@ void UWorld::RemoveFromWorld( ULevel* Level, bool bAllowIncrementalRemoval, FNet
 
 	Level->bIsDisassociatingLevel = true;
 
-	auto BeginRemoval = [Level, this]()
+	auto BeginRemoval = [Level, this, OwningLevelStreaming]()
 	{
 		FWorldDelegates::PreLevelRemovedFromWorld.Broadcast(Level, this);
 		Level->bIsBeingRemoved = true;
 
+		if (OwningLevelStreaming)
+		{
+			FLevelStreamingDelegates::OnLevelBeginMakingInvisible.Broadcast(this, OwningLevelStreaming, Level);
+		}
+
 		if (Level->bRequireFullVisibilityToRender && Level->bIsVisible)
 		{
 			Level->bIsVisible = false;
@@ -3918,7 +3929,7 @@ void UWorld::UpdateLevelStreaming()
 
 				if (bRedetermineTarget)
 				{
-					bShouldContinueToConsider = FStreamingLevelPrivateAccessor::DetermineTargetState(StreamingLevel);
+					bShouldContinueToConsider = FStreamingLevelPrivateAccessor::UpdateTargetState(StreamingLevel);
 				}
 			}
 
@@ -3995,7 +4006,7 @@ void UWorld::AddStreamingLevel(ULevelStreaming* StreamingLevelToAdd)
 			{
 				StreamingLevels.Add(StreamingLevelToAdd);
 				FStreamingLevelPrivateAccessor::OnLevelAdded(StreamingLevelToAdd);
-				if (FStreamingLevelPrivateAccessor::DetermineTargetState(StreamingLevelToAdd))
+				if (FStreamingLevelPrivateAccessor::UpdateTargetState(StreamingLevelToAdd))
 				{
 					StreamingLevelsToConsider.Add(StreamingLevelToAdd);
 				}
@@ -4104,7 +4115,7 @@ void UWorld::PopulateStreamingLevelsToConsider()
 			++NumStreamingLevelsBeingLoaded;
 		}
 
-		if (FStreamingLevelPrivateAccessor::DetermineTargetState(StreamingLevel))
+		if (FStreamingLevelPrivateAccessor::UpdateTargetState(StreamingLevel))
 		{
 			StreamingLevelsToConsider.Add(StreamingLevel);
 		}
@@ -4115,7 +4126,7 @@ void UWorld::UpdateStreamingLevelShouldBeConsidered(ULevelStreaming* StreamingLe
 {
 	if (StreamingLevelToConsider && ensure(StreamingLevelToConsider->GetWorld() == this) && StreamingLevelToConsider->GetLevelStreamingState() != ELevelStreamingState::Removed)
 	{
-		if (FStreamingLevelPrivateAccessor::DetermineTargetState(StreamingLevelToConsider))
+		if (FStreamingLevelPrivateAccessor::UpdateTargetState(StreamingLevelToConsider))
 		{
 			StreamingLevelsToConsider.Add(StreamingLevelToConsider);
 		}
diff --git a/Engine/Source/Runtime/Engine/Private/WorldPartition/WorldPartitionLevelStreamingDynamic.cpp b/Engine/Source/Runtime/Engine/Private/WorldPartition/WorldPartitionLevelStreamingDynamic.cpp
--- a/Engine/Source/Runtime/Engine/Private/WorldPartition/WorldPartitionLevelStreamingDynamic.cpp
+++ b/Engine/Source/Runtime/Engine/Private/WorldPartition/WorldPartitionLevelStreamingDynamic.cpp
@@ -288,7 +288,7 @@ bool UWorldPartitionLevelStreamingDynamic::RequestLevel(UWorld* InPersistentWorl
 				}
 				else
 				{
-					CurrentState = ELevelStreamingState::Loading;
+					SetCurrentState(ELevelStreamingState::Loading);
 				}
 			}
 		}
diff --git a/Engine/Source/Runtime/Engine/Public/Streaming/LevelStreamingDelegates.h b/Engine/Source/Runtime/Engine/Public/Streaming/LevelStreamingDelegates.h
new file mode 100644
--- /dev/null
+++ b/Engine/Source/Runtime/Engine/Public/Streaming/LevelStreamingDelegates.h
@@ -0,0 +1,44 @@
+// Copyright Epic Games, Inc. All Rights Reserved.
+// Globally-bound delegates for level streaming status updates
+
+#pragma once
+
+#include "Delegates/Delegate.h"
+
+class UWorld;
+class ULevelStreaming;
+class ULevel;
+
+enum class ELevelStreamingState : uint8;
+enum class ELevelStreamingTargetState : uint8;
+
+// TODO: what is the best way to capture a level which is queued for unload while it is during async load or add to world?
+//	Are level async loads actually cancelled?
+
+struct ENGINE_API FLevelStreamingDelegates
+{
+	// Called when a new target state for a streaming level is determined. 
+	// Can be used to measure latency between requests and application
+	// (e.g. limits on current level loads & visibility changes)
+	static TMulticastDelegate<void(
+		UWorld*, 
+		const ULevelStreaming*, 
+		ULevel* LevelIfLoaded,
+		ELevelStreamingState Current,
+		ELevelStreamingTargetState PrevTarget, 
+		ELevelStreamingTargetState NewTarget)> OnLevelStreamingTargetStateChanged;
+
+	// Called when a streaming level enters a new state.
+	// When a level is newly added to the streaming system, PreviousState and NewState may be equal.
+	static TMulticastDelegate<void(
+		UWorld*, 
+		const ULevelStreaming*,
+		ULevel* LevelIfLoaded,
+		ELevelStreamingState PreviousState,
+		ELevelStreamingState NewState
+		)> OnLevelStreamingStateChanged;
+
+	// Notifications on when work actually beings on making a level visible or invisible, so profiling tools can measure the difference between time spent waiting to be added to the world and time spend actively working
+	static TMulticastDelegate<void(UWorld*, const ULevelStreaming*, ULevel* LoadedLevel)> OnLevelBeginMakingVisible;
+	static TMulticastDelegate<void(UWorld*, const ULevelStreaming*, ULevel* LoadedLevel)> OnLevelBeginMakingInvisible;
+};
