- Fix SRVs not being created when IsSRVNeeded is true
- CachedNumIndices is cleared when Empty is called so Num() returns 0

#rnx

diff --git a/Engine/Source/Runtime/Engine/Public/RawIndexBuffer.h b/Engine/Source/Runtime/Engine/Public/RawIndexBuffer.h
--- a/Engine/Source/Runtime/Engine/Public/RawIndexBuffer.h
+++ b/Engine/Source/Runtime/Engine/Public/RawIndexBuffer.h
@@ -398,11 +398,12 @@ public:
 	*/
 	virtual void InitRHI() override
 	{
+		const bool bHadIndexData = Num() > 0;
 		IndexBufferRHI = CreateRHIBuffer_RenderThread();
 
-		if (IndexBufferRHI && IsSRVNeeded() && Num())
+		if (IndexBufferRHI && IsSRVNeeded() && bHadIndexData)
 		{
-			SRVValue = RHICreateShaderResourceView(Num() ? IndexBufferRHI : nullptr);
+			SRVValue = RHICreateShaderResourceView(IndexBufferRHI->GetSize() > 0 ? IndexBufferRHI : nullptr);
 		}
 	}
 	
