Allow const delegates to be used for asynchronous world collision tests

diff --git a/Engine/Source/Runtime/Engine/Classes/Engine/World.h b/Engine/Source/Runtime/Engine/Classes/Engine/World.h
--- a/Engine/Source/Runtime/Engine/Classes/Engine/World.h
+++ b/Engine/Source/Runtime/Engine/Classes/Engine/World.h
@@ -2175,7 +2175,7 @@ public:
 	 *  @param  TraceChannel    The 'channel' that this ray is in, used to determine which components to hit
 	 *  @param  Params          Additional parameters used for the trace
 	 * 	@param 	ResponseParam	ResponseContainer to be used for this trace
-	 *	@param	InDeleagte		Delegate function to be called - to see example, search FTraceDelegate
+	 *	@param	InDelegate		Delegate function to be called - to see example, search FTraceDelegate
 	 *							Example can be void MyActor::TraceDone(const FTraceHandle& TraceHandle, FTraceDatum & TraceData)
 	 *							Before sending to the function, 
 	 *						
@@ -2184,7 +2184,7 @@ public:
 	 * 
 	 *	@param	UserData		UserData
 	 */ 
-	FTraceHandle	AsyncLineTraceByChannel(EAsyncTraceType InTraceType, const FVector& Start,const FVector& End, ECollisionChannel TraceChannel, const FCollisionQueryParams& Params = FCollisionQueryParams::DefaultQueryParam, const FCollisionResponseParams& ResponseParam = FCollisionResponseParams::DefaultResponseParam, FTraceDelegate * InDelegate=NULL, uint32 UserData = 0 );
+	FTraceHandle	AsyncLineTraceByChannel(EAsyncTraceType InTraceType, const FVector& Start,const FVector& End, ECollisionChannel TraceChannel, const FCollisionQueryParams& Params = FCollisionQueryParams::DefaultQueryParam, const FCollisionResponseParams& ResponseParam = FCollisionResponseParams::DefaultResponseParam, const FTraceDelegate* InDelegate = nullptr, uint32 UserData = 0 );
 
 	/**
 	 * Interface for Async. Pretty much same parameter set except you can optional set delegate to be called when execution is completed and you can set UserData if you'd like
@@ -2196,7 +2196,7 @@ public:
 	 *  @param  End             End location of the ray
 	 *	@param	ObjectQueryParams	List of object types it's looking for
 	 *  @param  Params          Additional parameters used for the trace
-	 *	@param	InDeleagte		Delegate function to be called - to see example, search FTraceDelegate
+	 *	@param	InDelegate		Delegate function to be called - to see example, search FTraceDelegate
 	 *							Example can be void MyActor::TraceDone(const FTraceHandle& TraceHandle, FTraceDatum & TraceData)
 	 *							Before sending to the function, 
 	 *						
@@ -2205,7 +2205,7 @@ public:
 	 * 
 	 *	@param	UserData		UserData
 	 */ 
-	FTraceHandle	AsyncLineTraceByObjectType(EAsyncTraceType InTraceType, const FVector& Start,const FVector& End, const FCollisionObjectQueryParams& ObjectQueryParams, const FCollisionQueryParams& Params = FCollisionQueryParams::DefaultQueryParam, FTraceDelegate * InDelegate=NULL, uint32 UserData = 0 );
+	FTraceHandle	AsyncLineTraceByObjectType(EAsyncTraceType InTraceType, const FVector& Start,const FVector& End, const FCollisionObjectQueryParams& ObjectQueryParams, const FCollisionQueryParams& Params = FCollisionQueryParams::DefaultQueryParam, const FTraceDelegate* InDelegate = nullptr, uint32 UserData = 0 );
 
 	/**
 	 * Interface for Async. Pretty much same parameter set except you can optional set delegate to be called when execution is completed and you can set UserData if you'd like
@@ -2217,7 +2217,7 @@ public:
 	 *  @param  End             End location of the ray
 	 *  @param  ProfileName		The 'profile' used to determine which components to hit
 	 *  @param  Params          Additional parameters used for the trace
-	 *	@param	InDeleagte		Delegate function to be called - to see example, search FTraceDelegate
+	 *	@param	InDelegate		Delegate function to be called - to see example, search FTraceDelegate
 	 *							Example can be void MyActor::TraceDone(const FTraceHandle& TraceHandle, FTraceDatum & TraceData)
 	 *							Before sending to the function,
 	 *
@@ -2226,7 +2226,7 @@ public:
 	 *
 	 *	@param	UserData		UserData
 	 */
-	FTraceHandle	AsyncLineTraceByProfile(EAsyncTraceType InTraceType, const FVector& Start, const FVector& End, FName ProfileName, const FCollisionQueryParams& Params = FCollisionQueryParams::DefaultQueryParam, FTraceDelegate* InDelegate = NULL, uint32 UserData = 0);
+	FTraceHandle	AsyncLineTraceByProfile(EAsyncTraceType InTraceType, const FVector& Start, const FVector& End, FName ProfileName, const FCollisionQueryParams& Params = FCollisionQueryParams::DefaultQueryParam, const FTraceDelegate* InDelegate = nullptr, uint32 UserData = 0);
 
 	/**
 	 * Interface for Async trace
@@ -2241,7 +2241,7 @@ public:
 	 *  @param	CollisionShape		CollisionShape - supports Box, Sphere, Capsule
 	 *  @param  Params          Additional parameters used for the trace
 	 * 	@param 	ResponseParam	ResponseContainer to be used for this trace	 
-	 *	@param	InDeleagte		Delegate function to be called - to see example, search FTraceDelegate
+	 *	@param	InDelegate		Delegate function to be called - to see example, search FTraceDelegate
 	 *							Example can be void MyActor::TraceDone(const FTraceHandle& TraceHandle, FTraceDatum & TraceData)
 	 *							Before sending to the function, 
 	 *						
@@ -2250,7 +2250,7 @@ public:
 	 * 
 	 *	@param	UserData		UserData
 	 */ 
-	FTraceHandle	AsyncSweepByChannel(EAsyncTraceType InTraceType, const FVector& Start, const FVector& End, const FQuat& Rot, ECollisionChannel TraceChannel, const FCollisionShape& CollisionShape, const FCollisionQueryParams& Params = FCollisionQueryParams::DefaultQueryParam, const FCollisionResponseParams& ResponseParam = FCollisionResponseParams::DefaultResponseParam, FTraceDelegate * InDelegate = NULL, uint32 UserData = 0);
+	FTraceHandle	AsyncSweepByChannel(EAsyncTraceType InTraceType, const FVector& Start, const FVector& End, const FQuat& Rot, ECollisionChannel TraceChannel, const FCollisionShape& CollisionShape, const FCollisionQueryParams& Params = FCollisionQueryParams::DefaultQueryParam, const FCollisionResponseParams& ResponseParam = FCollisionResponseParams::DefaultResponseParam, const FTraceDelegate* InDelegate = nullptr, uint32 UserData = 0);
 
 	/**
 	 * Interface for Async trace
@@ -2264,7 +2264,7 @@ public:
 	 *	@param	ObjectQueryParams	List of object types it's looking for
 	 *  @param	CollisionShape		CollisionShape - supports Box, Sphere, Capsule
 	 *  @param  Params          Additional parameters used for the trace
-	 *	@param	InDeleagte		Delegate function to be called - to see example, search FTraceDelegate
+	 *	@param	InDelegate		Delegate function to be called - to see example, search FTraceDelegate
 	 *							Example can be void MyActor::TraceDone(const FTraceHandle& TraceHandle, FTraceDatum & TraceData)
 	 *							Before sending to the function, 
 	 *						
@@ -2273,7 +2273,7 @@ public:
 	 * 
 	 *	@param	UserData		UserData
 	 */ 
-	FTraceHandle	AsyncSweepByObjectType(EAsyncTraceType InTraceType, const FVector& Start, const FVector& End, const FQuat& Rot, const FCollisionObjectQueryParams& ObjectQueryParams, const FCollisionShape& CollisionShape, const FCollisionQueryParams& Params = FCollisionQueryParams::DefaultQueryParam, FTraceDelegate * InDelegate = NULL, uint32 UserData = 0);
+	FTraceHandle	AsyncSweepByObjectType(EAsyncTraceType InTraceType, const FVector& Start, const FVector& End, const FQuat& Rot, const FCollisionObjectQueryParams& ObjectQueryParams, const FCollisionShape& CollisionShape, const FCollisionQueryParams& Params = FCollisionQueryParams::DefaultQueryParam, const FTraceDelegate* InDelegate = nullptr, uint32 UserData = 0);
 
 	/**
 	 * Interface for Async trace
@@ -2287,7 +2287,7 @@ public:
 	 *  @param  ProfileName     The 'profile' used to determine which components to hit
 	 *  @param	CollisionShape	CollisionShape - supports Box, Sphere, Capsule
 	 *  @param  Params          Additional parameters used for the trace
-	 *	@param	InDeleagte		Delegate function to be called - to see example, search FTraceDelegate
+	 *	@param	InDelegate		Delegate function to be called - to see example, search FTraceDelegate
 	 *							Example can be void MyActor::TraceDone(const FTraceHandle& TraceHandle, FTraceDatum & TraceData)
 	 *							Before sending to the function,
 	 *
@@ -2296,7 +2296,7 @@ public:
 	 *
 	 *	@param	UserData		UserData
 	 */
-	FTraceHandle	AsyncSweepByProfile(EAsyncTraceType InTraceType, const FVector& Start, const FVector& End, const FQuat& Rot, FName ProfileName, const FCollisionShape& CollisionShape, const FCollisionQueryParams& Params = FCollisionQueryParams::DefaultQueryParam, FTraceDelegate* InDelegate = NULL, uint32 UserData = 0);
+	FTraceHandle	AsyncSweepByProfile(EAsyncTraceType InTraceType, const FVector& Start, const FVector& End, const FQuat& Rot, FName ProfileName, const FCollisionShape& CollisionShape, const FCollisionQueryParams& Params = FCollisionQueryParams::DefaultQueryParam, const FTraceDelegate* InDelegate = nullptr, uint32 UserData = 0);
 
 	// overlap functions
 
@@ -2307,12 +2307,11 @@ public:
 	 * the data is available only in the next frame after request is made - in other words, if request is made in frame X, you can get the result in frame (X+1)
 	 *
 	 *  @param  Pos             Location of center of shape to test against the world
-	 *	@param	bMultiTrace		true if you'd like to do multi trace, or false otherwise
 	 *  @param  TraceChannel    The 'channel' that this query is in, used to determine which components to hit
 	 *  @param	CollisionShape		CollisionShape - supports Box, Sphere, Capsule
 	 *  @param  Params          Additional parameters used for the trace
 	 * 	@param 	ResponseParam	ResponseContainer to be used for this trace
-	 *	@param	InDeleagte		Delegate function to be called - to see example, search FTraceDelegate
+	 *	@param	InDelegate		Delegate function to be called - to see example, search FTraceDelegate
 	 *							Example can be void MyActor::TraceDone(const FTraceHandle& TraceHandle, FTraceDatum & TraceData)
 	 *							Before sending to the function, 
 	 *						
@@ -2321,7 +2320,7 @@ public:
 	 * 
 	 *	@param UserData			UserData
 	 */ 
-	FTraceHandle	AsyncOverlapByChannel(const FVector& Pos, const FQuat& Rot, ECollisionChannel TraceChannel, const FCollisionShape& CollisionShape, const FCollisionQueryParams& Params = FCollisionQueryParams::DefaultQueryParam, const FCollisionResponseParams& ResponseParam = FCollisionResponseParams::DefaultResponseParam, FOverlapDelegate * InDelegate = NULL, uint32 UserData = 0);
+	FTraceHandle	AsyncOverlapByChannel(const FVector& Pos, const FQuat& Rot, ECollisionChannel TraceChannel, const FCollisionShape& CollisionShape, const FCollisionQueryParams& Params = FCollisionQueryParams::DefaultQueryParam, const FCollisionResponseParams& ResponseParam = FCollisionResponseParams::DefaultResponseParam, const FOverlapDelegate* InDelegate = nullptr, uint32 UserData = 0);
 
 	/**
 	 * Interface for Async trace
@@ -2333,7 +2332,7 @@ public:
 	 *	@param	ObjectQueryParams	List of object types it's looking for
 	 *  @param	CollisionShape		CollisionShape - supports Box, Sphere, Capsule
 	 *  @param  Params          Additional parameters used for the trace
-	 *	@param	InDeleagte		Delegate function to be called - to see example, search FTraceDelegate
+	 *	@param	InDelegate		Delegate function to be called - to see example, search FTraceDelegate
 	 *							Example can be void MyActor::TraceDone(const FTraceHandle& TraceHandle, FTraceDatum & TraceData)
 	 *							Before sending to the function, 
 	 *						
@@ -2342,7 +2341,7 @@ public:
 	 * 
 	 *	@param UserData			UserData
 	 */ 
-	FTraceHandle	AsyncOverlapByObjectType(const FVector& Pos, const FQuat& Rot, const FCollisionObjectQueryParams& ObjectQueryParams, const FCollisionShape& CollisionShape, const FCollisionQueryParams& Params = FCollisionQueryParams::DefaultQueryParam, FOverlapDelegate * InDelegate = NULL, uint32 UserData = 0);
+	FTraceHandle	AsyncOverlapByObjectType(const FVector& Pos, const FQuat& Rot, const FCollisionObjectQueryParams& ObjectQueryParams, const FCollisionShape& CollisionShape, const FCollisionQueryParams& Params = FCollisionQueryParams::DefaultQueryParam, const FOverlapDelegate* InDelegate = nullptr, uint32 UserData = 0);
 
 	/**
 	 * Interface for Async trace
@@ -2354,7 +2353,7 @@ public:
 	 *  @param  ProfileName     The 'profile' used to determine which components to hit
 	 *  @param	CollisionShape		CollisionShape - supports Box, Sphere, Capsule
 	 *  @param  Params          Additional parameters used for the trace
-	 *	@param	InDeleagte		Delegate function to be called - to see example, search FTraceDelegate
+	 *	@param	InDelegate		Delegate function to be called - to see example, search FTraceDelegate
 	 *							Example can be void MyActor::TraceDone(const FTraceHandle& TraceHandle, FTraceDatum & TraceData)
 	 *							Before sending to the function,
 	 *
@@ -2363,7 +2362,7 @@ public:
 	 *
 	 *	@param UserData			UserData
 	 */
-	FTraceHandle	AsyncOverlapByProfile(const FVector& Pos, const FQuat& Rot, FName ProfileName, const FCollisionShape& CollisionShape, const FCollisionQueryParams& Params = FCollisionQueryParams::DefaultQueryParam, FOverlapDelegate* InDelegate = NULL, uint32 UserData = 0);
+	FTraceHandle	AsyncOverlapByProfile(const FVector& Pos, const FQuat& Rot, FName ProfileName, const FCollisionShape& CollisionShape, const FCollisionQueryParams& Params = FCollisionQueryParams::DefaultQueryParam, const FOverlapDelegate* InDelegate = nullptr, uint32 UserData = 0);
 
 	/**
 	 * Query function 
diff --git a/Engine/Source/Runtime/Engine/Private/Collision/WorldCollisionAsync.cpp b/Engine/Source/Runtime/Engine/Private/Collision/WorldCollisionAsync.cpp
--- a/Engine/Source/Runtime/Engine/Private/Collision/WorldCollisionAsync.cpp
+++ b/Engine/Source/Runtime/Engine/Private/Collision/WorldCollisionAsync.cpp
@@ -366,17 +366,17 @@ FWorldAsyncTraceState::FWorldAsyncTraceState()
 	DataBuffer[CurrentFrame].bAsyncAllowed = true;
 }
 
-FTraceHandle UWorld::AsyncLineTraceByChannel(EAsyncTraceType InTraceType, const FVector& Start, const FVector& End, ECollisionChannel TraceChannel, const FCollisionQueryParams& Params /* = FCollisionQueryParams::DefaultQueryParam */, const FCollisionResponseParams& ResponseParam /* = FCollisionResponseParams::DefaultResponseParam */, FTraceDelegate * InDelegate/* =NULL */, uint32 UserData /* = 0 */)
+FTraceHandle UWorld::AsyncLineTraceByChannel(EAsyncTraceType InTraceType, const FVector& Start, const FVector& End, ECollisionChannel TraceChannel, const FCollisionQueryParams& Params /* = FCollisionQueryParams::DefaultQueryParam */, const FCollisionResponseParams& ResponseParam /* = FCollisionResponseParams::DefaultResponseParam */, const FTraceDelegate* InDelegate/* = nullptr */, uint32 UserData /* = 0 */)
 {
 	return StartNewTrace(AsyncTraceState, FTraceDatum(this, FCollisionShape::LineShape, Params, ResponseParam, FCollisionObjectQueryParams::DefaultObjectQueryParam, TraceChannel, UserData, InTraceType, Start, End, FQuat::Identity, InDelegate, AsyncTraceState.CurrentFrame));
 }
 
-FTraceHandle UWorld::AsyncLineTraceByObjectType(EAsyncTraceType InTraceType, const FVector& Start, const FVector& End, const FCollisionObjectQueryParams& ObjectQueryParams, const FCollisionQueryParams& Params /* = FCollisionQueryParams::DefaultQueryParam */, FTraceDelegate* InDelegate/* =NULL */, uint32 UserData /* = 0 */)
+FTraceHandle UWorld::AsyncLineTraceByObjectType(EAsyncTraceType InTraceType, const FVector& Start, const FVector& End, const FCollisionObjectQueryParams& ObjectQueryParams, const FCollisionQueryParams& Params /* = FCollisionQueryParams::DefaultQueryParam */, const FTraceDelegate* InDelegate/* = nullptr */, uint32 UserData /* = 0 */)
 {
 	return StartNewTrace(AsyncTraceState, FTraceDatum(this, FCollisionShape::LineShape, Params, FCollisionResponseParams::DefaultResponseParam, ObjectQueryParams, DefaultCollisionChannel, UserData, InTraceType, Start, End, FQuat::Identity, InDelegate, AsyncTraceState.CurrentFrame));
 }
 
-FTraceHandle UWorld::AsyncLineTraceByProfile(EAsyncTraceType InTraceType, const FVector& Start, const FVector& End, FName ProfileName, const FCollisionQueryParams& Params /* = FCollisionQueryParams::DefaultQueryParam */, FTraceDelegate* InDelegate/* =NULL */, uint32 UserData /* = 0 */)
+FTraceHandle UWorld::AsyncLineTraceByProfile(EAsyncTraceType InTraceType, const FVector& Start, const FVector& End, FName ProfileName, const FCollisionQueryParams& Params /* = FCollisionQueryParams::DefaultQueryParam */, const FTraceDelegate* InDelegate/* = nullptr */, uint32 UserData /* = 0 */)
 {
 	ECollisionChannel TraceChannel;
 	FCollisionResponseParams ResponseParam;
@@ -385,17 +385,17 @@ FTraceHandle UWorld::AsyncLineTraceByProfile(EAsyncTraceType InTraceType, const
 	return StartNewTrace(AsyncTraceState, FTraceDatum(this, FCollisionShape::LineShape, Params, ResponseParam, FCollisionObjectQueryParams::DefaultObjectQueryParam, TraceChannel, UserData, InTraceType, Start, End, FQuat::Identity, InDelegate, AsyncTraceState.CurrentFrame));
 }
 
-FTraceHandle UWorld::AsyncSweepByChannel(EAsyncTraceType InTraceType, const FVector& Start, const FVector& End, const FQuat& Rot, ECollisionChannel TraceChannel, const FCollisionShape& CollisionShape, const FCollisionQueryParams& Params /* = FCollisionQueryParams::DefaultQueryParam */, const FCollisionResponseParams& ResponseParam /* = FCollisionResponseParams::DefaultResponseParam */, FTraceDelegate* InDelegate /* = NULL */, uint32 UserData /* = 0 */)
+FTraceHandle UWorld::AsyncSweepByChannel(EAsyncTraceType InTraceType, const FVector& Start, const FVector& End, const FQuat& Rot, ECollisionChannel TraceChannel, const FCollisionShape& CollisionShape, const FCollisionQueryParams& Params /* = FCollisionQueryParams::DefaultQueryParam */, const FCollisionResponseParams& ResponseParam /* = FCollisionResponseParams::DefaultResponseParam */, const FTraceDelegate* InDelegate/* = nullptr */, uint32 UserData /* = 0 */)
 {
 	return StartNewTrace(AsyncTraceState, FTraceDatum(this, CollisionShape, Params, ResponseParam, FCollisionObjectQueryParams::DefaultObjectQueryParam, TraceChannel, UserData, InTraceType, Start, End, Rot, InDelegate, AsyncTraceState.CurrentFrame));
 }
 
-FTraceHandle UWorld::AsyncSweepByObjectType(EAsyncTraceType InTraceType, const FVector& Start, const FVector& End, const FQuat& Rot, const FCollisionObjectQueryParams& ObjectQueryParams, const FCollisionShape& CollisionShape, const FCollisionQueryParams& Params /* = FCollisionQueryParams::DefaultQueryParam */, FTraceDelegate* InDelegate /* = NULL */, uint32 UserData /* = 0 */)
+FTraceHandle UWorld::AsyncSweepByObjectType(EAsyncTraceType InTraceType, const FVector& Start, const FVector& End, const FQuat& Rot, const FCollisionObjectQueryParams& ObjectQueryParams, const FCollisionShape& CollisionShape, const FCollisionQueryParams& Params /* = FCollisionQueryParams::DefaultQueryParam */, const FTraceDelegate* InDelegate/* = nullptr */, uint32 UserData /* = 0 */)
 {
 	return StartNewTrace(AsyncTraceState, FTraceDatum(this, CollisionShape, Params, FCollisionResponseParams::DefaultResponseParam, ObjectQueryParams, DefaultCollisionChannel, UserData, InTraceType, Start, End, Rot, InDelegate, AsyncTraceState.CurrentFrame));
 }
 
-FTraceHandle UWorld::AsyncSweepByProfile(EAsyncTraceType InTraceType, const FVector& Start, const FVector& End, const FQuat& Rot, FName ProfileName, const FCollisionShape& CollisionShape, const FCollisionQueryParams& Params /* = FCollisionQueryParams::DefaultQueryParam */, FTraceDelegate* InDelegate /* = NULL */, uint32 UserData /* = 0 */)
+FTraceHandle UWorld::AsyncSweepByProfile(EAsyncTraceType InTraceType, const FVector& Start, const FVector& End, const FQuat& Rot, FName ProfileName, const FCollisionShape& CollisionShape, const FCollisionQueryParams& Params /* = FCollisionQueryParams::DefaultQueryParam */, const FTraceDelegate* InDelegate/* = nullptr */, uint32 UserData /* = 0 */)
 {
 	ECollisionChannel TraceChannel;
 	FCollisionResponseParams ResponseParam;
@@ -405,17 +405,17 @@ FTraceHandle UWorld::AsyncSweepByProfile(EAsyncTraceType InTraceType, const FVec
 }
 
 // overlap functions
-FTraceHandle UWorld::AsyncOverlapByChannel(const FVector& Pos, const FQuat& Rot, ECollisionChannel TraceChannel, const FCollisionShape& CollisionShape, const FCollisionQueryParams& Params /* = FCollisionQueryParams::DefaultQueryParam */, const FCollisionResponseParams& ResponseParam /* = FCollisionResponseParams::DefaultResponseParam */, FOverlapDelegate * InDelegate /* = NULL */, uint32 UserData /* = 0 */)
+FTraceHandle UWorld::AsyncOverlapByChannel(const FVector& Pos, const FQuat& Rot, ECollisionChannel TraceChannel, const FCollisionShape& CollisionShape, const FCollisionQueryParams& Params /* = FCollisionQueryParams::DefaultQueryParam */, const FCollisionResponseParams& ResponseParam /* = FCollisionResponseParams::DefaultResponseParam */, const FOverlapDelegate* InDelegate/* = nullptr */, uint32 UserData /* = 0 */)
 {
 	return StartNewTrace(AsyncTraceState, FOverlapDatum(this, CollisionShape, Params, ResponseParam, FCollisionObjectQueryParams::DefaultObjectQueryParam, TraceChannel, UserData, Pos, Rot, InDelegate, AsyncTraceState.CurrentFrame));
 }
 
-FTraceHandle UWorld::AsyncOverlapByObjectType(const FVector& Pos, const FQuat& Rot, const FCollisionObjectQueryParams& ObjectQueryParams, const FCollisionShape& CollisionShape, const FCollisionQueryParams& Params /* = FCollisionQueryParams::DefaultQueryParam */, FOverlapDelegate * InDelegate /* = NULL */, uint32 UserData /* = 0 */)
+FTraceHandle UWorld::AsyncOverlapByObjectType(const FVector& Pos, const FQuat& Rot, const FCollisionObjectQueryParams& ObjectQueryParams, const FCollisionShape& CollisionShape, const FCollisionQueryParams& Params /* = FCollisionQueryParams::DefaultQueryParam */, const FOverlapDelegate* InDelegate/* = nullptr */, uint32 UserData /* = 0 */)
 {
 	return StartNewTrace(AsyncTraceState, FOverlapDatum(this, CollisionShape, Params, FCollisionResponseParams::DefaultResponseParam, ObjectQueryParams, DefaultCollisionChannel, UserData, Pos, Rot, InDelegate, AsyncTraceState.CurrentFrame));
 }
 
-FTraceHandle UWorld::AsyncOverlapByProfile(const FVector& Pos, const FQuat& Rot, FName ProfileName, const FCollisionShape& CollisionShape, const FCollisionQueryParams& Params /* = FCollisionQueryParams::DefaultQueryParam */, FOverlapDelegate* InDelegate /* = NULL */, uint32 UserData /* = 0 */)
+FTraceHandle UWorld::AsyncOverlapByProfile(const FVector& Pos, const FQuat& Rot, FName ProfileName, const FCollisionShape& CollisionShape, const FCollisionQueryParams& Params /* = FCollisionQueryParams::DefaultQueryParam */, const FOverlapDelegate* InDelegate /* = nullptr */, uint32 UserData /* = 0 */)
 {
 	ECollisionChannel TraceChannel;
 	FCollisionResponseParams ResponseParam;
diff --git a/Engine/Source/Runtime/Engine/Public/WorldCollision.h b/Engine/Source/Runtime/Engine/Public/WorldCollision.h
--- a/Engine/Source/Runtime/Engine/Public/WorldCollision.h
+++ b/Engine/Source/Runtime/Engine/Public/WorldCollision.h
@@ -177,7 +177,7 @@ struct FTraceDatum : public FBaseTraceDatum
 
 	/** Set Trace Datum for each shape type **/
 	FTraceDatum(UWorld* World, const FCollisionShape& CollisionShape, const FCollisionQueryParams& Param, const struct FCollisionResponseParams& InResponseParam, const struct FCollisionObjectQueryParams& InObjectQueryParam,
-		ECollisionChannel Channel, uint32 InUserData, EAsyncTraceType InTraceType, const FVector& InStart, const FVector& InEnd, const FQuat& InRot, FTraceDelegate* InDelegate, int32 FrameCounter)
+		ECollisionChannel Channel, uint32 InUserData, EAsyncTraceType InTraceType, const FVector& InStart, const FVector& InEnd, const FQuat& InRot, const FTraceDelegate* InDelegate, int32 FrameCounter)
 	{
 		Set(World, CollisionShape, Param, InResponseParam, InObjectQueryParam, Channel, InUserData, FrameCounter);
 		Start = InStart;
@@ -221,7 +221,7 @@ struct FOverlapDatum : FBaseTraceDatum
 
 	FOverlapDatum(UWorld * World, const FCollisionShape& CollisionShape, const FCollisionQueryParams& Param, const struct FCollisionResponseParams &InResponseParam, const struct FCollisionObjectQueryParams& InObjectQueryParam, 
 		ECollisionChannel Channel, uint32 InUserData,
-		const FVector& InPos, const FQuat& InRot, FOverlapDelegate * InDelegate, int32 FrameCounter)
+		const FVector& InPos, const FQuat& InRot, const FOverlapDelegate* InDelegate, int32 FrameCounter)
 	{
 		Set(World, CollisionShape, Param, InResponseParam, InObjectQueryParam, Channel, InUserData, FrameCounter);
 		Pos = InPos;
