Global, Material and MeshMaterial shader compiling WITH_EDITOR cleanup.
None of the shader compiling code runs in non-Editor mode, so don't compile it for those targets.

diff --git a/Engine/Plugins/Experimental/Water/Source/Runtime/Private/WaterZoneActor.cpp b/Engine/Plugins/Experimental/Water/Source/Runtime/Private/WaterZoneActor.cpp
--- a/Engine/Plugins/Experimental/Water/Source/Runtime/Private/WaterZoneActor.cpp
+++ b/Engine/Plugins/Experimental/Water/Source/Runtime/Private/WaterZoneActor.cpp
@@ -440,7 +440,9 @@ bool AWaterZone::UpdateWaterInfoTexture()
 				{
 					if (!MaterialResource->IsGameThreadShaderMapComplete())
 					{
+#if WITH_EDITOR
 						MaterialResource->SubmitCompileJobs_GameThread(EShaderCompileJobPriority::ForceLocal);
+#endif
 						bHasIncompleteShaderMaps = true;
 					}
 				}
diff --git a/Engine/Plugins/VirtualProduction/CameraCalibrationCore/Source/CameraCalibrationCore/Private/LensDistortionModelHandlerBase.cpp b/Engine/Plugins/VirtualProduction/CameraCalibrationCore/Source/CameraCalibrationCore/Private/LensDistortionModelHandlerBase.cpp
--- a/Engine/Plugins/VirtualProduction/CameraCalibrationCore/Source/CameraCalibrationCore/Private/LensDistortionModelHandlerBase.cpp
+++ b/Engine/Plugins/VirtualProduction/CameraCalibrationCore/Source/CameraCalibrationCore/Private/LensDistortionModelHandlerBase.cpp
@@ -233,7 +233,9 @@ bool ULensDistortionModelHandlerBase::IsDisplacementMapMaterialReady(UMaterialIn
 			return true;
 		}
 
+#if WITH_EDITOR
 		MaterialResource->SubmitCompileJobs_GameThread(EShaderCompileJobPriority::ForceLocal);
+#endif
 
 		return MaterialResource->IsGameThreadShaderMapComplete();
 	}
diff --git a/Engine/Source/Runtime/Engine/Private/Materials/MaterialShader.cpp b/Engine/Source/Runtime/Engine/Private/Materials/MaterialShader.cpp
--- a/Engine/Source/Runtime/Engine/Private/Materials/MaterialShader.cpp
+++ b/Engine/Source/Runtime/Engine/Private/Materials/MaterialShader.cpp
@@ -1005,9 +1005,6 @@ void FMaterialShaderMapId::UpdateFromParameterSet(const FStaticParameterSet& Sta
 		StaticComponentMaskParameter.bOverride = false;
 	}
 }	
-#endif // WITH_EDITOR
-
-#if WITH_EDITOR
 
 uint32 GetTypeHash(FPlatformTypeLayoutParameters Params) { return HashCombine(Params.Flags, Params.MaxFieldAlignment); }
 
@@ -1091,9 +1088,7 @@ void FMaterialShaderMapId::AppendKeyString(FString& KeyString) const
 		KeyString += TEXT("_NewHLSL");
 	}
 }
-#endif // WITH_EDITOR
 
-#if WITH_EDITOR
 void FMaterialShaderMapId::SetShaderDependencies(const TArray<FShaderType*>& ShaderTypes, const TArray<const FShaderPipelineType*>& ShaderPipelineTypes, const TArray<FVertexFactoryType*>& VFTypes, EShaderPlatform ShaderPlatform)
 {
 	if (!FPlatformProperties::RequiresCookedData() && AllowShaderCompiling())
@@ -1124,7 +1119,6 @@ void FMaterialShaderMapId::SetShaderDependencies(const TArray<FShaderType*>& Sha
 		}
 	}
 }
-#endif // WITH_EDITOR
 
 static void PrepareMaterialShaderCompileJob(EShaderPlatform Platform,
 	EShaderPermutationFlags PermutationFlags,
@@ -1247,6 +1241,7 @@ FShader* FMaterialShaderType::FinishCompileShader(
 
 	return Shader;
 }
+#endif // WITH_EDITOR
 
 bool FMaterialShaderType::ShouldCompilePermutation(EShaderPlatform Platform, const FMaterialShaderParameters& MaterialParameters, int32 PermutationId, EShaderPermutationFlags Flags) const
 {
@@ -1267,11 +1262,13 @@ bool FMaterialShaderType::ShouldCompilePipeline(const FShaderPipelineType* Shade
 	return true;
 }
 
+#if WITH_EDITOR
 void FMaterialShaderType::SetupCompileEnvironment(EShaderPlatform Platform, const FMaterialShaderParameters& MaterialParameters, int32 PermutationId, EShaderPermutationFlags PermutationFlags, FShaderCompilerEnvironment& Environment) const
 {
 	// Allow the shader type to modify its compile environment.
 	ModifyCompilationEnvironment(FMaterialShaderPermutationParameters(Platform, MaterialParameters, PermutationId, PermutationFlags), Environment);
 }
+#endif // WITH_EDITOR
 
 /**
 * Finds the shader map for a material.
@@ -1288,31 +1285,6 @@ TRefCountPtr<FMaterialShaderMap> FMaterialShaderMap::FindId(const FMaterialShade
 	return Result;
 }
 
-#if ALLOW_SHADERMAP_DEBUG_DATA
-/** Flushes the given shader types from any loaded FMaterialShaderMap's. */
-void FMaterialShaderMap::FlushShaderTypes(TArray<const FShaderType*>& ShaderTypesToFlush, TArray<const FShaderPipelineType*>& ShaderPipelineTypesToFlush, TArray<const FVertexFactoryType*>& VFTypesToFlush)
-{
-	FScopeLock AllMatSMAccess(&AllMaterialShaderMapsGuard);
-	for (int32 ShaderMapIndex = 0; ShaderMapIndex < AllMaterialShaderMaps.Num(); ShaderMapIndex++)
-	{
-		FMaterialShaderMap* CurrentShaderMap = AllMaterialShaderMaps[ShaderMapIndex];
-
-		for (int32 ShaderTypeIndex = 0; ShaderTypeIndex < ShaderTypesToFlush.Num(); ShaderTypeIndex++)
-		{
-			CurrentShaderMap->FlushShadersByShaderType(ShaderTypesToFlush[ShaderTypeIndex]);
-		}
-		for (int32 VFTypeIndex = 0; VFTypeIndex < VFTypesToFlush.Num(); VFTypeIndex++)
-		{
-			CurrentShaderMap->FlushShadersByVertexFactoryType(VFTypesToFlush[VFTypeIndex]);
-		}
-		for (int32 TypeIndex = 0; TypeIndex < ShaderPipelineTypesToFlush.Num(); TypeIndex++)
-		{
-			CurrentShaderMap->FlushShadersByShaderPipelineType(ShaderPipelineTypesToFlush[TypeIndex]);
-		}
-	}
-}
-#endif
-
 #if WITH_EDITOR
 void FMaterialShaderMap::GetAllOutdatedTypes(TArray<const FShaderType*>& OutdatedShaderTypes, TArray<const FShaderPipelineType*>& OutdatedShaderPipelineTypes, TArray<const FVertexFactoryType*>& OutdatedFactoryTypes)
 {
@@ -1698,6 +1670,8 @@ void FMaterialShaderMapContent::Finalize(const FShaderMapResourceCode* Code)
 	Hasher.GetHash(ShaderContentHash.Hash);
 }
 
+#if WITH_EDITOR
+
 static TMap<uint32, TRefCountPtr<FMaterialShaderMap>>& GetCompilingShaderMapLookup()
 {
 	static TMap<uint32, TRefCountPtr<FMaterialShaderMap>> Map;
@@ -1783,7 +1757,7 @@ void FMaterialShaderMap::SubmitCompileJobs(uint32 CompilingShaderMapId,
 	const FMaterialShaderParameters MaterialParameters(Material);
 	const FMaterialShaderMapLayout& Layout = AcquireMaterialShaderMapLayout(ShaderPlatform, LocalPermutationFlags, MaterialParameters);
 
-#if ALLOW_SHADERMAP_DEBUG_DATA && WITH_EDITOR
+#if ALLOW_SHADERMAP_DEBUG_DATA
 	FString DebugExtensionStr(TEXT(""));
 	FString DebugDescriptionStr(TEXT(""));
 
@@ -1826,10 +1800,8 @@ void FMaterialShaderMap::SubmitCompileJobs(uint32 CompilingShaderMapId,
 			}
 
 			// Verify that the shader map Id contains inputs for any shaders that will be put into this shader map
-#if WITH_EDITOR
 			check(ShaderMapId.ContainsVertexFactoryType(MeshLayout.VertexFactoryType));
 			check(ShaderMapId.ContainsShaderType(ShaderType, kUniqueShaderPermutationId));
-#endif
 
 			NumShadersPerVF++;
 			// only compile the shader if we don't already have it and it is not a pipelined one
@@ -1865,20 +1837,19 @@ void FMaterialShaderMap::SubmitCompileJobs(uint32 CompilingShaderMapId,
 			auto& StageTypes = Pipeline->GetStages();
 
 			// Verify that the shader map Id contains inputs for any shaders that will be put into this shader map
-#if WITH_EDITOR
 			check(ShaderMapId.ContainsShaderPipelineType(Pipeline));
 			check(ShaderMapId.ContainsVertexFactoryType(MeshLayout.VertexFactoryType));
-#endif
+
 			if (Pipeline->ShouldOptimizeUnusedOutputs(ShaderPlatform))
 			{
 				NumShadersPerVF += StageTypes.Num();
-#if WITH_EDITOR
+
 				for (auto* ShaderType : StageTypes)
 				{
 					// Verify that the shader map Id contains inputs for any shaders that will be put into this shader map
 					check(ShaderMapId.ContainsShaderType(ShaderType, kUniqueShaderPermutationId));
 				}
-#endif
+
 				// Make a pipeline job with all the stages
 				FMeshMaterialShaderType::BeginCompileShaderPipeline(InPriority,
 					CompilingShaderMapId,
@@ -1931,10 +1902,9 @@ void FMaterialShaderMap::SubmitCompileJobs(uint32 CompilingShaderMapId,
 			continue;
 		}
 
-#if WITH_EDITOR
 		// Verify that the shader map Id contains inputs for any shaders that will be put into this shader map
 		check(ShaderMapId.ContainsShaderType(ShaderType, kUniqueShaderPermutationId));
-#endif
+
 		// Compile this material shader for this material.
 		TArray<FString> ShaderErrors;
 
@@ -1972,9 +1942,8 @@ void FMaterialShaderMap::SubmitCompileJobs(uint32 CompilingShaderMapId,
 			auto& StageTypes = Pipeline->GetStages();
 
 			// Verify that the shader map Id contains inputs for any shaders that will be put into this shader map
-#if WITH_EDITOR
 			check(ShaderMapId.ContainsShaderPipelineType(Pipeline));
-#endif
+
 			if (Pipeline->ShouldOptimizeUnusedOutputs(ShaderPlatform))
 			{
 				NumShaders += StageTypes.Num();
@@ -2048,7 +2017,7 @@ void FMaterialShaderMap::Compile(
 	ShaderMapId = InShaderMapId;
 	bIsPersistent = Material->IsPersistent();
 
-#if ALLOW_SHADERMAP_DEBUG_DATA && WITH_EDITOR
+#if ALLOW_SHADERMAP_DEBUG_DATA
 	// Store the material name for debugging purposes.
 	// Note: Material instances with static parameters will have the same FriendlyName for their shader maps!
 	NewContent->FriendlyName = Material->GetFriendlyName();
@@ -2226,7 +2195,6 @@ FShader* FMaterialShaderMap::ProcessCompilationResultsForSingleJob(FShaderCompil
 		}
 	}
 
-#if WITH_EDITOR
 	// add shader source
 	{
 		// Keep the preprocessed source list sorted by a name constructed from VF/ShaderType/PermutationId and deduplicate entries
@@ -2238,12 +2206,10 @@ FShader* FMaterialShaderMap::ProcessCompilationResultsForSingleJob(FShaderCompil
 			GetMutableContent()->ShaderProcessedSource.EmplaceAt(Index, Key, *CurrentJob.Output.OptionalFinalShaderSource);
 		}
 	}
-#endif
 
 	return Shader;
 }
 
-#if WITH_EDITOR
 void FMaterialShaderMap::ProcessCompilationResults(const TArray<FShaderCommonCompileJobPtr>& InCompilationResults, int32& InOutJobIndex, float& TimeBudget)
 {
 	check(!bCompilationFinalized);
@@ -2319,6 +2285,7 @@ void FMaterialShaderMap::ProcessCompilationResults(const TArray<FShaderCommonCom
 	}
 	while ((TimeBudget > 0.0f) && InOutJobIndex < InCompilationResults.Num());
 }
+
 #endif // WITH_EDITOR
 
 class FMaterialShaderMapLayoutCache
@@ -2741,9 +2708,7 @@ void FMaterialShaderMap::LoadMissingShadersFromMemory(const FMaterial* Material)
 	}
 #endif
 }
-#endif // WITH_EDITOR
 
-#if WITH_EDITOR
 const FMemoryImageString* FMaterialShaderMap::GetShaderSource(const FVertexFactoryType* VertexFactoryType, const FShaderType* ShaderType, int32 PermutationId) const
 {
 	FHashedName Key = GetPreprocessedSourceKey(VertexFactoryType, ShaderType, PermutationId);
@@ -2893,7 +2858,9 @@ void FMaterialShaderMap::Release()
 				check(GetShaderPlatform()>= EShaderPlatform::SP_NumPlatforms || GIdToMaterialShaderMap[GetShaderPlatform()].FindRef(ShaderMapId) != this);
 			}
 
+#if WITH_EDITOR
 			FinalizedClone.SafeRelease();
+#endif // WITH_EDITOR
 			check(!bDeletedThroughDeferredCleanup);
 			bDeletedThroughDeferredCleanup = true;
 		}
@@ -2905,7 +2872,6 @@ void FMaterialShaderMap::Release()
 }
 
 FMaterialShaderMap::FMaterialShaderMap() :
-	CompilingId(0u),
 	NumRefs(0),
 	bDeletedThroughDeferredCleanup(false),
 	bRegistered(false),
@@ -2958,6 +2924,7 @@ FMaterialShaderMap::~FMaterialShaderMap()
 #endif
 }
 
+#if WITH_EDITOR
 FMaterialShaderMap* FMaterialShaderMap::AcquireFinalizedClone()
 {
 	checkSlow(IsInGameThread());
@@ -2988,9 +2955,7 @@ FMaterialShaderMap* FMaterialShaderMap::AcquireFinalizedClone()
 	Clone->bCompiledSuccessfully = bCompiledSuccessfully;
 	Clone->bIsPersistent = bIsPersistent;
 	Clone->AssignCopy(*this);
-#if WITH_EDITOR
 	Clone->AssociateWithAssets(GetAssociatedAssets());
-#endif
 
 	FinalizedClone = Clone;
 	return Clone;
@@ -3000,60 +2965,7 @@ FMaterialShaderMap* FMaterialShaderMap::GetFinalizedClone() const
 {
 	return FinalizedClone;
 }
-
-/**
- * Removes all entries in the cache with exceptions based on a shader type
- * @param ShaderType - The shader type to flush
- */
-void FMaterialShaderMap::FlushShadersByShaderType(const FShaderType* ShaderType)
-{
-	// flush from all the vertex factory shader maps
-	for (int32 Index = 0; Index < GetContent()->OrderedMeshShaderMaps.Num(); Index++)
-	{
-		FMeshMaterialShaderMap* MeshShaderMap = GetContent()->OrderedMeshShaderMaps[Index];
-		if (MeshShaderMap)
-		{
-			MeshShaderMap->FlushShadersByShaderType(ShaderType);
-		}
-	}
-
-	if (ShaderType->GetMaterialShaderType())
-	{
-		const int32 PermutationCount = ShaderType->GetPermutationCount();
-		for (int32 PermutationId = 0; PermutationId < PermutationCount; ++PermutationId)
-		{
-			GetMutableContent()->RemoveShaderTypePermutaion(ShaderType->GetMaterialShaderType(), PermutationId);
-		}
-	}
-}
-
-void FMaterialShaderMap::FlushShadersByShaderPipelineType(const FShaderPipelineType* ShaderPipelineType)
-{
-	// flush from all the vertex factory shader maps
-	for (int32 Index = 0;Index < GetContent()->OrderedMeshShaderMaps.Num();Index++)
-	{
-		FMeshMaterialShaderMap* MeshShaderMap = GetContent()->OrderedMeshShaderMaps[Index];
-		if (MeshShaderMap)
-		{
-			MeshShaderMap->FlushShadersByShaderPipelineType(ShaderPipelineType);
-		}
-	}
-
-	if (ShaderPipelineType->IsMaterialTypePipeline())
-	{
-		GetMutableContent()->RemoveShaderPipelineType(ShaderPipelineType);
-	}
-}
-
-
-/**
- * Removes all entries in the cache with exceptions based on a vertex factory type
- * @param ShaderType - The shader type to flush
- */
-void FMaterialShaderMap::FlushShadersByVertexFactoryType(const FVertexFactoryType* VertexFactoryType)
-{
-	GetMutableContent()->RemoveMeshShaderMap(VertexFactoryType->GetHashedName());
-}
+#endif // WITH_EDITOR
 
 bool FMaterialShaderMap::Serialize(FArchive& Ar, bool bInlineShaderResources, bool bLoadedByCookedMaterial, bool bInlineShaderCode)
 {
@@ -3078,54 +2990,6 @@ bool FMaterialShaderMap::Serialize(FArchive& Ar, bool bInlineShaderResources, bo
 	return bSerialized;
 }
 
-/*void FMaterialShaderMap::RegisterSerializedShaders(bool bLoadedByCookedMaterial)
-{
-	check(IsInGameThread());
-
-	TShaderMapContent<FMaterialShaderType>::RegisterSerializedShaders(bLoadedByCookedMaterial);
-	
-	for (FMeshMaterialShaderMap* MeshShaderMap : OrderedMeshShaderMaps)
-	{
-		if (MeshShaderMap)
-		{
-			MeshShaderMap->RegisterSerializedShaders(bLoadedByCookedMaterial);
-		}
-	}
-
-	// Trim the mesh shader maps by removing empty entries
-	for (int32 VFIndex = 0; VFIndex < OrderedMeshShaderMaps.Num(); VFIndex++)
-	{
-		if (OrderedMeshShaderMaps[VFIndex] && OrderedMeshShaderMaps[VFIndex]->IsEmpty())
-		{
-			OrderedMeshShaderMaps[VFIndex] = nullptr;
-		}
-	}
-
-	for (int32 Index = MeshShaderMaps.Num() - 1; Index >= 0; Index--)
-	{
-		if (MeshShaderMaps[Index].IsEmpty())
-		{
-			MeshShaderMaps.RemoveAt(Index);
-		}
-	}
-}
-
-void FMaterialShaderMap::DiscardSerializedShaders()
-{
-	TShaderMapContent<FMaterialShaderType>::DiscardSerializedShaders();
-
-	for (int32 VFIndex = 0; VFIndex < OrderedMeshShaderMaps.Num(); VFIndex++)
-	{
-		OrderedMeshShaderMaps[VFIndex] = nullptr;
-	}
-
-	for (int32 Index = MeshShaderMaps.Num() - 1; Index >= 0; Index--)
-	{
-		MeshShaderMaps[Index].DiscardSerializedShaders();
-	}
-	MeshShaderMaps.Empty();
-}*/
-
 #if WITH_EDITOR
 uint32 FMaterialShaderMap::GetMaxTextureSamplers() const
 {
diff --git a/Engine/Source/Runtime/Engine/Private/Materials/MaterialShared.cpp b/Engine/Source/Runtime/Engine/Private/Materials/MaterialShared.cpp
--- a/Engine/Source/Runtime/Engine/Private/Materials/MaterialShared.cpp
+++ b/Engine/Source/Runtime/Engine/Private/Materials/MaterialShared.cpp
@@ -908,10 +908,13 @@ const FMaterialShaderMap* FMaterial::GetShaderMapToUse() const
 		// If we are accessing uniform texture expressions on the game thread, use results from a shader map whose compile is in flight that matches this material
 		// This allows querying what textures a material uses even when it is being asynchronously compiled
 		ShaderMapToUse = GameThreadShaderMap;
+
+#if WITH_EDITOR
 		if (!ShaderMapToUse && GameThreadCompilingShaderMapId != 0u)
 		{
 			ShaderMapToUse = FMaterialShaderMap::FindCompilingShaderMap(GameThreadCompilingShaderMapId);
 		}
+#endif // WITH_EDITOR
 
 		checkf(!ShaderMapToUse || ShaderMapToUse->GetNumRefs() > 0, TEXT("NumRefs %i, GameThreadShaderMap 0x%08x"), ShaderMapToUse->GetNumRefs(), GetGameThreadShaderMap());
 	}
@@ -1177,6 +1180,7 @@ void FMaterial::SetInlineShaderMap(FMaterialShaderMap* InMaterialShaderMap)
 	});
 }
 
+#if WITH_EDITOR
 void FMaterial::SetCompilingShaderMap(FMaterialShaderMap* InMaterialShaderMap)
 {
 	checkSlow(IsInGameThread());
@@ -1218,6 +1222,7 @@ bool FMaterial::ReleaseGameThreadCompilingShaderMap()
 	}
 	return bReleased;
 }
+#endif // WITH_EDITOR
 
 void FMaterial::ReleaseRenderThreadCompilingShaderMap()
 {
@@ -2136,9 +2141,9 @@ bool FMaterial::PrepareDestroy_GameThread()
 {
 	check(IsInGameThread());
 
+#if WITH_EDITOR
 	const bool bReleasedCompilingId = ReleaseGameThreadCompilingShaderMap();
 
-#if WITH_EDITOR
 	if (GIsEditor)
 	{
 		const FSetElementId FoundId = EditorLoadedMaterialResources.FindId(this);
@@ -2148,17 +2153,21 @@ bool FMaterial::PrepareDestroy_GameThread()
 			EditorLoadedMaterialResources.Remove(FoundId);
 		}
 	}
-#endif // WITH_EDITOR
 
 	return bReleasedCompilingId;
+#else
+	return false;
+#endif
 }
 
 void FMaterial::PrepareDestroy_RenderThread()
 {
 	check(IsInRenderingThread());
 
+#if WITH_EDITOR
 	RenderingThreadCompilingShaderMapId = 0u;
 	RenderingThreadPendingCompilerEnvironment.SafeRelease();
+#endif
 }
 
 void FMaterial::DeferredDelete(FMaterial* InMaterial)
@@ -2188,9 +2197,11 @@ void FMaterial::DeferredDelete(FMaterial* InMaterial)
  */
 FMaterial::~FMaterial()
 {
+#if WITH_EDITOR
 	check(GameThreadCompilingShaderMapId == 0u);
 	check(RenderingThreadCompilingShaderMapId == 0u);
 	check(!RenderingThreadPendingCompilerEnvironment.IsValid());
+#endif // WITH_EDITOR
 
 #if UE_CHECK_FMATERIAL_LIFETIME
 	const uint32 NumRemainingRefs = GetRefCount();
@@ -2202,6 +2213,8 @@ FMaterial::~FMaterial()
 #endif // WITH_EDITOR
 }
 
+#if WITH_EDITOR
+
 /** Populates OutEnvironment with defines needed to compile shaders for this material. */
 void FMaterial::SetupMaterialEnvironment(
 	EShaderPlatform Platform,
@@ -2409,6 +2422,7 @@ void FMaterial::SetupMaterialEnvironment(
 		}
 	}
 }
+#endif // WITH_EDITOR
 
 /**
  * Caches the material shaders for this material with no static parameters on the given platform.
@@ -2946,7 +2960,6 @@ void FMaterial::CacheGivenTypes(EShaderPlatform Platform, const TArray<const FVe
 		GShaderCompilingManager->SubmitJobs(CompileJobs, GetBaseMaterialPathName(), GameThreadShaderMap->GetDebugDescription());
 	}
 }
-#endif // WITH_EDITOR
 
 bool FMaterial::Translate_Legacy(const FMaterialShaderMapId& ShaderMapId,
 	const FStaticParameterSet& InStaticParameters,
@@ -2955,7 +2968,6 @@ bool FMaterial::Translate_Legacy(const FMaterialShaderMapId& ShaderMapId,
 	FMaterialCompilationOutput& OutCompilationOutput,
 	TRefCountPtr<FSharedShaderCompilerEnvironment>& OutMaterialEnvironment)
 {
-#if WITH_EDITORONLY_DATA
 	FHLSLMaterialTranslator MaterialTranslator(this, OutCompilationOutput, InStaticParameters, InPlatform, GetQualityLevel(), ShaderMapId.FeatureLevel, InTargetPlatform);
 	const bool bSuccess = MaterialTranslator.Translate();
 	if (bSuccess)
@@ -2969,10 +2981,6 @@ bool FMaterial::Translate_Legacy(const FMaterialShaderMapId& ShaderMapId,
 		OutMaterialEnvironment->IncludeVirtualPathToContentsMap.Add(TEXT("/Engine/Generated/Material.ush"), MaterialShaderCode);
 	}
 	return bSuccess;
-#else
-	checkNoEntry();
-	return false;
-#endif
 }
 
 bool FMaterial::Translate_New(const FMaterialShaderMapId& ShaderMapId,
@@ -2982,13 +2990,8 @@ bool FMaterial::Translate_New(const FMaterialShaderMapId& ShaderMapId,
 	FMaterialCompilationOutput& OutCompilationOutput,
 	TRefCountPtr<FSharedShaderCompilerEnvironment>& OutMaterialEnvironment)
 {
-#if WITH_EDITOR
 	const FMaterialCompileTargetParameters TargetParams(InPlatform, ShaderMapId.FeatureLevel, InTargetPlatform);
 	return MaterialEmitHLSL(TargetParams, InStaticParameters, *this, OutCompilationOutput, OutMaterialEnvironment);
-#else
-	checkNoEntry();
-	return false;
-#endif
 }
 
 bool FMaterial::Translate(const FMaterialShaderMapId& InShaderMapId,
@@ -2998,7 +3001,6 @@ bool FMaterial::Translate(const FMaterialShaderMapId& InShaderMapId,
 	FMaterialCompilationOutput& OutCompilationOutput,
 	TRefCountPtr<FSharedShaderCompilerEnvironment>& OutMaterialEnvironment)
 {
-#if WITH_EDITOR
 	if (InShaderMapId.bUsingNewHLSLGenerator)
 	{
 		return Translate_New(InShaderMapId, InStaticParameters, InPlatform, InTargetPlatform, OutCompilationOutput, OutMaterialEnvironment);
@@ -3007,10 +3009,6 @@ bool FMaterial::Translate(const FMaterialShaderMapId& InShaderMapId,
 	{
 		return Translate_Legacy(InShaderMapId, InStaticParameters, InPlatform, InTargetPlatform, OutCompilationOutput, OutMaterialEnvironment);
 	}
-#else
-	checkNoEntry();
-	return false;
-#endif
 }
 
 /**
@@ -3028,7 +3026,6 @@ bool FMaterial::BeginCompileShaderMap(
 	EMaterialShaderPrecompileMode PrecompileMode,
 	const ITargetPlatform* TargetPlatform)
 {
-#if WITH_EDITORONLY_DATA
 	bool bSuccess = false;
 
 	STAT(double MaterialCompileTime = 0);
@@ -3037,9 +3034,7 @@ bool FMaterial::BeginCompileShaderMap(
 
 	SCOPE_SECONDS_COUNTER(MaterialCompileTime);
 
-#if WITH_EDITOR
 	NewShaderMap->AssociateWithAsset(GetAssetPath());
-#endif
 
 	// Generate the material shader code.
 	FMaterialCompilationOutput NewCompilationOutput;
@@ -3048,10 +3043,9 @@ bool FMaterial::BeginCompileShaderMap(
 
 	if(bSuccess)
 	{
-#if WITH_EDITOR
 		FShaderCompileUtilities::GenerateBrdfHeaders((EShaderPlatform)Platform);
 		FShaderCompileUtilities::ApplyDerivedDefines(*MaterialEnvironment, nullptr, (EShaderPlatform)Platform);
-#endif
+
 		{
 			FShaderParametersMetadata* UniformBufferStruct = NewCompilationOutput.UniformExpressionSet.CreateBufferStruct();
 			SetupMaterialEnvironment(Platform, *UniformBufferStruct, NewCompilationOutput.UniformExpressionSet, *MaterialEnvironment);
@@ -3124,12 +3118,10 @@ bool FMaterial::BeginCompileShaderMap(
 	INC_FLOAT_STAT_BY(STAT_ShaderCompiling_MaterialShaders,(float)MaterialCompileTime);
 
 	return bSuccess;
-#else
-	UE_LOG(LogMaterial, Fatal,TEXT("Not supported."));
-	return false;
-#endif
 }
 
+#endif // WITH_EDITOR
+
 /**
  * Should the shader for this material with the given platform, shader type and vertex 
  * factory type combination be compiled
@@ -3275,7 +3267,6 @@ bool FMaterial::TryGetShaders(const FMaterialShaderTypes& InTypes, const FVertex
 	const bool bIsInGameThread = IsInGameThread();
 	const FMaterialShaderMap* ShaderMap = bIsInGameThread ? GameThreadShaderMap : RenderingThreadShaderMap;
 	const bool bShaderMapComplete = bIsInGameThread ? IsGameThreadShaderMapComplete() : IsRenderingThreadShaderMapComplete();
-	const uint32 CompilingShaderMapId = bIsInGameThread ? GameThreadCompilingShaderMapId : RenderingThreadCompilingShaderMapId;
 
 	if (ShaderMap == nullptr)
 	{
@@ -3338,21 +3329,26 @@ bool FMaterial::TryGetShaders(const FMaterialShaderTypes& InTypes, const FVertex
 						GODSCManager->AddThreadedShaderPipelineRequest(ShaderPlatform, GetFeatureLevel(), GetQualityLevel(), MaterialName, VFTypeName, PipelineName, ShaderStageNamesToCompile);
 					}
 				}
-				else 
+				else
 #endif
-				if (CompilingShaderMapId != 0u)
 				{
-					if (!bShaderMapComplete)
+#if WITH_EDITOR
+					const uint32 CompilingShaderMapId = bIsInGameThread ? GameThreadCompilingShaderMapId : RenderingThreadCompilingShaderMapId;
+					if (CompilingShaderMapId != 0u)
 					{
-						if (InVertexFactoryType)
-						{
-							FMeshMaterialShaderType::BeginCompileShaderPipeline(EShaderCompileJobPriority::ForceLocal, CompilingShaderMapId, kUniqueShaderPermutationId, ShaderPlatform, PermutationFlags, this, RenderingThreadPendingCompilerEnvironment, InVertexFactoryType, InTypes.PipelineType, CompileJobs, nullptr, nullptr);
-						}
-						else
+						if (!bShaderMapComplete)
 						{
-							FMaterialShaderType::BeginCompileShaderPipeline(EShaderCompileJobPriority::ForceLocal, CompilingShaderMapId, ShaderPlatform, PermutationFlags, this, RenderingThreadPendingCompilerEnvironment, InTypes.PipelineType, CompileJobs, nullptr, nullptr);
+							if (InVertexFactoryType)
+							{
+								FMeshMaterialShaderType::BeginCompileShaderPipeline(EShaderCompileJobPriority::ForceLocal, CompilingShaderMapId, kUniqueShaderPermutationId, ShaderPlatform, PermutationFlags, this, RenderingThreadPendingCompilerEnvironment, InVertexFactoryType, InTypes.PipelineType, CompileJobs, nullptr, nullptr);
+							}
+							else
+							{
+								FMaterialShaderType::BeginCompileShaderPipeline(EShaderCompileJobPriority::ForceLocal, CompilingShaderMapId, ShaderPlatform, PermutationFlags, this, RenderingThreadPendingCompilerEnvironment, InTypes.PipelineType, CompileJobs, nullptr, nullptr);
+							}
 						}
 					}
+#endif // WITH_EDITOR
 				}
 			}
 		}
@@ -3390,19 +3386,24 @@ bool FMaterial::TryGetShaders(const FMaterialShaderTypes& InTypes, const FVertex
 					}
 					else
 #endif
-					if (CompilingShaderMapId != 0u)
 					{
-						if (!bShaderMapComplete)
+#if WITH_EDITOR
+						const uint32 CompilingShaderMapId = bIsInGameThread ? GameThreadCompilingShaderMapId : RenderingThreadCompilingShaderMapId;
+						if (CompilingShaderMapId != 0u)
 						{
-							if (InVertexFactoryType)
+							if (!bShaderMapComplete)
 							{
-								ShaderType->AsMeshMaterialShaderType()->BeginCompileShader(EShaderCompileJobPriority::ForceLocal, CompilingShaderMapId, PermutationId, ShaderPlatform, PermutationFlags, this, RenderingThreadPendingCompilerEnvironment, InVertexFactoryType, CompileJobs, nullptr, nullptr);
-							}
-							else
-							{
-								ShaderType->AsMaterialShaderType()->BeginCompileShader(EShaderCompileJobPriority::ForceLocal, CompilingShaderMapId, PermutationId, this, RenderingThreadPendingCompilerEnvironment, ShaderPlatform, PermutationFlags, CompileJobs, nullptr, nullptr);
+								if (InVertexFactoryType)
+								{
+									ShaderType->AsMeshMaterialShaderType()->BeginCompileShader(EShaderCompileJobPriority::ForceLocal, CompilingShaderMapId, PermutationId, ShaderPlatform, PermutationFlags, this, RenderingThreadPendingCompilerEnvironment, InVertexFactoryType, CompileJobs, nullptr, nullptr);
+								}
+								else
+								{
+									ShaderType->AsMaterialShaderType()->BeginCompileShader(EShaderCompileJobPriority::ForceLocal, CompilingShaderMapId, PermutationId, this, RenderingThreadPendingCompilerEnvironment, ShaderPlatform, PermutationFlags, CompileJobs, nullptr, nullptr);
+								}
 							}
 						}
+#endif // WITH_EDITOR
 					}
 				}
 			}
@@ -4202,6 +4203,7 @@ void FMaterialRenderProxy::ReleaseResource()
 	}
 }
 
+#if WITH_EDITOR
 void FMaterial::SubmitCompileJobs_GameThread(EShaderCompileJobPriority Priority)
 {
 	check(IsInGameThread());
@@ -4238,6 +4240,7 @@ void FMaterial::SubmitCompileJobs_RenderThread(EShaderCompileJobPriority Priorit
 		}
 	}
 }
+#endif // WITH_EDITOR
 
 const FMaterial& FMaterialRenderProxy::GetMaterialWithFallback(ERHIFeatureLevel::Type InFeatureLevel, const FMaterialRenderProxy*& OutFallbackMaterialRenderProxy) const
 {
@@ -4255,10 +4258,12 @@ const FMaterial& FMaterialRenderProxy::GetMaterialWithFallback(ERHIFeatureLevel:
 		while (!Material || !Material->IsRenderingThreadShaderMapComplete());
 		OutFallbackMaterialRenderProxy = FallbackMaterialProxy;
 
+#if WITH_EDITOR
 		if (BaseMaterial)
 		{
 			BaseMaterial->SubmitCompileJobs_RenderThread(EShaderCompileJobPriority::Normal);
 		}
+#endif // WITH_EDITOR
 	}
 	return *Material;
 }
diff --git a/Engine/Source/Runtime/Engine/Private/Materials/MeshMaterialShader.cpp b/Engine/Source/Runtime/Engine/Private/Materials/MeshMaterialShader.cpp
--- a/Engine/Source/Runtime/Engine/Private/Materials/MeshMaterialShader.cpp
+++ b/Engine/Source/Runtime/Engine/Private/Materials/MeshMaterialShader.cpp
@@ -22,6 +22,8 @@ namespace MaterialMeshCookStats
 }
 #endif
 
+#if WITH_EDITOR
+
 static void PrepareMeshMaterialShaderCompileJob(EShaderPlatform Platform,
 	EShaderPermutationFlags PermutationFlags,
 	const FMaterial* Material,
@@ -178,6 +180,8 @@ FShader* FMeshMaterialShaderType::FinishCompileShader(
 	return Shader;
 }
 
+#endif // WITH_EDITOR
+
 bool FMeshMaterialShaderType::ShouldCompilePermutation(EShaderPlatform Platform, const FMaterialShaderParameters& MaterialParameters, const FVertexFactoryType* VertexFactoryType, int32 PermutationId, EShaderPermutationFlags Flags) const
 {
 	return FShaderType::ShouldCompilePermutation(FMeshMaterialShaderPermutationParameters(Platform, MaterialParameters, VertexFactoryType, PermutationId, Flags));
@@ -217,13 +221,13 @@ bool FMeshMaterialShaderType::ShouldCompileVertexFactoryPipeline(const FShaderPi
 	return true;
 }
 
+#if WITH_EDITOR
 void FMeshMaterialShaderType::SetupCompileEnvironment(EShaderPlatform Platform, const FMaterialShaderParameters& MaterialParameters, const FVertexFactoryType* VertexFactoryType, int32 PermutationId, EShaderPermutationFlags Flags, FShaderCompilerEnvironment& Environment) const
 {
 	// Allow the shader type to modify its compile environment.
 	FShaderType::ModifyCompilationEnvironment(FMeshMaterialShaderPermutationParameters(Platform, MaterialParameters, VertexFactoryType, PermutationId, Flags), Environment);
 }
 
-#if WITH_EDITOR
 void FMeshMaterialShaderMap::LoadMissingShadersFromMemory(
 	const FSHAHash& MaterialShaderMapHash, 
 	const FMaterial* Material, 
@@ -299,27 +303,3 @@ void FMeshMaterialShaderMap::LoadMissingShadersFromMemory(
 #endif
 }
 #endif // WITH_EDITOR
-
-/**
- * Removes all entries in the cache with exceptions based on a shader type
- * @param ShaderType - The shader type to flush
- */
-void FMeshMaterialShaderMap::FlushShadersByShaderType(const FShaderType* ShaderType)
-{
-	if (ShaderType->GetMeshMaterialShaderType())
-	{
-		const int32 PermutationCount = ShaderType->GetPermutationCount();
-		for (int32 PermutationId = 0; PermutationId < PermutationCount; ++PermutationId)
-		{
-			RemoveShaderTypePermutaion(ShaderType, PermutationId);
-		}
-	}
-}
-
-void FMeshMaterialShaderMap::FlushShadersByShaderPipelineType(const FShaderPipelineType* ShaderPipelineType)
-{
-	if (ShaderPipelineType->IsMeshMaterialTypePipeline())
-	{
-		RemoveShaderPipelineType(ShaderPipelineType);
-	}
-}
diff --git a/Engine/Source/Runtime/Engine/Private/ShaderCompiler/ShaderCompiler.cpp b/Engine/Source/Runtime/Engine/Private/ShaderCompiler/ShaderCompiler.cpp
--- a/Engine/Source/Runtime/Engine/Private/ShaderCompiler/ShaderCompiler.cpp
+++ b/Engine/Source/Runtime/Engine/Private/ShaderCompiler/ShaderCompiler.cpp
@@ -4920,6 +4920,7 @@ bool FShaderCompilingManager::HandlePotentialRetryOnError(TMap<int32, FShaderMap
 {
 	bool bRetryCompile = false;
 
+#if WITH_EDITORONLY_DATA
 	for (TMap<int32, FShaderMapFinalizeResults>::TIterator It(CompletedShaderMaps); It; ++It)
 	{
 		FShaderMapFinalizeResults& Results = It.Value();
@@ -4927,6 +4928,7 @@ bool FShaderCompilingManager::HandlePotentialRetryOnError(TMap<int32, FShaderMap
 		if (!Results.bAllJobsSucceeded)
 		{
 			bool bSpecialEngineMaterial = false;
+
 			const FMaterialShaderMap* ShaderMap = FMaterialShaderMap::FindCompilingShaderMap(It.Key());
 			if (ShaderMap)
 			{
@@ -4940,7 +4942,6 @@ bool FShaderCompilingManager::HandlePotentialRetryOnError(TMap<int32, FShaderMap
 				}
 			}
 
-#if WITH_EDITORONLY_DATA
 
 			if (UE_LOG_ACTIVE(LogShaders, Log) 
 				// Always log detailed errors when a special engine material or global shader fails to compile, as those will be fatal errors
@@ -5048,7 +5049,6 @@ bool FShaderCompilingManager::HandlePotentialRetryOnError(TMap<int32, FShaderMap
 					break;
 				}
 			}
-#endif	//WITH_EDITORONLY_DATA
 		}
 	}
 
@@ -5114,6 +5114,7 @@ bool FShaderCompilingManager::HandlePotentialRetryOnError(TMap<int32, FShaderMap
 
 		check(CompletedShaderMaps.Num() == OriginalNumShaderMaps);
 	}
+#endif	//WITH_EDITORONLY_DATA
 
 	return bRetryCompile;
 }
@@ -6553,6 +6554,7 @@ void ProcessCookOnTheFlyShaders(bool bReloadGlobalShaders, const TArray<uint8>&
 */
 void RecompileGlobalShaders()
 {
+#if WITH_EDITOR
 	if (!FPlatformProperties::RequiresCookedData())
 	{
 		// Flush pending accesses to the existing global shaders.
@@ -6566,6 +6568,7 @@ void RecompileGlobalShaders()
 
 		GShaderCompilingManager->ProcessAsyncResults(false, true);
 	}
+#endif // WITH_EDITOR
 }
 
 void GetOutdatedShaderTypes(TArray<const FShaderType*>& OutdatedShaderTypes, TArray<const FShaderPipelineType*>& OutdatedShaderPipelineTypes, TArray<const FVertexFactoryType*>& OutdatedFactoryTypes)
@@ -6614,6 +6617,7 @@ bool RecompileShaders(const TCHAR* Cmd, FOutputDevice& Ar)
 		return true;
 	}
 
+#if WITH_EDITOR
 	FString FlagStr(FParse::Token(Cmd, 0));
 	if( FlagStr.Len() > 0 )
 	{
@@ -6756,13 +6760,17 @@ bool RecompileShaders(const TCHAR* Cmd, FOutputDevice& Ar)
 
 		GWarn->EndSlowTask();
 
-		return 1;
+		return true;
 	}
 
 	UE_LOG(LogShaderCompilers, Warning, TEXT("Invalid parameter. Options are: \n'Changed', 'Global', 'Material [name]', 'All'."));
-	return 1;
+#endif // WITH_EDITOR
+
+	return true;
 }
 
+#if WITH_EDITOR
+
 static void PrepareGlobalShaderCompileJob(EShaderPlatform Platform,
 	EShaderPermutationFlags PermutationFlags,
 	const FShaderPipelineType* ShaderPipeline,
@@ -6880,6 +6888,9 @@ FShader* FGlobalShaderTypeCompiler::FinishCompileShader(const FGlobalShaderType*
 	return Shader;
 }
 
+#endif // WITH_EDITOR
+
+#if WITH_EDITORONLY_DATA
 namespace ShaderCompilerUtil
 {
 	FOnGlobalShadersCompilation GOnGlobalShdersCompilationDelegate;
@@ -6889,6 +6900,7 @@ FOnGlobalShadersCompilation& GetOnGlobalShaderCompilation()
 {
 	return ShaderCompilerUtil::GOnGlobalShdersCompilationDelegate;
 }
+#endif // WITH_EDITORONLY_DATA
 
 /**
 * Makes sure all global shaders are loaded and/or compiled for the passed in platform.
@@ -6933,11 +6945,13 @@ void VerifyGlobalShaders(EShaderPlatform Platform, const ITargetPlatform* Target
 		bErrorOnMissing = true;
 	}
 
+#if WITH_EDITOR
 	// All jobs, single & pipeline
 	TArray<FShaderCommonCompileJobPtr> GlobalShaderJobs;
 
 	// Add the single jobs first
 	TMap<TShaderTypePermutation<const FShaderType>, FShaderCompileJob*> SharedShaderJobs;
+#endif // WITH_EDITOR
 
 	for (TLinkedList<FShaderType*>::TIterator ShaderTypeIt(FShaderType::GetTypeList()); ShaderTypeIt; ShaderTypeIt.Next())
 	{
@@ -6961,19 +6975,23 @@ void VerifyGlobalShaders(EShaderPlatform Platform, const ITargetPlatform* Target
 						UE_LOG(LogShaders, Fatal, TEXT("Missing global shader %s's permutation %i, Please make sure cooking was successful."),
 							GlobalShaderType->GetName(), PermutationId);
 					}
-
-					if (OutdatedShaderTypes)
+					else
 					{
-						// Remove old shader, if it exists
-						GlobalShaderMap->RemoveShaderTypePermutaion(GlobalShaderType, PermutationId);
-					}
+#if WITH_EDITOR
+						if (OutdatedShaderTypes)
+						{
+							// Remove old shader, if it exists
+							GlobalShaderMap->RemoveShaderTypePermutaion(GlobalShaderType, PermutationId);
+						}
 
-					// Compile this global shader type.
-					FGlobalShaderTypeCompiler::BeginCompileShader(GlobalShaderType, PermutationId, Platform, PermutationFlags, GlobalShaderJobs);
-					//TShaderTypePermutation<const FShaderType> ShaderTypePermutation(GlobalShaderType, PermutationId);
-					//check(!SharedShaderJobs.Find(ShaderTypePermutation));
-					//SharedShaderJobs.Add(ShaderTypePermutation, Job);
-					PermutationCountToCompile++;
+						// Compile this global shader type.
+						FGlobalShaderTypeCompiler::BeginCompileShader(GlobalShaderType, PermutationId, Platform, PermutationFlags, GlobalShaderJobs);
+						//TShaderTypePermutation<const FShaderType> ShaderTypePermutation(GlobalShaderType, PermutationId);
+						//check(!SharedShaderJobs.Find(ShaderTypePermutation));
+						//SharedShaderJobs.Add(ShaderTypePermutation, Job);
+						PermutationCountToCompile++;
+#endif // WITH_EDITOR
+					}
 				}
 			}
 		}
@@ -7009,37 +7027,42 @@ void VerifyGlobalShaders(EShaderPlatform Platform, const ITargetPlatform* Target
 				{
 					UE_LOG(LogShaders, Fatal, TEXT("Missing global shader pipeline %s, Please make sure cooking was successful."), Pipeline->GetName());
 				}
-
-				if (!bEmptyMap)
-				{
-					UE_LOG(LogShaders, Log, TEXT("	%s"), Pipeline->GetName());
-				}
-
-				if (Pipeline->ShouldOptimizeUnusedOutputs(Platform))
-				{
-					// Make a pipeline job with all the stages
-					FGlobalShaderTypeCompiler::BeginCompileShaderPipeline(Platform, PermutationFlags, Pipeline, GlobalShaderJobs);
-				}
 				else
 				{
-					// If sharing shaders amongst pipelines, add this pipeline as a dependency of an existing individual job
-					for (const FShaderType* ShaderType : Pipeline->GetStages())
+#if WITH_EDITOR
+					if (!bEmptyMap)
+					{
+						UE_LOG(LogShaders, Log, TEXT("	%s"), Pipeline->GetName());
+					}
+
+					if (Pipeline->ShouldOptimizeUnusedOutputs(Platform))
 					{
-						TShaderTypePermutation<const FShaderType> ShaderTypePermutation(ShaderType, kUniqueShaderPermutationId);
-
-						FShaderCompileJob** Job = SharedShaderJobs.Find(ShaderTypePermutation);
-						checkf(Job, TEXT("Couldn't find existing shared job for global shader %s on pipeline %s!"), ShaderType->GetName(), Pipeline->GetName());
-						auto* SingleJob = (*Job)->GetSingleShaderJob();
-						check(SingleJob);
-						auto& SharedPipelinesInJob = SingleJob->SharingPipelines.FindOrAdd(nullptr);
-						check(!SharedPipelinesInJob.Contains(Pipeline));
-						SharedPipelinesInJob.Add(Pipeline);
+						// Make a pipeline job with all the stages
+						FGlobalShaderTypeCompiler::BeginCompileShaderPipeline(Platform, PermutationFlags, Pipeline, GlobalShaderJobs);
 					}
+					else
+					{
+						// If sharing shaders amongst pipelines, add this pipeline as a dependency of an existing individual job
+						for (const FShaderType* ShaderType : Pipeline->GetStages())
+						{
+							TShaderTypePermutation<const FShaderType> ShaderTypePermutation(ShaderType, kUniqueShaderPermutationId);
+
+							FShaderCompileJob** Job = SharedShaderJobs.Find(ShaderTypePermutation);
+							checkf(Job, TEXT("Couldn't find existing shared job for global shader %s on pipeline %s!"), ShaderType->GetName(), Pipeline->GetName());
+							auto* SingleJob = (*Job)->GetSingleShaderJob();
+							check(SingleJob);
+							auto& SharedPipelinesInJob = SingleJob->SharingPipelines.FindOrAdd(nullptr);
+							check(!SharedPipelinesInJob.Contains(Pipeline));
+							SharedPipelinesInJob.Add(Pipeline);
+						}
+					}
+#endif // WITH_EDITOR
 				}
 			}
 		}
 	}
 
+#if WITH_EDITOR
 	if (GlobalShaderJobs.Num() > 0)
 	{
 		GetOnGlobalShaderCompilation().Broadcast();
@@ -7063,6 +7086,7 @@ void VerifyGlobalShaders(EShaderPlatform Platform, const ITargetPlatform* Target
 			GShaderCompilingManager->FinishCompilation(TEXT("Global"), ShaderMapIds);
 		}
 	}
+#endif // WITH_EDITOR
 }
 
 void VerifyGlobalShaders(EShaderPlatform Platform, bool bLoadedFromCacheFile, const TArray<const FShaderType*>* OutdatedShaderTypes, const TArray<const FShaderPipelineType*>* OutdatedShaderPipelineTypes)
@@ -7972,6 +7996,7 @@ void RecompileShadersForRemote(
 
 void BeginRecompileGlobalShaders(const TArray<const FShaderType*>& OutdatedShaderTypes, const TArray<const FShaderPipelineType*>& OutdatedShaderPipelineTypes, EShaderPlatform ShaderPlatform, const ITargetPlatform* TargetPlatform)
 {
+#if WITH_EDITOR
 	if (!FPlatformProperties::RequiresCookedData())
 	{
 		// Flush pending accesses to the existing global shaders.
@@ -7987,6 +8012,7 @@ void BeginRecompileGlobalShaders(const TArray<const FShaderType*>& OutdatedShade
 			VerifyGlobalShaders(ShaderPlatform, TargetPlatform, false, &OutdatedShaderTypes, &OutdatedShaderPipelineTypes);
 		}
 	}
+#endif
 }
 
 void FinishRecompileGlobalShaders()
@@ -7995,6 +8021,8 @@ void FinishRecompileGlobalShaders()
 	GShaderCompilingManager->ProcessAsyncResults(false, true);
 }
 
+#if WITH_EDITOR
+
 static inline FShader* ProcessCompiledJob(FShaderCompileJob* SingleJob, const FShaderPipelineType* Pipeline, TArray<EShaderPlatform>& ShaderPlatformsProcessed, TArray<const FShaderPipelineType*>& OutSharedPipelines)
 {
 	const FGlobalShaderType* GlobalShaderType = SingleJob->Key.ShaderType->GetGlobalShaderType();
@@ -8137,6 +8165,7 @@ void SaveGlobalShadersForRemoteRecompile(FArchive& Ar, EShaderPlatform ShaderPla
 		GlobalShaderMap->SaveToGlobalArchive(Ar);
 	}
 }
+#endif // WITH_EDITOR
 
 void LoadGlobalShadersForRemoteRecompile(FArchive& Ar, EShaderPlatform ShaderPlatform)
 {
diff --git a/Engine/Source/Runtime/Engine/Public/MaterialShaderType.h b/Engine/Source/Runtime/Engine/Public/MaterialShaderType.h
--- a/Engine/Source/Runtime/Engine/Public/MaterialShaderType.h
+++ b/Engine/Source/Runtime/Engine/Public/MaterialShaderType.h
@@ -121,6 +121,7 @@ public:
 			InSourceFilename);
 	}
 
+#if WITH_EDITOR
 	/**
 	 * Enqueues a compilation for a new shader of this type.
 	 * @param Material - The material to link the shader with.
@@ -163,6 +164,7 @@ public:
 		const FShaderPipelineType* ShaderPipeline,
 		const FString& InDebugDescription
 		) const;
+#endif // WITH_EDITOR
 
 	/**
 	 * Checks if the shader type should be cached for a particular platform and material.
@@ -174,12 +176,14 @@ public:
 
 	static bool ShouldCompilePipeline(const FShaderPipelineType* ShaderPipelineType, EShaderPlatform Platform, const FMaterialShaderParameters& MaterialParameters, EShaderPermutationFlags Flags);
 
+#if WITH_EDITOR
 	/**
 	 * Sets up the environment used to compile an instance of this shader type.
 	 * @param Platform - Platform to compile for.
 	 * @param Environment - The shader compile environment that the function modifies.
 	 */
 	void SetupCompileEnvironment(EShaderPlatform Platform, const FMaterialShaderParameters& MaterialParameters, int32 PermutationId, EShaderPermutationFlags Flags, FShaderCompilerEnvironment& Environment) const;
+#endif // WITH_EDITOR
 };
 
 struct FMaterialShaderTypes
diff --git a/Engine/Source/Runtime/Engine/Public/MaterialShared.h b/Engine/Source/Runtime/Engine/Public/MaterialShared.h
--- a/Engine/Source/Runtime/Engine/Public/MaterialShared.h
+++ b/Engine/Source/Runtime/Engine/Public/MaterialShared.h
@@ -1020,18 +1020,6 @@ public:
 		EShaderPlatform Platform);
 #endif
 
-	/**
-	 * Removes all entries in the cache with exceptions based on a shader type
-	 * @param ShaderType - The shader type to flush
-	 */
-	void FlushShadersByShaderType(const FShaderType* ShaderType);
-
-		/**
-	 * Removes all entries in the cache with exceptions based on a shader type
-	 * @param ShaderType - The shader type to flush
-	 */
-	void FlushShadersByShaderPipelineType(const FShaderPipelineType* ShaderPipelineType);
-
 	// Accessors.
 	inline const FHashedName& GetVertexFactoryTypeName() const { return VertexFactoryTypeName; }
 
@@ -1127,14 +1115,9 @@ public:
 	 */
 	static TRefCountPtr<FMaterialShaderMap> FindId(const FMaterialShaderMapId& ShaderMapId, EShaderPlatform Platform);
 
+#if WITH_EDITOR
 	static FMaterialShaderMap* FindCompilingShaderMap(uint32 CompilingId);
 
-#if ALLOW_SHADERMAP_DEBUG_DATA
-	/** Flushes the given shader types from any loaded FMaterialShaderMap's. */
-	static void FlushShaderTypes(TArray<const FShaderType*>& ShaderTypesToFlush, TArray<const FShaderPipelineType*>& ShaderPipelineTypesToFlush, TArray<const FVertexFactoryType*>& VFTypesToFlush);
-#endif
-
-#if WITH_EDITOR
 	/** Gets outdated types from all loaded material shader maps */
 	static void GetAllOutdatedTypes(TArray<const FShaderType*>& OutdatedShaderTypes, TArray<const FShaderPipelineType*>& OutdatedShaderPipelineTypes, TArray<const FVertexFactoryType*>& OutdatedFactoryTypes);
 
@@ -1164,9 +1147,10 @@ public:
 	FMaterialShaderMap();
 	virtual ~FMaterialShaderMap();
 	
+#if WITH_EDITOR
 	FMaterialShaderMap* AcquireFinalizedClone();
 	FMaterialShaderMap* GetFinalizedClone() const;
-
+#endif // WITH_EDITOR
 
 	// ShaderMap interface
 	TShaderRef<FShader> GetShader(FShaderType* ShaderType, int32 PermutationId = 0) const
@@ -1185,6 +1169,7 @@ public:
 
 	uint32 GetMaxNumInstructionsForShader(FShaderType* ShaderType) const { return GetContent()->GetMaxNumInstructionsForShader(*this, ShaderType); }
 
+#if WITH_EDITOR
 	void SubmitCompileJobs(uint32 CompilingShaderMapId,
 		const FMaterial* Material,
 		const TRefCountPtr<FSharedShaderCompilerEnvironment>& MaterialEnvironment,
@@ -1205,7 +1190,6 @@ public:
 		EMaterialShaderPrecompileMode PrecompileMode
 		);
 
-#if WITH_EDITOR
 	/** Sorts the incoming compiled jobs into the appropriate mesh shader maps */
 	void ProcessCompilationResults(const TArray<TRefCountPtr<FShaderCommonCompileJob>>& ICompilationResults, int32& InOutJobIndex, float& TimeBudget);
 #endif
@@ -1253,24 +1237,6 @@ public:
 	ENGINE_API void AddRef();
 	ENGINE_API void Release();
 
-	/**
-	 * Removes all entries in the cache with exceptions based on a shader type
-	 * @param ShaderType - The shader type to flush
-	 */
-	void FlushShadersByShaderType(const FShaderType* ShaderType);
-
-	/**
-	 * Removes all entries in the cache with exceptions based on a shader pipeline type
-	 * @param ShaderPipelineType - The shader pipeline type to flush
-	 */
-	void FlushShadersByShaderPipelineType(const FShaderPipelineType* ShaderPipelineType);
-
-	/**
-	 * Removes all entries in the cache with exceptions based on a vertex factory type
-	 * @param ShaderType - The shader type to flush
-	 */
-	void FlushShadersByVertexFactoryType(const FVertexFactoryType* VertexFactoryType);
-
 	/** Serializes the shader map. */
 	bool Serialize(FArchive& Ar, bool bInlineShaderResources=true, bool bLoadedByCookedMaterial=false, bool bInlineShaderCode=false);
 
@@ -1281,6 +1247,7 @@ public:
 
 	/** Backs up any FShaders in this shader map to memory through serialization and clears FShader references. */
 	TArray<uint8>* BackupShadersToMemory();
+
 	/** Recreates FShaders from the passed in memory, handling shader key changes. */
 	void RestoreShadersFromMemory(const TArray<uint8>& ShaderData);
 
@@ -1304,16 +1271,20 @@ public:
 	const FMaterialShaderMapId& GetShaderMapId() const { return ShaderMapId; }
 	const FSHAHash& GetShaderContentHash() const { return GetContent()->ShaderContentHash; }
 
+#if WITH_EDITOR
 	uint32 AcquireCompilingId(const TRefCountPtr<FSharedShaderCompilerEnvironment>& InMaterialEnvironment);
 	void ReleaseCompilingId();
 	uint32 GetCompilingId() const { return CompilingId; }
 	const TRefCountPtr<FSharedShaderCompilerEnvironment>& GetPendingCompilerEnvironment() const { return PendingCompilerEnvironment; }
+#endif // WITH_EDITOR
 
 	bool IsCompilationFinalized() const { return bCompilationFinalized; }
 	bool CompiledSuccessfully() const { return bCompiledSuccessfully; }
+#if WITH_EDITOR
 	void SetCompiledSuccessfully(bool bSuccess) { bCompiledSuccessfully = bSuccess; }
 	void AddCompilingDependency(FMaterial* Material);
 	void RemoveCompilingDependency(FMaterial* Material);
+#endif // WITH_EDITOR
 
 #if WITH_EDITORONLY_DATA
 	const TCHAR* GetFriendlyName() const { return *GetContent()->FriendlyName; }
@@ -1414,9 +1385,11 @@ private:
 	static FCriticalSection AllMaterialShaderMapsGuard;
 #endif
 
+#if WITH_EDITOR
 	TRefCountPtr<FMaterialShaderMap> FinalizedClone;
 	TRefCountPtr<FSharedShaderCompilerEnvironment> PendingCompilerEnvironment;
 	TArray<TRefCountPtr<FMaterial>> CompilingMaterialDependencies;
+#endif // WITH_EDITOR
 
 	FUniformBufferLayoutRHIRef UniformBufferLayout;
 
@@ -1430,8 +1403,10 @@ private:
 	/** Tracks material resources and their shader maps that need to be compiled but whose compilation is being deferred. */
 	//static TMap<TRefCountPtr<FMaterialShaderMap>, TArray<FMaterial*> > ShaderMapsBeingCompiled;
 
+#if WITH_EDITOR
 	/** Uniquely identifies this shader map during compilation, needed for deferred compilation where shaders from multiple shader maps are compiled together. */
-	uint32 CompilingId;
+	uint32 CompilingId = 0;
+#endif // WITH_EDITOR
 
 	mutable int32 NumRefs;
 
@@ -1452,7 +1427,9 @@ private:
 	/** Indicates whether the shader map should be stored in the shader cache. */
 	uint32 bIsPersistent : 1;
 
+#if WITH_EDITOR
 	FShader* ProcessCompilationResultsForSingleJob(class FShaderCompileJob* SingleJob, const FShaderPipelineType* ShaderPipeline, const FSHAHash& MaterialShaderMapHash);
+#endif
 
 	friend ENGINE_API void DumpMaterialStats( EShaderPlatform Platform );
 	friend class FShaderCompilingManager;
@@ -1677,8 +1654,6 @@ public:
 	 * Minimal initialization constructor.
 	 */
 	FMaterial():
-		GameThreadCompilingShaderMapId(0u),
-		RenderingThreadCompilingShaderMapId(0u),
 		RenderingThreadShaderMap(NULL),
 		QualityLevel(EMaterialQualityLevel::Num),
 		FeatureLevel(ERHIFeatureLevel::Num),
@@ -2126,8 +2101,10 @@ public:
 
 	ENGINE_API bool ShouldCacheShaders(const EShaderPlatform ShaderPlatform, const FMaterialShaderTypes& InTypes, const FVertexFactoryType* InVertexFactoryType) const;
 
+#if WITH_EDITOR
 	ENGINE_API void SubmitCompileJobs_GameThread(EShaderCompileJobPriority Priority);
 	ENGINE_API void SubmitCompileJobs_RenderThread(EShaderCompileJobPriority Priority) const;
+#endif // WITH_EDITOR
 
 	/** Returns a string that describes the material's usage for debugging purposes. */
 	virtual FString GetMaterialUsageDescription() const = 0;
@@ -2275,10 +2252,14 @@ protected:
 
 	bool GetLoadedCookedShaderMapId() const { return bLoadedCookedShaderMapId; }
 
+#if WITH_EDITOR
 	void SetCompilingShaderMap(FMaterialShaderMap* InMaterialShaderMap);
+#endif
 
 private:
+#if WITH_EDITOR
 	bool ReleaseGameThreadCompilingShaderMap();
+#endif
 	void ReleaseRenderThreadCompilingShaderMap();
 
 #if WITH_EDITOR
@@ -2296,13 +2277,13 @@ private:
 
 	TSharedPtr<FMaterialShaderMap::FAsyncLoadContext> CacheShadersPending;
 	TUniqueFunction<bool ()> CacheShadersCompletion;
-#endif // WITH_EDITOR
 
-	uint32 GameThreadCompilingShaderMapId;
-	uint32 RenderingThreadCompilingShaderMapId;
+	uint32 GameThreadCompilingShaderMapId = 0;
+	uint32 RenderingThreadCompilingShaderMapId = 0;
 
 	TRefCountPtr<FSharedShaderCompilerEnvironment> GameThreadPendingCompilerEnvironment;
 	TRefCountPtr<FSharedShaderCompilerEnvironment> RenderingThreadPendingCompilerEnvironment;
+#endif // WITH_EDITOR
 
 	/** 
 	 * Used to prevent submitting the material more than once during CacheMeshDrawCommands unless priority has been increased. 
@@ -2393,6 +2374,7 @@ private:
 	FRWLock PrecachePSOVFLock;
 	TArray<FPrecacheVertexTypeWithParams> PrecachedPSOVertexFactories;
 
+#if WITH_EDITOR
 	/**
 	* Compiles this material for Platform, storing the result in OutShaderMap if the compile was synchronous
 	*/
@@ -2431,6 +2413,7 @@ private:
 		const FUniformExpressionSet& InUniformExpressionSet,
 		FShaderCompilerEnvironment& OutEnvironment
 		) const;
+#endif // WITH_EDITOR
 
 	/** Get the float precision mode for the material and shader taking into account any project wide precision values */
 	static void GetOutputPrecision(EMaterialFloatPrecisionMode FloatPrecisionMode, bool& bFullPrecisionInPS, bool& bFullPrecisionInMaterial);
diff --git a/Engine/Source/Runtime/Engine/Public/MeshMaterialShaderType.h b/Engine/Source/Runtime/Engine/Public/MeshMaterialShaderType.h
--- a/Engine/Source/Runtime/Engine/Public/MeshMaterialShaderType.h
+++ b/Engine/Source/Runtime/Engine/Public/MeshMaterialShaderType.h
@@ -74,6 +74,7 @@ public:
 			InSourceFilename);
 	}
 
+#if WITH_EDITOR
 	/**
 	 * Enqueues a compilation for a new shader of this type.
 	 * @param Platform - The platform to compile for.
@@ -121,6 +122,7 @@ public:
 		const FShaderPipelineType* ShaderPipeline,
 		const FString& InDebugDescription
 		) const;
+#endif // WITH_EDITOR
 
 	/**
 	 * Checks if the shader type should be cached for a particular platform, material, and vertex factory type.
@@ -137,11 +139,13 @@ public:
 
 	static bool ShouldCompilePipeline(const FShaderPipelineType* ShaderPipelineType, EShaderPlatform Platform, const FMaterialShaderParameters& MaterialParameters, const FVertexFactoryType* VertexFactoryType, EShaderPermutationFlags Flags);
 
+#if WITH_EDITOR
 	/**
 	 * Sets up the environment used to compile an instance of this shader type.
 	 * @param Platform - Platform to compile for.
 	 * @param Environment - The shader compile environment that the function modifies.
 	 */
 	void SetupCompileEnvironment(EShaderPlatform Platform, const FMaterialShaderParameters& MaterialParameters, const FVertexFactoryType* VertexFactoryType, int32 PermutationId, EShaderPermutationFlags PermutationFlags, FShaderCompilerEnvironment& Environment) const;
+#endif // WITH_EDITOR
 };
 
diff --git a/Engine/Source/Runtime/Engine/Public/ShaderCompiler.h b/Engine/Source/Runtime/Engine/Public/ShaderCompiler.h
--- a/Engine/Source/Runtime/Engine/Public/ShaderCompiler.h
+++ b/Engine/Source/Runtime/Engine/Public/ShaderCompiler.h
@@ -497,6 +497,7 @@ private:
 	class IConsoleObject* PrintStatsCmd;
 };
 
+#if WITH_EDITOR
 class FGlobalShaderTypeCompiler
 {
 public:
@@ -513,6 +514,7 @@ public:
 	/** Either returns an equivalent existing shader of this type, or constructs a new instance. */
 	static FShader* FinishCompileShader(const FGlobalShaderType* ShaderType, const FShaderCompileJob& CompileJob, const FShaderPipelineType* ShaderPipelineType);
 };
+#endif // WITH_EDITOR
 
 class FShaderCompileThreadRunnableBase : public FRunnable
 {
@@ -1199,9 +1201,11 @@ extern bool RecompileShaders(const TCHAR* Cmd, FOutputDevice& Ar);
 /** Returns whether all global shader types containing the substring are complete and ready for rendering. if type name is null, check everything */
 extern ENGINE_API bool IsGlobalShaderMapComplete(const TCHAR* TypeNameSubstring = nullptr);
 
+#if WITH_EDITORONLY_DATA
 /** Returns the delegate triggered when global shaders compilation jobs start. */
 DECLARE_MULTICAST_DELEGATE(FOnGlobalShadersCompilation);
 extern ENGINE_API FOnGlobalShadersCompilation& GetOnGlobalShaderCompilation();
+#endif // WITH_EDITORONLY_DATA
 
 /**
 * Makes sure all global shaders are loaded and/or compiled for the passed in platform.
@@ -1240,11 +1244,13 @@ extern ENGINE_API void BeginRecompileGlobalShaders(const TArray<const FShaderTyp
 /** Finishes recompiling global shaders.  Must be called after BeginRecompileGlobalShaders. */
 extern ENGINE_API void FinishRecompileGlobalShaders();
 
+#if WITH_EDITOR
 /** Called by the shader compiler to process completed global shader jobs. */
 extern ENGINE_API void ProcessCompiledGlobalShaders(const TArray<FShaderCommonCompileJobPtr>& CompilationResults);
 
 /** Serializes a global shader map to an archive (used with recompiling shaders for a remote console) */
 extern ENGINE_API void SaveGlobalShadersForRemoteRecompile(FArchive& Ar, EShaderPlatform ShaderPlatform);
+#endif // WITH_EDITOR
 
 /** Serializes a global shader map to an archive (used with recompiling shaders for a remote console) */
 extern ENGINE_API void LoadGlobalShadersForRemoteRecompile(FArchive& Ar, EShaderPlatform ShaderPlatform);
diff --git a/Engine/Source/Runtime/RenderCore/Private/GlobalShader.cpp b/Engine/Source/Runtime/RenderCore/Private/GlobalShader.cpp
--- a/Engine/Source/Runtime/RenderCore/Private/GlobalShader.cpp
+++ b/Engine/Source/Runtime/RenderCore/Private/GlobalShader.cpp
@@ -504,6 +504,7 @@ FGlobalShader::FGlobalShader(const ShaderMetaType::CompiledShaderInitializerType
 :	FShader(Initializer)
 {}
 
+#if WITH_EDITOR
 void FGlobalShaderType::SetupCompileEnvironment(EShaderPlatform Platform, int32 PermutationId, EShaderPermutationFlags Flags, FShaderCompilerEnvironment& Environment) const
 {
 	FGlobalShaderPermutationParameters Parameters(GetFName(), Platform, PermutationId, Flags);
@@ -514,6 +515,7 @@ void FGlobalShaderType::SetupCompileEnvironment(EShaderPlatform Platform, int32
 	// Allow the shader type to modify its compile environment.
 	ModifyCompilationEnvironment(Parameters, Environment);
 }
+#endif // WITH_EDITOR
 
 void BackupGlobalShaderMap(FGlobalShaderBackupData& OutGlobalShaderBackup)
 {
diff --git a/Engine/Source/Runtime/RenderCore/Public/GlobalShader.h b/Engine/Source/Runtime/RenderCore/Public/GlobalShader.h
--- a/Engine/Source/Runtime/RenderCore/Public/GlobalShader.h
+++ b/Engine/Source/Runtime/RenderCore/Public/GlobalShader.h
@@ -127,12 +127,14 @@ public:
 		return FShaderType::ShouldCompilePermutation(FGlobalShaderPermutationParameters(GetFName(), Platform, PermutationId, Flags));
 	}
 
+#if WITH_EDITOR
 	/**
 	 * Sets up the environment used to compile an instance of this shader type.
 	 * @param Platform - Platform to compile for.
 	 * @param Environment - The shader compile environment that the function modifies.
 	 */
 	void SetupCompileEnvironment(EShaderPlatform Platform, int32 PermutationId, EShaderPermutationFlags Flags, FShaderCompilerEnvironment& Environment) const;
+#endif // WITH_EDITOR
 
 	static bool ShouldCompilePipeline(const FShaderPipelineType* ShaderPipelineType, EShaderPlatform Platform, EShaderPermutationFlags Flags);
 };
