Add new world delegate for when an actor has been removed from the world. This allows us to know the earliest point an actor can be renamed and not cause issues with replication.
This fixes an issue when attempting to rename actors after a frame delay as GC could have happened inbetween those times causing some actors to be missed.

diff --git a/Engine/Source/Runtime/Engine/Classes/Engine/World.h b/Engine/Source/Runtime/Engine/Classes/Engine/World.h
--- a/Engine/Source/Runtime/Engine/Classes/Engine/World.h
+++ b/Engine/Source/Runtime/Engine/Classes/Engine/World.h
@@ -149,6 +149,7 @@ ENGINE_API DECLARE_LOG_CATEGORY_EXTERN(LogSpawn, Warning, All);
 
 DECLARE_MULTICAST_DELEGATE_OneParam(FOnActorSpawned, AActor*);
 DECLARE_MULTICAST_DELEGATE_OneParam(FOnActorDestroyed, AActor*);
+DECLARE_MULTICAST_DELEGATE_OneParam(FOnActorRemovedFromWorld, AActor*);
 DECLARE_MULTICAST_DELEGATE_OneParam(FOnFeatureLevelChanged, ERHIFeatureLevel::Type);
 DECLARE_MULTICAST_DELEGATE_OneParam(FOnMovieSceneSequenceTick, float);
 
@@ -1497,6 +1498,9 @@ private:
 	/** a delegate that broadcasts a notification whenever an actor is destroyed */
 	mutable FOnActorDestroyed OnActorDestroyed;
 
+	/** Broadcasts when an actor has been removed from the world. The earliest point an actor can safely be renamed and not affect replication */
+	mutable FOnActorRemovedFromWorld OnActorRemovedFromWorld;
+
 	/** Reset Async Trace Buffer **/
 	void ResetAsyncTrace();
 
@@ -2743,6 +2747,12 @@ public:
 	/** Remove a listener for OnActorDestroyed events */
 	void RemoveOnActorDestroyededHandler(FDelegateHandle InHandle) const;
 
+	/** Add a listener for OnActorRemovedFromWorld events */
+	FDelegateHandle AddOnActorRemovedFromWorldHandler(const FOnActorRemovedFromWorld::FDelegate& InHandler) const;
+
+	/** Remove a listener for OnActorRemovedFromWorld events */
+	void RemoveOnActorRemovedFromWorldHandler(FDelegateHandle InHandle) const;
+
 	/**
 	 * Returns whether the passed in actor is part of any of the loaded levels actors array.
 	 * Warning: Will return true for pending kill actors!
diff --git a/Engine/Source/Runtime/Engine/Private/LevelActor.cpp b/Engine/Source/Runtime/Engine/Private/LevelActor.cpp
--- a/Engine/Source/Runtime/Engine/Private/LevelActor.cpp
+++ b/Engine/Source/Runtime/Engine/Private/LevelActor.cpp
@@ -996,6 +996,8 @@ bool UWorld::DestroyActor( AActor* ThisActor, bool bNetForce, bool bShouldModify
 #endif
 	}
 
+	OnActorRemovedFromWorld.Broadcast(ThisActor);
+
 	// Clean up the actor's components.
 	ThisActor->UnregisterAllComponents();
 
diff --git a/Engine/Source/Runtime/Engine/Private/World.cpp b/Engine/Source/Runtime/Engine/Private/World.cpp
--- a/Engine/Source/Runtime/Engine/Private/World.cpp
+++ b/Engine/Source/Runtime/Engine/Private/World.cpp
@@ -5357,6 +5357,16 @@ void UWorld::RemoveOnActorDestroyededHandler(FDelegateHandle InHandle) const
 	OnActorDestroyed.Remove(InHandle);
 }
 
+FDelegateHandle UWorld::AddOnActorRemovedFromWorldHandler(const FOnActorRemovedFromWorld::FDelegate& InHandler) const
+{
+	return OnActorRemovedFromWorld.Add(InHandler);
+}
+
+void UWorld::RemoveOnActorRemovedFromWorldHandler(FDelegateHandle InHandle) const
+{
+	OnActorRemovedFromWorld.Remove(InHandle);
+}
+
 FDelegateHandle UWorld::AddMovieSceneSequenceTickHandler(const FOnMovieSceneSequenceTick::FDelegate& InHandler)
 {
 	return MovieSceneSequenceTick.Add(InHandler);
