Expose DataLayerLoadingPolicy for licensees to be able to override the default OR behavior of Editor Is Loaded flag.

diff --git a/Engine/Source/Runtime/Engine/Private/WorldPartition/DataLayer/DataLayerLoadingPolicy.cpp b/Engine/Source/Runtime/Engine/Private/WorldPartition/DataLayer/DataLayerLoadingPolicy.cpp
new file mode 100644
--- /dev/null
+++ b/Engine/Source/Runtime/Engine/Private/WorldPartition/DataLayer/DataLayerLoadingPolicy.cpp
@@ -0,0 +1,19 @@
+// Copyright Epic Games, Inc. All Rights Reserved.
+
+#include "WorldPartition/DataLayer/DataLayerLoadingPolicy.h"
+#include "WorldPartition/DataLayer/DataLayerInstance.h"
+
+#if WITH_EDITOR
+bool UDataLayerLoadingPolicy::ResolveIsLoadedInEditor(TArray<const UDataLayerInstance*>& InDataLayers) const
+{
+	check(!InDataLayers.IsEmpty());
+	for (const UDataLayerInstance* DataLayerInstance : InDataLayers)
+	{
+		if (DataLayerInstance->IsEffectiveLoadedInEditor())
+		{
+			return true;
+		}
+	}
+	return false;
+}
+#endif
\ No newline at end of file
diff --git a/Engine/Source/Runtime/Engine/Private/WorldPartition/DataLayer/DataLayerSubsystem.cpp b/Engine/Source/Runtime/Engine/Private/WorldPartition/DataLayer/DataLayerSubsystem.cpp
--- a/Engine/Source/Runtime/Engine/Private/WorldPartition/DataLayer/DataLayerSubsystem.cpp
+++ b/Engine/Source/Runtime/Engine/Private/WorldPartition/DataLayer/DataLayerSubsystem.cpp
@@ -1,6 +1,7 @@
 // Copyright Epic Games, Inc. All Rights Reserved.
 
 #include "WorldPartition/DataLayer/DataLayerSubsystem.h"
+#include "WorldPartition/DataLayer/DataLayerLoadingPolicy.h"
 #include "WorldPartition/DataLayer/WorldDataLayers.h"
 #include "WorldPartition/DataLayer/WorldDataLayersActorDesc.h"
 #include "WorldPartition/DataLayer/DataLayerAsset.h"
@@ -76,6 +77,14 @@ void UDataLayerSubsystem::Initialize(FSubsystemCollectionBase& Collection)
 	Super::Initialize(Collection);
 
 #if WITH_EDITOR
+	UClass* DataLayerLoadingPolicyClassValue = DataLayerLoadingPolicyClass.Get();
+	if (!DataLayerLoadingPolicyClassValue)
+	{
+		DataLayerLoadingPolicyClassValue = UDataLayerLoadingPolicy::StaticClass();
+	}
+	DataLayerLoadingPolicy = NewObject<UDataLayerLoadingPolicy>(this, DataLayerLoadingPolicyClassValue);
+	check(DataLayerLoadingPolicy);
+
 	if (GEditor)
 	{
 		FModuleManager::LoadModuleChecked<IDataLayerEditorModule>("DataLayerEditor");
@@ -106,6 +115,23 @@ bool UDataLayerSubsystem::DoesSupportWorldType(const EWorldType::Type WorldType)
 
 #if WITH_EDITOR
 
+bool UDataLayerSubsystem::ResolveIsLoadedInEditor(const TArray<FName>& InDataLayerInstanceNames) const
+{
+	TArray<const UDataLayerInstance*> DataLayerInstances;
+	for (const FName& DataLayerInstanceName : InDataLayerInstanceNames)
+	{
+		if (const UDataLayerInstance* DataLayerInstance = GetDataLayerInstance(DataLayerInstanceName))
+		{
+			DataLayerInstances.Add(DataLayerInstance);
+		}
+	}
+	if (!DataLayerInstances.IsEmpty())
+	{
+		return DataLayerLoadingPolicy->ResolveIsLoadedInEditor(DataLayerInstances);
+	}
+	return true;
+}
+
 bool UDataLayerSubsystem::CanResolveDataLayers() const
 {
 	return !WorldDataLayerCollection.IsEmpty();
diff --git a/Engine/Source/Runtime/Engine/Private/WorldPartition/WorldPartitionActorLoaderInterface.cpp b/Engine/Source/Runtime/Engine/Private/WorldPartition/WorldPartitionActorLoaderInterface.cpp
--- a/Engine/Source/Runtime/Engine/Private/WorldPartition/WorldPartitionActorLoaderInterface.cpp
+++ b/Engine/Source/Runtime/Engine/Private/WorldPartition/WorldPartitionActorLoaderInterface.cpp
@@ -143,19 +143,7 @@ bool IWorldPartitionActorLoaderInterface::ILoaderAdapter::PassDataLayersFilter(c
 		}
 		else
 		{
-			uint32 NumValidLayers = 0;
-			for (const FName& DataLayerInstanceName : ActorDescProxy.GetDataLayers())
-			{
-				if (const UDataLayerInstance* DataLayerInstance = DataLayerSubsystem->GetDataLayerInstance(DataLayerInstanceName))
-				{
-					if (DataLayerInstance->IsEffectiveLoadedInEditor())
-					{
-						return true;
-					}
-					NumValidLayers++;
-				}
-			}
-			return !NumValidLayers;
+			return DataLayerSubsystem->ResolveIsLoadedInEditor(ActorDescProxy.GetDataLayers());
 		}
 	}
 
diff --git a/Engine/Source/Runtime/Engine/Public/WorldPartition/DataLayer/DataLayerLoadingPolicy.h b/Engine/Source/Runtime/Engine/Public/WorldPartition/DataLayer/DataLayerLoadingPolicy.h
new file mode 100644
--- /dev/null
+++ b/Engine/Source/Runtime/Engine/Public/WorldPartition/DataLayer/DataLayerLoadingPolicy.h
@@ -0,0 +1,19 @@
+// Copyright Epic Games, Inc. All Rights Reserved.
+
+#pragma once
+
+#include "CoreMinimal.h"
+#include "UObject/Object.h"
+#include "DataLayerLoadingPolicy.generated.h"
+
+class UDataLayerInstance;
+
+UCLASS()
+class UDataLayerLoadingPolicy : public UObject
+{
+	GENERATED_BODY()
+public:
+#if WITH_EDITOR
+	virtual bool ResolveIsLoadedInEditor(TArray<const UDataLayerInstance*>& InDataLayers) const;
+#endif
+};
\ No newline at end of file
diff --git a/Engine/Source/Runtime/Engine/Public/WorldPartition/DataLayer/DataLayerSubsystem.h b/Engine/Source/Runtime/Engine/Public/WorldPartition/DataLayer/DataLayerSubsystem.h
--- a/Engine/Source/Runtime/Engine/Public/WorldPartition/DataLayer/DataLayerSubsystem.h
+++ b/Engine/Source/Runtime/Engine/Public/WorldPartition/DataLayer/DataLayerSubsystem.h
@@ -14,6 +14,7 @@
  */
 
 class UActorDescContainer;
+class UDataLayerLoadingPolicy;
 class UCanvas;
 
 DECLARE_DYNAMIC_MULTICAST_DELEGATE_TwoParams(FOnDataLayerRuntimeStateChanged, const UDataLayerInstance*, DataLayer, EDataLayerRuntimeState, State);
@@ -42,7 +43,7 @@ private:
 };
 #endif
 
-UCLASS()
+UCLASS(Config = Engine)
 class ENGINE_API UDataLayerSubsystem : public UWorldSubsystem
 {
 	GENERATED_BODY()
@@ -82,6 +83,7 @@ public:
 	//~ End Blueprint callable functions
 
 #if WITH_EDITOR
+	bool ResolveIsLoadedInEditor(const TArray<FName>& InDataLayerInstanceNames) const;
 	bool CanResolveDataLayers() const;
 
 	bool RemoveDataLayer(const UDataLayerInstance* InDataLayer);
@@ -219,6 +221,12 @@ public:
 	//~ End Deprecated
 
 private:
+	UPROPERTY(Config)
+	TSoftClassPtr<UDataLayerLoadingPolicy> DataLayerLoadingPolicyClass;
+
+	UPROPERTY(Transient)
+	TObjectPtr<UDataLayerLoadingPolicy> DataLayerLoadingPolicy;
+
 #if WITH_EDITOR
 	void RegisterEarlyLoadedWorldDataLayers();
 	void OnActorDescContainerInitialized(UActorDescContainer* InActorDescContainer);
