Separating UpdateStreamingState from Camera updates in CSV and SCOPE CYCLE COUNTER

diff --git a/Engine/Source/Runtime/Engine/Private/LevelTick.cpp b/Engine/Source/Runtime/Engine/Private/LevelTick.cpp
--- a/Engine/Source/Runtime/Engine/Private/LevelTick.cpp
+++ b/Engine/Source/Runtime/Engine/Private/LevelTick.cpp
@@ -1607,7 +1607,7 @@ void UWorld::Tick( ELevelTick TickType, float DeltaSeconds )
 				}
 			}
 
-			// Update cameras and streaming volumes
+			// Update cameras
 			{
 				SCOPE_CYCLE_COUNTER(STAT_UpdateCameraTime);
 				CSV_SCOPED_TIMING_STAT_EXCLUSIVE(Camera);
@@ -1626,8 +1626,11 @@ void UWorld::Tick( ELevelTick TickType, float DeltaSeconds )
 						}
 					}
 				}
+			}
 
-				if( !bIsPaused && IsGameWorld())
+			// Update streaming volumes
+			{
+				if (!bIsPaused && IsGameWorld())
 				{
 					// Update world's required streaming levels
 					InternalUpdateStreamingState();
diff --git a/Engine/Source/Runtime/Engine/Private/Streaming/TextureStreamingHelpers.h b/Engine/Source/Runtime/Engine/Private/Streaming/TextureStreamingHelpers.h
--- a/Engine/Source/Runtime/Engine/Private/Streaming/TextureStreamingHelpers.h
+++ b/Engine/Source/Runtime/Engine/Private/Streaming/TextureStreamingHelpers.h
@@ -28,6 +28,7 @@ DECLARE_CYCLE_STAT_EXTERN(TEXT("RemoveFromWorld Time"),STAT_RemoveFromWorldTime,
 DECLARE_CYCLE_STAT_EXTERN(TEXT("UpdateLevelStreaming Time"),STAT_UpdateLevelStreamingTime,STATGROUP_StreamingDetails, );
 DECLARE_CYCLE_STAT_EXTERN(TEXT("Manage LevelsToConsider"), STAT_ManageLevelsToConsider, STATGROUP_StreamingDetails, );
 DECLARE_CYCLE_STAT_EXTERN(TEXT("Volume Streaming Tick"),STAT_VolumeStreamingTickTime,STATGROUP_StreamingDetails, );
+DECLARE_CYCLE_STAT_EXTERN(TEXT("UpdateStreamingState Time"), STAT_UpdateStreamingState, STATGROUP_StreamingDetails, );
 DECLARE_DWORD_COUNTER_STAT_EXTERN(TEXT("Streaming Volumes"),STAT_VolumeStreamingChecks,STATGROUP_StreamingDetails, );
 
 
diff --git a/Engine/Source/Runtime/Engine/Private/World.cpp b/Engine/Source/Runtime/Engine/Private/World.cpp
--- a/Engine/Source/Runtime/Engine/Private/World.cpp
+++ b/Engine/Source/Runtime/Engine/Private/World.cpp
@@ -2706,6 +2706,7 @@ DEFINE_STAT(STAT_AddToWorldTime);
 DEFINE_STAT(STAT_RemoveFromWorldTime);
 DEFINE_STAT(STAT_UpdateLevelStreamingTime);
 DEFINE_STAT(STAT_ManageLevelsToConsider);
+DEFINE_STAT(STAT_UpdateStreamingState);
 
 /**
  * Static helper function for Add/RemoveToWorld to determine whether we've already spent all the allotted time.
@@ -3857,6 +3858,9 @@ void UWorld::BlockTillLevelStreamingCompleted()
 
 void UWorld::InternalUpdateStreamingState()
 {
+	CSV_SCOPED_TIMING_STAT_EXCLUSIVE(UpdateStreamingState);
+	SCOPE_CYCLE_COUNTER(STAT_UpdateStreamingState);
+
 	// Update streaming levels state using streaming volumes.
 	// Issues level streaming load/unload requests based on local players being inside/outside level streaming volumes.
 	ProcessLevelStreamingVolumes();
