Fixed content bundle not loading on server when server streaming is disabled.

diff --git a/Engine/Source/Runtime/Engine/Private/WorldPartition/WorldPartitionRuntimeSpatialHash.cpp b/Engine/Source/Runtime/Engine/Private/WorldPartition/WorldPartitionRuntimeSpatialHash.cpp
--- a/Engine/Source/Runtime/Engine/Private/WorldPartition/WorldPartitionRuntimeSpatialHash.cpp
+++ b/Engine/Source/Runtime/Engine/Private/WorldPartition/WorldPartitionRuntimeSpatialHash.cpp
@@ -495,6 +495,29 @@ void FSpatialHashStreamingGrid::ForEachRuntimeCell(const FGridCellCoord& Coords,
 	});
 }
 
+void FSpatialHashStreamingGrid::ForEachRuntimeCell(TFunctionRef<bool(const UWorldPartitionRuntimeCell*)> Func) const
+{
+	auto ForEachGridLevel = [Func](const TArray<FSpatialHashStreamingGridLevel>& InGridLevels)
+	{
+		for (const FSpatialHashStreamingGridLevel& GridLevel : InGridLevels)
+		{
+			for (const FSpatialHashStreamingGridLayerCell& LayerCell : GridLevel.LayerCells)
+			{
+				for (const UWorldPartitionRuntimeSpatialHashCell* Cell : LayerCell.GridCells)
+				{
+					if (!Func(Cell))
+					{
+						return;
+					}
+				}
+			}
+		}
+	};
+
+	ForEachGridLevel(GridLevels);
+	ForEachGridLevel(InjectedGridLevels);
+}
+
 void FSpatialHashStreamingGrid::GetAlwaysLoadedCells(const UDataLayerSubsystem* DataLayerSubsystem, TSet<const UWorldPartitionRuntimeCell*>& OutActivateCells, TSet<const UWorldPartitionRuntimeCell*>& OutLoadCells) const
 {
 	if (GridLevels.Num() > 0)
@@ -1522,19 +1545,7 @@ void UWorldPartitionRuntimeSpatialHash::ForEachStreamingCells(TFunctionRef<bool(
 	{
 		if (IsCellRelevantFor(StreamingGrid.bClientOnlyVisible))
 		{
-			for (const FSpatialHashStreamingGridLevel& GridLevel : StreamingGrid.GridLevels)
-			{
-				for (const FSpatialHashStreamingGridLayerCell& LayerCell : GridLevel.LayerCells)
-				{
-					for (const UWorldPartitionRuntimeSpatialHashCell* Cell : LayerCell.GridCells)
-					{
-						if (!Func(Cell))
-						{
-							break;
-						}
-					}
-				}
-			}
+			StreamingGrid.ForEachRuntimeCell(Func);
 		}
 	});
 }
diff --git a/Engine/Source/Runtime/Engine/Public/WorldPartition/WorldPartitionRuntimeSpatialHash.h b/Engine/Source/Runtime/Engine/Public/WorldPartition/WorldPartitionRuntimeSpatialHash.h
--- a/Engine/Source/Runtime/Engine/Public/WorldPartition/WorldPartitionRuntimeSpatialHash.h
+++ b/Engine/Source/Runtime/Engine/Public/WorldPartition/WorldPartitionRuntimeSpatialHash.h
@@ -146,6 +146,7 @@ struct FSpatialHashStreamingGrid
 	void GetAlwaysLoadedCells(const UDataLayerSubsystem* DataLayerSubsystem, TSet<const UWorldPartitionRuntimeCell*>& OutActivateCells, TSet<const UWorldPartitionRuntimeCell*>& OutLoadCells) const;
 	void Draw2D(UCanvas* Canvas, UWorld* World, const TArray<FWorldPartitionStreamingSource>& Sources, const FBox& Region, const FBox2D& GridScreenBounds, TFunctionRef<FVector2D(const FVector2D&)> WorldToScreen) const;
 	void Draw3D(UWorld* World, const TArray<FWorldPartitionStreamingSource>& Sources, const FTransform& Transform) const;
+	void ForEachRuntimeCell(TFunctionRef<bool(const UWorldPartitionRuntimeCell*)> Func) const;
 
 private:
 	void ForEachRuntimeCell(const FGridCellCoord& Coords, TFunctionRef<void(const UWorldPartitionRuntimeCell*)> Func) const;
