Downgrade ensures on unexpected target states to logs.

diff --git a/Engine/Source/Runtime/Engine/Private/ProfilingDebugging/LevelStreamingProfilingSubsystem.cpp b/Engine/Source/Runtime/Engine/Private/ProfilingDebugging/LevelStreamingProfilingSubsystem.cpp
--- a/Engine/Source/Runtime/Engine/Private/ProfilingDebugging/LevelStreamingProfilingSubsystem.cpp
+++ b/Engine/Source/Runtime/Engine/Private/ProfilingDebugging/LevelStreamingProfilingSubsystem.cpp
@@ -251,7 +251,7 @@ void ULevelStreamingProfilingSubsystem::OnLevelStreamingTargetStateChanged(UWorl
 		case ELevelStreamingState::Unloaded:
 			return OnLevelQueuedForLoading(World, StreamingLevel);
 		default:
-			ensure(false);
+			UE_LOG(LogLevelStreamingProfiling, Warning, TEXT("Unexpected desired state LoadedNotVisible in current state %s"), ::EnumToString(CurrentState));
 			return;
 		}
 	case ELevelStreamingTargetState::LoadedVisible:
@@ -262,7 +262,7 @@ void ULevelStreamingProfilingSubsystem::OnLevelStreamingTargetStateChanged(UWorl
 		case ELevelStreamingState::LoadedNotVisible:
 			return OnLevelQueuedForAddToWorld(World, StreamingLevel, LevelIfLoaded);
 		default:
-			ensure(false);
+			UE_LOG(LogLevelStreamingProfiling, Warning, TEXT("Unexpected desired state LoadedVisible in current state %s"), ::EnumToString(CurrentState));
 			return;
 		}
 	case ELevelStreamingTargetState::Unloaded:
