Texture.cpp : add SourceFormat asset registry tag

diff --git a/Engine/Source/Runtime/Engine/Private/Texture.cpp b/Engine/Source/Runtime/Engine/Private/Texture.cpp
--- a/Engine/Source/Runtime/Engine/Private/Texture.cpp
+++ b/Engine/Source/Runtime/Engine/Private/Texture.cpp
@@ -1177,6 +1177,10 @@ void UTexture::GetAssetRegistryTags(TArray<FAssetRegistryTag>& OutTags) const
 
 	OutTags.Add(FAssetRegistryTag("SourceCompression", Source.GetSourceCompressionAsString(), FAssetRegistryTag::TT_Alphabetical));
 
+	OutTags.Add(FAssetRegistryTag("SourceFormat", 
+		StaticEnum<ETextureSourceFormat>()->GetDisplayNameTextByValue(Source.GetFormat()).ToString(),
+		FAssetRegistryTag::TT_Alphabetical));
+	
 	Super::GetAssetRegistryTags(OutTags);
 }
 #endif
