World Partition: implement GetActors on runtime streaming cells.

#rnx

diff --git a/Engine/Source/Runtime/Engine/Private/WorldPartition/WorldPartitionRuntimeLevelStreamingCell.cpp b/Engine/Source/Runtime/Engine/Private/WorldPartition/WorldPartitionRuntimeLevelStreamingCell.cpp
--- a/Engine/Source/Runtime/Engine/Private/WorldPartition/WorldPartitionRuntimeLevelStreamingCell.cpp
+++ b/Engine/Source/Runtime/Engine/Private/WorldPartition/WorldPartitionRuntimeLevelStreamingCell.cpp
@@ -5,6 +5,7 @@
 #include "WorldPartition/WorldPartition.h"
 #include "Engine/Level.h"
 #include "Engine/World.h"
+#include "Misc/Paths.h"
 
 #include "WorldPartition/HLOD/HLODSubsystem.h"
 #include "WorldPartition/HLOD/HLODActor.h"
@@ -59,6 +60,22 @@ bool UWorldPartitionRuntimeLevelStreamingCell::HasActors() const
 #endif
 }
 
+TArray<FName> UWorldPartitionRuntimeLevelStreamingCell::GetActors() const
+{
+	TArray<FName> Actors;
+
+#if WITH_EDITOR
+	Actors.Reserve(Packages.Num());
+
+	for (const FWorldPartitionRuntimeCellObjectMapping& Package : Packages)
+	{
+		Actors.Add(*FPaths::GetExtension(Package.Path.ToString()));
+	}
+#endif
+
+	return Actors;
+}
+
 void UWorldPartitionRuntimeLevelStreamingCell::CreateAndSetLevelStreaming(const FString& InPackageName)
 {
 	LevelStreaming = CreateLevelStreaming(InPackageName);
diff --git a/Engine/Source/Runtime/Engine/Public/WorldPartition/WorldPartitionRuntimeCell.h b/Engine/Source/Runtime/Engine/Public/WorldPartition/WorldPartitionRuntimeCell.h
--- a/Engine/Source/Runtime/Engine/Public/WorldPartition/WorldPartitionRuntimeCell.h
+++ b/Engine/Source/Runtime/Engine/Public/WorldPartition/WorldPartitionRuntimeCell.h
@@ -174,6 +174,7 @@ class ENGINE_API UWorldPartitionRuntimeCell : public UObject, public IWorldParti
 	virtual FName GetGridName() const { return DebugInfo.GridName; }
 	bool GetClientOnlyVisible() const { return bClientOnlyVisible; }
 	virtual FGuid const& GetContentBundleID() const { return ContentBundleID; }
+	virtual TArray<FName> GetActors() const PURE_VIRTUAL(UWorldPartitionRuntimeCell::GetActors, return TArray<FName>(););
 
 	IWorldPartitionRuntimeCellOwner* GetCellOwner() const { return CastChecked<IWorldPartitionRuntimeCellOwner>(GetOuter()); }
 
diff --git a/Engine/Source/Runtime/Engine/Public/WorldPartition/WorldPartitionRuntimeLevelStreamingCell.h b/Engine/Source/Runtime/Engine/Public/WorldPartition/WorldPartitionRuntimeLevelStreamingCell.h
--- a/Engine/Source/Runtime/Engine/Public/WorldPartition/WorldPartitionRuntimeLevelStreamingCell.h
+++ b/Engine/Source/Runtime/Engine/Public/WorldPartition/WorldPartitionRuntimeLevelStreamingCell.h
@@ -33,6 +33,7 @@ class UWorldPartitionRuntimeLevelStreamingCell : public UWorldPartitionRuntimeSp
 	class UWorldPartitionLevelStreamingDynamic* GetLevelStreaming() const;
 
 	bool HasActors() const;
+	virtual TArray<FName> GetActors() const override;
 
 	void CreateAndSetLevelStreaming(const FString& InPackageName);
 	class UWorldPartitionLevelStreamingDynamic* CreateLevelStreaming(const FString& InPackageName = FString()) const;
