World Partition: engine side changes required for 22855434.

#rnx

diff --git a/Engine/Source/Runtime/Engine/Public/WorldPartition/WorldPartitionRuntimeSpatialHash.h b/Engine/Source/Runtime/Engine/Public/WorldPartition/WorldPartitionRuntimeSpatialHash.h
--- a/Engine/Source/Runtime/Engine/Public/WorldPartition/WorldPartitionRuntimeSpatialHash.h
+++ b/Engine/Source/Runtime/Engine/Public/WorldPartition/WorldPartitionRuntimeSpatialHash.h
@@ -294,12 +294,12 @@ public:
 
 	static FString GetCellCoordString(const FGridCellCoord& InCellGlobalCoord);
 
-protected:
+	void ForEachStreamingGrid(TFunctionRef<void(const FSpatialHashStreamingGrid&)> Func) const;
 
+protected:
 	const FSpatialHashStreamingGrid* GetStreamingGridByName(FName InGridName) const;
 	void ForEachStreamingGrid(TFunctionRef<void(FSpatialHashStreamingGrid&)> Func);
 	void ForEachStreamingGridBreakable(TFunctionRef<bool(const FSpatialHashStreamingGrid&)> Func) const;
-	void ForEachStreamingGrid(TFunctionRef<void(const FSpatialHashStreamingGrid&)> Func) const;
 
 	virtual EWorldPartitionStreamingPerformance GetStreamingPerformanceForCell(const UWorldPartitionRuntimeCell* Cell) const override;
 
