Fix Strata anisotropic material conversion

diff --git a/Engine/Source/Runtime/Engine/Private/Materials/Material.cpp b/Engine/Source/Runtime/Engine/Private/Materials/Material.cpp
--- a/Engine/Source/Runtime/Engine/Private/Materials/Material.cpp
+++ b/Engine/Source/Runtime/Engine/Private/Materials/Material.cpp
@@ -3144,12 +3144,59 @@ void UMaterial::ConvertMaterialToStrataMaterial()
 		UMaterialExpressionBreakMaterialAttributes* BreakMatAtt = NewObject<UMaterialExpressionBreakMaterialAttributes>(this);
 		MoveConnectionTo(EditorOnly->MaterialAttributes, BreakMatAtt, 0);
 
+		// Check if Anisotropy input is actually used/connected
+		// This is needed/important as if anisotropy is connected a converted material becomes 'complex' instead of 'single'. 
+		// Since this path uses material attributes, there is no direct way to evaluate if anisotropy is connected or not. To find this,
+		// we loop over all the material attribute make/set nodes to figure out if the anisotropy value is ever set.
+		bool bIsAnisotropyConnected = false;
+		{
+			// Gather 'make' attributes nodes
+			{
+				TArray<UMaterialExpressionMakeMaterialAttributes*> MakeAttributeExpressions;
+				GetAllExpressionsInMaterialAndFunctionsOfType<UMaterialExpressionMakeMaterialAttributes>(MakeAttributeExpressions);
+				for (UMaterialExpressionMakeMaterialAttributes* Expression : MakeAttributeExpressions)
+				{
+					bIsAnisotropyConnected = Expression->GetExpressionInput(MP_Anisotropy) != nullptr;
+					if (bIsAnisotropyConnected)
+					{
+						break;
+					}
+				}
+			}
+
+			// Gather 'set' attributes nodes
+			if (!bIsAnisotropyConnected)
+			{
+				const FGuid AnisotropyGuid = FMaterialAttributeDefinitionMap::GetID(MP_Anisotropy);
+				TArray<UMaterialExpressionSetMaterialAttributes*> SetAttributeExpressions;
+				GetAllExpressionsInMaterialAndFunctionsOfType<UMaterialExpressionSetMaterialAttributes>(SetAttributeExpressions);
+				for (UMaterialExpressionSetMaterialAttributes* Expression : SetAttributeExpressions)
+				{
+					for (const FGuid& AttributeGuid : Expression->AttributeSetTypes)
+					{
+						if (AttributeGuid == AnisotropyGuid)
+						{
+							bIsAnisotropyConnected = true;
+							break;
+						}
+					}
+					if (bIsAnisotropyConnected)
+					{
+						break;
+					}
+				}
+			}
+		}
+
 		ConvertNode = NewObject<UMaterialExpressionStrataLegacyConversion>(this);
 		ConvertNode->BaseColor.Connect(0, BreakMatAtt);
 		ConvertNode->Metallic.Connect(1, BreakMatAtt);
 		ConvertNode->Specular.Connect(2, BreakMatAtt);
 		ConvertNode->Roughness.Connect(3, BreakMatAtt);
-		ConvertNode->Anisotropy.Connect(4, BreakMatAtt);
+		if (bIsAnisotropyConnected)
+		{
+			ConvertNode->Anisotropy.Connect(4, BreakMatAtt);
+		}
 		ConvertNode->EmissiveColor.Connect(5, BreakMatAtt);
 		ConvertNode->Normal.Connect(8, BreakMatAtt);
 		ConvertNode->Tangent.Connect(9, BreakMatAtt);
diff --git a/Engine/Source/Runtime/Engine/Private/Materials/MaterialExpressions.cpp b/Engine/Source/Runtime/Engine/Private/Materials/MaterialExpressions.cpp
--- a/Engine/Source/Runtime/Engine/Private/Materials/MaterialExpressions.cpp
+++ b/Engine/Source/Runtime/Engine/Private/Materials/MaterialExpressions.cpp
@@ -22271,6 +22271,7 @@ FStrataOperator* UMaterialExpressionStrataLegacyConversion::StrataGenerateMateri
 		SlabOperator.bBSDFHasSSS = true;
 		SlabOperator.bBSDFHasMFPPluggedIn = true;
 		SlabOperator.bBSDFHasFuzz = true;
+		SlabOperator.bBSDFHasAnisotropy = Anisotropy.IsConnected();
 
 		return &SlabOperator;
 	};
