World Partition - HLOD: Replicate HLOD actors if any of their components needs to be replicated
* Also, HLOD actors don't need to be loaded on the server if they have nothing to replicate

diff --git a/Engine/Source/Runtime/Engine/Private/WorldPartition/HLOD/HLODActor.cpp b/Engine/Source/Runtime/Engine/Private/WorldPartition/HLOD/HLODActor.cpp
--- a/Engine/Source/Runtime/Engine/Private/WorldPartition/HLOD/HLODActor.cpp
+++ b/Engine/Source/Runtime/Engine/Private/WorldPartition/HLOD/HLODActor.cpp
@@ -38,6 +38,11 @@ AWorldPartitionHLOD::AWorldPartitionHLOD(const FObjectInitializer& ObjectInitial
 	SetCanBeDamaged(false);
 	SetActorEnableCollision(false);
 
+	// Set HLOD actors to replicate by default, since the CDO's GetIsReplicated() is used to tell if a class type might replicate or not.
+	// The real need for replication will be adjusted depending on the presence of owned components that
+	// needs to be replicated.
+	bReplicates = true;
+
 #if WITH_EDITORONLY_DATA
 	HLODHash = 0;
 	HLODBounds = FBox(EForceInit::ForceInit);
@@ -46,13 +51,13 @@ AWorldPartitionHLOD::AWorldPartitionHLOD(const FObjectInitializer& ObjectInitial
 
 void AWorldPartitionHLOD::SetVisibility(bool bInVisible)
 {
-	// When propagating visibility state to children, SetVisibility dirties all attached components.
-	// Because we know that the visibility flag of all components of an HLOD actor are always in sync, 
-	// we test on RootComponent to check if call is required. This way we avoid dirtying all primitive proxies render state.
-	if (RootComponent && (RootComponent->GetVisibleFlag() != bInVisible))
+	ForEachComponent<USceneComponent>(false, [bInVisible](USceneComponent* SceneComponent)
 	{
-		RootComponent->SetVisibility(bInVisible, /*bPropagateToChildren*/ true);
-	}
+		if (SceneComponent && (SceneComponent->GetVisibleFlag() != bInVisible))
+		{
+			SceneComponent->SetVisibility(bInVisible, false);
+		}
+	});
 }
 
 void AWorldPartitionHLOD::BeginPlay()
@@ -95,6 +100,12 @@ void AWorldPartitionHLOD::Serialize(FArchive& Ar)
 #endif
 }
 
+bool AWorldPartitionHLOD::NeedsLoadForServer() const
+{
+	// Only needed on server if this HLOD actor has anything to replicate to clients
+	return GetIsReplicated();
+}
+
 #if WITH_EDITOR
 
 void AWorldPartitionHLOD::PostLoad()
@@ -125,8 +136,8 @@ void AWorldPartitionHLOD::PostLoad()
 }
 
 void AWorldPartitionHLOD::RerunConstructionScripts()
-{}
-
+{
+}
 
 TUniquePtr<FWorldPartitionActorDesc> AWorldPartitionHLOD::CreateClassActorDesc() const
 {
@@ -139,32 +150,34 @@ void AWorldPartitionHLOD::SetHLODComponents(const TArray<UActorComponent*>& InHL
 
 	Modify();
 
-	TArray<UActorComponent*> ComponentsToRemove;
-	GetComponents(ComponentsToRemove);
+	TArray<UActorComponent*> ComponentsToRemove = GetInstanceComponents();
 	for (UActorComponent* ComponentToRemove : ComponentsToRemove)
 	{
 		ComponentToRemove->DestroyComponent();
 	}
 
-	for(UActorComponent* HLODComponent : InHLODComponents)
+	// We'll turn on replication for this actor only if it contains a replicated component
+	check(!IsActorInitialized());
+	bReplicates = false;
+
+	for(UActorComponent* Component : InHLODComponents)
 	{
-		HLODComponent->Rename(nullptr, this);
-		AddInstanceComponent(HLODComponent);
+		Component->Rename(nullptr, this);
+		AddInstanceComponent(Component);
 
-		if (USceneComponent* HLODSceneComponent = Cast<USceneComponent>(HLODComponent))
+		const bool ComponentReplicates = Component->GetIsReplicated();
+		bReplicates |= ComponentReplicates;
+
+		if (USceneComponent* SceneComponent = Cast<USceneComponent>(Component))
 		{
-			if (RootComponent == nullptr)
-			{
-				SetRootComponent(HLODSceneComponent);
-			}
-			else
+			// Prefer a replicated root component
+			if (!RootComponent || (!RootComponent->GetIsReplicated() && ComponentReplicates))
 			{
-				// Attach to root component, but don't mess world tranform
-				HLODSceneComponent->AttachToComponent(RootComponent, FAttachmentTransformRules::KeepWorldTransform);
+				RootComponent = SceneComponent;
 			}
 		}
 	
-		HLODComponent->RegisterComponent();
+		Component->RegisterComponent();
 	}
 }
 
diff --git a/Engine/Source/Runtime/Engine/Public/WorldPartition/HLOD/HLODActor.h b/Engine/Source/Runtime/Engine/Public/WorldPartition/HLOD/HLODActor.h
--- a/Engine/Source/Runtime/Engine/Public/WorldPartition/HLOD/HLODActor.h
+++ b/Engine/Source/Runtime/Engine/Public/WorldPartition/HLOD/HLODActor.h
@@ -62,11 +62,12 @@ public:
 protected:
 	//~ Begin UObject Interface.
 	virtual void Serialize(FArchive& Ar) override;
+	virtual bool NeedsLoadForServer() const override;
 #if WITH_EDITOR
 	virtual void PostLoad() override;
 	virtual void RerunConstructionScripts() override;
-	virtual bool CanEditChange(const FProperty* InProperty) const { return false; }
-	virtual bool CanEditChangeComponent(const UActorComponent* Component, const FProperty* InProperty) const { return false; }
+	virtual bool CanEditChange(const FProperty* InProperty) const override { return false; }
+	virtual bool CanEditChangeComponent(const UActorComponent* Component, const FProperty* InProperty) const override { return false; }
 #endif
 	//~ End UObject Interface.
 
@@ -81,7 +82,7 @@ protected:
 	virtual FBox GetStreamingBounds() const override;
 
 	virtual bool ShouldImport(FStringView ActorPropString, bool IsMovingLevel) override { return false; }
-	virtual bool IsLockLocation() const { return true; }
+	virtual bool IsLockLocation() const override { return true; }
 	virtual bool IsUserManaged() const override { return false; }
 #endif
 	//~ End AActor Interface.
