Downgrade a warning to a log.

#rnx

diff --git a/Engine/Source/Runtime/Engine/Private/Actor.cpp b/Engine/Source/Runtime/Engine/Private/Actor.cpp
--- a/Engine/Source/Runtime/Engine/Private/Actor.cpp
+++ b/Engine/Source/Runtime/Engine/Private/Actor.cpp
@@ -888,7 +888,7 @@ void AActor::Serialize(FArchive& Ar)
 			FGuid FixupContentBundleGuid = ContentBundlePaths::GetContentBundleGuidFromExternalActorPackagePath(GetPackage()->GetFName().ToString());
 			if (ContentBundleGuid != FixupContentBundleGuid)
 			{
-				UE_LOG(LogActor, Warning, TEXT("Actor ContentBundleGuid was fixed up: %s"), *GetName());
+				UE_LOG(LogActor, Log, TEXT("Actor ContentBundleGuid was fixed up: %s"), *GetName());
 				ContentBundleGuid = FixupContentBundleGuid;
 			}
 		}
diff --git a/Engine/Source/Runtime/Engine/Private/WorldPartition/WorldPartitionActorDesc.cpp b/Engine/Source/Runtime/Engine/Private/WorldPartition/WorldPartitionActorDesc.cpp
--- a/Engine/Source/Runtime/Engine/Private/WorldPartition/WorldPartitionActorDesc.cpp
+++ b/Engine/Source/Runtime/Engine/Private/WorldPartition/WorldPartitionActorDesc.cpp
@@ -489,7 +489,7 @@ PRAGMA_ENABLE_DEPRECATION_WARNINGS
 			FGuid FixupContentBundleGuid = ContentBundlePaths::GetContentBundleGuidFromExternalActorPackagePath(ActorPackage.ToString());
 			if (ContentBundleGuid != FixupContentBundleGuid)
 			{
-				UE_LOG(LogWorldPartition, Warning, TEXT("ActorDesc ContentBundleGuid was fixed up: %s"), *GetActorName().ToString());
+				UE_LOG(LogWorldPartition, Log, TEXT("ActorDesc ContentBundleGuid was fixed up: %s"), *GetActorName().ToString());
 				ContentBundleGuid = FixupContentBundleGuid;
 			}
 		}
