- Fix for seaprate translucency not being combined in TSR pass

#rnx

diff --git a/Engine/Source/Runtime/Engine/Public/ShowFlags.h b/Engine/Source/Runtime/Engine/Public/ShowFlags.h
--- a/Engine/Source/Runtime/Engine/Public/ShowFlags.h
+++ b/Engine/Source/Runtime/Engine/Public/ShowFlags.h
@@ -376,6 +376,7 @@ private:
 		SetVisualizeBuffer(false);
 		SetVisualizeNanite(false);
 		SetVisualizeLumen(false);
+		SetVisualizeStrata(false);
 		SetVisualizeVirtualShadowMap(false);
 		SetVectorFields(false);
 		SetGBufferHints(false);
