Add a ResetDeviceProperty function to allow you to clear any effects that you may have set on an input device once a force feedback effect is over. Pass in the platform user during evaluation in case users want to create stateful logic based on the player

#rnx

diff --git a/Engine/Source/Developer/AssetTools/Private/AssetTypeActions/AssetTypeActions_ForceFeedbackEffect.cpp b/Engine/Source/Developer/AssetTools/Private/AssetTypeActions/AssetTypeActions_ForceFeedbackEffect.cpp
--- a/Engine/Source/Developer/AssetTools/Private/AssetTypeActions/AssetTypeActions_ForceFeedbackEffect.cpp
+++ b/Engine/Source/Developer/AssetTools/Private/AssetTypeActions/AssetTypeActions_ForceFeedbackEffect.cpp
@@ -151,6 +151,7 @@ void FAssetTypeActions_ForceFeedbackEffect::PlayEffect(UForceFeedbackEffect* Eff
 
 void FAssetTypeActions_ForceFeedbackEffect::StopEffect() 
 {
+	PreviewForceFeedbackEffect.ResetDeviceProperties();
 	PreviewForceFeedbackEffect.ForceFeedbackEffect = nullptr;
 
 	IInputInterface* InputInterface = FSlateApplication::Get().GetInputInterface();
@@ -247,6 +248,7 @@ void FPreviewForceFeedbackEffect::Tick( float DeltaTime )
 
 	if (!Update(DeltaTime, ForceFeedbackValues))
 	{
+		ResetDeviceProperties();
 		ForceFeedbackEffect = nullptr;
 	}
 
diff --git a/Engine/Source/Runtime/Engine/Classes/GameFramework/ForceFeedbackEffect.h b/Engine/Source/Runtime/Engine/Classes/GameFramework/ForceFeedbackEffect.h
--- a/Engine/Source/Runtime/Engine/Classes/GameFramework/ForceFeedbackEffect.h
+++ b/Engine/Source/Runtime/Engine/Classes/GameFramework/ForceFeedbackEffect.h
@@ -97,6 +97,9 @@ struct ENGINE_API FActiveForceFeedbackEffect
 	// Updates the final force feedback values based on this effect.  Returns true if the effect should continue playing, false if it is finished.
 	bool Update(float DeltaTime, FForceFeedbackValues& Values);
 
+	/** Reset any device properties that may need to be after the duration of this effect has ended. */
+	void ResetDeviceProperties();
+
 	// Gets the current values at the stored play time
 	void GetValues(FForceFeedbackValues& Values) const;
 };
@@ -129,4 +132,7 @@ class UForceFeedbackEffect : public UObject
 	void GetValues(const float EvalTime, FForceFeedbackValues& Values, float ValueMultiplier = 1.f) const;
 
 	void SetDeviceProperties(const FPlatformUserId PlatformUser, const float DeltaTime, const float EvalTime);
+
+	/** Reset any device properties that may need to be after the duration of this effect has ended. */
+	void ResetDeviceProperties(const FPlatformUserId PlatformUser);
 };
diff --git a/Engine/Source/Runtime/Engine/Classes/GameFramework/InputDeviceProperties.h b/Engine/Source/Runtime/Engine/Classes/GameFramework/InputDeviceProperties.h
--- a/Engine/Source/Runtime/Engine/Classes/GameFramework/InputDeviceProperties.h
+++ b/Engine/Source/Runtime/Engine/Classes/GameFramework/InputDeviceProperties.h
@@ -29,13 +29,15 @@ public:
 
 	/**
 	* Evaluate this device property for a given duration. 
+	* If overriding in Blueprints, make sure to call the parent function!
 	* 
-	* @param DeltaTime		Delta time
-	* @param Duration		The number of seconds that this property has been active. Use this to get things like curve data over time.
-	* @return				A pointer to the evaluated input device property.
+ 	* @param PlatformUser		The platform user that should receive this device property change
+	* @param DeltaTime			Delta time
+	* @param Duration			The number of seconds that this property has been active. Use this to get things like curve data over time.
+	* @return					A pointer to the evaluated input device property.
 	*/
 	UFUNCTION(BlueprintCallable, BlueprintNativeEvent, Category = "InputDevice")
-	void EvaluateDeviceProperty(const float DeltaTime, const float Duration);
+	void EvaluateDeviceProperty(const FPlatformUserId PlatformUser, const float DeltaTime, const float Duration);
 
 	/** 
 	* Native C++ implementation of EvaluateDeviceProperty.
@@ -43,10 +45,28 @@ public:
 	* Override this to alter your device property in native code.
 	* @see UInputDeviceProperty::EvaluateDeviceProperty
 	*/
-	virtual void EvaluateDeviceProperty_Implementation(const float DeltaTime, const float Duration) { }
+	virtual void EvaluateDeviceProperty_Implementation(const FPlatformUserId PlatformUser, const float DeltaTime, const float Duration);
 
 	/**
-	* Apply the given device property to the Input Interface 
+	* Reset the current device property. Provides an opportunity to reset device state after evaluation is complete. 
+	* If overriding in Blueprints, make sure to call the parent function!
+	* 
+	* @param PlatformUser		The platform user that should receive this device property change
+	*/
+	UFUNCTION(BlueprintCallable, BlueprintNativeEvent, Category = "InputDevice")
+	void ResetDeviceProperty(const FPlatformUserId PlatformUser);
+
+	/**
+	* Native C++ implementation of ResetDeviceProperty
+	* Override this in C++ to alter the device property behavior in native code. 
+	* 
+	* @see ResetDeviceProperty
+	*/
+	virtual void ResetDeviceProperty_Implementation(const FPlatformUserId PlatformUser);
+
+	/**
+	* Apply the given device property to the Input Interface. 
+	* Note: To remove any applied affects of this device property, call ResetDeviceProperty.
 	* 
 	* @param UserId		The owning Platform User whose input device this property should be applied to.
 	*/
@@ -92,7 +112,8 @@ public:
 
 	UColorInputDeviceProperty();
 
-	virtual void EvaluateDeviceProperty_Implementation(const float DeltaTime, const float Duration) override;
+	virtual void EvaluateDeviceProperty_Implementation(const FPlatformUserId PlatformUser, const float DeltaTime, const float Duration) override;
+	virtual void ResetDeviceProperty_Implementation(const FPlatformUserId PlatformUser) override;
 	virtual FInputDeviceProperty* GetInternalDeviceProperty() override;
 
 protected:
diff --git a/Engine/Source/Runtime/Engine/Private/GameFramework/ForceFeedbackEffect.cpp b/Engine/Source/Runtime/Engine/Private/GameFramework/ForceFeedbackEffect.cpp
--- a/Engine/Source/Runtime/Engine/Private/GameFramework/ForceFeedbackEffect.cpp
+++ b/Engine/Source/Runtime/Engine/Private/GameFramework/ForceFeedbackEffect.cpp
@@ -91,12 +91,23 @@ void UForceFeedbackEffect::SetDeviceProperties(const FPlatformUserId PlatformUse
 	{
 		if (DeviceProp)
 		{
-			DeviceProp->EvaluateDeviceProperty(DeltaTime, EvalTime);
+			DeviceProp->EvaluateDeviceProperty(PlatformUser, DeltaTime, EvalTime);
 			DeviceProp->ApplyDeviceProperty(PlatformUser);
 		}
 	}
 }
 
+void UForceFeedbackEffect::ResetDeviceProperties(const FPlatformUserId PlatformUser)
+{
+	for (TObjectPtr<UInputDeviceProperty> DeviceProp : DeviceProperties)
+	{
+		if (DeviceProp)
+		{
+			DeviceProp->ResetDeviceProperty(PlatformUser);
+		}
+	}
+}
+
 void FActiveForceFeedbackEffect::GetValues(FForceFeedbackValues& Values) const
 {
 	if (ForceFeedbackEffect)
@@ -134,4 +145,12 @@ bool FActiveForceFeedbackEffect::Update(const float DeltaTime, FForceFeedbackVal
 	ForceFeedbackEffect->SetDeviceProperties(PlatformUser, DeltaTime, EvalTime);
 
 	return true;
+}
+
+void FActiveForceFeedbackEffect::ResetDeviceProperties()
+{
+	if (ForceFeedbackEffect)
+	{
+		ForceFeedbackEffect->ResetDeviceProperties(PlatformUser);
+	}
 }
\ No newline at end of file
diff --git a/Engine/Source/Runtime/Engine/Private/GameFramework/InputDeviceProperties.cpp b/Engine/Source/Runtime/Engine/Private/GameFramework/InputDeviceProperties.cpp
--- a/Engine/Source/Runtime/Engine/Private/GameFramework/InputDeviceProperties.cpp
+++ b/Engine/Source/Runtime/Engine/Private/GameFramework/InputDeviceProperties.cpp
@@ -24,6 +24,16 @@ void UInputDeviceProperty::ApplyDeviceProperty(const FPlatformUserId UserId)
 	}
 }
 
+void UInputDeviceProperty::EvaluateDeviceProperty_Implementation(const FPlatformUserId PlatformUser, const float DeltaTime, const float Duration)
+{
+
+}
+
+void UInputDeviceProperty::ResetDeviceProperty_Implementation(const FPlatformUserId PlatformUser)
+{
+
+}
+
 ///////////////////////////////////////////////////////////////////////
 // UColorInputDeviceProperty
 
@@ -42,7 +52,7 @@ UColorInputDeviceProperty::UColorInputDeviceProperty()
 	}
 }
 
-void UColorInputDeviceProperty::EvaluateDeviceProperty_Implementation(const float DeltaTime, const float Duration)
+void UColorInputDeviceProperty::EvaluateDeviceProperty_Implementation(const FPlatformUserId PlatformUser, const float DeltaTime, const float Duration)
 {
 	InternalProperty.bEnable = bEnable;
 	
@@ -50,6 +60,13 @@ void UColorInputDeviceProperty::EvaluateDeviceProperty_Implementation(const floa
 	InternalProperty.Color = CurveColor.ToFColorSRGB();
 }
 
+void UColorInputDeviceProperty::ResetDeviceProperty_Implementation(const FPlatformUserId PlatformUser)
+{
+	// Disabling the light will reset the color
+	InternalProperty.bEnable = false;
+	ApplyDeviceProperty(PlatformUser);
+}
+
 FInputDeviceProperty* UColorInputDeviceProperty::GetInternalDeviceProperty()
 {
 	return &InternalProperty;
diff --git a/Engine/Source/Runtime/Engine/Private/PlayerController.cpp b/Engine/Source/Runtime/Engine/Private/PlayerController.cpp
--- a/Engine/Source/Runtime/Engine/Private/PlayerController.cpp
+++ b/Engine/Source/Runtime/Engine/Private/PlayerController.cpp
@@ -4100,6 +4100,8 @@ void APlayerController::ClientPlayForceFeedback_Internal_Implementation( UForceF
 			{
 				if (ActiveForceFeedbackEffects[Index].Parameters.Tag == Params.Tag)
 				{
+					// Reset the device properties on an active effect before removal
+					ActiveForceFeedbackEffects[Index].ResetDeviceProperties();
 					ActiveForceFeedbackEffects.RemoveAtSwap(Index);
 				}
 			}
@@ -4127,6 +4129,11 @@ void APlayerController::ClientStopForceFeedback_Implementation( UForceFeedbackEf
 {
 	if (ForceFeedbackEffect == NULL && Tag == NAME_None)
 	{
+		// Reset all device properties
+		for (FActiveForceFeedbackEffect& Effect : ActiveForceFeedbackEffects)
+		{
+			Effect.ResetDeviceProperties();
+		}
 		ActiveForceFeedbackEffects.Empty();
 	}
 	else
@@ -4136,6 +4143,8 @@ void APlayerController::ClientStopForceFeedback_Implementation( UForceFeedbackEf
 			if (    (ForceFeedbackEffect == NULL || ActiveForceFeedbackEffects[Index].ForceFeedbackEffect == ForceFeedbackEffect)
 				 && (Tag == NAME_None || ActiveForceFeedbackEffects[Index].Parameters.Tag == Tag) )
 			{
+				// Reset the device properties on an active effect before removal
+				ActiveForceFeedbackEffects[Index].ResetDeviceProperties();
 				ActiveForceFeedbackEffects.RemoveAtSwap(Index);
 			}
 		}
@@ -4538,6 +4547,8 @@ void APlayerController::ProcessForceFeedbackAndHaptics(const float DeltaTime, co
 			{
 				if (!ActiveForceFeedbackEffects[Index].Update(DeltaTime, ForceFeedbackValues))
 				{
+					// Reset any device properties that may need it (i.e. trigger resistance) 
+					ActiveForceFeedbackEffects[Index].ResetDeviceProperties();
 					ActiveForceFeedbackEffects.RemoveAtSwap(Index);
 				}
 			}
