World Partition - HLOD: Destructible HLOD support using HLOD modifiers

diff --git a/Engine/Plugins/Editor/WorldPartitionHLODUtilities/Source/Private/WorldPartition/HLOD/Modifiers/HLODModifierMeshDestruction.cpp b/Engine/Plugins/Editor/WorldPartitionHLODUtilities/Source/Private/WorldPartition/HLOD/Modifiers/HLODModifierMeshDestruction.cpp
new file mode 100644
--- /dev/null
+++ b/Engine/Plugins/Editor/WorldPartitionHLODUtilities/Source/Private/WorldPartition/HLOD/Modifiers/HLODModifierMeshDestruction.cpp
@@ -0,0 +1,173 @@
+// Copyright Epic Games, Inc. All Rights Reserved.
+
+#include "WorldPartition/HLOD/Modifiers/HLODModifierMeshDestruction.h"
+
+#include "Algo/Transform.h"
+#include "IMeshMergeExtension.h"
+#include "Materials/MaterialInstanceConstant.h"
+#include "Materials/MaterialInterface.h"
+#include "MeshMergeDataTracker.h"
+#include "MeshMergeModule.h"
+#include "Modules/ModuleManager.h"
+#include "WorldPartition/HLOD/HLODBuilder.h"
+#include "WorldPartition/HLOD/HLODDestruction.h"
+#include "WorldPartition/HLOD/DestructibleHLODComponent.h"
+
+
+UWorldPartitionHLODModifierMeshDestruction::UWorldPartitionHLODModifierMeshDestruction(const FObjectInitializer& ObjectInitializer)
+	: Super(ObjectInitializer)
+{
+}
+
+
+
+bool UWorldPartitionHLODModifierMeshDestruction::CanModifyHLOD(TSubclassOf<UHLODBuilder> InHLODBuilderClass) const
+{
+	return true;
+}
+
+
+// Use vertex color attributes to store component indices
+// In the material, this allows us to mask destructed building parts
+class FDestructionMeshMergeExtension : public IMeshMergeExtension
+{
+public:
+	const FHLODBuildContext& HLODBuildContext;
+
+	TArray<AActor*> DestructibleActors;
+	UMaterialInterface* DestructibleMaterial;
+
+	FDestructionMeshMergeExtension(const FHLODBuildContext& InHLODBuildContext) : HLODBuildContext(InHLODBuildContext)
+	{
+		FModuleManager::LoadModuleChecked<IMeshMergeModule>("MeshMergeUtilities").GetUtilities().RegisterExtension(this);
+	}
+
+	virtual ~FDestructionMeshMergeExtension()
+	{
+		FModuleManager::LoadModuleChecked<IMeshMergeModule>("MeshMergeUtilities").GetUtilities().UnregisterExtension(this);
+	}
+
+	virtual void OnCreatedMergedRawMeshes(const TArray<UStaticMeshComponent*>& MergedComponents, const class FMeshMergeDataTracker& DataTracker, TArray<FMeshDescription>& MergedMeshLODs) override
+	{
+		check(DestructibleActors.IsEmpty());
+
+		// Create component to actor map & compute components to wedges offsets.
+		TMap<uint32, uint32> ComponentToActorIndex;
+		TArray<uint32> ComponentToWedgeOffsets;
+		ComponentToWedgeOffsets.SetNumZeroed(MergedComponents.Num());
+
+		for (int32 ComponentIndex = 0; ComponentIndex < MergedComponents.Num(); ++ComponentIndex)
+		{
+			ComponentToWedgeOffsets[ComponentIndex] = DataTracker.GetComponentToWedgeMappng(ComponentIndex, /*LODIndex=*/0);
+
+			UStaticMeshComponent* StaticMeshComponent = MergedComponents[ComponentIndex];
+			if (AActor* Actor = StaticMeshComponent->GetOwner())
+			{
+				if (Actor->Implements<UWorldPartitionDestructibleInHLODInterface>())
+				{
+					uint32 ActorIndex = DestructibleActors.AddUnique(Actor);
+					ComponentToActorIndex.Add(ComponentIndex, ActorIndex);
+				}
+			}
+		}
+
+		// Clear vertex colors
+		FMeshDescription& MergedMesh = MergedMeshLODs[0];
+		TVertexInstanceAttributesRef<FVector4f> VertexInstanceColors = MergedMesh.VertexInstanceAttributes().GetAttributesRef<FVector4f>(MeshAttribute::VertexInstance::Color);
+		for (FVertexInstanceID VertexInstanceID : MergedMesh.VertexInstances().GetElementIDs())
+		{
+			VertexInstanceColors[VertexInstanceID] = FVector4f(0.0f, 0.0f, 0.0f, 0.0f);
+		}
+
+		// Store actor index + 1 in the wedge colors for each part of the merged mesh
+		for (int32 ComponentIndex = 0; ComponentIndex < MergedComponents.Num(); ++ComponentIndex)
+		{
+			const uint32 ComponentStart = ComponentToWedgeOffsets[ComponentIndex];
+			const uint32 ComponentEnd = ComponentToWedgeOffsets.IsValidIndex(ComponentIndex + 1) ? ComponentToWedgeOffsets[ComponentIndex + 1] : MergedMesh.VertexInstances().Num();
+			FColor ActorColor = FColor::Black;
+
+			const uint32* ActorIndex = ComponentToActorIndex.Find(ComponentIndex);
+			if (ActorIndex)
+			{
+				ActorColor.DWColor() = *ActorIndex + 1;
+			}
+
+			for (uint32 Index = ComponentStart; Index < ComponentEnd; ++Index)
+			{
+				FVertexInstanceID VertexInstanceID(Index);
+				VertexInstanceColors[FVertexInstanceID(Index)] = FLinearColor(ActorColor);
+			}
+		}
+	}
+
+	virtual void OnCreatedProxyMaterial(const TArray<UStaticMeshComponent*>& MergedComponents, UMaterialInterface* ProxyMaterial) override
+	{
+		UMaterialInstanceConstant* Instance = Cast<UMaterialInstanceConstant>(ProxyMaterial);
+		if (!Instance)
+		{
+			return;
+		}
+
+		TSet<AActor*> DestructibleActorsSet;
+		for (UStaticMeshComponent* Component : MergedComponents)
+		{
+			if (Component->GetOwner()->Implements<UWorldPartitionDestructibleInHLODInterface>())
+			{
+				DestructibleActorsSet.Add(Component->GetOwner());
+			}
+		}
+
+		// Ensure a destructible material is used
+		const static FName EnableInstanceDestroyingParamName(TEXT("EnableInstanceDestroying"));
+		FGuid ParamGUID;
+		bool bEnableInstanceDestroying = false;
+		bool bFoundParam = Instance->GetStaticSwitchParameterValue(EnableInstanceDestroyingParamName, bEnableInstanceDestroying, ParamGUID);
+		if (bFoundParam && bEnableInstanceDestroying)
+		{
+			Instance->SetScalarParameterValueEditorOnly(FMaterialParameterInfo("NumInstances"), FMath::RoundUpToPowerOfTwo(DestructibleActorsSet.Num()));
+			Instance->BasePropertyOverrides.TwoSided = false;
+			Instance->BasePropertyOverrides.bOverride_TwoSided = true;
+
+			FStaticParameterSet CurrentParameters = Instance->GetStaticParameters();
+
+			Instance->UpdateStaticPermutation(CurrentParameters);
+			Instance->InitStaticPermutation();
+
+			DestructibleMaterial = Instance;
+		}
+	}
+
+	UWorldPartitionDestructibleHLODMeshComponent* GetDestructibleHLODComponent()
+	{
+		UWorldPartitionDestructibleHLODMeshComponent* Component = nullptr;
+
+		if (!DestructibleActors.IsEmpty())
+		{
+			Component = NewObject<UWorldPartitionDestructibleHLODMeshComponent>();
+			Component->SetIsReplicated(true);
+			Component->SetWorldLocation(HLODBuildContext.WorldPosition);
+			TArray<FName> DestructibleActorsNames;
+			Algo::Transform(DestructibleActors, DestructibleActorsNames, [](AActor* Actor) { return Actor->GetFName(); });
+			Component->SetDestructibleActors(DestructibleActorsNames);
+			Component->SetDestructibleHLODMaterial(DestructibleMaterial);
+		}
+
+		return Component;
+	}
+};
+
+void UWorldPartitionHLODModifierMeshDestruction::BeginHLODBuild(const FHLODBuildContext& InHLODBuildContext)
+{
+	DestructionMeshMergeExtension = new FDestructionMeshMergeExtension(InHLODBuildContext);
+}
+
+void UWorldPartitionHLODModifierMeshDestruction::EndHLODBuild(TArray<UActorComponent*>& InOutComponents)
+{
+	UWorldPartitionDestructibleHLODMeshComponent* DestructibleHLODComponent = DestructionMeshMergeExtension->GetDestructibleHLODComponent();
+	if (DestructibleHLODComponent)
+	{
+		InOutComponents.Add(DestructibleHLODComponent);
+	}
+
+	delete DestructionMeshMergeExtension;
+}
\ No newline at end of file
diff --git a/Engine/Plugins/Editor/WorldPartitionHLODUtilities/Source/Public/WorldPartition/HLOD/Modifiers/HLODModifierMeshDestruction.h b/Engine/Plugins/Editor/WorldPartitionHLODUtilities/Source/Public/WorldPartition/HLOD/Modifiers/HLODModifierMeshDestruction.h
new file mode 100644
--- /dev/null
+++ b/Engine/Plugins/Editor/WorldPartitionHLODUtilities/Source/Public/WorldPartition/HLOD/Modifiers/HLODModifierMeshDestruction.h
@@ -0,0 +1,23 @@
+// Copyright Epic Games, Inc. All Rights Reserved.
+
+#pragma once
+
+#include "WorldPartition/HLOD/HLODModifier.h"
+#include "HLODModifierMeshDestruction.generated.h"
+
+
+class FDestructionMeshMergeExtension;
+
+
+UCLASS()
+class WORLDPARTITIONHLODUTILITIES_API UWorldPartitionHLODModifierMeshDestruction : public UWorldPartitionHLODModifier
+{
+	GENERATED_UCLASS_BODY()
+
+	virtual bool CanModifyHLOD(TSubclassOf<UHLODBuilder> InHLODBuilderClass) const;
+	virtual void BeginHLODBuild(const FHLODBuildContext& InHLODBuildContext);
+	virtual void EndHLODBuild(TArray<UActorComponent*>&InOutComponents);
+
+private:
+	FDestructionMeshMergeExtension* DestructionMeshMergeExtension;
+};
diff --git a/Engine/Source/Runtime/Engine/Private/WorldPartition/HLOD/DestructibleHLODComponent.cpp b/Engine/Source/Runtime/Engine/Private/WorldPartition/HLOD/DestructibleHLODComponent.cpp
new file mode 100644
--- /dev/null
+++ b/Engine/Source/Runtime/Engine/Private/WorldPartition/HLOD/DestructibleHLODComponent.cpp
@@ -0,0 +1,260 @@
+// Copyright Epic Games, Inc. All Rights Reserved.
+
+#include "WorldPartition/HLOD/DestructibleHLODComponent.h"
+
+#include "Engine/StaticMesh.h"
+#include "Engine/Texture2DDynamic.h"
+#include "Materials/MaterialInstanceDynamic.h"
+#include "Materials/MaterialInterface.h"
+#include "Net/UnrealNetwork.h"
+
+#define LOCTEXT_NAMESPACE "DestructibleHLODComponent"
+
+DEFINE_LOG_CATEGORY_STATIC(LogHLODDestruction, Log, All);
+
+UWorldPartitionDestructibleHLODComponent::UWorldPartitionDestructibleHLODComponent(const FObjectInitializer& ObjectInitializer)
+	: Super(ObjectInitializer)
+{
+	SetIsReplicatedByDefault(true);
+}
+
+const TArray<FName>& UWorldPartitionDestructibleHLODComponent::GetDestructibleActors() const
+{
+	return DestructibleActors;
+}
+
+void UWorldPartitionDestructibleHLODComponent::SetDestructibleActors(const TArray<FName>& InDestructibleActors)
+{
+	DestructibleActors = InDestructibleActors;
+}
+
+void FWorldPartitionDestructibleHLODState::Initialize(UWorldPartitionDestructibleHLODComponent* InDestructibleHLODComponent)
+{
+	OwnerComponent = InDestructibleHLODComponent;
+	
+	NumDestructibleActors = InDestructibleHLODComponent->GetDestructibleActors().Num();
+
+	const ENetMode NetMode = OwnerComponent->GetNetMode();
+	bIsClient = NetMode != ENetMode::NM_DedicatedServer;
+	bIsServer = NetMode == ENetMode::NM_DedicatedServer || NetMode == ENetMode::NM_ListenServer;
+
+	if (bIsServer)
+	{
+		DamagedActors.Reserve(NumDestructibleActors);
+		ActorsToDamagedActorsMapping.SetNum(NumDestructibleActors);
+		for (auto& Entry : ActorsToDamagedActorsMapping)
+		{
+			Entry = INDEX_NONE;
+		}
+	}
+	
+	if (bIsClient)
+	{
+		VisibilityBuffer.SetNumUninitialized(FMath::RoundUpToPowerOfTwo(NumDestructibleActors));
+		FMemory::Memset(VisibilityBuffer.GetData(), FWorldPartitionDestructibleHLODDamagedActorState::MAX_HEALTH, VisibilityBuffer.Num());
+	}
+}
+
+void FWorldPartitionDestructibleHLODState::SetActorHealth(int32 InActorIndex, uint8 InActorHealth)
+{
+	if (InActorIndex < 0 || InActorIndex >= NumDestructibleActors)
+	{
+		UE_LOG(LogHLODDestruction, Error, TEXT("Invalid actor index provided to SetActorHealth() (%d, max = %d)"), InActorIndex, NumDestructibleActors);
+		return;
+	}
+
+	// If we are the server, replicate damage to clients
+	if (IsServer())
+	{
+		int32 DamagedActorIdx = ActorsToDamagedActorsMapping[InActorIndex];
+
+		// Add a new damaged actor entry
+		if (DamagedActorIdx == INDEX_NONE)
+		{
+			DamagedActorIdx = DamagedActors.Num();
+			ActorsToDamagedActorsMapping[InActorIndex] = DamagedActorIdx;
+			DamagedActors.Emplace(InActorIndex);
+		}
+
+		FWorldPartitionDestructibleHLODDamagedActorState& DamagedActor = DamagedActors[DamagedActorIdx];
+		check(DamagedActor.ActorIndex == InActorIndex);
+
+		DamagedActor.ActorHealth = InActorHealth;
+		MarkItemDirty(DamagedActor);
+	}
+
+	// If we are the client, directly update the visibility buffer
+	if (IsClient())
+	{
+		VisibilityBuffer[InActorIndex] = InActorHealth;
+	}
+}
+
+void FWorldPartitionDestructibleHLODState::PostReplicatedAdd(const TArrayView<int32>& AddedIndices, int32 FinalSize)
+{
+	PostReplicatedChange(AddedIndices, FinalSize);
+}
+
+void FWorldPartitionDestructibleHLODState::PostReplicatedChange(const TArrayView<int32>& ChangedIndices, int32 FinalSize)
+{
+	check(IsClient());
+
+	for (int32 ChangedIndex : ChangedIndices)
+	{
+		if (DamagedActors.IsValidIndex(ChangedIndex))
+		{
+			const FWorldPartitionDestructibleHLODDamagedActorState& DamagedActorState = DamagedActors[ChangedIndex];
+			VisibilityBuffer[DamagedActorState.ActorIndex] = DamagedActorState.ActorHealth;
+		}
+		else
+		{
+			UE_LOG(LogHLODDestruction, Error, TEXT("Invalid damaged actor index %d (num = %d)"), ChangedIndex, DamagedActors.Num());
+		}
+	}
+
+	OwnerComponent->OnDestructionStateUpdated();
+}
+
+UWorldPartitionDestructibleHLODMeshComponent::UWorldPartitionDestructibleHLODMeshComponent(const FObjectInitializer& ObjectInitializer)
+	: Super(ObjectInitializer)
+{
+}
+
+void UWorldPartitionDestructibleHLODMeshComponent::GetLifetimeReplicatedProps(TArray< class FLifetimeProperty >& OutLifetimeProps) const
+{
+	Super::GetLifetimeReplicatedProps(OutLifetimeProps);
+
+	// Turn on push based replication for low-frequency variables. This saves the server time in determining what to replicate.
+	FDoRepLifetimeParams Params;
+	Params.bIsPushBased = true;
+
+	DOREPLIFETIME_WITH_PARAMS_FAST(UWorldPartitionDestructibleHLODMeshComponent, DestructibleHLODState, Params);
+}
+
+void UWorldPartitionDestructibleHLODMeshComponent::BeginPlay()
+{
+	Super::BeginPlay();
+
+	DestructibleHLODState.Initialize(this);
+}
+
+void UWorldPartitionDestructibleHLODMeshComponent::OnDestructionStateUpdated()
+{
+	UpdateVisibilityTexture();
+}
+
+void UWorldPartitionDestructibleHLODMeshComponent::DestroyActor(int32 ActorIndex)
+{
+	DamageActor(ActorIndex, 0.0f);
+}
+
+void UWorldPartitionDestructibleHLODMeshComponent::DamageActor(int32 ActorIndex, float RemainingHealthPercent)
+{
+	// Set percentage health as uint8 in the visibility texture
+	const uint8 HealthInt = FMath::Clamp(RemainingHealthPercent, 0, 1) * FWorldPartitionDestructibleHLODDamagedActorState::MAX_HEALTH;
+
+	DestructibleHLODState.SetActorHealth(ActorIndex, HealthInt);
+
+	UpdateVisibilityTexture();
+}
+
+void UWorldPartitionDestructibleHLODMeshComponent::SetupVisibilityTexture()
+{
+	if (!VisibilityTexture && DestructibleHLODMaterial)
+	{
+		// Retrieve number of instance stored inside of this LOD actor
+		float NumberOfInstances = 0.0f;
+		if (DestructibleHLODMaterial->GetScalarParameterValue(FMaterialParameterInfo("NumInstances"), NumberOfInstances, true))
+		{
+			// Create dynamic texture size of (NumInstances, 1) if required
+			uint32 TextureSize = FMath::TruncToInt(NumberOfInstances);
+
+			FTexture2DDynamicCreateInfo CreateInfo;
+			CreateInfo.Format = PF_G8;
+			CreateInfo.Filter = TF_Nearest;
+			CreateInfo.SamplerAddressMode = AM_Clamp;
+			CreateInfo.bSRGB = false;
+
+			UTexture2DDynamic* DynamicInstanceTexture = UTexture2DDynamic::Create(TextureSize, 1, CreateInfo);
+			if (DynamicInstanceTexture)
+			{
+				VisibilityTexture = DynamicInstanceTexture;
+			}
+
+			if (VisibilityTexture)
+			{
+				// Create dynamic material instance if required and set it to us the dynamic texture
+				UMaterialInstanceDynamic* MaterialInstance = UMaterialInstanceDynamic::Create(DestructibleHLODMaterial, this);
+				MaterialInstance->SetTextureParameterValue("InstanceVisibilityTexture", VisibilityTexture);
+				VisibilityMaterial = MaterialInstance;
+			}
+
+			if (VisibilityMaterial)
+			{
+				// For components that are referencing the destructible material, override it with our MID
+				GetOwner()->ForEachComponent<UStaticMeshComponent>(false, [this](UStaticMeshComponent* SMComponent)
+				{
+					if (UMaterialInterface* Material = SMComponent->GetMaterial(0))
+					{
+						if (Material == DestructibleHLODMaterial)
+						{
+							SMComponent->SetMaterial(0, VisibilityMaterial);
+						}
+					}
+				});
+			}
+		}
+	}
+}
+
+void UWorldPartitionDestructibleHLODMeshComponent::UpdateVisibilityTexture()
+{
+	if (!DestructibleHLODState.IsClient() || !FApp::CanEverRender())
+	{
+		return;
+	}
+
+	SetupVisibilityTexture();
+
+	if (VisibilityTexture)
+	{
+		FTexture2DDynamicResource* TextureResource = static_cast<FTexture2DDynamicResource*>(VisibilityTexture->GetResource());
+
+		if (FApp::CanEverRender() && ensure(TextureResource))
+		{
+			ENQUEUE_RENDER_COMMAND(FUpdateHLODVisibilityTexture)(
+				[TextureResource, VisibilityBuffer = DestructibleHLODState.GetVisibilityBuffer()](FRHICommandListImmediate& RHICmdList)
+			{
+					UpdateVisibilityTexture_RenderThread(TextureResource, VisibilityBuffer);
+			});
+		}
+	}
+}
+
+void UWorldPartitionDestructibleHLODMeshComponent::SetDestructibleHLODMaterial(UMaterialInterface* InDestructibleMaterial)
+{
+	DestructibleHLODMaterial = InDestructibleMaterial;
+}
+
+void UWorldPartitionDestructibleHLODMeshComponent::UpdateVisibilityTexture_RenderThread(FTexture2DDynamicResource* TextureResource, const TArray<uint8>& VisibilityBuffer)
+{
+	check(IsInRenderingThread());
+
+	FRHITexture2D* TextureRHI = TextureResource->GetTexture2DRHI();
+
+	uint32 DestStride = 0;
+	uint8* DestData = reinterpret_cast<uint8*>(RHILockTexture2D(TextureRHI, 0, RLM_WriteOnly, DestStride, false, false));
+
+#if !(UE_BUILD_SHIPPING)
+	if ((uint32)VisibilityBuffer.Num() > DestStride)
+	{
+		UE_LOG(LogHLODDestruction, Fatal, TEXT("UpdateVisibilityTexture_RenderThread: copy dest (%d) is smaller than source (%d). Memory will be stomped."), DestStride, VisibilityBuffer.Num());
+	}
+#endif
+
+	FMemory::Memcpy(DestData, VisibilityBuffer.GetData(), FMath::Min((uint32)VisibilityBuffer.Num(), DestStride));
+
+	RHIUnlockTexture2D(TextureRHI, 0, false, false);
+}
+
+#undef LOCTEXT_NAMESPACE
diff --git a/Engine/Source/Runtime/Engine/Private/WorldPartition/HLOD/HLODDestruction.cpp b/Engine/Source/Runtime/Engine/Private/WorldPartition/HLOD/HLODDestruction.cpp
new file mode 100644
--- /dev/null
+++ b/Engine/Source/Runtime/Engine/Private/WorldPartition/HLOD/HLODDestruction.cpp
@@ -0,0 +1,99 @@
+// Copyright Epic Games, Inc. All Rights Reserved.
+
+#include "WorldPartition/HLOD/HLODDestruction.h"
+
+#include "GameFramework/Actor.h"
+#include "LevelUtils.h"
+#include "WorldPartition/HLOD/DestructibleHLODComponent.h"
+#include "WorldPartition/HLOD/HLODActor.h"
+#include "WorldPartition/HLOD/HLODSubsystem.h"
+#include "WorldPartition/WorldPartitionLevelStreamingDynamic.h"
+
+
+namespace
+{
+	static const UWorldPartitionRuntimeCell* GetActorRuntimeCell(AActor* InActor)
+	{
+		const ULevel* ActorLevel = InActor->GetLevel();
+		const ULevelStreaming* LevelStreaming = FLevelUtils::FindStreamingLevel(ActorLevel);
+		const UWorldPartitionLevelStreamingDynamic* LevelStreamingDynamic = Cast<UWorldPartitionLevelStreamingDynamic>(LevelStreaming);
+		return LevelStreamingDynamic ? LevelStreamingDynamic->GetWorldPartitionRuntimeCell() : nullptr;
+	}
+
+	static FWorldPartitionHLODDestructionTag ResolveHLODDestructionTagForActor(AActor* InActor)
+	{
+		TRACE_CPUPROFILER_EVENT_SCOPE(ResolveHLODDestructionTagForActor);
+
+		const UWorldPartitionRuntimeCell* RuntimeCell = GetActorRuntimeCell(InActor);
+
+		const FName ActorFName = InActor->GetFName();
+
+		FWorldPartitionHLODDestructionTag HLODDestructionTag;
+
+		if (RuntimeCell)
+		{
+			for (AWorldPartitionHLOD* HLODActor : InActor->GetWorld()->GetSubsystem<UHLODSubsystem>()->GetHLODActorsForCell(RuntimeCell))
+			{
+				for (UActorComponent* ActorComponent : HLODActor->GetComponents())
+				{
+					if (UWorldPartitionDestructibleHLODComponent* DestructibleHLODComponent = Cast<UWorldPartitionDestructibleHLODComponent>(ActorComponent))
+					{
+						int32 ActorIdx = DestructibleHLODComponent->GetDestructibleActors().IndexOfByKey(ActorFName);
+						if (ActorIdx != INDEX_NONE)
+						{
+							HLODDestructionTag.ActorIndex = ActorIdx;
+							HLODDestructionTag.HLODDestructionComponent = DestructibleHLODComponent;
+							break;
+						}
+					}
+				}
+
+				if (HLODDestructionTag.IsValid())
+				{
+					break;
+				}
+			}
+		}
+
+		return HLODDestructionTag;
+	}
+
+	FWorldPartitionHLODDestructionTag GetDestructionTagForActor(AActor* SourceActor)
+	{
+		FWorldPartitionHLODDestructionTag Tag = IWorldPartitionDestructibleInHLODInterface::Execute_GetHLODDestructionTag(SourceActor);
+		if (!Tag.IsValid())
+		{
+			Tag = ResolveHLODDestructionTagForActor(SourceActor);
+			IWorldPartitionDestructibleInHLODInterface::Execute_SetHLODDestructionTag(SourceActor, Tag);
+		}
+		return Tag;
+	}
+}
+
+void UWorldPartitionDestructibleInHLODSupportLibrary::DestroyInHLOD(const TScriptInterface<IWorldPartitionDestructibleInHLODInterface>& DestructibleInHLOD)
+{
+	AActor* SourceActor = CastChecked<AActor>(DestructibleInHLOD.GetObject());
+
+	if (SourceActor->HasAuthority())
+	{
+		FWorldPartitionHLODDestructionTag Tag = GetDestructionTagForActor(SourceActor);
+		if (Tag.IsValid())
+		{
+			Tag.HLODDestructionComponent->DestroyActor(Tag.ActorIndex);
+		}
+	}
+}
+
+void UWorldPartitionDestructibleInHLODSupportLibrary::DamageInHLOD(const TScriptInterface<IWorldPartitionDestructibleInHLODInterface>& DestructibleInHLOD, float DamagePercent)
+{
+	AActor* SourceActor = CastChecked<AActor>(DestructibleInHLOD.GetObject());
+
+	if (SourceActor->HasAuthority())
+	{
+		FWorldPartitionHLODDestructionTag Tag = GetDestructionTagForActor(SourceActor);
+		if (Tag.IsValid())
+		{
+			Tag.HLODDestructionComponent->DamageActor(Tag.ActorIndex, DamagePercent);
+		}
+	}
+}
diff --git a/Engine/Source/Runtime/Engine/Public/WorldPartition/HLOD/DestructibleHLODComponent.h b/Engine/Source/Runtime/Engine/Public/WorldPartition/HLOD/DestructibleHLODComponent.h
new file mode 100644
--- /dev/null
+++ b/Engine/Source/Runtime/Engine/Public/WorldPartition/HLOD/DestructibleHLODComponent.h
@@ -0,0 +1,153 @@
+// Copyright Epic Games, Inc. All Rights Reserved.
+
+#pragma once
+
+#include "CoreMinimal.h"
+#include "Containers/Map.h"
+#include "Containers/SparseArray.h"
+#include "Components/SceneComponent.h"
+#include "Net/Serialization/FastArraySerializer.h"
+#include "DestructibleHLODComponent.generated.h"
+
+
+class UMaterialInterface;
+class UMaterialInstanceDynamic;
+class UTexture2DDynamic;
+class FTexture2DDynamicResource;
+
+
+UCLASS(Abstract, HideCategories=(Tags, Sockets, ComponentTick, ComponentReplication, Activation, Cooking, Events, AssetUserData, Collision))
+class ENGINE_API UWorldPartitionDestructibleHLODComponent : public USceneComponent
+{
+	GENERATED_UCLASS_BODY()
+
+public:
+	virtual void DestroyActor(int32 ActorIndex) PURE_VIRTUAL(UWorldPartitionDestructibleHLODComponent::DestroyActor, );
+	virtual void DamageActor(int32 ActorIndex, float DamagePercent) PURE_VIRTUAL(UWorldPartitionDestructibleHLODComponent::DamageActor, );
+	virtual void OnDestructionStateUpdated() PURE_VIRTUAL(UWorldPartitionDestructibleHLODComponent::OnDestructionStateUpdated, );
+
+	const TArray<FName>& GetDestructibleActors() const;
+	void SetDestructibleActors(const TArray<FName>& InDestructibleActors);
+
+protected:
+	// Name of the destructible actors from the source cell.
+	UPROPERTY()
+	TArray<FName> DestructibleActors;
+};
+
+
+// Entry for a damaged actor
+USTRUCT()
+struct FWorldPartitionDestructibleHLODDamagedActorState : public FFastArraySerializerItem
+{
+	GENERATED_USTRUCT_BODY()
+
+public:
+	static const uint8 MAX_HEALTH = 0xFF;
+
+	FWorldPartitionDestructibleHLODDamagedActorState()
+		: ActorIndex(INDEX_NONE)
+		, ActorHealth(MAX_HEALTH)
+	{
+	}
+
+	FWorldPartitionDestructibleHLODDamagedActorState(int32 InActorIndex)
+		: ActorIndex(InActorIndex)
+		, ActorHealth(MAX_HEALTH)
+	{
+	}
+
+	bool operator == (const FWorldPartitionDestructibleHLODDamagedActorState& Other) const
+	{
+		return ActorIndex == Other.ActorIndex && ActorHealth == Other.ActorHealth;
+	}
+
+	UPROPERTY()
+	int32 ActorIndex;
+
+	UPROPERTY()
+	uint8 ActorHealth;
+};
+
+
+// Replicated state of the destructible HLOD
+USTRUCT()
+struct FWorldPartitionDestructibleHLODState : public FFastArraySerializer
+{
+	GENERATED_BODY()
+
+public:
+	bool NetDeltaSerialize(FNetDeltaSerializeInfo& DeltaParms)
+	{
+		return FFastArraySerializer::FastArrayDeltaSerialize<FWorldPartitionDestructibleHLODDamagedActorState, FWorldPartitionDestructibleHLODState>(DamagedActors, DeltaParms, *this);
+	}
+
+	void Initialize(UWorldPartitionDestructibleHLODComponent* InDestructibleHLODComponent);
+	void SetActorHealth(int32 InActorIndex, uint8 InActorHealth);
+
+	const bool IsClient() const { return bIsClient; }
+	const bool IsServer() const { return bIsServer; }
+	const TArray<uint8>& GetVisibilityBuffer() const { return VisibilityBuffer;	}
+
+	// ~ FFastArraySerializer Contract Begin
+	void PostReplicatedAdd(const TArrayView<int32>& AddedIndices, int32 FinalSize);
+	void PostReplicatedChange(const TArrayView<int32>& ChangedIndices, int32 FinalSize);
+	// ~ FFastArraySerializer Contract End
+
+private:
+	UPROPERTY()
+	TArray<FWorldPartitionDestructibleHLODDamagedActorState> DamagedActors;
+
+	UPROPERTY(NotReplicated)
+	TObjectPtr<UWorldPartitionDestructibleHLODComponent> OwnerComponent;
+
+	// Server only, map of actors indices to their damage info in the DamagedActors array
+	TArray<int32> ActorsToDamagedActorsMapping;	
+
+	// Client only, visibility buffer that is meant to be sent to the GPU
+	TArray<uint8> VisibilityBuffer;
+
+	bool bIsServer = false;
+	bool bIsClient = false;
+	int32 NumDestructibleActors = 0;
+};
+
+
+template<> struct TStructOpsTypeTraits<FWorldPartitionDestructibleHLODState> : public TStructOpsTypeTraitsBase2<FWorldPartitionDestructibleHLODState>
+{
+	enum { WithNetDeltaSerializer = true };
+};
+
+
+UCLASS()
+class ENGINE_API UWorldPartitionDestructibleHLODMeshComponent : public UWorldPartitionDestructibleHLODComponent
+{
+	GENERATED_UCLASS_BODY()
+
+public:
+	void SetDestructibleHLODMaterial(UMaterialInterface* InDestructibleMaterial);
+
+	virtual void DestroyActor(int32 ActorIndex) override;
+	virtual void DamageActor(int32 ActorIndex, float DamagePercent) override;
+	virtual void OnDestructionStateUpdated() override;
+
+private:
+	virtual void BeginPlay() override;
+
+	void SetupVisibilityTexture();
+	void UpdateVisibilityTexture();
+	static void UpdateVisibilityTexture_RenderThread(FTexture2DDynamicResource* TextureResource, const TArray<uint8>& VisibilityBuffer);
+
+private:
+	UPROPERTY()
+	TObjectPtr<UMaterialInterface> DestructibleHLODMaterial;
+
+	UPROPERTY(Transient, Replicated)
+	FWorldPartitionDestructibleHLODState DestructibleHLODState;
+
+	UPROPERTY(Transient)
+	TObjectPtr<UMaterialInstanceDynamic> VisibilityMaterial;
+
+	UPROPERTY(Transient)
+	TObjectPtr<UTexture2DDynamic> VisibilityTexture;
+};
diff --git a/Engine/Source/Runtime/Engine/Public/WorldPartition/HLOD/HLODDestruction.h b/Engine/Source/Runtime/Engine/Public/WorldPartition/HLOD/HLODDestruction.h
new file mode 100644
--- /dev/null
+++ b/Engine/Source/Runtime/Engine/Public/WorldPartition/HLOD/HLODDestruction.h
@@ -0,0 +1,69 @@
+// Copyright Epic Games, Inc. All Rights Reserved.
+
+#pragma once
+
+#include "CoreMinimal.h"
+#include "Kismet/BlueprintFunctionLibrary.h"
+#include "UObject/Interface.h"
+#include "HLODDestruction.generated.h"
+
+
+class UWorldPartitionDestructibleHLODComponent;
+
+
+USTRUCT(BlueprintType)
+struct ENGINE_API FWorldPartitionHLODDestructionTag
+{
+	GENERATED_BODY()
+
+	FWorldPartitionHLODDestructionTag()
+		: HLODDestructionComponent()
+		, ActorIndex(INDEX_NONE)
+	{
+	}
+
+	bool IsValid() const
+	{
+		return HLODDestructionComponent != nullptr && ActorIndex != INDEX_NONE;
+	}
+
+	UPROPERTY()
+	TObjectPtr<UWorldPartitionDestructibleHLODComponent> HLODDestructionComponent;
+
+	UPROPERTY()
+	int32 ActorIndex;
+};
+
+
+UINTERFACE(NotBlueprintable, meta=(DisplayName="Destructible in HLOD Interface"))
+class ENGINE_API UWorldPartitionDestructibleInHLODInterface : public UInterface
+{
+	GENERATED_BODY()
+};
+
+
+class ENGINE_API IWorldPartitionDestructibleInHLODInterface
+{
+	GENERATED_BODY()
+
+public:
+	UFUNCTION(BlueprintNativeEvent, meta=(BlueprintInternalUseOnly = "true"))
+	void SetHLODDestructionTag(const FWorldPartitionHLODDestructionTag& InDestructionTag);
+
+	UFUNCTION(BlueprintNativeEvent, meta=(BlueprintInternalUseOnly = "true"))
+	FWorldPartitionHLODDestructionTag GetHLODDestructionTag() const;
+};
+
+
+UCLASS()
+class ENGINE_API UWorldPartitionDestructibleInHLODSupportLibrary : public UBlueprintFunctionLibrary
+{
+	GENERATED_BODY()
+
+public:
+	UFUNCTION(BlueprintCallable, Category = "HLOD Destruction")
+	static void DestroyInHLOD(const TScriptInterface<IWorldPartitionDestructibleInHLODInterface>& DestructibleInHLOD);
+
+	UFUNCTION(BlueprintCallable, Category="HLOD Destruction")
+	static void DamageInHLOD(const TScriptInterface<IWorldPartitionDestructibleInHLODInterface>& DestructibleInHLOD, float DamagePercent);
+};
\ No newline at end of file
