World Partition: added Insights scopes to most parts of world partition/actor descriptor container/content bundle initialization steps.

#rnx

diff --git a/Engine/Source/Runtime/Engine/Private/WorldPartition/ActorDescContainer.cpp b/Engine/Source/Runtime/Engine/Private/WorldPartition/ActorDescContainer.cpp
--- a/Engine/Source/Runtime/Engine/Private/WorldPartition/ActorDescContainer.cpp
+++ b/Engine/Source/Runtime/Engine/Private/WorldPartition/ActorDescContainer.cpp
@@ -33,6 +33,8 @@ void UActorDescContainer::Initialize(UWorld* InWorld, FName InPackageName)
 
 void UActorDescContainer::Initialize(const FInitializeParams& InitParams)
 {
+	TRACE_CPUPROFILER_EVENT_SCOPE(UActorDescContainer::Initialize);
+
 	check(!World || World == InitParams.World);
 	World = InitParams.World;
 
@@ -47,31 +49,41 @@ void UActorDescContainer::Initialize(const FInitializeParams& InitParams)
 
 		// Do a synchronous scan of the level external actors path.			
 		IAssetRegistry& AssetRegistry = FModuleManager::LoadModuleChecked<FAssetRegistryModule>(TEXT("AssetRegistry")).Get();
-		AssetRegistry.ScanPathsSynchronous({ ContainerExternalActorsPath }, /*bForceRescan*/true, /*bIgnoreDenyListScanFilters*/false);
+		{
+			TRACE_CPUPROFILER_EVENT_SCOPE(ScanPathsSynchronous);
+			AssetRegistry.ScanPathsSynchronous({ ContainerExternalActorsPath }, /*bForceRescan*/true, /*bIgnoreDenyListScanFilters*/false);
+		}
 
 		FARFilter Filter;
 		Filter.bRecursivePaths = true;
 		Filter.bIncludeOnlyOnDiskAssets = true;
 		Filter.PackagePaths.Add(*ContainerExternalActorsPath);
 
+		TRACE_CPUPROFILER_EVENT_SCOPE(GetAssets);
 		AssetRegistry.GetAssets(Filter, Assets);
 	}
 
-	for (const FAssetData& Asset : Assets)
 	{
-		TUniquePtr<FWorldPartitionActorDesc> ActorDesc = FWorldPartitionActorDescUtils::GetActorDescriptorFromAssetData(Asset);
-
-		if (ActorDesc.IsValid() && (!InitParams.FilterActorDesc || InitParams.FilterActorDesc(ActorDesc.Get())))
+		TRACE_CPUPROFILER_EVENT_SCOPE(RegisterDescriptors);
+		for (const FAssetData& Asset : Assets)
 		{
-			AddActorDescriptor(ActorDesc.Release());
-		}
-		else
-		{
-			InvalidActors.Emplace(MoveTemp(ActorDesc));
+			TUniquePtr<FWorldPartitionActorDesc> ActorDesc = FWorldPartitionActorDescUtils::GetActorDescriptorFromAssetData(Asset);
+
+			if (ActorDesc.IsValid() && (!InitParams.FilterActorDesc || InitParams.FilterActorDesc(ActorDesc.Get())))
+			{
+				AddActorDescriptor(ActorDesc.Release());
+			}
+			else
+			{
+				InvalidActors.Emplace(MoveTemp(ActorDesc));
+			}
 		}
 	}
 
-	OnActorDescContainerInitialized.Broadcast(this);
+	{
+		TRACE_CPUPROFILER_EVENT_SCOPE(OnActorDescContainerInitialized)
+		OnActorDescContainerInitialized.Broadcast(this);
+	}
 
 	RegisterEditorDelegates();
 
diff --git a/Engine/Source/Runtime/Engine/Private/WorldPartition/ContentBundle/ContentBundleContainer.cpp b/Engine/Source/Runtime/Engine/Private/WorldPartition/ContentBundle/ContentBundleContainer.cpp
--- a/Engine/Source/Runtime/Engine/Private/WorldPartition/ContentBundle/ContentBundleContainer.cpp
+++ b/Engine/Source/Runtime/Engine/Private/WorldPartition/ContentBundle/ContentBundleContainer.cpp
@@ -36,6 +36,7 @@ UWorld* FContentBundleContainer::GetInjectedWorld() const
 
 void FContentBundleContainer::Initialize()
 {
+	TRACE_CPUPROFILER_EVENT_SCOPE(FContentBundleContainer::Initialize);
 	UE_LOG(LogContentBundle, Log, TEXT("[Container: %s] Creating new contrainer."), *GetInjectedWorld()->GetName());
 
 #if WITH_EDITOR
@@ -251,16 +252,19 @@ void FContentBundleContainer::DeinitializeContentBundle(FContentBundleBase& Cont
 
 void FContentBundleContainer::InjectContentBundle(FContentBundleBase& ContentBundle)
 {
+	TRACE_CPUPROFILER_EVENT_SCOPE(InjectContentBundle);
 	ContentBundle.InjectContent();
 }
 
 void FContentBundleContainer::RemoveContentBundle(FContentBundleBase& ContentBundle)
 {
+	TRACE_CPUPROFILER_EVENT_SCOPE(RemoveContentBundle);
 	ContentBundle.RemoveContent();
 }
 
 void FContentBundleContainer::OnContentBundleClientRegistered(TSharedPtr<FContentBundleClient>& ContentBundleClient)
 {
+	TRACE_CPUPROFILER_EVENT_SCOPE(InitializeContentBundle);
 	InitializeContentBundle(ContentBundleClient);
 }
 
diff --git a/Engine/Source/Runtime/Engine/Private/WorldPartition/WorldPartition.cpp b/Engine/Source/Runtime/Engine/Private/WorldPartition/WorldPartition.cpp
--- a/Engine/Source/Runtime/Engine/Private/WorldPartition/WorldPartition.cpp
+++ b/Engine/Source/Runtime/Engine/Private/WorldPartition/WorldPartition.cpp
@@ -521,10 +521,7 @@ void UWorldPartition::Initialize(UWorld* InWorld, const FTransform& InTransform)
 				);
 		}
 
-		{
-			TRACE_CPUPROFILER_EVENT_SCOPE(UActorDescContainer::Initialize);
-			ActorDescContainer = RegisterActorDescContainer(PackageName);
-		}
+		ActorDescContainer = RegisterActorDescContainer(PackageName);
 
 		{
 			TRACE_CPUPROFILER_EVENT_SCOPE(UActorDescContainer::Hash);
@@ -554,7 +551,10 @@ void UWorldPartition::Initialize(UWorld* InWorld, const FTransform& InTransform)
 	{
 		// Here we need to flush any async loading before starting any synchronous load (mixing synchronous and asynchronous load on the same package is not handled properly).
 		// Also, this will ensure that FindObject will work on an external actor that was loading asynchronously.
-		FlushAsyncLoading();
+		{
+			TRACE_CPUPROFILER_EVENT_SCOPE(FlushAsyncLoading);
+			FlushAsyncLoading();
+		}
 
 		// Make sure to preload only AWorldDataLayers actor first
 		for (FActorDescList::TIterator<> ActorDescIterator(ActorDescContainer); ActorDescIterator; ++ActorDescIterator)
@@ -582,6 +582,8 @@ void UWorldPartition::Initialize(UWorld* InWorld, const FTransform& InTransform)
 		// Apply level transform on actors already part of the level
 		if (!GetInstanceTransform().Equals(FTransform::Identity))
 		{
+			TRACE_CPUPROFILER_EVENT_SCOPE(ApplyLevelTransform);
+
 			check(!OuterWorld->PersistentLevel->bAlreadyMovedActors);
 			for (AActor* Actor : OuterWorld->PersistentLevel->Actors)
 			{
@@ -601,7 +603,10 @@ void UWorldPartition::Initialize(UWorld* InWorld, const FTransform& InTransform)
 	if (bIsEditor && !bIsCooking)
 	{
 		// Load the always loaded cell, don't call LoadCells to avoid creating a transaction
-		AlwaysLoadedActors->Load();
+		{
+			TRACE_CPUPROFILER_EVENT_SCOPE(LoadAlwaysLoaded);
+			AlwaysLoadedActors->Load();
+		}
 
 		// Load more cells depending on the user's settings
 		// Skipped when running from a commandlet and for subpartitions
@@ -610,6 +615,7 @@ void UWorldPartition::Initialize(UWorld* InWorld, const FTransform& InTransform)
 			// Load last loaded regions
 			if (GetMutableDefault<UWorldPartitionEditorPerProjectUserSettings>()->GetEnableLoadingOfLastLoadedRegions())
 			{
+				TRACE_CPUPROFILER_EVENT_SCOPE(LoadLastLoadedRegions);
 				LoadLastLoadedRegions();
 			}
 		}
