SceneManagement - allow actor position to be explicitly specified in  FDynamicPrimitiveUniformBuffer.   The pre-existing Set() methods use the center of the WS bounds, but this isn't what you want in all cases.

diff --git a/Engine/Source/Runtime/Engine/Private/SceneManagement.cpp b/Engine/Source/Runtime/Engine/Private/SceneManagement.cpp
--- a/Engine/Source/Runtime/Engine/Private/SceneManagement.cpp
+++ b/Engine/Source/Runtime/Engine/Private/SceneManagement.cpp
@@ -350,6 +350,7 @@ void FMeshElementCollector::AddMesh(int32 ViewIndex, FMeshBatch& MeshBatch)
 void FDynamicPrimitiveUniformBuffer::Set(
 	const FMatrix& LocalToWorld,
 	const FMatrix& PreviousLocalToWorld,
+	const FVector& ActorPositionWS,
 	const FBoxSphereBounds& WorldBounds,
 	const FBoxSphereBounds& LocalBounds,
 	const FBoxSphereBounds& PreSkinnedLocalBounds,
@@ -364,7 +365,7 @@ void FDynamicPrimitiveUniformBuffer::Set(
 		.Defaults()
 			.LocalToWorld(LocalToWorld)
 			.PreviousLocalToWorld(PreviousLocalToWorld)
-			.ActorWorldPosition(WorldBounds.Origin)
+			.ActorWorldPosition(ActorPositionWS)
 			.WorldBounds(WorldBounds)
 			.LocalBounds(LocalBounds)
 			.PreSkinnedLocalBounds(PreSkinnedLocalBounds)
@@ -377,6 +378,20 @@ void FDynamicPrimitiveUniformBuffer::Set(
 	UniformBuffer.InitResource();
 }
 
+void FDynamicPrimitiveUniformBuffer::Set(
+	const FMatrix& LocalToWorld,
+	const FMatrix& PreviousLocalToWorld,
+	const FBoxSphereBounds& WorldBounds,
+	const FBoxSphereBounds& LocalBounds,
+	const FBoxSphereBounds& PreSkinnedLocalBounds,
+	bool bReceivesDecals,
+	bool bHasPrecomputedVolumetricLightmap,
+	bool bOutputVelocity,
+	const FCustomPrimitiveData* CustomPrimitiveData)
+{
+	Set(LocalToWorld, PreviousLocalToWorld, WorldBounds.Origin, WorldBounds, LocalBounds, PreSkinnedLocalBounds, bReceivesDecals, bHasPrecomputedVolumetricLightmap, bOutputVelocity, CustomPrimitiveData);
+}
+
 void FDynamicPrimitiveUniformBuffer::Set(
 	const FMatrix& LocalToWorld,
 	const FMatrix& PreviousLocalToWorld,
diff --git a/Engine/Source/Runtime/Engine/Public/SceneManagement.h b/Engine/Source/Runtime/Engine/Public/SceneManagement.h
--- a/Engine/Source/Runtime/Engine/Public/SceneManagement.h
+++ b/Engine/Source/Runtime/Engine/Public/SceneManagement.h
@@ -2631,6 +2631,18 @@ public:
 
 	TUniformBuffer<FPrimitiveUniformShaderParameters> UniformBuffer;
 
+	ENGINE_API void Set(
+		const FMatrix& LocalToWorld,
+		const FMatrix& PreviousLocalToWorld,
+		const FVector& ActorPositionWS, 
+		const FBoxSphereBounds& WorldBounds,
+		const FBoxSphereBounds& LocalBounds,
+		const FBoxSphereBounds& PreSkinnedLocalBounds,
+		bool bReceivesDecals,
+		bool bHasPrecomputedVolumetricLightmap,
+		bool bOutputVelocity,
+		const FCustomPrimitiveData* CustomPrimitiveData);
+
 	ENGINE_API void Set(
 		const FMatrix& LocalToWorld,
 		const FMatrix& PreviousLocalToWorld,
