World Partition: don't force a full rescan when parsing the asset registry to initialize actor descriptor containers.

#rnx

diff --git a/Engine/Source/Runtime/Engine/Private/WorldPartition/ActorDescContainer.cpp b/Engine/Source/Runtime/Engine/Private/WorldPartition/ActorDescContainer.cpp
--- a/Engine/Source/Runtime/Engine/Private/WorldPartition/ActorDescContainer.cpp
+++ b/Engine/Source/Runtime/Engine/Private/WorldPartition/ActorDescContainer.cpp
@@ -51,7 +51,7 @@ void UActorDescContainer::Initialize(const FInitializeParams& InitParams)
 		IAssetRegistry& AssetRegistry = FModuleManager::LoadModuleChecked<FAssetRegistryModule>(TEXT("AssetRegistry")).Get();
 		{
 			TRACE_CPUPROFILER_EVENT_SCOPE(ScanPathsSynchronous);
-			AssetRegistry.ScanPathsSynchronous({ ContainerExternalActorsPath }, /*bForceRescan*/true, /*bIgnoreDenyListScanFilters*/false);
+			AssetRegistry.ScanPathsSynchronous({ ContainerExternalActorsPath }, /*bForceRescan*/false, /*bIgnoreDenyListScanFilters*/false);
 		}
 
 		FARFilter Filter;
