Fixed non-unity build
#rnx

diff --git a/Engine/Source/Runtime/Engine/Private/WorldPartition/HLOD/DestructibleHLODComponent.cpp b/Engine/Source/Runtime/Engine/Private/WorldPartition/HLOD/DestructibleHLODComponent.cpp
--- a/Engine/Source/Runtime/Engine/Private/WorldPartition/HLOD/DestructibleHLODComponent.cpp
+++ b/Engine/Source/Runtime/Engine/Private/WorldPartition/HLOD/DestructibleHLODComponent.cpp
@@ -2,11 +2,15 @@
 
 #include "WorldPartition/HLOD/DestructibleHLODComponent.h"
 
+#include "Components/StaticMeshComponent.h"
 #include "Engine/StaticMesh.h"
 #include "Engine/Texture2DDynamic.h"
+#include "GameFramework/Actor.h"
 #include "Materials/MaterialInstanceDynamic.h"
 #include "Materials/MaterialInterface.h"
 #include "Net/UnrealNetwork.h"
+#include "RenderingThread.h"
+#include "TextureResource.h"
 
 #define LOCTEXT_NAMESPACE "DestructibleHLODComponent"
 
