Fixed non-unity build
#rnx

diff --git a/Engine/Source/Runtime/Engine/Public/WorldPartition/HLOD/HLODLayer.h b/Engine/Source/Runtime/Engine/Public/WorldPartition/HLOD/HLODLayer.h
--- a/Engine/Source/Runtime/Engine/Public/WorldPartition/HLOD/HLODLayer.h
+++ b/Engine/Source/Runtime/Engine/Public/WorldPartition/HLOD/HLODLayer.h
@@ -11,6 +11,8 @@
 #include "EngineDefines.h"
 #include "Engine/MeshMerging.h"
 
+#include "WorldPartition/HLOD/HLODBuilder.h"
+
 #include "HLODLayer.generated.h"
 
 class AActor;
