Fix anim timeline trimming tools creating invalid key data

Prevent trimming from reducing key count below 2

diff --git a/Engine/Source/Runtime/Engine/Private/Animation/AnimSequenceHelpers.cpp b/Engine/Source/Runtime/Engine/Private/Animation/AnimSequenceHelpers.cpp
--- a/Engine/Source/Runtime/Engine/Private/Animation/AnimSequenceHelpers.cpp
+++ b/Engine/Source/Runtime/Engine/Private/Animation/AnimSequenceHelpers.cpp
@@ -677,6 +677,7 @@ bool AnimationData::Trim(UAnimSequence* InSequence, float TrimStart, float TrimE
 	{
 		const FFrameRate& FrameRate = DataModel->GetFrameRate();
 
+
 		// if there is only one key, there is nothing to trim away
 		if (NumKeys <= 1)
 		{
@@ -782,9 +783,16 @@ void AnimationData::RemoveKeys(UAnimSequence* InSequence, int32 StartKeyIndex, i
 
 		auto ShrinkKeys = [&](auto& KeyData)
 		{
-			if (KeyData.Num() >= (StartKeyIndex + NumKeysToRemove))
+			// Dont allow us to trim below 2 keys (1 frame)
+			int32 NumKeyDataToRemove = NumKeysToRemove;
+			if((KeyData.Num() - NumKeysToRemove) < 2)
+			{
+				NumKeyDataToRemove = KeyData.Num() - 2;
+			}
+
+			if (KeyData.Num() >= (StartKeyIndex + NumKeyDataToRemove))
 			{
-				KeyData.RemoveAt(StartKeyIndex, NumKeysToRemove);
+				KeyData.RemoveAt(StartKeyIndex, NumKeyDataToRemove);
 				check(KeyData.Num() > 0);
 				KeyData.Shrink();
 			}
