- Change SRV creation to use a FShaderResourceViewInitializer as this allows for null SRV

#rnx

diff --git a/Engine/Source/Runtime/Engine/Public/RawIndexBuffer.h b/Engine/Source/Runtime/Engine/Public/RawIndexBuffer.h
--- a/Engine/Source/Runtime/Engine/Public/RawIndexBuffer.h
+++ b/Engine/Source/Runtime/Engine/Public/RawIndexBuffer.h
@@ -403,7 +403,8 @@ public:
 
 		if (IndexBufferRHI && IsSRVNeeded() && bHadIndexData)
 		{
-			SRVValue = RHICreateShaderResourceView(IndexBufferRHI->GetSize() > 0 ? IndexBufferRHI : nullptr);
+			// If the index buffer is a placeholder we still need to create a FRHIShaderResourceView.
+			SRVValue = RHICreateShaderResourceView(FShaderResourceViewInitializer(IndexBufferRHI->GetSize() > 0 ? IndexBufferRHI : nullptr, sizeof(INDEX_TYPE) == 2 ? PF_R16_UINT : PF_R32_UINT));
 		}
 	}
 	
