Handle applying cooked meta-data in its PostLoad rather than its owner's PostLoad

BPs can PostLoad out-of-order from the optional data being loaded, which results in some meta-data being missed. This change means that the optional data will always apply itself when loaded.

#rnx

diff --git a/Engine/Source/Runtime/Engine/Private/BlueprintGeneratedClass.cpp b/Engine/Source/Runtime/Engine/Private/BlueprintGeneratedClass.cpp
--- a/Engine/Source/Runtime/Engine/Private/BlueprintGeneratedClass.cpp
+++ b/Engine/Source/Runtime/Engine/Private/BlueprintGeneratedClass.cpp
@@ -144,15 +144,6 @@ void UBlueprintGeneratedClass::PostLoad()
 	Super::PostLoad();
 
 #if WITH_EDITORONLY_DATA
-	if (bCooked)
-	{
-		if (const UClassCookedMetaData* CookedMetaData = FindCookedMetaData())
-		{
-			CookedMetaData->ApplyMetaData(this);
-			PurgeCookedMetaData();
-		}
-	}
-
 #if WITH_EDITOR
 	// Make BPGC from a cooked package standalone so it doesn't get GCed
 	if (GEditor && bCooked)
diff --git a/Engine/Source/Runtime/Engine/Private/CookedMetaData.cpp b/Engine/Source/Runtime/Engine/Private/CookedMetaData.cpp
--- a/Engine/Source/Runtime/Engine/Private/CookedMetaData.cpp
+++ b/Engine/Source/Runtime/Engine/Private/CookedMetaData.cpp
@@ -6,6 +6,40 @@
 
 #include UE_INLINE_GENERATED_CPP_BY_NAME(CookedMetaData)
 
+namespace CookedMetaDataUtil::Internal
+{
+
+void PrepareCookedMetaDataForPurge(UObject* CookedMetaDataPtr)
+{
+	// Skip the rename for cooked packages, as IO store cannot currently handle renames
+	if (!CookedMetaDataPtr->GetPackage()->HasAnyPackageFlags(PKG_Cooked))
+	{
+		FNameBuilder BaseMetaDataName(CookedMetaDataPtr->GetFName());
+		BaseMetaDataName << TEXT("_PURGED");
+		CookedMetaDataPtr->Rename(FNameBuilder(MakeUniqueObjectName(CookedMetaDataPtr->GetOuter(), CookedMetaDataPtr->GetClass(), FName(BaseMetaDataName))).ToString(), nullptr, REN_DoNotDirty | REN_DontCreateRedirectors | REN_ForceNoResetLoaders | REN_NonTransactional);
+	}
+
+	CookedMetaDataPtr->ClearFlags(RF_Standalone | RF_Public);
+}
+
+template <typename CookedMetaDataOuterType, typename CookedMetaDataType>
+void PostLoadCookedMetaData(CookedMetaDataType* CookedMetaDataPtr)
+{
+#if WITH_EDITORONLY_DATA
+	checkf(CookedMetaDataPtr->GetPackage()->HasAnyPackageFlags(PKG_Cooked), TEXT("Cooked meta-data should only be loaded for a cooked package!"));
+
+	if (CookedMetaDataOuterType* Owner = CastChecked<CookedMetaDataOuterType>(CookedMetaDataPtr->GetOuter()))
+	{
+		Owner->ConditionalPostLoad();
+		CookedMetaDataPtr->ApplyMetaData(Owner);
+		PrepareCookedMetaDataForPurge(CookedMetaDataPtr);
+	}
+#endif
+}
+
+} // namespace CookedMetaDataUtil::Internal
+
+
 bool FObjectCookedMetaDataStore::HasMetaData() const
 {
 	return ObjectMetaData.Num() > 0;
@@ -104,6 +138,12 @@ void FStructCookedMetaDataStore::ApplyMetaData(UStruct* TargetStruct) const
 }
 
 
+void UEnumCookedMetaData::PostLoad()
+{
+	Super::PostLoad();
+	CookedMetaDataUtil::Internal::PostLoadCookedMetaData<UEnum, UEnumCookedMetaData>(this);
+}
+
 bool UEnumCookedMetaData::HasMetaData() const
 {
 	return EnumMetaData.HasMetaData();
@@ -120,6 +160,13 @@ void UEnumCookedMetaData::ApplyMetaData(UEnum* TargetEnum) const
 }
 
 
+
+void UStructCookedMetaData::PostLoad()
+{
+	Super::PostLoad();
+	CookedMetaDataUtil::Internal::PostLoadCookedMetaData<UScriptStruct, UStructCookedMetaData>(this);
+}
+
 bool UStructCookedMetaData::HasMetaData() const
 {
 	return StructMetaData.HasMetaData();
@@ -136,6 +183,12 @@ void UStructCookedMetaData::ApplyMetaData(UScriptStruct* TargetStruct) const
 }
 
 
+void UClassCookedMetaData::PostLoad()
+{
+	Super::PostLoad();
+	CookedMetaDataUtil::Internal::PostLoadCookedMetaData<UClass, UClassCookedMetaData>(this);
+}
+
 bool UClassCookedMetaData::HasMetaData() const
 {
 	return ClassMetaData.HasMetaData()
@@ -170,19 +223,3 @@ void UClassCookedMetaData::ApplyMetaData(UClass* TargetClass) const
 		}
 	}
 }
-
-namespace CookedMetaDataUtil::Internal
-{
-
-void RenameCookedMetaDataForPurge(UObject* CookedMetaDataPtr)
-{
-	// Skip the rename for cooked packages, as IO store cannot currently handle renames
-	if (!CookedMetaDataPtr->GetPackage()->HasAnyPackageFlags(PKG_Cooked))
-	{
-		FNameBuilder BaseMetaDataName(CookedMetaDataPtr->GetFName());
-		BaseMetaDataName << TEXT("_PURGED");
-		CookedMetaDataPtr->Rename(FNameBuilder(MakeUniqueObjectName(CookedMetaDataPtr->GetOuter(), CookedMetaDataPtr->GetClass(), FName(BaseMetaDataName))).ToString(), nullptr, REN_DoNotDirty | REN_DontCreateRedirectors | REN_ForceNoResetLoaders | REN_NonTransactional);
-	}
-}
-
-} // namespace CookedMetaDataUtil::Internal
diff --git a/Engine/Source/Runtime/Engine/Private/UserDefinedEnum.cpp b/Engine/Source/Runtime/Engine/Private/UserDefinedEnum.cpp
--- a/Engine/Source/Runtime/Engine/Private/UserDefinedEnum.cpp
+++ b/Engine/Source/Runtime/Engine/Private/UserDefinedEnum.cpp
@@ -82,15 +82,6 @@ void UUserDefinedEnum::PostLoad()
 {
 	Super::PostLoad();
 
-	if (GetPackage()->HasAnyPackageFlags(PKG_Cooked))
-	{
-		if (const UEnumCookedMetaData* CookedMetaData = FindCookedMetaData())
-		{
-			CookedMetaData->ApplyMetaData(this);
-			PurgeCookedMetaData();
-		}
-	}
-
 	FEnumEditorUtils::UpdateAfterPathChanged(this);
 	if (NumEnums() > 1 && DisplayNameMap.Num() == 0) // >1 because User Defined Enums always have a "MAX" entry
 	{
diff --git a/Engine/Source/Runtime/Engine/Private/UserDefinedStruct.cpp b/Engine/Source/Runtime/Engine/Private/UserDefinedStruct.cpp
--- a/Engine/Source/Runtime/Engine/Private/UserDefinedStruct.cpp
+++ b/Engine/Source/Runtime/Engine/Private/UserDefinedStruct.cpp
@@ -163,15 +163,6 @@ void UUserDefinedStruct::PostLoad()
 {
 	Super::PostLoad();
 
-	if (GetPackage()->HasAnyPackageFlags(PKG_Cooked))
-	{
-		if (const UStructCookedMetaData* CookedMetaData = FindCookedMetaData())
-		{
-			CookedMetaData->ApplyMetaData(this);
-			PurgeCookedMetaData();
-		}
-	}
-
 	ValidateGuid();
 }
 
diff --git a/Engine/Source/Runtime/Engine/Public/CookedMetaData.h b/Engine/Source/Runtime/Engine/Public/CookedMetaData.h
--- a/Engine/Source/Runtime/Engine/Public/CookedMetaData.h
+++ b/Engine/Source/Runtime/Engine/Public/CookedMetaData.h
@@ -66,12 +66,14 @@ private:
 /**
  * Cooked meta-data for a UEnum.
  */
-UCLASS(Optional)
+UCLASS(Optional, Within=Enum)
 class ENGINE_API UEnumCookedMetaData : public UObject
 {
 public:
 	GENERATED_BODY()
 
+	virtual void PostLoad() override;
+	
 	virtual bool HasMetaData() const;
 	virtual void CacheMetaData(const UEnum* SourceEnum);
 	virtual void ApplyMetaData(UEnum* TargetEnum) const;
@@ -84,12 +86,14 @@ protected:
 /**
  * Cooked meta-data for a UScriptStruct, including its nested FProperty data.
  */
-UCLASS(Optional)
+UCLASS(Optional, Within=ScriptStruct)
 class ENGINE_API UStructCookedMetaData : public UObject
 {
 public:
 	GENERATED_BODY()
 
+	virtual void PostLoad() override;
+	
 	virtual bool HasMetaData() const;
 	virtual void CacheMetaData(const UScriptStruct* SourceStruct);
 	virtual void ApplyMetaData(UScriptStruct* TargetStruct) const;
@@ -102,12 +106,14 @@ protected:
 /**
  * Cooked meta-data for a UClass, including its nested FProperty and UFunction data.
  */
-UCLASS(Optional)
+UCLASS(Optional, Within=Class)
 class ENGINE_API UClassCookedMetaData : public UObject
 {
 public:
 	GENERATED_BODY()
 
+	virtual void PostLoad() override;
+	
 	virtual bool HasMetaData() const;
 	virtual void CacheMetaData(const UClass* SourceClass);
 	virtual void ApplyMetaData(UClass* TargetClass) const;
@@ -125,7 +131,7 @@ namespace CookedMetaDataUtil
 
 namespace Internal
 {
-ENGINE_API void RenameCookedMetaDataForPurge(UObject* CookedMetaDataPtr);
+ENGINE_API void PrepareCookedMetaDataForPurge(UObject* CookedMetaDataPtr);
 }
 
 template <typename CookedMetaDataType>
@@ -143,8 +149,7 @@ CookedMetaDataType* FindCookedMetaData(UObject* Outer, const TCHAR* Name)
 template <typename CookedMetaDataType, typename CookedMetaDataPtrType>
 void PurgeCookedMetaData(CookedMetaDataPtrType& CookedMetaDataPtr)
 {
-	Internal::RenameCookedMetaDataForPurge(CookedMetaDataPtr);
-	CookedMetaDataPtr->ClearFlags(RF_Standalone | RF_Public);
+	Internal::PrepareCookedMetaDataForPurge(CookedMetaDataPtr);
 	CookedMetaDataPtr = nullptr;
 }
 
