Skeleton compatibility improvements

Skeleton compatibility is now bi-directional. Specifying a compatible skeleton A -> B now implies B -> A.
Skeleton compatibility is now an editor-only concern. The runtime will attempt to do the 'best it can' via name -> name mappings. Only the editor will prevent assigning incompatible skeletons in (e.g.) asset pickers etc.
Skeleton compatibility checks in editor can now be disabled in the editor preferences (and each asset picker now has a checkbox option in its view settings that allows for quick access to this).

Moves FSkeletonRemapping to its own file (which is now private).
Skeleton remappings are now generated on demand on worker threads just before animation decompression and stored in a registry, guarded by FRWScopeLock for thread-safety.

Fixed some anim BP compiler edge cases where asset references on pins were not getting preloaded correctly, causing skeletons to be erroneously reported as missing.

Exposed the current asset registry filter in SAssetView so that menu extensions can access it (and use it to provide context)

diff --git a/Engine/Plugins/Animation/AnimToTexture/Source/AnimToTextureEditor/Private/AnimToTextureBPLibrary.cpp b/Engine/Plugins/Animation/AnimToTexture/Source/AnimToTextureEditor/Private/AnimToTextureBPLibrary.cpp
--- a/Engine/Plugins/Animation/AnimToTexture/Source/AnimToTextureEditor/Private/AnimToTextureBPLibrary.cpp
+++ b/Engine/Plugins/Animation/AnimToTexture/Source/AnimToTextureEditor/Private/AnimToTextureBPLibrary.cpp
@@ -188,9 +188,9 @@ void UAnimToTextureBPLibrary::AnimationToTexture(UAnimToTextureDataAsset* DataAs
 		}
 		
 		// Make sure SkeletalMesh is compatible with AnimSequence
-		if (!MeshComponent->GetSkeletalMeshAsset()->GetSkeleton()->IsCompatible(AnimSequence->GetSkeleton()))
+		if (!MeshComponent->GetSkeletalMeshAsset()->GetSkeleton()->IsCompatibleForEditor(AnimSequence->GetSkeleton()))
 		{
-			UE_LOG(LogTemp, Warning, TEXT("Invalid AnimSequence: %s for given SkeletalMesh: %s"), *AnimSequence->GetFName().ToString(), *MeshComponent->GetSkeletalMeshAsset()->GetFName().ToString());
+			UE_LOG(LogTemp, Warning, TEXT("Incompatible AnimSequence: %s for given SkeletalMesh: %s"), *AnimSequence->GetFName().ToString(), *MeshComponent->GetSkeletalMeshAsset()->GetFName().ToString());
 			continue;
 		}
 		// Set AnimSequence
diff --git a/Engine/Plugins/Animation/ControlRig/Source/ControlRigEditor/Private/Sequencer/ControlRigParameterTrackEditor.cpp b/Engine/Plugins/Animation/ControlRig/Source/ControlRigEditor/Private/Sequencer/ControlRigParameterTrackEditor.cpp
--- a/Engine/Plugins/Animation/ControlRig/Source/ControlRigEditor/Private/Sequencer/ControlRigParameterTrackEditor.cpp
+++ b/Engine/Plugins/Animation/ControlRig/Source/ControlRigEditor/Private/Sequencer/ControlRigParameterTrackEditor.cpp
@@ -720,7 +720,7 @@ public:
 				}
 				else if(Skeleton)
 				{
-					if (Skeleton->IsCompatibleSkeletonByAssetString(PreviewSkeleton))
+					if (Skeleton->IsCompatibleForEditor(PreviewSkeleton))
 					{
 						return true;
 					}
@@ -733,7 +733,7 @@ public:
 			}
 			else if (Skeleton)
 			{
-				if (Skeleton->IsCompatibleSkeletonByAssetString(PreviewSkeleton))
+				if (Skeleton->IsCompatibleForEditor(PreviewSkeleton))
 				{
 					return true;
 				}
@@ -745,7 +745,7 @@ public:
 			}
 			else if (Skeleton)
 			{
-				if (Skeleton->IsCompatibleSkeletonByAssetString(SourceHierarchyImport))
+				if (Skeleton->IsCompatibleForEditor(SourceHierarchyImport))
 				{
 					return true;
 				}
@@ -757,7 +757,7 @@ public:
 			}
 			else if (Skeleton)
 			{
-				if (Skeleton->IsCompatibleSkeletonByAssetString(SourceCurveImport))
+				if (Skeleton->IsCompatibleForEditor(SourceCurveImport))
 				{
 					return true;
 				}
@@ -4701,7 +4701,7 @@ void FControlRigParameterSection::AddAnimationSubMenuForFK(FMenuBuilder& MenuBui
 		AssetPickerConfig.OnShouldFilterAsset = FOnShouldFilterAsset::CreateRaw(this, &FControlRigParameterSection::ShouldFilterAssetForFK);
 		AssetPickerConfig.Filter.bRecursiveClasses = true;
 		AssetPickerConfig.Filter.ClassPaths.Add(UAnimSequenceBase::StaticClass()->GetClassPathName());
-		AssetPickerConfig.Filter.TagsAndValues.Add(TEXT("Skeleton"), FAssetData(Skeleton).GetExportTextName());
+		AssetPickerConfig.OnShouldFilterAsset = FOnShouldFilterAsset::CreateUObject(Skeleton, &USkeleton::ShouldFilterAsset, TEXT("Skeleton"));
 		AssetPickerConfig.SaveSettingsName = TEXT("SequencerAssetPicker");
 		AssetPickerConfig.AdditionalReferencingAssets.Add(FAssetData(Sequence));
 	}
diff --git a/Engine/Plugins/Animation/DeformerGraph/Source/OptimusEditor/Private/OptimusEditor.cpp b/Engine/Plugins/Animation/DeformerGraph/Source/OptimusEditor/Private/OptimusEditor.cpp
--- a/Engine/Plugins/Animation/DeformerGraph/Source/OptimusEditor/Private/OptimusEditor.cpp
+++ b/Engine/Plugins/Animation/DeformerGraph/Source/OptimusEditor/Private/OptimusEditor.cpp
@@ -94,10 +94,6 @@ void FOptimusEditor::Construct(
 	// graph component's scene component bindings.
 	PersonaToolkit->GetPreviewScene()->RegisterOnPreviewMeshChanged(FOnPreviewMeshChanged::CreateSP(this, &FOptimusEditor::HandlePreviewMeshChanged));
 
-	// TODO: Do we need this?
-	TSharedRef<IAssetFamily> AssetFamily = PersonaModule.CreatePersonaAssetFamily(InDeformerObject);
-	AssetFamily->RecordAssetOpened(FAssetData(InDeformerObject));
-	
 	// Construct a new graph with a default name
 	// TODO: Use a document manager like blueprints.
 	// FIXME: The deformer asset shouldn't really be the owner.
diff --git a/Engine/Plugins/Animation/IKRig/Source/IKRigEditor/Private/RetargetEditor/IKRetargetEditor.cpp b/Engine/Plugins/Animation/IKRig/Source/IKRigEditor/Private/RetargetEditor/IKRetargetEditor.cpp
--- a/Engine/Plugins/Animation/IKRig/Source/IKRigEditor/Private/RetargetEditor/IKRetargetEditor.cpp
+++ b/Engine/Plugins/Animation/IKRig/Source/IKRigEditor/Private/RetargetEditor/IKRetargetEditor.cpp
@@ -53,9 +53,6 @@ void FIKRetargetEditor::InitAssetEditor(
 	
 	FPersonaModule& PersonaModule = FModuleManager::LoadModuleChecked<FPersonaModule>("Persona");
 	PersonaToolkit = PersonaModule.CreatePersonaToolkit(InAsset, PersonaToolkitArgs);
-	
-	TSharedRef<IAssetFamily> AssetFamily = PersonaModule.CreatePersonaAssetFamily(InAsset);
-	AssetFamily->RecordAssetOpened(FAssetData(InAsset));
 
 	const bool bCreateDefaultStandaloneMenu = true;
 	const bool bCreateDefaultToolbar = true;
diff --git a/Engine/Plugins/Animation/IKRig/Source/IKRigEditor/Private/RetargetEditor/IKRetargetEditorController.cpp b/Engine/Plugins/Animation/IKRig/Source/IKRigEditor/Private/RetargetEditor/IKRetargetEditorController.cpp
--- a/Engine/Plugins/Animation/IKRig/Source/IKRigEditor/Private/RetargetEditor/IKRetargetEditorController.cpp
+++ b/Engine/Plugins/Animation/IKRig/Source/IKRigEditor/Private/RetargetEditor/IKRetargetEditorController.cpp
@@ -197,7 +197,7 @@ void FIKRetargetEditorController::HandleRetargeterNeedsInitialized(UIKRetargeter
 	{
 		// verify that animation is compatible with (potentially) new source mesh
 		USkeletalMesh* NewSourceMesh = GetSkeletalMesh(ERetargetSourceOrTarget::Source);
-		if (AnimThatWasPlaying->GetSkeleton()->IsCompatible(NewSourceMesh->GetSkeleton()))
+		if (AnimThatWasPlaying->GetSkeleton()->IsCompatibleForEditor(NewSourceMesh->GetSkeleton()))
 		{
 			SourceAnimInstance->SetAnimationAsset(AnimToRestore);
 			SourceAnimInstance->SetPlaying(true);
diff --git a/Engine/Plugins/Animation/IKRig/Source/IKRigEditor/Private/RetargetEditor/IKRetargetPoseExporter.cpp b/Engine/Plugins/Animation/IKRig/Source/IKRigEditor/Private/RetargetEditor/IKRetargetPoseExporter.cpp
--- a/Engine/Plugins/Animation/IKRig/Source/IKRigEditor/Private/RetargetEditor/IKRetargetPoseExporter.cpp
+++ b/Engine/Plugins/Animation/IKRig/Source/IKRigEditor/Private/RetargetEditor/IKRetargetPoseExporter.cpp
@@ -466,7 +466,7 @@ bool FIKRetargetPoseExporter::OnShouldFilterSequenceToImport(const FAssetData& A
 		return true;
 	}
 
-	return !DesiredSkeleton->IsCompatibleSkeletonByAssetData(AssetData);
+	return !DesiredSkeleton->IsCompatibleForEditor(AssetData);
 }
 
 FReply FIKRetargetPoseExporter::OnImportPoseFromSequence()
diff --git a/Engine/Plugins/Animation/IKRig/Source/IKRigEditor/Private/RetargetEditor/SIKRetargetAssetBrowser.cpp b/Engine/Plugins/Animation/IKRig/Source/IKRigEditor/Private/RetargetEditor/SIKRetargetAssetBrowser.cpp
--- a/Engine/Plugins/Animation/IKRig/Source/IKRigEditor/Private/RetargetEditor/SIKRetargetAssetBrowser.cpp
+++ b/Engine/Plugins/Animation/IKRig/Source/IKRigEditor/Private/RetargetEditor/SIKRetargetAssetBrowser.cpp
@@ -232,7 +232,7 @@ bool SIKRetargetAssetBrowser::OnShouldFilterAsset(const struct FAssetData& Asset
 		return true;
 	}
 
-	return !DesiredSkeleton->IsCompatibleSkeletonByAssetData(AssetData);
+	return !DesiredSkeleton->IsCompatibleForEditor(AssetData);
 }
 
 // ------------------------------------------BEGIN  SBatchExportDialog ----------------------------
diff --git a/Engine/Plugins/Animation/IKRig/Source/IKRigEditor/Private/RigEditor/IKRigToolkit.cpp b/Engine/Plugins/Animation/IKRig/Source/IKRigEditor/Private/RigEditor/IKRigToolkit.cpp
--- a/Engine/Plugins/Animation/IKRig/Source/IKRigEditor/Private/RigEditor/IKRigToolkit.cpp
+++ b/Engine/Plugins/Animation/IKRig/Source/IKRigEditor/Private/RigEditor/IKRigToolkit.cpp
@@ -60,9 +60,6 @@ void FIKRigEditorToolkit::InitAssetEditor(
 	const FPersonaModule& PersonaModule = FModuleManager::LoadModuleChecked<FPersonaModule>("Persona");
 	PersonaToolkit = PersonaModule.CreatePersonaToolkit(IKRigAsset, PersonaToolkitArgs);
 
-	const TSharedRef<IAssetFamily> AssetFamily = PersonaModule.CreatePersonaAssetFamily(IKRigAsset);
-	AssetFamily->RecordAssetOpened(FAssetData(IKRigAsset));
-
 	static constexpr bool bCreateDefaultStandaloneMenu = true;
 	static constexpr bool bCreateDefaultToolbar = true;
 	FAssetEditorToolkit::InitAssetEditor(
diff --git a/Engine/Plugins/Animation/IKRig/Source/IKRigEditor/Private/RigEditor/SIKRigAssetBrowser.cpp b/Engine/Plugins/Animation/IKRig/Source/IKRigEditor/Private/RigEditor/SIKRigAssetBrowser.cpp
--- a/Engine/Plugins/Animation/IKRig/Source/IKRigEditor/Private/RigEditor/SIKRigAssetBrowser.cpp
+++ b/Engine/Plugins/Animation/IKRig/Source/IKRigEditor/Private/RigEditor/SIKRigAssetBrowser.cpp
@@ -142,7 +142,7 @@ bool SIKRigAssetBrowser::OnShouldFilterAsset(const struct FAssetData& AssetData)
 		return true;
 	}
 
-	return !DesiredSkeleton->IsCompatibleSkeletonByAssetData(AssetData);
+	return !DesiredSkeleton->IsCompatibleForEditor(AssetData);
 }
 
 #undef LOCTEXT_NAMESPACE
diff --git a/Engine/Plugins/Animation/MLDeformer/MLDeformerFramework/Source/MLDeformerFrameworkEditor/Private/MLDeformerEditorModel.cpp b/Engine/Plugins/Animation/MLDeformer/MLDeformerFramework/Source/MLDeformerFrameworkEditor/Private/MLDeformerEditorModel.cpp
--- a/Engine/Plugins/Animation/MLDeformer/MLDeformerFramework/Source/MLDeformerFrameworkEditor/Private/MLDeformerEditorModel.cpp
+++ b/Engine/Plugins/Animation/MLDeformer/MLDeformerFramework/Source/MLDeformerFrameworkEditor/Private/MLDeformerEditorModel.cpp
@@ -934,9 +934,9 @@ namespace UE::MLDeformer
 		FText Result;
 		if (InSkelMesh && InAnimSeq)
 		{
-			if (!InSkelMesh->GetSkeleton()->IsCompatible(InAnimSeq->GetSkeleton()))
+			if (!InSkelMesh->GetSkeleton()->IsCompatibleForEditor(InAnimSeq->GetSkeleton()))
 			{
-				Result = LOCTEXT("SkeletonMismatch", "The base skeletal mesh and anim sequence use different skeletons. The animation might not play correctly.");
+				Result = LOCTEXT("SkeletonMismatch", "The base skeletal mesh and anim sequence use incompatible skeletons. The animation might not play correctly.");
 			}
 		}
 		return Result;
diff --git a/Engine/Plugins/Animation/MLDeformer/MLDeformerFramework/Source/MLDeformerFrameworkEditor/Private/MLDeformerEditorToolkit.cpp b/Engine/Plugins/Animation/MLDeformer/MLDeformerFramework/Source/MLDeformerFrameworkEditor/Private/MLDeformerEditorToolkit.cpp
--- a/Engine/Plugins/Animation/MLDeformer/MLDeformerFramework/Source/MLDeformerFrameworkEditor/Private/MLDeformerEditorToolkit.cpp
+++ b/Engine/Plugins/Animation/MLDeformer/MLDeformerFramework/Source/MLDeformerFrameworkEditor/Private/MLDeformerEditorToolkit.cpp
@@ -74,9 +74,6 @@ namespace UE::MLDeformer
 		FPersonaModule& PersonaModule = FModuleManager::LoadModuleChecked<FPersonaModule>("Persona");
 		PersonaToolkit = PersonaModule.CreatePersonaToolkit(DeformerAsset, PersonaToolkitArgs);
 
-		TSharedRef<IAssetFamily> AssetFamily = PersonaModule.CreatePersonaAssetFamily(DeformerAsset);
-		AssetFamily->RecordAssetOpened(FAssetData(DeformerAsset));
-
 		const bool bCreateDefaultStandaloneMenu = true;
 		const bool bCreateDefaultToolbar = true;
 		FAssetEditorToolkit::InitAssetEditor(
diff --git a/Engine/Plugins/Animation/MLDeformer/MLDeformerFramework/Source/MLDeformerFrameworkEditor/Private/MLDeformerModelDetails.cpp b/Engine/Plugins/Animation/MLDeformer/MLDeformerFramework/Source/MLDeformerFrameworkEditor/Private/MLDeformerModelDetails.cpp
--- a/Engine/Plugins/Animation/MLDeformer/MLDeformerFramework/Source/MLDeformerFrameworkEditor/Private/MLDeformerModelDetails.cpp
+++ b/Engine/Plugins/Animation/MLDeformer/MLDeformerFramework/Source/MLDeformerFrameworkEditor/Private/MLDeformerModelDetails.cpp
@@ -345,7 +345,7 @@ namespace UE::MLDeformer
 
 	bool FMLDeformerModelDetails::FilterAnimSequences(const FAssetData& AssetData, USkeleton* Skeleton)
 	{
-		if (Skeleton && Skeleton->IsCompatibleSkeletonByAssetData(AssetData))
+		if (Skeleton && Skeleton->IsCompatibleForEditor(AssetData))
 		{
 			return false;
 		}
diff --git a/Engine/Plugins/Animation/MLDeformer/MLDeformerFramework/Source/MLDeformerFrameworkEditor/Private/MLDeformerVizSettingsDetails.cpp b/Engine/Plugins/Animation/MLDeformer/MLDeformerFramework/Source/MLDeformerFrameworkEditor/Private/MLDeformerVizSettingsDetails.cpp
--- a/Engine/Plugins/Animation/MLDeformer/MLDeformerFramework/Source/MLDeformerFrameworkEditor/Private/MLDeformerVizSettingsDetails.cpp
+++ b/Engine/Plugins/Animation/MLDeformer/MLDeformerFramework/Source/MLDeformerFrameworkEditor/Private/MLDeformerVizSettingsDetails.cpp
@@ -175,7 +175,7 @@ namespace UE::MLDeformer
 
 	bool FMLDeformerVizSettingsDetails::FilterAnimSequences(const FAssetData& AssetData, USkeleton* Skeleton)
 	{
-		if (Skeleton && Skeleton->IsCompatibleSkeletonByAssetData(AssetData))
+		if (Skeleton && Skeleton->IsCompatibleForEditor(AssetData))
 		{
 			return false;
 		}
diff --git a/Engine/Plugins/Developer/AnimationSharing/Source/AnimationSharing/Private/AnimationSharingManager.cpp b/Engine/Plugins/Developer/AnimationSharing/Source/AnimationSharing/Private/AnimationSharingManager.cpp
--- a/Engine/Plugins/Developer/AnimationSharing/Source/AnimationSharing/Private/AnimationSharingManager.cpp
+++ b/Engine/Plugins/Developer/AnimationSharing/Source/AnimationSharing/Private/AnimationSharingManager.cpp
@@ -403,7 +403,7 @@ void UAnimationSharingManager::RegisterActorWithSkeleton(AActor* InActor, const
 	{
 		uint32 ArrayIndex = Skeletons.IndexOfByPredicate([SharingSkeleton](const USkeleton* Skeleton)
 		{
-			return (Skeleton == SharingSkeleton) || (Skeleton->IsCompatible(SharingSkeleton));
+			return (Skeleton == SharingSkeleton);
 		});
 		return ArrayIndex;
 	}();
diff --git a/Engine/Plugins/Experimental/Animation/PoseSearch/Source/Runtime/Private/PoseSearch.cpp b/Engine/Plugins/Experimental/Animation/PoseSearch/Source/Runtime/Private/PoseSearch.cpp
--- a/Engine/Plugins/Experimental/Animation/PoseSearch/Source/Runtime/Private/PoseSearch.cpp
+++ b/Engine/Plugins/Experimental/Animation/PoseSearch/Source/Runtime/Private/PoseSearch.cpp
@@ -1069,7 +1069,7 @@ bool UPoseSearchDatabase::IsValidForIndexing() const
 			}
 
 			const USkeleton* SeqSkeleton = DbSequence.Sequence->GetSkeleton();
-			if (!SeqSkeleton || !SeqSkeleton->IsCompatible(Schema->Skeleton))
+			if (!SeqSkeleton)
 			{
 				bValid = false;
 				break;
@@ -1085,7 +1085,7 @@ bool UPoseSearchDatabase::IsValidForIndexing() const
 			}
 
 			const USkeleton* SeqSkeleton = DbBlendSpace.BlendSpace->GetSkeleton();
-			if (!SeqSkeleton || !SeqSkeleton->IsCompatible(Schema->Skeleton))
+			if (!SeqSkeleton)
 			{
 				bValid = false;
 				break;
@@ -4425,7 +4425,7 @@ bool BuildIndex(const UAnimSequence* Sequence, UPoseSearchSequenceMetaData* Sequ
 	}
 
 	USkeleton* SeqSkeleton = Sequence->GetSkeleton();
-	if (!SeqSkeleton || !SeqSkeleton->IsCompatible(SequenceMetaData->Schema->Skeleton))
+	if (!SeqSkeleton)
 	{
 		return false;
 	}
diff --git a/Engine/Plugins/Experimental/Mutable/Source/CustomizableObjectEditor/Private/MuCOE/CustomizableObjectEditorViewportClient.cpp b/Engine/Plugins/Experimental/Mutable/Source/CustomizableObjectEditor/Private/MuCOE/CustomizableObjectEditorViewportClient.cpp
--- a/Engine/Plugins/Experimental/Mutable/Source/CustomizableObjectEditor/Private/MuCOE/CustomizableObjectEditorViewportClient.cpp
+++ b/Engine/Plugins/Experimental/Mutable/Source/CustomizableObjectEditor/Private/MuCOE/CustomizableObjectEditorViewportClient.cpp
@@ -2098,7 +2098,7 @@ void FCustomizableObjectEditorViewportClient::SetAnimation(UAnimationAsset* Anim
 	{
 		if (SkeletalMeshComponent.IsValid() && Animation != nullptr
 			&& SkeletalMeshComponent->GetSkinnedAsset() != nullptr
-			&& SkeletalMeshComponent->GetSkinnedAsset()->GetSkeleton()->IsCompatible(Animation->GetSkeleton())
+			&& SkeletalMeshComponent->GetSkinnedAsset()->GetSkeleton()->IsCompatibleForEditor(Animation->GetSkeleton())
 			)
 		{
 			SetRealtime(true);
diff --git a/Engine/Plugins/Importers/USDImporter/Source/USDSchemas/Private/USDSkelRootTranslator.cpp b/Engine/Plugins/Importers/USDImporter/Source/USDSchemas/Private/USDSkelRootTranslator.cpp
--- a/Engine/Plugins/Importers/USDImporter/Source/USDSchemas/Private/USDSkelRootTranslator.cpp
+++ b/Engine/Plugins/Importers/USDImporter/Source/USDSchemas/Private/USDSkelRootTranslator.cpp
@@ -926,7 +926,7 @@ namespace UsdSkelRootTranslatorImpl
 		else
 		{
 			// Force skeletons to be compatible if they aren't (we need both ways!)
-			if ( !Skeleton->IsCompatible( AnimBP->TargetSkeleton ) )
+			if ( !Skeleton->IsCompatibleForEditor( AnimBP->TargetSkeleton ) )
 			{
 				AnimBP->TargetSkeleton->AddCompatibleSkeleton( Skeleton );
 				Skeleton->AddCompatibleSkeleton( AnimBP->TargetSkeleton );
diff --git a/Engine/Plugins/Runtime/HairStrands/Source/HairStrandsEditor/Private/GroomCustomAssetEditorToolkit.cpp b/Engine/Plugins/Runtime/HairStrands/Source/HairStrandsEditor/Private/GroomCustomAssetEditorToolkit.cpp
--- a/Engine/Plugins/Runtime/HairStrands/Source/HairStrandsEditor/Private/GroomCustomAssetEditorToolkit.cpp
+++ b/Engine/Plugins/Runtime/HairStrands/Source/HairStrandsEditor/Private/GroomCustomAssetEditorToolkit.cpp
@@ -168,7 +168,7 @@ bool FGroomCustomAssetEditorToolkit::OnShouldFilterAnimAsset(const FAssetData& A
 	if (PreviewSkeletalMeshComponent != nullptr)
 	{
 		USkeleton* Skeleton = PreviewSkeletalMeshComponent->GetSkeletalMeshAsset()->GetSkeleton();
-		if (Skeleton && Skeleton->IsCompatibleSkeletonByAssetData(AssetData))
+		if (Skeleton && Skeleton->IsCompatibleForEditor(AssetData))
 		{
 			return false;
 		}
@@ -379,7 +379,7 @@ static UAnimationAsset* GetFirstCompatibleAnimAsset(const USkeletalMeshComponent
 	// Filter binding asset which match the groom asset (as the tag/value filter above does not work)
 	for (FAssetData& Asset : AnimAssetData)
 	{
-		if (InSkeleton->IsCompatibleSkeletonByAssetData(Asset))
+		if (InSkeleton->IsCompatibleForEditor(Asset))
 		{
 			return (UAnimationAsset*)Asset.GetAsset();
 		}
diff --git a/Engine/Source/Developer/AssetTools/Private/AssetTypeActions/AssetTypeActions_AnimBlueprint.cpp b/Engine/Source/Developer/AssetTools/Private/AssetTypeActions/AssetTypeActions_AnimBlueprint.cpp
--- a/Engine/Source/Developer/AssetTools/Private/AssetTypeActions/AssetTypeActions_AnimBlueprint.cpp
+++ b/Engine/Source/Developer/AssetTools/Private/AssetTypeActions/AssetTypeActions_AnimBlueprint.cpp
@@ -108,7 +108,7 @@ void FAssetTypeActions_AnimBlueprint::GetActions(const TArray<UObject*>& InObjec
 							{
 								// Filter on compatible skeletons
 								const FString ExportTextName = AssetData.GetExportTextName();
-								return !LocalAnimBlueprint->TargetSkeleton->IsCompatibleSkeletonByAssetString(ExportTextName);
+								return !LocalAnimBlueprint->TargetSkeleton->IsCompatibleForEditor(ExportTextName);
 							}
 						}
 						
diff --git a/Engine/Source/Editor/AnimGraph/Private/AnimGraphNode_AimOffsetLookAt.cpp b/Engine/Source/Editor/AnimGraph/Private/AnimGraphNode_AimOffsetLookAt.cpp
--- a/Engine/Source/Editor/AnimGraph/Private/AnimGraphNode_AimOffsetLookAt.cpp
+++ b/Engine/Source/Editor/AnimGraph/Private/AnimGraphNode_AimOffsetLookAt.cpp
@@ -175,11 +175,6 @@ void UAnimGraphNode_AimOffsetLookAt::ValidateAnimNodeDuringCompilation(class USk
 		const USkeleton* BlendSpaceSkeleton = BlendSpaceToCheck->GetSkeleton();
 		if (BlendSpaceSkeleton) // if blend space doesn't have skeleton, it might be due to blend space not loaded yet, @todo: wait with anim blueprint compilation until all assets are loaded?
 		{
-			if (ForSkeleton && !ForSkeleton->IsCompatible(BlendSpaceSkeleton))
-			{
-				MessageLog.Error(TEXT("@@ references blendspace that uses an incompatible skeleton @@"), this, BlendSpaceSkeleton);
-			}
-
 			// Temporary fix where skeleton is not fully loaded during AnimBP compilation and thus the socket name check is invalid UE-39499 (NEED FIX) 
 			if (!BlendSpaceSkeleton->HasAnyFlags(RF_NeedPostLoad))
 			{
diff --git a/Engine/Source/Editor/AnimGraph/Private/AnimGraphNode_AssetPlayerBase.cpp b/Engine/Source/Editor/AnimGraph/Private/AnimGraphNode_AssetPlayerBase.cpp
--- a/Engine/Source/Editor/AnimGraph/Private/AnimGraphNode_AssetPlayerBase.cpp
+++ b/Engine/Source/Editor/AnimGraph/Private/AnimGraphNode_AssetPlayerBase.cpp
@@ -352,7 +352,7 @@ bool UAnimGraphNode_AssetPlayerBase::IsActionFilteredOut(class FBlueprintActionF
 			UAnimationAsset* Asset = GetAnimationAsset();
 			if(Asset)
 			{
-				if (AnimBlueprint->TargetSkeleton == nullptr || !AnimBlueprint->TargetSkeleton->IsCompatible(Asset->GetSkeleton()))
+				if (AnimBlueprint->TargetSkeleton == nullptr || !AnimBlueprint->TargetSkeleton->IsCompatibleForEditor(Asset->GetSkeleton()))
 				{
 					// Asset does not use a compatible skeleton with the Blueprint, cannot use
 					bIsFilteredOut = true;
@@ -361,7 +361,7 @@ bool UAnimGraphNode_AssetPlayerBase::IsActionFilteredOut(class FBlueprintActionF
 			}
 			else if(!UnloadedSkeletonName.IsEmpty())
 			{
-				if(AnimBlueprint->TargetSkeleton == nullptr || !AnimBlueprint->TargetSkeleton->IsCompatibleSkeletonByAssetString(UnloadedSkeletonName))
+				if(AnimBlueprint->TargetSkeleton == nullptr || !AnimBlueprint->TargetSkeleton->IsCompatibleForEditor(UnloadedSkeletonName))
 				{
 					bIsFilteredOut = true;
 					break;
@@ -534,11 +534,22 @@ void UAnimGraphNode_AssetPlayerBase::ValidateAnimNodeDuringCompilationHelper(USk
 		USkeleton* AssetSkeleton = AssetToCheck->GetSkeleton();
 		
 		// if asset doesn't have a skeleton, it might be due to the asset no being not loaded yet
-		if(AssetSkeleton && (!ForSkeleton || !ForSkeleton->IsCompatible(AssetSkeleton)))
+		if(AssetSkeleton == nullptr)
 		{
-			MessageLog.Error(*FText::Format(LOCTEXT("IncompatibleSkeletonFormat", "@@ references {0} that uses an incompatible skeleton @@"), InAssetType->GetDisplayNameText()).ToString(), this, AssetSkeleton);
+			MessageLog.Error(*FText::Format(LOCTEXT("MissingSkeletonFormat", "@@ references {0} that uses a missing skeleton @@"), InAssetType->GetDisplayNameText()).ToString(), this, AssetSkeleton);
 		}
 	}
 }
 
+void UAnimGraphNode_AssetPlayerBase::PreloadRequiredAssetsHelper(UAnimationAsset* InAsset, UEdGraphPin* InExposedPin)
+{
+	UAnimationAsset* AssetToLoad = InAsset;
+	if(InExposedPin != nullptr && AssetToLoad == nullptr)
+	{
+		AssetToLoad = Cast<UAnimationAsset>(InExposedPin->DefaultObject);
+	}
+
+	PreloadObject(AssetToLoad);
+}
+
 #undef LOCTEXT_NAMESPACE
\ No newline at end of file
diff --git a/Engine/Source/Editor/AnimGraph/Private/AnimGraphNode_BlendSpaceBase.cpp b/Engine/Source/Editor/AnimGraph/Private/AnimGraphNode_BlendSpaceBase.cpp
--- a/Engine/Source/Editor/AnimGraph/Private/AnimGraphNode_BlendSpaceBase.cpp
+++ b/Engine/Source/Editor/AnimGraph/Private/AnimGraphNode_BlendSpaceBase.cpp
@@ -56,7 +56,7 @@ void UAnimGraphNode_BlendSpaceBase::CustomizePinData(UEdGraphPin* Pin, FName Sou
 
 void UAnimGraphNode_BlendSpaceBase::PreloadRequiredAssets()
 {
-	PreloadObject(GetBlendSpace());
+	PreloadRequiredAssetsHelper(GetBlendSpace(), FindPin(GET_MEMBER_NAME_STRING_CHECKED(FAnimNode_BlendSpacePlayer, BlendSpace)));
 
 	Super::PreloadRequiredAssets();
 }
diff --git a/Engine/Source/Editor/AnimGraph/Private/AnimGraphNode_BlendSpaceGraphBase.cpp b/Engine/Source/Editor/AnimGraph/Private/AnimGraphNode_BlendSpaceGraphBase.cpp
--- a/Engine/Source/Editor/AnimGraph/Private/AnimGraphNode_BlendSpaceGraphBase.cpp
+++ b/Engine/Source/Editor/AnimGraph/Private/AnimGraphNode_BlendSpaceGraphBase.cpp
@@ -570,7 +570,7 @@ bool UAnimGraphNode_BlendSpaceGraphBase::IsActionFilteredOut(class FBlueprintAct
 		{
 			if(!SkeletonName.IsEmpty())
 			{
-				if(AnimBlueprint->TargetSkeleton == nullptr || !AnimBlueprint->TargetSkeleton->IsCompatibleSkeletonByAssetString(SkeletonName))
+				if(AnimBlueprint->TargetSkeleton == nullptr || !AnimBlueprint->TargetSkeleton->IsCompatibleForEditor(SkeletonName))
 				{
 					bIsFilteredOut = true;
 					break;
diff --git a/Engine/Source/Editor/AnimGraph/Private/AnimGraphNode_Mirror.cpp b/Engine/Source/Editor/AnimGraph/Private/AnimGraphNode_Mirror.cpp
--- a/Engine/Source/Editor/AnimGraph/Private/AnimGraphNode_Mirror.cpp
+++ b/Engine/Source/Editor/AnimGraph/Private/AnimGraphNode_Mirror.cpp
@@ -72,9 +72,9 @@ void UAnimGraphNode_Mirror::ValidateAnimNodeDuringCompilation(class USkeleton* F
 	else if (ForSkeleton)
 	{
 		const USkeleton* MirrorTableSkeleton = Node.GetMirrorDataTable()->Skeleton;
-		if (!ForSkeleton->IsCompatible(MirrorTableSkeleton))
+		if (MirrorTableSkeleton == nullptr)
 		{
-			MessageLog.Error(TEXT("@@ has a mirror data table that is not compatible with the current skeleton. Please update the table or create a new table for the skeleton."), this);
+			MessageLog.Warning(TEXT("@@ has a mirror data table that has a missing skeleton. Please update the table or create a new table for the skeleton."), this);
 		}
 	}
 	
@@ -264,7 +264,7 @@ bool UAnimGraphNode_Mirror::HasMirrorDataTableForBlueprints(const TArray<UBluepr
 				if (Asset.IsAssetLoaded())
 				{
 					UMirrorDataTable* MirrorDataTable = Cast<UMirrorDataTable>(Asset.GetAsset());
-					if (AnimBlueprint->TargetSkeleton && AnimBlueprint->TargetSkeleton->IsCompatible(MirrorDataTable->Skeleton))
+					if (AnimBlueprint->TargetSkeleton && AnimBlueprint->TargetSkeleton->IsCompatibleForEditor(MirrorDataTable->Skeleton))
 					{
 						return true; 
 					}
@@ -296,7 +296,7 @@ bool UAnimGraphNode_Mirror::IsActionFilteredOut(class FBlueprintActionFilter con
 			UMirrorDataTable* MirrorDataTable = Node.GetMirrorDataTable();
 			if (MirrorDataTable)
 			{
-				if(AnimBlueprint->TargetSkeleton == nullptr || !AnimBlueprint->TargetSkeleton->IsCompatible(MirrorDataTable->Skeleton))
+				if(AnimBlueprint->TargetSkeleton == nullptr || !AnimBlueprint->TargetSkeleton->IsCompatibleForEditor(MirrorDataTable->Skeleton))
 				{
 					// Mirror Data Table does not use the same skeleton as the Blueprint, cannot use
 					bIsFilteredOut = true;
@@ -305,7 +305,7 @@ bool UAnimGraphNode_Mirror::IsActionFilteredOut(class FBlueprintActionFilter con
 			}
 			else if(!UnloadedSkeletonName.IsEmpty())
 			{
-				if(AnimBlueprint->TargetSkeleton == nullptr || !AnimBlueprint->TargetSkeleton->IsCompatibleSkeletonByAssetString(UnloadedSkeletonName))
+				if(AnimBlueprint->TargetSkeleton == nullptr || !AnimBlueprint->TargetSkeleton->IsCompatibleForEditor(UnloadedSkeletonName))
 				{
 					bIsFilteredOut = true;
 					break;
diff --git a/Engine/Source/Editor/AnimGraph/Private/AnimGraphNode_PoseByName.cpp b/Engine/Source/Editor/AnimGraph/Private/AnimGraphNode_PoseByName.cpp
--- a/Engine/Source/Editor/AnimGraph/Private/AnimGraphNode_PoseByName.cpp
+++ b/Engine/Source/Editor/AnimGraph/Private/AnimGraphNode_PoseByName.cpp
@@ -18,7 +18,7 @@ UAnimGraphNode_PoseByName::UAnimGraphNode_PoseByName(const FObjectInitializer& O
 
 void UAnimGraphNode_PoseByName::PreloadRequiredAssets()
 {
-	PreloadObject(Node.PoseAsset);
+	PreloadRequiredAssetsHelper(Node.PoseAsset, FindPin(GET_MEMBER_NAME_STRING_CHECKED(FAnimNode_PoseHandler, PoseAsset)));
 
 	Super::PreloadRequiredAssets();
 }
@@ -122,11 +122,10 @@ void UAnimGraphNode_PoseByName::ValidateAnimNodeDuringCompilation(class USkeleto
 	}
 	else if (ForSkeleton)
 	{
-		USkeleton* SeqSkeleton = PoseAssetToCheck->GetSkeleton();
-		if (SeqSkeleton && // if anim sequence doesn't have skeleton, it might be due to anim sequence not loaded yet, @todo: wait with anim blueprint compilation until all assets are loaded?
-			!ForSkeleton->IsCompatible(SeqSkeleton))
+		USkeleton* PoseAssetSkeleton = PoseAssetToCheck->GetSkeleton();
+		if (PoseAssetSkeleton == nullptr)
 		{
-			MessageLog.Error(TEXT("@@ references sequence that uses an incompatible skeleton @@"), this, SeqSkeleton);
+			MessageLog.Error(TEXT("@@ references pose asset that uses a missing skeleton @@"), this, PoseAssetSkeleton);
 		}
 	}
 }
diff --git a/Engine/Source/Editor/AnimGraph/Private/AnimGraphNode_PoseHandler.cpp b/Engine/Source/Editor/AnimGraph/Private/AnimGraphNode_PoseHandler.cpp
--- a/Engine/Source/Editor/AnimGraph/Private/AnimGraphNode_PoseHandler.cpp
+++ b/Engine/Source/Editor/AnimGraph/Private/AnimGraphNode_PoseHandler.cpp
@@ -41,7 +41,7 @@ void UAnimGraphNode_PoseHandler::OnOverrideAssets(IAnimBlueprintNodeOverrideAsse
 
 void UAnimGraphNode_PoseHandler::PreloadRequiredAssets()
 {
-	PreloadObject(GetPoseHandlerNode()->PoseAsset);
+	PreloadRequiredAssetsHelper(GetPoseHandlerNode()->PoseAsset, FindPin(GET_MEMBER_NAME_STRING_CHECKED(FAnimNode_PoseHandler, PoseAsset)));
 
 	Super::PreloadRequiredAssets();
 }
diff --git a/Engine/Source/Editor/AnimGraph/Private/AnimGraphNode_RotationOffsetBlendSpace.cpp b/Engine/Source/Editor/AnimGraph/Private/AnimGraphNode_RotationOffsetBlendSpace.cpp
--- a/Engine/Source/Editor/AnimGraph/Private/AnimGraphNode_RotationOffsetBlendSpace.cpp
+++ b/Engine/Source/Editor/AnimGraph/Private/AnimGraphNode_RotationOffsetBlendSpace.cpp
@@ -161,10 +161,9 @@ void UAnimGraphNode_RotationOffsetBlendSpace::ValidateAnimNodeDuringCompilation(
 	else if (ForSkeleton)
 	{
 		USkeleton* BlendSpaceSkeleton = BlendSpaceToCheck->GetSkeleton();
-		if (BlendSpaceSkeleton && // if blend space doesn't have skeleton, it might be due to blend space not loaded yet, @todo: wait with anim blueprint compilation until all assets are loaded?
-			!ForSkeleton->IsCompatible(BlendSpaceSkeleton))
+		if (BlendSpaceSkeleton == nullptr)
 		{
-			MessageLog.Error(TEXT("@@ references blendspace that uses an incompatible skeleton @@"), this, BlendSpaceSkeleton);
+			MessageLog.Error(TEXT("@@ references blendspace that uses a missing skeleton @@"), this, BlendSpaceSkeleton);
 		}
 	}
 
diff --git a/Engine/Source/Editor/AnimGraph/Private/AnimGraphNode_SequenceEvaluator.cpp b/Engine/Source/Editor/AnimGraph/Private/AnimGraphNode_SequenceEvaluator.cpp
--- a/Engine/Source/Editor/AnimGraph/Private/AnimGraphNode_SequenceEvaluator.cpp
+++ b/Engine/Source/Editor/AnimGraph/Private/AnimGraphNode_SequenceEvaluator.cpp
@@ -29,7 +29,7 @@ UAnimGraphNode_SequenceEvaluator::UAnimGraphNode_SequenceEvaluator(const FObject
 
 void UAnimGraphNode_SequenceEvaluator::PreloadRequiredAssets()
 {
-	PreloadObject(Node.GetSequence());
+	PreloadRequiredAssetsHelper(Node.GetSequence(), FindPin(GET_MEMBER_NAME_STRING_CHECKED(FAnimNode_SequenceEvaluator, Sequence)));
 
 	Super::PreloadRequiredAssets();
 }
diff --git a/Engine/Source/Editor/AnimGraph/Private/AnimGraphNode_SequencePlayer.cpp b/Engine/Source/Editor/AnimGraph/Private/AnimGraphNode_SequencePlayer.cpp
--- a/Engine/Source/Editor/AnimGraph/Private/AnimGraphNode_SequencePlayer.cpp
+++ b/Engine/Source/Editor/AnimGraph/Private/AnimGraphNode_SequencePlayer.cpp
@@ -45,7 +45,7 @@ void UAnimGraphNode_SequencePlayer::Serialize(FArchive& Ar)
 
 void UAnimGraphNode_SequencePlayer::PreloadRequiredAssets()
 {
-	PreloadObject(Node.GetSequence());
+	PreloadRequiredAssetsHelper(Node.GetSequence(), FindPin(GET_MEMBER_NAME_STRING_CHECKED(FAnimNode_SequencePlayer, Sequence)));
 
 	Super::PreloadRequiredAssets();
 }
diff --git a/Engine/Source/Editor/AnimGraph/Private/AnimationGraphSchema.cpp b/Engine/Source/Editor/AnimGraph/Private/AnimationGraphSchema.cpp
--- a/Engine/Source/Editor/AnimGraph/Private/AnimationGraphSchema.cpp
+++ b/Engine/Source/Editor/AnimGraph/Private/AnimationGraphSchema.cpp
@@ -564,7 +564,7 @@ void UAnimationGraphSchema::SpawnNodeFromAsset(UAnimationAsset* Asset, const FVe
 
 	UAnimBlueprint* AnimBlueprint = Cast<UAnimBlueprint>(FBlueprintEditorUtils::FindBlueprintForGraph(Graph));
 
-	const bool bSkelMatch = (AnimBlueprint != nullptr) && (AnimBlueprint->TargetSkeleton != nullptr) && (AnimBlueprint->TargetSkeleton->IsCompatible(Asset->GetSkeleton()));
+	const bool bSkelMatch = (AnimBlueprint != nullptr) && (AnimBlueprint->TargetSkeleton != nullptr) && (AnimBlueprint->TargetSkeleton->IsCompatibleForEditor(Asset->GetSkeleton()));
 	const bool bTypeMatch = (PinIfAvailable == nullptr) || UAnimationGraphSchema::IsLocalSpacePosePin(PinIfAvailable->PinType);
 	const bool bDirectionMatch = (PinIfAvailable == nullptr) || (PinIfAvailable->Direction == EGPD_Input);
 
@@ -673,7 +673,7 @@ void UAnimationGraphSchema::GetAssetsNodeHoverMessage(const TArray<FAssetData>&
 
 	// this one only should happen when there is an Anim Blueprint
 	UAnimBlueprint* AnimBlueprint = Cast<UAnimBlueprint>(FBlueprintEditorUtils::FindBlueprintForNode(HoverNode));
-	const bool bSkelMatch = (AnimBlueprint != NULL) && (AnimBlueprint->TargetSkeleton->IsCompatible(Asset->GetSkeleton()));
+	const bool bSkelMatch = (AnimBlueprint != NULL) && (AnimBlueprint->TargetSkeleton->IsCompatibleForEditor(Asset->GetSkeleton()));
 
 	if (!bSkelMatch)
 	{
@@ -705,7 +705,7 @@ void UAnimationGraphSchema::GetAssetsPinHoverMessage(const TArray<FAssetData>& A
 	// this one only should happen when there is an Anim Blueprint
 	UAnimBlueprint* AnimBlueprint = Cast<UAnimBlueprint>(FBlueprintEditorUtils::FindBlueprintForNode(HoverPin->GetOwningNode()));
 
-	const bool bSkelMatch = (AnimBlueprint != NULL) && (AnimBlueprint->TargetSkeleton->IsCompatible(Asset->GetSkeleton()));
+	const bool bSkelMatch = (AnimBlueprint != NULL) && (AnimBlueprint->TargetSkeleton->IsCompatibleForEditor(Asset->GetSkeleton()));
 	const bool bTypeMatch = UAnimationGraphSchema::IsLocalSpacePosePin(HoverPin->PinType);
 	const bool bDirectionMatch = HoverPin->Direction == EGPD_Input;
 
@@ -726,7 +726,7 @@ void UAnimationGraphSchema::GetAssetsGraphHoverMessage(const TArray<FAssetData>&
 	if (UAnimationAsset* AnimationAsset = FAssetData::GetFirstAsset<UAnimationAsset>(Assets))
 	{
 		UAnimBlueprint* AnimBlueprint = Cast<UAnimBlueprint>(FBlueprintEditorUtils::FindBlueprintForGraph(HoverGraph));
-		const bool bSkelMatch = (AnimBlueprint != nullptr) && (AnimBlueprint->TargetSkeleton != nullptr) && (AnimBlueprint->TargetSkeleton->IsCompatible(AnimationAsset->GetSkeleton()));
+		const bool bSkelMatch = (AnimBlueprint != nullptr) && (AnimBlueprint->TargetSkeleton != nullptr) && (AnimBlueprint->TargetSkeleton->IsCompatibleForEditor(AnimationAsset->GetSkeleton()));
 		if (!bSkelMatch)
 		{
 			OutOkIcon = false;
diff --git a/Engine/Source/Editor/AnimGraph/Public/AnimGraphNode_AssetPlayerBase.h b/Engine/Source/Editor/AnimGraph/Public/AnimGraphNode_AssetPlayerBase.h
--- a/Engine/Source/Editor/AnimGraph/Public/AnimGraphNode_AssetPlayerBase.h
+++ b/Engine/Source/Editor/AnimGraph/Public/AnimGraphNode_AssetPlayerBase.h
@@ -61,6 +61,9 @@ public:
 	// Helper function to validate player nodes
 	void ValidateAnimNodeDuringCompilationHelper(USkeleton* ForSkeleton, FCompilerResultsLog& MessageLog, UAnimationAsset* InAsset, TSubclassOf<UAnimationAsset> InAssetType, UEdGraphPin* InExposedPin, FName InPropertyName);
 
+	// Helper function to preload required assets (including those referenced on pins)
+	void PreloadRequiredAssetsHelper(UAnimationAsset* InAsset, UEdGraphPin* InExposedPin);
+	
 protected:
 	// Helper functions to build a title for an asset player node
 	FText GetNodeTitleHelper(ENodeTitleType::Type InTitleType, UEdGraphPin* InAssetPin, const FText& InAssetDesc, const TFunction<FText(UAnimationAsset*)> InPostFixFunctionRef = nullptr) const;
diff --git a/Engine/Source/Editor/AnimationBlueprintEditor/Private/AnimationBlueprintEditor.cpp b/Engine/Source/Editor/AnimationBlueprintEditor/Private/AnimationBlueprintEditor.cpp
--- a/Engine/Source/Editor/AnimationBlueprintEditor/Private/AnimationBlueprintEditor.cpp
+++ b/Engine/Source/Editor/AnimationBlueprintEditor/Private/AnimationBlueprintEditor.cpp
@@ -422,8 +422,7 @@ void FAnimationBlueprintEditor::InitAnimationBlueprintEditor(const EToolkitMode:
 	PersonaToolkit->GetPreviewScene()->SetDefaultAnimationMode(EPreviewSceneDefaultAnimationMode::AnimationBlueprint);
 	PersonaToolkit->GetPreviewScene()->RegisterOnPreviewMeshChanged(FOnPreviewMeshChanged::CreateSP(this, &FAnimationBlueprintEditor::HandlePreviewMeshChanged));
 
-	TSharedRef<IAssetFamily> AssetFamily = PersonaModule.CreatePersonaAssetFamily(InAnimBlueprint);
-	AssetFamily->RecordAssetOpened(FAssetData(InAnimBlueprint));
+	PersonaModule.RecordAssetOpened(InAnimBlueprint);
 
 	if(InAnimBlueprint->BlueprintType != BPTYPE_Interface && !InAnimBlueprint->bIsTemplate)
 	{
@@ -2269,7 +2268,7 @@ bool FAnimationBlueprintEditor::HandleShouldFilterAsset(
 	
 	const USkeleton* Skeleton = PersonaPreviewSceneDescription->PreviewMesh->GetSkeleton();
 	const FString SkeletonTag = InAssetData.GetTagValueRef<FString>(InTag);
-	if (Skeleton && Skeleton->IsCompatibleSkeletonByAssetString(SkeletonTag))
+	if (Skeleton && Skeleton->IsCompatibleForEditor(SkeletonTag))
 	{
 		return false;
 	}
diff --git a/Engine/Source/Editor/AnimationEditor/Private/AnimationEditor.cpp b/Engine/Source/Editor/AnimationEditor/Private/AnimationEditor.cpp
--- a/Engine/Source/Editor/AnimationEditor/Private/AnimationEditor.cpp
+++ b/Engine/Source/Editor/AnimationEditor/Private/AnimationEditor.cpp
@@ -536,8 +536,7 @@ TSharedPtr<SDockTab> FAnimationEditor::OpenNewAnimationDocumentTab(UAnimationAss
 		}
 
 		// let the asset family know too
-		TSharedRef<IAssetFamily> AssetFamily = PersonaModule.CreatePersonaAssetFamily(InAnimAsset);
-		AssetFamily->RecordAssetOpened(FAssetData(InAnimAsset));
+		PersonaModule.RecordAssetOpened(FAssetData(InAnimAsset));
 	}
 
 	return OpenedTab;
diff --git a/Engine/Source/Editor/ContentBrowser/Public/SAssetView.h b/Engine/Source/Editor/ContentBrowser/Public/SAssetView.h
--- a/Engine/Source/Editor/ContentBrowser/Public/SAssetView.h
+++ b/Engine/Source/Editor/ContentBrowser/Public/SAssetView.h
@@ -294,6 +294,9 @@ public:
 	/** Notifies the asset view that the filter-list filter has changed */
 	void SetBackendFilter(const FARFilter& InBackendFilter);
 
+	/** Get the current backend filter */
+	const FARFilter& GetBackendFilter() const { return BackendFilter; }
+	
 	/** Handler for when a data source requests folder item creation */
 	void NewFolderItemRequested(const FContentBrowserItemTemporaryContext& NewItemContext);
 
diff --git a/Engine/Source/Editor/DetailCustomizations/Private/AnimSequenceDetails.cpp b/Engine/Source/Editor/DetailCustomizations/Private/AnimSequenceDetails.cpp
--- a/Engine/Source/Editor/DetailCustomizations/Private/AnimSequenceDetails.cpp
+++ b/Engine/Source/Editor/DetailCustomizations/Private/AnimSequenceDetails.cpp
@@ -124,12 +124,17 @@ void FAnimSequenceDetails::CustomizeDetails(IDetailLayoutBuilder& DetailBuilder)
 	if (TargetSkeleton.IsValid())
 	{
 		RegisterRetargetSourceChanged();
+
+		// Add each retarget source
+		TArray<FName> RetargetSources;
+		TargetSkeleton->GetRetargetSources(RetargetSources);
+		
 		// go through profile and see if it has mine
-		for (auto Iter = TargetSkeleton->AnimRetargetSources.CreateConstIterator(); Iter; ++Iter)
+		for (FName& RetargetSource : RetargetSources)
 		{
-			RetargetSourceComboList.Add( MakeShareable( new FString ( Iter.Key().ToString() )));
+			RetargetSourceComboList.Add( MakeShareable( new FString ( RetargetSource.ToString() )));
 
-			if (Iter.Key() == CurrentPoseName) 
+			if (RetargetSource == CurrentPoseName) 
 			{
 				InitialSelected = RetargetSourceComboList.Last();
 			}
@@ -613,9 +618,9 @@ void SAnimationRefPoseViewport::Tick( const FGeometry& AllottedGeometry, const d
 		{
 			Description->SetText( FText::Format( LOCTEXT( "Previewing", "Previewing {0}" ), FText::FromString( Component->AnimClass->GetName() ) ) );
 		}
-		else if ( PreviewAnimationSequence && !PreviewAnimationSequence->GetSkeleton()->IsCompatible(TargetSkeleton) )
+		else if ( PreviewAnimationSequence && !PreviewAnimationSequence->GetSkeleton()->IsCompatibleForEditor(TargetSkeleton) )
 		{
-			Description->SetText( FText::Format( LOCTEXT( "IncorrectSkeleton", "The preview asset doesn't work for the skeleton '{0}'" ), FText::FromString( TargetSkeletonName ) ) );
+			Description->SetText( FText::Format( LOCTEXT( "IncorrectSkeleton", "The preview asset is incompatible with the skeleton '{0}'" ), FText::FromString( TargetSkeletonName ) ) );
 		}
 		else if ( Component->GetSkeletalMeshAsset() == NULL )
 		{
diff --git a/Engine/Source/Editor/DetailCustomizations/Private/AnimationAssetDetail.cpp b/Engine/Source/Editor/DetailCustomizations/Private/AnimationAssetDetail.cpp
--- a/Engine/Source/Editor/DetailCustomizations/Private/AnimationAssetDetail.cpp
+++ b/Engine/Source/Editor/DetailCustomizations/Private/AnimationAssetDetail.cpp
@@ -40,7 +40,7 @@ void FAnimationAssetDetails::CustomizeDetails(IDetailLayoutBuilder& DetailBuilde
 	{
 		if (UAnimationAsset* TestAsset = Cast<UAnimationAsset>(SelectionIt->Get()))
 		{
-			if (TargetSkeleton.IsValid() && !TestAsset->GetSkeleton()->IsCompatible(TargetSkeleton.Get()))
+			if (TargetSkeleton.IsValid() && !TestAsset->GetSkeleton()->IsCompatibleForEditor(TargetSkeleton.Get()))
 			{
 				TargetSkeleton = nullptr;
 				break;
@@ -86,7 +86,7 @@ bool FAnimationAssetDetails::ShouldFilterAsset(const FAssetData& AssetData)
 {
 	if (TargetSkeleton.IsValid())
 	{
-		return !TargetSkeleton->IsCompatibleSkeletonByAssetData(AssetData);
+		return !TargetSkeleton->IsCompatibleForEditor(AssetData);
 	}
 
 	return true;
diff --git a/Engine/Source/Editor/DetailCustomizations/Private/PoseAssetDetails.cpp b/Engine/Source/Editor/DetailCustomizations/Private/PoseAssetDetails.cpp
--- a/Engine/Source/Editor/DetailCustomizations/Private/PoseAssetDetails.cpp
+++ b/Engine/Source/Editor/DetailCustomizations/Private/PoseAssetDetails.cpp
@@ -324,7 +324,7 @@ bool FPoseAssetDetails::ShouldFilterAsset(const FAssetData& AssetData)
 {
 	if (TargetSkeleton.IsValid())
 	{
-		return !TargetSkeleton->IsCompatibleSkeletonByAssetData(AssetData);
+		return !TargetSkeleton->IsCompatibleForEditor(AssetData);
 	}
 
 	return true;
@@ -644,7 +644,7 @@ void FPoseAssetDetails::CachePoseAssetData()
 		 SourceAnimationPropertyHandle->GetValue(ObjectSet);
 
 		 UAnimSequence* AnimSequenceSelected = Cast<UAnimSequence>(ObjectSet);
-		 if (AnimSequenceSelected && PoseAsset->GetSkeleton()->IsCompatible(AnimSequenceSelected->GetSkeleton()))
+		 if (AnimSequenceSelected && PoseAsset->GetSkeleton()->IsCompatibleForEditor(AnimSequenceSelected->GetSkeleton()))
 		 {
 			 FScopedTransaction Transaction(LOCTEXT("UpdatePoseSourceAnimation_Transaction", "Update Pose"));
 			 PoseAsset->Modify();
diff --git a/Engine/Source/Editor/DetailCustomizations/Private/SkeletalMeshComponentDetails.cpp b/Engine/Source/Editor/DetailCustomizations/Private/SkeletalMeshComponentDetails.cpp
--- a/Engine/Source/Editor/DetailCustomizations/Private/SkeletalMeshComponentDetails.cpp
+++ b/Engine/Source/Editor/DetailCustomizations/Private/SkeletalMeshComponentDetails.cpp
@@ -258,7 +258,7 @@ EVisibility FSkeletalMeshComponentDetails::VisibilityForAnimationMode(EAnimation
 bool FSkeletalMeshComponentDetails::OnShouldFilterAnimAsset( const FAssetData& AssetData )
 {
 	// Check the compatible skeletons.
-	if (Skeleton && Skeleton->IsCompatibleSkeletonByAssetData(AssetData))
+	if (Skeleton && Skeleton->IsCompatibleForEditor(AssetData))
 	{
 		return false;
 	}
@@ -395,7 +395,7 @@ void FSkeletalMeshComponentDetails::UseSelectedAnimBlueprint()
 		{
 			if(USkeleton* AnimBlueprintSkeleton = AnimBlueprintToAssign->TargetSkeleton)
 			{
-				if (Skeleton && Skeleton->IsCompatible(AnimBlueprintSkeleton))
+				if (Skeleton && Skeleton->IsCompatibleForEditor(AnimBlueprintSkeleton))
 				{
 					OnClassPicked(AnimBlueprintToAssign->GetAnimBlueprintGeneratedClass());
 				}
diff --git a/Engine/Source/Editor/MovieSceneTools/Private/TrackEditors/SkeletalAnimationTrackEditor.cpp b/Engine/Source/Editor/MovieSceneTools/Private/TrackEditors/SkeletalAnimationTrackEditor.cpp
--- a/Engine/Source/Editor/MovieSceneTools/Private/TrackEditors/SkeletalAnimationTrackEditor.cpp
+++ b/Engine/Source/Editor/MovieSceneTools/Private/TrackEditors/SkeletalAnimationTrackEditor.cpp
@@ -509,7 +509,7 @@ public:
 			return true;
 		}
 
-		return !(Skeleton && Skeleton->IsCompatibleSkeletonByAssetData(AssetData));
+		return !(Skeleton && Skeleton->IsCompatibleForEditor(AssetData));
 	}
 
 private:
@@ -1183,7 +1183,7 @@ bool FSkeletalAnimationTrackEditor::HandleAssetAdded(UObject* Asset, const FGuid
 		{
 			USkeleton* Skeleton = AcquireSkeletonFromObjectGuid(TargetObjectGuid, GetSequencer());
 
-			if (Skeleton && Skeleton->IsCompatible(AnimSequence->GetSkeleton()))
+			if (Skeleton && Skeleton->IsCompatibleForEditor(AnimSequence->GetSkeleton()))
 			{
 				UObject* Object = SequencerPtr->FindSpawnedObjectOrTemplate(TargetObjectGuid);
 				
@@ -1793,7 +1793,7 @@ bool FSkeletalAnimationTrackEditor::FilterAnimSequences(const FAssetData& AssetD
 		return true;
 	}
 
-	if (Skeleton && Skeleton->IsCompatibleSkeletonByAssetData(AssetData) == false)
+	if (Skeleton && Skeleton->IsCompatibleForEditor(AssetData) == false)
 	{
 		return true;
 	}
@@ -2006,7 +2006,7 @@ bool FSkeletalAnimationTrackEditor::OnAllowDrop(const FDragDropEvent& DragDropEv
 		for (USkeletalMeshComponent* SkeletalMeshComponent : SkeletalMeshComponents)
 		{
 			USkeleton* Skeleton = GetSkeletonFromComponent(SkeletalMeshComponent);
-			if (bValidAnimSequence && Skeleton && Skeleton->IsCompatible(AnimSequence->GetSkeleton()))
+			if (bValidAnimSequence && Skeleton && Skeleton->IsCompatibleForEditor(AnimSequence->GetSkeleton()))
 			{
 				FFrameRate TickResolution = SequencerPtr->GetFocusedTickResolution();
 				FFrameNumber LengthInFrames = TickResolution.AsFrameNumber(AnimSequence->GetPlayLength());
@@ -2069,7 +2069,7 @@ FReply FSkeletalAnimationTrackEditor::OnDrop(const FDragDropEvent& DragDropEvent
 		{
 			USkeleton* Skeleton = GetSkeletonFromComponent(SkeletalMeshComponent);
 
-			if (bValidAnimSequence && Skeleton && Skeleton->IsCompatible(AnimSequence->GetSkeleton()))
+			if (bValidAnimSequence && Skeleton && Skeleton->IsCompatibleForEditor(AnimSequence->GetSkeleton()))
 			{
 				UObject* BoundObject = SequencerPtr.IsValid() ? SequencerPtr->FindSpawnedObjectOrTemplate(DragDropParams.TargetObjectGuid) : nullptr;
 
diff --git a/Engine/Source/Editor/Persona/Persona.Build.cs b/Engine/Source/Editor/Persona/Persona.Build.cs
--- a/Engine/Source/Editor/Persona/Persona.Build.cs
+++ b/Engine/Source/Editor/Persona/Persona.Build.cs
@@ -26,7 +26,6 @@ public class Persona : ModuleRules
                 "AssetRegistry", 
                 "MainFrame",
                 "DesktopPlatform",
-                "ContentBrowser",
                 "AssetTools",
                 "MeshReductionInterface",
                 "SequenceRecorder",
@@ -77,13 +76,13 @@ public class Persona : ModuleRules
 				"StatusBar",
 				"ToolWidgets",
 				"InteractiveToolsFramework",
+				"ContentBrowser",
             }
         );
 
         DynamicallyLoadedModuleNames.AddRange(
             new string[] {
-                "ContentBrowser",
-                "Documentation",
+	            "Documentation",
                 "MainFrame",
                 "DesktopPlatform",
                 "SkeletonEditor",
diff --git a/Engine/Source/Editor/Persona/Private/AnimGraphNodeDetails.cpp b/Engine/Source/Editor/Persona/Private/AnimGraphNodeDetails.cpp
--- a/Engine/Source/Editor/Persona/Private/AnimGraphNodeDetails.cpp
+++ b/Engine/Source/Editor/Persona/Private/AnimGraphNodeDetails.cpp
@@ -275,7 +275,7 @@ TSharedRef<SWidget> FAnimGraphNodeDetails::CreatePropertyWidget(FProperty* Targe
 bool FAnimGraphNodeDetails::OnShouldFilterAnimAsset( const FAssetData& AssetData, UClass* NodeToFilterFor ) const
 {
 	FAssetDataTagMapSharedView::FFindTagResult Result = AssetData.TagsAndValues.FindTag("Skeleton");
-	if (Result.IsSet() && TargetSkeleton && TargetSkeleton->IsCompatibleSkeletonByAssetData(AssetData))
+	if (Result.IsSet() && TargetSkeleton && TargetSkeleton->IsCompatibleForEditor(AssetData))
 	{
 		const UClass* AssetClass = AssetData.GetClass();
 		// If node is an 'asset player', only let you select the right kind of asset for it
@@ -1073,7 +1073,7 @@ TSharedRef<SWidget> SParentPlayerTreeRow::GenerateWidgetForColumn(const FName& C
 bool SParentPlayerTreeRow::OnShouldFilterAsset(const FAssetData& AssetData)
 {
 	const USkeleton* CurrentSkeleton = CastChecked<UAnimBlueprint>(BlueprintEditor.Pin()->GetBlueprintObj())->TargetSkeleton;
-	return CurrentSkeleton != nullptr && !CurrentSkeleton->IsCompatibleSkeletonByAssetData(AssetData);
+	return CurrentSkeleton != nullptr && !CurrentSkeleton->IsCompatibleForEditor(AssetData);
 }
 
 void SParentPlayerTreeRow::OnAssetSelected(const FAssetData& AssetData)
diff --git a/Engine/Source/Editor/Persona/Private/AnimInstanceDetails.cpp b/Engine/Source/Editor/Persona/Private/AnimInstanceDetails.cpp
--- a/Engine/Source/Editor/Persona/Private/AnimInstanceDetails.cpp
+++ b/Engine/Source/Editor/Persona/Private/AnimInstanceDetails.cpp
@@ -114,5 +114,5 @@ TSharedRef<SWidget> FAnimInstanceDetails::CreateFilteredObjectPropertyWidget(FPr
 
 bool FAnimInstanceDetails::OnShouldFilterAnimAsset(const FAssetData& AssetData, const USkeleton* TargetSkeleton) const
 {
-	return !(TargetSkeleton && TargetSkeleton->IsCompatibleSkeletonByAssetData(AssetData));
+	return !(TargetSkeleton && TargetSkeleton->IsCompatibleForEditor(AssetData));
 }
diff --git a/Engine/Source/Editor/Persona/Private/AnimationEditorPreviewScene.cpp b/Engine/Source/Editor/Persona/Private/AnimationEditorPreviewScene.cpp
--- a/Engine/Source/Editor/Persona/Private/AnimationEditorPreviewScene.cpp
+++ b/Engine/Source/Editor/Persona/Private/AnimationEditorPreviewScene.cpp
@@ -159,7 +159,7 @@ void FAnimationEditorPreviewScene::SetPreviewMesh(USkeletalMesh* NewPreviewMesh,
 	if (NewPreviewMesh != nullptr && GetEditableSkeleton().IsValid() && !GetEditableSkeleton()->GetSkeleton().IsCompatibleMesh(NewPreviewMesh))
 	{
 		const USkeleton& Skeleton = GetEditableSkeleton()->GetSkeleton();
-		if (NewPreviewMesh->GetSkeleton() && Skeleton.IsCompatible(NewPreviewMesh->GetSkeleton()))
+		if (NewPreviewMesh->GetSkeleton() && Skeleton.IsCompatibleForEditor(NewPreviewMesh->GetSkeleton()))
 		{
 			SetPreviewMeshInternal(NewPreviewMesh);
 		}	
@@ -649,7 +649,7 @@ void FAnimationEditorPreviewScene::SetPreviewAnimationAsset(UAnimationAsset* Ani
 			}
 
 			// Treat it as invalid if it's got a bogus skeleton pointer
-			if (Skeleton != nullptr && !Skeleton->IsCompatible(AnimAsset->GetSkeleton()))
+			if (AnimAsset->GetSkeleton() == nullptr)
 			{
 				return;
 			}
diff --git a/Engine/Source/Editor/Persona/Private/Customization/BlendSampleDetails.cpp b/Engine/Source/Editor/Persona/Private/Customization/BlendSampleDetails.cpp
--- a/Engine/Source/Editor/Persona/Private/Customization/BlendSampleDetails.cpp
+++ b/Engine/Source/Editor/Persona/Private/Customization/BlendSampleDetails.cpp
@@ -547,7 +547,7 @@ bool FBlendSampleDetails::ShouldFilterAssetStatic(const FAssetData& AssetData, c
 	if (AssetData.GetTagValue(SkeletonTagName, SkeletonName))
 	{
 		// Check whether or not the skeletons are compatible
-		if (BlendSpaceBase->GetSkeleton()->IsCompatibleSkeletonByAssetData(AssetData))
+		if (BlendSpaceBase->GetSkeleton()->IsCompatibleForEditor(AssetData))
 		{
 			// If so check if the additive animation type is compatible with the blend space
 			const FName AdditiveTypeTagName = GET_MEMBER_NAME_CHECKED(UAnimSequence, AdditiveAnimType);
diff --git a/Engine/Source/Editor/Persona/Private/PersonaAssetFamily.cpp b/Engine/Source/Editor/Persona/Private/PersonaAssetFamily.cpp
--- a/Engine/Source/Editor/Persona/Private/PersonaAssetFamily.cpp
+++ b/Engine/Source/Editor/Persona/Private/PersonaAssetFamily.cpp
@@ -8,6 +8,7 @@
 #include "PhysicsEngine/PhysicsAsset.h"
 #include "Styling/AppStyle.h"
 #include "AssetToolsModule.h"
+#include "Preferences/PersonaOptions.h"
 
 #define LOCTEXT_NAMESPACE "PersonaAssetFamily"
 
@@ -44,6 +45,11 @@ FPersonaAssetFamily::FPersonaAssetFamily(const UObject* InFromObject)
 	}
 }
 
+void FPersonaAssetFamily::Initialize()
+{
+	GetMutableDefault<UPersonaOptions>()->RegisterOnUpdateSettings(UPersonaOptions::FOnUpdateSettingsMulticaster::FDelegate::CreateSP(this, &FPersonaAssetFamily::OnSettingsChange));
+}
+
 void FPersonaAssetFamily::GetAssetTypes(TArray<UClass*>& OutAssetTypes) const
 {
 	OutAssetTypes.Reset();
@@ -61,21 +67,8 @@ static void FindAssets(const USkeleton* InSkeleton, TArray<FAssetData>& OutAsset
 	{
 		return;
 	}
-	FAssetRegistryModule& AssetRegistryModule = FModuleManager::LoadModuleChecked<FAssetRegistryModule>("AssetRegistry");
-	FARFilter Filter;
-	Filter.bRecursiveClasses = true;
-	Filter.ClassPaths.Add(AssetType::StaticClass()->GetClassPathName());
-	Filter.TagsAndValues.Add(SkeletonTag, FAssetData(InSkeleton).GetExportTextName());
-	
-	// Also include all compatible assets.
-	FString CompatibleTagValue;
-	for (const auto& CompatibleSkeleton : InSkeleton->GetCompatibleSkeletons())
-	{
-		CompatibleTagValue = FString::Format(TEXT("{0}'{1}'"), { *USkeleton::StaticClass()->GetPathName(), *CompatibleSkeleton.ToString() });
-		Filter.TagsAndValues.Add(SkeletonTag, CompatibleTagValue);
-	}
 
-	AssetRegistryModule.Get().GetAssets(Filter, OutAssetData);
+	InSkeleton->GetCompatibleAssets(AssetType::StaticClass(), *SkeletonTag.ToString(), OutAssetData);
 }
 
 FAssetData FPersonaAssetFamily::FindAssetOfType(UClass* InAssetClass) const
@@ -173,8 +166,7 @@ void FPersonaAssetFamily::FindAssetsOfType(UClass* InAssetClass, TArray<FAssetDa
 	{
 		if (InAssetClass->IsChildOf<USkeleton>())
 		{
-			// we should always have a skeleton here, this asset family is based on it
-			OutAssets.Add(FAssetData(Skeleton.Get()));
+			Skeleton.Get()->GetCompatibleSkeletonAssets(OutAssets);
 		}
 		else if (InAssetClass->IsChildOf<UAnimationAsset>())
 		{
@@ -311,7 +303,7 @@ bool FPersonaAssetFamily::IsAssetCompatible(const FAssetData& InAssetData) const
 		{
 			if (Skeleton.Get())
 			{
-				return Skeleton.Get()->IsCompatibleSkeletonByAssetData(InAssetData);
+				return Skeleton.Get()->IsCompatibleForEditor(InAssetData);
 			}
 		}
 		else if (Class->IsChildOf<UAnimationAsset>() || Class->IsChildOf<USkeletalMesh>())
@@ -322,7 +314,7 @@ bool FPersonaAssetFamily::IsAssetCompatible(const FAssetData& InAssetData) const
 			{
 				if (Skeleton.Get())
 				{
-					return Skeleton.Get()->IsCompatibleSkeletonByAssetData(InAssetData);
+					return Skeleton.Get()->IsCompatibleForEditor(Result.GetValue());
 				}
 			}
 		}
@@ -334,7 +326,7 @@ bool FPersonaAssetFamily::IsAssetCompatible(const FAssetData& InAssetData) const
 			{
 				if (Skeleton.Get())
 				{
-					return Skeleton.Get()->IsCompatibleSkeletonByAssetString(Result.GetValue());
+					return Skeleton.Get()->IsCompatibleForEditor(Result.GetValue());
 				}
 			}
 		}
@@ -387,11 +379,7 @@ void FPersonaAssetFamily::RecordAssetOpened(const FAssetData& InAssetData)
 		UClass* Class = InAssetData.GetClass();
 		if (Class)
 		{
-			if (Class->IsChildOf<USkeleton>())
-			{
-				Skeleton = Cast<USkeleton>(InAssetData.GetAsset());
-			}
-			else if (Class->IsChildOf<UAnimationAsset>())
+			if (Class->IsChildOf<UAnimationAsset>())
 			{
 				AnimationAsset = Cast<UAnimationAsset>(InAssetData.GetAsset());
 			}
@@ -478,5 +466,10 @@ void FPersonaAssetFamily::FindCounterpartAssets(const UObject* InAsset, const US
 	}
 }
 
+void FPersonaAssetFamily::OnSettingsChange(const UPersonaOptions* InOptions, EPropertyChangeType::Type InChangeType)
+{
+	OnAssetFamilyChanged.Broadcast();
+}
+
 #undef LOCTEXT_NAMESPACE
  
diff --git a/Engine/Source/Editor/Persona/Private/PersonaAssetFamily.h b/Engine/Source/Editor/Persona/Private/PersonaAssetFamily.h
--- a/Engine/Source/Editor/Persona/Private/PersonaAssetFamily.h
+++ b/Engine/Source/Editor/Persona/Private/PersonaAssetFamily.h
@@ -11,8 +11,9 @@
 
 class UAnimBlueprint;
 class UPhysicsAsset;
+class UPersonaOptions;
 
-class FPersonaAssetFamily : public IAssetFamily
+class FPersonaAssetFamily : public IAssetFamily, public TSharedFromThis<FPersonaAssetFamily>
 {
 public:
 	FPersonaAssetFamily(const UObject* InFromObject);
@@ -30,12 +31,23 @@ public:
 	virtual void RecordAssetOpened(const FAssetData& InAssetData) override;
 	DECLARE_DERIVED_EVENT(FPersonaAssetFamily, IAssetFamily::FOnAssetOpened, FOnAssetOpened)
 	virtual FOnAssetOpened& GetOnAssetOpened() override { return OnAssetOpened;  }
+	DECLARE_DERIVED_EVENT(FPersonaAssetFamily, IAssetFamily::FOnAssetFamilyChanged, FOnAssetFamilyChanged)
+	virtual FOnAssetFamilyChanged& GetOnAssetFamilyChanged() override { return OnAssetFamilyChanged; }
 
 	/** Helper functions for constructor and other systems that need to discover meshes/skeletons from related assets */
 	static void FindCounterpartAssets(const UObject* InAsset, TWeakObjectPtr<const USkeleton>& OutSkeleton, TWeakObjectPtr<const USkeletalMesh>& OutMesh);
 	static void FindCounterpartAssets(const UObject* InAsset, const USkeleton*& OutSkeleton, const USkeletalMesh*& OutMesh);
 
 private:
+	/** Initialization to avoid shared ptr access in constructor */
+	void Initialize();
+	
+	/** Hande key persona settings changes (e.g. skeleton compatibilty) */
+	void OnSettingsChange(const UPersonaOptions* InOptions, EPropertyChangeType::Type InChangeType);
+	
+private:
+	friend class FPersonaAssetFamilyManager;
+	
 	/** The skeleton that links all assets */
 	TWeakObjectPtr<const USkeleton> Skeleton;
 
@@ -53,4 +65,7 @@ private:
 
 	/** Event fired when an asset is opened */
 	FOnAssetOpened OnAssetOpened;
+
+	/** Event fired when an asset family changes (e.g. relationships are altered) */
+	FOnAssetFamilyChanged OnAssetFamilyChanged;
 };
diff --git a/Engine/Source/Editor/Persona/Private/PersonaAssetFamilyManager.cpp b/Engine/Source/Editor/Persona/Private/PersonaAssetFamilyManager.cpp
--- a/Engine/Source/Editor/Persona/Private/PersonaAssetFamilyManager.cpp
+++ b/Engine/Source/Editor/Persona/Private/PersonaAssetFamilyManager.cpp
@@ -14,20 +14,40 @@ FPersonaAssetFamilyManager& FPersonaAssetFamilyManager::Get()
 TSharedRef<IAssetFamily> FPersonaAssetFamilyManager::CreatePersonaAssetFamily(const UObject* InAsset)
 {
 	// compact any invalid entries
-	AssetFamilies.RemoveAll([](const TWeakPtr<class IAssetFamily>& InAssetFamily) { return !InAssetFamily.IsValid(); });
+	AssetFamilies.RemoveAll([](const TWeakPtr<IAssetFamily>& InAssetFamily) { return !InAssetFamily.IsValid(); });
 
-	// look for an existing matching asset family
-	FAssetData AssetData(InAsset);
-	for (TWeakPtr<class IAssetFamily>& AssetFamily : AssetFamilies)
+	// Create new asset family
+	TSharedRef<FPersonaAssetFamily> NewAssetFamily = MakeShared<FPersonaAssetFamily>(InAsset);
+	NewAssetFamily->Initialize();
+	AssetFamilies.Add(NewAssetFamily);
+	return NewAssetFamily;
+}
+
+void FPersonaAssetFamilyManager::BroadcastAssetFamilyChange()
+{
+	// Create copy as delegate can modify the AssetFamilies array
+	TArray<TWeakPtr<IAssetFamily>> AssetFamiliesCopy = AssetFamilies;
+	for (TWeakPtr<IAssetFamily>& AssetFamily : AssetFamiliesCopy)
 	{
-		if (AssetFamily.Pin()->IsAssetCompatible(AssetData))
+		if (TSharedPtr<IAssetFamily> PinnedAssetFamily = AssetFamily.Pin())
 		{
-			return AssetFamily.Pin().ToSharedRef();
+			PinnedAssetFamily->GetOnAssetFamilyChanged().Broadcast();
 		}
 	}
-
-	// not found - make a new one
-	TSharedRef<IAssetFamily> NewAssetFamily = MakeShareable(new FPersonaAssetFamily(InAsset));
-	AssetFamilies.Add(NewAssetFamily);
-	return NewAssetFamily;
 }
+
+void FPersonaAssetFamilyManager::RecordAssetOpened(const FAssetData& InAssetData) const
+{
+	for (const TWeakPtr<IAssetFamily>& AssetFamily : AssetFamilies)
+	{
+		if (TSharedPtr<IAssetFamily> PinnedAssetFamily = AssetFamily.Pin())
+		{
+			if(PinnedAssetFamily->IsAssetCompatible(InAssetData))
+			{
+				PRAGMA_DISABLE_DEPRECATION_WARNINGS
+				PinnedAssetFamily->RecordAssetOpened(InAssetData);
+				PRAGMA_ENABLE_DEPRECATION_WARNINGS
+			}
+		}
+	}
+}
\ No newline at end of file
diff --git a/Engine/Source/Editor/Persona/Private/PersonaAssetFamilyManager.h b/Engine/Source/Editor/Persona/Private/PersonaAssetFamilyManager.h
--- a/Engine/Source/Editor/Persona/Private/PersonaAssetFamilyManager.h
+++ b/Engine/Source/Editor/Persona/Private/PersonaAssetFamilyManager.h
@@ -5,6 +5,7 @@
 #include "CoreMinimal.h"
 
 class IAssetFamily;
+struct FAssetData;
 
 /** Central registry of persona asset families */
 class FPersonaAssetFamilyManager
@@ -13,14 +14,20 @@ public:
 	/** Singleton access */
 	static FPersonaAssetFamilyManager& Get();
 
-	/** Create a skeleton tree for the requested skeleton */
-	TSharedRef<class IAssetFamily> CreatePersonaAssetFamily(const UObject* InAsset);
+	/** Create an asset family using the supplied asset */
+	TSharedRef<IAssetFamily> CreatePersonaAssetFamily(const UObject* InAsset);
 
+	/** All asset families have changed (e.g. universal compatibility) */
+	void BroadcastAssetFamilyChange();
+
+	/** Record an asset being opened - forward to all compatible asset families */
+	void RecordAssetOpened(const FAssetData& InAssetData) const;
+	
 private:
 	/** Hidden constructor */
 	FPersonaAssetFamilyManager() {}
 
 private:
 	/** All current asset families */
-	TArray<TWeakPtr<class IAssetFamily>> AssetFamilies;
+	TArray<TWeakPtr<IAssetFamily>> AssetFamilies;
 };
diff --git a/Engine/Source/Editor/Persona/Private/PersonaMeshDetails.cpp b/Engine/Source/Editor/Persona/Private/PersonaMeshDetails.cpp
--- a/Engine/Source/Editor/Persona/Private/PersonaMeshDetails.cpp
+++ b/Engine/Source/Editor/Persona/Private/PersonaMeshDetails.cpp
@@ -4298,7 +4298,7 @@ void FPersonaMeshDetails::CustomizeDetails( IDetailLayoutBuilder& DetailLayout )
 			{
 				if(UAnimBlueprint* AnimBlueprint = Cast<UAnimBlueprint>(PostProcessAnimBlueprintClass->ClassGeneratedBy))
 				{
-					return !SkeletalMeshPtr->GetSkeleton()->IsCompatible(AnimBlueprint->TargetSkeleton) ? EVisibility::Visible : EVisibility::Collapsed;
+					return !SkeletalMeshPtr->GetSkeleton()->IsCompatibleForEditor(AnimBlueprint->TargetSkeleton) ? EVisibility::Visible : EVisibility::Collapsed;
 				}
 			}
 		}
@@ -4431,7 +4431,7 @@ bool FPersonaMeshDetails::OnShouldFilterPostProcessBlueprint(const FAssetData& A
 		if (!SkelMesh->IsCompiling())
 		{
 			const FString SkeletonName = AssetData.GetTagValueRef<FString>("TargetSkeleton");
-			return !SkelMesh->GetSkeleton()->IsCompatibleSkeletonByAssetString(SkeletonName);
+			return !SkelMesh->GetSkeleton()->IsCompatibleForEditor(SkeletonName);
 		}
 	}
 
@@ -6564,6 +6564,6 @@ void FPersonaMeshDetails::OnPreviewMeshChanged(USkeletalMesh* OldSkeletalMesh, U
 
 bool FPersonaMeshDetails::FilterOutBakePose(const FAssetData& AssetData, USkeleton* Skeleton) const
 {
-	return !(Skeleton && Skeleton->IsCompatibleSkeletonByAssetData(AssetData));
+	return !(Skeleton && Skeleton->IsCompatibleForEditor(AssetData));
 }
 #undef LOCTEXT_NAMESPACE
\ No newline at end of file
diff --git a/Engine/Source/Editor/Persona/Private/PersonaModule.cpp b/Engine/Source/Editor/Persona/Private/PersonaModule.cpp
--- a/Engine/Source/Editor/Persona/Private/PersonaModule.cpp
+++ b/Engine/Source/Editor/Persona/Private/PersonaModule.cpp
@@ -97,6 +97,9 @@
 #include "SkeletalMeshReferenceSectionDetails.h"
 #include "PersonaToolMenuContext.h"
 #include "ToolMenu.h"
+#include "ContentBrowserMenuContexts.h"
+#include "SAssetView.h"
+#include "ToolMenus.h"
 
 IMPLEMENT_MODULE( FPersonaModule, Persona );
 
@@ -154,6 +157,8 @@ void FPersonaModule::StartupModule()
 	FKismetEditorUtilities::RegisterOnBlueprintCreatedCallback(this, UAnimNotifyState::StaticClass(), FKismetEditorUtilities::FOnBlueprintCreated::CreateRaw(this, &FPersonaModule::HandleNewAnimNotifyStateBlueprintCreated));
 
 	GetMutableDefault<UPersonaOptions>()->SetFlags(RF_Transactional);
+
+	RegisterToolMenuExtensions();
 }
 
 void FPersonaModule::ShutdownModule()
@@ -234,6 +239,16 @@ TSharedRef<class IAssetFamily> FPersonaModule::CreatePersonaAssetFamily(const UO
 	return FPersonaAssetFamilyManager::Get().CreatePersonaAssetFamily(InAsset);
 }
 
+void FPersonaModule::BroadcastAssetFamilyChange() const
+{
+	return FPersonaAssetFamilyManager::Get().BroadcastAssetFamilyChange();
+}
+
+void FPersonaModule::RecordAssetOpened(const FAssetData& InAssetData) const
+{
+	return FPersonaAssetFamilyManager::Get().RecordAssetOpened(InAssetData);
+}
+
 TSharedRef<SWidget> FPersonaModule::CreateAssetFamilyShortcutWidget(const TSharedRef<class FWorkflowCentricApplication>& InHostingApp, const TSharedRef<class IAssetFamily>& InAssetFamily) const
 {
 	return SNew(SAssetFamilyShortcutBar, InHostingApp, InAssetFamily);
@@ -927,7 +942,7 @@ namespace UE::Persona::Private
 					FString TagValue;
 					if (AssetData.GetTagValue("Skeleton", TagValue))
 					{
-						return !InWeakPersonaToolkit.Pin()->GetSkeleton()->IsCompatibleSkeletonByAssetData(AssetData);
+						return !InWeakPersonaToolkit.Pin()->GetSkeleton()->IsCompatibleForEditor(AssetData);
 					}
 				}
 				return true;
@@ -988,7 +1003,7 @@ namespace UE::Persona::Private
 					FString TagValue;
 					if (AssetData.GetTagValue("Skeleton", TagValue))
 					{
-						return !InWeakPersonaToolkit.Pin()->GetSkeleton()->IsCompatibleSkeletonByAssetData(AssetData);
+						return !InWeakPersonaToolkit.Pin()->GetSkeleton()->IsCompatibleForEditor(AssetData);
 					}
 				}
 				return true;
@@ -1648,4 +1663,74 @@ TSharedRef<SWidget> FPersonaModule::CreateBlendSpaceEditWidget(UBlendSpace* InBl
 		.StatusBarName(InArgs.StatusBarName);
 }
 
+void FPersonaModule::RegisterToolMenuExtensions()
+{
+	// Register content browser menu extensions
+	UToolMenu* ToolMenu = UToolMenus::Get()->ExtendMenu("ContentBrowser.AssetViewOptions");
+	ToolMenu->AddDynamicSection("AnimationAssets", FNewToolMenuDelegate::CreateLambda([](UToolMenu* InMenu)
+	{
+		if(UContentBrowserAssetViewContextMenuContext* MenuContext = InMenu->Context.FindContext<UContentBrowserAssetViewContextMenuContext>())
+		{
+			if(TSharedPtr<SAssetView> AssetView = MenuContext->AssetView.Pin())
+			{
+				// Check whether this asset view is showing animation-related assets
+				TArray<FTopLevelAssetPath> BaseAnimAssetClassPaths =
+				{
+					UAnimationAsset::StaticClass()->GetClassPathName(),
+					UAnimBlueprint::StaticClass()->GetClassPathName(),
+					USkeletalMesh::StaticClass()->GetClassPathName(),
+					USkeleton::StaticClass()->GetClassPathName()
+				};
+				TSet<FTopLevelAssetPath> AnimAssetClassPaths;
+				IAssetRegistry& AssetRegistry = FModuleManager::LoadModuleChecked<FAssetRegistryModule>(TEXT("AssetRegistry")).Get();
+				AssetRegistry.GetDerivedClassNames(BaseAnimAssetClassPaths, TSet<FTopLevelAssetPath>(), AnimAssetClassPaths);
+
+				auto AssetClassMatches = [&AnimAssetClassPaths](const FTopLevelAssetPath& InPath)
+				{
+					for(const FTopLevelAssetPath& Path : AnimAssetClassPaths)
+					{
+						if(InPath == Path)
+						{
+							return true;
+						}
+					}
+					return false;
+				};
+
+				// Add a menu item to allow incompatible skeletons to be selected if this is selecting animation assets
+				const FARFilter& BackendFilter = AssetView->GetBackendFilter();
+				if(BackendFilter.ClassPaths.ContainsByPredicate(AssetClassMatches))
+				{
+					FToolMenuSection& Section = InMenu->AddSection("AnimationAssets", LOCTEXT("AnimationAssetsSection", "Animation Assets"));
+					Section.AddMenuEntry(
+						"AllowIncompatibleSkeletons",
+						LOCTEXT("AllowIncompatibleSkeletons", "Allow Incompatible Skeletons"),
+						LOCTEXT("AllowIncompatibleSkeletonsTooltip", "Whether to allow animation assets that are incompatible with the current skeleton/skeletal mesh to be selected."),
+						FSlateIcon(),
+						FToolUIActionChoice(
+							FUIAction(
+								FExecuteAction::CreateLambda([WeakAssetView = TWeakPtr<SAssetView>(AssetView)]()
+								{
+									UPersonaOptions* PersonaOptions = GetMutableDefault<UPersonaOptions>();
+									PersonaOptions->bAllowIncompatibleSkeletonSelection = !PersonaOptions->bAllowIncompatibleSkeletonSelection;
+									if(TSharedPtr<SAssetView> AssetView = WeakAssetView.Pin())
+									{
+										AssetView->RequestSlowFullListRefresh();
+									}
+								}),
+								FCanExecuteAction(),
+								FIsActionChecked::CreateLambda([]()
+								{
+									return GetDefault<UPersonaOptions>()->bAllowIncompatibleSkeletonSelection;
+								})
+							)
+						),
+						EUserInterfaceActionType::ToggleButton
+					);
+				}
+			}
+		}
+	}));
+}
+
 #undef LOCTEXT_NAMESPACE
diff --git a/Engine/Source/Editor/Persona/Private/PersonaPreviewSceneAnimationController.cpp b/Engine/Source/Editor/Persona/Private/PersonaPreviewSceneAnimationController.cpp
--- a/Engine/Source/Editor/Persona/Private/PersonaPreviewSceneAnimationController.cpp
+++ b/Engine/Source/Editor/Persona/Private/PersonaPreviewSceneAnimationController.cpp
@@ -51,7 +51,7 @@ IDetailPropertyRow* UPersonaPreviewSceneAnimationController::AddPreviewControlle
 
 bool UPersonaPreviewSceneAnimationController::HandleShouldFilterAsset(const FAssetData& InAssetData, const USkeleton* InSkeleton) const
 {
-	if (InSkeleton && InSkeleton->IsCompatibleSkeletonByAssetData(InAssetData))
+	if (InSkeleton && InSkeleton->IsCompatibleForEditor(InAssetData))
 	{
 		return false;
 	}
diff --git a/Engine/Source/Editor/Persona/Private/PersonaToolkit.cpp b/Engine/Source/Editor/Persona/Private/PersonaToolkit.cpp
--- a/Engine/Source/Editor/Persona/Private/PersonaToolkit.cpp
+++ b/Engine/Source/Editor/Persona/Private/PersonaToolkit.cpp
@@ -240,7 +240,7 @@ void FPersonaToolkit::SetAnimationAsset(class UAnimationAsset* InAnimationAsset)
 {
 	if (InAnimationAsset != nullptr)
 	{
-		check(Skeleton->IsCompatible(InAnimationAsset->GetSkeleton()));
+		ensure(Skeleton->IsCompatibleForEditor(InAnimationAsset->GetSkeleton()));
 	}
 
 	AnimationAsset = InAnimationAsset;
diff --git a/Engine/Source/Editor/Persona/Private/PreviewSceneCustomizations.cpp b/Engine/Source/Editor/Persona/Private/PreviewSceneCustomizations.cpp
--- a/Engine/Source/Editor/Persona/Private/PreviewSceneCustomizations.cpp
+++ b/Engine/Source/Editor/Persona/Private/PreviewSceneCustomizations.cpp
@@ -482,7 +482,7 @@ bool FPreviewSceneDescriptionCustomization::HandleShouldFilterAsset(const FAsset
 
 	const USkeleton* Skeleton = PersonaPreviewSceneDescription->PreviewMesh->GetSkeleton();
 	const FString SkeletonTag = InAssetData.GetTagValueRef<FString>(InTag);
-	if (Skeleton && Skeleton->IsCompatibleSkeletonByAssetString(SkeletonTag))
+	if (Skeleton && Skeleton->IsCompatibleForEditor(SkeletonTag))
 	{
 		return false;
 	}
@@ -688,7 +688,7 @@ bool FPreviewMeshCollectionEntryCustomization::HandleShouldFilterAsset(const FAs
 		return false;
 	}
 
-	if (Skeleton && Skeleton->IsCompatibleSkeletonByAssetData(InAssetData))
+	if (Skeleton && Skeleton->IsCompatibleForEditor(InAssetData))
 	{
 		return false;
 	}
diff --git a/Engine/Source/Editor/Persona/Private/SAnimSegmentsPanel.cpp b/Engine/Source/Editor/Persona/Private/SAnimSegmentsPanel.cpp
--- a/Engine/Source/Editor/Persona/Private/SAnimSegmentsPanel.cpp
+++ b/Engine/Source/Editor/Persona/Private/SAnimSegmentsPanel.cpp
@@ -413,17 +413,16 @@ void SAnimSegmentsPanel::FillSubMenu(FMenuBuilder& MenuBuilder, int32 AnimSegmen
 			AssetPickerConfig.Filter.bRecursiveClasses = false;
 			AssetPickerConfig.bAllowNullSelection = false;
 
-			USkeleton* Skeleton = OldSequenceBase->GetSkeleton();
-			AssetPickerConfig.Filter.TagsAndValues.Add(TEXT("Skeleton"), FAssetData(Skeleton).GetExportTextName());
-
 			// only do this for anim sequence because we don't know additive or not otherwise from asset registry
+			TEnumAsByte<EAdditiveAnimationType> AdditiveType;
 			bool bFilterAdditive = OldSequenceBase->GetClass() == UAnimSequence::StaticClass();
-			if (bFilterAdditive)
+			if(bFilterAdditive)
 			{
-				// we do just check additiveanimtype, not IsValidAdditive because we only check asset registry string, we just assume this only checks additive anim type
-				// in order to check IsValidAdditive, we have to load all animations, which is too slow
-				AssetPickerConfig.OnShouldFilterAsset = FOnShouldFilterAsset::CreateRaw(this, &SAnimSegmentsPanel::ShouldFilter, CastChecked<UAnimSequence>(OldSequenceBase)->AdditiveAnimType);
+				AdditiveType = CastChecked<UAnimSequence>(OldSequenceBase)->AdditiveAnimType;
 			}
+
+			AssetPickerConfig.OnShouldFilterAsset = FOnShouldFilterAsset::CreateRaw(this, &SAnimSegmentsPanel::ShouldFilter, OldSequenceBase->GetSkeleton(), bFilterAdditive, AdditiveType);
+
 			/** The delegate that fires when an asset was selected */
 			AssetPickerConfig.OnAssetSelected = FOnAssetSelected::CreateRaw(this, &SAnimSegmentsPanel::ReplaceAnimSegment, AnimSegmentIndex);
 
@@ -445,12 +444,20 @@ void SAnimSegmentsPanel::FillSubMenu(FMenuBuilder& MenuBuilder, int32 AnimSegmen
 	}
 }
 
-bool SAnimSegmentsPanel::ShouldFilter(const FAssetData& DataToDisplay, TEnumAsByte<EAdditiveAnimationType> InAdditiveType)
+bool SAnimSegmentsPanel::ShouldFilter(const FAssetData& DataToDisplay, USkeleton* InSkeleton, bool bInFilterAdditive, TEnumAsByte<EAdditiveAnimationType> InAdditiveType)
 {
-	UEnum* AdditiveTypeEnum = StaticEnum<EAdditiveAnimationType>();
-	const FString EnumString = DataToDisplay.GetTagValueRef<FString>(GET_MEMBER_NAME_CHECKED(UAnimSequence, AdditiveAnimType));
-	EAdditiveAnimationType AdditiveType = (!EnumString.IsEmpty() ? (EAdditiveAnimationType)AdditiveTypeEnum->GetValueByName(*EnumString) : AAT_None);
-	return (AdditiveType != InAdditiveType);
+	bool bFilter = false;
+	if(bInFilterAdditive)
+	{
+		UEnum* AdditiveTypeEnum = StaticEnum<EAdditiveAnimationType>();
+		const FString EnumString = DataToDisplay.GetTagValueRef<FString>(GET_MEMBER_NAME_CHECKED(UAnimSequence, AdditiveAnimType));
+		EAdditiveAnimationType AdditiveType = (!EnumString.IsEmpty() ? (EAdditiveAnimationType)AdditiveTypeEnum->GetValueByName(*EnumString) : AAT_None);
+		bFilter |= (AdditiveType != InAdditiveType);
+	}
+
+	bFilter |= InSkeleton->ShouldFilterAsset(DataToDisplay);
+
+	return bFilter;
 }
 
 void SAnimSegmentsPanel::OnTrackDragDrop( TSharedPtr<FDragDropOperation> DragDropOp, float DataPos )
diff --git a/Engine/Source/Editor/Persona/Private/SAnimSegmentsPanel.h b/Engine/Source/Editor/Persona/Private/SAnimSegmentsPanel.h
--- a/Engine/Source/Editor/Persona/Private/SAnimSegmentsPanel.h
+++ b/Engine/Source/Editor/Persona/Private/SAnimSegmentsPanel.h
@@ -18,6 +18,7 @@ struct FAssetData;
 class FMenuBuilder;
 class UAnimSequenceBase;
 struct FAnimSegment;
+class USkeleton;
 
 DECLARE_DELEGATE( FOnPreAnimUpdate )
 DECLARE_DELEGATE( FOnPostAnimUpdate )
@@ -135,7 +136,7 @@ private:
 	void				ReplaceAnimSegment(UAnimSequenceBase* NewSequenceBase, float NewStartPos);
 	void				ReplaceAnimSegment(int32 AnimSegmentIndex, UAnimSequenceBase* NewSequenceBase);
 	void				ReplaceAnimSegment(const FAssetData& NewSequenceData, int32 AnimSegmentIndex);
-	bool				ShouldFilter(const FAssetData& DataToDisplay, TEnumAsByte<EAdditiveAnimationType> InAdditiveType);
+	bool				ShouldFilter(const FAssetData& DataToDisplay, USkeleton* InSkeleton, bool bInFilterAdditive, TEnumAsByte<EAdditiveAnimationType> InAdditiveType);
 
 	// Remove all selected anim segments in all segment tracks
 	void RemoveSelectedAnimSegments();
diff --git a/Engine/Source/Editor/Persona/Private/SAnimViewportToolBar.cpp b/Engine/Source/Editor/Persona/Private/SAnimViewportToolBar.cpp
--- a/Engine/Source/Editor/Persona/Private/SAnimViewportToolBar.cpp
+++ b/Engine/Source/Editor/Persona/Private/SAnimViewportToolBar.cpp
@@ -881,7 +881,7 @@ void SAnimViewportToolBar::FillCharacterMirrorMenu(FMenuBuilder& MenuBuilder) co
 		AssetPickerConfig.Filter.ClassPaths.Add(UMirrorDataTable::StaticClass()->GetClassPathName());
 		AssetPickerConfig.Filter.bRecursiveClasses = false;
 		AssetPickerConfig.bAllowNullSelection = true;
-		AssetPickerConfig.Filter.TagsAndValues.Add(TEXT("Skeleton"), FAssetData(Skeleton).GetExportTextName());
+		AssetPickerConfig.OnShouldFilterAsset = FOnShouldFilterAsset::CreateUObject(Skeleton, &USkeleton::ShouldFilterAsset, TEXT("Skeleton"));
 		AssetPickerConfig.InitialAssetSelection = FAssetData(PreviewInstance->GetMirrorDataTable());
 		AssetPickerConfig.OnAssetSelected = FOnAssetSelected::CreateRaw(const_cast<SAnimViewportToolBar*>(this), &SAnimViewportToolBar::OnMirrorDataTableSelected);
 		AssetPickerConfig.InitialAssetViewType = EAssetViewType::List;
diff --git a/Engine/Source/Editor/Persona/Private/SAnimationSequenceBrowser.cpp b/Engine/Source/Editor/Persona/Private/SAnimationSequenceBrowser.cpp
--- a/Engine/Source/Editor/Persona/Private/SAnimationSequenceBrowser.cpp
+++ b/Engine/Source/Editor/Persona/Private/SAnimationSequenceBrowser.cpp
@@ -1545,7 +1545,7 @@ bool SAnimationSequenceBrowser::HandleFilterAsset(const FAssetData& InAssetData)
 		const USkeleton* DesiredSkeleton = PersonaToolkitPtr.Pin()->GetSkeleton();
 		if (DesiredSkeleton)
 		{
-			return !DesiredSkeleton->IsCompatibleSkeletonByAssetData(InAssetData);
+			return !DesiredSkeleton->IsCompatibleForEditor(InAssetData);
 		}
 	}
 
diff --git a/Engine/Source/Editor/Persona/Private/SAssetFamilyShortcutBar.cpp b/Engine/Source/Editor/Persona/Private/SAssetFamilyShortcutBar.cpp
--- a/Engine/Source/Editor/Persona/Private/SAssetFamilyShortcutBar.cpp
+++ b/Engine/Source/Editor/Persona/Private/SAssetFamilyShortcutBar.cpp
@@ -28,6 +28,7 @@
 #include "Styling/ToolBarStyle.h"
 #include "AssetToolsModule.h"
 #include "IAssetTools.h"
+#include "PersonaAssetFamilyManager.h"
 
 #define LOCTEXT_NAMESPACE "SAssetFamilyShortcutBar"
 
@@ -439,7 +440,7 @@ public:
 
 			// switch to new asset if needed
 			FAssetData NewAssetData = AssetFamily->FindAssetOfType(AssetData.GetClass());
-			if (!bAssetBeingEdited && NewAssetData.IsValid() && NewAssetData != AssetData)
+			if (!bAssetBeingEdited && NewAssetData != AssetData)
 			{
 				AssetData = NewAssetData;
 
@@ -516,23 +517,38 @@ private:
 	bool bPackageDirty;
 };
 
-void SAssetFamilyShortcutBar::Construct(const FArguments& InArgs, const TSharedRef<class FWorkflowCentricApplication>& InHostingApp, const TSharedRef<class IAssetFamily>& InAssetFamily)
+void SAssetFamilyShortcutBar::Construct(const FArguments& InArgs, const TSharedRef<FWorkflowCentricApplication>& InHostingApp, const TSharedRef<IAssetFamily>& InAssetFamily)
 {
+	WeakHostingApp = InHostingApp;
+	AssetFamily = InAssetFamily;
+
 	ThumbnailPool = MakeShareable(new FAssetThumbnailPool(16, false));
 
-	TSharedRef<SHorizontalBox> HorizontalBox = SNew(SHorizontalBox);
+	InAssetFamily->GetOnAssetFamilyChanged().AddSP(this, &SAssetFamilyShortcutBar::OnAssetFamilyChanged);
+
+	HorizontalBox = SNew(SHorizontalBox);
+
+	BuildShortcuts();
 
+	ChildSlot
+	[
+		HorizontalBox.ToSharedRef()
+	];
+}
+
+void SAssetFamilyShortcutBar::BuildShortcuts()
+{
 	TArray<UClass*> AssetTypes;
-	InAssetFamily->GetAssetTypes(AssetTypes);
+	AssetFamily->GetAssetTypes(AssetTypes);
 
 	for (UClass* Class : AssetTypes)
 	{
-		FAssetData AssetData = InAssetFamily->FindAssetOfType(Class);
+		FAssetData AssetData = AssetFamily->FindAssetOfType(Class);
 		HorizontalBox->AddSlot()
 		.AutoWidth()
 		.Padding(0.0f, 4.0f, 16.0f, 4.0f)
 		[
-			SNew(SAssetShortcut, InHostingApp, InAssetFamily, AssetData, ThumbnailPool.ToSharedRef())
+			SNew(SAssetShortcut, WeakHostingApp.Pin().ToSharedRef(), AssetFamily.ToSharedRef(), AssetData, ThumbnailPool.ToSharedRef())
 			.Visibility_Lambda([Class]()
 			{
 				IAssetTools& AssetTools = FModuleManager::LoadModuleChecked<FAssetToolsModule>("AssetTools").Get();
@@ -540,11 +556,22 @@ void SAssetFamilyShortcutBar::Construct(const FArguments& InArgs, const TSharedR
 			})
 		];
 	}
+}
 
-	ChildSlot
-	[
-		HorizontalBox
-	];
+void SAssetFamilyShortcutBar::OnAssetFamilyChanged()
+{
+	HorizontalBox->ClearChildren();
+
+	const TArray<UObject*>* CurrentObjects = WeakHostingApp.Pin()->GetObjectsCurrentlyBeingEdited();
+	if(CurrentObjects && CurrentObjects->Num() > 0 && (*CurrentObjects)[0])
+	{
+		AssetFamily->GetOnAssetFamilyChanged().RemoveAll((this));
+		AssetFamily.Reset();
+		AssetFamily = FPersonaAssetFamilyManager::Get().CreatePersonaAssetFamily((*CurrentObjects)[0]);
+		AssetFamily->GetOnAssetFamilyChanged().AddSP(this, &SAssetFamilyShortcutBar::OnAssetFamilyChanged);
+
+		BuildShortcuts();
+	}
 }
 
 #undef LOCTEXT_NAMESPACE
diff --git a/Engine/Source/Editor/Persona/Private/SAssetFamilyShortcutBar.h b/Engine/Source/Editor/Persona/Private/SAssetFamilyShortcutBar.h
--- a/Engine/Source/Editor/Persona/Private/SAssetFamilyShortcutBar.h
+++ b/Engine/Source/Editor/Persona/Private/SAssetFamilyShortcutBar.h
@@ -7,6 +7,9 @@
 #include "Widgets/SCompoundWidget.h"
 #include "AssetThumbnail.h"
 
+class SHorizontalBox;
+class FWorkflowCentricApplication;
+
 class SAssetFamilyShortcutBar : public SCompoundWidget
 {
 public:
@@ -18,6 +21,19 @@ public:
 	void Construct(const FArguments& InArgs, const TSharedRef<class FWorkflowCentricApplication>& InHostingApp, const TSharedRef<class IAssetFamily>& InAssetFamily);
 
 private:
+	void BuildShortcuts();
+
+	void OnAssetFamilyChanged();
+	
 	/** The thumbnail pool for displaying asset shortcuts */
 	TSharedPtr<class FAssetThumbnailPool> ThumbnailPool;
+
+	/** Container widget for shortcuts */
+	TSharedPtr<SHorizontalBox> HorizontalBox;
+
+	/** Asset family we represent */
+	TSharedPtr<IAssetFamily> AssetFamily;
+
+	/** App we are embedded in */
+	TWeakPtr<FWorkflowCentricApplication> WeakHostingApp;
 };
diff --git a/Engine/Source/Editor/Persona/Private/SRetargetSourceWindow.cpp b/Engine/Source/Editor/Persona/Private/SRetargetSourceWindow.cpp
--- a/Engine/Source/Editor/Persona/Private/SRetargetSourceWindow.cpp
+++ b/Engine/Source/Editor/Persona/Private/SRetargetSourceWindow.cpp
@@ -438,8 +438,14 @@ void SRetargetSourceWindow::OnAddRetargetSource()
 
 	const USkeleton& Skeleton = EditableSkeletonPtr.Pin()->GetSkeleton();
 
-	FString SkeletonString = FAssetData(&Skeleton).GetExportTextName();
-	AssetPickerConfig.Filter.TagsAndValues.Add(USkeletalMesh::GetSkeletonMemberName(), SkeletonString);
+	AssetPickerConfig.OnShouldFilterAsset = FOnShouldFilterAsset::CreateLambda([&Skeleton](const FAssetData& InAssetData)
+	{
+		if(Skeleton.IsCompatibleForEditor(InAssetData))
+		{
+			return false;
+		}
+		return true;
+	});
 
 	TSharedRef<SWidget> Widget = SNew(SBox)
 		.WidthOverride(384.f)
diff --git a/Engine/Source/Editor/Persona/Public/IAssetFamily.h b/Engine/Source/Editor/Persona/Public/IAssetFamily.h
--- a/Engine/Source/Editor/Persona/Public/IAssetFamily.h
+++ b/Engine/Source/Editor/Persona/Public/IAssetFamily.h
@@ -54,9 +54,14 @@ public:
 	virtual UClass* GetAssetFamilyClass(UClass* InClass) const = 0;
 
 	/** Record that an asset was opened */
+	UE_DEPRECATED(5.2, "Please use FPersonaModule::RecordAssetOpened to inform all asset familes rather than just one.")
 	virtual void RecordAssetOpened(const FAssetData& InAssetData) = 0;
 
 	/** Event fired when an asset is opened */
 	DECLARE_EVENT_OneParam(IAssetFamily, FOnAssetOpened, UObject*)
 	virtual FOnAssetOpened& GetOnAssetOpened() = 0;
+
+	/** Event fired when an asset family changes (e.g. relationships are altered) */
+	DECLARE_EVENT(IAssetFamily, FOnAssetFamilyChanged)
+	virtual FOnAssetFamilyChanged& GetOnAssetFamilyChanged() = 0;
 };
diff --git a/Engine/Source/Editor/Persona/Public/PersonaModule.h b/Engine/Source/Editor/Persona/Public/PersonaModule.h
--- a/Engine/Source/Editor/Persona/Public/PersonaModule.h
+++ b/Engine/Source/Editor/Persona/Public/PersonaModule.h
@@ -266,6 +266,12 @@ public:
 	/** Create an asset family for the supplied persona asset */
 	virtual TSharedRef<IAssetFamily> CreatePersonaAssetFamily(const UObject* InAsset) const;
 
+	/** Broadcast event that all asset families need to change */
+	virtual void BroadcastAssetFamilyChange() const;
+
+	/** Record that an asset was opened (forward to relevant asset families) */
+	virtual void RecordAssetOpened(const FAssetData& InAssetData) const;
+	
 	/** Create a shortcut widget for an asset family */
 	virtual TSharedRef<SWidget> CreateAssetFamilyShortcutWidget(const TSharedRef<class FWorkflowCentricApplication>& InHostingApp, const TSharedRef<class IAssetFamily>& InAssetFamily) const;
 
@@ -460,6 +466,7 @@ private:
 	
 	static bool HandleAssetCreated(const TArray<UObject*> NewAssets);
 
+	void RegisterToolMenuExtensions();
 private:
 	TSharedPtr<FExtensibilityManager> MenuExtensibilityManager;
 	TSharedPtr<FExtensibilityManager> ToolBarExtensibilityManager;
diff --git a/Engine/Source/Editor/PhysicsAssetEditor/Private/PhysicsAssetEditor.cpp b/Engine/Source/Editor/PhysicsAssetEditor/Private/PhysicsAssetEditor.cpp
--- a/Engine/Source/Editor/PhysicsAssetEditor/Private/PhysicsAssetEditor.cpp
+++ b/Engine/Source/Editor/PhysicsAssetEditor/Private/PhysicsAssetEditor.cpp
@@ -187,8 +187,7 @@ void FPhysicsAssetEditor::InitPhysicsAssetEditor(const EToolkitMode::Type Mode,
 	FPersonaModule& PersonaModule = FModuleManager::LoadModuleChecked<FPersonaModule>("Persona");
 	PersonaToolkit = PersonaModule.CreatePersonaToolkit(SharedData->PhysicsAsset, PersonaToolkitArgs);
 
-	TSharedRef<IAssetFamily> AssetFamily = PersonaModule.CreatePersonaAssetFamily(ObjectToEdit);
-	AssetFamily->RecordAssetOpened(FAssetData(ObjectToEdit));
+	PersonaModule.RecordAssetOpened(FAssetData(ObjectToEdit));
 
 	FSkeletonTreeArgs SkeletonTreeArgs;
 	SkeletonTreeArgs.OnSelectionChanged = FOnSkeletonTreeSelectionChanged::CreateSP(this, &FPhysicsAssetEditor::HandleSelectionChanged);
diff --git a/Engine/Source/Editor/SkeletalMeshEditor/Private/SkeletalMeshEditor.cpp b/Engine/Source/Editor/SkeletalMeshEditor/Private/SkeletalMeshEditor.cpp
--- a/Engine/Source/Editor/SkeletalMeshEditor/Private/SkeletalMeshEditor.cpp
+++ b/Engine/Source/Editor/SkeletalMeshEditor/Private/SkeletalMeshEditor.cpp
@@ -244,8 +244,7 @@ void FSkeletalMeshEditor::InitSkeletalMeshEditor(const EToolkitMode::Type Mode,
 
 	PersonaToolkit->GetPreviewScene()->SetDefaultAnimationMode(EPreviewSceneDefaultAnimationMode::ReferencePose);
 
-	TSharedRef<IAssetFamily> AssetFamily = PersonaModule.CreatePersonaAssetFamily(InSkeletalMesh);
-	AssetFamily->RecordAssetOpened(FAssetData(InSkeletalMesh));
+	PersonaModule.RecordAssetOpened(FAssetData(InSkeletalMesh));
 
 	TSharedPtr<IPersonaPreviewScene> PreviewScene = PersonaToolkit->GetPreviewScene();
 
diff --git a/Engine/Source/Editor/SkeletonEditor/Private/EditableSkeleton.cpp b/Engine/Source/Editor/SkeletonEditor/Private/EditableSkeleton.cpp
--- a/Engine/Source/Editor/SkeletonEditor/Private/EditableSkeleton.cpp
+++ b/Engine/Source/Editor/SkeletonEditor/Private/EditableSkeleton.cpp
@@ -24,6 +24,8 @@
 #include "SBlendProfilePicker.h"
 #include "AssetNotifications.h"
 #include "BlueprintActionDatabase.h"
+#include "IAssetFamily.h"
+#include "PersonaModule.h"
 #include "AssetRegistry/ARFilter.h"
 #include "AssetRegistry/AssetRegistryModule.h"
 #include "SSkeletonWidget.h"
@@ -1324,6 +1326,14 @@ void FEditableSkeleton::RenameRetargetSource(const FName InOldName, const FName
 		TagsAndValues.Add(GET_MEMBER_NAME_CHECKED(UAnimSequence, RetargetSource), InOldName.ToString());
 		AssetRegistryModule.Get().GetAssetsByTagValues(TagsAndValues, AssetList);
 
+		// This could get assets of different skeletons here, so trim them
+		const FString SkeletonName = FAssetData(Skeleton).GetExportTextName();
+		AssetList.RemoveAll([&SkeletonName](const FAssetData& InAssetData)
+		{
+			FString AssetSkeletonName = InAssetData.GetTagValueRef<FString>(TEXT("Skeleton"));
+			return AssetSkeletonName != SkeletonName;
+		});
+		
 		// ask users if they'd like to continue and/or fix up
 		if (AssetList.Num() > 0)
 		{
@@ -1421,6 +1431,14 @@ void FEditableSkeleton::DeleteRetargetSources(const TArray<FName>& InRetargetSou
 			TagsAndValues.Add(GET_MEMBER_NAME_CHECKED(UAnimSequence, RetargetSource), PoseFound->PoseName.ToString());
 			AssetRegistryModule.Get().GetAssetsByTagValues(TagsAndValues, AssetList);
 
+			// This could get assets of different skeletons here, so trim them
+			const FString SkeletonName = FAssetData(Skeleton).GetExportTextName();
+			AssetList.RemoveAll([&SkeletonName](const FAssetData& InAssetData)
+			{
+				FString AssetSkeletonName = InAssetData.GetTagValueRef<FString>(TEXT("Skeleton"));
+				return AssetSkeletonName != SkeletonName;
+			});
+			
 			// ask users if they'd like to continue and/or fix up
 			if (AssetList.Num() > 0)
 			{
@@ -1510,6 +1528,10 @@ void FEditableSkeleton::AddCompatibleSkeleton(const USkeleton* InCompatibleSkele
 	Skeleton->Modify();
 
 	Skeleton->AddCompatibleSkeleton(InCompatibleSkeleton);
+
+	// Inform asset families
+	FPersonaModule& PersonaModule = FModuleManager::LoadModuleChecked<FPersonaModule>("Persona");
+	PersonaModule.BroadcastAssetFamilyChange();
 }
 
 void FEditableSkeleton::RemoveCompatibleSkeleton(const USkeleton* InCompatibleSkeleton)
@@ -1518,6 +1540,10 @@ void FEditableSkeleton::RemoveCompatibleSkeleton(const USkeleton* InCompatibleSk
 	Skeleton->Modify();
 
 	Skeleton->RemoveCompatibleSkeleton(InCompatibleSkeleton);
+
+	// Inform asset families
+	FPersonaModule& PersonaModule = FModuleManager::LoadModuleChecked<FPersonaModule>("Persona");
+	PersonaModule.BroadcastAssetFamilyChange();
 }
 
 void FEditableSkeleton::RefreshRigConfig()
diff --git a/Engine/Source/Editor/SkeletonEditor/Private/SkeletonEditor.cpp b/Engine/Source/Editor/SkeletonEditor/Private/SkeletonEditor.cpp
--- a/Engine/Source/Editor/SkeletonEditor/Private/SkeletonEditor.cpp
+++ b/Engine/Source/Editor/SkeletonEditor/Private/SkeletonEditor.cpp
@@ -99,8 +99,7 @@ void FSkeletonEditor::InitSkeletonEditor(const EToolkitMode::Type Mode, const TS
 
 	PersonaToolkit->GetPreviewScene()->SetDefaultAnimationMode(EPreviewSceneDefaultAnimationMode::ReferencePose);
 
-	TSharedRef<IAssetFamily> AssetFamily = PersonaModule.CreatePersonaAssetFamily(InSkeleton);
-	AssetFamily->RecordAssetOpened(FAssetData(InSkeleton));
+	PersonaModule.RecordAssetOpened(FAssetData(InSkeleton));
 
 	FSkeletonTreeArgs SkeletonTreeArgs;
 	SkeletonTreeArgs.OnSelectionChanged = FOnSkeletonTreeSelectionChanged::CreateSP(this, &FSkeletonEditor::HandleSelectionChanged);
diff --git a/Engine/Source/Editor/UnrealEd/Classes/Preferences/PersonaOptions.h b/Engine/Source/Editor/UnrealEd/Classes/Preferences/PersonaOptions.h
--- a/Engine/Source/Editor/UnrealEd/Classes/Preferences/PersonaOptions.h
+++ b/Engine/Source/Editor/UnrealEd/Classes/Preferences/PersonaOptions.h
@@ -239,6 +239,10 @@ class UNREALED_API UPersonaOptions : public UObject
 	UPROPERTY(config)
 	TArray<FName> TimelineEnabledSnaps;
 
+	/** Whether to allow animation assets that are incompatible with the current skeleton/skeletal mesh to be selected. */
+	UPROPERTY(EditAnywhere, config, Category = "Assets")
+	bool bAllowIncompatibleSkeletonSelection;
+
 public:
 	void SetShowGrid( bool bInShowGrid );
 	void SetHighlightOrigin( bool bInHighlightOrigin );
@@ -257,6 +261,7 @@ public:
 	void SetNotifyTimingNodeColor(const FLinearColor& InColor);
 	void SetBranchingPointTimingNodeColor(const FLinearColor& InColor);
 	FAssetEditorOptions& GetAssetEditorOptions(const FName& InContext);
+	bool GetAllowIncompatibleSkeletonSelection() const;
 
 	DECLARE_MULTICAST_DELEGATE_TwoParams(FOnUpdateSettingsMulticaster, const UPersonaOptions*, EPropertyChangeType::Type);
 	FOnUpdateSettingsMulticaster OnSettingsChange;
diff --git a/Engine/Source/Editor/UnrealEd/Private/Animation/DebugSkelMeshComponent.cpp b/Engine/Source/Editor/UnrealEd/Private/Animation/DebugSkelMeshComponent.cpp
--- a/Engine/Source/Editor/UnrealEd/Private/Animation/DebugSkelMeshComponent.cpp
+++ b/Engine/Source/Editor/UnrealEd/Private/Animation/DebugSkelMeshComponent.cpp
@@ -517,17 +517,6 @@ FString UDebugSkelMeshComponent::GetPreviewText() const
 
 void UDebugSkelMeshComponent::InitAnim(bool bForceReinit)
 {
-	// If we already have PreviewInstance and its asset's Skeleton isn't compatible with the mesh's Skeleton
-	// then we need to clear it up to avoid an issue
-	if ( PreviewInstance && PreviewInstance->GetCurrentAsset() && GetSkeletalMeshAsset())
-	{
-		if (!GetSkeletalMeshAsset()->GetSkeleton()->IsCompatible(PreviewInstance->GetCurrentAsset()->GetSkeleton()))
-		{
-			// if it doesn't match, just clear it
-			PreviewInstance->SetAnimationAsset(nullptr);
-		}
-	}
-
 	if (PreviewInstance != nullptr && AnimScriptInstance == PreviewInstance && bForceReinit)
 	{
 		// Reset current animation data
diff --git a/Engine/Source/Editor/UnrealEd/Private/Animation/EditorAnimCompositeSegment.cpp b/Engine/Source/Editor/UnrealEd/Private/Animation/EditorAnimCompositeSegment.cpp
--- a/Engine/Source/Editor/UnrealEd/Private/Animation/EditorAnimCompositeSegment.cpp
+++ b/Engine/Source/Editor/UnrealEd/Private/Animation/EditorAnimCompositeSegment.cpp
@@ -45,7 +45,7 @@ bool UEditorAnimCompositeSegment::ApplyChangesToMontage()
 	{
 		if(Composite->AnimationTrack.AnimSegments.IsValidIndex(AnimSegmentIndex))
 		{
-			if (AnimSegment.GetAnimReference() && Composite->GetSkeleton()->IsCompatible(AnimSegment.GetAnimReference()->GetSkeleton()))
+			if (AnimSegment.GetAnimReference() && Composite->GetSkeleton()->IsCompatibleForEditor(AnimSegment.GetAnimReference()->GetSkeleton()))
 			{
 				Composite->AnimationTrack.AnimSegments[AnimSegmentIndex] = AnimSegment;
 				return true;
diff --git a/Engine/Source/Editor/UnrealEd/Private/Animation/EditorAnimSegment.cpp b/Engine/Source/Editor/UnrealEd/Private/Animation/EditorAnimSegment.cpp
--- a/Engine/Source/Editor/UnrealEd/Private/Animation/EditorAnimSegment.cpp
+++ b/Engine/Source/Editor/UnrealEd/Private/Animation/EditorAnimSegment.cpp
@@ -48,7 +48,7 @@ bool UEditorAnimSegment::ApplyChangesToMontage()
 	{
 		if(Montage->SlotAnimTracks.IsValidIndex(AnimSlotIndex) && Montage->SlotAnimTracks[AnimSlotIndex].AnimTrack.AnimSegments.IsValidIndex(AnimSegmentIndex) )
 		{
-			if (AnimSegment.GetAnimReference() && Montage->GetSkeleton()->IsCompatible(AnimSegment.GetAnimReference()->GetSkeleton()))
+			if (AnimSegment.GetAnimReference() && Montage->GetSkeleton()->IsCompatibleForEditor(AnimSegment.GetAnimReference()->GetSkeleton()))
 			{
 				Montage->SlotAnimTracks[AnimSlotIndex].AnimTrack.AnimSegments[AnimSegmentIndex] = AnimSegment;
 				Montage->UpdateLinkableElements(AnimSlotIndex, AnimSegmentIndex);
diff --git a/Engine/Source/Editor/UnrealEd/Private/LevelEditorViewport.cpp b/Engine/Source/Editor/UnrealEd/Private/LevelEditorViewport.cpp
--- a/Engine/Source/Editor/UnrealEd/Private/LevelEditorViewport.cpp
+++ b/Engine/Source/Editor/UnrealEd/Private/LevelEditorViewport.cpp
@@ -819,20 +819,16 @@ static bool AttemptApplyObjToComponent(UObject* ObjToUse, USceneComponent* Compo
 						const FScopedTransaction Transaction(LOCTEXT("DropAnimBlueprintOnObject", "Drop Anim Blueprint On Object"));
 						SkeletalMeshComponent->Modify();
 
-						// If the component doesn't have a mesh or the anim blueprint's skeleton isn't compatible with the existing mesh's skeleton, the mesh should change
-						const bool bShouldChangeMesh = !SkeletalMeshComponent->GetSkeletalMeshAsset() || !SkeletalMeshComponent->GetSkeletalMeshAsset()->GetSkeleton()->IsCompatible(AnimBPSkeleton);
+						// If the component doesn't have a mesh, the mesh should change
+						const bool bShouldChangeMesh = !SkeletalMeshComponent->GetSkeletalMeshAsset();
 
 						if (bShouldChangeMesh)
 						{
 							SkeletalMeshComponent->SetSkeletalMesh(AnimBPSkeleton->GetPreviewMesh(true));
 						}
 
-						// Verify that the skeletons are compatible before changing the anim BP
-						if (SkeletalMeshComponent->GetSkeletalMeshAsset() && SkeletalMeshComponent->GetSkeletalMeshAsset()->GetSkeleton()->IsCompatible(AnimBPSkeleton))
-						{
-							SkeletalMeshComponent->SetAnimInstanceClass(DroppedObjAsAnimBlueprint->GeneratedClass);
-							bResult = true;
-						}
+						SkeletalMeshComponent->SetAnimInstanceClass(DroppedObjAsAnimBlueprint->GeneratedClass);
+						bResult = true;
 					}
 				}
 			}
@@ -854,8 +850,8 @@ static bool AttemptApplyObjToComponent(UObject* ObjToUse, USceneComponent* Compo
 						const FScopedTransaction Transaction(LOCTEXT("DropAnimationOnObject", "Drop Animation On Object"));
 						SkeletalMeshComponent->Modify();
 
-						// If the component doesn't have a mesh or the anim blueprint's skeleton isn't compatible with the existing mesh's skeleton, the mesh should change
-						const bool bShouldChangeMesh = !SkeletalMeshComponent->GetSkeletalMeshAsset() || !SkeletalMeshComponent->GetSkeletalMeshAsset()->GetSkeleton()->IsCompatible(AnimSkeleton);
+						// If the component doesn't have a mesh, the mesh should change
+						const bool bShouldChangeMesh = !SkeletalMeshComponent->GetSkeletalMeshAsset();
 
 						if (bShouldChangeMesh)
 						{
diff --git a/Engine/Source/Editor/UnrealEd/Private/PreferenceStubs.cpp b/Engine/Source/Editor/UnrealEd/Private/PreferenceStubs.cpp
--- a/Engine/Source/Editor/UnrealEd/Private/PreferenceStubs.cpp
+++ b/Engine/Source/Editor/UnrealEd/Private/PreferenceStubs.cpp
@@ -154,6 +154,10 @@ UPersonaOptions::UPersonaOptions(const FObjectInitializer& ObjectInitializer)
 	TimelineEnabledSnaps = { "CompositeSegment", "MontageSection" };
 
 	bExpandTreeOnSelection = true;
+
+	bAllowIncompatibleSkeletonSelection = false;
+
+	USkeleton::AreAllSkeletonsCompatibleDelegate.BindUObject(this, &UPersonaOptions::GetAllowIncompatibleSkeletonSelection);
 }
 
 void UPersonaOptions::SetShowGrid( bool bInShowGrid )
@@ -293,3 +297,8 @@ void UPersonaOptions::PostEditChangeProperty(struct FPropertyChangedEvent& Prope
 
 	OnSettingsChange.Broadcast(this, PropertyChangedEvent.ChangeType);
 }
+
+bool UPersonaOptions::GetAllowIncompatibleSkeletonSelection() const
+{
+	return bAllowIncompatibleSkeletonSelection;
+}
\ No newline at end of file
diff --git a/Engine/Source/Runtime/AnimGraphRuntime/Private/AnimCustomInstanceHelper.cpp b/Engine/Source/Runtime/AnimGraphRuntime/Private/AnimCustomInstanceHelper.cpp
--- a/Engine/Source/Runtime/AnimGraphRuntime/Private/AnimCustomInstanceHelper.cpp
+++ b/Engine/Source/Runtime/AnimGraphRuntime/Private/AnimCustomInstanceHelper.cpp
@@ -18,5 +18,5 @@ bool FAnimCustomInstanceHelper::ShouldCreateCustomInstancePlayer(const USkeletal
 	UAnimInstance* AnimInstance = SkeletalMeshComponent->GetAnimInstance();
 
 	return (AnimInstance == nullptr || SkeletalMeshComponent->GetAnimationMode() != EAnimationMode::AnimationBlueprint ||
-		AnimInstance->GetClass() != SkeletalMeshComponent->AnimClass || !Skeleton || !Skeleton->IsCompatible(AnimInstance->CurrentSkeleton));
+		AnimInstance->GetClass() != SkeletalMeshComponent->AnimClass || !Skeleton || !AnimInstance->CurrentSkeleton);
 }
\ No newline at end of file
diff --git a/Engine/Source/Runtime/AnimGraphRuntime/Private/AnimNodes/AnimNode_BlendSpacePlayer.cpp b/Engine/Source/Runtime/AnimGraphRuntime/Private/AnimNodes/AnimNode_BlendSpacePlayer.cpp
--- a/Engine/Source/Runtime/AnimGraphRuntime/Private/AnimNodes/AnimNode_BlendSpacePlayer.cpp
+++ b/Engine/Source/Runtime/AnimGraphRuntime/Private/AnimNodes/AnimNode_BlendSpacePlayer.cpp
@@ -79,7 +79,7 @@ void FAnimNode_BlendSpacePlayerBase::UpdateInternal(const FAnimationUpdateContex
 	DECLARE_SCOPE_HIERARCHICAL_COUNTER_ANIMNODE(UpdateInternal)
 
 	UBlendSpace* CurrentBlendSpace = GetBlendSpace();
-	if ((CurrentBlendSpace != nullptr) && (Context.AnimInstanceProxy->IsSkeletonCompatible(CurrentBlendSpace->GetSkeleton())))
+	if (CurrentBlendSpace != nullptr && CurrentBlendSpace->GetSkeleton() != nullptr)
 	{	
 		if (PreviousBlendSpace != CurrentBlendSpace)
 		{
@@ -125,7 +125,7 @@ void FAnimNode_BlendSpacePlayerBase::Evaluate_AnyThread(FPoseContext& Output)
 	DECLARE_SCOPE_HIERARCHICAL_COUNTER_ANIMNODE(Evaluate_AnyThread) 
 
 	UBlendSpace* CurrentBlendSpace = GetBlendSpace();
-	if ((CurrentBlendSpace != nullptr) && (Output.AnimInstanceProxy->IsSkeletonCompatible(CurrentBlendSpace->GetSkeleton())))
+	if (CurrentBlendSpace != nullptr && CurrentBlendSpace->GetSkeleton() != nullptr)
 	{
 		FAnimationPoseData AnimationPoseData(Output);
 		CurrentBlendSpace->GetAnimationPose(BlendSampleDataCache, FAnimExtractContext(static_cast<double>(InternalTimeAccumulator), Output.AnimInstanceProxy->ShouldExtractRootMotion(), DeltaTimeRecord, GetLoop()), AnimationPoseData);
diff --git a/Engine/Source/Runtime/AnimGraphRuntime/Private/AnimNodes/AnimNode_PoseBlendNode.cpp b/Engine/Source/Runtime/AnimGraphRuntime/Private/AnimNodes/AnimNode_PoseBlendNode.cpp
--- a/Engine/Source/Runtime/AnimGraphRuntime/Private/AnimNodes/AnimNode_PoseBlendNode.cpp
+++ b/Engine/Source/Runtime/AnimGraphRuntime/Private/AnimNodes/AnimNode_PoseBlendNode.cpp
@@ -46,9 +46,9 @@ void FAnimNode_PoseBlendNode::Evaluate_AnyThread(FPoseContext& Output)
 
 	bool bValidPose = false;
 
-	if (CurrentPoseAsset.IsValid() && (PoseExtractContext.PoseCurves.Num() > 0) && (Output.AnimInstanceProxy->IsSkeletonCompatible(CurrentPoseAsset->GetSkeleton())))
+	const UPoseAsset* CachedPoseAsset = CurrentPoseAsset.Get();
+	if (PoseAsset && PoseExtractContext.PoseCurves.Num() > 0 && PoseAsset->GetSkeleton() != nullptr)
 	{
-		const UPoseAsset* CachedPoseAsset = CurrentPoseAsset.Get();
 		FPoseContext CurrentPose(Output);
 		// only give pose curve, we don't set any more curve here
 		for (int32 PoseIdx = 0; PoseIdx < PoseExtractContext.PoseCurves.Num(); ++PoseIdx)
diff --git a/Engine/Source/Runtime/AnimGraphRuntime/Private/AnimNodes/AnimNode_PoseByName.cpp b/Engine/Source/Runtime/AnimGraphRuntime/Private/AnimNodes/AnimNode_PoseByName.cpp
--- a/Engine/Source/Runtime/AnimGraphRuntime/Private/AnimNodes/AnimNode_PoseByName.cpp
+++ b/Engine/Source/Runtime/AnimGraphRuntime/Private/AnimNodes/AnimNode_PoseByName.cpp
@@ -47,7 +47,8 @@ void FAnimNode_PoseByName::Evaluate_AnyThread(FPoseContext& Output)
 {
 	DECLARE_SCOPE_HIERARCHICAL_COUNTER_ANIMNODE(Evaluate_AnyThread)
 	// make sure we have curve to eval
-	if ((CurrentPoseAsset.IsValid()) && (PoseExtractContext.PoseCurves.Num() > 0) && (Output.AnimInstanceProxy->IsSkeletonCompatible(CurrentPoseAsset->GetSkeleton())))
+	const UPoseAsset* CachedPoseAsset = CurrentPoseAsset.Get();
+	if (CachedPoseAsset && PoseExtractContext.PoseCurves.Num() > 0 && CurrentPoseAsset->GetSkeleton() != nullptr)
 	{
 		// we only have one 
 		PoseExtractContext.PoseCurves[0].Value = PoseWeight;
diff --git a/Engine/Source/Runtime/AnimGraphRuntime/Private/AnimNodes/AnimNode_PoseDriver.cpp b/Engine/Source/Runtime/AnimGraphRuntime/Private/AnimNodes/AnimNode_PoseDriver.cpp
--- a/Engine/Source/Runtime/AnimGraphRuntime/Private/AnimNodes/AnimNode_PoseDriver.cpp
+++ b/Engine/Source/Runtime/AnimGraphRuntime/Private/AnimNodes/AnimNode_PoseDriver.cpp
@@ -374,9 +374,10 @@ void FAnimNode_PoseDriver::Evaluate_AnyThread(FPoseContext& Output)
 		if (OutputWeights.Num() > 0)
 		{
 			// If we want to drive poses, and PoseAsset is assigned and compatible
+			const UPoseAsset* CachedPoseAsset = CurrentPoseAsset.Get();
 			if (DriveOutput == EPoseDriverOutput::DrivePoses &&
-				CurrentPoseAsset.IsValid() &&
-				Output.AnimInstanceProxy->IsSkeletonCompatible(CurrentPoseAsset->GetSkeleton()))
+				CachedPoseAsset &&
+				CurrentPoseAsset->GetSkeleton() != nullptr)
 			{
 				FPoseContext CurrentPose(Output);
 
diff --git a/Engine/Source/Runtime/AnimGraphRuntime/Private/AnimNodes/AnimNode_SequenceEvaluator.cpp b/Engine/Source/Runtime/AnimGraphRuntime/Private/AnimNodes/AnimNode_SequenceEvaluator.cpp
--- a/Engine/Source/Runtime/AnimGraphRuntime/Private/AnimNodes/AnimNode_SequenceEvaluator.cpp
+++ b/Engine/Source/Runtime/AnimGraphRuntime/Private/AnimNodes/AnimNode_SequenceEvaluator.cpp
@@ -42,7 +42,7 @@ void FAnimNode_SequenceEvaluatorBase::UpdateAssetPlayer(const FAnimationUpdateCo
 		// Clamp input to a valid position on this sequence's time line.
 		CurrentExplicitTime = FMath::Clamp(CurrentExplicitTime, 0.f, CurrentSequence->GetPlayLength());
 
-		if ((!GetTeleportToExplicitTime() || (GetGroupName() != NAME_None) || (GetGroupMethod() == EAnimSyncMethod::Graph)) && (Context.AnimInstanceProxy->IsSkeletonCompatible(CurrentSequence->GetSkeleton())))
+		if ((!GetTeleportToExplicitTime() || (GetGroupName() != NAME_None) || (GetGroupMethod() == EAnimSyncMethod::Graph)) && CurrentSequence->GetSkeleton() != nullptr)
 		{
 			if (bReinitialized)
 			{
@@ -89,7 +89,7 @@ void FAnimNode_SequenceEvaluatorBase::Evaluate_AnyThread(FPoseContext& Output)
 	DECLARE_SCOPE_HIERARCHICAL_COUNTER_ANIMNODE(Evaluate_AnyThread)
 	check(Output.AnimInstanceProxy != nullptr);
 	UAnimSequenceBase* CurrentSequence = GetSequence();
-	if ((CurrentSequence != nullptr) && (Output.AnimInstanceProxy->IsSkeletonCompatible(CurrentSequence->GetSkeleton())))
+	if (CurrentSequence != nullptr && CurrentSequence->GetSkeleton() != nullptr)
 	{
 		FAnimationPoseData AnimationPoseData(Output);
 		CurrentSequence->GetAnimationPose(AnimationPoseData, FAnimExtractContext(static_cast<double>(InternalTimeAccumulator), Output.AnimInstanceProxy->ShouldExtractRootMotion(), DeltaTimeRecord, GetShouldLoop()));
diff --git a/Engine/Source/Runtime/AnimGraphRuntime/Public/AnimNodes/AnimNode_BlendSpacePlayer.h b/Engine/Source/Runtime/AnimGraphRuntime/Public/AnimNodes/AnimNode_BlendSpacePlayer.h
--- a/Engine/Source/Runtime/AnimGraphRuntime/Public/AnimNodes/AnimNode_BlendSpacePlayer.h
+++ b/Engine/Source/Runtime/AnimGraphRuntime/Public/AnimNodes/AnimNode_BlendSpacePlayer.h
@@ -121,6 +121,7 @@ struct ANIMGRAPHRUNTIME_API FAnimNode_BlendSpacePlayer : public FAnimNode_BlendS
 {
 	GENERATED_BODY()
 
+	friend class UAnimGraphNode_BlendSpaceBase;
 	friend class UAnimGraphNode_BlendSpacePlayer;
 	friend class UAnimGraphNode_BlendSpaceEvaluator;
 	friend class UAnimGraphNode_RotationOffsetBlendSpace;
diff --git a/Engine/Source/Runtime/Engine/Classes/Animation/SkeletalMeshActor.h b/Engine/Source/Runtime/Engine/Classes/Animation/SkeletalMeshActor.h
--- a/Engine/Source/Runtime/Engine/Classes/Animation/SkeletalMeshActor.h
+++ b/Engine/Source/Runtime/Engine/Classes/Animation/SkeletalMeshActor.h
@@ -85,9 +85,6 @@ public:
 	//~ End AActor Interface
 
 private:
-	// utility function to see if it can play animation or not
-	bool CanPlayAnimation(class UAnimSequenceBase* AnimAssetBase=NULL) const;
-
 	// currently actively playing montage
 	TMap<FName, TWeakObjectPtr<class UAnimMontage>> CurrentlyPlayingMontages;
 
diff --git a/Engine/Source/Runtime/Engine/Classes/Animation/Skeleton.h b/Engine/Source/Runtime/Engine/Classes/Animation/Skeleton.h
--- a/Engine/Source/Runtime/Engine/Classes/Animation/Skeleton.h
+++ b/Engine/Source/Runtime/Engine/Classes/Animation/Skeleton.h
@@ -33,6 +33,10 @@ class USkinnedAsset;
 class FPackageReloadedEvent;
 struct FAssetData;
 enum class EPackageReloadPhase : uint8;
+class USkeleton;
+
+// Delegate used to control global skeleton compatibility
+DECLARE_DELEGATE_RetVal(bool, FAreAllSkeletonsCompatible);
 
 /** This is a mapping table between bone in a particular skeletal mesh and bone of this skeleton set. */
 USTRUCT()
@@ -280,174 +284,6 @@ public:
 	}
 };
 
-struct FSkeletonRemapping
-{
-	FSkeletonRemapping() = default;
-	FSkeletonRemapping(const USkeleton* InSourceSkeleton, const USkeleton* InTargetSkeleton);
-
-	const TWeakObjectPtr<const USkeleton>& GetSourceSkeleton() const { return SourceSkeleton; }
-	const TWeakObjectPtr<const USkeleton>& GetTargetSkeleton() const { return TargetSkeleton; }
-
-	/**
-	 * Compose this remapping with another remapping in place.  The other remapping's source skeleton must match
-	 * this remapping's target skeleton.  The result will map from this remapping's source skeleton to the other
-	 * remapping's target skeleton
-	 *
-	 * @param	OtherSkeletonRemapping		Skeleton remapping to compose into this remapping
-	 */
-	void ComposeWith(const FSkeletonRemapping& OtherSkeletonRemapping);
-
-	/**
-	 * Get the target skeleton bone index that corresponds to the specified bone on the source skeleton
-	 *
-	 * @param	SourceSkeletonBoneIndex		Skeleton bone index on the source skeleton
-	 * @return								Skeleton bone index on the target skeleton (or INDEX_NONE)
-	 */
-	int32 GetTargetSkeletonBoneIndex(int32 SourceSkeletonBoneIndex) const
-	{
-		return (SourceToTargetBoneIndexes.IsValidIndex(SourceSkeletonBoneIndex))
-				   ? SourceToTargetBoneIndexes[SourceSkeletonBoneIndex]
-				   : INDEX_NONE;
-	}
-
-	/**
-	 * Get the source skeleton bone index that corresponds to the specified bone on the target skeleton
-	 *
-	 * @param	TargetSkeletonBoneIndex		Skeleton bone index on the target skeleton
-	 * @return								Skeleton bone index on the source skeleton (or INDEX_NONE)
-	 */
-	int32 GetSourceSkeletonBoneIndex(int32 TargetSkeletonBoneIndex) const
-	{
-		return (TargetToSourceBoneIndexes.IsValidIndex(TargetSkeletonBoneIndex))
-				   ? TargetToSourceBoneIndexes[TargetSkeletonBoneIndex]
-				   : INDEX_NONE;
-	}
-
-	/**
-	 * Get the specified bone transform retargeted from the source skeleton onto the target skeleton, corrected
-	 * for differences between source and target rest poses
-	 *
-	 * @param	TargetSkeletonBoneIndex		Skeleton bone index on the target skeleton
-	 * @param	SourceTransform				Bone transform from the corresponding bone on the source skeleton
-	 * @return								Transform mapped onto the target skeleton
-	 */
-	FTransform RetargetBoneTransformToTargetSkeleton(int32 TargetSkeletonBoneIndex, const FTransform& SourceTransform) const
-	{
-		return FTransform(
-			RetargetBoneRotationToTargetSkeleton(TargetSkeletonBoneIndex, SourceTransform.GetRotation()),
-			RetargetBoneTranslationToTargetSkeleton(TargetSkeletonBoneIndex, SourceTransform.GetTranslation()),
-			SourceTransform.GetScale3D());
-	}
-
-	/**
-	 * Get the specified bone translation retargeted from the source skeleton onto the target skeleton, corrected
-	 * for differences between source and target rest poses
-	 *
-	 * @param	TargetSkeletonBoneIndex		Skeleton bone index on the target skeleton
-	 * @param	SourceTranslation			Bone translation from the corresponding bone on the source skeleton
-	 * @return								Translation mapped onto the target skeleton
-	 */
-	FVector RetargetBoneTranslationToTargetSkeleton(int32 TargetSkeletonBoneIndex, const FVector& SourceTranslation) const
-	{
-		// Compute the translation part of FTransform(Q1) * Source * FTransform(Q0)
-		const TTuple<FQuat, FQuat>& QQ = RetargetingTable[TargetSkeletonBoneIndex];
-		return QQ.Get<0>().RotateVector(SourceTranslation);
-	}
-
-	/**
-	 * Get the specified bone rotation retargeted from the source skeleton onto the target skeleton, corrected
-	 * for differences between source and target rest poses
-	 *
-	 * @param	TargetSkeletonBoneIndex		Skeleton bone index on the target skeleton
-	 * @param	SourceRotation				Bone rotation from the corresponding bone on the source skeleton
-	 * @return								Rotation mapped onto the target skeleton
-	 */
-	FQuat RetargetBoneRotationToTargetSkeleton(int32 TargetSkeletonBoneIndex, const FQuat& SourceRotation) const
-	{
-		// Compute the rotation part of FTransform(Q1) * Source * FTransform(Q0)
-		const TTuple<FQuat, FQuat>& QQ = RetargetingTable[TargetSkeletonBoneIndex];
-		return QQ.Get<0>() * SourceRotation * QQ.Get<1>();
-	}
-
-	/**
-	 * Get the specified additive transform retargeted from the source skeleton onto the target skeleton, corrected
-	 * for differences between source and target rest poses
-	 *
-	 * @param	TargetSkeletonBoneIndex		Skeleton bone index on the target skeleton
-	 * @param	SourceTransform				Bone transform from the corresponding bone on the source skeleton
-	 * @return								Transform mapped onto the target skeleton
-	 */
-	FTransform RetargetAdditiveTransformToTargetSkeleton(int32 TargetSkeletonBoneIndex, const FTransform& SourceTransform) const
-	{
-		return FTransform(
-			RetargetAdditiveRotationToTargetSkeleton(TargetSkeletonBoneIndex, SourceTransform.GetRotation()),
-			RetargetAdditiveTranslationToTargetSkeleton(TargetSkeletonBoneIndex, SourceTransform.GetTranslation()),
-			SourceTransform.GetScale3D());
-	}
-
-	/**
-	 * Get the specified additive translation retargeted from the source skeleton onto the target skeleton, corrected
-	 * for differences between source and target rest poses
-	 *
-	 * @param	TargetSkeletonBoneIndex		Skeleton bone index on the target skeleton
-	 * @param	SourceTranslation			Bone translation from the corresponding bone on the source skeleton
-	 * @return								Translation mapped onto the target skeleton
-	 */
-	FVector RetargetAdditiveTranslationToTargetSkeleton(int32 TargetSkeletonBoneIndex, const FVector& SourceTranslation) const
-	{
-		// Compute the translation part of FTransform(Q0.Inverse) * Source * FTransform(Q0)
-		const TTuple<FQuat, FQuat>& QQ = RetargetingTable[TargetSkeletonBoneIndex];
-		return QQ.Get<0>().RotateVector(SourceTranslation);
-	}
-
-	/**
-	 * Get the specified additive rotation retargeted from the source skeleton onto the target skeleton, corrected
-	 * for differences between source and target rest poses
-	 *
-	 * @param	TargetSkeletonBoneIndex		Skeleton bone index on the target skeleton
-	 * @param	SourceRotation				Bone rotation from the corresponding bone on the source skeleton
-	 * @return								Rotation mapped onto the target skeleton
-	 */
-	FQuat RetargetAdditiveRotationToTargetSkeleton(int32 TargetSkeletonBoneIndex, const FQuat& SourceRotation) const
-	{
-		// Compute the rotation part of FTransform(Q0.Inverse) * Source * FTransform(Q0)
-		const TTuple<FQuat, FQuat>& QQ = RetargetingTable[TargetSkeletonBoneIndex];
-		return QQ.Get<0>() * SourceRotation * QQ.Get<0>().Inverse();
-	}
-
-	/**
-	 * Get the curve mapping array, which maps from the source curve UID to the target curve UID.
-	 * You can access this array like this:
-	 * 
-	 * \code{.cpp}
-	 * const SmartName::UID_Type TargetCurveUID = GetSourceToTargetCurveMapping()[SourceCurveUID];
-	 * \endcode
-	 */
-	const TArray<SmartName::UID_Type>& GetSourceToTargetCurveMapping() const { return SourceToTargetCurveMapping; }
-
-	/**
-	 * Generates the mapping table for the curves, based on curve names.
-	 * Basically this will look at the curve names in the smart name table on the source, and tries to find matching ones in the table of 
-	 * the target skeleton. It then maps the UID's.
-	 */
-	void GenerateCurveMapping();
-
-private:
-	TWeakObjectPtr<const USkeleton> SourceSkeleton;
-	TWeakObjectPtr<const USkeleton> TargetSkeleton;
-
-	// Table of target skeleton bone indexes (indexed by source skeleton bone index)
-	TArray<int32> SourceToTargetBoneIndexes;
-
-	// Table of source skeleton bone indexes (indexed by target skeleton bone index)
-	TArray<int32> TargetToSourceBoneIndexes;
-	
-	// Maps curve UIDs between source and target (indexed by source curve UID).
-	TArray<SmartName::UID_Type> SourceToTargetCurveMapping;
-
-	// Table of precalculated constants for retargeting from source to target (indexed by target skeleton bone index)
-	TArray<TTuple<FQuat, FQuat>> RetargetingTable;
-};
 
 /**
  *	USkeleton : that links between mesh and animation
@@ -506,7 +342,6 @@ protected:
 public:
 	//~ Begin UObject Interface.
 #if WITH_EDITOR
-	ENGINE_API virtual void PostEditChangeProperty(struct FPropertyChangedEvent& PropertyChangedEvent);
 	ENGINE_API virtual void PreEditUndo() override;
 	ENGINE_API virtual void PostEditUndo() override;
 #endif
@@ -565,6 +400,8 @@ public:
 	const TArray<uint16>& GetDefaultCurveUIDList() const { return DefaultCurveUIDList; }
 
 	ENGINE_API const TArray<TSoftObjectPtr<USkeleton>>& GetCompatibleSkeletons() const { return CompatibleSkeletons; }
+
+	UE_DEPRECATED(5.2, "Please use UE::Anim::FSkeletonRemappingRegistry::GetRemapping.")
 	ENGINE_API const FSkeletonRemapping* GetSkeletonRemapping(const USkeleton* SourceSkeleton) const;
 
 #if WITH_EDITOR
@@ -815,45 +652,57 @@ public:
 	 * @param	(out) List of Direct Children
 	 */
 	ENGINE_API int32 GetChildBones(int32 ParentBoneIndex, TArray<int32> & Children) const;
-
-#endif
+	
 	/**
-	 * Check if animation content authored on the supplied skeleton may be played on this skeleton
-	 */
-	ENGINE_API bool IsCompatible(const USkeleton* InSkeleton) const;
+	 * Check if animation content authored on the supplied skeleton may be played on this skeleton.
+	 * Note that animations may not always be correct if they were not authored on this skeleton, and may require
+	 * retargeting or other fixup.
+	 */	
+	ENGINE_API bool IsCompatibleForEditor(const USkeleton* InSkeleton) const;
 
 	/**
 	 * Check if this skeleton is compatible with a given other asset, if that is a skeleton.
-	 */
-	ENGINE_API bool IsCompatibleSkeletonByAssetData(const FAssetData& AssetData, const TCHAR* InTag = TEXT("Skeleton")) const;
+	 */	
+	ENGINE_API bool IsCompatibleForEditor(const FAssetData& AssetData, const TCHAR* InTag = TEXT("Skeleton")) const;
 
 	/**
 	 * Check if this skeleton is compatible with another skeleton asset that is identified by the string returned by AssetData(SkeletonPtr).GetExportTextName().
 	 */
-	ENGINE_API bool IsCompatibleSkeletonByAssetString(const FString& SkeletonAssetString) const;
+	ENGINE_API bool IsCompatibleForEditor(const FString& SkeletonAssetString) const;
+
+	/** Wrapper for !IsCompatibleForEditor, used as a convenience function for binding to FOnShouldFilterAsset in asset pickers. */
+	ENGINE_API bool ShouldFilterAsset(const FAssetData& InAssetData, const TCHAR* InTag = TEXT("Skeleton")) const;
+	
+	/** Get all skeleton assets that are compatible with this skeleton (not just the internal list, but also reciprocally and implicitly compatible skeletons) */
+	ENGINE_API void GetCompatibleSkeletonAssets(TArray<FAssetData>& OutAssets) const;
 
-	DECLARE_EVENT_OneParam(USkeleton, FSkeletonDestructEvent, const USkeleton*);
-	FSkeletonDestructEvent OnSkeletonDestructEvent;
+	/** Get compatible assets given the asset's class and skeleton tag.*/
+	ENGINE_API void GetCompatibleAssets(UClass* AssetClass, const TCHAR* InTag, TArray<FAssetData>& OutAssets) const;
+#endif
+
+	UE_DEPRECATED(5.2, "Compatibility is now an editor-only concern. Please use IsCompatibleForEditor.")
+	ENGINE_API bool IsCompatible(const USkeleton* InSkeleton) const;
 
+	UE_DEPRECATED(5.2, "Compatibility is now an editor-only concern. Please use IsCompatibleForEditor.")
+	ENGINE_API bool IsCompatibleSkeletonByAssetData(const FAssetData& AssetData, const TCHAR* InTag = TEXT("Skeleton")) const;
+
+	UE_DEPRECATED(5.2, "Compatibility is now an editor-only concern. Please use IsCompatibleForEditor.")
+	ENGINE_API bool IsCompatibleSkeletonByAssetString(const FString& SkeletonAssetString) const;
+	
 	DECLARE_EVENT(USkeleton, FSmartNamesChangedEvent);
 	FSmartNamesChangedEvent OnSmartNamesChangedEvent;
 
+#if WITH_EDITORONLY_DATA
+	/**
+	 * Global compatibility delegate, used to override skeleton compatibility. If this returns true, no additional
+	 * compatibility checks are made.
+	 */
+	ENGINE_API static FAreAllSkeletonsCompatible AreAllSkeletonsCompatibleDelegate;
+#endif
 public:
 	UFUNCTION(BlueprintCallable, Category=Skeleton)
 	ENGINE_API void AddCompatibleSkeleton(const USkeleton* SourceSkeleton);
 	ENGINE_API void RemoveCompatibleSkeleton(const USkeleton* SourceSkeleton);
-protected:
-	mutable FCriticalSection SkeletonRemappingMutex;
-
-	/** The skeleton remappings, which map bones from a source skeleton to this skeleton (the target). */
-	TMap<const USkeleton*, FSkeletonRemapping*> SkeletonRemappings;
-
-	/**
- 	 * The function that is called when another skeleton is being destructed.
-	 * This is used to unregister skeleton remappings that aren't needed anymore.
-	 * @param Skeleton The skeleton that is being destructed.
-	 */
-	void HandleSkeletonDestruct(const USkeleton* Skeleton);
 
 	/**
 	 * Handle smart name changes.
@@ -861,41 +710,6 @@ protected:
 	 **/
 	void HandleSmartNamesChangedEvent();
 
-	/**
-	 * Build or update all skeleton remappings for all loaded skeletons that we are compatible with.
-	 * And if bBidirectional is set to true also update all remappings in other skeletons that are compatible with this skeleton.
-	 * Keep in mind this only generates mappings to skeletons that are loaded, as otherwise we can't generate the mappings yet, even though the list of compatible skeletons might contain
-	 * a larger set of skeletons. We do not load all skeletons upfront.
-	 * @param bBidirectional When enabled we update or create all mappings for both skeletons we are compatible with and skeletons that are compatible with ourselves.
-	 *        When set to false it only updates or creates mappings to skeletons we are compatible with.
-	 */
-	void BuildSkeletonRemappings(bool bBidirectional);
-
-	/**
-	 * Remove all skeleton remappings stored inside our skeleton.
-	 * This frees up memory.
-	 */
-	void ClearSkeletonRemappings();
-
-	/**
-	 * Remove the skeleton remappings for a specific source skeleton.
-	 * This only works if the source skeleton is inside the compatible skeleton list of this skeleton.
-	 * If a skeleton is passed in that isn't in the compatibility list this method will essentially do nothing.
-	 * The skeleton will remain inside the compatibility list, it is just the remapping data that is being removed.
-	 * @param SourceSkeleton The source skeleton we are compatible with and want to remove the generated mapping data for.
-	 */
-	void RemoveSkeletonRemapping(const USkeleton* SourceSkeleton);
-
-	/**
-	 * Create or update a skeletal remapping to a given source skeleton.
-	 * @param SourceSkeleton The skeleton from which we receive animation data.
-	 * @param bRebuildIfExists When set to true this method will delete the old mapping and generate a new one. This is useful when the hierarchy changes for example.
-	 */
-	void CreateSkeletonRemappingIfNeeded(const USkeleton* SourceSkeleton, bool bRebuildIfExists = false);
-
-	static TArray<USkeleton*> LoadedSkeletons; // The set of skeletons that are currently loaded.
-	static FCriticalSection LoadedSkeletonsMutex; // The mutex used for thread safety when registering and removing items from the LoadedSkeletons array.
-
 public:
 	/** 
 	 * Indexing naming convention
@@ -967,20 +781,13 @@ public:
 	 */
 	ENGINE_API bool RecreateBoneTree(USkinnedAsset* InSkinnedAsset);
 
-	// @todo document
-	const TArray<FTransform>& GetRefLocalPoses( FName RetargetSource = NAME_None ) const 
-	{
-		if ( RetargetSource != NAME_None ) 
-		{
-			const FReferencePose* FoundRetargetSource = AnimRetargetSources.Find(RetargetSource);
-			if (FoundRetargetSource)
-			{
-				return FoundRetargetSource->ReferencePose;
-			}
-		}
-
-		return ReferenceSkeleton.GetRefBonePose();	
-	}
+	
+	/**
+	 * Get the local-space ref pose for the specified retarget source.
+	 * @param	RetargetSource	The name of the retarget source to find
+	 * @return the transforms for the retarget source reference pose. If the retarget source is not found, this returns the skeleton's reference pose.
+	 */
+	ENGINE_API const TArray<FTransform>& GetRefLocalPoses( FName RetargetSource = NAME_None ) const;
 
 #if WITH_EDITORONLY_DATA
 	/**
@@ -989,6 +796,12 @@ public:
 	 * @return NAME_None if a retarget source was not found, or a valid name if it was
 	 */
 	ENGINE_API FName GetRetargetSourceForMesh(USkinnedAsset* InSkinnedAsset) const;
+
+	/**
+	 * Get all the retarget source names for this skeleton.
+	 * @param	OutRetargetSources	The retarget source names array to be filled
+	 */
+	ENGINE_API void GetRetargetSources(TArray<FName>& OutRetargetSources) const;
 #endif
 
 	/** 
@@ -1130,6 +943,10 @@ public:
 	ENGINE_API static const FName CurveNameTag;
 	ENGINE_API static const FString CurveTagDelimiter;
 
+	// Asset registry information for compatible skeletons
+	ENGINE_API static const FName CompatibleSkeletonsNameTag;
+	ENGINE_API static const FString CompatibleSkeletonsTagDelimiter;
+
 	// rig Configs
 	ENGINE_API static const FName RigTag;
 	ENGINE_API void SetRigConfig(URig * Rig);
diff --git a/Engine/Source/Runtime/Engine/Private/AnimBlueprint.cpp b/Engine/Source/Runtime/Engine/Private/AnimBlueprint.cpp
--- a/Engine/Source/Runtime/Engine/Private/AnimBlueprint.cpp
+++ b/Engine/Source/Runtime/Engine/Private/AnimBlueprint.cpp
@@ -414,7 +414,7 @@ bool UAnimBlueprint::IsCompatible(const UAnimBlueprint* InAnimBlueprint) const
 		return true;
 	}
 
-	return (TargetSkeleton != nullptr && TargetSkeleton->IsCompatible(InAnimBlueprint->TargetSkeleton));
+	return (TargetSkeleton != nullptr && TargetSkeleton->IsCompatibleForEditor(InAnimBlueprint->TargetSkeleton));
 }
 
 bool UAnimBlueprint::IsCompatibleByAssetString(const FString& InSkeletonAsset, bool bInIsTemplate, bool bInIsInterface) const
@@ -431,6 +431,6 @@ bool UAnimBlueprint::IsCompatibleByAssetString(const FString& InSkeletonAsset, b
 		return true;
 	}
 
-	return (TargetSkeleton != nullptr && TargetSkeleton->IsCompatibleSkeletonByAssetString(InSkeletonAsset));
+	return (TargetSkeleton != nullptr && TargetSkeleton->IsCompatibleForEditor(InSkeletonAsset));
 }
 #endif
diff --git a/Engine/Source/Runtime/Engine/Private/Animation/AnimInstance.cpp b/Engine/Source/Runtime/Engine/Private/Animation/AnimInstance.cpp
--- a/Engine/Source/Runtime/Engine/Private/Animation/AnimInstance.cpp
+++ b/Engine/Source/Runtime/Engine/Private/Animation/AnimInstance.cpp
@@ -2002,7 +2002,7 @@ UAnimMontage* UAnimInstance::PlaySlotAnimationAsDynamicMontage_WithBlendArgs(UAn
 
 UAnimMontage* UAnimInstance::PlaySlotAnimationAsDynamicMontage_WithBlendSettings(UAnimSequenceBase* Asset, FName SlotNodeName, const FMontageBlendSettings& BlendInSettings, const FMontageBlendSettings& BlendOutSettings, float InPlayRate, int32 LoopCount, float BlendOutTriggerTime, float InTimeToStartMontageAt)
 {
-	if (Asset && CurrentSkeleton->IsCompatible(Asset->GetSkeleton()))
+	if (Asset && Asset->GetSkeleton())
 	{
 		// create asset using the information
 		UAnimMontage* NewMontage = UAnimMontage::CreateSlotAnimationAsDynamicMontage_WithBlendSettings(Asset, SlotNodeName, BlendInSettings, BlendOutSettings, InPlayRate, LoopCount, BlendOutTriggerTime);
@@ -2093,7 +2093,7 @@ float UAnimInstance::Montage_PlayInternal(UAnimMontage* MontageToPlay, const FMo
 
 	if (MontageToPlay && (MontageToPlay->GetPlayLength() > 0.f) && MontageToPlay->HasValidSlotSetup())
 	{
-		if (CurrentSkeleton && CurrentSkeleton->IsCompatible(MontageToPlay->GetSkeleton()))
+		if (CurrentSkeleton && MontageToPlay->GetSkeleton())
 		{
 			const FName NewMontageGroupName = MontageToPlay->GetGroupName();
 			if (bStopAllMontages)
@@ -2934,18 +2934,6 @@ void UAnimInstance::PerformLinkedLayerOverlayOperation(TSubclassOf<UAnimInstance
 	if (IAnimClassInterface* AnimBlueprintClass = IAnimClassInterface::GetFromClass(GetClass()))
 	{
 		UClass* NewClass = InClass.Get();
-		if(NewClass)
-		{
-			// Verify target skeleton matches at runtime
-			IAnimClassInterface* LinkedAnimBlueprintClass = IAnimClassInterface::GetFromClass(NewClass);
-			USkeleton* LinkedSkeleton = LinkedAnimBlueprintClass->GetTargetSkeleton();
-			USkeleton* OuterSkeleton = AnimBlueprintClass->GetTargetSkeleton();
-			if(!LinkedSkeleton->IsCompatible(OuterSkeleton))
-			{
-				UE_LOG(LogAnimation, Warning, TEXT("Linking layer: Linked instance class has a mismatched target skeleton. Expected %s, found %s. Anim instance: %s Linked layer class : %s"), OuterSkeleton ? *OuterSkeleton->GetName() : TEXT("null"), LinkedSkeleton ? *LinkedSkeleton->GetName() : TEXT("null"), *this->GetName(), *NewClass->GetName());
-				return;
-			}
-		}
 
 		// Make sure we have valid objects as initialization can route back out of linked instances into this outer graph
 		GetProxyOnAnyThread<FAnimInstanceProxy>().InitializeObjects(this);
diff --git a/Engine/Source/Runtime/Engine/Private/Animation/AnimMontage.cpp b/Engine/Source/Runtime/Engine/Private/Animation/AnimMontage.cpp
--- a/Engine/Source/Runtime/Engine/Private/Animation/AnimMontage.cpp
+++ b/Engine/Source/Runtime/Engine/Private/Animation/AnimMontage.cpp
@@ -454,20 +454,17 @@ void UAnimMontage::PostLoad()
 		}
 	}
 
-	// verify if skeleton matches, otherwise clear it, this can happen if anim sequence has been modified when this hasn't been loaded. 
-	if (const USkeleton* MySkeleton = GetSkeleton())
+	// verify if skeleton is valid, otherwise clear it, this can happen if anim sequence has been modified when this hasn't been loaded. 
+	for (int32 I=0; I<SlotAnimTracks.Num(); ++I)
 	{
-		for (int32 I=0; I<SlotAnimTracks.Num(); ++I)
+		if ( SlotAnimTracks[I].AnimTrack.AnimSegments.Num() > 0 )
 		{
-			if ( SlotAnimTracks[I].AnimTrack.AnimSegments.Num() > 0 )
+			UAnimSequenceBase* SequenceBase = SlotAnimTracks[I].AnimTrack.AnimSegments[0].GetAnimReference();
+			if (SequenceBase && SequenceBase->GetSkeleton() == nullptr)
 			{
-				UAnimSequenceBase* SequenceBase = SlotAnimTracks[I].AnimTrack.AnimSegments[0].GetAnimReference();
-				if (SequenceBase && !MySkeleton->IsCompatible(SequenceBase->GetSkeleton()))
-				{
-					SlotAnimTracks[I].AnimTrack.AnimSegments[0].SetAnimReference(nullptr);
-					MarkPackageDirty();
-					break;
-				}
+				SlotAnimTracks[I].AnimTrack.AnimSegments[0].SetAnimReference(nullptr);
+				MarkPackageDirty();
+				break;
 			}
 		}
 	}
diff --git a/Engine/Source/Runtime/Engine/Private/Animation/AnimNode_LinkedAnimGraph.cpp b/Engine/Source/Runtime/Engine/Private/Animation/AnimNode_LinkedAnimGraph.cpp
--- a/Engine/Source/Runtime/Engine/Private/Animation/AnimNode_LinkedAnimGraph.cpp
+++ b/Engine/Source/Runtime/Engine/Private/Animation/AnimNode_LinkedAnimGraph.cpp
@@ -302,19 +302,6 @@ void FAnimNode_LinkedAnimGraph::ReinitializeLinkedAnimInstance(const UAnimInstan
 void FAnimNode_LinkedAnimGraph::SetAnimClass(TSubclassOf<UAnimInstance> InClass, const UAnimInstance* InOwningAnimInstance)
 {
 	UClass* NewClass = InClass.Get();
-	if(NewClass)
-	{
-		// Verify target skeleton match at runtime
-		IAnimClassInterface* LinkedAnimBlueprintClass = IAnimClassInterface::GetFromClass(NewClass);
-		IAnimClassInterface* OuterAnimBlueprintClass = IAnimClassInterface::GetFromClass(InOwningAnimInstance->GetClass());
-		USkeleton* LinkedSkeleton = LinkedAnimBlueprintClass->GetTargetSkeleton();
-		USkeleton* OuterSkeleton = OuterAnimBlueprintClass->GetTargetSkeleton();
-		if(!LinkedSkeleton->IsCompatible(OuterSkeleton))
-		{
-			UE_LOG(LogAnimation, Warning, TEXT("Setting linked anim instance class: Class has a mismatched target skeleton. Expected %s, found %s."), OuterSkeleton ? *OuterSkeleton->GetName() : TEXT("null"), LinkedSkeleton ? *LinkedSkeleton->GetName() : TEXT("null"));
-			return;
-		}
-	}
 
 	// Verified OK, so set it now
 	TSubclassOf<UAnimInstance> OldClass = InstanceClass;
diff --git a/Engine/Source/Runtime/Engine/Private/Animation/AnimNode_SequencePlayer.cpp b/Engine/Source/Runtime/Engine/Private/Animation/AnimNode_SequencePlayer.cpp
--- a/Engine/Source/Runtime/Engine/Private/Animation/AnimNode_SequencePlayer.cpp
+++ b/Engine/Source/Runtime/Engine/Private/Animation/AnimNode_SequencePlayer.cpp
@@ -91,7 +91,7 @@ void FAnimNode_SequencePlayerBase::UpdateAssetPlayer(const FAnimationUpdateConte
 		CurrentSequence = nullptr;
 	}
 
-	if ((CurrentSequence != nullptr) && (Context.AnimInstanceProxy->IsSkeletonCompatible(CurrentSequence->GetSkeleton())))
+	if (CurrentSequence != nullptr && CurrentSequence->GetSkeleton() != nullptr)
 	{
 		const float CurrentPlayRate = GetPlayRate();
 		const float CurrentPlayRateBasis = GetPlayRateBasis();
@@ -119,7 +119,7 @@ void FAnimNode_SequencePlayerBase::Evaluate_AnyThread(FPoseContext& Output)
 	DECLARE_SCOPE_HIERARCHICAL_COUNTER_ANIMNODE(Evaluate_AnyThread);
 
 	UAnimSequenceBase* CurrentSequence = GetSequence();
-	if ((CurrentSequence != nullptr) && (Output.AnimInstanceProxy->IsSkeletonCompatible(CurrentSequence->GetSkeleton())))
+	if (CurrentSequence != nullptr && CurrentSequence->GetSkeleton() != nullptr)
 	{
 		const bool bExpectedAdditive = Output.ExpectsAdditivePose();
 		const bool bIsAdditive = CurrentSequence->IsValidAdditive();
diff --git a/Engine/Source/Runtime/Engine/Private/Animation/AnimSequence.cpp b/Engine/Source/Runtime/Engine/Private/Animation/AnimSequence.cpp
--- a/Engine/Source/Runtime/Engine/Private/Animation/AnimSequence.cpp
+++ b/Engine/Source/Runtime/Engine/Private/Animation/AnimSequence.cpp
@@ -936,7 +936,7 @@ void UAnimSequence::PostLoad()
 	// but this causes issue since once this happens this is unrecoverable until you delete from outside of editor
 	if (IsValidAdditive())
 	{
-		if (RefPoseSeq && GetSkeleton() && !GetSkeleton()->IsCompatible(RefPoseSeq->GetSkeleton()))
+		if (RefPoseSeq && RefPoseSeq->GetSkeleton() == nullptr)
 		{
 			// if this happens, there was a issue with retargeting
 			UE_LOG(LogAnimation, Warning, TEXT("Animation %s - Invalid additive animation base animation (%s)"), *GetName(), *RefPoseSeq->GetName());
diff --git a/Engine/Source/Runtime/Engine/Private/Animation/AnimSingleNodeInstance.cpp b/Engine/Source/Runtime/Engine/Private/Animation/AnimSingleNodeInstance.cpp
--- a/Engine/Source/Runtime/Engine/Private/Animation/AnimSingleNodeInstance.cpp
+++ b/Engine/Source/Runtime/Engine/Private/Animation/AnimSingleNodeInstance.cpp
@@ -50,8 +50,8 @@ void UAnimSingleNodeInstance::SetAnimationAsset(class UAnimationAsset* NewAsset,
 		}
 		else if (CurrentAsset != nullptr)
 		{
-			// if we have an asset, make sure their skeleton matches, otherwise, null it
-			if (!MeshComponent->GetSkeletalMeshAsset()->GetSkeleton()->IsCompatible(CurrentAsset->GetSkeleton()))
+			// if we have an asset, make sure their skeleton is valid, otherwise, null it
+			if (CurrentAsset->GetSkeleton() == nullptr)
 			{
 				// clear asset since we do not have matching skeleton
 				CurrentAsset = nullptr;
diff --git a/Engine/Source/Runtime/Engine/Private/Animation/AnimSingleNodeInstanceProxy.cpp b/Engine/Source/Runtime/Engine/Private/Animation/AnimSingleNodeInstanceProxy.cpp
--- a/Engine/Source/Runtime/Engine/Private/Animation/AnimSingleNodeInstanceProxy.cpp
+++ b/Engine/Source/Runtime/Engine/Private/Animation/AnimSingleNodeInstanceProxy.cpp
@@ -475,10 +475,7 @@ void FAnimNode_SingleNode::Evaluate_AnyThread(FPoseContext& Output)
 		// if it has a preview pose asset, we have to handle that after we do all animation
 		if (const UPoseAsset* PoseAsset = Proxy->CurrentAsset->PreviewPoseAsset)
 		{
-			USkeleton* MySkeleton = Proxy->CurrentAsset->GetSkeleton();
-
-			// if skeleton doesn't match it won't work
-			if (MySkeleton->IsCompatible(PoseAsset->GetSkeleton()))
+			if (PoseAsset->GetSkeleton() != nullptr)
 			{
 				const TArray<FSmartName>& PoseNames = PoseAsset->GetPoseNames();
 
diff --git a/Engine/Source/Runtime/Engine/Private/Animation/AnimationDecompression.cpp b/Engine/Source/Runtime/Engine/Private/Animation/AnimationDecompression.cpp
--- a/Engine/Source/Runtime/Engine/Private/Animation/AnimationDecompression.cpp
+++ b/Engine/Source/Runtime/Engine/Private/Animation/AnimationDecompression.cpp
@@ -7,7 +7,8 @@
 
 #include "AnimationRuntime.h"
 #include "AnimEncoding.h"
-
+#include "SkeletonRemappingRegistry.h"
+#include "SkeletonRemapping.h"
 
 CSV_DECLARE_CATEGORY_MODULE_EXTERN(ENGINE_API, Animation);
 DECLARE_CYCLE_STAT(TEXT("Build Anim Track Pairs"), STAT_BuildAnimTrackPairs, STATGROUP_Anim);
@@ -46,7 +47,7 @@ void DecompressPose(FCompactPose& OutPose,
 	const int32 NumTracks = CompressedData.CompressedTrackToSkeletonMapTable.Num();
 
 	const USkeleton* TargetSkeleton = RequiredBones.GetSkeletonAsset();
-	const FSkeletonRemapping* SkeletonRemapping = TargetSkeleton->GetSkeletonRemapping(DecompressionContext.GetSourceSkeleton());
+	const FSkeletonRemapping& SkeletonRemapping = UE::Anim::FSkeletonRemappingRegistry::Get().GetRemapping(DecompressionContext.GetSourceSkeleton(), TargetSkeleton);
 
 	BoneTrackArray& RotationScalePairs = FGetBonePoseScratchArea::Get().RotationScalePairs;
 	BoneTrackArray& TranslationPairs = FGetBonePoseScratchArea::Get().TranslationPairs;
@@ -73,7 +74,7 @@ void DecompressPose(FCompactPose& OutPose,
 		for (int32 TrackIndex = (bFirstTrackIsRootBone ? 1 : 0); TrackIndex < NumTracks; TrackIndex++)
 		{
 			const int32 SourceSkeletonBoneIndex = CompressedData.GetSkeletonIndexFromTrackIndex(TrackIndex);
-			const int32 TargetSkeletonBoneIndex = (SkeletonRemapping) ? SkeletonRemapping->GetTargetSkeletonBoneIndex(SourceSkeletonBoneIndex) : SourceSkeletonBoneIndex;
+			const int32 TargetSkeletonBoneIndex = SkeletonRemapping.IsValid() ? SkeletonRemapping.GetTargetSkeletonBoneIndex(SourceSkeletonBoneIndex) : SourceSkeletonBoneIndex;
 
 			if (TargetSkeletonBoneIndex != INDEX_NONE)
 			{
@@ -136,24 +137,24 @@ void DecompressPose(FCompactPose& OutPose,
 			CompressedData.BoneCompressionCodec->DecompressBone(DecompressionContext, TrackIndex, RootAtom);
 
 			// Retarget the root onto the target skeleton (correcting for differences in rest poses)
-			if (SkeletonRemapping)
+			if (SkeletonRemapping.RequiresReferencePoseRetarget())
 			{
 				const int32 TargetSkeletonBoneIndex = 0;
 
 				if (DecompressionContext.IsAdditiveAnimation())
 				{
-					RootAtom.SetRotation(SkeletonRemapping->RetargetAdditiveRotationToTargetSkeleton(TargetSkeletonBoneIndex, RootAtom.GetRotation()));
+					RootAtom.SetRotation(SkeletonRemapping.RetargetAdditiveRotationToTargetSkeleton(TargetSkeletonBoneIndex, RootAtom.GetRotation()));
 					if (TargetSkeleton->GetBoneTranslationRetargetingMode(TargetSkeletonBoneIndex) != EBoneTranslationRetargetingMode::Skeleton)
 					{
-						RootAtom.SetTranslation(SkeletonRemapping->RetargetAdditiveTranslationToTargetSkeleton(TargetSkeletonBoneIndex, RootAtom.GetTranslation()));
+						RootAtom.SetTranslation(SkeletonRemapping.RetargetAdditiveTranslationToTargetSkeleton(TargetSkeletonBoneIndex, RootAtom.GetTranslation()));
 					}
 				}
 				else
 				{
-					RootAtom.SetRotation(SkeletonRemapping->RetargetBoneRotationToTargetSkeleton(TargetSkeletonBoneIndex, RootAtom.GetRotation()));
+					RootAtom.SetRotation(SkeletonRemapping.RetargetBoneRotationToTargetSkeleton(TargetSkeletonBoneIndex, RootAtom.GetRotation()));
 					if (TargetSkeleton->GetBoneTranslationRetargetingMode(TargetSkeletonBoneIndex) != EBoneTranslationRetargetingMode::Skeleton)
 					{
-						RootAtom.SetTranslation(SkeletonRemapping->RetargetBoneTranslationToTargetSkeleton(TargetSkeletonBoneIndex, RootAtom.GetTranslation()));
+						RootAtom.SetTranslation(SkeletonRemapping.RetargetBoneTranslationToTargetSkeleton(TargetSkeletonBoneIndex, RootAtom.GetTranslation()));
 					}
 				}
 			}
@@ -171,17 +172,17 @@ void DecompressPose(FCompactPose& OutPose,
 	}
 
 	// Retarget the pose onto the target skeleton (correcting for differences in rest poses)
-	if (SkeletonRemapping)
+	if (SkeletonRemapping.RequiresReferencePoseRetarget())
 	{
 		if (DecompressionContext.IsAdditiveAnimation())
 		{
 			for (FCompactPoseBoneIndex BoneIndex(bFirstTrackIsRootBone ? 1 : 0); BoneIndex < OutPose.GetNumBones(); ++BoneIndex)
 			{
 				const int32 TargetSkeletonBoneIndex = RequiredBones.GetSkeletonIndex(BoneIndex);
-				OutPose[BoneIndex].SetRotation(SkeletonRemapping->RetargetAdditiveRotationToTargetSkeleton(TargetSkeletonBoneIndex, OutPose[BoneIndex].GetRotation()));
+				OutPose[BoneIndex].SetRotation(SkeletonRemapping.RetargetAdditiveRotationToTargetSkeleton(TargetSkeletonBoneIndex, OutPose[BoneIndex].GetRotation()));
 				if (TargetSkeleton->GetBoneTranslationRetargetingMode(TargetSkeletonBoneIndex) != EBoneTranslationRetargetingMode::Skeleton)
 				{
-					OutPose[BoneIndex].SetTranslation(SkeletonRemapping->RetargetAdditiveTranslationToTargetSkeleton(TargetSkeletonBoneIndex, OutPose[BoneIndex].GetTranslation()));
+					OutPose[BoneIndex].SetTranslation(SkeletonRemapping.RetargetAdditiveTranslationToTargetSkeleton(TargetSkeletonBoneIndex, OutPose[BoneIndex].GetTranslation()));
 				}
 			}
 		}
@@ -190,10 +191,10 @@ void DecompressPose(FCompactPose& OutPose,
 			for (FCompactPoseBoneIndex BoneIndex(bFirstTrackIsRootBone ? 1 : 0); BoneIndex < OutPose.GetNumBones(); ++BoneIndex)
 			{
 				const int32 TargetSkeletonBoneIndex = RequiredBones.GetSkeletonIndex(BoneIndex);
-				OutPose[BoneIndex].SetRotation(SkeletonRemapping->RetargetBoneRotationToTargetSkeleton(TargetSkeletonBoneIndex, OutPose[BoneIndex].GetRotation()));
+				OutPose[BoneIndex].SetRotation(SkeletonRemapping.RetargetBoneRotationToTargetSkeleton(TargetSkeletonBoneIndex, OutPose[BoneIndex].GetRotation()));
 				if (TargetSkeleton->GetBoneTranslationRetargetingMode(TargetSkeletonBoneIndex) != EBoneTranslationRetargetingMode::Skeleton)
 				{
-					OutPose[BoneIndex].SetTranslation(SkeletonRemapping->RetargetBoneTranslationToTargetSkeleton(TargetSkeletonBoneIndex, OutPose[BoneIndex].GetTranslation()));
+					OutPose[BoneIndex].SetTranslation(SkeletonRemapping.RetargetBoneTranslationToTargetSkeleton(TargetSkeletonBoneIndex, OutPose[BoneIndex].GetTranslation()));
 				}
 			}
 		}
@@ -241,10 +242,10 @@ void DecompressPose(FCompactPose& OutPose,
 
 			// Remap the base pose onto the target skeleton so that we are working entirely in target space
 			FTransform BaseTransform = AuthoredOnRefSkeleton[SourceSkeletonBoneIndex];
-			if (SkeletonRemapping)
+			if (SkeletonRemapping.RequiresReferencePoseRetarget())
 			{
-				const int32 TargetSkeletonBoneIndex = SkeletonRemapping->GetTargetSkeletonBoneIndex(SourceSkeletonBoneIndex);
-				BaseTransform = SkeletonRemapping->RetargetBoneTransformToTargetSkeleton(TargetSkeletonBoneIndex, BaseTransform);
+				const int32 TargetSkeletonBoneIndex = SkeletonRemapping.GetTargetSkeletonBoneIndex(SourceSkeletonBoneIndex);
+				BaseTransform = SkeletonRemapping.RetargetBoneTransformToTargetSkeleton(TargetSkeletonBoneIndex, BaseTransform);
 			}
 
 			// Apply the retargeting as if it were an additive difference between the current skeleton and the retarget skeleton. 
diff --git a/Engine/Source/Runtime/Engine/Private/Animation/AnimationRuntime.cpp b/Engine/Source/Runtime/Engine/Private/Animation/AnimationRuntime.cpp
--- a/Engine/Source/Runtime/Engine/Private/Animation/AnimationRuntime.cpp
+++ b/Engine/Source/Runtime/Engine/Private/Animation/AnimationRuntime.cpp
@@ -16,6 +16,8 @@
 #include "Animation/AnimationPoseData.h"
 #include "Animation/BlendProfile.h"
 #include "Animation/MirrorDataTable.h"
+#include "SkeletonRemappingRegistry.h"
+#include "SkeletonRemapping.h"
 
 DEFINE_LOG_CATEGORY(LogAnimation);
 DEFINE_LOG_CATEGORY(LogRootMotion);
@@ -2598,10 +2600,10 @@ void FAnimationRuntime::RetargetBoneTransform(const USkeleton* SourceSkeleton, c
 	if (SourceSkeleton)
 	{
 		const USkeleton* TargetSkeleton = RequiredBones.GetSkeletonAsset();
-		const FSkeletonRemapping* SkeletonRemapping = TargetSkeleton->GetSkeletonRemapping(SourceSkeleton);
+		const FSkeletonRemapping& SkeletonRemapping = UE::Anim::FSkeletonRemappingRegistry::Get().GetRemapping(SourceSkeleton, TargetSkeleton);
 
 		const int32 TargetSkeletonBoneIndex = RequiredBones.GetSkeletonIndex(BoneIndex);
-		const int32 SourceSkeletonBoneIndex = (SkeletonRemapping) ? SkeletonRemapping->GetSourceSkeletonBoneIndex(TargetSkeletonBoneIndex) : SkeletonBoneIndex;
+		const int32 SourceSkeletonBoneIndex = SkeletonRemapping.IsValid() ? SkeletonRemapping.GetSourceSkeletonBoneIndex(TargetSkeletonBoneIndex) : SkeletonBoneIndex;
 
 		switch (TargetSkeleton->GetBoneTranslationRetargetingMode(TargetSkeletonBoneIndex))
 		{
@@ -2635,9 +2637,9 @@ void FAnimationRuntime::RetargetBoneTransform(const USkeleton* SourceSkeleton, c
 
 					// Remap the base pose onto the target skeleton so that we are working entirely in target space
 					FTransform BaseTransform = AuthoredOnRefSkeleton[SourceSkeletonBoneIndex];
-					if (SkeletonRemapping)
+					if (SkeletonRemapping.RequiresReferencePoseRetarget())
 					{
-						BaseTransform = SkeletonRemapping->RetargetBoneTransformToTargetSkeleton(TargetSkeletonBoneIndex, BaseTransform);
+						BaseTransform = SkeletonRemapping.RetargetBoneTransformToTargetSkeleton(TargetSkeletonBoneIndex, BaseTransform);
 					}
 
 					// Apply the retargeting as if it were an additive difference between the current skeleton and the retarget skeleton. 
diff --git a/Engine/Source/Runtime/Engine/Private/Animation/BlendSpace.cpp b/Engine/Source/Runtime/Engine/Private/Animation/BlendSpace.cpp
--- a/Engine/Source/Runtime/Engine/Private/Animation/BlendSpace.cpp
+++ b/Engine/Source/Runtime/Engine/Private/Animation/BlendSpace.cpp
@@ -885,12 +885,7 @@ void UBlendSpace::GetAnimationPose_Internal(TArray<FBlendSampleData>& BlendSampl
 			}
 			else
 			{
-				if (Sample.Animation
-#if WITH_EDITOR
-					// verify if Sample.Animation->GetSkeleton matches
-					&& ensure(GetSkeleton()->IsCompatible(Sample.Animation->GetSkeleton()))
-#endif // WITH_EDITOR
-					)
+				if (Sample.Animation && Sample.Animation->GetSkeleton() != nullptr)
 				{
 					const float Time = FMath::Clamp<float>(BlendSampleDataCache[I].Time, 0.f, Sample.Animation->GetPlayLength());
 
@@ -1467,7 +1462,10 @@ bool UBlendSpace::IsAnimationCompatibleWithSkeleton(const UAnimSequence* Animati
 {
 	// Check if the animation sequences skeleton is compatible with the blendspace one
 	const USkeleton* MySkeleton = GetSkeleton();
-	const bool bIsAnimationCompatible = AnimationSequence && MySkeleton && AnimationSequence->GetSkeleton() && MySkeleton->IsCompatible(AnimationSequence->GetSkeleton());
+	bool bIsAnimationCompatible = AnimationSequence && MySkeleton && AnimationSequence->GetSkeleton();
+#if WITH_EDITORONLY_DATA
+	bIsAnimationCompatible = bIsAnimationCompatible && MySkeleton->IsCompatibleForEditor(AnimationSequence->GetSkeleton());
+#endif
 	return bIsAnimationCompatible;
 }
 
diff --git a/Engine/Source/Runtime/Engine/Private/Animation/SkeletalMeshActor.cpp b/Engine/Source/Runtime/Engine/Private/Animation/SkeletalMeshActor.cpp
--- a/Engine/Source/Runtime/Engine/Private/Animation/SkeletalMeshActor.cpp
+++ b/Engine/Source/Runtime/Engine/Private/Animation/SkeletalMeshActor.cpp
@@ -150,12 +150,6 @@ void ASkeletalMeshActor::OnRep_ReplicatedMaterial1()
 	SkeletalMeshComponent->SetMaterial(1, ReplicatedMaterial1);
 }
 
-bool ASkeletalMeshActor::CanPlayAnimation(class UAnimSequenceBase* AnimAssetBase/*=NULL*/) const
-{
-	return (SkeletalMeshComponent->GetSkeletalMeshAsset() && SkeletalMeshComponent->GetSkeletalMeshAsset()->GetSkeleton() &&
-		(!AnimAssetBase || SkeletalMeshComponent->GetSkeletalMeshAsset()->GetSkeleton()->IsCompatible(AnimAssetBase->GetSkeleton())));
-}
-
 #if WITH_EDITOR
 
 bool ASkeletalMeshActor::GetReferencedContentObjects(TArray<UObject*>& Objects) const
diff --git a/Engine/Source/Runtime/Engine/Private/Animation/Skeleton.cpp b/Engine/Source/Runtime/Engine/Private/Animation/Skeleton.cpp
--- a/Engine/Source/Runtime/Engine/Private/Animation/Skeleton.cpp
+++ b/Engine/Source/Runtime/Engine/Private/Animation/Skeleton.cpp
@@ -31,6 +31,7 @@
 #include "Animation/AnimBlueprint.h"
 #include "UObject/AnimObjectVersion.h"
 #include "EngineUtils.h"
+#include "SkeletonRemappingRegistry.h"
 #include "Misc/ScopeLock.h"
 
 #include UE_INLINE_GENERATED_CPP_BY_NAME(Skeleton)
@@ -45,9 +46,16 @@ const FString USkeleton::AnimNotifyTagDelimiter = TEXT(";");
 const FName USkeleton::CurveNameTag = FName(TEXT("CurveNameList"));
 const FString USkeleton::CurveTagDelimiter = TEXT(";");
 
+const FName USkeleton::CompatibleSkeletonsNameTag = FName(TEXT("CompatibleSkeletonList"));
+const FString USkeleton::CompatibleSkeletonsTagDelimiter = TEXT(";");
+
 const FName USkeleton::RigTag = FName(TEXT("Rig"));
 #endif 
 
+#if WITH_EDITORONLY_DATA
+FAreAllSkeletonsCompatible USkeleton::AreAllSkeletonsCompatibleDelegate;
+#endif
+
 // Names of smartname containers for skeleton properties
 const FName USkeleton::AnimCurveMappingName = FName(TEXT("AnimationCurves"));
 const FName USkeleton::AnimTrackCurveMappingName = FName(TEXT("AnimationTrackCurves"));
@@ -55,11 +63,6 @@ const FName USkeleton::AnimTrackCurveMappingName = FName(TEXT("AnimationTrackCur
 const FName FAnimSlotGroup::DefaultGroupName = FName(TEXT("DefaultGroup"));
 const FName FAnimSlotGroup::DefaultSlotName = FName(TEXT("DefaultSlot"));
 
-// Skeleton remapping related.
-TArray<USkeleton*> USkeleton::LoadedSkeletons;		// We keep track of a list of loaded skeletons, because we have no thread safe object iterator at the time this code is written.
-FCriticalSection USkeleton::LoadedSkeletonsMutex;	// The mutex we use to lock the LoadedSkeletons array.
-
-
 void SerializeReferencePose(FArchive& Ar, FReferencePose& P, UObject* Outer)
 {
 	Ar.UsingCustomVersion(FAnimPhysObjectVersion::GUID);
@@ -114,30 +117,25 @@ namespace VirtualBoneNameHelpers
 	}
 }
 
-bool USkeleton::IsCompatible(const USkeleton* InSkeleton) const
+#if WITH_EDITORONLY_DATA
+bool USkeleton::IsCompatibleForEditor(const USkeleton* InSkeleton) const
 {
-	if (InSkeleton == nullptr)
-	{
-		return false;
-	}
+	return IsCompatibleForEditor(FAssetData(InSkeleton));
+}
 
-	if (InSkeleton == this)
+bool USkeleton::IsCompatibleForEditor(const FAssetData& AssetData, const TCHAR* InTag) const
+{
+	if(AssetData.GetClass() == USkeleton::StaticClass())
 	{
-		return true;
+		return IsCompatibleForEditor(AssetData.GetExportTextName());
 	}
-
-	for (const TSoftObjectPtr<USkeleton>& CompatibleSkeleton : CompatibleSkeletons)
+	else
 	{
-		if (CompatibleSkeleton == InSkeleton->CachedSoftObjectPtr)
-		{
-			return true;
-		}
+		return IsCompatibleForEditor(AssetData.GetTagValueRef<FString>(InTag));
 	}
-
-	return false;
 }
 
-bool USkeleton::IsCompatibleSkeletonByAssetString(const FString& SkeletonAssetString) const
+bool USkeleton::IsCompatibleForEditor(const FString& SkeletonAssetString) const
 {
 	// First check against itself.
 	const FString SkeletonString = FAssetData(this).GetExportTextName();
@@ -146,6 +144,12 @@ bool USkeleton::IsCompatibleSkeletonByAssetString(const FString& SkeletonAssetSt
 		return true;
 	}
 
+	// Let the global delegate override any per-skeleton settings
+	if(AreAllSkeletonsCompatibleDelegate.IsBound() && AreAllSkeletonsCompatibleDelegate.Execute())
+	{
+		return true;
+	}
+
 	// Now check against the list of compatible skeletons and see if we're dealing with the same asset.
 	const FSoftObjectPath InPath(SkeletonAssetString);
 	for (const TSoftObjectPtr<USkeleton>& CompatibleSkeleton : CompatibleSkeletons)
@@ -156,305 +160,148 @@ bool USkeleton::IsCompatibleSkeletonByAssetString(const FString& SkeletonAssetSt
 		}
 	}
 
-	return false;
-}
-
-void USkeleton::AddCompatibleSkeleton(const USkeleton* SourceSkeleton)
-{
-	CompatibleSkeletons.AddUnique(SourceSkeleton);
+	// Check if the other skeleton is compatible with this via the asset registry
+	const IAssetRegistry& AssetRegistry = FModuleManager::LoadModuleChecked<FAssetRegistryModule>("AssetRegistry").Get();
+	const FAssetData SkeletonAssetData = AssetRegistry.GetAssetByObjectPath(FSoftObjectPath(*SkeletonAssetString));	
+	const FString TagValue = SkeletonAssetData.GetTagValueRef<FString>(USkeleton::CompatibleSkeletonsNameTag);
+	if (!TagValue.IsEmpty())
+	{
+		TArray<FString> OtherCompatibleSkeletons;
+		if (TagValue.ParseIntoArray(OtherCompatibleSkeletons, *USkeleton::CompatibleSkeletonsTagDelimiter, true) > 0)
+		{
+			for (const FString& OtherCompatibleSkeleton : OtherCompatibleSkeletons)
+			{
+				if (SkeletonString == OtherCompatibleSkeleton)
+				{
+					return true;
+				}
+			}
+		}
+	}
 
-	ClearSkeletonRemappings();
-	BuildSkeletonRemappings(false);
+	return false;
 }
 
-void USkeleton::RemoveCompatibleSkeleton(const USkeleton* SourceSkeleton)
+bool USkeleton::ShouldFilterAsset(const FAssetData& InAssetData, const TCHAR* InTag) const
 {
-	CompatibleSkeletons.Remove(SourceSkeleton);
-
-	ClearSkeletonRemappings();
-	BuildSkeletonRemappings(false);
+	return !IsCompatibleForEditor(InAssetData, InTag);
 }
 
-bool USkeleton::IsCompatibleSkeletonByAssetData(const FAssetData& AssetData, const TCHAR* InTag) const
+void USkeleton::GetCompatibleSkeletonAssets(TArray<FAssetData>& OutAssets) const
 {
-	return IsCompatibleSkeletonByAssetString(AssetData.GetTagValueRef<FString>(InTag));
-}
+	const IAssetRegistry& AssetRegistry = FModuleManager::LoadModuleChecked<FAssetRegistryModule>("AssetRegistry").Get();
 
-FSkeletonRemapping::FSkeletonRemapping(const USkeleton* InSourceSkeleton, const USkeleton* InTargetSkeleton)
-	: SourceSkeleton(InSourceSkeleton)
-	, TargetSkeleton(InTargetSkeleton)
-{
-	if (!GetSourceSkeleton().IsValid() || !GetTargetSkeleton().IsValid())
+	// If we are compatible with all, then just add all skeletons
+	if(AreAllSkeletonsCompatibleDelegate.IsBound() && AreAllSkeletonsCompatibleDelegate.Execute())
 	{
-		return;
-	}
+		TArray<FAssetData> AllSkeletons;
+		FARFilter ARFilter;
+		ARFilter.ClassPaths.Add(USkeleton::StaticClass()->GetClassPathName());
+		AssetRegistry.GetAssets(ARFilter, AllSkeletons);
 
-	const FReferenceSkeleton& SourceReferenceSkeleton = InSourceSkeleton->GetReferenceSkeleton();
-	const FReferenceSkeleton& TargetReferenceSkeleton = InTargetSkeleton->GetReferenceSkeleton();
-
-	const int32 SourceNumBones = SourceReferenceSkeleton.GetNum();
-	const int32 TargetNumBones = TargetReferenceSkeleton.GetNum();
-
-	SourceToTargetBoneIndexes.SetNumUninitialized(SourceNumBones);
-	TargetToSourceBoneIndexes.SetNumUninitialized(TargetNumBones);
-	RetargetingTable.SetNumUninitialized(TargetNumBones);
-
-	TArrayView<const FTransform> SourceLocalTransforms;
-	TArrayView<const FTransform> TargetLocalTransforms;
-/*
-	// Build the mapping from source to target bones through the remapping rig if one exists
-	if (true)	// TODO_SKELETON_REMAPPING: Use a remapping rig if it exists
-	{
-		for (int32 SourceBoneIndex = 0; SourceBoneIndex < SourceNumBones; ++SourceBoneIndex)
-		{
-			const FName BoneName = SourceReferenceSkeleton.GetBoneName(SourceBoneIndex);
-			int32 TargetBoneIndex = TargetReferenceSkeleton.FindBoneIndex(BoneName);
-			SourceToTargetBoneIndexes[SourceBoneIndex] = TargetBoneIndex;
-		}
-
-		// Get the matched source and target rest poses from the remapping rig
-		//
-		// TODO_SKELETON_REMAPPING: For now we'll just use the skeleton's rest poses, but we should really be pulling these
-		// poses from a remapping rig so that the use can better align the poses.  Note that we'll want the remapping rig
-		// to be independent from the USkeleton itself because a skeleton may want to participate in multiple remappings with
-		// different other skeletons and may have to pose itself differently to align with the different other skeletons
-		//
-		SourceLocalTransforms = MakeArrayView(SourceReferenceSkeleton.GetRefBonePose());
-		TargetLocalTransforms = MakeArrayView(TargetReferenceSkeleton.GetRefBonePose());
-	}
-	else // Fall back to simple name matching if there is no rig
-*/
-	{
-		// Match source to target bones by name lookup between source and target skeletons
-		for (int32 SourceBoneIndex = 0; SourceBoneIndex < SourceNumBones; ++SourceBoneIndex)
-		{
-			const FName BoneName = SourceReferenceSkeleton.GetBoneName(SourceBoneIndex);
-			const int32 TargetBoneIndex = TargetReferenceSkeleton.FindBoneIndex(BoneName);
-			SourceToTargetBoneIndexes[SourceBoneIndex] = TargetBoneIndex;
-		}
-
-		// Get the matched (hopefully) source and target rest poses from the source and target skeletons
-		SourceLocalTransforms = MakeArrayView(SourceReferenceSkeleton.GetRefBonePose());
-		TargetLocalTransforms = MakeArrayView(TargetReferenceSkeleton.GetRefBonePose());
-	}
-
-	// Force the roots to map onto each other regardless of their names
-	SourceToTargetBoneIndexes[0] = 0;
-
-	// Build the reverse mapping from target back to source bones
-	FMemory::Memset(TargetToSourceBoneIndexes.GetData(), static_cast<uint8>(INDEX_NONE), TargetNumBones * TargetToSourceBoneIndexes.GetTypeSize());
-	for (int32 SourceBoneIndex = 0; SourceBoneIndex < SourceNumBones; ++SourceBoneIndex)
-	{
-		const int32 TargetBoneIndex = SourceToTargetBoneIndexes[SourceBoneIndex];
-		if (TargetBoneIndex != INDEX_NONE)
-		{
-			TargetToSourceBoneIndexes[TargetBoneIndex] = SourceBoneIndex;
-		}
-	}
-
-	TArray<FTransform> SourceComponentTransforms;
-	TArray<FTransform> TargetComponentTransforms;
-	FAnimationRuntime::FillUpComponentSpaceTransforms(SourceReferenceSkeleton, SourceLocalTransforms, SourceComponentTransforms);
-	FAnimationRuntime::FillUpComponentSpaceTransforms(TargetReferenceSkeleton, TargetLocalTransforms, TargetComponentTransforms);
-
-	// Calculate the retargeting constants to map from source skeleton space to target skeleton space
-	//
-	// Simply remapping joint indices is usually not sufficient to give us the desired result pose if the source and target
-	// skeletons are built with different conventions for joint orientations.  We therefore need to compute a remapping
-	// between the source and target joint orientations, which we do in terms of delta rotations from the rest pose:
-	//
-	//		Q = P * D * R
-	//
-	// where:
-	//
-	//		Q is the final joint orientation in component space
-	//		P is the parent joint orientation in component space
-	//		D is the delta rotation that we want to remap
-	//		R is the local space rest pose orientation for the joint
-	//
-	// We want to find a mapping from the source to the target:
-	//
-	//		Ps * Ds * Rs  -->  Pt * Dt * Rt
-	//
-	// such that the deltas produce equivalent rotations on the mesh even if their parent rotation frames are different.
-	// In other words, we need to find Dt such that its component space rotation is equivalent to Ds.  To convert a rotation
-	// from local space (D) to component space (C), given its parent (P), we have:
-	//
-	//		C * P = P * D
-	//		C = P * D * P⁻¹
-	//
-	// Setting the source and target component space rotations to be equal (Cs = Ct) then gives us:
-	//
-	//		Ps * Ds * Ps⁻¹ = Pt * Dt * Pt⁻¹
-	//
-	// which we then solve for Dt to get:
-	//
-	//		Dt = Pt⁻¹ * Ps * Ds * Ps⁻¹ * Pt
-	//
-	// However, when we're remapping an animation pose, we will have the local transforms rather than the deltas from the
-	// rest pose, so we also need to convert between these local transforms and the equivalent deltas:
-	//
-	//		L = D * R
-	//		D = L * R⁻¹
-	//
-	// Combining that with our equation for Dt, we get:
-	//
-	//		Lt * Rt⁻¹ = Pt⁻¹ * Ps * Ls * Rs⁻¹ * Ps⁻¹ * Pt
-	//
-	// Solving for Lt then gives us:
-	//
-	//		Lt = Pt⁻¹ * Ps * Ls * Rs⁻¹ * Ps⁻¹ * Pt * Rt
-	//
-	// Finally, factoring out the constant terms (which we precompute here) gives us:
-	//
-	//		Lt = Q0 * Ls * Q1
-	//		Q0 = Pt⁻¹ * Ps
-	//		Q1 = Rs⁻¹ * Ps⁻¹ * Pt * Rt
-	//
-	// Note that when remapping additive animations, we drop the rest pose terms, but we still need to convert between the
-	// source and target rotation frames.  Dropping Rs and Rt from the equations above gives us:
-	//
-	//		Lt = Pt⁻¹ * Ps * Ls * Ps⁻¹ * Pt			(for additive animations)
-	//
-	// which is equivalent to the following in terms of our precomputed constants:
-	//
-	//		Lt = Q0 * Ls * Q0⁻¹						(for additive animations)
-	//
-	RetargetingTable[0] = MakeTuple(FQuat::Identity, SourceLocalTransforms[0].GetRotation().Inverse() * TargetLocalTransforms[0].GetRotation());
-	for (int32 TargetBoneIndex = 1; TargetBoneIndex < TargetNumBones; ++TargetBoneIndex)
-	{
-		const int32 SourceBoneIndex = TargetToSourceBoneIndexes[TargetBoneIndex];
-		if (SourceBoneIndex != INDEX_NONE)
-		{
-			const int32 SourceParentIndex = SourceReferenceSkeleton.GetParentIndex(SourceBoneIndex);
-			const int32 TargetParentIndex = TargetReferenceSkeleton.GetParentIndex(TargetBoneIndex);
-			check(SourceParentIndex != INDEX_NONE);
-			check(TargetParentIndex != INDEX_NONE);
-
-			const FQuat PS = SourceComponentTransforms[SourceParentIndex].GetRotation();
-			const FQuat PT = TargetComponentTransforms[TargetParentIndex].GetRotation();
-
-			const FQuat RS = SourceLocalTransforms[SourceBoneIndex].GetRotation();
-			const FQuat RT = TargetLocalTransforms[TargetBoneIndex].GetRotation();
-
-			const FQuat Q0 = PT.Inverse() * PS;
-			const FQuat Q1 = RS.Inverse() * PS.Inverse() * PT * RT;
-
-			RetargetingTable[TargetBoneIndex] = MakeTuple(Q0, Q1);
+		OutAssets.Append(AllSkeletons);
+	}
+	else
+	{
+		// Always add 'this'
+		FAssetData AssetDataThis(this);
+		OutAssets.Add(AssetDataThis);
+		
+		// Add skeletons in our compatibility list
+		for(TSoftObjectPtr<USkeleton> CompatibleSkeleton : CompatibleSkeletons)
+		{
+			const FAssetData SkeletonAssetData = AssetRegistry.GetAssetByObjectPath(CompatibleSkeleton.ToSoftObjectPath());
+			if(SkeletonAssetData.IsValid())
+			{
+				OutAssets.Add(SkeletonAssetData);
+			}
 		}
-		else
+
+		// Add skeletons where we are listed in their compatibility list
 		{
-			RetargetingTable[TargetBoneIndex] = MakeTuple(FQuat::Identity, FQuat::Identity);
+			TArray<FAssetData> SkeletonsCompatibleWithThis;
+			FARFilter ARFilter;
+			ARFilter.ClassPaths.Add(USkeleton::StaticClass()->GetClassPathName());
+			ARFilter.TagsAndValues.Add(CompatibleSkeletonsNameTag, AssetDataThis.GetExportTextName());
+			AssetRegistry.GetAssets(ARFilter, SkeletonsCompatibleWithThis);
+
+			OutAssets.Append(SkeletonsCompatibleWithThis);
 		}
 	}
-
-	GenerateCurveMapping();
 }
 
-void FSkeletonRemapping::ComposeWith(const FSkeletonRemapping& OtherSkeletonRemapping)
+void USkeleton::GetCompatibleAssets(UClass* AssetClass, const TCHAR* InTag, TArray<FAssetData>& OutAssets) const
 {
-	check(OtherSkeletonRemapping.SourceSkeleton == TargetSkeleton);
-
-	TargetSkeleton = OtherSkeletonRemapping.TargetSkeleton;
-
-	const int32 SourceNumBones = SourceToTargetBoneIndexes.Num();
-	const int32 TargetNumBones = OtherSkeletonRemapping.TargetToSourceBoneIndexes.Num();
-
-	TArray<int32> NewTargetToSourceBoneIndexes;
-	TArray<TTuple<FQuat, FQuat>> NewRetargetingTable;
-
-	NewTargetToSourceBoneIndexes.SetNumUninitialized(TargetNumBones);
-	NewRetargetingTable.SetNumUninitialized(TargetNumBones);
-
-	// Compose the retargeting constants
-	for (int32 NewTargetBoneIndex = 0; NewTargetBoneIndex < TargetNumBones; ++NewTargetBoneIndex)
+	const IAssetRegistry& AssetRegistry = FModuleManager::LoadModuleChecked<FAssetRegistryModule>("AssetRegistry").Get();
+	
+	if(AreAllSkeletonsCompatibleDelegate.IsBound() && AreAllSkeletonsCompatibleDelegate.Execute())
 	{
-		const int32 OldTargetBoneIndex = OtherSkeletonRemapping.TargetToSourceBoneIndexes[NewTargetBoneIndex];
-		const int32 OldSourceBoneIndex = (OldTargetBoneIndex != INDEX_NONE) ? TargetToSourceBoneIndexes[OldTargetBoneIndex] : INDEX_NONE;
-
-		if (OldSourceBoneIndex != INDEX_NONE)
-		{
-			const TTuple<FQuat, FQuat>& OldQQ = RetargetingTable[OldTargetBoneIndex];
-			const TTuple<FQuat, FQuat>& NewQQ = OtherSkeletonRemapping.RetargetingTable[NewTargetBoneIndex];
-			const FQuat Q0 = NewQQ.Get<0>() * OldQQ.Get<0>();
-			const FQuat Q1 = OldQQ.Get<1>() * NewQQ.Get<1>();
-
-			NewTargetToSourceBoneIndexes[NewTargetBoneIndex] = OldSourceBoneIndex;
-			NewRetargetingTable[NewTargetBoneIndex] = MakeTuple(Q0, Q1);
-		}
-		else
-		{
-			NewTargetToSourceBoneIndexes[NewTargetBoneIndex] = INDEX_NONE;
-			NewRetargetingTable[NewTargetBoneIndex] = MakeTuple(FQuat::Identity, FQuat::Identity);
-		}
+		// Compatible with all, so return all assets of type
+		FARFilter ARFilter;
+		ARFilter.ClassPaths.Add(AssetClass->GetClassPathName());
+		ARFilter.bRecursiveClasses = true;
+		AssetRegistry.GetAssets(ARFilter, OutAssets);
 	}
+	else
+	{
+		FARFilter ARFilter;
+		ARFilter.ClassPaths.Add(AssetClass->GetClassPathName());
+		ARFilter.bRecursiveClasses = true;
 
-	TargetToSourceBoneIndexes = MoveTemp(NewTargetToSourceBoneIndexes);
-	RetargetingTable = MoveTemp(NewRetargetingTable);
+		TArray<FAssetData> CompatibleSkeletonAssets;
+		GetCompatibleSkeletonAssets(CompatibleSkeletonAssets);
 
-	// Rebuild the mapping from source bones to the target bones
-	FMemory::Memset(SourceToTargetBoneIndexes.GetData(), static_cast<uint8>(INDEX_NONE), SourceNumBones * SourceToTargetBoneIndexes.GetTypeSize());
-	for (int32 TargetBoneIndex = 0; TargetBoneIndex < TargetNumBones; ++TargetBoneIndex)
-	{
-		const int32 SourceBoneIndex = TargetToSourceBoneIndexes[TargetBoneIndex];
-		if (SourceBoneIndex != INDEX_NONE)
+		for (const FAssetData& CompatibleSkeleton : CompatibleSkeletonAssets)
 		{
-			SourceToTargetBoneIndexes[SourceBoneIndex] = TargetBoneIndex;
+			ARFilter.TagsAndValues.Add(InTag, CompatibleSkeleton.GetExportTextName());
 		}
+
+		AssetRegistry.GetAssets(ARFilter, OutAssets);
 	}
 }
 
-void FSkeletonRemapping::GenerateCurveMapping()
-{
-	const USkeleton* Source = GetSourceSkeleton().Get();
-	const USkeleton* Target = GetTargetSkeleton().Get();
-	check(Source && Target);
-
-	// Get the curve mappings.
-	const FSmartNameMapping* SourceMapping = Source->GetSmartNameContainer(USkeleton::AnimCurveMappingName);
-	const FSmartNameMapping* TargetMapping = Target->GetSmartNameContainer(USkeleton::AnimCurveMappingName);
-	if (!SourceMapping) // Source has no curves, so there is nothing to add.
-	{
-		return;
-	}
+#endif
 
-	// Check whether we have curves or not in both source and target.
-	bool bHasSourceCurves = false;
-	bool bHasTargetCurves = false;
-	if (TargetMapping)
+bool USkeleton::IsCompatible(const USkeleton* InSkeleton) const
+{
+#if WITH_EDITORONLY_DATA
+	return IsCompatibleForEditor(InSkeleton);
+#else
+	if (InSkeleton == nullptr)
 	{
-		bHasTargetCurves = (TargetMapping->GetMaxUID() != SmartName::MaxUID);
-		bHasSourceCurves = (SourceMapping->GetMaxUID() != SmartName::MaxUID);
+		return false;
 	}
 
-	// Init the mapping table.
-	const SmartName::UID_Type NumItems = SourceMapping->GetMaxUID() + 1;
-	if (bHasSourceCurves)
-	{
-		SourceToTargetCurveMapping.Init(MAX_uint16, NumItems);
-	}
+	return true;
+#endif
+}
 
-	// If we have no source or target curves (while having valid smart name mappings), there is nothing to do.
-	if (!bHasTargetCurves && !bHasSourceCurves)
-	{
-		return;
-	}
+bool USkeleton::IsCompatibleSkeletonByAssetString(const FString& SkeletonAssetString) const
+{
+#if WITH_EDITORONLY_DATA
+	return IsCompatibleForEditor(SkeletonAssetString);
+#else
+	return false;
+#endif
+}
 
-	// For every source curve, try to find the curve with the same name in the target.
-	for (SmartName::UID_Type Index = 0; Index < NumItems; ++Index)
-	{
-		FName CurveName;
-		SourceMapping->GetName(Index, CurveName);
+void USkeleton::AddCompatibleSkeleton(const USkeleton* SourceSkeleton)
+{
+	CompatibleSkeletons.AddUnique(SourceSkeleton);
+}
 
-		// Make sure the curve name is valid and we can actually find a curve with the same name in the target.
-		FSmartName TargetSmartName;
-		if (!CurveName.IsValid() || !TargetMapping->FindSmartName(CurveName, TargetSmartName))
-		{
-			continue;
-		}
+void USkeleton::RemoveCompatibleSkeleton(const USkeleton* SourceSkeleton)
+{
+	CompatibleSkeletons.Remove(SourceSkeleton);
+}
 
-		SourceToTargetCurveMapping[Index] = TargetSmartName.UID;
-	}
+bool USkeleton::IsCompatibleSkeletonByAssetData(const FAssetData& AssetData, const TCHAR* InTag) const
+{
+#if WITH_EDITORONLY_DATA
+	return IsCompatibleForEditor(AssetData, InTag);
+#else
+	return false;
+#endif
 }
 
 USkeleton::USkeleton(const FObjectInitializer& ObjectInitializer)
@@ -469,44 +316,16 @@ USkeleton::USkeleton(const FObjectInitializer& ObjectInitializer)
 	{
 		FCoreUObjectDelegates::OnPackageReloaded.AddStatic(&USkeleton::HandlePackageReloaded);
 	}
-	else
-	{
-		if (!HasAnyFlags(RF_NeedPostLoad))
-		{
-			FScopeLock Lock(&LoadedSkeletonsMutex);
-			LoadedSkeletons.Add(this);
-		}
-	}
 }
 
 void USkeleton::BeginDestroy()
 {
-	//  Remove the loaded skeleton from the list.
-	{
-		FScopeLock Lock(&LoadedSkeletonsMutex);
-		LoadedSkeletons.Remove(this);
-
-		// Also unregister this skeleton from other skeletons this one is listening to (events it is listening to).
-		for (USkeleton* Skeleton : LoadedSkeletons)
-		{
-			if (Skeleton->IsCompatible(this))
-			{
-				Skeleton->OnSkeletonDestructEvent.RemoveAll(this);
-				Skeleton->RemoveCompatibleSkeleton(this);
-			}
-		}
-	}
-
 	Super::BeginDestroy();
 	if (HasAnyFlags(RF_ClassDefaultObject))
 	{
 		FCoreUObjectDelegates::OnPackageReloaded.RemoveAll(this);
 	}
 
-	// Clear our own skeleton mappings.
-	ClearSkeletonRemappings();
-	OnSkeletonDestructEvent.Broadcast(this); // Inform others about this skeleton being destructed, so we can remove mappings inside those skeletons.
-	OnSkeletonDestructEvent.Clear();
 	OnSmartNamesChangedEvent.Clear();
 }
 
@@ -551,45 +370,7 @@ void USkeleton::PostLoad()
 	{
 		CurveMappingTable->InitializeCurveMetaData(this);
 	}
-
-	// Add this skeleton to the loaded skeletons list and register some event listeners.
-	{
-		FScopeLock Lock(&LoadedSkeletonsMutex);
-
-		// Make any loaded skeleton that is compatible with this skeleton listen to destruct events.
-		// This is done so when this skeleton gets destructed, other skeletons can remove mappings to this skeleton.
-		for (USkeleton* Skeleton : LoadedSkeletons)
-		{
-			if (!Skeleton || Skeleton == this)
-			{
-				continue;
-			}
-
-			// If the already loaded other skeleton is compatible with this newly loaded skeleton, then let the
-			// already loaded other skeleton listen to the destruction event of this newly loaded skeleton.
-			if (Skeleton->IsCompatible(this))
-			{
-				Skeleton->CreateSkeletonRemappingIfNeeded(this);
-				if (!OnSkeletonDestructEvent.IsBoundToObject(Skeleton))
-				{
-					OnSkeletonDestructEvent.AddUObject(Skeleton, &USkeleton::HandleSkeletonDestruct);
-				}
-			}
-
-			// If we are compatible with the already loaded skeleton, we want to start listening to destruct events of that one.
-			if (IsCompatible(Skeleton))
-			{
-				CreateSkeletonRemappingIfNeeded(Skeleton);
-				if (!Skeleton->OnSkeletonDestructEvent.IsBoundToObject(this))
-				{
-					Skeleton->OnSkeletonDestructEvent.AddUObject(this, &USkeleton::HandleSkeletonDestruct);
-				}
-			}
-		}
-
-		LoadedSkeletons.Add(this);
-	}
-
+	
 	// Listen to smart name changes so we can regenerate our skeleton mappings.
 	if (!OnSmartNamesChangedEvent.IsBoundToObject(this))
 	{
@@ -606,71 +387,7 @@ void USkeleton::PostLoad()
 // Regenerate all required skeleton remappings whenever curves change.
 void USkeleton::HandleSmartNamesChangedEvent()
 {
-	FScopeLock Lock(&LoadedSkeletonsMutex);
-	for (USkeleton* Skeleton : LoadedSkeletons)
-	{
-		if (Skeleton == this)
-		{
-			continue;
-		}
-
-		if (!Skeleton)
-		{
-			continue;
-		}
-
-		if (Skeleton->IsCompatible(this))
-		{
-			Skeleton->CreateSkeletonRemappingIfNeeded(this, true);
-		}
-
-		if (IsCompatible(Skeleton))
-		{
-			CreateSkeletonRemappingIfNeeded(Skeleton, true);
-		}
-	}
-}
-
-void USkeleton::HandleSkeletonDestruct(const USkeleton* Skeleton)
-{
-	OnSkeletonDestructEvent.RemoveAll(Skeleton);
-
-	// We aren't interested to hear that we loaded ourselves as we need no mappings to ourselves.
-	// Also if this skeleton isn't in the compatible list, we can ignore it.
-	if (Skeleton == this || !Skeleton || !IsCompatible(Skeleton))
-	{
-		return;
-	}
-
-	// Check if we already have a mapping. If so, remove the existing one and create a new one.
-	// This might only happen when a skeleton asset gets reimported.
-	RemoveSkeletonRemapping(Skeleton);
-}
-
-void USkeleton::CreateSkeletonRemappingIfNeeded(const USkeleton* SourceSkeleton, bool bRebuildIfExists)
-{
-	FScopeLock Lock(&SkeletonRemappingMutex);
-	check(SourceSkeleton != nullptr && SourceSkeleton != this && IsCompatible(SourceSkeleton));
-
-	// Check if we already have a mapping. If so, remove the existing one and create a new one.
-	// This might only happen when a skeleton asset gets reimported.
-	FSkeletonRemapping** Mapping = SkeletonRemappings.Find(SourceSkeleton);
-	if (Mapping)	
-	{
-		if (!bRebuildIfExists)
-		{
-			return;
-		}
-
-		delete *Mapping;
-		SkeletonRemappings.Remove(SourceSkeleton);
-	}
-
-	// Create, add and return the new entry.
-	FSkeletonRemapping* NewMapping = new FSkeletonRemapping(SourceSkeleton, this);
-	SkeletonRemappings.Add(SourceSkeleton, NewMapping);
-
-	UE_LOG(LogAnimation, Verbose, TEXT("Recreating remapping between %s and %s"), *FAssetData(this).AssetName.ToString(), *FAssetData(SourceSkeleton).AssetName.ToString());
+	UE::Anim::FSkeletonRemappingRegistry::Get().RefreshCurveMappings(this);
 }
 
 void USkeleton::PostDuplicate(bool bDuplicateForPIE)
@@ -682,10 +399,6 @@ void USkeleton::PostDuplicate(bool bDuplicateForPIE)
 		// regenerate Guid
 		RegenerateGuid();
 	}
-
-	// Initialize skeleton remappings.
-	ClearSkeletonRemappings();
-	BuildSkeletonRemappings(true);
 }
 
 void USkeleton::Serialize( FArchive& Ar )
@@ -797,18 +510,6 @@ void USkeleton::Serialize( FArchive& Ar )
 }
 
 #if WITH_EDITOR
-void USkeleton::PostEditChangeProperty(FPropertyChangedEvent& PropertyChangedEvent)
-{
-	Super::PostEditChangeProperty(PropertyChangedEvent);
-
-	// Rebuild skeleton mappings involving this skeleton when we modified the compatible skeleton list.
-	if (PropertyChangedEvent.Property && PropertyChangedEvent.Property->GetFName() == GET_MEMBER_NAME_CHECKED(USkeleton, CompatibleSkeletons))
-	{
-		ClearSkeletonRemappings();
-		BuildSkeletonRemappings(false); // Build it, but not bidirectional, meaning only rebuild mappings to skeletons we are compatible with.
-	}
-}
-
 void USkeleton::PreEditUndo()
 {
 	// Undoing so clear cached data as it will now be stale
@@ -853,71 +554,9 @@ void USkeleton::ConvertToFReferenceSkeleton()
 #endif
 }
 
-void USkeleton::ClearSkeletonRemappings()
-{
-	FScopeLock Lock(&SkeletonRemappingMutex);
-	for (auto& Mapping : SkeletonRemappings)
-	{
-		delete Mapping.Value;
-	}
-
-	SkeletonRemappings.Reset();
-}
-
 const FSkeletonRemapping* USkeleton::GetSkeletonRemapping(const USkeleton* SourceSkeleton) const
 {
-	check(SourceSkeleton);
-	if (SourceSkeleton == this)
-	{
-		return nullptr;
-	}
-
-	// Find the existing one.
-	FScopeLock Lock(&SkeletonRemappingMutex);
-	const FSkeletonRemapping* const* Mapping = SkeletonRemappings.Find(SourceSkeleton);
-	if (Mapping)
-	{
-		return *Mapping;
-	}
-
-	return nullptr;
-}
-
-void USkeleton::BuildSkeletonRemappings(bool bBidirectional)
-{
-	for (USkeleton* Skeleton : LoadedSkeletons)
-	{
-		if (Skeleton == this)
-		{
-			continue;
-		}
-
-		// Update the remapping in the other skeleton.
-		if (bBidirectional && Skeleton->IsCompatible(this))
-		{
-			Skeleton->CreateSkeletonRemappingIfNeeded(this, true);
-		}
-
-		// Update the remapping to this skeleton inside our own remapping table.
-		if (IsCompatible(Skeleton))
-		{
-			CreateSkeletonRemappingIfNeeded(Skeleton, true);
-		}
-	}
-}
-
-void USkeleton::RemoveSkeletonRemapping(const USkeleton* SourceSkeleton)
-{
-	check(SourceSkeleton && SourceSkeleton != this);
-
-	// Locate the mapping and remove if found.
-	FScopeLock Lock(&SkeletonRemappingMutex);
-	FSkeletonRemapping** Mapping = SkeletonRemappings.Find(SourceSkeleton);
-	if (Mapping)
-	{
-		delete *Mapping;
-		SkeletonRemappings.Remove(SourceSkeleton);
-	}
+	return &UE::Anim::FSkeletonRemappingRegistry::Get().GetRemapping(SourceSkeleton, this);
 }
 
 bool USkeleton::DoesParentChainMatch(int32 StartBoneIndex, const USkinnedAsset* InSkinnedAsset) const
@@ -1361,6 +1000,18 @@ void USkeleton::SetBoneTranslationRetargetingMode(const int32 BoneIndex, EBoneTr
 	}
 }
 
+const TArray<FTransform>& USkeleton::GetRefLocalPoses( FName RetargetSource ) const 
+{
+	if ( RetargetSource != NAME_None ) 
+	{
+		if (const FReferencePose* FoundRetargetSource = AnimRetargetSources.Find(RetargetSource))
+		{
+			return FoundRetargetSource->ReferencePose;
+		}
+	}
+	return ReferenceSkeleton.GetRefBonePose();
+}
+
 #if WITH_EDITORONLY_DATA
 
 FName USkeleton::GetRetargetSourceForMesh(USkinnedAsset* InSkinnedAsset) const
@@ -1377,6 +1028,14 @@ FName USkeleton::GetRetargetSourceForMesh(USkinnedAsset* InSkinnedAsset) const
 	return NAME_None;
 }
 
+void USkeleton::GetRetargetSources(TArray<FName>& OutRetargetSources) const
+{
+	for(const TPair<FName, FReferencePose>& AnimRetargetSource : AnimRetargetSources)
+	{
+		OutRetargetSources.Add(AnimRetargetSource.Key);
+	}
+}
+
 #endif
 
 int32 USkeleton::GetRawAnimationTrackIndex(const int32 InSkeletonBoneIndex, const UAnimSequence* InAnimSeq)
@@ -1432,7 +1091,7 @@ USkeletalMesh* USkeleton::GetPreviewMesh(bool bFindIfNotSet/*=false*/)
 #if WITH_EDITORONLY_DATA
 	USkeletalMesh* PreviewMesh = PreviewSkeletalMesh.LoadSynchronous();
 
-	if(PreviewMesh && !IsCompatible(PreviewMesh->GetSkeleton())) // fix mismatched skeleton
+	if(PreviewMesh && !IsCompatibleForEditor(PreviewMesh->GetSkeleton())) // fix mismatched skeleton
 	{
 		PreviewSkeletalMesh.Reset();
 		PreviewMesh = nullptr;
@@ -1690,7 +1349,7 @@ void USkeleton::HandleSkeletonHierarchyChange()
 	}
 
 	// Full rebuild of all compatible with this and with ones we are compatible with.
-	BuildSkeletonRemappings(true);
+	UE::Anim::FSkeletonRemappingRegistry::Get().RefreshMappings(this);
 
 	// Fix up loaded animations (any animations that aren't loaded will be fixed on load)
 	int32 NumLoadedAssets = 0;
@@ -2355,7 +2014,7 @@ void USkeleton::HandleVirtualBoneChanges()
 	const bool bRebuildNameMap = false;
 	ReferenceSkeleton.RebuildRefSkeleton(this, bRebuildNameMap);
 
-	BuildSkeletonRemappings(true);
+	UE::Anim::FSkeletonRemappingRegistry::Get().RefreshMappings(this);
 
 	for (TObjectIterator<USkeletalMesh> ItMesh; ItMesh; ++ItMesh)
 	{
@@ -2555,6 +2214,17 @@ void USkeleton::GetAssetRegistryTags(TArray<FAssetRegistryTag>& OutTags) const
 	FString RigFullName = (RigConfig.Rig)? RigConfig.Rig->GetFullName() : TEXT("");
 
 	OutTags.Add(FAssetRegistryTag(USkeleton::RigTag, RigFullName, FAssetRegistryTag::TT_Hidden));
+
+	TStringBuilder<256> CompatibleSkeletonsBuilder;
+	
+	for(const TSoftObjectPtr<USkeleton>& CompatibleSkeleton : CompatibleSkeletons)
+	{
+		const FString ExportPath = FObjectPropertyBase::GetExportPath(USkeleton::StaticClass()->GetClassPathName(), CompatibleSkeleton.ToSoftObjectPath().GetAssetPathString());
+		CompatibleSkeletonsBuilder.Append(ExportPath);
+		CompatibleSkeletonsBuilder.Append(USkeleton::CompatibleSkeletonsTagDelimiter);
+	}
+
+	OutTags.Add(FAssetRegistryTag(USkeleton::CompatibleSkeletonsNameTag, CompatibleSkeletonsBuilder.ToString(), FAssetRegistryTag::TT_Hidden));
 }
 
 #endif //WITH_EDITOR
diff --git a/Engine/Source/Runtime/Engine/Private/Animation/SkeletonRemapping.cpp b/Engine/Source/Runtime/Engine/Private/Animation/SkeletonRemapping.cpp
new file mode 100644
--- /dev/null
+++ b/Engine/Source/Runtime/Engine/Private/Animation/SkeletonRemapping.cpp
@@ -0,0 +1,297 @@
+// Copyright Epic Games, Inc. All Rights Reserved.
+
+#include "SkeletonRemapping.h"
+
+#include "AnimationRuntime.h"
+#include "Animation/Skeleton.h"
+
+FSkeletonRemapping::FSkeletonRemapping(const USkeleton* InSourceSkeleton, const USkeleton* InTargetSkeleton)
+	: SourceSkeleton(InSourceSkeleton)
+	, TargetSkeleton(InTargetSkeleton)
+{
+	GenerateMapping();
+}
+
+void FSkeletonRemapping::RegenerateMapping()
+{
+	SourceToTargetBoneIndexes.Reset();
+	TargetToSourceBoneIndexes.Reset();
+	RetargetingTable.Reset();
+
+	GenerateMapping();
+}
+
+void FSkeletonRemapping::GenerateMapping()
+{
+	if (!GetSourceSkeleton().IsValid() || !GetTargetSkeleton().IsValid())
+	{
+		return;
+	}
+	
+	const FReferenceSkeleton& SourceReferenceSkeleton = SourceSkeleton->GetReferenceSkeleton();
+	const FReferenceSkeleton& TargetReferenceSkeleton = TargetSkeleton->GetReferenceSkeleton();
+
+	const int32 SourceNumBones = SourceReferenceSkeleton.GetNum();
+	const int32 TargetNumBones = TargetReferenceSkeleton.GetNum();
+
+	SourceToTargetBoneIndexes.SetNumUninitialized(SourceNumBones);
+	TargetToSourceBoneIndexes.SetNumUninitialized(TargetNumBones);
+	RetargetingTable.SetNumUninitialized(TargetNumBones);
+
+	TArrayView<const FTransform> SourceLocalTransforms;
+	TArrayView<const FTransform> TargetLocalTransforms;
+/*
+	// Build the mapping from source to target bones through the remapping rig if one exists
+	if (true)	// TODO_SKELETON_REMAPPING: Use a remapping rig if it exists
+	{
+		for (int32 SourceBoneIndex = 0; SourceBoneIndex < SourceNumBones; ++SourceBoneIndex)
+		{
+			const FName BoneName = SourceReferenceSkeleton.GetBoneName(SourceBoneIndex);
+			int32 TargetBoneIndex = TargetReferenceSkeleton.FindBoneIndex(BoneName);
+			SourceToTargetBoneIndexes[SourceBoneIndex] = TargetBoneIndex;
+		}
+
+		// Get the matched source and target rest poses from the remapping rig
+		//
+		// TODO_SKELETON_REMAPPING: For now we'll just use the skeleton's rest poses, but we should really be pulling these
+		// poses from a remapping rig so that the use can better align the poses.  Note that we'll want the remapping rig
+		// to be independent from the USkeleton itself because a skeleton may want to participate in multiple remappings with
+		// different other skeletons and may have to pose itself differently to align with the different other skeletons
+		//
+		SourceLocalTransforms = MakeArrayView(SourceReferenceSkeleton.GetRefBonePose());
+		TargetLocalTransforms = MakeArrayView(TargetReferenceSkeleton.GetRefBonePose());
+	}
+	else // Fall back to simple name matching if there is no rig
+*/
+	{
+		// Match source to target bones by name lookup between source and target skeletons
+		for (int32 SourceBoneIndex = 0; SourceBoneIndex < SourceNumBones; ++SourceBoneIndex)
+		{
+			const FName BoneName = SourceReferenceSkeleton.GetBoneName(SourceBoneIndex);
+			const int32 TargetBoneIndex = TargetReferenceSkeleton.FindBoneIndex(BoneName);
+			SourceToTargetBoneIndexes[SourceBoneIndex] = TargetBoneIndex;
+		}
+
+		// Get the matched (hopefully) source and target rest poses from the source and target skeletons
+		SourceLocalTransforms = MakeArrayView(SourceReferenceSkeleton.GetRefBonePose());
+		TargetLocalTransforms = MakeArrayView(TargetReferenceSkeleton.GetRefBonePose());
+	}
+
+	// Force the roots to map onto each other regardless of their names
+	SourceToTargetBoneIndexes[0] = 0;
+
+	// Build the reverse mapping from target back to source bones
+	FMemory::Memset(TargetToSourceBoneIndexes.GetData(), static_cast<uint8>(INDEX_NONE), TargetNumBones * TargetToSourceBoneIndexes.GetTypeSize());
+	for (int32 SourceBoneIndex = 0; SourceBoneIndex < SourceNumBones; ++SourceBoneIndex)
+	{
+		const int32 TargetBoneIndex = SourceToTargetBoneIndexes[SourceBoneIndex];
+		if (TargetBoneIndex != INDEX_NONE)
+		{
+			TargetToSourceBoneIndexes[TargetBoneIndex] = SourceBoneIndex;
+		}
+	}
+
+	TArray<FTransform> SourceComponentTransforms;
+	TArray<FTransform> TargetComponentTransforms;
+	FAnimationRuntime::FillUpComponentSpaceTransforms(SourceReferenceSkeleton, SourceLocalTransforms, SourceComponentTransforms);
+	FAnimationRuntime::FillUpComponentSpaceTransforms(TargetReferenceSkeleton, TargetLocalTransforms, TargetComponentTransforms);
+
+	// Calculate the retargeting constants to map from source skeleton space to target skeleton space
+	//
+	// Simply remapping joint indices is usually not sufficient to give us the desired result pose if the source and target
+	// skeletons are built with different conventions for joint orientations.  We therefore need to compute a remapping
+	// between the source and target joint orientations, which we do in terms of delta rotations from the rest pose:
+	//
+	//		Q = P * D * R
+	//
+	// where:
+	//
+	//		Q is the final joint orientation in component space
+	//		P is the parent joint orientation in component space
+	//		D is the delta rotation that we want to remap
+	//		R is the local space rest pose orientation for the joint
+	//
+	// We want to find a mapping from the source to the target:
+	//
+	//		Ps * Ds * Rs  -->  Pt * Dt * Rt
+	//
+	// such that the deltas produce equivalent rotations on the mesh even if their parent rotation frames are different.
+	// In other words, we need to find Dt such that its component space rotation is equivalent to Ds.  To convert a rotation
+	// from local space (D) to component space (C), given its parent (P), we have:
+	//
+	//		C * P = P * D
+	//		C = P * D * P⁻¹
+	//
+	// Setting the source and target component space rotations to be equal (Cs = Ct) then gives us:
+	//
+	//		Ps * Ds * Ps⁻¹ = Pt * Dt * Pt⁻¹
+	//
+	// which we then solve for Dt to get:
+	//
+	//		Dt = Pt⁻¹ * Ps * Ds * Ps⁻¹ * Pt
+	//
+	// However, when we're remapping an animation pose, we will have the local transforms rather than the deltas from the
+	// rest pose, so we also need to convert between these local transforms and the equivalent deltas:
+	//
+	//		L = D * R
+	//		D = L * R⁻¹
+	//
+	// Combining that with our equation for Dt, we get:
+	//
+	//		Lt * Rt⁻¹ = Pt⁻¹ * Ps * Ls * Rs⁻¹ * Ps⁻¹ * Pt
+	//
+	// Solving for Lt then gives us:
+	//
+	//		Lt = Pt⁻¹ * Ps * Ls * Rs⁻¹ * Ps⁻¹ * Pt * Rt
+	//
+	// Finally, factoring out the constant terms (which we precompute here) gives us:
+	//
+	//		Lt = Q0 * Ls * Q1
+	//		Q0 = Pt⁻¹ * Ps
+	//		Q1 = Rs⁻¹ * Ps⁻¹ * Pt * Rt
+	//
+	// Note that when remapping additive animations, we drop the rest pose terms, but we still need to convert between the
+	// source and target rotation frames.  Dropping Rs and Rt from the equations above gives us:
+	//
+	//		Lt = Pt⁻¹ * Ps * Ls * Ps⁻¹ * Pt			(for additive animations)
+	//
+	// which is equivalent to the following in terms of our precomputed constants:
+	//
+	//		Lt = Q0 * Ls * Q0⁻¹						(for additive animations)
+	//
+	RetargetingTable[0] = MakeTuple(FQuat::Identity, SourceLocalTransforms[0].GetRotation().Inverse() * TargetLocalTransforms[0].GetRotation());
+	for (int32 TargetBoneIndex = 1; TargetBoneIndex < TargetNumBones; ++TargetBoneIndex)
+	{
+		const int32 SourceBoneIndex = TargetToSourceBoneIndexes[TargetBoneIndex];
+		if (SourceBoneIndex != INDEX_NONE)
+		{
+			const int32 SourceParentIndex = SourceReferenceSkeleton.GetParentIndex(SourceBoneIndex);
+			const int32 TargetParentIndex = TargetReferenceSkeleton.GetParentIndex(TargetBoneIndex);
+			check(SourceParentIndex != INDEX_NONE);
+			check(TargetParentIndex != INDEX_NONE);
+
+			const FQuat PS = SourceComponentTransforms[SourceParentIndex].GetRotation();
+			const FQuat PT = TargetComponentTransforms[TargetParentIndex].GetRotation();
+
+			const FQuat RS = SourceLocalTransforms[SourceBoneIndex].GetRotation();
+			const FQuat RT = TargetLocalTransforms[TargetBoneIndex].GetRotation();
+
+			const FQuat Q0 = PT.Inverse() * PS;
+			const FQuat Q1 = RS.Inverse() * PS.Inverse() * PT * RT;
+
+			RetargetingTable[TargetBoneIndex] = MakeTuple(Q0, Q1);
+		}
+		else
+		{
+			RetargetingTable[TargetBoneIndex] = MakeTuple(FQuat::Identity, FQuat::Identity);
+		}
+	}
+
+	GenerateCurveMapping();
+}
+
+void FSkeletonRemapping::ComposeWith(const FSkeletonRemapping& OtherSkeletonRemapping)
+{
+	check(OtherSkeletonRemapping.SourceSkeleton == TargetSkeleton);
+
+	TargetSkeleton = OtherSkeletonRemapping.TargetSkeleton;
+
+	const int32 SourceNumBones = SourceToTargetBoneIndexes.Num();
+	const int32 TargetNumBones = OtherSkeletonRemapping.TargetToSourceBoneIndexes.Num();
+
+	TArray<int32> NewTargetToSourceBoneIndexes;
+	TArray<TTuple<FQuat, FQuat>> NewRetargetingTable;
+
+	NewTargetToSourceBoneIndexes.SetNumUninitialized(TargetNumBones);
+	NewRetargetingTable.SetNumUninitialized(TargetNumBones);
+
+	// Compose the retargeting constants
+	for (int32 NewTargetBoneIndex = 0; NewTargetBoneIndex < TargetNumBones; ++NewTargetBoneIndex)
+	{
+		const int32 OldTargetBoneIndex = OtherSkeletonRemapping.TargetToSourceBoneIndexes[NewTargetBoneIndex];
+		const int32 OldSourceBoneIndex = (OldTargetBoneIndex != INDEX_NONE) ? TargetToSourceBoneIndexes[OldTargetBoneIndex] : INDEX_NONE;
+
+		if (OldSourceBoneIndex != INDEX_NONE)
+		{
+			const TTuple<FQuat, FQuat>& OldQQ = RetargetingTable[OldTargetBoneIndex];
+			const TTuple<FQuat, FQuat>& NewQQ = OtherSkeletonRemapping.RetargetingTable[NewTargetBoneIndex];
+			const FQuat Q0 = NewQQ.Get<0>() * OldQQ.Get<0>();
+			const FQuat Q1 = OldQQ.Get<1>() * NewQQ.Get<1>();
+
+			NewTargetToSourceBoneIndexes[NewTargetBoneIndex] = OldSourceBoneIndex;
+			NewRetargetingTable[NewTargetBoneIndex] = MakeTuple(Q0, Q1);
+		}
+		else
+		{
+			NewTargetToSourceBoneIndexes[NewTargetBoneIndex] = INDEX_NONE;
+			NewRetargetingTable[NewTargetBoneIndex] = MakeTuple(FQuat::Identity, FQuat::Identity);
+		}
+	}
+
+	TargetToSourceBoneIndexes = MoveTemp(NewTargetToSourceBoneIndexes);
+	RetargetingTable = MoveTemp(NewRetargetingTable);
+
+	// Rebuild the mapping from source bones to the target bones
+	FMemory::Memset(SourceToTargetBoneIndexes.GetData(), static_cast<uint8>(INDEX_NONE), SourceNumBones * SourceToTargetBoneIndexes.GetTypeSize());
+	for (int32 TargetBoneIndex = 0; TargetBoneIndex < TargetNumBones; ++TargetBoneIndex)
+	{
+		const int32 SourceBoneIndex = TargetToSourceBoneIndexes[TargetBoneIndex];
+		if (SourceBoneIndex != INDEX_NONE)
+		{
+			SourceToTargetBoneIndexes[SourceBoneIndex] = TargetBoneIndex;
+		}
+	}
+}
+
+void FSkeletonRemapping::GenerateCurveMapping()
+{
+	const USkeleton* Source = GetSourceSkeleton().Get();
+	const USkeleton* Target = GetTargetSkeleton().Get();
+	check(Source && Target);
+
+	// Get the curve mappings.
+	const FSmartNameMapping* SourceMapping = Source->GetSmartNameContainer(USkeleton::AnimCurveMappingName);
+	const FSmartNameMapping* TargetMapping = Target->GetSmartNameContainer(USkeleton::AnimCurveMappingName);
+	if (!SourceMapping) // Source has no curves, so there is nothing to add.
+	{
+		return;
+	}
+
+	// Check whether we have curves or not in both source and target.
+	bool bHasSourceCurves = false;
+	bool bHasTargetCurves = false;
+	if (TargetMapping)
+	{
+		bHasTargetCurves = (TargetMapping->GetMaxUID() != SmartName::MaxUID);
+		bHasSourceCurves = (SourceMapping->GetMaxUID() != SmartName::MaxUID);
+	}
+
+	// Init the mapping table.
+	const SmartName::UID_Type NumItems = SourceMapping->GetMaxUID() + 1;
+	if (bHasSourceCurves)
+	{
+		SourceToTargetCurveMapping.Init(MAX_uint16, NumItems);
+	}
+
+	// If we have no source or target curves (while having valid smart name mappings), there is nothing to do.
+	if (!bHasTargetCurves && !bHasSourceCurves)
+	{
+		return;
+	}
+
+	// For every source curve, try to find the curve with the same name in the target.
+	for (SmartName::UID_Type Index = 0; Index < NumItems; ++Index)
+	{
+		FName CurveName;
+		SourceMapping->GetName(Index, CurveName);
+
+		// Make sure the curve name is valid and we can actually find a curve with the same name in the target.
+		FSmartName TargetSmartName;
+		if (!CurveName.IsValid() || !TargetMapping->FindSmartName(CurveName, TargetSmartName))
+		{
+			continue;
+		}
+
+		SourceToTargetCurveMapping[Index] = TargetSmartName.UID;
+	}
+}
\ No newline at end of file
diff --git a/Engine/Source/Runtime/Engine/Private/Animation/SkeletonRemapping.h b/Engine/Source/Runtime/Engine/Private/Animation/SkeletonRemapping.h
new file mode 100644
--- /dev/null
+++ b/Engine/Source/Runtime/Engine/Private/Animation/SkeletonRemapping.h
@@ -0,0 +1,196 @@
+// Copyright Epic Games, Inc. All Rights Reserved.
+
+#pragma once
+
+#include "CoreMinimal.h"
+#include "Animation/AnimTypes.h"
+
+class USkeleton;
+
+struct FSkeletonRemapping
+{
+	FSkeletonRemapping() = default;
+	FSkeletonRemapping(const USkeleton* InSourceSkeleton, const USkeleton* InTargetSkeleton);
+
+	const TWeakObjectPtr<const USkeleton>& GetSourceSkeleton() const { return SourceSkeleton; }
+	const TWeakObjectPtr<const USkeleton>& GetTargetSkeleton() const { return TargetSkeleton; }
+
+	/**
+	 * Compose this remapping with another remapping in place.  The other remapping's source skeleton must match
+	 * this remapping's target skeleton.  The result will map from this remapping's source skeleton to the other
+	 * remapping's target skeleton
+	 *
+	 * @param	OtherSkeletonRemapping		Skeleton remapping to compose into this remapping
+	 */
+	void ComposeWith(const FSkeletonRemapping& OtherSkeletonRemapping);
+
+	/**
+	 * Get the target skeleton bone index that corresponds to the specified bone on the source skeleton
+	 *
+	 * @param	SourceSkeletonBoneIndex		Skeleton bone index on the source skeleton
+	 * @return								Skeleton bone index on the target skeleton (or INDEX_NONE)
+	 */
+	int32 GetTargetSkeletonBoneIndex(int32 SourceSkeletonBoneIndex) const
+	{
+		return (SourceToTargetBoneIndexes.IsValidIndex(SourceSkeletonBoneIndex))
+				   ? SourceToTargetBoneIndexes[SourceSkeletonBoneIndex]
+				   : INDEX_NONE;
+	}
+
+	/**
+	 * Get the source skeleton bone index that corresponds to the specified bone on the target skeleton
+	 *
+	 * @param	TargetSkeletonBoneIndex		Skeleton bone index on the target skeleton
+	 * @return								Skeleton bone index on the source skeleton (or INDEX_NONE)
+	 */
+	int32 GetSourceSkeletonBoneIndex(int32 TargetSkeletonBoneIndex) const
+	{
+		return (TargetToSourceBoneIndexes.IsValidIndex(TargetSkeletonBoneIndex))
+				   ? TargetToSourceBoneIndexes[TargetSkeletonBoneIndex]
+				   : INDEX_NONE;
+	}
+
+	/**
+	 * Get the specified bone transform retargeted from the source skeleton onto the target skeleton, corrected
+	 * for differences between source and target rest poses
+	 *
+	 * @param	TargetSkeletonBoneIndex		Skeleton bone index on the target skeleton
+	 * @param	SourceTransform				Bone transform from the corresponding bone on the source skeleton
+	 * @return								Transform mapped onto the target skeleton
+	 */
+	FTransform RetargetBoneTransformToTargetSkeleton(int32 TargetSkeletonBoneIndex, const FTransform& SourceTransform) const
+	{
+		return FTransform(
+			RetargetBoneRotationToTargetSkeleton(TargetSkeletonBoneIndex, SourceTransform.GetRotation()),
+			RetargetBoneTranslationToTargetSkeleton(TargetSkeletonBoneIndex, SourceTransform.GetTranslation()),
+			SourceTransform.GetScale3D());
+	}
+
+	/**
+	 * Get the specified bone translation retargeted from the source skeleton onto the target skeleton, corrected
+	 * for differences between source and target rest poses
+	 *
+	 * @param	TargetSkeletonBoneIndex		Skeleton bone index on the target skeleton
+	 * @param	SourceTranslation			Bone translation from the corresponding bone on the source skeleton
+	 * @return								Translation mapped onto the target skeleton
+	 */
+	FVector RetargetBoneTranslationToTargetSkeleton(int32 TargetSkeletonBoneIndex, const FVector& SourceTranslation) const
+	{
+		// Compute the translation part of FTransform(Q1) * Source * FTransform(Q0)
+		const TTuple<FQuat, FQuat>& QQ = RetargetingTable[TargetSkeletonBoneIndex];
+		return QQ.Get<0>().RotateVector(SourceTranslation);
+	}
+
+	/**
+	 * Get the specified bone rotation retargeted from the source skeleton onto the target skeleton, corrected
+	 * for differences between source and target rest poses
+	 *
+	 * @param	TargetSkeletonBoneIndex		Skeleton bone index on the target skeleton
+	 * @param	SourceRotation				Bone rotation from the corresponding bone on the source skeleton
+	 * @return								Rotation mapped onto the target skeleton
+	 */
+	FQuat RetargetBoneRotationToTargetSkeleton(int32 TargetSkeletonBoneIndex, const FQuat& SourceRotation) const
+	{
+		// Compute the rotation part of FTransform(Q1) * Source * FTransform(Q0)
+		const TTuple<FQuat, FQuat>& QQ = RetargetingTable[TargetSkeletonBoneIndex];
+		return QQ.Get<0>() * SourceRotation * QQ.Get<1>();
+	}
+
+	/**
+	 * Get the specified additive transform retargeted from the source skeleton onto the target skeleton, corrected
+	 * for differences between source and target rest poses
+	 *
+	 * @param	TargetSkeletonBoneIndex		Skeleton bone index on the target skeleton
+	 * @param	SourceTransform				Bone transform from the corresponding bone on the source skeleton
+	 * @return								Transform mapped onto the target skeleton
+	 */
+	FTransform RetargetAdditiveTransformToTargetSkeleton(int32 TargetSkeletonBoneIndex, const FTransform& SourceTransform) const
+	{
+		return FTransform(
+			RetargetAdditiveRotationToTargetSkeleton(TargetSkeletonBoneIndex, SourceTransform.GetRotation()),
+			RetargetAdditiveTranslationToTargetSkeleton(TargetSkeletonBoneIndex, SourceTransform.GetTranslation()),
+			SourceTransform.GetScale3D());
+	}
+
+	/**
+	 * Get the specified additive translation retargeted from the source skeleton onto the target skeleton, corrected
+	 * for differences between source and target rest poses
+	 *
+	 * @param	TargetSkeletonBoneIndex		Skeleton bone index on the target skeleton
+	 * @param	SourceTranslation			Bone translation from the corresponding bone on the source skeleton
+	 * @return								Translation mapped onto the target skeleton
+	 */
+	FVector RetargetAdditiveTranslationToTargetSkeleton(int32 TargetSkeletonBoneIndex, const FVector& SourceTranslation) const
+	{
+		// Compute the translation part of FTransform(Q0.Inverse) * Source * FTransform(Q0)
+		const TTuple<FQuat, FQuat>& QQ = RetargetingTable[TargetSkeletonBoneIndex];
+		return QQ.Get<0>().RotateVector(SourceTranslation);
+	}
+
+	/**
+	 * Get the specified additive rotation retargeted from the source skeleton onto the target skeleton, corrected
+	 * for differences between source and target rest poses
+	 *
+	 * @param	TargetSkeletonBoneIndex		Skeleton bone index on the target skeleton
+	 * @param	SourceRotation				Bone rotation from the corresponding bone on the source skeleton
+	 * @return								Rotation mapped onto the target skeleton
+	 */
+	FQuat RetargetAdditiveRotationToTargetSkeleton(int32 TargetSkeletonBoneIndex, const FQuat& SourceRotation) const
+	{
+		// Compute the rotation part of FTransform(Q0.Inverse) * Source * FTransform(Q0)
+		const TTuple<FQuat, FQuat>& QQ = RetargetingTable[TargetSkeletonBoneIndex];
+		return QQ.Get<0>() * SourceRotation * QQ.Get<0>().Inverse();
+	}
+
+	/**
+	 * Get the curve mapping array, which maps from the source curve UID to the target curve UID.
+	 * You can access this array like this:
+	 * 
+	 * \code{.cpp}
+	 * const SmartName::UID_Type TargetCurveUID = GetSourceToTargetCurveMapping()[SourceCurveUID];
+	 * \endcode
+	 */
+	const TArray<SmartName::UID_Type>& GetSourceToTargetCurveMapping() const { return SourceToTargetCurveMapping; }
+
+	/** Refreshes the mapping. Empties mapping arrays and regenerates them against the stored skeletons. */
+	void RegenerateMapping();
+	
+	/**
+	 * Generates the mapping table for the curves, based on curve names.
+	 * Basically this will look at the curve names in the smart name table on the source, and tries to find matching ones in the table of 
+	 * the target skeleton. It then maps the UID's.
+	 */
+	void GenerateCurveMapping();
+
+	/** Check to see if a reference pose retarget is required between the source & target hierarchies */
+	bool RequiresReferencePoseRetarget() const
+	{
+		return RetargetingTable.Num() > 0;
+	}
+	
+	/** Check whether this mapping is valid - whether the source and target hierarchies still exist */
+	bool IsValid() const
+	{
+		return SourceSkeleton.IsValid() && TargetSkeleton.IsValid();
+	}
+
+private:
+	/** Internal helper function - performs mapping generation */
+	void GenerateMapping();
+	
+private:
+	TWeakObjectPtr<const USkeleton> SourceSkeleton;
+	TWeakObjectPtr<const USkeleton> TargetSkeleton;
+
+	// Table of target skeleton bone indexes (indexed by source skeleton bone index)
+	TArray<int32> SourceToTargetBoneIndexes;
+
+	// Table of source skeleton bone indexes (indexed by target skeleton bone index)
+	TArray<int32> TargetToSourceBoneIndexes;
+	
+	// Maps curve UIDs between source and target (indexed by source curve UID).
+	TArray<SmartName::UID_Type> SourceToTargetCurveMapping;
+
+	// Table of precalculated constants for retargeting from source to target (indexed by target skeleton bone index)
+	TArray<TTuple<FQuat, FQuat>> RetargetingTable;
+};
\ No newline at end of file
diff --git a/Engine/Source/Runtime/Engine/Private/Animation/SkeletonRemappingRegistry.cpp b/Engine/Source/Runtime/Engine/Private/Animation/SkeletonRemappingRegistry.cpp
new file mode 100644
--- /dev/null
+++ b/Engine/Source/Runtime/Engine/Private/Animation/SkeletonRemappingRegistry.cpp
@@ -0,0 +1,133 @@
+// Copyright Epic Games, Inc. All Rights Reserved.
+
+#include "SkeletonRemappingRegistry.h"
+#include "SkeletonRemapping.h"
+#include "Animation/Skeleton.h"
+
+namespace UE::Anim
+{
+
+FSkeletonRemappingRegistry* GSkeletonRemappingRegistry = nullptr;
+FSkeletonRemapping DefaultMapping;
+
+struct FSkeletonRemappingRegistryPrivate
+{
+	static FDelegateHandle PostGarbageCollectHandle;
+
+	static void HandlePostGarbageCollect()
+	{
+		// Compact the registry on GC
+		if(GSkeletonRemappingRegistry)
+		{
+			FRWScopeLock WriteLock(GSkeletonRemappingRegistry->MappingsLock, SLT_Write);
+			
+			for(auto Iter = GSkeletonRemappingRegistry->Mappings.CreateIterator(); Iter; ++Iter)
+			{
+				const FSkeletonRemappingRegistry::FWeakSkeletonPair& SkeletonPair = Iter.Key();
+				if(SkeletonPair.Key.Get() == nullptr || SkeletonPair.Value.Get() == nullptr)
+				{
+					Iter.RemoveCurrent();
+				}
+			}
+
+			for(auto Iter = GSkeletonRemappingRegistry->PerSkeletonMappings.CreateIterator(); Iter; ++Iter)
+			{
+				const TWeakObjectPtr<const USkeleton>& SkeletonPtr = Iter.Key();
+				if(SkeletonPtr.Get() == nullptr)
+				{
+					Iter.RemoveCurrent();
+				}
+			}
+		}
+	}
+};
+
+FDelegateHandle FSkeletonRemappingRegistryPrivate::PostGarbageCollectHandle;
+
+FSkeletonRemappingRegistry& FSkeletonRemappingRegistry::Get()
+{
+	checkf(GSkeletonRemappingRegistry, TEXT("Skeleton remapping registry is not instanced. It is only valid to access this while the engine module is loaded."));
+	return *GSkeletonRemappingRegistry;
+}
+
+void FSkeletonRemappingRegistry::Init()
+{
+	GSkeletonRemappingRegistry = new FSkeletonRemappingRegistry();
+	FSkeletonRemappingRegistryPrivate::PostGarbageCollectHandle = FCoreUObjectDelegates::GetPostGarbageCollect().AddStatic(&FSkeletonRemappingRegistryPrivate::HandlePostGarbageCollect);
+}
+
+void FSkeletonRemappingRegistry::Destroy()
+{
+	FCoreUObjectDelegates::GetPostGarbageCollect().Remove(FSkeletonRemappingRegistryPrivate::PostGarbageCollectHandle);
+	delete GSkeletonRemappingRegistry;
+	GSkeletonRemappingRegistry = nullptr;
+}
+
+const FSkeletonRemapping& FSkeletonRemappingRegistry::GetRemapping(const USkeleton* InSourceSkeleton, const USkeleton* InTargetSkeleton)
+{
+	if(InSourceSkeleton == InTargetSkeleton)
+	{
+		return DefaultMapping;
+	}
+	
+	const FWeakSkeletonPair Pair(InSourceSkeleton, InTargetSkeleton);
+
+	{
+		FRWScopeLock ReadLock(MappingsLock, SLT_ReadOnly);
+		const TSharedPtr<FSkeletonRemapping>* ExistingMapping = Mappings.Find(Pair);
+		if(ExistingMapping && ExistingMapping->IsValid())
+		{
+			return *(*ExistingMapping).Get();
+		}
+	}
+
+	// No existing, so create a new mapping
+	TSharedPtr<FSkeletonRemapping> NewMapping = MakeShared<FSkeletonRemapping>(InSourceSkeleton, InTargetSkeleton);
+	
+	{
+		FRWScopeLock WriteLock(MappingsLock, SLT_Write);
+
+		// Add to global mapping
+		Mappings.Add(Pair, NewMapping);
+
+		// Add to per-skeleton mappings
+		PerSkeletonMappings.Add(InSourceSkeleton, NewMapping);
+		PerSkeletonMappings.Add(InTargetSkeleton, NewMapping);
+
+		return *NewMapping;
+	}
+}
+
+void FSkeletonRemappingRegistry::RefreshCurveMappings(USkeleton* InSkeleton)
+{
+	TArray<TSharedPtr<FSkeletonRemapping>> ExistingMappings;
+	{
+		FRWScopeLock ReadLock(MappingsLock, SLT_ReadOnly);
+		PerSkeletonMappings.MultiFind(InSkeleton, ExistingMappings);
+	}
+	{
+		FRWScopeLock WriteLock(MappingsLock, SLT_Write);
+		for(const TSharedPtr<FSkeletonRemapping>& ExistingMapping : ExistingMappings)
+		{
+			ExistingMapping->GenerateCurveMapping();
+		}
+	}
+}
+
+void FSkeletonRemappingRegistry::RefreshMappings(USkeleton* InSkeleton)
+{
+	TArray<TSharedPtr<FSkeletonRemapping>> ExistingMappings;
+	{
+		FRWScopeLock ReadLock(MappingsLock, SLT_ReadOnly);
+		PerSkeletonMappings.MultiFind(InSkeleton, ExistingMappings);
+	}
+	{
+		FRWScopeLock WriteLock(MappingsLock, SLT_Write);
+		for(const TSharedPtr<FSkeletonRemapping>& ExistingMapping : ExistingMappings)
+		{
+			ExistingMapping->RegenerateMapping();
+		}
+	}
+}
+
+}
diff --git a/Engine/Source/Runtime/Engine/Private/Animation/SkeletonRemappingRegistry.h b/Engine/Source/Runtime/Engine/Private/Animation/SkeletonRemappingRegistry.h
new file mode 100644
--- /dev/null
+++ b/Engine/Source/Runtime/Engine/Private/Animation/SkeletonRemappingRegistry.h
@@ -0,0 +1,58 @@
+// Copyright Epic Games, Inc. All Rights Reserved.
+
+#pragma once
+
+#include "CoreMinimal.h"
+#include "UObject/WeakObjectPtr.h"
+
+struct FSkeletonRemapping;
+class USkeleton;
+class FEngineModule;
+
+namespace UE::Anim
+{
+
+struct FSkeletonRemappingRegistryPrivate;
+
+// Global registry of skeleton remappings
+// Remappings are created on-demand when calling GetRemapping.
+// Calling public functions from multiple threads is expected. Data races are guarded by a FRWLock.
+class FSkeletonRemappingRegistry
+{
+public:
+	// Access the global registry
+	static FSkeletonRemappingRegistry& Get();
+
+	// Get a mapping between a source and target skeleton, built on demand
+	// Remapping is only valid locally to this call and could be free'd later (they are stored as shared ptrs)
+	// Do not assume ownership of the mapping after this call. 
+	const FSkeletonRemapping& GetRemapping(const USkeleton* InSourceSkeleton, const USkeleton* InTargetSkeleton);
+
+	// Refresh any existing mappings that the supplied skeleton is used with
+	void RefreshMappings(USkeleton* InSkeleton);
+	
+	// Refresh any existing curve mappings that the supplied skeleton is used with
+	void RefreshCurveMappings(USkeleton* InSkeleton);
+
+private:
+	friend FSkeletonRemappingRegistryPrivate;
+	friend ::FEngineModule;
+	
+private:
+	// Initialize the global registry
+	static void Init();
+
+	// Shutdown the global registry
+	static void Destroy();
+	
+	// Mappings map is guarded with a RW lock
+	FRWLock MappingsLock;
+
+	using FWeakSkeletonPair = TPair<TWeakObjectPtr<const USkeleton>, TWeakObjectPtr<const USkeleton>>;
+	
+	// All mappings
+	TMap<FWeakSkeletonPair, TSharedPtr<FSkeletonRemapping>> Mappings;
+	TMultiMap<TWeakObjectPtr<const USkeleton>, TSharedPtr<FSkeletonRemapping>> PerSkeletonMappings;
+};
+
+}
\ No newline at end of file
diff --git a/Engine/Source/Runtime/Engine/Private/BoneContainer.cpp b/Engine/Source/Runtime/Engine/Private/BoneContainer.cpp
--- a/Engine/Source/Runtime/Engine/Private/BoneContainer.cpp
+++ b/Engine/Source/Runtime/Engine/Private/BoneContainer.cpp
@@ -6,7 +6,8 @@
 #include "Engine/SkeletalMesh.h"
 #include "EngineLogs.h"
 #include "Animation/AnimCurveTypes.h"
-
+#include "Animation/SkeletonRemappingRegistry.h"
+#include "Animation/SkeletonRemapping.h"
 #include "HAL/LowLevelMemTracker.h"
 
 #include UE_INLINE_GENERATED_CPP_BY_NAME(BoneContainer)
@@ -33,10 +34,10 @@ FSkeletonRemappingCurve::FSkeletonRemappingCurve(FBlendedCurve& InCurve, FBoneCo
 	, BoneContainer(InBoneContainer)
 {
 	LLM_SCOPE_BYNAME(TEXT("Animation/BoneContainer"));
-	const FSkeletonRemapping* SkeletonRemapping = BoneContainer.GetSkeletonAsset()->GetSkeletonRemapping(SourceSkeleton);
-	if (SkeletonRemapping) // No remapping is required, just continue as we normally would.
+	const FSkeletonRemapping& SkeletonRemapping = UE::Anim::FSkeletonRemappingRegistry::Get().GetRemapping(SourceSkeleton, BoneContainer.GetSkeletonAsset());
+	if (SkeletonRemapping.IsValid()) // No remapping is required, just continue as we normally would.
 	{
-		const FCachedSkeletonCurveMapping& CurveMapping = BoneContainer.GetOrCreateCachedCurveMapping(SkeletonRemapping);
+		const FCachedSkeletonCurveMapping& CurveMapping = BoneContainer.GetOrCreateCachedCurveMapping(&SkeletonRemapping);
 		Curve.UIDToArrayIndexLUT = &CurveMapping.UIDToArrayIndices;
 		bIsRemapping = true;
 	}
diff --git a/Engine/Source/Runtime/Engine/Private/Components/SkeletalMeshComponent.cpp b/Engine/Source/Runtime/Engine/Private/Components/SkeletalMeshComponent.cpp
--- a/Engine/Source/Runtime/Engine/Private/Components/SkeletalMeshComponent.cpp
+++ b/Engine/Source/Runtime/Engine/Private/Components/SkeletalMeshComponent.cpp
@@ -549,9 +549,9 @@ bool USkeletalMeshComponent::NeedToSpawnAnimScriptInstance() const
 		// Fall back to mesh skeleton (anim BP could have been a template)
 		AnimSkeleton = MeshAsset ? MeshAsset->GetSkeleton() : nullptr;
 	}
-	const bool bSkeletonCompatible = (MeshAsset && AnimSkeleton && MeshAsset->GetSkeleton()) ? MeshAsset->GetSkeleton()->IsCompatible(AnimSkeleton) : false;
+	const bool bSkeletonsExist = (MeshAsset && AnimSkeleton && MeshAsset->GetSkeleton());
 	const bool bSkelMeshCompatible = (MeshAsset && AnimSkeleton) ? AnimSkeleton->IsCompatibleMesh(MeshAsset, false) : false;
-	const bool bAnimSkelValid = !AnimClassInterface || (bSkeletonCompatible && bSkelMeshCompatible);
+	const bool bAnimSkelValid = !AnimClassInterface || (bSkeletonsExist && bSkelMeshCompatible);
 
 	if (AnimationMode == EAnimationMode::AnimationBlueprint && AnimClass && bAnimSkelValid)
 	{
@@ -787,11 +787,11 @@ void USkeletalMeshComponent::InitAnim(bool bForceReinit)
 
 		const bool bClearAnimInstance = AnimScriptInstance && !AnimSkeleton;
 		const bool bSkeletonMismatch = AnimSkeleton && (AnimScriptInstance->CurrentSkeleton!=GetSkeletalMeshAsset()->GetSkeleton());
-		const bool bSkeletonCompatible = AnimSkeleton && GetSkeletalMeshAsset()->GetSkeleton() && !bSkeletonMismatch && GetSkeletalMeshAsset()->GetSkeleton()->IsCompatible(AnimSkeleton);
+		const bool bSkeletonsExist = AnimSkeleton && GetSkeletalMeshAsset()->GetSkeleton() && !bSkeletonMismatch;
 
 		LastPoseTickFrame = 0;
 
-		if (bBlueprintMismatch || bSkeletonMismatch || !bSkeletonCompatible || bClearAnimInstance)
+		if (bBlueprintMismatch || bSkeletonMismatch || !bSkeletonsExist || bClearAnimInstance)
 		{
 			ClearAnimScriptInstance();
 		}
@@ -1128,14 +1128,14 @@ void USkeletalMeshComponent::PostEditChangeProperty(FPropertyChangedEvent& Prope
 
 		if ( PropertyThatChanged->GetFName() == GET_MEMBER_NAME_CHECKED( FSingleAnimationPlayData, AnimToPlay ))
 		{
-			// make sure the animation skeleton matches the current skeletalmesh
-			if (AnimationData.AnimToPlay && SkelMesh && SkelMesh->GetSkeleton() && SkelMesh->GetSkeleton()->IsCompatible(AnimationData.AnimToPlay->GetSkeleton()))
+			// make sure the animation skeleton is valid
+			if (AnimationData.AnimToPlay && AnimationData.AnimToPlay->GetSkeleton())
 			{
 				PlayAnimation(AnimationData.AnimToPlay, false);
 			}
 			else
 			{
-				UE_LOG(LogAnimation, Warning, TEXT("Invalid animation"));
+				UE_LOG(LogAnimation, Warning, TEXT("Invalid animation skeleton"));
 				AnimationData.AnimToPlay = nullptr;
 			}
 		}
@@ -3856,11 +3856,6 @@ void USkeletalMeshComponent::ValidateAnimation()
 				UE_LOG(LogAnimation, Warning, TEXT("Animation %s is incompatible because it has no skeleton, removing animation from actor."), *AnimationData.AnimToPlay->GetName());
 				AnimationData.AnimToPlay = nullptr;
 			}
-			else if (!AnimationData.AnimToPlay->GetSkeleton()->IsCompatible(SkelMesh->GetSkeleton()))
-			{
-				UE_LOG(LogAnimation, Warning, TEXT("Animation %s is incompatible with the skeletal mesh's skeleton, removing animation from actor."), *AnimationData.AnimToPlay->GetName());
-				AnimationData.AnimToPlay = nullptr;
-			}
 		}
 	}
 	else if (AnimationMode == EAnimationMode::AnimationBlueprint)
@@ -3878,11 +3873,6 @@ void USkeletalMeshComponent::ValidateAnimation()
 				UE_LOG(LogAnimation, Warning, TEXT("AnimBP %s is incompatible because mesh %s has no skeleton, removing AnimBP from actor."), *AnimClass->GetName(), *SkelMesh->GetName());
 				AnimClass = nullptr;
 			}
-			else if (!AnimClassInterface->GetTargetSkeleton()->IsCompatible(SkelMesh->GetSkeleton()))
-			{
-				UE_LOG(LogAnimation, Warning, TEXT("AnimBP %s is incompatible with skeleton %s, removing AnimBP from actor."), *AnimClass->GetName(), *SkelMesh->GetSkeleton()->GetName());
-				AnimClass = nullptr;
-			}
 		}
 	}
 }
diff --git a/Engine/Source/Runtime/Engine/Private/SkeletalMeshLODSettings.cpp b/Engine/Source/Runtime/Engine/Private/SkeletalMeshLODSettings.cpp
--- a/Engine/Source/Runtime/Engine/Private/SkeletalMeshLODSettings.cpp
+++ b/Engine/Source/Runtime/Engine/Private/SkeletalMeshLODSettings.cpp
@@ -59,7 +59,7 @@ bool USkeletalMeshLODSettings::SetLODSettingsToMesh(USkeletalMesh* InMesh, int32
 		// if we have available bake pose
 		// it's possible for skeleton to be null if this happens in the middle of importing
 		// so if skeleton is null, we allow copy (the GetBakePose will check correct skeleton when get it)
-		if (Setting.BakePose && (!InMesh->GetSkeleton() || InMesh->GetSkeleton()->IsCompatible(Setting.BakePose->GetSkeleton())))
+		if (Setting.BakePose && (!InMesh->GetSkeleton() || !Setting.BakePose->GetSkeleton()))
 		{
 			LODInfo->BakePose = Setting.BakePose;
 		}
diff --git a/Engine/Source/Runtime/Engine/Private/UnrealEngine.cpp b/Engine/Source/Runtime/Engine/Private/UnrealEngine.cpp
--- a/Engine/Source/Runtime/Engine/Private/UnrealEngine.cpp
+++ b/Engine/Source/Runtime/Engine/Private/UnrealEngine.cpp
@@ -245,6 +245,7 @@ UnrealEngine.cpp: Implements the UEngine class and helpers.
 #include "Engine/CoreSettings.h"
 #include "IEyeTrackerModule.h"
 #include "Interfaces/IPluginManager.h"
+#include "Animation/SkeletonRemappingRegistry.h"
 
 #if !UE_BUILD_SHIPPING
 #include "GenericPlatform/GenericPlatformCrashContext.h"
@@ -330,6 +331,8 @@ void FEngineModule::StartupModule()
 #endif
 
 	FSkinWeightProfileManager::OnStartup();
+
+	UE::Anim::FSkeletonRemappingRegistry::Init();
 }
 
 void FEngineModule::ShutdownModule()
@@ -345,6 +348,8 @@ void FEngineModule::ShutdownModule()
 	FParticleSystemWorldManager::OnShutdown();
 
 	FSkinWeightProfileManager::OnShutdown();
+
+	UE::Anim::FSkeletonRemappingRegistry::Destroy();
 }
 
 /* Global variables
diff --git a/Engine/Source/Runtime/Engine/Public/Animation/AnimInstanceProxy.h b/Engine/Source/Runtime/Engine/Public/Animation/AnimInstanceProxy.h
--- a/Engine/Source/Runtime/Engine/Public/Animation/AnimInstanceProxy.h
+++ b/Engine/Source/Runtime/Engine/Public/Animation/AnimInstanceProxy.h
@@ -244,10 +244,10 @@ public:
 	/** Trigger any anim notifies */
 	void TriggerAnimNotifies(USkeletalMeshComponent* SkelMeshComp, float DeltaSeconds);
 
-	/** Check whether animation content authored on the supplied skeleton may be played on this instance's skeleton */
+	UE_DEPRECATED(5.2, "Skeleton compatibility is now an editor-only concern, this function is no longer used.")
 	bool IsSkeletonCompatible(USkeleton const* InSkeleton) const
 	{
-		return Skeleton && Skeleton->IsCompatible(InSkeleton);
+		return Skeleton != nullptr && InSkeleton != nullptr;
 	}
 
 	/** Check whether we should extract root motion */
diff --git a/Engine/Source/Runtime/MovieSceneTracks/Private/Evaluation/MovieSceneSkeletalAnimationTemplate.cpp b/Engine/Source/Runtime/MovieSceneTracks/Private/Evaluation/MovieSceneSkeletalAnimationTemplate.cpp
--- a/Engine/Source/Runtime/MovieSceneTracks/Private/Evaluation/MovieSceneSkeletalAnimationTemplate.cpp
+++ b/Engine/Source/Runtime/MovieSceneTracks/Private/Evaluation/MovieSceneSkeletalAnimationTemplate.cpp
@@ -34,8 +34,10 @@ bool ShouldUsePreviewPlayback(IMovieScenePlayer& Player, UObject& RuntimeObject)
 
 bool CanPlayAnimation(USkeletalMeshComponent* SkeletalMeshComponent, UAnimSequenceBase* AnimAssetBase)
 {
-	return (SkeletalMeshComponent->GetSkeletalMeshAsset() && SkeletalMeshComponent->GetSkeletalMeshAsset()->GetSkeleton() &&
-		(!AnimAssetBase || SkeletalMeshComponent->GetSkeletalMeshAsset()->GetSkeleton()->IsCompatible(AnimAssetBase->GetSkeleton())));
+	return (SkeletalMeshComponent->GetSkeletalMeshAsset() &&
+			SkeletalMeshComponent->GetSkeletalMeshAsset()->GetSkeleton() &&
+			AnimAssetBase != nullptr &&
+			AnimAssetBase->GetSkeleton() != nullptr);
 }
 
 void ResetAnimSequencerInstance(UObject& InObject, const UE::MovieScene::FRestoreStateParams& Params)
