- Change streaming methods to use full SRV update

#rnx

diff --git a/Engine/Source/Runtime/Engine/Public/RawIndexBuffer.h b/Engine/Source/Runtime/Engine/Public/RawIndexBuffer.h
--- a/Engine/Source/Runtime/Engine/Public/RawIndexBuffer.h
+++ b/Engine/Source/Runtime/Engine/Public/RawIndexBuffer.h
@@ -509,7 +509,7 @@ public:
 			Batcher.QueueUpdateRequest(IndexBufferRHI, IntermediateBuffer);
 			if (SRVValue)
 			{
-				Batcher.QueueUpdateRequest(SRVValue, IndexBufferRHI);
+				Batcher.QueueUpdateRequest(SRVValue, IndexBufferRHI, sizeof(INDEX_TYPE), sizeof(INDEX_TYPE) == 2 ? PF_R16_UINT : PF_R32_UINT);
 			}
 		}
 	}
@@ -523,7 +523,7 @@ public:
 		}
 		if (SRVValue)
 		{
-			Batcher.QueueUpdateRequest(SRVValue, nullptr);
+			Batcher.QueueUpdateRequest(SRVValue, nullptr, 0, 0);
 		}
 	}
 
