Post Process weight applies to LumenDiffuseColorBoost and LumenSkylightLeaking

diff --git a/Engine/Source/Runtime/Engine/Private/SceneView.cpp b/Engine/Source/Runtime/Engine/Private/SceneView.cpp
--- a/Engine/Source/Runtime/Engine/Private/SceneView.cpp
+++ b/Engine/Source/Runtime/Engine/Private/SceneView.cpp
@@ -1661,20 +1661,9 @@ void FSceneView::OverridePostProcessSettings(const FPostProcessSettings& Src, fl
 			Dest.LumenMaxTraceDistance = Src.LumenMaxTraceDistance;
 		}
 
-		if (Src.bOverride_LumenDiffuseColorBoost)
-		{
-			Dest.LumenDiffuseColorBoost = Src.LumenDiffuseColorBoost;
-		}
-
-		if (Src.bOverride_LumenSkylightLeaking)
-		{
-			Dest.LumenSkylightLeaking = Src.LumenSkylightLeaking;
-		}
-
-		if (Src.bOverride_LumenFullSkylightLeakingDistance)
-		{
-			Dest.LumenFullSkylightLeakingDistance = Src.LumenFullSkylightLeakingDistance;
-		}
+		LERP_PP(LumenDiffuseColorBoost);
+		LERP_PP(LumenSkylightLeaking);
+		LERP_PP(LumenFullSkylightLeakingDistance);
 
 		if (Src.bOverride_LumenRayLightingMode)
 		{
