Fix an issue where the Mirror Data Table was not automatically adding sync markers on creation

diff --git a/Engine/Source/Runtime/Engine/Private/Animation/MirrorDataTable.cpp b/Engine/Source/Runtime/Engine/Private/Animation/MirrorDataTable.cpp
--- a/Engine/Source/Runtime/Engine/Private/Animation/MirrorDataTable.cpp
+++ b/Engine/Source/Runtime/Engine/Private/Animation/MirrorDataTable.cpp
@@ -182,7 +182,7 @@ void UMirrorDataTable::FindReplaceMirroredNames()
 
 	const FReferenceSkeleton& RefSkeleton = Skeleton->GetReferenceSkeleton();
 
-	TSet<FName> NamesByCategory[3];
+	TSet<FName> NamesByCategory[4];
 	ForeachRow<FMirrorTableRow>(TEXT("UMirrorDataTable::FindReplaceMirroredNames"), [&NamesByCategory] (const FName& Key, const FMirrorTableRow& Value) mutable
 		{
 			NamesByCategory[Value.MirrorEntryType].Add(Value.Name);
@@ -195,7 +195,7 @@ void UMirrorDataTable::FindReplaceMirroredNames()
 	{
 		// directly add rows to avoid using FDataTableEditorUtils, which is not appropriate at this point
 		// equivalent to FDataTableEditorUtils::AddRow(DataTable, BoneName);
-		static const FString CategoryName[] = { TEXT(":Bone"), TEXT(":Notify"), TEXT(":Curve") };
+		static const FString CategoryName[] = { TEXT(":Bone"), TEXT(":Notify"), TEXT(":Curve"), TEXT(":SyncMarker") };
 		if (!NamesByCategory[RowType].Contains(Name))
 		{
 			FName RowName = Name;
@@ -253,6 +253,15 @@ void UMirrorDataTable::FindReplaceMirroredNames()
 		}
 	}
 
+	for (const FName& SyncMarker : Skeleton->GetExistingMarkerNames())
+	{
+		FName MirroredName = FindReplace(SyncMarker);
+		if (!MirroredName.IsNone() && Skeleton->GetExistingMarkerNames().Contains(MirroredName))
+		{
+			AddMirrorRow(SyncMarker, MirroredName, EMirrorRowType::SyncMarker);
+		}
+	}
+
 	const FSmartNameMapping* CurveSmartNames = Skeleton->GetSmartNameContainer(USkeleton::AnimCurveMappingName);
 	if (CurveSmartNames)
 	{
