World Partition: fix crash when repair errors button is pressed in level outliner.

#rnx

diff --git a/Engine/Source/Editor/SceneOutliner/Private/ActorBrowsingMode.cpp b/Engine/Source/Editor/SceneOutliner/Private/ActorBrowsingMode.cpp
--- a/Engine/Source/Editor/SceneOutliner/Private/ActorBrowsingMode.cpp
+++ b/Engine/Source/Editor/SceneOutliner/Private/ActorBrowsingMode.cpp
@@ -1122,7 +1122,10 @@ void FActorBrowsingMode::RepairErrors() const
 			{
 				for (const TUniquePtr<FWorldPartitionActorDesc>& InvalidActor : ActorDescContainer->GetInvalidActors())
 				{
-					InvalidActorDescs.Add(InvalidActor.Get());
+					if (FWorldPartitionActorDesc* InvalidActorDesc = InvalidActor.Get())
+					{
+						InvalidActorDescs.Add(InvalidActorDesc);
+					}
 				}
 				ActorDescContainer->ClearInvalidActors();
 			});
diff --git a/Engine/Source/Runtime/Engine/Private/WorldPartition/ActorDescContainer.cpp b/Engine/Source/Runtime/Engine/Private/WorldPartition/ActorDescContainer.cpp
--- a/Engine/Source/Runtime/Engine/Private/WorldPartition/ActorDescContainer.cpp
+++ b/Engine/Source/Runtime/Engine/Private/WorldPartition/ActorDescContainer.cpp
@@ -69,7 +69,7 @@ void UActorDescContainer::Initialize(const FInitializeParams& InitParams)
 		{
 			TUniquePtr<FWorldPartitionActorDesc> ActorDesc = FWorldPartitionActorDescUtils::GetActorDescriptorFromAssetData(Asset);
 
-			if (ActorDesc.IsValid() && (!InitParams.FilterActorDesc || InitParams.FilterActorDesc(ActorDesc.Get())))
+			if (ActorDesc.IsValid() && ActorDesc->GetNativeClass().IsValid() && (!InitParams.FilterActorDesc || InitParams.FilterActorDesc(ActorDesc.Get())))
 			{
 				AddActorDescriptor(ActorDesc.Release());
 			}
diff --git a/Engine/Source/Runtime/Engine/Private/WorldPartition/WorldPartitionActorDescUtils.cpp b/Engine/Source/Runtime/Engine/Private/WorldPartition/WorldPartitionActorDescUtils.cpp
--- a/Engine/Source/Runtime/Engine/Private/WorldPartition/WorldPartitionActorDescUtils.cpp
+++ b/Engine/Source/Runtime/Engine/Private/WorldPartition/WorldPartitionActorDescUtils.cpp
@@ -63,7 +63,8 @@ TUniquePtr<FWorldPartitionActorDesc> FWorldPartitionActorDescUtils::GetActorDesc
 		if (!ActorDescInitData.NativeClass)
 		{
 			UE_LOG(LogWorldPartition, Log, TEXT("Invalid class for actor guid `%s` ('%s') from package '%s'"), *NewActorDesc->GetGuid().ToString(), *NewActorDesc->GetActorName().ToString(), *NewActorDesc->GetActorPackage().ToString());
-			return nullptr;
+			NewActorDesc->NativeClass.Reset();
+			return NewActorDesc;
 		}
 		/*else if (UClass* Class = FindObject<UClass>(InAssetData.AssetClassPath); !Class)
 		{
