Added missing Content Bundle ID in world partition cell debug name

diff --git a/Engine/Source/Runtime/Engine/Private/WorldPartition/WorldPartitionRuntimeCell.cpp b/Engine/Source/Runtime/Engine/Private/WorldPartition/WorldPartitionRuntimeCell.cpp
--- a/Engine/Source/Runtime/Engine/Private/WorldPartition/WorldPartitionRuntimeCell.cpp
+++ b/Engine/Source/Runtime/Engine/Private/WorldPartition/WorldPartitionRuntimeCell.cpp
@@ -65,20 +65,28 @@ void UWorldPartitionRuntimeCell::UpdateDebugName()
 	Builder += FString::Printf(TEXT("L%d_X%d_Y%d"), DebugInfo.CoordZ, DebugInfo.CoordX, DebugInfo.CoordY);
 	int32 DataLayerCount = DataLayers.Num();
 
-	const UDataLayerSubsystem* DataLayerSubsystem = UWorld::GetSubsystem<UDataLayerSubsystem>(GetCellOwner()->GetOuterWorld());
-	TArray<const UDataLayerInstance*> DataLayerObjects;
-	if (DataLayerSubsystem && DataLayerCount > 0)
+	if (DataLayerCount > 0)
 	{
-		Builder += TEXT(" DL[");
-		for (int i = 0; i < DataLayerCount; ++i)
+		if (const UDataLayerSubsystem* DataLayerSubsystem = UWorld::GetSubsystem<UDataLayerSubsystem>(GetCellOwner()->GetOuterWorld()))
 		{
-			const UDataLayerInstance* DataLayer = DataLayerSubsystem->GetDataLayerInstance(DataLayers[i]);
-			DataLayerObjects.Add(DataLayer);
-			Builder += DataLayer->GetDataLayerShortName();
-			Builder += TEXT(",");
+			TArray<const UDataLayerInstance*> DataLayerObjects;
+			Builder += TEXT(" DL[");
+			for (int i = 0; i < DataLayerCount; ++i)
+			{
+				const UDataLayerInstance* DataLayer = DataLayerSubsystem->GetDataLayerInstance(DataLayers[i]);
+				DataLayerObjects.Add(DataLayer);
+				Builder += DataLayer->GetDataLayerShortName();
+				Builder += TEXT(",");
+			}
+			Builder += FString::Printf(TEXT("ID:%X]"), FDataLayersID(DataLayerObjects).GetHash());
 		}
-		Builder += FString::Printf(TEXT("ID:%X]"), FDataLayersID(DataLayerObjects).GetHash());
 	}
+
+	if (ContentBundleID.IsValid())
+	{
+		Builder += FString::Printf(TEXT(" CB[%X]"), GetTypeHash(ContentBundleID));
+	}
+
 	DebugInfo.Name = Builder.ToString();
 }
 
