Fix missing OFPA dependencies causing missing assets on cook builds

diff --git a/Engine/Source/Runtime/Engine/Private/ExternalActorAssetDependencyGatherer.cpp b/Engine/Source/Runtime/Engine/Private/ExternalActorAssetDependencyGatherer.cpp
new file mode 100644
--- /dev/null
+++ b/Engine/Source/Runtime/Engine/Private/ExternalActorAssetDependencyGatherer.cpp
@@ -0,0 +1,35 @@
+// Copyright Epic Games, Inc. All Rights Reserved.
+#include "ExternalActorAssetDependencyGatherer.h"
+
+#if WITH_EDITOR
+
+#include "Engine/World.h"
+#include "Engine/Level.h"
+#include "HAL/FileManager.h"
+
+// Register FExternalActorAssetDependencyGatherer for UWorld class
+REGISTER_ASSETDEPENDENCY_GATHERER(FExternalActorAssetDependencyGatherer, UWorld);
+
+void FExternalActorAssetDependencyGatherer::GatherDependencies(const FAssetData& AssetData, const FAssetRegistryState& AssetRegistryState, TFunctionRef<FARCompiledFilter(const FARFilter&)> CompileFilterFunc, TArray<IAssetDependencyGatherer::FGathereredDependency>& OutDependencies, TArray<FString>& OutDependencyDirectories) const
+{
+	if (ULevel::GetIsLevelUsingExternalActorsFromAsset(AssetData))
+	{
+		const FString ExternalActorsPath = ULevel::GetExternalActorsPath(AssetData.PackageName.ToString());
+		OutDependencyDirectories.Add(ExternalActorsPath);
+
+		FARFilter Filter;
+		Filter.bRecursivePaths = true;
+		Filter.bIncludeOnlyOnDiskAssets = true;
+		Filter.PackagePaths.Add(*ExternalActorsPath);
+
+		TArray<FAssetData> FilteredAssets;
+		AssetRegistryState.GetAssets(CompileFilterFunc(Filter), {}, FilteredAssets, true);
+
+		for (const FAssetData& FilteredAsset : FilteredAssets)
+		{
+			OutDependencies.Emplace(IAssetDependencyGatherer::FGathereredDependency{ FilteredAsset.PackageName, UE::AssetRegistry::EDependencyProperty::Game });
+		}
+	}
+}
+
+#endif
\ No newline at end of file
diff --git a/Engine/Source/Runtime/Engine/Private/ExternalActorAssetDependencyGatherer.h b/Engine/Source/Runtime/Engine/Private/ExternalActorAssetDependencyGatherer.h
new file mode 100644
--- /dev/null
+++ b/Engine/Source/Runtime/Engine/Private/ExternalActorAssetDependencyGatherer.h
@@ -0,0 +1,21 @@
+// Copyright Epic Games, Inc. All Rights Reserved.
+
+#pragma once
+
+#if WITH_EDITOR
+
+#include "AssetRegistry/AssetDependencyGatherer.h"
+#include "AssetRegistry/AssetData.h"
+#include "Engine/World.h"
+
+class FExternalActorAssetDependencyGatherer : public IAssetDependencyGatherer
+{
+public:
+	
+	FExternalActorAssetDependencyGatherer() = default;
+	virtual ~FExternalActorAssetDependencyGatherer() = default;
+
+	virtual void GatherDependencies(const FAssetData& AssetData, const FAssetRegistryState& AssetRegistryState, TFunctionRef<FARCompiledFilter(const FARFilter&)> CompileFilterFunc, TArray<IAssetDependencyGatherer::FGathereredDependency>& OutDependencies, TArray<FString>& OutDependencyDirectories) const override;
+};
+
+#endif
\ No newline at end of file
diff --git a/Engine/Source/Runtime/Engine/Private/WorldPartition/WorldPartitionAssetDependencyGatherer.cpp b/Engine/Source/Runtime/Engine/Private/WorldPartition/WorldPartitionAssetDependencyGatherer.cpp
deleted file mode 100644
diff --git a/Engine/Source/Runtime/Engine/Private/WorldPartition/WorldPartitionAssetDependencyGatherer.h b/Engine/Source/Runtime/Engine/Private/WorldPartition/WorldPartitionAssetDependencyGatherer.h
deleted file mode 100644
