Temp crash guard.

#rnx

diff --git a/Engine/Source/Runtime/Engine/Private/WorldPartition/ContentBundle/ContentBundleContainer.cpp b/Engine/Source/Runtime/Engine/Private/WorldPartition/ContentBundle/ContentBundleContainer.cpp
--- a/Engine/Source/Runtime/Engine/Private/WorldPartition/ContentBundle/ContentBundleContainer.cpp
+++ b/Engine/Source/Runtime/Engine/Private/WorldPartition/ContentBundle/ContentBundleContainer.cpp
@@ -62,13 +62,16 @@ void FContentBundleContainer::Initialize()
 
 void FContentBundleContainer::Deinitialize()
 {
-	UE_LOG(LogContentBundle, Log, TEXT("[Container: %s] Deleting container."), *GetInjectedWorld()->GetName());
+	if (GetInjectedWorld())
+	{
+		UE_LOG(LogContentBundle, Log, TEXT("[Container: %s] Deleting container."), *GetInjectedWorld()->GetName());
 
 #if WITH_EDITOR
-	UWorldPartition* WorldPartition = GetInjectedWorld()->GetWorldPartition();
-	WorldPartition->OnPreGenerateStreaming.RemoveAll(this);
-	WorldPartition->OnBeginCook.RemoveAll(this);
+		UWorldPartition* WorldPartition = GetInjectedWorld()->GetWorldPartition();
+		WorldPartition->OnPreGenerateStreaming.RemoveAll(this);
+		WorldPartition->OnBeginCook.RemoveAll(this);
 #endif
+	}
 
 	UnregisterContentBundleClientEvents();
 	DeinitializeContentBundles();
