Remove recursive calls for fields.

diff --git a/Engine/Source/Runtime/Engine/Classes/PhysicsField/PhysicsFieldComponent.h b/Engine/Source/Runtime/Engine/Classes/PhysicsField/PhysicsFieldComponent.h
--- a/Engine/Source/Runtime/Engine/Classes/PhysicsField/PhysicsFieldComponent.h
+++ b/Engine/Source/Runtime/Engine/Classes/PhysicsField/PhysicsFieldComponent.h
@@ -186,11 +186,14 @@ public:
 	 */
 	void UpdateInstance(const float TimeSeconds);
 
+	/** Build the field array for params and bounds */
+	static void BuildFieldArray(TQueue<FFieldNodeBase*>& FieldNodesToProcess, TArray<FFieldNodeBase*>& FieldNodes);
+
 	/** Update the offsets and params given a node */
-	void BuildNodeParams(FFieldNodeBase* FieldNode, const TMap<FFieldNodeBase*, float> CommandTimes, const float PreviousTime);
+	void BuildNodeParams(const TArray<FFieldNodeBase*>& FieldNodes, const TMap<FFieldNodeBase*, float> CommandTimes, const float PreviousTime);
 
 	/** Update the bounds given a node */
-	static void BuildNodeBounds(FFieldNodeBase* FieldNode, FVector& MinBounds, FVector& MaxBounds, float& MaxMagnitude);
+	static void BuildNodeBounds(const TArray<FFieldNodeBase*>& FieldNodes, FVector& MinBounds, FVector& MaxBounds, float& MaxMagnitude);
 
 	/** The field system resource. */
 	FPhysicsFieldResource* FieldResource = nullptr;
diff --git a/Engine/Source/Runtime/Engine/Private/PhysicsField/PhysicsFieldComponent.cpp b/Engine/Source/Runtime/Engine/Private/PhysicsField/PhysicsFieldComponent.cpp
--- a/Engine/Source/Runtime/Engine/Private/PhysicsField/PhysicsFieldComponent.cpp
+++ b/Engine/Source/Runtime/Engine/Private/PhysicsField/PhysicsFieldComponent.cpp
@@ -777,8 +777,12 @@ void FPhysicsFieldInstance::UpdateInstance(const float TimeSeconds)
 			float MaxMagnitude = 1.0;
 			if (TargetNode)
 			{
-				BuildNodeBounds(TargetNode, MinBound, MaxBound, MaxMagnitude);
-				BuildNodeParams(TargetNode, CommandTimes, 0.0);
+				TArray<FFieldNodeBase*> FieldNodes;
+				TQueue<FFieldNodeBase*> FieldNodesToProcess;
+				FieldNodesToProcess.Enqueue(TargetNode);
+				BuildFieldArray(FieldNodesToProcess, FieldNodes);
+				BuildNodeBounds(FieldNodes, MinBound, MaxBound, MaxMagnitude);
+				BuildNodeParams(FieldNodes, CommandTimes, 0.0);
 				if (TargetRoots.Num() > 1) delete TargetNode;
 			}
 			TargetsOffsets[TargetType + 1] = NodesOffsets.Num() - PreviousNodes;
@@ -817,9 +821,236 @@ void FPhysicsFieldInstance::UpdateInstance(const float TimeSeconds)
 	}
 }
 
-void FPhysicsFieldInstance::BuildNodeParams(FFieldNodeBase* FieldNode, const TMap<FFieldNodeBase*,float> CommandTimes, const float PreviousTime)
+void FPhysicsFieldInstance::BuildFieldArray(TQueue<FFieldNodeBase*>& FieldNodesToProcess, TArray<FFieldNodeBase*>& FieldNodes)
 {
-	if (FieldNode)
+	while (!FieldNodesToProcess.IsEmpty())
+	{
+		FFieldNodeBase* FieldNode;
+		FieldNodesToProcess.Dequeue(FieldNode);
+		if (!FieldNode)
+		{
+			continue;
+		}
+		if (FieldNode->SerializationType() == FFieldNodeBase::ESerializationType::FieldNode_FUniformInteger)
+		{
+			FieldNodes.Add(FieldNode);
+		}
+		else if (FieldNode->SerializationType() == FFieldNodeBase::ESerializationType::FieldNode_FRadialIntMask)
+		{
+			FieldNodes.Add(FieldNode);
+		}
+		else if (FieldNode->SerializationType() == FFieldNodeBase::ESerializationType::FieldNode_FUniformScalar)
+		{
+			FieldNodes.Add(FieldNode);
+		}
+		else if (FieldNode->SerializationType() == FFieldNodeBase::ESerializationType::FieldNode_FWaveScalar)
+		{
+			FieldNodes.Add(FieldNode);
+		}
+		else if (FieldNode->SerializationType() == FFieldNodeBase::ESerializationType::FieldNode_FRadialFalloff)
+		{
+			FieldNodes.Add(FieldNode);
+		}
+		else if (FieldNode->SerializationType() == FFieldNodeBase::ESerializationType::FieldNode_FPlaneFalloff)
+		{
+			FieldNodes.Add(FieldNode);
+		}
+		else if (FieldNode->SerializationType() == FFieldNodeBase::ESerializationType::FieldNode_FBoxFalloff)
+		{
+			FieldNodes.Add(FieldNode);
+		}
+		else if (FieldNode->SerializationType() == FFieldNodeBase::ESerializationType::FieldNode_FNoiseField)
+		{
+			FieldNodes.Add(FieldNode);
+		}
+		else if (FieldNode->SerializationType() == FFieldNodeBase::ESerializationType::FieldNode_FUniformVector)
+		{
+			FieldNodes.Add(FieldNode);
+		}
+		else if (FieldNode->SerializationType() == FFieldNodeBase::ESerializationType::FieldNode_FRadialVector)
+		{
+			FieldNodes.Add(FieldNode);
+		}
+		else if (FieldNode->SerializationType() == FFieldNodeBase::ESerializationType::FieldNode_FRandomVector)
+		{
+			FieldNodes.Add(FieldNode);
+		}
+		else if (FieldNode->SerializationType() == FFieldNodeBase::ESerializationType::FieldNode_FSumScalar)
+		{
+			if (FieldNode->UserData)
+			{
+				FieldNodes.Add(FieldNode);
+			}
+			else
+			{
+				FSumScalar* LocalNode = StaticCast<FSumScalar*>(FieldNode);
+
+				if (LocalNode->ScalarRight)
+				{
+					FieldNodesToProcess.Enqueue(LocalNode->ScalarRight.Get());
+					LocalNode->ScalarRight->UserData = 0;
+				}
+				if (LocalNode->ScalarLeft)
+				{
+					FieldNodesToProcess.Enqueue(LocalNode->ScalarLeft.Get());
+					LocalNode->ScalarLeft->UserData = 0;
+				}
+				FieldNodesToProcess.Enqueue(LocalNode);
+				FieldNode->UserData = 1;
+			}
+		}
+		else if (FieldNode->SerializationType() == FFieldNodeBase::ESerializationType::FieldNode_FSumVector)
+		{
+			if (FieldNode->UserData)
+			{
+				FieldNodes.Add(FieldNode);
+			}
+			else
+			{
+				FSumVector* LocalNode = StaticCast<FSumVector*>(FieldNode);
+				
+				if (LocalNode->Scalar)
+				{
+					FieldNodesToProcess.Enqueue(LocalNode->Scalar.Get());
+					LocalNode->Scalar->UserData = 0;
+				}
+				if (LocalNode->VectorRight)
+				{
+					FieldNodesToProcess.Enqueue(LocalNode->VectorRight.Get());
+					LocalNode->VectorRight->UserData = 0;
+				}
+				if (LocalNode->VectorLeft)
+				{
+					FieldNodesToProcess.Enqueue(LocalNode->VectorLeft.Get());
+					LocalNode->VectorLeft->UserData = 0;
+				}
+				FieldNodesToProcess.Enqueue(LocalNode);
+				FieldNode->UserData = 1;
+			}
+		}
+		else if (FieldNode->SerializationType() == FFieldNodeBase::ESerializationType::FieldNode_FConversionField)
+		{
+			if (FieldNode->Type() == FFieldNodeBase::EFieldType::EField_Int32)
+			{
+				if (FieldNode->UserData)
+				{
+					FieldNodes.Add(FieldNode);
+				}
+				else
+				{ 
+					FConversionField<float, int32>* LocalNode = StaticCast<FConversionField<float, int32>*>(FieldNode);
+					
+					if (LocalNode->InputField)
+					{
+						FieldNodesToProcess.Enqueue(LocalNode->InputField.Get());
+						LocalNode->InputField->UserData = 0;
+					}
+					FieldNodesToProcess.Enqueue(LocalNode);
+					FieldNode->UserData = 1;
+				}
+			}
+			else if (FieldNode->Type() == FFieldNodeBase::EFieldType::EField_Float)
+			{
+				if (FieldNode->UserData)
+				{
+					FieldNodes.Add(FieldNode);
+				}
+				else
+				{ 
+					FConversionField<int32, float>* LocalNode = StaticCast<FConversionField<int32, float>*>(FieldNode);
+					
+					if (LocalNode->InputField)
+					{
+						FieldNodesToProcess.Enqueue(LocalNode->InputField.Get());
+						LocalNode->InputField->UserData = 0;
+					}
+					FieldNodesToProcess.Enqueue(LocalNode);
+					FieldNode->UserData = 1;
+				}
+			}
+		}
+		else if (FieldNode->SerializationType() == FFieldNodeBase::ESerializationType::FieldNode_FCullingField)
+		{
+			if (FieldNode->Type() == FFieldNodeBase::EFieldType::EField_Int32)
+			{
+				if (FieldNode->UserData)
+				{
+					FieldNodes.Add(FieldNode);
+				}
+				else
+				{ 
+					FCullingField<int32>* LocalNode = StaticCast<FCullingField<int32>*>(FieldNode);
+					
+					if (LocalNode->Culling)
+					{
+						FieldNodesToProcess.Enqueue(LocalNode->Culling.Get());
+						LocalNode->Culling->UserData = 0;
+					}
+					if (LocalNode->Input)
+					{
+						FieldNodesToProcess.Enqueue(LocalNode->Input.Get());
+						LocalNode->Input->UserData = 0;
+					}
+					FieldNodesToProcess.Enqueue(LocalNode);
+					FieldNode->UserData = 1;
+				}
+			}
+			else if (FieldNode->Type() == FFieldNodeBase::EFieldType::EField_Float)
+			{
+				if (FieldNode->UserData)
+				{
+					FieldNodes.Add(FieldNode);
+				}
+				else
+				{
+					FCullingField<float>* LocalNode = StaticCast<FCullingField<float>*>(FieldNode);
+				
+					if (LocalNode->Culling)
+					{
+						FieldNodesToProcess.Enqueue(LocalNode->Culling.Get());
+						LocalNode->Culling->UserData = 0;
+					}
+					if (LocalNode->Input)
+					{
+						FieldNodesToProcess.Enqueue(LocalNode->Input.Get());
+						LocalNode->Input->UserData = 0;
+					}
+					FieldNodesToProcess.Enqueue(LocalNode);
+					FieldNode->UserData = 1;
+				}
+			}
+			else if (FieldNode->Type() == FFieldNodeBase::EFieldType::EField_FVector)
+			{
+				if (FieldNode->UserData)
+				{
+					FieldNodes.Add(FieldNode);
+				}
+				else
+				{
+					FCullingField<FVector>* LocalNode = StaticCast<FCullingField<FVector>*>(FieldNode);
+					
+					if (LocalNode->Culling)
+					{
+						FieldNodesToProcess.Enqueue(LocalNode->Culling.Get());
+						LocalNode->Culling->UserData = 0;
+					}
+					if (LocalNode->Input)
+					{
+						FieldNodesToProcess.Enqueue(LocalNode->Input.Get());
+						LocalNode->Input->UserData = 0;
+					}
+					FieldNodesToProcess.Enqueue(LocalNode);
+					FieldNode->UserData = 1;
+				}
+			}
+		}
+	}
+}
+
+
+void FPhysicsFieldInstance::BuildNodeParams(const TArray<FFieldNodeBase*>& FieldNodes, const TMap<FFieldNodeBase*,float> CommandTimes, const float PreviousTime)
+{
+	for (FFieldNodeBase* FieldNode : FieldNodes)
 	{
 		const float NextTime = CommandTimes.Find(FieldNode) ? CommandTimes[FieldNode] : PreviousTime;
 		if (FieldNode->SerializationType() == FFieldNodeBase::ESerializationType::FieldNode_FUniformInteger)
@@ -977,10 +1208,6 @@ void FPhysicsFieldInstance::BuildNodeParams(FFieldNodeBase* FieldNode, const TMa
 		else if (FieldNode->SerializationType() == FFieldNodeBase::ESerializationType::FieldNode_FSumScalar)
 		{
 			FSumScalar* LocalNode = StaticCast<FSumScalar*>(FieldNode);
-
-			BuildNodeParams(LocalNode->ScalarRight.Get(), CommandTimes, NextTime);
-			BuildNodeParams(LocalNode->ScalarLeft.Get(), CommandTimes, NextTime);
-
 			NodesOffsets.Add(NodesParams.Num());
 			NodesParams.Add(FieldNode->Type());
 			NodesParams.Add(FFieldNodeBase::ESerializationType::FieldNode_FSumScalar);
@@ -992,11 +1219,6 @@ void FPhysicsFieldInstance::BuildNodeParams(FFieldNodeBase* FieldNode, const TMa
 		else if (FieldNode->SerializationType() == FFieldNodeBase::ESerializationType::FieldNode_FSumVector)
 		{
 			FSumVector* LocalNode = StaticCast<FSumVector*>(FieldNode);
-
-			BuildNodeParams(LocalNode->Scalar.Get(), CommandTimes, NextTime);
-			BuildNodeParams(LocalNode->VectorRight.Get(), CommandTimes, NextTime);
-			BuildNodeParams(LocalNode->VectorLeft.Get(), CommandTimes, NextTime);
-
 			NodesOffsets.Add(NodesParams.Num());
 			NodesParams.Add(FieldNode->Type());
 			NodesParams.Add(FFieldNodeBase::ESerializationType::FieldNode_FSumVector);
@@ -1011,9 +1233,6 @@ void FPhysicsFieldInstance::BuildNodeParams(FFieldNodeBase* FieldNode, const TMa
 			if (FieldNode->Type() == FFieldNodeBase::EFieldType::EField_Int32)
 			{
 				FConversionField<float, int32>* LocalNode = StaticCast<FConversionField<float, int32>*>(FieldNode);
-
-				BuildNodeParams(LocalNode->InputField.Get(), CommandTimes, NextTime);
-
 				NodesOffsets.Add(NodesParams.Num());
 				NodesParams.Add(FieldNode->Type());
 				NodesParams.Add(FFieldNodeBase::ESerializationType::FieldNode_FConversionField);
@@ -1022,9 +1241,6 @@ void FPhysicsFieldInstance::BuildNodeParams(FFieldNodeBase* FieldNode, const TMa
 			else if (FieldNode->Type() == FFieldNodeBase::EFieldType::EField_Float)
 			{
 				FConversionField<int32, float>* LocalNode = StaticCast<FConversionField<int32, float>*>(FieldNode);
-
-				BuildNodeParams(LocalNode->InputField.Get(), CommandTimes, NextTime);
-
 				NodesOffsets.Add(NodesParams.Num());
 				NodesParams.Add(FieldNode->Type());
 				NodesParams.Add(FFieldNodeBase::ESerializationType::FieldNode_FConversionField);
@@ -1036,10 +1252,6 @@ void FPhysicsFieldInstance::BuildNodeParams(FFieldNodeBase* FieldNode, const TMa
 			if (FieldNode->Type() == FFieldNodeBase::EFieldType::EField_Int32)
 			{
 				FCullingField<int32>* LocalNode = StaticCast<FCullingField<int32>*>(FieldNode);
-
-				BuildNodeParams(LocalNode->Culling.Get(), CommandTimes, NextTime);
-				BuildNodeParams(LocalNode->Input.Get(), CommandTimes, NextTime);
-
 				NodesOffsets.Add(NodesParams.Num());
 				NodesParams.Add(FieldNode->Type());
 				NodesParams.Add(FFieldNodeBase::ESerializationType::FieldNode_FCullingField);
@@ -1050,10 +1262,6 @@ void FPhysicsFieldInstance::BuildNodeParams(FFieldNodeBase* FieldNode, const TMa
 			else if (FieldNode->Type() == FFieldNodeBase::EFieldType::EField_Float)
 			{
 				FCullingField<float>* LocalNode = StaticCast<FCullingField<float>*>(FieldNode);
-
-				BuildNodeParams(LocalNode->Culling.Get(), CommandTimes, NextTime);
-				BuildNodeParams(LocalNode->Input.Get(), CommandTimes, NextTime);
-
 				NodesOffsets.Add(NodesParams.Num());
 				NodesParams.Add(FieldNode->Type());
 				NodesParams.Add(FFieldNodeBase::ESerializationType::FieldNode_FCullingField);
@@ -1064,10 +1272,6 @@ void FPhysicsFieldInstance::BuildNodeParams(FFieldNodeBase* FieldNode, const TMa
 			else if (FieldNode->Type() == FFieldNodeBase::EFieldType::EField_FVector)
 			{
 				FCullingField<FVector>* LocalNode = StaticCast<FCullingField<FVector>*>(FieldNode);
-
-				BuildNodeParams(LocalNode->Culling.Get(), CommandTimes, NextTime);
-				BuildNodeParams(LocalNode->Input.Get(), CommandTimes, NextTime);
-
 				NodesOffsets.Add(NodesParams.Num());
 				NodesParams.Add(FieldNode->Type());
 				NodesParams.Add(FFieldNodeBase::ESerializationType::FieldNode_FCullingField);
@@ -1079,43 +1283,61 @@ void FPhysicsFieldInstance::BuildNodeParams(FFieldNodeBase* FieldNode, const TMa
 	}
 }
 
-void FPhysicsFieldInstance::BuildNodeBounds(FFieldNodeBase* FieldNode, FVector& MinBounds, FVector& MaxBounds, float& MaxMagnitude)
+void FPhysicsFieldInstance::BuildNodeBounds(const TArray<FFieldNodeBase*>& FieldNodes, FVector& MinBounds, FVector& MaxBounds, float& MaxMagnitude)
 {
+	TArray<FVector> AllMinBounds;
+	TArray<FVector> AllMaxBounds;
+	TArray<float> AllMaxMagnitude;
+
 	MinBounds = FVector(-FLT_MAX);
 	MaxBounds = FVector(FLT_MAX);
 	MaxMagnitude = 1.0;
 
-	if (FieldNode)
+	for (int32 Index = 0; Index < FieldNodes.Num(); ++Index)
 	{
+		FFieldNodeBase* FieldNode = FieldNodes[Index];
+		FieldNode->UserData = Index;
 		if (FieldNode->SerializationType() == FFieldNodeBase::ESerializationType::FieldNode_FUniformInteger)
 		{
 			FUniformInteger* LocalNode = StaticCast<FUniformInteger*>(FieldNode);
-			MaxMagnitude = LocalNode->Magnitude;
+			AllMinBounds.Add(MinBounds);
+			AllMaxBounds.Add(MaxBounds);
+			AllMaxMagnitude.Add(LocalNode->Magnitude);
 		}
 		else if (FieldNode->SerializationType() == FFieldNodeBase::ESerializationType::FieldNode_FUniformScalar)
 		{
 			FUniformScalar* LocalNode = StaticCast<FUniformScalar*>(FieldNode);
-			MaxMagnitude = LocalNode->Magnitude;
+			AllMinBounds.Add(MinBounds);
+			AllMaxBounds.Add(MaxBounds);
+			AllMaxMagnitude.Add(LocalNode->Magnitude);
 		}
 		else if (FieldNode->SerializationType() == FFieldNodeBase::ESerializationType::FieldNode_FWaveScalar)
 		{
 			FWaveScalar* LocalNode = StaticCast<FWaveScalar*>(FieldNode);
-			MaxMagnitude = LocalNode->Magnitude;
+			AllMinBounds.Add(MinBounds);
+			AllMaxBounds.Add(MaxBounds);
+			AllMaxMagnitude.Add(LocalNode->Magnitude);
 		}
 		else if (FieldNode->SerializationType() == FFieldNodeBase::ESerializationType::FieldNode_FUniformVector)
 		{
 			FUniformVector* LocalNode = StaticCast<FUniformVector*>(FieldNode);
-			MaxMagnitude = LocalNode->Magnitude;
+			AllMinBounds.Add(MinBounds);
+			AllMaxBounds.Add(MaxBounds);
+			AllMaxMagnitude.Add(LocalNode->Magnitude);
 		}
 		else if (FieldNode->SerializationType() == FFieldNodeBase::ESerializationType::FieldNode_FRadialVector)
 		{
 			FRadialVector* LocalNode = StaticCast<FRadialVector*>(FieldNode);
-			MaxMagnitude = LocalNode->Magnitude;
+			AllMinBounds.Add(MinBounds);
+			AllMaxBounds.Add(MaxBounds);
+			AllMaxMagnitude.Add(LocalNode->Magnitude);
 		}
 		else if (FieldNode->SerializationType() == FFieldNodeBase::ESerializationType::FieldNode_FRandomVector)
 		{
 			FRandomVector* LocalNode = StaticCast<FRandomVector*>(FieldNode);
-			MaxMagnitude = LocalNode->Magnitude;
+			AllMinBounds.Add(MinBounds);
+			AllMaxBounds.Add(MaxBounds);
+			AllMaxMagnitude.Add(LocalNode->Magnitude);
 		}
 		else if (FieldNode->SerializationType() == FFieldNodeBase::ESerializationType::FieldNode_FRadialIntMask)
 		{
@@ -1124,6 +1346,10 @@ void FPhysicsFieldInstance::BuildNodeBounds(FFieldNodeBase* FieldNode, FVector&
 			MinBounds = (LocalNode->ExteriorValue == 0) ? LocalNode->Position - FVector(LocalNode->Radius) : FVector(-FLT_MAX);
 			MaxBounds = (LocalNode->ExteriorValue == 0) ? LocalNode->Position + FVector(LocalNode->Radius) : FVector(FLT_MAX);
 			MaxMagnitude = 1.0;
+
+			AllMinBounds.Add(MinBounds);
+			AllMaxBounds.Add(MaxBounds);
+			AllMaxMagnitude.Add(MaxMagnitude);
 		}
 		else if (FieldNode->SerializationType() == FFieldNodeBase::ESerializationType::FieldNode_FPlaneFalloff)
 		{
@@ -1131,6 +1357,10 @@ void FPhysicsFieldInstance::BuildNodeBounds(FFieldNodeBase* FieldNode, FVector&
 			MinBounds = FVector(-FLT_MAX);
 			MaxBounds = FVector(FLT_MAX);
 			MaxMagnitude = LocalNode->Magnitude;
+			
+			AllMinBounds.Add(MinBounds);
+			AllMaxBounds.Add(MaxBounds);
+			AllMaxMagnitude.Add(MaxMagnitude);
 		}
 		else if (FieldNode->SerializationType() == FFieldNodeBase::ESerializationType::FieldNode_FRadialFalloff)
 		{
@@ -1139,6 +1369,10 @@ void FPhysicsFieldInstance::BuildNodeBounds(FFieldNodeBase* FieldNode, FVector&
 			MinBounds = (LocalNode->Default == 0) ? LocalNode->Position - FVector(LocalNode->Radius) : FVector(-FLT_MAX);
 			MaxBounds = (LocalNode->Default == 0) ? LocalNode->Position + FVector(LocalNode->Radius) : FVector(FLT_MAX);
 			MaxMagnitude = LocalNode->Magnitude;
+
+			AllMinBounds.Add(MinBounds);
+			AllMaxBounds.Add(MaxBounds);
+			AllMaxMagnitude.Add(MaxMagnitude);
 		}
 		else if (FieldNode->SerializationType() == FFieldNodeBase::ESerializationType::FieldNode_FBoxFalloff)
 		{
@@ -1153,6 +1387,10 @@ void FPhysicsFieldInstance::BuildNodeBounds(FFieldNodeBase* FieldNode, FVector&
 				MaxBounds = BoundingBox.Max;
 			}
 			MaxMagnitude = LocalNode->Magnitude;
+
+			AllMinBounds.Add(MinBounds);
+			AllMaxBounds.Add(MaxBounds);
+			AllMaxMagnitude.Add(MaxMagnitude);
 		}
 		else if (FieldNode->SerializationType() == FFieldNodeBase::ESerializationType::FieldNode_FSumScalar)
 		{
@@ -1160,8 +1398,18 @@ void FPhysicsFieldInstance::BuildNodeBounds(FFieldNodeBase* FieldNode, FVector&
 
 			FVector MinBoundsA(-FLT_MAX), MaxBoundsA(FLT_MAX), MinBoundsB(-FLT_MAX), MaxBoundsB(FLT_MAX);
 			float MaxMagnitudeA = 0.0, MaxMagnitudeB = 0.0;
-			BuildNodeBounds(LocalNode->ScalarRight.Get(), MinBoundsA, MaxBoundsA, MaxMagnitudeA);
-			BuildNodeBounds(LocalNode->ScalarLeft.Get(), MinBoundsB, MaxBoundsB, MaxMagnitudeB);
+			if (LocalNode->ScalarRight)
+			{
+				MinBoundsA = AllMinBounds[LocalNode->ScalarRight->UserData];
+				MaxBoundsA = AllMaxBounds[LocalNode->ScalarRight->UserData];
+				MaxMagnitudeA = AllMaxMagnitude[LocalNode->ScalarRight->UserData];
+			}
+			if (LocalNode->ScalarLeft)
+			{
+				MinBoundsB = AllMinBounds[LocalNode->ScalarLeft->UserData];
+				MaxBoundsB = AllMaxBounds[LocalNode->ScalarLeft->UserData];
+				MaxMagnitudeB = AllMaxMagnitude[LocalNode->ScalarLeft->UserData];
+			}
 
 			if (LocalNode->Operation == EFieldOperationType::Field_Multiply ||
 				LocalNode->Operation == EFieldOperationType::Field_Divide)
@@ -1177,6 +1425,10 @@ void FPhysicsFieldInstance::BuildNodeBounds(FFieldNodeBase* FieldNode, FVector&
 				MaxBounds = MaxVector(MaxBoundsA, MaxBoundsB);
 				MaxMagnitude = FMath::Max(MaxMagnitudeA,MaxMagnitudeB);
 			}
+
+			AllMinBounds.Add(MinBounds);
+			AllMaxBounds.Add(MaxBounds);
+			AllMaxMagnitude.Add(MaxMagnitude);
 		}
 		else if (FieldNode->SerializationType() == FFieldNodeBase::ESerializationType::FieldNode_FSumVector)
 		{
@@ -1184,9 +1436,24 @@ void FPhysicsFieldInstance::BuildNodeBounds(FFieldNodeBase* FieldNode, FVector&
 
 			FVector MinBoundsA(-FLT_MAX), MaxBoundsA(FLT_MAX), MinBoundsB(-FLT_MAX), MaxBoundsB(FLT_MAX), MinBoundsC(-FLT_MAX), MaxBoundsC(FLT_MAX);
 			float MaxMagnitudeA = 0.0, MaxMagnitudeB = 0.0, MaxMagnitudeC = 0.0;
-			BuildNodeBounds(LocalNode->Scalar.Get(), MinBoundsA, MaxBoundsA, MaxMagnitudeA);
-			BuildNodeBounds(LocalNode->VectorRight.Get(), MinBoundsB, MaxBoundsB, MaxMagnitudeB);
-			BuildNodeBounds(LocalNode->VectorLeft.Get(), MinBoundsC, MaxBoundsC, MaxMagnitudeC);
+			if (LocalNode->Scalar)
+			{
+				MinBoundsA = AllMinBounds[LocalNode->Scalar->UserData];
+				MaxBoundsA = AllMaxBounds[LocalNode->Scalar->UserData];
+				MaxMagnitudeA = AllMaxMagnitude[LocalNode->Scalar->UserData];
+			}
+			if (LocalNode->VectorRight)
+			{
+				MinBoundsB = AllMinBounds[LocalNode->VectorRight->UserData];
+				MaxBoundsB = AllMaxBounds[LocalNode->VectorRight->UserData];
+				MaxMagnitudeB = AllMaxMagnitude[LocalNode->VectorRight->UserData];
+			}
+			if (LocalNode->VectorLeft)
+			{
+				MinBoundsC = AllMinBounds[LocalNode->VectorLeft->UserData];
+				MaxBoundsC = AllMaxBounds[LocalNode->VectorLeft->UserData];
+				MaxMagnitudeC = AllMaxMagnitude[LocalNode->VectorLeft->UserData];
+			}
 
 			if (LocalNode->Operation == EFieldOperationType::Field_Multiply ||
 				LocalNode->Operation == EFieldOperationType::Field_Divide)
@@ -1205,6 +1472,10 @@ void FPhysicsFieldInstance::BuildNodeBounds(FFieldNodeBase* FieldNode, FVector&
 			MinBounds = MaxVector(MinBounds, MinBoundsA);
 			MaxBounds = MinVector(MaxBounds, MaxBoundsA);
 			MaxMagnitude = MaxMagnitude * MaxMagnitudeA;
+
+			AllMinBounds.Add(MinBounds);
+			AllMaxBounds.Add(MaxBounds);
+			AllMaxMagnitude.Add(MaxMagnitude);
 		}
 		else if (FieldNode->SerializationType() == FFieldNodeBase::ESerializationType::FieldNode_FCullingField)
 		{
@@ -1214,8 +1485,18 @@ void FPhysicsFieldInstance::BuildNodeBounds(FFieldNodeBase* FieldNode, FVector&
 
 				FVector MinBoundsA(-FLT_MAX), MaxBoundsA(FLT_MAX), MinBoundsB(-FLT_MAX), MaxBoundsB(FLT_MAX);
 				float MaxMagnitudeA = 0.0, MaxMagnitudeB = 0.0;
-				BuildNodeBounds(LocalNode->Culling.Get(), MinBoundsA, MaxBoundsA, MaxMagnitudeA);
-				BuildNodeBounds(LocalNode->Input.Get(), MinBoundsB, MaxBoundsB, MaxMagnitudeB);
+				if (LocalNode->Culling)
+				{
+					MinBoundsA = AllMinBounds[LocalNode->Culling->UserData];
+					MaxBoundsA = AllMaxBounds[LocalNode->Culling->UserData];
+					MaxMagnitudeA = AllMaxMagnitude[LocalNode->Culling->UserData];
+				}
+				if (LocalNode->Input)
+				{
+					MinBoundsB = AllMinBounds[LocalNode->Input->UserData];
+					MaxBoundsB = AllMaxBounds[LocalNode->Input->UserData];
+					MaxMagnitudeB = AllMaxMagnitude[LocalNode->Input->UserData];
+				}
 
 				if (LocalNode->Operation == EFieldCullingOperationType::Field_Culling_Inside)
 				{
@@ -1228,6 +1509,10 @@ void FPhysicsFieldInstance::BuildNodeBounds(FFieldNodeBase* FieldNode, FVector&
 					MaxBounds = MinVector(MaxBoundsA, MaxBoundsB);
 				}
 				MaxMagnitude = MaxMagnitudeB;
+
+				AllMinBounds.Add(MinBounds);
+				AllMaxBounds.Add(MaxBounds);
+				AllMaxMagnitude.Add(MaxMagnitude);
 			}
 			else if (FieldNode->Type() == FFieldNodeBase::EFieldType::EField_Float)
 			{
@@ -1235,8 +1520,18 @@ void FPhysicsFieldInstance::BuildNodeBounds(FFieldNodeBase* FieldNode, FVector&
 
 				FVector MinBoundsA(-FLT_MAX), MaxBoundsA(FLT_MAX), MinBoundsB(-FLT_MAX), MaxBoundsB(FLT_MAX);
 				float MaxMagnitudeA = 0.0, MaxMagnitudeB = 0.0;
-				BuildNodeBounds(LocalNode->Culling.Get(), MinBoundsA, MaxBoundsA, MaxMagnitudeA);
-				BuildNodeBounds(LocalNode->Input.Get(), MinBoundsB, MaxBoundsB, MaxMagnitudeB);
+				if (LocalNode->Culling)
+				{
+					MinBoundsA = AllMinBounds[LocalNode->Culling->UserData];
+					MaxBoundsA = AllMaxBounds[LocalNode->Culling->UserData];
+					MaxMagnitudeA = AllMaxMagnitude[LocalNode->Culling->UserData];
+				}
+				if (LocalNode->Input)
+				{
+					MinBoundsB = AllMinBounds[LocalNode->Input->UserData];
+					MaxBoundsB = AllMaxBounds[LocalNode->Input->UserData];
+					MaxMagnitudeB = AllMaxMagnitude[LocalNode->Input->UserData];
+				}
 
 				if (LocalNode->Operation == EFieldCullingOperationType::Field_Culling_Inside)
 				{
@@ -1249,6 +1544,10 @@ void FPhysicsFieldInstance::BuildNodeBounds(FFieldNodeBase* FieldNode, FVector&
 					MaxBounds = MinVector(MaxBoundsA, MaxBoundsB);
 				}
 				MaxMagnitude = MaxMagnitudeB;
+
+				AllMinBounds.Add(MinBounds);
+				AllMaxBounds.Add(MaxBounds);
+				AllMaxMagnitude.Add(MaxMagnitude);
 			}
 			else if (FieldNode->Type() == FFieldNodeBase::EFieldType::EField_FVector)
 			{
@@ -1256,8 +1555,18 @@ void FPhysicsFieldInstance::BuildNodeBounds(FFieldNodeBase* FieldNode, FVector&
 
 				FVector MinBoundsA(-FLT_MAX), MaxBoundsA(FLT_MAX), MinBoundsB(-FLT_MAX), MaxBoundsB(FLT_MAX);
 				float MaxMagnitudeA = 0.0, MaxMagnitudeB = 0.0;
-				BuildNodeBounds(LocalNode->Culling.Get(), MinBoundsA, MaxBoundsA, MaxMagnitudeA);
-				BuildNodeBounds(LocalNode->Input.Get(), MinBoundsB, MaxBoundsB, MaxMagnitudeB);
+				if (LocalNode->Culling)
+				{
+					MinBoundsA = AllMinBounds[LocalNode->Culling->UserData];
+					MaxBoundsA = AllMaxBounds[LocalNode->Culling->UserData];
+					MaxMagnitudeA = AllMaxMagnitude[LocalNode->Culling->UserData];
+				}
+				if (LocalNode->Input)
+				{
+					MinBoundsB = AllMinBounds[LocalNode->Input->UserData];
+					MaxBoundsB = AllMaxBounds[LocalNode->Input->UserData];
+					MaxMagnitudeB = AllMaxMagnitude[LocalNode->Input->UserData];
+				}
 
 				if (LocalNode->Operation == EFieldCullingOperationType::Field_Culling_Inside)
 				{
@@ -1270,6 +1579,10 @@ void FPhysicsFieldInstance::BuildNodeBounds(FFieldNodeBase* FieldNode, FVector&
 					MaxBounds = MinVector(MaxBoundsA, MaxBoundsB);
 				}
 				MaxMagnitude = MaxMagnitudeB;
+
+				AllMinBounds.Add(MinBounds);
+				AllMaxBounds.Add(MaxBounds);
+				AllMaxMagnitude.Add(MaxMagnitude);
 			}
 		}
 		else if (FieldNode->SerializationType() == FFieldNodeBase::ESerializationType::FieldNode_FConversionField)
@@ -1279,25 +1592,47 @@ void FPhysicsFieldInstance::BuildNodeBounds(FFieldNodeBase* FieldNode, FVector&
 				FConversionField<float, int32>* LocalNode = StaticCast<FConversionField<float, int32>*>(FieldNode);
 				FVector MinBoundsA(-FLT_MAX), MaxBoundsA(FLT_MAX);
 				float MaxMagnitudeA = 0.0;
+				if (LocalNode->InputField)
+				{
+					MinBoundsA = AllMinBounds[LocalNode->InputField->UserData];
+					MaxBoundsA = AllMaxBounds[LocalNode->InputField->UserData];
+					MaxMagnitudeA = AllMaxMagnitude[LocalNode->InputField->UserData];
+				}
 
-				BuildNodeBounds(LocalNode->InputField.Get(), MinBoundsA, MaxBoundsA, MaxMagnitudeA);
 				MinBounds = MinBoundsA;
 				MaxBounds = MaxBoundsA;
 				MaxMagnitude = MaxMagnitudeA;
+
+				AllMinBounds.Add(MinBounds);
+				AllMaxBounds.Add(MaxBounds);
+				AllMaxMagnitude.Add(MaxMagnitude);
 			}
 			else if (FieldNode->Type() == FFieldNodeBase::EFieldType::EField_Float)
 			{
 				FConversionField<int32, float>* LocalNode = StaticCast<FConversionField<int32, float>*>(FieldNode);
 				FVector MinBoundsA(-FLT_MAX), MaxBoundsA(FLT_MAX);
 				float MaxMagnitudeA = 0.0;
+				if (LocalNode->InputField)
+				{
+					MinBoundsA = AllMinBounds[LocalNode->InputField->UserData];
+					MaxBoundsA = AllMaxBounds[LocalNode->InputField->UserData];
+					MaxMagnitudeA = AllMaxMagnitude[LocalNode->InputField->UserData];
+				}
 
-				BuildNodeBounds(LocalNode->InputField.Get(), MinBoundsA, MaxBoundsA, MaxMagnitudeA);
 				MinBounds = MinBoundsA;
 				MaxBounds = MaxBoundsA;
 				MaxMagnitude = MaxMagnitudeA;
+
+				AllMinBounds.Add(MinBounds);
+				AllMaxBounds.Add(MaxBounds);
+				AllMaxMagnitude.Add(MaxMagnitude);
 			}
 		}
 	}
+
+	MinBounds = AllMinBounds.Last();
+	MaxBounds = AllMaxBounds.Last();
+	MaxMagnitude = AllMaxMagnitude.Last();
 }
 
 /**
@@ -1462,7 +1797,11 @@ void UPhysicsFieldComponent::BuildCommandBounds(FFieldSystemCommand& FieldComman
 	FieldCommand.BoundingBox.Max = FVector(FLT_MAX);
 	FieldCommand.MaxMagnitude = 1.0;
 
-	FPhysicsFieldInstance::BuildNodeBounds(FieldCommand.RootNode.Get(), FieldCommand.BoundingBox.Min, FieldCommand.BoundingBox.Max, FieldCommand.MaxMagnitude);
+	TArray<FFieldNodeBase*> FieldNodes;
+	TQueue<FFieldNodeBase*> FieldNodesToProcess;
+	FieldNodesToProcess.Enqueue(FieldCommand.RootNode.Get());
+	FPhysicsFieldInstance::BuildFieldArray(FieldNodesToProcess, FieldNodes);
+	FPhysicsFieldInstance::BuildNodeBounds(FieldNodes, FieldCommand.BoundingBox.Min, FieldCommand.BoundingBox.Max, FieldCommand.MaxMagnitude);
 }
 
 void UPhysicsFieldComponent::AddTransientCommand(const FFieldSystemCommand& FieldCommand, const bool bIsWorldField)
diff --git a/Engine/Source/Runtime/Experimental/Chaos/Public/Field/FieldSystem.h b/Engine/Source/Runtime/Experimental/Chaos/Public/Field/FieldSystem.h
--- a/Engine/Source/Runtime/Experimental/Chaos/Public/Field/FieldSystem.h
+++ b/Engine/Source/Runtime/Experimental/Chaos/Public/Field/FieldSystem.h
@@ -435,6 +435,8 @@ public:
 		FieldNode_FReturnResultsTerminal
 	};
 
+	uint32 UserData;
+
 	FFieldNodeBase() : Name("") {}
 	virtual ~FFieldNodeBase() {}
 	virtual EFieldType Type() const { check(false); return EFieldType::EField_None; }
