World Partition - HLOD: Added events for HLOD registration/unregistration in the world

diff --git a/Engine/Source/Runtime/Engine/Private/WorldPartition/HLOD/HLODSubsystem.cpp b/Engine/Source/Runtime/Engine/Private/WorldPartition/HLOD/HLODSubsystem.cpp
--- a/Engine/Source/Runtime/Engine/Private/WorldPartition/HLOD/HLODSubsystem.cpp
+++ b/Engine/Source/Runtime/Engine/Private/WorldPartition/HLOD/HLODSubsystem.cpp
@@ -22,6 +22,7 @@
 #include "Components/StaticMeshComponent.h"
 #include "HAL/FileManager.h"
 #include "Misc/FileHelper.h"
+#include "SceneViewExtension.h"
 #include "StaticMeshResources.h"
 #include "Templates/UniquePtr.h"
 
@@ -68,30 +69,48 @@ static TAutoConsoleVariable<int32> CVarHLODWarmupVTSizeClamp(
 
 namespace FHLODSubsystem
 {
-	const UWorldPartitionRuntimeCell* GetActorRuntimeCell(AActor* InActor)
+    static const UWorldPartitionRuntimeCell* GetActorRuntimeCell(AActor* InActor)
+    {
+	    const ULevel* ActorLevel = InActor->GetLevel();
+	    const ULevelStreaming* LevelStreaming = FLevelUtils::FindStreamingLevel(ActorLevel);
+	    const UWorldPartitionLevelStreamingDynamic* LevelStreamingDynamic = Cast<UWorldPartitionLevelStreamingDynamic>(LevelStreaming);
+	    return LevelStreamingDynamic ? LevelStreamingDynamic->GetWorldPartitionRuntimeCell() : nullptr;
+    }
+    
+    static UWorldPartition* GetWorldPartition(AWorldPartitionHLOD* InWorldPartitionHLOD)
+    {
+	    // Alwaysloaded Cell level will have a WorldPartition
+	    if (UWorldPartition* WorldPartition = InWorldPartitionHLOD->GetLevel()->GetWorldPartition())
+	    {
+		    return WorldPartition;
+	    } // If not find it through the cell
+	    else if (const UWorldPartitionRuntimeCell* RuntimeCell = GetActorRuntimeCell(InWorldPartitionHLOD))
+	    {
+		    return RuntimeCell->GetCellOwner()->GetOuterWorld()->GetWorldPartition();
+	    }
+    
+	    return nullptr;
+    }
+}
+
+
+class FHLODResourcesResidencySceneViewExtension : public FWorldSceneViewExtension
+{
+public:
+	FHLODResourcesResidencySceneViewExtension(const FAutoRegister& AutoRegister, UWorld* InWorld)
+		: FWorldSceneViewExtension(AutoRegister, InWorld)
 	{
-		const ULevel* ActorLevel = InActor->GetLevel();
-		const ULevelStreaming* LevelStreaming = FLevelUtils::FindStreamingLevel(ActorLevel);
-		const UWorldPartitionLevelStreamingDynamic* LevelStreamingDynamic = Cast<UWorldPartitionLevelStreamingDynamic>(LevelStreaming);
-		return LevelStreamingDynamic ? LevelStreamingDynamic->GetWorldPartitionRuntimeCell() : nullptr;
 	}
 
-	UWorldPartition* GetWorldPartition(AWorldPartitionHLOD* InWorldPartitionHLOD)
+	virtual void SetupViewFamily(FSceneViewFamily& InViewFamily) override {}
+	virtual void SetupView(FSceneViewFamily& InViewFamily, FSceneView& InView) override {}
+	virtual void BeginRenderViewFamily(FSceneViewFamily& InViewFamily) override
 	{
-		// Alwaysloaded Cell level will have a WorldPartition
-		if (UWorldPartition* WorldPartition = InWorldPartitionHLOD->GetLevel()->GetWorldPartition())
-		{
-			return WorldPartition;
-		} // If not find it through the cell
-		else if (const UWorldPartitionRuntimeCell* RuntimeCell = GetActorRuntimeCell(InWorldPartitionHLOD))
-		{
-			return RuntimeCell->GetCellOwner()->GetOuterWorld()->GetWorldPartition();
-		}
-
-		return nullptr;
+		GetWorld()->GetSubsystem<UHLODSubsystem>()->OnBeginRenderViews(InViewFamily);
 	}
 };
 
+
 UHLODSubsystem::UHLODSubsystem()
 	: UWorldSubsystem()
 {
@@ -225,6 +244,8 @@ void UHLODSubsystem::RegisterHLODActor(AWorldPartitionHLOD* InWorldPartitionHLOD
 			InWorldPartitionHLOD->SetVisibility(false);
 		}
 	}
+
+	HLODActorRegisteredEvent.Broadcast(InWorldPartitionHLOD);
 }
 
 void UHLODSubsystem::UnregisterHLODActor(AWorldPartitionHLOD* InWorldPartitionHLOD)
@@ -249,6 +270,8 @@ void UHLODSubsystem::UnregisterHLODActor(AWorldPartitionHLOD* InWorldPartitionHL
 			check(NumRemoved == 1);
 		}
 	}
+
+	HLODActorUnregisteredEvent.Broadcast(InWorldPartitionHLOD);
 }
 
 const TArray<AWorldPartitionHLOD*>& UHLODSubsystem::GetHLODActorsForCell(const UWorldPartitionRuntimeCell* InCell) const
@@ -520,11 +543,6 @@ void UHLODSubsystem::OnBeginRenderViews(const FSceneViewFamily& InViewFamily)
 	}
 }
 
-void FHLODResourcesResidencySceneViewExtension::BeginRenderViewFamily(FSceneViewFamily& InViewFamily)
-{
-	GetWorld()->GetSubsystem<UHLODSubsystem>()->OnBeginRenderViews(InViewFamily);
-}
-
 #if WITH_EDITOR
 
 static void DumpHLODStats(UWorld* World)
diff --git a/Engine/Source/Runtime/Engine/Public/WorldPartition/HLOD/HLODSubsystem.h b/Engine/Source/Runtime/Engine/Public/WorldPartition/HLOD/HLODSubsystem.h
--- a/Engine/Source/Runtime/Engine/Public/WorldPartition/HLOD/HLODSubsystem.h
+++ b/Engine/Source/Runtime/Engine/Public/WorldPartition/HLOD/HLODSubsystem.h
@@ -4,28 +4,19 @@
 #include "Subsystems/WorldSubsystem.h"
 #include "Containers/Array.h"
 #include "Containers/Map.h"
-#include "UObject/NameTypes.h"
-#include "SceneViewExtension.h"
-
 #include "HLODSubsystem.generated.h"
 
+
+class AWorldPartitionHLOD;
+class FSceneViewFamily;
+class FHLODResourcesResidencySceneViewExtension;
 class UWorldPartition;
 class UWorldPartitionRuntimeCell;
-class AWorldPartitionHLOD;
 
-class FHLODResourcesResidencySceneViewExtension : public FWorldSceneViewExtension
-{
-public:
-	FHLODResourcesResidencySceneViewExtension(const FAutoRegister& AutoRegister, UWorld* InWorld)
-		: FWorldSceneViewExtension(AutoRegister, InWorld)
-	{
-	}
 
-	virtual void SetupViewFamily(FSceneViewFamily& InViewFamily) override {}
-	virtual void SetupView(FSceneViewFamily& InViewFamily, FSceneView& InView) override {}
+DECLARE_MULTICAST_DELEGATE_OneParam(FWorldPartitionHLODActorRegisteredEvent, AWorldPartitionHLOD* /* InHLODActor */);
+DECLARE_MULTICAST_DELEGATE_OneParam(FWorldPartitionHLODActorUnregisteredEvent, AWorldPartitionHLOD* /* InHLODActor */);
 
-	virtual void BeginRenderViewFamily(FSceneViewFamily& InViewFamily) override;
-};
 
 /**
  * UHLODSubsystem
@@ -59,6 +50,9 @@ public:
 	const TArray<AWorldPartitionHLOD*>& GetHLODActorsForCell(const UWorldPartitionRuntimeCell* InCell) const;
 
 	static bool IsHLODEnabled();
+
+	FWorldPartitionHLODActorRegisteredEvent& OnHLODActorRegisteredEvent() { return HLODActorRegisteredEvent; }
+	FWorldPartitionHLODActorUnregisteredEvent& OnHLODActorUnregisteredEvent() { return HLODActorUnregisteredEvent; }
 	
 private:
 	struct FCellData
@@ -97,4 +91,8 @@ private:
 
 	friend class FHLODResourcesResidencySceneViewExtension;
 	TSharedPtr<FHLODResourcesResidencySceneViewExtension, ESPMode::ThreadSafe> SceneViewExtension;
+
+	FWorldPartitionHLODActorRegisteredEvent		HLODActorRegisteredEvent;
+	FWorldPartitionHLODActorUnregisteredEvent	HLODActorUnregisteredEvent;
 };
+
