Renamed FSpatialHashStreamingGrid::GetAlwaysLoadedCells to FSpatialHashStreamingGrid::GetNonSpatiallyLoadedCells for clarity.

diff --git a/Engine/Source/Runtime/Engine/Private/WorldPartition/WorldPartitionRuntimeSpatialHash.cpp b/Engine/Source/Runtime/Engine/Private/WorldPartition/WorldPartitionRuntimeSpatialHash.cpp
--- a/Engine/Source/Runtime/Engine/Private/WorldPartition/WorldPartitionRuntimeSpatialHash.cpp
+++ b/Engine/Source/Runtime/Engine/Private/WorldPartition/WorldPartitionRuntimeSpatialHash.cpp
@@ -280,7 +280,7 @@ void FSpatialHashStreamingGrid::GetCells(const TArray<FWorldPartitionStreamingSo
 		});
 	}
 
-	GetAlwaysLoadedCells(DataLayerSubsystem, OutActivateCells.GetCells(), OutLoadCells.GetCells());
+	GetNonSpatiallyLoadedCells(DataLayerSubsystem, OutActivateCells.GetCells(), OutLoadCells.GetCells());
 
 	if (!GRuntimeSpatialHashUseAlignedGridLevels && GRuntimeSpatialHashSnapNonAlignedGridLevelsToLowerLevels)
 	{
@@ -518,7 +518,7 @@ void FSpatialHashStreamingGrid::ForEachRuntimeCell(TFunctionRef<bool(const UWorl
 	ForEachGridLevel(InjectedGridLevels);
 }
 
-void FSpatialHashStreamingGrid::GetAlwaysLoadedCells(const UDataLayerSubsystem* DataLayerSubsystem, TSet<const UWorldPartitionRuntimeCell*>& OutActivateCells, TSet<const UWorldPartitionRuntimeCell*>& OutLoadCells) const
+void FSpatialHashStreamingGrid::GetNonSpatiallyLoadedCells(const UDataLayerSubsystem* DataLayerSubsystem, TSet<const UWorldPartitionRuntimeCell*>& OutActivateCells, TSet<const UWorldPartitionRuntimeCell*>& OutLoadCells) const
 {
 	if (GridLevels.Num() > 0)
 	{
@@ -1607,7 +1607,7 @@ void UWorldPartitionRuntimeSpatialHash::ForEachStreamingCellsSources(const TArra
 		{
 			if (IsCellRelevantFor(StreamingGrid.bClientOnlyVisible))
 			{
-				StreamingGrid.GetAlwaysLoadedCells(DataLayerSubsystem, ActivateStreamingSourceCells.GetCells(), LoadStreamingSourceCells.GetCells());
+				StreamingGrid.GetNonSpatiallyLoadedCells(DataLayerSubsystem, ActivateStreamingSourceCells.GetCells(), LoadStreamingSourceCells.GetCells());
 			}
 		});
 	}
diff --git a/Engine/Source/Runtime/Engine/Public/WorldPartition/WorldPartitionRuntimeSpatialHash.h b/Engine/Source/Runtime/Engine/Public/WorldPartition/WorldPartitionRuntimeSpatialHash.h
--- a/Engine/Source/Runtime/Engine/Public/WorldPartition/WorldPartitionRuntimeSpatialHash.h
+++ b/Engine/Source/Runtime/Engine/Public/WorldPartition/WorldPartitionRuntimeSpatialHash.h
@@ -143,7 +143,7 @@ struct FSpatialHashStreamingGrid
 	int64 GetCellSize(int32 Level) const;
 	void GetCells(const FWorldPartitionStreamingQuerySource& QuerySource, TSet<const UWorldPartitionRuntimeCell*>& OutCells, bool bEnableZCulling) const;
 	void GetCells(const TArray<FWorldPartitionStreamingSource>& Sources, const class UDataLayerSubsystem* DataLayerSubsystem, UWorldPartitionRuntimeHash::FStreamingSourceCells& OutActivateCells, UWorldPartitionRuntimeHash::FStreamingSourceCells& OutLoadCells, bool bEnableZCulling) const;
-	void GetAlwaysLoadedCells(const UDataLayerSubsystem* DataLayerSubsystem, TSet<const UWorldPartitionRuntimeCell*>& OutActivateCells, TSet<const UWorldPartitionRuntimeCell*>& OutLoadCells) const;
+	void GetNonSpatiallyLoadedCells(const UDataLayerSubsystem* DataLayerSubsystem, TSet<const UWorldPartitionRuntimeCell*>& OutActivateCells, TSet<const UWorldPartitionRuntimeCell*>& OutLoadCells) const;
 	void Draw2D(UCanvas* Canvas, UWorld* World, const TArray<FWorldPartitionStreamingSource>& Sources, const FBox& Region, const FBox2D& GridScreenBounds, TFunctionRef<FVector2D(const FVector2D&)> WorldToScreen) const;
 	void Draw3D(UWorld* World, const TArray<FWorldPartitionStreamingSource>& Sources, const FTransform& Transform) const;
 	void ForEachRuntimeCell(TFunctionRef<bool(const UWorldPartitionRuntimeCell*)> Func) const;
