World Partition: don't force a GC when unloading data layers.

#rnx

diff --git a/Engine/Source/Runtime/Engine/Classes/Engine/CoreSettings.h b/Engine/Source/Runtime/Engine/Classes/Engine/CoreSettings.h
--- a/Engine/Source/Runtime/Engine/Classes/Engine/CoreSettings.h
+++ b/Engine/Source/Runtime/Engine/Classes/Engine/CoreSettings.h
@@ -147,8 +147,6 @@ extern ENGINE_API float GLevelStreamingUnregisterComponentsTimeLimit;
 extern ENGINE_API int32 GLevelStreamingForceGCAfterLevelStreamedOut;
 /** Whether to kick off incremental GC when there are over the specified amount of levels still waiting to be purged. */
 extern ENGINE_API int32 GLevelStreamingContinuouslyIncrementalGCWhileLevelsPendingPurge;
-/** Whether to override GLevelStreamingContinuouslyIncrementalGCWhileLevelsPendingPurge temporarily until at least one level needs to be purged.  */
-extern ENGINE_API int32 GLevelStreamingContinuouslyIncrementalGCWhileLevelsPendingPurgeOverride;
 /** Enables level streaming requests while async loading (of anything) while the match is already in progress and no loading screen is up. */
 extern ENGINE_API int32 GLevelStreamingAllowLevelRequestsWhileAsyncLoadingInMatch;
 /** When we're already loading this many levels and actively in match, don't allow any more requests until one of those completes.  Set to zero to disable. */
diff --git a/Engine/Source/Runtime/Engine/Private/CoreSettings.cpp b/Engine/Source/Runtime/Engine/Private/CoreSettings.cpp
--- a/Engine/Source/Runtime/Engine/Private/CoreSettings.cpp
+++ b/Engine/Source/Runtime/Engine/Private/CoreSettings.cpp
@@ -25,7 +25,6 @@ int32 GLevelStreamingComponentsUnregistrationGranularity = 5;
 int32 GLevelStreamingRouteActorInitializationGranularity = 10;
 int32 GLevelStreamingForceGCAfterLevelStreamedOut = 1;
 int32 GLevelStreamingContinuouslyIncrementalGCWhileLevelsPendingPurge = 1;
-int32 GLevelStreamingContinuouslyIncrementalGCWhileLevelsPendingPurgeOverride = 0;
 int32 GLevelStreamingAllowLevelRequestsWhileAsyncLoadingInMatch = 1;
 int32 GLevelStreamingMaxLevelRequestsAtOnceWhileInMatch = 0;
 int32 GLevelStreamingForceVerifyLevelsGotRemovedByGC = 0;
diff --git a/Engine/Source/Runtime/Engine/Private/World.cpp b/Engine/Source/Runtime/Engine/Private/World.cpp
--- a/Engine/Source/Runtime/Engine/Private/World.cpp
+++ b/Engine/Source/Runtime/Engine/Private/World.cpp
@@ -3953,8 +3953,7 @@ void UWorld::UpdateLevelStreaming()
 	if (GLevelStreamingContinuouslyIncrementalGCWhileLevelsPendingPurge)
 	{
 		// Figure out whether there are any levels we haven't collected garbage yet.
-		const int32 LevelStreamingContinuouslyIncrementalGCWhileLevelsPendingPurge = GLevelStreamingContinuouslyIncrementalGCWhileLevelsPendingPurgeOverride ? 1 : GLevelStreamingContinuouslyIncrementalGCWhileLevelsPendingPurge;
-		const bool bAreLevelsPendingPurge = FLevelStreamingGCHelper::GetNumLevelsPendingPurge() >= LevelStreamingContinuouslyIncrementalGCWhileLevelsPendingPurge;
+		const bool bAreLevelsPendingPurge = FLevelStreamingGCHelper::GetNumLevelsPendingPurge() >= GLevelStreamingContinuouslyIncrementalGCWhileLevelsPendingPurge;
 
 		// Request a 'soft' GC if there are levels pending purge and there are levels to be loaded. In the case of a blocking
 		// load this is going to guarantee GC firing first thing afterwards and otherwise it is going to sneak in right before
@@ -3962,7 +3961,6 @@ void UWorld::UpdateLevelStreaming()
 		if (bAreLevelsPendingPurge)
 		{
 			GEngine->ForceGarbageCollection(false);
-			GLevelStreamingContinuouslyIncrementalGCWhileLevelsPendingPurgeOverride = 0;
 		}
 	}
 
diff --git a/Engine/Source/Runtime/Engine/Private/WorldPartition/DataLayer/WorldDataLayers.cpp b/Engine/Source/Runtime/Engine/Private/WorldPartition/DataLayer/WorldDataLayers.cpp
--- a/Engine/Source/Runtime/Engine/Private/WorldPartition/DataLayer/WorldDataLayers.cpp
+++ b/Engine/Source/Runtime/Engine/Private/WorldPartition/DataLayer/WorldDataLayers.cpp
@@ -184,10 +184,6 @@ void AWorldDataLayers::SetDataLayerRuntimeState(const UDataLayerInstance* InData
 			{
 				ActiveDataLayerNames.Add(InDataLayerInstance->GetDataLayerFName());
 			}
-			else if (InState == EDataLayerRuntimeState::Unloaded)
-			{
-				GLevelStreamingContinuouslyIncrementalGCWhileLevelsPendingPurgeOverride = 1;
-			}
 
 			// Update Replicated Properties
 			RepActiveDataLayerNames = ActiveDataLayerNames.Array();
