Optimized UWorldPartitionStreamingPolicy::UpdateStreamingState
Also added trace cpu profiler events

diff --git a/Engine/Source/Runtime/Engine/Private/WorldPartition/WorldPartitionStreamingPolicy.cpp b/Engine/Source/Runtime/Engine/Private/WorldPartition/WorldPartitionStreamingPolicy.cpp
--- a/Engine/Source/Runtime/Engine/Private/WorldPartition/WorldPartitionStreamingPolicy.cpp
+++ b/Engine/Source/Runtime/Engine/Private/WorldPartition/WorldPartitionStreamingPolicy.cpp
@@ -54,11 +54,13 @@ static FAutoConsoleVariableRef CVarBlockOnSlowStreaming(
 	GBlockOnSlowStreaming,
 	TEXT("Set if streaming needs to block when to slow to catchup."));
 
-static void SortStreamingCellsByImportance(const TSet<const UWorldPartitionRuntimeCell*>& InCells, TArray<const UWorldPartitionRuntimeCell*>& OutSortedCells)
+static void SortStreamingCellsByImportance(TArray<const UWorldPartitionRuntimeCell*>& InOutCells)
 {
-	TRACE_CPUPROFILER_EVENT_SCOPE(SortStreamingCellsByImportance);
-	OutSortedCells = InCells.Array();
-	Algo::Sort(OutSortedCells, [](const UWorldPartitionRuntimeCell* CellA, const UWorldPartitionRuntimeCell* CellB) { return CellA->SortCompare(CellB) < 0; });
+	if (InOutCells.Num() > 1)
+	{
+		TRACE_CPUPROFILER_EVENT_SCOPE(SortStreamingCellsByImportance);
+		Algo::Sort(InOutCells, [](const UWorldPartitionRuntimeCell* CellA, const UWorldPartitionRuntimeCell* CellB) { return CellA->SortCompare(CellB) < 0; });
+	}
 }
 
 UWorldPartitionStreamingPolicy::UWorldPartitionStreamingPolicy(const FObjectInitializer& ObjectInitializer)
@@ -138,15 +140,18 @@ void UWorldPartitionStreamingPolicy::UpdateStreamingSources()
 			}
 		}
 	}
-		
+
+	for (auto& Pair : StreamingSourcesVelocity)
+	{
+		Pair.Value.Invalidate();
+	}
+
 	// Update streaming sources velocity
 	const float CurrentTime = GetWorld()->GetTimeSeconds();
-	TSet<FName, DefaultKeyFuncs<FName>, TInlineSetAllocator<8>> ValidStreamingSources;
 	for (FWorldPartitionStreamingSource& StreamingSource : StreamingSources)
 	{
 		if (!StreamingSource.Name.IsNone())
 		{
-			ValidStreamingSources.Add(StreamingSource.Name);
 			FStreamingSourceVelocity& SourceVelocity = StreamingSourcesVelocity.FindOrAdd(StreamingSource.Name, FStreamingSourceVelocity(StreamingSource.Name));
 			StreamingSource.Velocity = SourceVelocity.GetAverageVelocity(StreamingSource.Location, CurrentTime);
 		}
@@ -155,7 +160,7 @@ void UWorldPartitionStreamingPolicy::UpdateStreamingSources()
 	// Cleanup StreamingSourcesVelocity
 	for (auto It(StreamingSourcesVelocity.CreateIterator()); It; ++It)
 	{
-		if (!ValidStreamingSources.Contains(It.Key()))
+		if (!It.Value().IsValid())
 		{
 			It.RemoveCurrent();
 		}
@@ -232,6 +237,8 @@ void UWorldPartitionStreamingPolicy::UpdateStreamingState()
 		// When world partition can't stream, all cells must be unloaded
 		if (WorldPartition->CanStream() && WorldPartition->RuntimeHash)
 		{
+			TRACE_CPUPROFILER_EVENT_SCOPE(UWorldPartitionStreamingPolicy::UpdateStreamingState_ForEachStreamingCellsSources);
+
 			UWorldPartitionRuntimeCell::DirtyStreamingSourceCacheEpoch();
 
 			WorldPartition->RuntimeHash->ForEachStreamingCellsSources(StreamingSources, [this](const UWorldPartitionRuntimeCell* Cell, EStreamingSourceTargetState TargetState)
@@ -255,6 +262,8 @@ void UWorldPartitionStreamingPolicy::UpdateStreamingState()
 
 	if (bIsServer)
 	{
+		TRACE_CPUPROFILER_EVENT_SCOPE(UWorldPartitionStreamingPolicy::UpdateStreamingState_ServerUpdate);
+
 		// Server will activate all non data layer cells at first and then load/activate/unload data layer cells only when the data layer states change
 		if (!bIsServerStreamingEnabled && 
 			(ServerStreamingEnabledEpoch == NewServerStreamingEnabledEpoch) &&
@@ -345,93 +354,65 @@ void UWorldPartitionStreamingPolicy::UpdateStreamingState()
 		return Cell->GetIsHLOD() && !UHLODSubsystem::IsHLODEnabled();
 	};
 
-	// Process cells to activate
-	auto ProcessCellsToActivate = [this, &ShouldSkipDisabledHLODCell](TSet<const UWorldPartitionRuntimeCell*>& Cells)
-	{
-		for (TSet<const UWorldPartitionRuntimeCell*>::TIterator It = Cells.CreateIterator(); It; ++It)
-		{
-			const UWorldPartitionRuntimeCell* Cell = *It;
-			if (ShouldSkipCellForPerformance(Cell) || ShouldSkipDisabledHLODCell(Cell))
-			{
-				It.RemoveCurrent();
-			}
-			else
-			{
-				Cell->MergeStreamingSourceInfo();
-			}
-		}
-	};
+	// Activation superseeds Loading
+	FrameLoadCells = FrameLoadCells.Difference(FrameActivateCells);
 
-	// Process cells to load
-	auto ProcessCellsToLoad = [this, bIsServer, &ShouldWaitForClientVisibility, &ShouldSkipDisabledHLODCell](TSet<const UWorldPartitionRuntimeCell*>& Cells)
+	// Determine cells to activate (sorted by importance)
+	TArray<const UWorldPartitionRuntimeCell*> ToActivateCells;
 	{
-		for (TSet<const UWorldPartitionRuntimeCell*>::TIterator It = Cells.CreateIterator(); It; ++It)
+		TRACE_CPUPROFILER_EVENT_SCOPE(UWorldPartitionStreamingPolicy::UpdateStreamingState_ToActivateCells);
+		for (const UWorldPartitionRuntimeCell* Cell : FrameActivateCells)
 		{
-			const UWorldPartitionRuntimeCell* Cell = *It;
-			// Only deactivated server cells need to call ShouldWaitForClientVisibility (those part of ActivatedCells)
-			const bool bIsServerCellToDeactivate = bIsServer && ActivatedCells.Contains(Cell);
-			const bool bShoudSkipServerCell = bIsServerCellToDeactivate && ShouldWaitForClientVisibility(Cell);
-
-			if (bShoudSkipServerCell || ShouldSkipCellForPerformance(Cell) || ShouldSkipDisabledHLODCell(Cell))
-			{
-				It.RemoveCurrent();
-			}
-			else
+			if (!ActivatedCells.Contains(Cell) && !ShouldSkipCellForPerformance(Cell) && !ShouldSkipDisabledHLODCell(Cell))
 			{
 				Cell->MergeStreamingSourceInfo();
+				ToActivateCells.Add(Cell);
 			}
 		}
-	};
+		SortStreamingCellsByImportance(ToActivateCells);
+	}
 
-	// Process cells to unload
-	auto ProcessCellsToUnload = [this, bIsServer, &ShouldWaitForClientVisibility](TSet<const UWorldPartitionRuntimeCell*>& Cells)
+	// Determine cells to load (sorted by importance)
+	TArray<const UWorldPartitionRuntimeCell*> ToLoadCells;
 	{
-		// Only loop for server as ShouldWaitForClientVisibility only concerns server
-		if (bIsServer)
+		TRACE_CPUPROFILER_EVENT_SCOPE(UWorldPartitionStreamingPolicy::UpdateStreamingState_ToLoadCells);
+		for (const UWorldPartitionRuntimeCell* Cell : FrameLoadCells)
 		{
-			for (TSet<const UWorldPartitionRuntimeCell*>::TIterator It = Cells.CreateIterator(); It; ++It)
+			if (!LoadedCells.Contains(Cell))
 			{
-				const UWorldPartitionRuntimeCell* Cell = *It;
-				if (ShouldWaitForClientVisibility(Cell))
+				// Only deactivated server cells need to call ShouldWaitForClientVisibility (those part of ActivatedCells)
+				const bool bIsServerCellToDeactivate = bIsServer && ActivatedCells.Contains(Cell);
+				const bool bShoudSkipServerCell = bIsServerCellToDeactivate && ShouldWaitForClientVisibility(Cell);
+				if (!bShoudSkipServerCell && !ShouldSkipCellForPerformance(Cell) && !ShouldSkipDisabledHLODCell(Cell))
 				{
-					It.RemoveCurrent();
+					Cell->MergeStreamingSourceInfo();
+					ToLoadCells.Add(Cell);
 				}
 			}
 		}
-	};
-
-	// Activation superseeds Loading
-	FrameLoadCells = FrameLoadCells.Difference(FrameActivateCells);
-
-	// Sort cells by importance
-	auto SortStreamingCells = [this](const TArray<FWorldPartitionStreamingSource>& Sources, const TSet<const UWorldPartitionRuntimeCell*>& InCells, TArray<const UWorldPartitionRuntimeCell*>& OutCells)
-	{
-		OutCells.Empty(InCells.Num());
-		SortStreamingCellsByImportance(InCells, OutCells);
-	};
-
-	// Determine cells to activate (sorted by importance)
-	TArray<const UWorldPartitionRuntimeCell*> ToActivateCells;
-	{
-		TSet<const UWorldPartitionRuntimeCell*> ToActivateCellsUnsorted = FrameActivateCells.Difference(ActivatedCells);
-		ProcessCellsToActivate(ToActivateCellsUnsorted);
-		SortStreamingCells(StreamingSources, ToActivateCellsUnsorted, ToActivateCells);
-	}
-
-	// Determine cells to load (sorted by importance)
-	TArray<const UWorldPartitionRuntimeCell*> ToLoadCells;
-	{
-		TSet<const UWorldPartitionRuntimeCell*> ToLoadCellsUnsorted = FrameLoadCells.Difference(LoadedCells);
-		ProcessCellsToLoad(ToLoadCellsUnsorted);
-		SortStreamingCells(StreamingSources, ToLoadCellsUnsorted, ToLoadCells);
+		SortStreamingCellsByImportance(ToLoadCells);
 	}
 
 	// Determine cells to unload
 	TArray<const UWorldPartitionRuntimeCell*> ToUnloadCells;
 	{
-		TSet<const UWorldPartitionRuntimeCell*> ToUnloadCellsUnsorted = ActivatedCells.Union(LoadedCells).Difference(FrameActivateCells.Union(FrameLoadCells));
-		ProcessCellsToUnload(ToUnloadCellsUnsorted);
-		ToUnloadCells = ToUnloadCellsUnsorted.Array();
+		TRACE_CPUPROFILER_EVENT_SCOPE(UWorldPartitionStreamingPolicy::UpdateStreamingState_ToUnloadCells);
+		auto BuildCellsToUnload = [this, &ToUnloadCells, bIsServer, ShouldWaitForClientVisibility](const TSet<const UWorldPartitionRuntimeCell*>& InCells)
+		{
+			for (const UWorldPartitionRuntimeCell* Cell : InCells)
+			{
+				if (!FrameActivateCells.Contains(Cell) && !FrameLoadCells.Contains(Cell))
+				{
+					if (!bIsServer || !ShouldWaitForClientVisibility(Cell))
+					{
+						ToUnloadCells.Add(Cell);
+					}
+				}
+			}
+		};
+
+		BuildCellsToUnload(ActivatedCells);
+		BuildCellsToUnload(LoadedCells);
 	}
 
 	if(World->bMatchStarted)
@@ -466,16 +447,15 @@ void UWorldPartitionStreamingPolicy::UpdateStreamingState()
 	// Sort cells and update streaming priority 
 	{
 		TRACE_CPUPROFILER_EVENT_SCOPE(UWorldPartitionStreamingPolicy::SortCellsAndUpdateStreamingPriority);
-		TSet<const UWorldPartitionRuntimeCell*> AddToWorldCells;
+		SortedAddToWorldCells.Reset();
 		for (const UWorldPartitionRuntimeCell* ActivatedCell : ActivatedCells)
 		{
 			if (!ActivatedCell->IsAddedToWorld() && !ActivatedCell->IsAlwaysLoaded())
 			{
-				AddToWorldCells.Add(ActivatedCell);
+				SortedAddToWorldCells.Add(ActivatedCell);
 			}
 		}
-
-		SortStreamingCells(StreamingSources, AddToWorldCells, SortedAddToWorldCells);
+		SortStreamingCellsByImportance(SortedAddToWorldCells);
 
 		// Update level streaming priority so that UWorld::UpdateLevelStreaming will naturally process the levels in the correct order
 		const int32 MaxPrio = SortedAddToWorldCells.Num();
@@ -508,16 +488,14 @@ void UWorldPartitionStreamingPolicy::UpdateDebugCellsStreamingPriority(const TSe
 {
 	if (FWorldPartitionDebugHelper::IsRuntimeSpatialHashCellStreamingPriorityShown())
 	{
-		TSet<const UWorldPartitionRuntimeCell*> CellsToSort;
-		CellsToSort.Append(ActivateStreamingCells);
-		CellsToSort.Append(LoadStreamingCells);
+		TArray<const UWorldPartitionRuntimeCell*> Cells = ActivateStreamingCells.Array();
+		Cells.Append(LoadStreamingCells.Array());
 
-		TArray<const UWorldPartitionRuntimeCell*> SortedCells;
-		SortStreamingCellsByImportance(CellsToSort, SortedCells);
+		SortStreamingCellsByImportance(Cells);
 
-		const int32 CellCount = SortedCells.Num();
+		const int32 CellCount = Cells.Num();
 		int32 CellPrio = 0;
-		for (const UWorldPartitionRuntimeCell* SortedCell : SortedCells)
+		for (const UWorldPartitionRuntimeCell* SortedCell : Cells)
 		{
 			const_cast<UWorldPartitionRuntimeCell*>(SortedCell)->SetDebugStreamingPriority(float(CellPrio++) / CellCount);
 		}
@@ -860,7 +838,8 @@ void UWorldPartitionStreamingPolicy::DrawRuntimeHash3D()
  */
 
 FStreamingSourceVelocity::FStreamingSourceVelocity(const FName& InSourceName)
-: SourceName(InSourceName)
+: bIsValid(false)
+, SourceName(InSourceName)
 , LastIndex(INDEX_NONE)
 , LastUpdateTime(-1.0)
 , VelocitiesHistorySum(0.f)
@@ -870,6 +849,8 @@ FStreamingSourceVelocity::FStreamingSourceVelocity(const FName& InSourceName)
 
 float FStreamingSourceVelocity::GetAverageVelocity(const FVector& NewPosition, const float CurrentTime)
 {
+	bIsValid = true;
+
 	const double TeleportDistance = 100;
 	const float MaxDeltaSeconds = 5.f;
 	const bool bIsFirstCall = (LastIndex == INDEX_NONE);
diff --git a/Engine/Source/Runtime/Engine/Public/WorldPartition/WorldPartitionStreamingPolicy.h b/Engine/Source/Runtime/Engine/Public/WorldPartition/WorldPartitionStreamingPolicy.h
--- a/Engine/Source/Runtime/Engine/Public/WorldPartition/WorldPartitionStreamingPolicy.h
+++ b/Engine/Source/Runtime/Engine/Public/WorldPartition/WorldPartitionStreamingPolicy.h
@@ -24,10 +24,13 @@ class UWorldPartition;
 struct FStreamingSourceVelocity
 {
 	FStreamingSourceVelocity(const FName& InSourceName);
+	void Invalidate() { bIsValid = false; }
+	bool IsValid() { return bIsValid; }
 	float GetAverageVelocity(const FVector& NewPosition, const float CurrentTime);
 
 private:
 	enum { VELOCITY_HISTORY_SAMPLE_COUNT = 16 };
+	bool bIsValid;
 	FName SourceName;
 	int32 LastIndex;
 	float LastUpdateTime;
