Chaos - add a flag to disable joint mass conditioning
- disable mass conditioning for all joints in TM_PhysGetSetMulti teeter-totter test

diff --git a/Engine/Source/Runtime/Engine/Classes/PhysicsEngine/ConstraintInstance.h b/Engine/Source/Runtime/Engine/Classes/PhysicsEngine/ConstraintInstance.h
--- a/Engine/Source/Runtime/Engine/Classes/PhysicsEngine/ConstraintInstance.h
+++ b/Engine/Source/Runtime/Engine/Classes/PhysicsEngine/ConstraintInstance.h
@@ -134,6 +134,15 @@ struct ENGINE_API FConstraintProfileProperties
 	UPROPERTY(EditAnywhere, Category = Projection)
 	uint8 bEnableProjection : 1;
 
+	/** 
+	 * Whether mass conditioning is enabled for this joint. Mass conditioning applies a non-physical scale to the mass and inertia of the two
+	 * bodies that only affects this joint, so that the mass and inertia ratios are smaller. This helps stabilize joints where the bodies
+	 * are very different sizes, and especially when the parent body is heavier than the child. However, it can lead to unrealistic
+	 * behaviour, especially when collisions are involved.
+	*/
+	UPROPERTY(EditAnywhere, Category = Projection)
+	uint8 bEnableMassConditioning : 1;
+
 	/** Whether it is possible to break the joint with angular force. */
 	UPROPERTY(EditAnywhere, AdvancedDisplay, Category = Angular)
 	uint8 bAngularBreakable : 1;
@@ -959,6 +968,21 @@ public:
 	void EnableParentDominates();
 	void DisableParentDominates();
 
+	/** Whether mass conditioning is enabled. @see FConstraintProfileProperties::bEnableMassConditioning */
+	bool IsMassConditioningEnabled() const
+	{
+		return ProfileInstance.bEnableMassConditioning;
+	}
+
+	/**
+	 * Enable maxx conditioning. @see FConstraintProfileProperties::bEnableMassConditioning
+	*/
+	void EnableMassConditioning();
+
+	/**
+	 * Disable maxx conditioning. @see FConstraintProfileProperties::bEnableMassConditioning
+	*/
+	void DisableMassConditioning();
 
 	float GetLastKnownScale() const { return LastKnownScale; }
 
diff --git a/Engine/Source/Runtime/Engine/Classes/PhysicsEngine/ConstraintInstanceBlueprintLibrary.h b/Engine/Source/Runtime/Engine/Classes/PhysicsEngine/ConstraintInstanceBlueprintLibrary.h
--- a/Engine/Source/Runtime/Engine/Classes/PhysicsEngine/ConstraintInstanceBlueprintLibrary.h
+++ b/Engine/Source/Runtime/Engine/Classes/PhysicsEngine/ConstraintInstanceBlueprintLibrary.h
@@ -98,6 +98,21 @@ public:
 	UFUNCTION(BlueprintCallable, Category = "Physics|Constraints")
 	static bool GetParentDominates(UPARAM(ref) FConstraintInstanceAccessor& Accessor);
 
+	/**
+	 * @brief Enable or disable mass conditioning for the constraint.
+	*/
+	UFUNCTION(BlueprintCallable, Category = "Physics|Constraints")
+	static void SetMassConditioningEnabled(
+		UPARAM(ref) FConstraintInstanceAccessor& Accessor,
+		bool bEnableMassConditioning);
+
+	/**
+	 * @brief Gets whether mass conditioning is enabled for the constraint.
+	*/
+	UFUNCTION(BlueprintCallable, Category = "Physics|Constraints")
+	static bool GetMassConditioningEnabled(UPARAM(ref) FConstraintInstanceAccessor& Accessor);
+
+
 	//---------------------------------------------------------------------------------------------------
 	//
 	// LINEAR LIMITS
diff --git a/Engine/Source/Runtime/Engine/Private/PhysicsEngine/ConstraintInstance.cpp b/Engine/Source/Runtime/Engine/Private/PhysicsEngine/ConstraintInstance.cpp
--- a/Engine/Source/Runtime/Engine/Private/PhysicsEngine/ConstraintInstance.cpp
+++ b/Engine/Source/Runtime/Engine/Private/PhysicsEngine/ConstraintInstance.cpp
@@ -197,6 +197,7 @@ FConstraintProfileProperties::FConstraintProfileProperties()
 	, bParentDominates(false)
 	, bEnableShockPropagation(false)
 	, bEnableProjection(true)
+	, bEnableMassConditioning(true)
 	, bAngularBreakable(false)
 	, bAngularPlasticity(false)
 	, bLinearBreakable(false)
@@ -556,6 +557,7 @@ void FConstraintProfileProperties::UpdateConstraintFlags_AssumesLocked(const FPh
 	FPhysicsInterface::SetProjectionEnabled_AssumesLocked(InConstraintRef, bEnableProjection, ProjectionLinearAlpha, ProjectionAngularAlpha, ProjectionLinearTolerance, ProjectionAngularTolerance);
 	FPhysicsInterface::SetShockPropagationEnabled_AssumesLocked(InConstraintRef, bEnableShockPropagation, ShockPropagationAlpha);
 	FPhysicsInterface::SetParentDominates_AssumesLocked(InConstraintRef, bParentDominates);
+	FPhysicsInterface::SetMassConditioningEnabled_AssumesLocked(InConstraintRef, bEnableMassConditioning);
 }
 
 
@@ -1345,6 +1347,24 @@ void FConstraintInstance::DisableParentDominates()
 	});
 }
 
+void FConstraintInstance::EnableMassConditioning()
+{
+	ProfileInstance.bEnableMassConditioning = true;
+	FPhysicsCommand::ExecuteWrite(ConstraintHandle, [&](const FPhysicsConstraintHandle& Constraint)
+	{
+		FPhysicsInterface::SetMassConditioningEnabled_AssumesLocked(Constraint, true);
+	});
+}
+
+void FConstraintInstance::DisableMassConditioning()
+{
+	ProfileInstance.bEnableMassConditioning = false;
+	FPhysicsCommand::ExecuteWrite(ConstraintHandle, [&](const FPhysicsConstraintHandle& Constraint)
+	{
+		FPhysicsInterface::SetMassConditioningEnabled_AssumesLocked(Constraint, false);
+	});
+}
+
 FConstraintInstance* FConstraintInstanceAccessor::Get() const
 {
 	if (Owner.IsValid())
diff --git a/Engine/Source/Runtime/Engine/Private/PhysicsEngine/ConstraintInstanceBlueprintLibrary.cpp b/Engine/Source/Runtime/Engine/Private/PhysicsEngine/ConstraintInstanceBlueprintLibrary.cpp
--- a/Engine/Source/Runtime/Engine/Private/PhysicsEngine/ConstraintInstanceBlueprintLibrary.cpp
+++ b/Engine/Source/Runtime/Engine/Private/PhysicsEngine/ConstraintInstanceBlueprintLibrary.cpp
@@ -118,6 +118,36 @@ bool UConstraintInstanceBlueprintLibrary::GetParentDominates(FConstraintInstance
 	return false;
 }
 
+
+void UConstraintInstanceBlueprintLibrary::SetMassConditioningEnabled(
+	FConstraintInstanceAccessor& Accessor,
+	bool bEnableMassConditioning
+)
+{
+	if (FConstraintInstance* ConstraintInstance = Accessor.Get())
+	{
+		Accessor.Modify();
+
+		if (bEnableMassConditioning)
+		{
+			ConstraintInstance->EnableMassConditioning();
+		}
+		else
+		{
+			ConstraintInstance->DisableMassConditioning();
+		}
+	}
+}
+
+bool UConstraintInstanceBlueprintLibrary::GetMassConditioningEnabled(FConstraintInstanceAccessor& Accessor)
+{
+	if (FConstraintInstance* ConstraintInstance = Accessor.Get())
+	{
+		return ConstraintInstance->IsMassConditioningEnabled();
+	}
+	return false;
+}
+
 //---------------------------------------------------------------------------------------------------
 //
 // LINEAR LIMITS
diff --git a/Engine/Source/Runtime/Engine/Private/PhysicsEngine/ImmediatePhysics/ImmediatePhysicsChaos/ImmediatePhysicsJointHandle_Chaos.cpp b/Engine/Source/Runtime/Engine/Private/PhysicsEngine/ImmediatePhysics/ImmediatePhysicsChaos/ImmediatePhysicsJointHandle_Chaos.cpp
--- a/Engine/Source/Runtime/Engine/Private/PhysicsEngine/ImmediatePhysics/ImmediatePhysicsChaos/ImmediatePhysicsJointHandle_Chaos.cpp
+++ b/Engine/Source/Runtime/Engine/Private/PhysicsEngine/ImmediatePhysics/ImmediatePhysicsChaos/ImmediatePhysicsJointHandle_Chaos.cpp
@@ -47,6 +47,7 @@ namespace ImmediatePhysics_Chaos
 
 		ConstraintSettings.bProjectionEnabled = Profile.bEnableProjection;
 		ConstraintSettings.bShockPropagationEnabled = Profile.bEnableShockPropagation;
+		ConstraintSettings.bMassConditioningEnabled = Profile.bEnableMassConditioning;
 
 		ConstraintSettings.LinearProjection = Profile.bEnableProjection ? Profile.ProjectionLinearAlpha : 0.0f;
 		ConstraintSettings.AngularProjection = Profile.bEnableProjection ? Profile.ProjectionAngularAlpha : 0.0f;
diff --git a/Engine/Source/Runtime/Experimental/Chaos/Private/Chaos/Joint/PBDJointCachedSolverGaussSeidel.cpp b/Engine/Source/Runtime/Experimental/Chaos/Private/Chaos/Joint/PBDJointCachedSolverGaussSeidel.cpp
--- a/Engine/Source/Runtime/Experimental/Chaos/Private/Chaos/Joint/PBDJointCachedSolverGaussSeidel.cpp
+++ b/Engine/Source/Runtime/Experimental/Chaos/Private/Chaos/Joint/PBDJointCachedSolverGaussSeidel.cpp
@@ -202,10 +202,22 @@ void FPBDJointCachedSolver::Init(
 
 	InitDerivedState();
 
-	FReal ConditionedInvMs[2];
-	FVec3 ConditionedInvILs[2];
-	FPBDJointUtilities::ConditionInverseMassAndInertia(Body0().InvM(), Body1().InvM(), Body0().InvILocal(), Body1().InvILocal(), SolverSettings.MinParentMassRatio, SolverSettings.MaxInertiaRatio, ConditionedInvMs[0], ConditionedInvMs[1], ConditionedInvILs[0], ConditionedInvILs[1]);
-
+	// Set the mass and inertia.
+	// If enabled, adjust the mass so that we limit the maximum mass and inertia ratios
+	FReal ConditionedInvMs[2] = 
+	{
+		Body0().InvM(),
+		Body1().InvM()
+	};
+	FVec3 ConditionedInvILs[2] = 
+	{
+		Body0().InvILocal(), 
+		Body1().InvILocal()
+	};
+	if (JointSettings.bMassConditioningEnabled)
+	{
+		FPBDJointUtilities::ConditionInverseMassAndInertia(Body0().InvM(), Body1().InvM(), Body0().InvILocal(), Body1().InvILocal(), SolverSettings.MinParentMassRatio, SolverSettings.MaxInertiaRatio, ConditionedInvMs[0], ConditionedInvMs[1], ConditionedInvILs[0], ConditionedInvILs[1]);
+	}
 	UpdateMass0(ConditionedInvMs[0], ConditionedInvILs[0]);
 	UpdateMass1(ConditionedInvMs[1], ConditionedInvILs[1]);
 
diff --git a/Engine/Source/Runtime/Experimental/Chaos/Private/Chaos/Joint/PBDJointSolverGaussSeidel.cpp b/Engine/Source/Runtime/Experimental/Chaos/Private/Chaos/Joint/PBDJointSolverGaussSeidel.cpp
--- a/Engine/Source/Runtime/Experimental/Chaos/Private/Chaos/Joint/PBDJointSolverGaussSeidel.cpp
+++ b/Engine/Source/Runtime/Experimental/Chaos/Private/Chaos/Joint/PBDJointSolverGaussSeidel.cpp
@@ -102,9 +102,18 @@ namespace Chaos
 		SolverBodies[0].SetInvMScale(JointSettings.ParentInvMassScale);
 		SolverBodies[1].SetInvMScale(FReal(1));
 
+		// Set the mass and inertia.
+		// If enabled, adjust the mass so that we limit the maximum mass and inertia ratios
 		InvMScales[0] = FReal(1);
 		InvMScales[1] = FReal(1);
-		FPBDJointUtilities::ConditionInverseMassAndInertia(Body0().InvM(), Body1().InvM(), Body0().InvILocal(), Body1().InvILocal(), SolverSettings.MinParentMassRatio, SolverSettings.MaxInertiaRatio, ConditionedInvMs[0], ConditionedInvMs[1], ConditionedInvILs[0], ConditionedInvILs[1]);
+		ConditionedInvMs[0] = Body0().InvM();
+		ConditionedInvMs[1] = Body1().InvM();
+		ConditionedInvILs[0] = Body0().InvILocal();
+		ConditionedInvILs[1] = Body1().InvILocal();
+		if (JointSettings.bMassConditioningEnabled)
+		{
+			FPBDJointUtilities::ConditionInverseMassAndInertia(Body0().InvM(), Body1().InvM(), Body0().InvILocal(), Body1().InvILocal(), SolverSettings.MinParentMassRatio, SolverSettings.MaxInertiaRatio, ConditionedInvMs[0], ConditionedInvMs[1], ConditionedInvILs[0], ConditionedInvILs[1]);
+		}
 		UpdateMass0();
 		UpdateMass1();
 
diff --git a/Engine/Source/Runtime/Experimental/Chaos/Private/Chaos/PBDJointConstraints.cpp b/Engine/Source/Runtime/Experimental/Chaos/Private/Chaos/PBDJointConstraints.cpp
--- a/Engine/Source/Runtime/Experimental/Chaos/Private/Chaos/PBDJointConstraints.cpp
+++ b/Engine/Source/Runtime/Experimental/Chaos/Private/Chaos/PBDJointConstraints.cpp
@@ -192,6 +192,7 @@ namespace Chaos
 		, bCollisionEnabled(true)
 		, bProjectionEnabled(false)
 		, bShockPropagationEnabled(false)
+		, bMassConditioningEnabled(true)
 		, LinearMotionTypes({ EJointMotionType::Locked, EJointMotionType::Locked, EJointMotionType::Locked })
 		, LinearLimit(FLT_MAX)
 		, AngularMotionTypes({ EJointMotionType::Free, EJointMotionType::Free, EJointMotionType::Free })
diff --git a/Engine/Source/Runtime/Experimental/Chaos/Public/Chaos/JointProperties.inl b/Engine/Source/Runtime/Experimental/Chaos/Public/Chaos/JointProperties.inl
--- a/Engine/Source/Runtime/Experimental/Chaos/Public/Chaos/JointProperties.inl
+++ b/Engine/Source/Runtime/Experimental/Chaos/Public/Chaos/JointProperties.inl
@@ -5,6 +5,7 @@ CHAOS_INNER_JOINT_PROPERTY(JointSettings, JointTransforms, ConnectorTransforms,
 CHAOS_INNER_JOINT_PROPERTY(JointSettings, CollisionEnabled, bCollisionEnabled, bool)
 CHAOS_INNER_JOINT_PROPERTY(JointSettings, ProjectionEnabled, bProjectionEnabled, bool)
 CHAOS_INNER_JOINT_PROPERTY(JointSettings, ShockPropagationEnabled, bShockPropagationEnabled, bool)
+CHAOS_INNER_JOINT_PROPERTY(JointSettings, MassConditioningEnabled, bMassConditioningEnabled, bool)
 CHAOS_INNER_JOINT_PROPERTY(JointSettings, ProjectionLinearAlpha, LinearProjection, FReal)
 CHAOS_INNER_JOINT_PROPERTY(JointSettings, ProjectionAngularAlpha, AngularProjection, FReal)
 CHAOS_INNER_JOINT_PROPERTY(JointSettings, ProjectionLinearTolerance, TeleportDistance, FReal)
diff --git a/Engine/Source/Runtime/Experimental/Chaos/Public/Chaos/PBDJointConstraintTypes.h b/Engine/Source/Runtime/Experimental/Chaos/Public/Chaos/PBDJointConstraintTypes.h
--- a/Engine/Source/Runtime/Experimental/Chaos/Public/Chaos/PBDJointConstraintTypes.h
+++ b/Engine/Source/Runtime/Experimental/Chaos/Public/Chaos/PBDJointConstraintTypes.h
@@ -139,6 +139,7 @@ namespace Chaos
 		bool bCollisionEnabled;
 		bool bProjectionEnabled;		// @chaos(todo): remove - implied by alpha and teleport settings
 		bool bShockPropagationEnabled;	// @chaos(todo): remove - implied by alpha
+		bool bMassConditioningEnabled;
 
 		TVector<EJointMotionType, 3> LinearMotionTypes;
 		FReal LinearLimit;
diff --git a/Engine/Source/Runtime/PhysicsCore/Private/ChaosEngineInterface.cpp b/Engine/Source/Runtime/PhysicsCore/Private/ChaosEngineInterface.cpp
--- a/Engine/Source/Runtime/PhysicsCore/Private/ChaosEngineInterface.cpp
+++ b/Engine/Source/Runtime/PhysicsCore/Private/ChaosEngineInterface.cpp
@@ -1343,6 +1343,17 @@ void FChaosEngineInterface::SetParentDominates_AssumesLocked(const FPhysicsConst
 	}
 }
 
+void FChaosEngineInterface::SetMassConditioningEnabled_AssumesLocked(const FPhysicsConstraintHandle& InConstraintRef, bool bInMassConditioningEnabled)
+{
+	if (InConstraintRef.IsValid() && InConstraintRef.Constraint->IsType(Chaos::EConstraintType::JointConstraintType))
+	{
+		if (Chaos::FJointConstraint* Constraint = static_cast<Chaos::FJointConstraint*>(InConstraintRef.Constraint))
+		{
+			Constraint->SetMassConditioningEnabled(bInMassConditioningEnabled);
+		}
+	}
+}
+
 void FChaosEngineInterface::SetBreakForces_AssumesLocked(const FPhysicsConstraintHandle& InConstraintRef,float InLinearBreakForce,float InAngularBreakTorque)
 {
 	if (InConstraintRef.IsValid() && InConstraintRef.Constraint->IsType(Chaos::EConstraintType::JointConstraintType))
diff --git a/Engine/Source/Runtime/PhysicsCore/Public/Chaos/ChaosEngineInterface.h b/Engine/Source/Runtime/PhysicsCore/Public/Chaos/ChaosEngineInterface.h
--- a/Engine/Source/Runtime/PhysicsCore/Public/Chaos/ChaosEngineInterface.h
+++ b/Engine/Source/Runtime/PhysicsCore/Public/Chaos/ChaosEngineInterface.h
@@ -498,7 +498,8 @@ public:
 	static void SetCollisionEnabled(const FPhysicsConstraintHandle& InConstraintRef,bool bInCollisionEnabled);
 	static void SetProjectionEnabled_AssumesLocked(const FPhysicsConstraintHandle& InConstraintRef,bool bInProjectionEnabled,float InLinearAlpha = 1.0f,float InAngularAlpha = 0.0f, float InLinearTolerance = 0.0f, float InAngularToleranceDeg = 0.0f);
 	static void SetShockPropagationEnabled_AssumesLocked(const FPhysicsConstraintHandle& InConstraintRef, bool bInShockPropagationEnabled, float InShockPropagationAlpha);
-	static void SetParentDominates_AssumesLocked(const FPhysicsConstraintHandle& InConstraintRef,bool bInParentDominates);
+	static void SetParentDominates_AssumesLocked(const FPhysicsConstraintHandle& InConstraintRef, bool bInParentDominates);
+	static void SetMassConditioningEnabled_AssumesLocked(const FPhysicsConstraintHandle& InConstraintRef, bool bInMassConditioningEnabled);
 	static void SetBreakForces_AssumesLocked(const FPhysicsConstraintHandle& InConstraintRef,float InLinearBreakForce,float InAngularBreakForce);
 	static void SetPlasticityLimits_AssumesLocked(const FPhysicsConstraintHandle& InConstraintRef, float InLinearPlasticityLimit, float InAngularPlasticityLimit, EConstraintPlasticityType InLinearPlasticityType);
 	static void SetContactTransferScale_AssumesLocked(const FPhysicsConstraintHandle& InConstraintRef, float InContactTransferScale);
