Expose "GetPrimaryPlatformUser" to blueprints so that you can use it as a back up in widgets. Fix a small comment typo in the description.

#rnx

diff --git a/Engine/Source/Runtime/ApplicationCore/Public/GenericPlatform/GenericPlatformInputDeviceMapper.h b/Engine/Source/Runtime/ApplicationCore/Public/GenericPlatform/GenericPlatformInputDeviceMapper.h
--- a/Engine/Source/Runtime/ApplicationCore/Public/GenericPlatform/GenericPlatformInputDeviceMapper.h
+++ b/Engine/Source/Runtime/ApplicationCore/Public/GenericPlatform/GenericPlatformInputDeviceMapper.h
@@ -81,7 +81,7 @@ public:
 
 	/**
 	 * Returns the 'Primary' Platform user for this platform.
-	 * This typcially has an internal ID of '0' and is used as the default platform user to
+	 * This typically has an internal ID of '0' and is used as the default platform user to
 	 * map devices such as the keyboard and mouse that don't get assigned unique ID's from their
 	 * owning platform code.
 	 */
diff --git a/Engine/Source/Runtime/Engine/Classes/GameFramework/PlatformInputDeviceMapperLibrary.h b/Engine/Source/Runtime/Engine/Classes/GameFramework/PlatformInputDeviceMapperLibrary.h
--- a/Engine/Source/Runtime/Engine/Classes/GameFramework/PlatformInputDeviceMapperLibrary.h
+++ b/Engine/Source/Runtime/Engine/Classes/GameFramework/PlatformInputDeviceMapperLibrary.h
@@ -66,6 +66,15 @@ public:
 	UFUNCTION(BlueprintPure, Category = "PlatformInputDevice")
 	static FPlatformUserId GetUserForUnpairedInputDevices();
 
+	/**
+	 * Returns the 'Primary' Platform user for this platform.
+	 * This typically has an internal ID of '0' and is used as the default platform user to
+	 * map devices such as the keyboard and mouse that don't get assigned unique ID's from their
+	 * owning platform code.
+	 */
+	UFUNCTION(BlueprintPure, Category = "PlatformInputDevice")
+	static FPlatformUserId GetPrimaryPlatformUser();
+
 	/** Returns true if the given Platform User Id is the user for unpaired input devices on this platform. */
 	UFUNCTION(BlueprintPure, Category = "PlatformInputDevice")
 	static bool IsUnpairedUserId(const FPlatformUserId PlatformId);
diff --git a/Engine/Source/Runtime/Engine/Private/PlatformInputDeviceMapperLibrary.cpp b/Engine/Source/Runtime/Engine/Private/PlatformInputDeviceMapperLibrary.cpp
--- a/Engine/Source/Runtime/Engine/Private/PlatformInputDeviceMapperLibrary.cpp
+++ b/Engine/Source/Runtime/Engine/Private/PlatformInputDeviceMapperLibrary.cpp
@@ -30,6 +30,11 @@ FPlatformUserId UPlatformInputDeviceMapperLibrary::GetUserForUnpairedInputDevice
 	return IPlatformInputDeviceMapper::Get().GetUserForUnpairedInputDevices();
 }
 
+FPlatformUserId UPlatformInputDeviceMapperLibrary::GetPrimaryPlatformUser()
+{
+	return IPlatformInputDeviceMapper::Get().GetPrimaryPlatformUser();
+}
+
 bool UPlatformInputDeviceMapperLibrary::IsUnpairedUserId(const FPlatformUserId PlatformId)
 {
 	return IPlatformInputDeviceMapper::Get().IsUnpairedUserId(PlatformId);
