Decoupled Zen, DerivedDataCache and Virtualization dependency from Studio Analytics. Systems can now register callbacks to append or modify Event payloads to analytics events by name via the FAnalytics interface.

diff --git a/Engine/Source/Developer/DerivedDataCache/Private/DerivedDataCache.cpp b/Engine/Source/Developer/DerivedDataCache/Private/DerivedDataCache.cpp
--- a/Engine/Source/Developer/DerivedDataCache/Private/DerivedDataCache.cpp
+++ b/Engine/Source/Developer/DerivedDataCache/Private/DerivedDataCache.cpp
@@ -18,6 +18,7 @@
 #include "DerivedDataPrivate.h"
 #include "DerivedDataRequest.h"
 #include "DerivedDataRequestOwner.h"
+#include "DerivedDataCacheUsageStats.h"
 #include "Experimental/Async/LazyEvent.h"
 #include "Features/IModularFeatures.h"
 #include "HAL/ThreadSafeCounter.h"
@@ -31,6 +32,8 @@
 #include "Stats/Stats.h"
 #include "Stats/StatsMisc.h"
 #include "ZenServerInterface.h"
+#include "AnalyticsEventAttribute.h"
+#include "Analytics.h"
 #include <atomic>
 
 DEFINE_STAT(STAT_DDC_NumGets);
@@ -661,12 +664,19 @@ public:
 
 		GVerifyDDC = FParse::Param(FCommandLine::Get(), TEXT("VerifyDDC"));
 
+		FAnalytics::Get().GetEventCallback(TEXT("Core.Loading"))->AddRaw(this, &FDerivedDataCache::OnAnalyticsEvent);
+		
 		UE_CLOG(GVerifyDDC, LogDerivedDataCache, Display, TEXT("Items retrieved from the DDC will be verified (-VerifyDDC)"));
 	}
 
 	/** Destructor, flushes all sync tasks **/
 	~FDerivedDataCache()
 	{
+		if (FAnalytics::IsAvailable())
+		{
+			FAnalytics::Get().GetEventCallback(TEXT("Core.Loading"))->RemoveAll(this);
+		}
+		
 		WaitForQuiescence(true);
 		FScopeLock ScopeLock(&SynchronizationObject);
 		for (TMap<uint32,FAsyncTask<FBuildAsyncWorker>*>::TIterator It(PendingTasks); It; ++It)
@@ -1034,6 +1044,77 @@ public:
 		GatherDerivedDataCacheSummaryStats(DDCSummaryStats);
 	}
 
+	void OnAnalyticsEvent(TArray<FAnalyticsEventAttribute>& Attributes)
+	{
+#if ENABLE_COOK_STATS
+
+		// Gather the latest resource stats
+		TArray<FDerivedDataCacheResourceStat> ResourceStats;
+
+		GatherDerivedDataCacheResourceStats(ResourceStats);
+
+		FDerivedDataCacheResourceStat ResourceStatsTotal(TEXT("Total"));
+
+		// Accumulate Totals
+		for (const FDerivedDataCacheResourceStat& Stat : ResourceStats)
+		{
+			ResourceStatsTotal += Stat;
+		}
+
+		ResourceStats.Emplace(ResourceStatsTotal);
+
+		// Append to the attributes
+		for (const FDerivedDataCacheResourceStat& Stat : ResourceStats)
+		{
+			FString BaseName = TEXT("DDC.Resource.") + Stat.AssetType;
+
+			BaseName = BaseName.Replace(TEXT("("), TEXT("")).Replace(TEXT(")"), TEXT(""));
+
+			{
+				FString AttrName = BaseName + TEXT(".BuildCount");
+				Attributes.Emplace(MoveTemp(AttrName), Stat.BuildCount);
+			}
+
+			{
+				FString AttrName = BaseName + TEXT(".BuildTimeSec");
+				Attributes.Emplace(MoveTemp(AttrName), Stat.BuildTimeSec);
+			}
+
+			{
+				FString AttrName = BaseName + TEXT(".BuildSizeMB");
+				Attributes.Emplace(MoveTemp(AttrName), Stat.BuildSizeMB);
+			}
+
+			{
+				FString AttrName = BaseName + TEXT(".LoadCount");
+				Attributes.Emplace(MoveTemp(AttrName), Stat.LoadCount);
+			}
+
+			{
+				FString AttrName = BaseName + TEXT(".LoadTimeSecLoadTimeSec");
+				Attributes.Emplace(MoveTemp(AttrName), Stat.LoadTimeSec);
+			}
+
+			{
+				FString AttrName = BaseName + TEXT(".LoadSizeMB");
+				Attributes.Emplace(MoveTemp(AttrName), Stat.LoadSizeMB);
+			}
+		}
+
+		// Gather the summary stats
+		FDerivedDataCacheSummaryStats SummaryStats;
+
+		GatherDerivedDataCacheSummaryStats(SummaryStats);
+
+		// Append to the attributes
+		for (const FDerivedDataCacheSummaryStat& Stat : SummaryStats.Stats)
+		{
+			FString FormattedAttrName = "DDC.Summary." + Stat.Key;
+			Attributes.Emplace(FormattedAttrName, Stat.Value);
+		}
+#endif
+	}
+
 	/** Get event delegate for data cache notifications */
 	virtual FOnDDCNotification& GetDDCNotificationEvent()
 	{
diff --git a/Engine/Source/Developer/Virtualization/Private/VirtualizationManager.cpp b/Engine/Source/Developer/Virtualization/Private/VirtualizationManager.cpp
--- a/Engine/Source/Developer/Virtualization/Private/VirtualizationManager.cpp
+++ b/Engine/Source/Developer/Virtualization/Private/VirtualizationManager.cpp
@@ -18,6 +18,7 @@
 #include "PackageVirtualizationProcess.h"
 #include "ProfilingDebugging/CookStats.h"
 #include "VirtualizationFilterSettings.h"
+#include "Analytics.h"
 
 #define LOCTEXT_NAMESPACE "Virtualization"
 
@@ -445,6 +446,11 @@ FVirtualizationManager::FVirtualizationManager()
 
 FVirtualizationManager::~FVirtualizationManager()
 {
+	if (FAnalytics::IsAvailable())
+	{
+		FAnalytics::Get().GetEventCallback(TEXT("Core.Loading"))->RemoveAll(this);
+	}
+
 	for (const TPair<IConsoleVariable*, FDelegateHandle>& KV : DebugValues.ConsoleDelegateHandles)
 	{
 		IConsoleVariable* ConsoleVariable = KV.Key;
@@ -497,6 +503,8 @@ bool FVirtualizationManager::Initialize(const FInitParams& InitParams)
 
 	MountBackends(InitParams.ConfigFile);
 
+	FAnalytics::Get().GetEventCallback(TEXT("Core.Loading"))->AddRaw(this, &FVirtualizationManager::OnAnalyticsEvent);
+
 	UE_LOG(LogVirtualization, Display, TEXT("Virtualization manager initialization completed"));
 
 	return true;
@@ -2006,6 +2014,71 @@ void FVirtualizationManager::BroadcastEvent(TConstArrayView<FPullRequest> Reques
 	}
 }
 
+void FVirtualizationManager::OnAnalyticsEvent(TArray<FAnalyticsEventAttribute>& Attributes)
+{
+	using namespace UE::Virtualization;
+
+	// Grab the Virtualization stats
+	if (IVirtualizationSystem::IsInitialized())
+	{
+		IVirtualizationSystem& System = IVirtualizationSystem::Get();
+
+		FPayloadActivityInfo PayloadActivityInfo = System.GetAccumualtedPayloadActivityInfo();
+
+		const FString BaseName = TEXT("Virtualization");
+
+		{
+			FString AttrName = BaseName + TEXT(".Enabled");
+			Attributes.Emplace(MoveTemp(AttrName), System.IsEnabled());
+		}
+
+		{
+			FString AttrName = BaseName + TEXT(".Cache.TimeSpent");
+			Attributes.Emplace(MoveTemp(AttrName), (double)PayloadActivityInfo.Cache.CyclesSpent * FPlatformTime::GetSecondsPerCycle());
+		}
+
+		{
+			FString AttrName = BaseName + TEXT(".Cache.PayloadCount");
+			Attributes.Emplace(MoveTemp(AttrName), (double)PayloadActivityInfo.Cache.PayloadCount);
+		}
+
+		{
+			FString AttrName = BaseName + TEXT(".Cache.TotalBytes");
+			Attributes.Emplace(MoveTemp(AttrName), (double)PayloadActivityInfo.Cache.TotalBytes);
+		}
+
+		{
+			FString AttrName = BaseName + TEXT(".Push.TimeSpent");
+			Attributes.Emplace(MoveTemp(AttrName), (double)PayloadActivityInfo.Push.CyclesSpent * FPlatformTime::GetSecondsPerCycle());
+		}
+
+		{
+			FString AttrName = BaseName + TEXT(".Push.PayloadCount");
+			Attributes.Emplace(MoveTemp(AttrName), (double)PayloadActivityInfo.Push.PayloadCount);
+		}
+
+		{
+			FString AttrName = BaseName + TEXT(".Push.TotalBytes");
+			Attributes.Emplace(MoveTemp(AttrName), (double)PayloadActivityInfo.Push.TotalBytes);
+		}
+
+		{
+			FString AttrName = BaseName + TEXT(".Pull.TimeSpent");
+			Attributes.Emplace(MoveTemp(AttrName), (double)PayloadActivityInfo.Pull.CyclesSpent * FPlatformTime::GetSecondsPerCycle());
+		}
+
+		{
+			FString AttrName = BaseName + TEXT(".Pull.PayloadCount");
+			Attributes.Emplace(MoveTemp(AttrName), (double)PayloadActivityInfo.Pull.PayloadCount);
+		}
+
+		{
+			FString AttrName = BaseName + TEXT(".Pull.TotalBytes");
+			Attributes.Emplace(MoveTemp(AttrName), (double)PayloadActivityInfo.Pull.TotalBytes);
+		}
+	}
+}
+
 } // namespace UE::Virtualization
 
 #undef UE_INLINE_ALLOCATION_COUNT
diff --git a/Engine/Source/Developer/Virtualization/Private/VirtualizationManager.h b/Engine/Source/Developer/Virtualization/Private/VirtualizationManager.h
--- a/Engine/Source/Developer/Virtualization/Private/VirtualizationManager.h
+++ b/Engine/Source/Developer/Virtualization/Private/VirtualizationManager.h
@@ -6,6 +6,8 @@
 #include "HAL/CriticalSection.h"
 #include "Logging/LogMacros.h"
 #include "Templates/UniquePtr.h"
+#include "AnalyticsEventAttribute.h"
+
 
 #include "Virtualization/VirtualizationSystem.h"
 
@@ -154,6 +156,7 @@ private:
 	void OnUpdateDebugMissBackendsFromConsole(const TArray<FString>& Args, FOutputDevice& OutputDevice);
 	void OnUpdateDebugMissChanceFromConsole(const TArray<FString>& Args, FOutputDevice& OutputDevice);
 	void OnUpdateDebugMissCountFromConsole(const TArray<FString>& Args, FOutputDevice& OutputDevice);
+	void OnAnalyticsEvent(TArray<FAnalyticsEventAttribute>& Attributes);
 
 	void UpdateBackendDebugState();
 
@@ -175,7 +178,7 @@ private:
 
 	void PullDataFromAllBackends(TArrayView<FPullRequest> Requests);
 	void PullDataFromBackend(IVirtualizationBackend& Backend, TArrayView<FPullRequest> Requests);
-
+	
 	bool ShouldVirtualizeAsset(const UObject* Owner) const;
 
 	/** 
diff --git a/Engine/Source/Developer/Virtualization/Virtualization.Build.cs b/Engine/Source/Developer/Virtualization/Virtualization.Build.cs
--- a/Engine/Source/Developer/Virtualization/Virtualization.Build.cs
+++ b/Engine/Source/Developer/Virtualization/Virtualization.Build.cs
@@ -14,6 +14,7 @@ public class Virtualization : ModuleRules
 				"CoreUObject",
 				"ApplicationCore",
 				"DerivedDataCache",
+				"Analytics",
 				"MessageLog",
 				"Projects",
 				"SourceControl"
diff --git a/Engine/Source/Developer/Zen/Internal/ZenServerInterface.cpp b/Engine/Source/Developer/Zen/Internal/ZenServerInterface.cpp
--- a/Engine/Source/Developer/Zen/Internal/ZenServerInterface.cpp
+++ b/Engine/Source/Developer/Zen/Internal/ZenServerInterface.cpp
@@ -26,6 +26,7 @@
 #include "Serialization/CompactBinaryValidation.h"
 #include "String/LexFromString.h"
 #include "SocketSubsystem.h"
+#include "Analytics.h"
 
 #if PLATFORM_WINDOWS
 #	include "Windows/AllowWindowsPlatformTypes.h"
@@ -710,6 +711,10 @@ FZenServiceInstance::FZenServiceInstance(FServiceSettings&& InSettings)
 
 FZenServiceInstance::~FZenServiceInstance()
 {
+	if (FAnalytics::IsAvailable())
+	{
+		FAnalytics::Get().GetEventCallback(TEXT("Core.Loading"))->RemoveAll(this);
+	}
 }
 
 bool 
@@ -767,6 +772,8 @@ FZenServiceInstance::Initialize()
 		bIsRunningLocally = IsLocalHost(HostName);
 	}
 	URL = WriteToString<64>(TEXT("http://"), HostName, TEXT(":"), Port, TEXT("/"));
+
+	FAnalytics::Get().GetEventCallback(TEXT("Core.Loading"))->AddRaw(this,&FZenServiceInstance::OnAnalyticsEvent);
 }
 
 void
@@ -1219,6 +1226,91 @@ FZenServiceInstance::GetStats(FZenStats& Stats)
 	return Stats.IsValid;
 }
 
+void FZenServiceInstance::OnAnalyticsEvent(TArray<FAnalyticsEventAttribute>& Attributes )
+{
+	// Grab the Zen summary stats
+	FZenStats ZenStats;
+
+	GetStats(ZenStats);
+
+	const FString BaseName = TEXT("Zen");
+
+	{
+		FString AttrName = BaseName + TEXT(".Enabled");
+		Attributes.Emplace(MoveTemp(AttrName), ZenStats.IsValid);
+	}
+
+	{
+		FString AttrName = BaseName + TEXT(".Cache.HitRatio");
+		Attributes.Emplace(MoveTemp(AttrName), ZenStats.CacheStats.HitRatio);
+	}
+
+	{
+		FString AttrName = BaseName + TEXT(".Cache.Hits");
+		Attributes.Emplace(MoveTemp(AttrName), ZenStats.CacheStats.Hits);
+	}
+
+	{
+		FString AttrName = BaseName + TEXT(".Cache.Misses");
+		Attributes.Emplace(MoveTemp(AttrName), ZenStats.CacheStats.Misses);
+	}
+
+	{
+		FString AttrName = BaseName + TEXT(".Cache.Size.Disk");
+		Attributes.Emplace(MoveTemp(AttrName), ZenStats.CacheStats.Size.Disk);
+	}
+
+	{
+		FString AttrName = BaseName + TEXT(".Cache.Size.Memory");
+		Attributes.Emplace(MoveTemp(AttrName), ZenStats.CacheStats.Size.Memory);
+	}
+
+	{
+		FString AttrName = BaseName + TEXT(".Cache.UpstreamHits");
+		Attributes.Emplace(MoveTemp(AttrName), ZenStats.CacheStats.UpstreamHits);
+	}
+
+	{
+		FString AttrName = BaseName + TEXT(".Cache.UpstreamRatio");
+		Attributes.Emplace(MoveTemp(AttrName), ZenStats.CacheStats.UpstreamRatio);
+	}
+
+	{
+		FString AttrName = BaseName + TEXT(".Cache.TotalUploadedMB");
+		Attributes.Emplace(MoveTemp(AttrName), ZenStats.UpstreamStats.TotalUploadedMB);
+	}
+
+	{
+		FString AttrName = BaseName + TEXT(".Upstream.TotalDownloadedMB");
+		Attributes.Emplace(MoveTemp(AttrName), ZenStats.UpstreamStats.TotalDownloadedMB);
+	}
+
+	{
+		FString AttrName = BaseName + TEXT(".Upstream.TotalUploadedMB");
+		Attributes.Emplace(MoveTemp(AttrName), ZenStats.UpstreamStats.TotalUploadedMB);
+	}
+
+	{
+		FString AttrName = BaseName + TEXT(".Cas.Size.Large");
+		Attributes.Emplace(MoveTemp(AttrName), ZenStats.CASStats.Size.Large);
+	}
+
+	{
+		FString AttrName = BaseName + TEXT(".Cas.Size.Small");
+		Attributes.Emplace(MoveTemp(AttrName), ZenStats.CASStats.Size.Small);
+	}
+
+	{
+		FString AttrName = BaseName + TEXT(".Cas.Size.Tiny");
+		Attributes.Emplace(MoveTemp(AttrName), ZenStats.CASStats.Size.Tiny);
+	}
+
+	{
+		FString AttrName = BaseName + TEXT(".Cas.Size.Total");
+		Attributes.Emplace(MoveTemp(AttrName), ZenStats.CASStats.Size.Total);
+	}
+}
+
 #endif // UE_WITH_ZEN
 
 }
diff --git a/Engine/Source/Developer/Zen/Public/ZenServerInterface.h b/Engine/Source/Developer/Zen/Public/ZenServerInterface.h
--- a/Engine/Source/Developer/Zen/Public/ZenServerInterface.h
+++ b/Engine/Source/Developer/Zen/Public/ZenServerInterface.h
@@ -15,6 +15,7 @@
 #include "Serialization/JsonWriter.h"
 #include "Templates/PimplPtr.h"
 #include "Templates/UniquePtr.h"
+#include "AnalyticsEventAttribute.h"
 #include "ZenGlobals.h"
 
 #if UE_WITH_ZEN
@@ -136,6 +137,7 @@ private:
 	void PromptUserToStopRunningServerInstance(const FString& ServerFilePath);
 	FString ConditionalUpdateLocalInstall();
 	static bool AutoLaunch(const FServiceAutoLaunchSettings& InSettings, FString&& ExecutablePath, FString& OutHostName, uint16& OutPort);
+	void OnAnalyticsEvent(TArray<FAnalyticsEventAttribute>& Attributes);
 
 	mutable TPimplPtr<class FZenHttpRequest> StatsHttpRequest;
 	mutable TFuture<FZenStats> StatsRequest;
diff --git a/Engine/Source/Developer/Zen/Zen.Build.cs b/Engine/Source/Developer/Zen/Zen.Build.cs
--- a/Engine/Source/Developer/Zen/Zen.Build.cs
+++ b/Engine/Source/Developer/Zen/Zen.Build.cs
@@ -9,6 +9,13 @@ public class Zen : ModuleRules
 	{
 		PrivateDependencyModuleNames.AddRange(new string[] { "Core", "Sockets", "SSL", "Json" });
 		PrivateIncludePathModuleNames.Add("DesktopPlatform");
+		
+		PublicDependencyModuleNames.AddRange(
+			new string[] {		
+				"Analytics",
+			}
+		);
+
 		AddEngineThirdPartyPrivateStaticDependencies(Target, "libcurl");
 		AddEngineThirdPartyPrivateStaticDependencies(Target, "OpenSSL");
 
diff --git a/Engine/Source/Runtime/Analytics/Analytics/Public/Analytics.h b/Engine/Source/Runtime/Analytics/Analytics/Public/Analytics.h
--- a/Engine/Source/Runtime/Analytics/Analytics/Public/Analytics.h
+++ b/Engine/Source/Runtime/Analytics/Analytics/Public/Analytics.h
@@ -5,6 +5,7 @@
 #include "AnalyticsBuildType.h"
 #include "AnalyticsProviderConfigurationDelegate.h"
 #include "Containers/UnrealString.h"
+#include "Containers/Map.h"
 #include "CoreGlobals.h"
 #include "CoreMinimal.h"
 #include "Delegates/Delegate.h"
@@ -19,6 +20,10 @@
 ANALYTICS_API DECLARE_LOG_CATEGORY_EXTERN(LogAnalytics, Display, All);
 
 class IAnalyticsProvider;
+struct FAnalyticsEventAttribute;
+
+/** Analytics Event Callback */
+DECLARE_MULTICAST_DELEGATE_OneParam(FAnalyticsEventCallback, TArray<FAnalyticsEventAttribute>&);
 
 /**
  * The public interface for interacting with analytics.
@@ -63,6 +68,37 @@ public:
 		return FModuleManager::Get().IsModuleLoaded( "Analytics" );
 	}
 
+	TSharedPtr<FAnalyticsEventCallback> GetEventCallback(const FString& EventName)
+	{
+		// Check to see if we already have an event callback registered with this name
+		TSharedPtr<FAnalyticsEventCallback>* CallBackPtr = EventCallbacks.Find(EventName);
+
+		if (CallBackPtr != nullptr)
+		{
+			// Previously registered event callback so return it
+			return *CallBackPtr;
+		}
+
+		// Create and register a new event callback
+		TSharedPtr<FAnalyticsEventCallback> CallBack = MakeShareable(new FAnalyticsEventCallback());
+
+		EventCallbacks.Add(EventName, CallBack);
+
+		return CallBack;
+	}
+
+	void FireEventCallbacks(const FString& EventName, TArray<FAnalyticsEventAttribute>& Attributes)
+	{
+		// Check to see if we already have an event callback registered with this name
+		TSharedPtr<FAnalyticsEventCallback>* CallBackPtr = EventCallbacks.Find(EventName);
+
+		if (CallBackPtr != nullptr)
+		{
+			// There are valid callbacks for this event so fire it
+			(*CallBackPtr)->Broadcast(Attributes);
+		}
+	}
+
 	//--------------------------------------------------------------------------
 	// Provider Factory functions.
 	//--------------------------------------------------------------------------
@@ -192,8 +228,13 @@ public:
 	 * Helper for writing configuration values from to an INI file (which will be a common scenario). 
 	 */
 	virtual void WriteConfigValueToIni(const FString& IniName, const FString& SectionName, const FString& KeyName, const FString& Value);
+
+
 private:
 	virtual void StartupModule() override;
 	virtual void ShutdownModule() override;
+
+	TMap< FString, TSharedPtr<FAnalyticsEventCallback> > EventCallbacks;
+
 };
 
diff --git a/Engine/Source/Runtime/Analytics/AnalyticsET/Private/IAnalyticsProviderET.cpp b/Engine/Source/Runtime/Analytics/AnalyticsET/Private/IAnalyticsProviderET.cpp
--- a/Engine/Source/Runtime/Analytics/AnalyticsET/Private/IAnalyticsProviderET.cpp
+++ b/Engine/Source/Runtime/Analytics/AnalyticsET/Private/IAnalyticsProviderET.cpp
@@ -477,16 +477,26 @@ void FAnalyticsProviderET::RecordEvent(FString&& EventName, const TArray<FAnalyt
 	// let higher level code filter the decision of whether to send the event
 	if (ShouldRecordEvent(EventName))
 	{
-		// fire any callbacks
+		// Make a copy of the incoming attributes so we can append any new attributes from the the callback functions
+		TArray<FAnalyticsEventAttribute> FinalAttributes;
+		
+		// Append the incoming attributes to the final attribute list
+		FinalAttributes.Append(Attributes);
+
+		// Fire any analytics callbacks that append attributes to the list
+		FAnalytics::Get().FireEventCallbacks(EventName, FinalAttributes);
+		
+		// Fire any read only callbacks
 		for (const auto& Cb : EventRecordedCallbacks)
 		{
 			// we no longer track if the event was Json, each attribute does.
-			Cb(EventName, Attributes, false);
+			Cb(EventName, FinalAttributes, false);
 		}
 
 		if (!Config.UseLegacyProtocol)
 		{
-			EventCache.AddToCache(MoveTemp(EventName), Attributes);
+			EventCache.AddToCache(MoveTemp(EventName), FinalAttributes);
+
 			// if we aren't caching events, flush immediately. This is really only for debugging as it will significantly affect bandwidth.
 			if (!bShouldCacheEvents)
 			{
@@ -495,7 +505,7 @@ void FAnalyticsProviderET::RecordEvent(FString&& EventName, const TArray<FAnalyt
 		}
 		else
 		{
-			FlushEventLegacy(EventName, Attributes);
+			FlushEventLegacy(EventName, FinalAttributes);
 		}
 	}
 	else
diff --git a/Engine/Source/Runtime/Engine/Private/Analytics/StudioAnalytics.cpp b/Engine/Source/Runtime/Engine/Private/Analytics/StudioAnalytics.cpp
--- a/Engine/Source/Runtime/Engine/Private/Analytics/StudioAnalytics.cpp
+++ b/Engine/Source/Runtime/Engine/Private/Analytics/StudioAnalytics.cpp
@@ -10,6 +10,7 @@
 #include "AnalyticsBuildType.h"
 #include "AnalyticsEventAttribute.h"
 #include "IAnalyticsProviderET.h"
+#include "Analytics.h"
 #include "AnalyticsET.h"
 #include "GeneralProjectSettings.h"
 #include "Misc/EngineVersion.h"
@@ -18,13 +19,6 @@
 #include "Interfaces/IAnalyticsProvider.h"
 #include "Templates/SharedPointer.h"
 #include "HAL/PlatformProcess.h"
-#include "DerivedDataCacheInterface.h"
-#include "DerivedDataCacheUsageStats.h"
-#include "Virtualization/VirtualizationSystem.h"
-
-#if WITH_EDITOR
-#include "ZenServerInterface.h"
-#endif
 
 bool FStudioAnalytics::bInitialized = false;
 volatile double FStudioAnalytics::TimeEstimation = 0;
@@ -148,8 +142,6 @@ void FStudioAnalytics::RecordEvent(const FString& EventName, const TArray<FAnaly
 
 void FStudioAnalytics::FireEvent_Loading(const FString& LoadingName, double SecondsSpentLoading, const TArray<FAnalyticsEventAttribute>& InAttributes)
 {
-	using namespace UE::Virtualization;
-
 	// Ignore anything less than a 1/4th a second.
 	if (SecondsSpentLoading < 0.250)
 	{
@@ -166,227 +158,14 @@ void FStudioAnalytics::FireEvent_Loading(const FString& LoadingName, double Seco
 	{
 		TArray<FAnalyticsEventAttribute> Attributes;
 
-		
 		Attributes.Emplace(TEXT("LoadingName"), LoadingName);
 		Attributes.Emplace(TEXT("LoadingSeconds"), SecondsSpentLoading);
 		Attributes.Append(InAttributes);
 
+		// Original loading event
 		FStudioAnalytics::GetProvider().RecordEvent(TEXT("Performance.Loading"), Attributes);
 
-#if ENABLE_COOK_STATS
-
-		TArray<FDerivedDataCacheResourceStat> DDCResourceStats;
-
-		// Grab the latest resource stats
-		GetDerivedDataCacheRef().GatherResourceStats(DDCResourceStats);
-
-		FDerivedDataCacheResourceStat DDCResourceStatsTotal(TEXT("Total"));
-
-		// Accumulate Totals
-		for (const FDerivedDataCacheResourceStat& Stat : DDCResourceStats)
-		{
-			DDCResourceStatsTotal += Stat;
-		}
-
-		DDCResourceStats.Emplace(DDCResourceStatsTotal);
-
-		for (const FDerivedDataCacheResourceStat& Stat : DDCResourceStats)
-		{
-			FString BaseName = TEXT("DDC.Resource.") + Stat.AssetType;
-
-			BaseName = BaseName.Replace(TEXT("("), TEXT("")).Replace(TEXT(")"), TEXT(""));
-			
-			{
-				FString AttrName = BaseName + TEXT(".BuildCount");
-				Attributes.Emplace(MoveTemp(AttrName), Stat.BuildCount);
-			}
-			
-			{
-				FString AttrName = BaseName + TEXT(".BuildTimeSec");
-				Attributes.Emplace(MoveTemp(AttrName), Stat.BuildTimeSec);
-			}
-
-			{
-				FString AttrName = BaseName + TEXT(".BuildSizeMB");
-				Attributes.Emplace(MoveTemp(AttrName), Stat.BuildSizeMB);
-			}
-
-			{
-				FString AttrName = BaseName + TEXT(".LoadCount");
-				Attributes.Emplace(MoveTemp(AttrName), Stat.LoadCount);
-			}
-
-			{
-				FString AttrName = BaseName + TEXT(".LoadTimeSec");
-				Attributes.Emplace(MoveTemp(AttrName), Stat.LoadTimeSec);
-			}
-
-			{
-				FString AttrName = BaseName + TEXT(".LoadSizeMB");
-				Attributes.Emplace(MoveTemp(AttrName), Stat.LoadSizeMB);
-			}
-		}
-
-		// Grab the DDC summary stats
-		{
-			FDerivedDataCacheSummaryStats DDCSummaryStats;
-			GetDerivedDataCacheRef().GatherSummaryStats(DDCSummaryStats);
-
-			for (const FDerivedDataCacheSummaryStat& Stat : DDCSummaryStats.Stats)
-			{
-				FString FormattedAttrName = "DDC.Summary." + Stat.Key;
-				Attributes.Emplace(FormattedAttrName, Stat.Value);
-			}
-
-			// Grab the Virtualization stats
-			if (IVirtualizationSystem::IsInitialized())
-			{
-				IVirtualizationSystem& System = IVirtualizationSystem::Get();
-
-				FPayloadActivityInfo PayloadActivityInfo = System.GetAccumualtedPayloadActivityInfo();
-
-				const FString BaseName = TEXT("Virtualization");
-
-				{
-					FString AttrName = BaseName + TEXT(".Enabled");
-					Attributes.Emplace(MoveTemp(AttrName), System.IsEnabled());
-				}
-
-				{
-					FString AttrName = BaseName + TEXT(".Cache.TimeSpent");
-					Attributes.Emplace(MoveTemp(AttrName), (double)PayloadActivityInfo.Cache.CyclesSpent * FPlatformTime::GetSecondsPerCycle());
-				}
-
-				{
-					FString AttrName = BaseName + TEXT(".Cache.PayloadCount");
-					Attributes.Emplace(MoveTemp(AttrName), (double)PayloadActivityInfo.Cache.PayloadCount);
-				}
-
-				{
-					FString AttrName = BaseName + TEXT(".Cache.TotalBytes");
-					Attributes.Emplace(MoveTemp(AttrName), (double)PayloadActivityInfo.Cache.TotalBytes);
-				}
-
-				{
-					FString AttrName = BaseName + TEXT(".Push.TimeSpent");
-					Attributes.Emplace(MoveTemp(AttrName), (double)PayloadActivityInfo.Push.CyclesSpent * FPlatformTime::GetSecondsPerCycle());
-				}
-
-				{
-					FString AttrName = BaseName + TEXT(".Push.PayloadCount");
-					Attributes.Emplace(MoveTemp(AttrName), (double)PayloadActivityInfo.Push.PayloadCount);
-				}
-
-				{
-					FString AttrName = BaseName + TEXT(".Push.TotalBytes");
-					Attributes.Emplace(MoveTemp(AttrName), (double)PayloadActivityInfo.Push.TotalBytes);
-				}
-
-				{
-					FString AttrName = BaseName + TEXT(".Pull.TimeSpent");
-					Attributes.Emplace(MoveTemp(AttrName), (double)PayloadActivityInfo.Pull.CyclesSpent * FPlatformTime::GetSecondsPerCycle());
-				}
-
-				{
-					FString AttrName = BaseName + TEXT(".Pull.PayloadCount");
-					Attributes.Emplace(MoveTemp(AttrName), (double)PayloadActivityInfo.Pull.PayloadCount);
-				}
-
-				{
-					FString AttrName = BaseName + TEXT(".Pull.TotalBytes");
-					Attributes.Emplace(MoveTemp(AttrName), (double)PayloadActivityInfo.Pull.TotalBytes);
-				}
-			}
-		}
-
-#if UE_WITH_ZEN
-		if (UE::Zen::IsDefaultServicePresent())
-		{
-			// Grab the Zen summary stats
-			UE::Zen::FZenStats ZenStats;
-			UE::Zen::GetDefaultServiceInstance().GetStats(ZenStats);
-
-			const FString BaseName = TEXT("Zen");
-
-			{
-				FString AttrName = BaseName + TEXT(".Enabled");
-				Attributes.Emplace(MoveTemp(AttrName), ZenStats.IsValid);
-			}
-
-			{
-				FString AttrName = BaseName + TEXT(".Cache.HitRatio");
-				Attributes.Emplace(MoveTemp(AttrName), ZenStats.CacheStats.HitRatio);
-			}
-
-			{
-				FString AttrName = BaseName + TEXT(".Cache.Hits");
-				Attributes.Emplace(MoveTemp(AttrName), ZenStats.CacheStats.Hits);
-			}
-
-			{
-				FString AttrName = BaseName + TEXT(".Cache.Misses");
-				Attributes.Emplace(MoveTemp(AttrName), ZenStats.CacheStats.Misses);
-			}
-
-			{
-				FString AttrName = BaseName + TEXT(".Cache.Size.Disk");
-				Attributes.Emplace(MoveTemp(AttrName), ZenStats.CacheStats.Size.Disk);
-			}
-
-			{
-				FString AttrName = BaseName + TEXT(".Cache.Size.Memory");
-				Attributes.Emplace(MoveTemp(AttrName), ZenStats.CacheStats.Size.Memory);
-			}
-
-			{
-				FString AttrName = BaseName + TEXT(".Cache.UpstreamHits");
-				Attributes.Emplace(MoveTemp(AttrName), ZenStats.CacheStats.UpstreamHits);
-			}
-
-			{
-				FString AttrName = BaseName + TEXT(".Cache.UpstreamRatio");
-				Attributes.Emplace(MoveTemp(AttrName), ZenStats.CacheStats.UpstreamRatio);
-			}
-
-			{
-				FString AttrName = BaseName + TEXT(".Cache.TotalUploadedMB");
-				Attributes.Emplace(MoveTemp(AttrName), ZenStats.UpstreamStats.TotalUploadedMB);
-			}
-
-			{
-				FString AttrName = BaseName + TEXT(".Upstream.TotalDownloadedMB");
-				Attributes.Emplace(MoveTemp(AttrName), ZenStats.UpstreamStats.TotalDownloadedMB);
-			}
-
-			{
-				FString AttrName = BaseName + TEXT(".Upstream.TotalUploadedMB");
-				Attributes.Emplace(MoveTemp(AttrName), ZenStats.UpstreamStats.TotalUploadedMB);
-			}
-
-			{
-				FString AttrName = BaseName + TEXT(".Cas.Size.Large");
-				Attributes.Emplace(MoveTemp(AttrName), ZenStats.CASStats.Size.Large);
-			}
-
-			{
-				FString AttrName = BaseName + TEXT(".Cas.Size.Small");
-				Attributes.Emplace(MoveTemp(AttrName), ZenStats.CASStats.Size.Small);
-			}
-
-			{
-				FString AttrName = BaseName + TEXT(".Cas.Size.Tiny");
-				Attributes.Emplace(MoveTemp(AttrName), ZenStats.CASStats.Size.Tiny);
-			}
-
-			{
-				FString AttrName = BaseName + TEXT(".Cas.Size.Total");
-				Attributes.Emplace(MoveTemp(AttrName), ZenStats.CASStats.Size.Total);
-			}
-		}
-#endif
-
-		// Store it all in the loading event
+		// New loading event which has callbacks from DDC, VA, Zen etc..
 		FStudioAnalytics::GetProvider().RecordEvent(TEXT("Core.Loading"), Attributes);
-#endif
 	}
-}
\ No newline at end of file
+}
