Fix potential race condition happenning when overwriting an audio chunk cache key from the wrong thread

#rnx

diff --git a/Engine/Source/Runtime/Engine/Private/AudioStreamingCache.cpp b/Engine/Source/Runtime/Engine/Private/AudioStreamingCache.cpp
--- a/Engine/Source/Runtime/Engine/Private/AudioStreamingCache.cpp
+++ b/Engine/Source/Runtime/Engine/Private/AudioStreamingCache.cpp
@@ -1636,11 +1636,22 @@ void FAudioChunkCache::KickOffAsyncLoad(FCacheElement* CacheElement, const FChun
 		CacheElement->DebugInfo.TimeLoadStarted = FPlatformTime::Cycles64();
 #endif
 
-
 		TFunction<void(bool)> OnLoadComplete = [OnLoadCompleted, CallbackThread, CacheElement, InKey, ChunkDataSize](bool bRequestFailed)
 		{
 			// Populate key and DataSize. The async read request was set up to write directly into CacheElement->ChunkData.
-			CacheElement->Key = InKey;
+			// The following condition should always be true and there should be no need
+			// to overwrite the Key as it can cause race condition between the callback thread
+			// and other threads trying to search for elements by key.
+			
+			// If this ensure is tripped for some reason, we must find the root cause, not remove the ensure.
+			ensure(CacheElement->Key == InKey);
+			// This can be removed later once we're sure the ensure is never tripped
+			// For now, avoid overwriting when both values are the same to avoid a race condition.
+			if (!(CacheElement->Key == InKey))
+			{
+				CacheElement->Key = InKey;
+			}
+
 			CacheElement->ChunkDataSize = ChunkDataSize;
 			CacheElement->bIsLoaded = true;
 
