Optimizations of world partition update streaming state
-  Avoid computing streaming source spherical sector intersection information when source is a sphere
- Avoid validating cell coordinate when already done
- Optimized FSquare2DGridHelper::ForEachIntersectingCells

diff --git a/Engine/Source/Runtime/Engine/Private/WorldPartition/RuntimeSpatialHash/RuntimeSpatialHashGridHelper.cpp b/Engine/Source/Runtime/Engine/Private/WorldPartition/RuntimeSpatialHash/RuntimeSpatialHashGridHelper.cpp
--- a/Engine/Source/Runtime/Engine/Private/WorldPartition/RuntimeSpatialHash/RuntimeSpatialHashGridHelper.cpp
+++ b/Engine/Source/Runtime/Engine/Private/WorldPartition/RuntimeSpatialHash/RuntimeSpatialHashGridHelper.cpp
@@ -257,85 +257,25 @@ FSquare2DGridHelper GetPartitionedActors(const FBox& WorldBounds, const FSpatial
 
 #endif // #if WITH_EDITOR
 
-FORCEINLINE static bool IsClockWise(const FVector2D& V1, const FVector2D& V2) 
+bool FSquare2DGridHelper::FGrid2D::DoesCircleSectorIntersectsCell(const FGridCellCoord2& Coords, const FSphericalSector& InShape) const
 {
-	return (V1 ^ V2) < 0;
-}
-
-FORCEINLINE static bool IsVectorInsideVectorPair(const FVector2D& TestVector, const FVector2D& V1, const FVector2D& V2)
-{
-	return IsClockWise(V1, TestVector) && !IsClockWise(V2, TestVector);
-}
-
-FORCEINLINE static bool IsPointInsideSector(const FVector2D& TestPoint, const FVector2D& SectorCenter, float SectorRadiusSquared, const FVector2D& SectorStart, const FVector2D& SectorEnd, float SectorAngle)
-{
-	const FVector2D TestVector = TestPoint - SectorCenter;
-	if (TestVector.SizeSquared() > SectorRadiusSquared)
-	{
-		return false;
-	}
-	
-	if (SectorAngle <= 180.0f)
-	{
-		return IsVectorInsideVectorPair(TestVector, SectorStart, SectorEnd);
-	}
-	else
-	{
-		return !IsVectorInsideVectorPair(TestVector, SectorEnd, SectorStart);
-	}
-}
-
-bool FSquare2DGridHelper::FGrid2D::DoesCircleSectorIntersectsCell(const FGridCellCoord2& Coords, const FVector2D& SectorCenter, float SectorRadiusSquared, const FVector2D& SectorStartVector, const FVector2D& SectorEndVector, float SectorAngle) const
-{
-	const int64 CellIndex = Coords.Y * GridSize + Coords.X;
 	FBox2D CellBounds;
-	GetCellBounds(CellIndex, CellBounds);
+	GetCellBounds(Coords, CellBounds);
 
-	// Test whether any cell corners are inside sector
-	if (IsPointInsideSector(FVector2D(CellBounds.Min.X, CellBounds.Min.Y), SectorCenter, SectorRadiusSquared, SectorStartVector, SectorEndVector, SectorAngle) ||
-		IsPointInsideSector(FVector2D(CellBounds.Max.X, CellBounds.Min.Y), SectorCenter, SectorRadiusSquared, SectorStartVector, SectorEndVector, SectorAngle) ||
-		IsPointInsideSector(FVector2D(CellBounds.Max.X, CellBounds.Max.Y), SectorCenter, SectorRadiusSquared, SectorStartVector, SectorEndVector, SectorAngle) ||
-		IsPointInsideSector(FVector2D(CellBounds.Min.X, CellBounds.Max.Y), SectorCenter, SectorRadiusSquared, SectorStartVector, SectorEndVector, SectorAngle))
-	{
-		return true;
-	}
-
-	// Test whether any sector point lies inside the cell bounds
-	if (CellBounds.IsInside(SectorCenter) ||
-		CellBounds.IsInside(SectorCenter + SectorStartVector) ||
-		CellBounds.IsInside(SectorCenter + SectorEndVector))
-	{
-		return true;
-	}
-
-	// Test whether closest point on cell from center is inside sector
-	if (IsPointInsideSector(CellBounds.GetClosestPointTo(SectorCenter), SectorCenter, SectorRadiusSquared, SectorStartVector, SectorEndVector, SectorAngle))
-	{
-		return true;
-	}
-
-	return false;
+	return InShape.IntersectsBox(CellBounds);
 }
 
 int32 FSquare2DGridHelper::FGrid2D::ForEachIntersectingCells(const FSphericalSector& InShape, TFunctionRef<void(const FGridCellCoord2&)> InOperation) const
 {
 	check(!InShape.IsNearlyZero());
 	const bool bIsSphere = InShape.IsSphere();
-	const FVector ScaledAxis = FVector(FVector2D(InShape.GetAxis()), 0).GetSafeNormal() * InShape.GetRadius();
-	const FVector SectorStart = FRotator(0, 0.5f * InShape.GetAngle(), 0).RotateVector(ScaledAxis);
-	const FVector SectorEnd = FRotator(0, -0.5f * InShape.GetAngle(), 0).RotateVector(ScaledAxis);
-
-	const FVector2D Center2D(InShape.GetCenter());
-	const FVector2D SectorStart2D = FVector2D(SectorStart);
-	const FVector2D SectorEnd2D = FVector2D(SectorEnd);
-	const float SectorRadiusSquared = InShape.GetRadius() * InShape.GetRadius();
 
 	int32 NumCells = 0;
 	const FSphere Sphere(InShape.GetCenter(), InShape.GetRadius());
-	ForEachIntersectingCells(Sphere, [this, bIsSphere, Center2D, SectorRadiusSquared, SectorStart2D, SectorEnd2D, SectorAngle = InShape.GetAngle(), InOperation, &NumCells](const FGridCellCoord2& Coords)
+	ForEachIntersectingCells(Sphere, [this, bIsSphere, InShape, InOperation, &NumCells](const FGridCellCoord2& Coords)
 	{
 		// If sector, filter coords that intersect sector
-		if (bIsSphere || DoesCircleSectorIntersectsCell(Coords, Center2D, SectorRadiusSquared, SectorStart2D, SectorEnd2D, SectorAngle))
+		if (bIsSphere || DoesCircleSectorIntersectsCell(Coords, InShape))
 		{
 			InOperation(Coords);
 			NumCells++;
diff --git a/Engine/Source/Runtime/Engine/Private/WorldPartition/RuntimeSpatialHash/RuntimeSpatialHashGridHelper.h b/Engine/Source/Runtime/Engine/Private/WorldPartition/RuntimeSpatialHash/RuntimeSpatialHashGridHelper.h
--- a/Engine/Source/Runtime/Engine/Private/WorldPartition/RuntimeSpatialHash/RuntimeSpatialHashGridHelper.h
+++ b/Engine/Source/Runtime/Engine/Private/WorldPartition/RuntimeSpatialHash/RuntimeSpatialHashGridHelper.h
@@ -70,11 +70,11 @@ struct FSquare2DGridHelper
 		 *
 		 * @return true if the specified coord was valid
 		 */
-		inline bool GetCellBounds(const FGridCellCoord2& InCoords, FBox2D& OutBounds) const
+		inline bool GetCellBounds(const FGridCellCoord2& InCoords, FBox2D& OutBounds, bool bCheckIsValidCoord = true) const
 		{
-			if (IsValidCoords(InCoords))
+			if (!bCheckIsValidCoord || IsValidCoords(InCoords))
 			{
-				const FVector2D Min = (FVector2D(Origin) - FVector2D(GridSize * CellSize * 0.5, GridSize * CellSize * 0.5)) + FVector2D(InCoords.X * CellSize, InCoords.Y * CellSize);
+				const FVector2D Min = (FVector2D(Origin) - FVector2D(GridSize * CellSize * 0.5)) + FVector2D(InCoords.X * CellSize, InCoords.Y * CellSize);
 				const FVector2D Max = Min + FVector2D(CellSize, CellSize);
 				OutBounds = FBox2D(Min, Max);
 				return true;
@@ -196,7 +196,7 @@ struct FSquare2DGridHelper
 						// Validate that generated coordinate is valid (in case we reached the 64-bit limit of cell index)
 						if (IsValidCoords(Coord))
 						{
-							if (!InOperation(FGridCellCoord2(x, y)))
+							if (!InOperation(Coord))
 							{
 								return NumCells;
 							}
@@ -225,16 +225,18 @@ struct FSquare2DGridHelper
 
 			// @todo_ow: rasterize circle instead?
 			const FBox Box(InSphere.Center - FVector(InSphere.W), InSphere.Center + FVector(InSphere.W));
+			const double SquareDistance = InSphere.W * InSphere.W;
+			const FVector2D SphereCenter(InSphere.Center);
 
-			ForEachIntersectingCells(Box, [this, &InSphere, &InOperation, &NumCells](const FGridCellCoord2& Coords)
+			ForEachIntersectingCells(Box, [this, SquareDistance, &SphereCenter, &InOperation, &NumCells](const FGridCellCoord2& Coords)
 			{
-				const int64 CellIndex = Coords.Y * GridSize + Coords.X;
-
+				// No need to check validity of coords as it's already done
+				const bool bCheckIsValidCoords = false;
 				FBox2D CellBounds;
-				GetCellBounds(CellIndex, CellBounds);
+				GetCellBounds(Coords, CellBounds, bCheckIsValidCoords);
 
-				FVector2D Delta = FVector2D(InSphere.Center) - FVector2D::Max(CellBounds.GetCenter() - CellBounds.GetExtent(), FVector2D::Min(FVector2D(InSphere.Center), CellBounds.GetCenter() + CellBounds.GetExtent()));
-				if ((Delta.X * Delta.X + Delta.Y * Delta.Y) < (InSphere.W * InSphere.W))
+				FVector2D Delta = SphereCenter - FVector2D::Max(CellBounds.Min, FVector2D::Min(SphereCenter, CellBounds.Max));
+				if ((Delta.X * Delta.X + Delta.Y * Delta.Y) < SquareDistance)
 				{
 					InOperation(Coords);
 					NumCells++;
@@ -253,7 +255,7 @@ struct FSquare2DGridHelper
 
 	private:
 
-		bool DoesCircleSectorIntersectsCell(const FGridCellCoord2& Coords, const FVector2D& SectorCenter, float SectorRadiusSquared, const FVector2D& SectorStartVector, const FVector2D& SectorEndVector, float SectorAngle) const;
+		bool DoesCircleSectorIntersectsCell(const FGridCellCoord2& Coords, const FSphericalSector& InShape) const;
 	};
 
 	struct FGridLevel : public FGrid2D
diff --git a/Engine/Source/Runtime/Engine/Private/WorldPartition/WorldPartitionStreamingSource.cpp b/Engine/Source/Runtime/Engine/Private/WorldPartition/WorldPartitionStreamingSource.cpp
new file mode 100644
--- /dev/null
+++ b/Engine/Source/Runtime/Engine/Private/WorldPartition/WorldPartitionStreamingSource.cpp
@@ -0,0 +1,112 @@
+// Copyright Epic Games, Inc. All Rights Reserved.
+
+#include "WorldPartition/WorldPartitionStreamingSource.h"
+
+FORCEINLINE static bool IsClockWise(const FVector2D& V1, const FVector2D& V2)
+{
+	return (V1 ^ V2) < 0;
+}
+
+FORCEINLINE static bool IsVectorInsideVectorPair(const FVector2D& TestVector, const FVector2D& V1, const FVector2D& V2)
+{
+	return IsClockWise(V1, TestVector) && !IsClockWise(V2, TestVector);
+}
+
+FORCEINLINE static bool IsPointInsideSector(const FVector2D& TestPoint, const FVector2D& SectorCenter, float SectorRadiusSquared, const FVector2D& SectorStart, const FVector2D& SectorEnd, float SectorAngle)
+{
+	const FVector2D TestVector = TestPoint - SectorCenter;
+	if (TestVector.SizeSquared() > SectorRadiusSquared)
+	{
+		return false;
+	}
+
+	if (SectorAngle <= 180.0f)
+	{
+		return IsVectorInsideVectorPair(TestVector, SectorStart, SectorEnd);
+	}
+	else
+	{
+		return !IsVectorInsideVectorPair(TestVector, SectorEnd, SectorStart);
+	}
+}
+
+TArray<TPair<FVector, FVector>> FSphericalSector::BuildDebugMesh() const
+{
+	TArray<TPair<FVector, FVector>> Segments;
+	if (!IsValid())
+	{
+		return Segments;
+	}
+
+	const int32 SegmentCount = FMath::Max(4, FMath::CeilToInt32(64 * (float)Angle / 360.f));
+	const FReal AngleStep = Angle / FReal(SegmentCount);
+	const FRotator ShapeRotation = FRotationMatrix::MakeFromX(Axis).Rotator();
+	const FVector ScaledAxis = FVector::ForwardVector * Radius;
+	const int32 RollCount = 16;
+
+	Segments.Reserve(2 * (RollCount + 1) * (SegmentCount + 2));
+	int32 LastArcStartIndex = -1;
+	for (int32 i = 0; i <= RollCount; ++i)
+	{
+		const float Roll = 180.f * i / float(RollCount);
+		const FTransform Transform(FRotator(0, 0, Roll) + ShapeRotation, Center);
+		FVector SegmentStart = Transform.TransformPosition(FRotator(0, -0.5f * Angle, 0).RotateVector(ScaledAxis));
+		Segments.Emplace(Center, SegmentStart);
+		int32 CurrentArcStartIndex = Segments.Num();
+		// Build sector arc
+		for (int32 j = 1; j <= SegmentCount; j++)
+		{
+			FVector SegmentEnd = Transform.TransformPosition(FRotator(0, -0.5f * Angle + (AngleStep * j), 0).RotateVector(ScaledAxis));
+			Segments.Emplace(SegmentStart, SegmentEnd);
+			SegmentStart = SegmentEnd;
+		}
+		Segments.Emplace(Center, SegmentStart);
+		if (i > 0)
+		{
+			// Connect sector arc to previous arc
+			for (int32 j = 0; j < SegmentCount; j++)
+			{
+				Segments.Emplace(Segments[LastArcStartIndex + j].Key, Segments[CurrentArcStartIndex + j].Key);
+			}
+			Segments.Emplace(Segments[LastArcStartIndex + SegmentCount - 1].Value, Segments[CurrentArcStartIndex + SegmentCount - 1].Value);
+		}
+		LastArcStartIndex = CurrentArcStartIndex;
+	}
+	return Segments;
+}
+
+bool FSphericalSector::IntersectsBox(const FBox2D& InBox) const
+{
+	const FVector ScaledAxis = FVector(FVector2D(Axis), 0).GetSafeNormal() * Radius;
+	const FVector SectorStart = FRotator(0, 0.5f * Angle, 0).RotateVector(ScaledAxis);
+	const FVector SectorEnd = FRotator(0, -0.5f * Angle, 0).RotateVector(ScaledAxis);
+	const FVector2D SectorCenter(Center);
+	const FVector2D SectorStartVector = FVector2D(SectorStart);
+	const FVector2D SectorEndVector = FVector2D(SectorEnd);
+	const float SectorRadiusSquared = Radius * Radius;
+
+	// Test whether any cell corners are inside sector
+	if (IsPointInsideSector(FVector2D(InBox.Min.X, InBox.Min.Y), SectorCenter, SectorRadiusSquared, SectorStartVector, SectorEndVector, Angle) ||
+		IsPointInsideSector(FVector2D(InBox.Max.X, InBox.Min.Y), SectorCenter, SectorRadiusSquared, SectorStartVector, SectorEndVector, Angle) ||
+		IsPointInsideSector(FVector2D(InBox.Max.X, InBox.Max.Y), SectorCenter, SectorRadiusSquared, SectorStartVector, SectorEndVector, Angle) ||
+		IsPointInsideSector(FVector2D(InBox.Min.X, InBox.Max.Y), SectorCenter, SectorRadiusSquared, SectorStartVector, SectorEndVector, Angle))
+	{
+		return true;
+	}
+
+	// Test whether any sector point lies inside the cell bounds
+	if (InBox.IsInside(SectorCenter) ||
+		InBox.IsInside(SectorCenter + SectorStartVector) ||
+		InBox.IsInside(SectorCenter + SectorEndVector))
+	{
+		return true;
+	}
+
+	// Test whether closest point on cell from center is inside sector
+	if (IsPointInsideSector(InBox.GetClosestPointTo(SectorCenter), SectorCenter, SectorRadiusSquared, SectorStartVector, SectorEndVector, Angle))
+	{
+		return true;
+	}
+
+	return false;
+}
diff --git a/Engine/Source/Runtime/Engine/Public/WorldPartition/WorldPartitionStreamingSource.h b/Engine/Source/Runtime/Engine/Public/WorldPartition/WorldPartitionStreamingSource.h
--- a/Engine/Source/Runtime/Engine/Public/WorldPartition/WorldPartitionStreamingSource.h
+++ b/Engine/Source/Runtime/Engine/Public/WorldPartition/WorldPartitionStreamingSource.h
@@ -63,50 +63,10 @@ public:
 	}
 
 	/** Helper method that builds a list of debug display segments */
-	TArray<TPair<FVector, FVector>> BuildDebugMesh() const
-	{
-		TArray<TPair<FVector, FVector>> Segments;
-		if (!IsValid())
-		{
-			return Segments;
-		}
-
-		const int32 SegmentCount = FMath::Max(4, FMath::CeilToInt32(64 * (float)Angle / 360.f));
-		const FReal AngleStep = Angle / FReal(SegmentCount);
-		const FRotator ShapeRotation = FRotationMatrix::MakeFromX(Axis).Rotator();
-		const FVector ScaledAxis = FVector::ForwardVector * Radius;
-		const int32 RollCount = 16;
+	TArray<TPair<FVector, FVector>> BuildDebugMesh() const;
 
-		Segments.Reserve(2 * (RollCount + 1) * (SegmentCount + 2));
-		int32 LastArcStartIndex = -1;
-		for (int32 i = 0; i <= RollCount; ++i)
-		{
-			const float Roll = 180.f * i / float(RollCount);
-			const FTransform Transform(FRotator(0, 0, Roll) + ShapeRotation, Center);
-			FVector SegmentStart = Transform.TransformPosition(FRotator(0, -0.5f * Angle, 0).RotateVector(ScaledAxis));
-			Segments.Emplace(Center, SegmentStart);
-			int32 CurrentArcStartIndex = Segments.Num();
-			// Build sector arc
-			for (int32 j = 1; j <= SegmentCount; j++)
-			{
-				FVector SegmentEnd = Transform.TransformPosition(FRotator(0, -0.5f * Angle + (AngleStep * j), 0).RotateVector(ScaledAxis));
-				Segments.Emplace(SegmentStart, SegmentEnd);
-				SegmentStart = SegmentEnd;
-			}
-			Segments.Emplace(Center, SegmentStart);
-			if (i > 0)
-			{
-				// Connect sector arc to previous arc
-				for (int32 j = 0; j < SegmentCount; j++)
-				{
-					Segments.Emplace(Segments[LastArcStartIndex + j].Key, Segments[CurrentArcStartIndex + j].Key);
-				}
-				Segments.Emplace(Segments[LastArcStartIndex + SegmentCount - 1].Value, Segments[CurrentArcStartIndex + SegmentCount - 1].Value);
-			}
-			LastArcStartIndex = CurrentArcStartIndex;
-		}
-		return Segments;
-	}
+	/** Returns whether spherical sector intersects 2D box */
+	bool IntersectsBox(const FBox2D& InBox) const;
 
 private:
 	/** Sphere center point. */
