Added missing include

#rnx

diff --git a/Engine/Source/Runtime/Engine/Private/WorldPartition/HLOD/HLODActor.cpp b/Engine/Source/Runtime/Engine/Private/WorldPartition/HLOD/HLODActor.cpp
--- a/Engine/Source/Runtime/Engine/Private/WorldPartition/HLOD/HLODActor.cpp
+++ b/Engine/Source/Runtime/Engine/Private/WorldPartition/HLOD/HLODActor.cpp
@@ -17,6 +17,7 @@
 #include "WorldPartition/HLOD/HLODLayer.h"
 #include "WorldPartition/HLOD/IWorldPartitionHLODUtilitiesModule.h"
 
+#include "RHI.h"
 #include "Modules/ModuleManager.h"
 #include "Engine/TextureStreamingTypes.h"
 #endif
