Optimized UDataLayerSubsystem::IsAnyDataLayerInEffectiveRuntimeState : avoid rebuilding EffectiveActiveDataLayerNames / EffectiveLoadedDataLayerNames every time.

diff --git a/Engine/Source/Runtime/Engine/Private/WorldPartition/DataLayer/DataLayerSubsystem.cpp b/Engine/Source/Runtime/Engine/Private/WorldPartition/DataLayer/DataLayerSubsystem.cpp
--- a/Engine/Source/Runtime/Engine/Private/WorldPartition/DataLayer/DataLayerSubsystem.cpp
+++ b/Engine/Source/Runtime/Engine/Private/WorldPartition/DataLayer/DataLayerSubsystem.cpp
@@ -66,6 +66,7 @@ void FDataLayersEditorBroadcast::StaticOnActorDataLayersEditorLoadingStateChange
 #endif
 
 UDataLayerSubsystem::UDataLayerSubsystem()
+	: bIsCachedEffectiveStateDirty(true)
 {
 #if WITH_EDITOR
 	DataLayerActorEditorContextID = 0;
@@ -76,6 +77,8 @@ void UDataLayerSubsystem::Initialize(FSubsystemCollectionBase& Collection)
 {
 	Super::Initialize(Collection);
 
+	bIsCachedEffectiveStateDirty = true;
+
 #if WITH_EDITOR
 	UClass* DataLayerLoadingPolicyClassValue = DataLayerLoadingPolicyClass.Get();
 	if (!DataLayerLoadingPolicyClassValue)
@@ -231,26 +234,41 @@ TArray<const UDataLayerInstance*> UDataLayerSubsystem::GetRuntimeDataLayerInstan
 
 #endif
 
-TSet<FName> UDataLayerSubsystem::GetEffectiveActiveDataLayerNames() const
+const TSet<FName>& UDataLayerSubsystem::GetEffectiveActiveDataLayerNames() const
 {
-	TSet<FName> EffectiveActiveDataLayerNames;
-	WorldDataLayerCollection.ForEachWorldDataLayers([&EffectiveActiveDataLayerNames](AWorldDataLayers* WorldDataLayers)
-	{
-		EffectiveActiveDataLayerNames.Append(WorldDataLayers->GetEffectiveActiveDataLayerNames());
-	});
-	
-	return EffectiveActiveDataLayerNames;
+	UpdateCachedEffectiveRuntimeStates();
+	return CachedEffectiveActiveDataLayerNames;
+}
+
+const TSet<FName>& UDataLayerSubsystem::GetEffectiveLoadedDataLayerNames() const
+{
+	UpdateCachedEffectiveRuntimeStates();
+	return CachedEffectiveLoadedDataLayerNames;
 }
 
-TSet<FName> UDataLayerSubsystem::GetEffectiveLoadedDataLayerNames() const
+void UDataLayerSubsystem::OnEffectiveRuntimeDataLayerStatesChanged(AWorldDataLayers* InWorldDataLayers)
 {
-	TSet<FName> EffectiveLoadedDataLayerNames;
-	WorldDataLayerCollection.ForEachWorldDataLayers([&EffectiveLoadedDataLayerNames](AWorldDataLayers* WorldDataLayers)
+	bIsCachedEffectiveStateDirty = true;
+}
+
+void UDataLayerSubsystem::UpdateCachedEffectiveRuntimeStates() const
+{
+	if (bIsCachedEffectiveStateDirty)
 	{
-		EffectiveLoadedDataLayerNames.Append(WorldDataLayers->GetEffectiveLoadedDataLayerNames());
-	});
+		bIsCachedEffectiveStateDirty = false;
+
+		CachedEffectiveActiveDataLayerNames.Reset();
+		WorldDataLayerCollection.ForEachWorldDataLayers([this](AWorldDataLayers* WorldDataLayers)
+		{
+			CachedEffectiveActiveDataLayerNames.Append(WorldDataLayers->GetEffectiveActiveDataLayerNames());
+		});
 
-	return EffectiveLoadedDataLayerNames;
+		CachedEffectiveLoadedDataLayerNames.Reset();
+		WorldDataLayerCollection.ForEachWorldDataLayers([this](AWorldDataLayers* WorldDataLayers)
+		{
+			CachedEffectiveLoadedDataLayerNames.Append(WorldDataLayers->GetEffectiveLoadedDataLayerNames());
+		});
+	}
 }
 
 void UDataLayerSubsystem::RegisterWorldDataLayer(AWorldDataLayers* WorldDataLayers)
@@ -263,6 +281,7 @@ void UDataLayerSubsystem::RegisterWorldDataLayer(AWorldDataLayers* WorldDataLaye
 
 	if (WorldDataLayerCollection.RegisterWorldDataLayer(WorldDataLayers))
 	{
+		bIsCachedEffectiveStateDirty = true;
 #if WITH_EDITOR
 		OnWorldDataLayerPostRegister.Broadcast(WorldDataLayers);
 #endif
@@ -276,7 +295,7 @@ void UDataLayerSubsystem::UnregisterWorldDataLayer(AWorldDataLayers* WorldDataLa
 #if WITH_EDITOR
 		OnWorldDataLayerPreUnregister.Broadcast(WorldDataLayers);
 #endif
-
+		bIsCachedEffectiveStateDirty = true;
 		WorldDataLayerCollection.UnregisterWorldDataLayer(WorldDataLayers);
 	}
 }
@@ -360,11 +379,26 @@ EDataLayerRuntimeState UDataLayerSubsystem::GetDataLayerEffectiveRuntimeStateByN
 
 bool UDataLayerSubsystem::IsAnyDataLayerInEffectiveRuntimeState(const TArray<FName>& InDataLayerNames, EDataLayerRuntimeState InState) const
 {
-	for (FName DataLayerName : InDataLayerNames)
+	if (InState == EDataLayerRuntimeState::Activated)
 	{
-		if (GetDataLayerEffectiveRuntimeStateByName(DataLayerName) == InState)
+		const TSet<FName>& Activated = GetEffectiveActiveDataLayerNames();
+		for (const FName& DataLayerName : InDataLayerNames)
 		{
-			return true;
+			if (Activated.Contains(DataLayerName))
+			{
+				return true;
+			}
+		}
+	}
+	else if (InState == EDataLayerRuntimeState::Loaded)
+	{
+		const TSet<FName>& Loaded = GetEffectiveLoadedDataLayerNames();
+		for (const FName& DataLayerName : InDataLayerNames)
+		{
+			if (Loaded.Contains(DataLayerName))
+			{
+				return true;
+			}
 		}
 	}
 	return false;
diff --git a/Engine/Source/Runtime/Engine/Private/WorldPartition/DataLayer/WorldDataLayers.cpp b/Engine/Source/Runtime/Engine/Private/WorldPartition/DataLayer/WorldDataLayers.cpp
--- a/Engine/Source/Runtime/Engine/Private/WorldPartition/DataLayer/WorldDataLayers.cpp
+++ b/Engine/Source/Runtime/Engine/Private/WorldPartition/DataLayer/WorldDataLayers.cpp
@@ -252,12 +252,20 @@ void AWorldDataLayers::OnRep_EffectiveActiveDataLayerNames()
 {
 	EffectiveActiveDataLayerNames.Reset();
 	EffectiveActiveDataLayerNames.Append(RepEffectiveActiveDataLayerNames);
+	if (UDataLayerSubsystem* DataLayerSubsystem = UWorld::GetSubsystem<UDataLayerSubsystem>(GetWorld()))
+	{
+		DataLayerSubsystem->OnEffectiveRuntimeDataLayerStatesChanged(this);
+	}
 }
 
 void AWorldDataLayers::OnRep_EffectiveLoadedDataLayerNames()
 {
 	EffectiveLoadedDataLayerNames.Reset();
 	EffectiveLoadedDataLayerNames.Append(RepEffectiveLoadedDataLayerNames);
+	if (UDataLayerSubsystem* DataLayerSubsystem = UWorld::GetSubsystem<UDataLayerSubsystem>(GetWorld()))
+	{
+		DataLayerSubsystem->OnEffectiveRuntimeDataLayerStatesChanged(this);
+	}
 }
 
 EDataLayerRuntimeState AWorldDataLayers::GetDataLayerEffectiveRuntimeStateByName(FName InDataLayerName) const
@@ -311,6 +319,11 @@ void AWorldDataLayers::ResolveEffectiveRuntimeState(const UDataLayerInstance* In
 		RepEffectiveActiveDataLayerNames = EffectiveActiveDataLayerNames.Array();
 		RepEffectiveLoadedDataLayerNames = EffectiveLoadedDataLayerNames.Array();
 
+		if (UDataLayerSubsystem* DataLayerSubsystem = UWorld::GetSubsystem<UDataLayerSubsystem>(GetWorld()))
+		{
+			DataLayerSubsystem->OnEffectiveRuntimeDataLayerStatesChanged(this);
+		}
+
 		++DataLayersStateEpoch;
 
 		if (bInNotifyChange)
diff --git a/Engine/Source/Runtime/Engine/Public/WorldPartition/DataLayer/DataLayerSubsystem.h b/Engine/Source/Runtime/Engine/Public/WorldPartition/DataLayer/DataLayerSubsystem.h
--- a/Engine/Source/Runtime/Engine/Public/WorldPartition/DataLayer/DataLayerSubsystem.h
+++ b/Engine/Source/Runtime/Engine/Public/WorldPartition/DataLayer/DataLayerSubsystem.h
@@ -120,8 +120,8 @@ public:
 	EDataLayerRuntimeState GetDataLayerEffectiveRuntimeState(const UDataLayerInstance* InDataLayer) const;
 	EDataLayerRuntimeState GetDataLayerEffectiveRuntimeStateByName(const FName& InDataLayerName) const;
 	bool IsAnyDataLayerInEffectiveRuntimeState(const TArray<FName>& InDataLayerNames, EDataLayerRuntimeState InState) const;
-	TSet<FName> GetEffectiveActiveDataLayerNames() const;
-	TSet<FName> GetEffectiveLoadedDataLayerNames() const;
+	const TSet<FName>& GetEffectiveActiveDataLayerNames() const;
+	const TSet<FName>& GetEffectiveLoadedDataLayerNames() const;
 
 #if WITH_EDITOR
 	/** Called once a AWorldDataLayers has been registered. */
@@ -243,6 +243,15 @@ private:
 	/** Console command used to set Runtime DataLayer state*/
 	static class FAutoConsoleCommand SetDataLayerRuntimeStateCommand;
 
+	void OnEffectiveRuntimeDataLayerStatesChanged(AWorldDataLayers* InWorldDataLayers);
+	void UpdateCachedEffectiveRuntimeStates() const;
+
+	mutable bool bIsCachedEffectiveStateDirty;
+	mutable TSet<FName> CachedEffectiveActiveDataLayerNames;
+	mutable TSet<FName> CachedEffectiveLoadedDataLayerNames;
+
+	friend class AWorldDataLayers;
+
 	FWorldDataLayersCollection WorldDataLayerCollection;
 };
 
