Add new interface function AllowsDevelopmentObjects to ITargetPlatform
This function makes a clearer distinction for what assets will be considered development vs what assets should be considered EditorOnly
A platform should be allowed to cook editor only objects but not development. The inclusion of AllowsDevelopmentObjects creates that distinction

diff --git a/Engine/Source/Developer/CookedEditor/Public/CookedEditorTargetPlatform.h b/Engine/Source/Developer/CookedEditor/Public/CookedEditorTargetPlatform.h
--- a/Engine/Source/Developer/CookedEditor/Public/CookedEditorTargetPlatform.h
+++ b/Engine/Source/Developer/CookedEditor/Public/CookedEditorTargetPlatform.h
@@ -151,6 +151,11 @@ public:
 		return true;
 	}
 
+	virtual bool AllowsDevelopmentObjects() const override
+	{
+		return false;
+	}
+
 	virtual bool AllowObject(const UObject* Obj) const override
 	{
 		// probably don't need this check, but it can't hurt
diff --git a/Engine/Source/Developer/TargetPlatform/Public/Common/TargetPlatformBase.h b/Engine/Source/Developer/TargetPlatform/Public/Common/TargetPlatformBase.h
--- a/Engine/Source/Developer/TargetPlatform/Public/Common/TargetPlatformBase.h
+++ b/Engine/Source/Developer/TargetPlatform/Public/Common/TargetPlatformBase.h
@@ -194,6 +194,11 @@ public:
 		return HasEditorOnlyData();
 	}
 
+	virtual bool AllowsDevelopmentObjects() const override
+	{
+		return AllowsEditorObjects();
+	}
+
 	virtual void GetBuildProjectSettingKeys(FString& OutSection, TArray<FString>& InBoolKeys, TArray<FString>& InIntKeys, TArray<FString>& InStringKeys) const override
 	{
 		// do nothing in the base class
diff --git a/Engine/Source/Developer/TargetPlatform/Public/Interfaces/ITargetPlatform.h b/Engine/Source/Developer/TargetPlatform/Public/Interfaces/ITargetPlatform.h
--- a/Engine/Source/Developer/TargetPlatform/Public/Interfaces/ITargetPlatform.h
+++ b/Engine/Source/Developer/TargetPlatform/Public/Interfaces/ITargetPlatform.h
@@ -327,6 +327,14 @@ public:
 	 */
 	virtual bool AllowsEditorObjects() const = 0;
 
+	/**
+	 * Checks whether this platform will allow development objects to be cooked. This is separate from AllowsEditorObjects
+	 * because cooked editors can be shipped with editor objects but still need to remove development assets.
+	 * 
+	 * @return true if this platform allows development objects to be cooked, false otherwise.
+	 */
+	virtual bool AllowsDevelopmentObjects() const = 0;
+
 	/**
 	 * Checks whether this platform is only a client (and must connect to a server to run).
 	 *
diff --git a/Engine/Source/Runtime/Engine/Classes/Engine/AssetManager.h b/Engine/Source/Runtime/Engine/Classes/Engine/AssetManager.h
--- a/Engine/Source/Runtime/Engine/Classes/Engine/AssetManager.h
+++ b/Engine/Source/Runtime/Engine/Classes/Engine/AssetManager.h
@@ -778,8 +778,8 @@ protected:
 	UPROPERTY()
 	bool bOnlyCookProductionAssets;
 
-	/** Suppresses bOnlyCookProductionAssets based on the AllowsEditorObjects() property of the TargetPlatforms being cooked. */
-	bool bTargetPlatformsAllowEditorObjects;
+	/** Suppresses bOnlyCookProductionAssets based on the AllowsDevelopmentObjects() property of the TargetPlatforms being cooked. */
+	bool bTargetPlatformsAllowDevelopmentObjects;
 
 	/** >0 if we are currently in bulk scanning mode */
 	UPROPERTY()
diff --git a/Engine/Source/Runtime/Engine/Private/AssetManager.cpp b/Engine/Source/Runtime/Engine/Private/AssetManager.cpp
--- a/Engine/Source/Runtime/Engine/Private/AssetManager.cpp
+++ b/Engine/Source/Runtime/Engine/Private/AssetManager.cpp
@@ -268,7 +268,7 @@ UAssetManager::UAssetManager()
 	bShouldUseSynchronousLoad = false;
 	bIsLoadingFromPakFiles = false;
 	bShouldAcquireMissingChunksOnLoad = false;
-	bTargetPlatformsAllowEditorObjects = false;
+	bTargetPlatformsAllowDevelopmentObjects = false;
 	NumBulkScanRequests = 0;
 	bIsManagementDatabaseCurrent = false;
 	bIsPrimaryAssetDirectoryCurrent = false;
@@ -3951,19 +3951,19 @@ void UAssetManager::ModifyCook(TArray<FName>& PackagesToCook, TArray<FName>& Pac
 	PRAGMA_ENABLE_DEPRECATION_WARNINGS;
 
 	check(TargetPlatforms.Num() > 0);
-	bTargetPlatformsAllowEditorObjects = TargetPlatforms[0]->AllowsEditorObjects();
+	bTargetPlatformsAllowDevelopmentObjects = TargetPlatforms[0]->AllowsDevelopmentObjects();
 	for (const ITargetPlatform* TargetPlatform : TargetPlatforms.Slice(1,TargetPlatforms.Num() - 1))
 	{
-		if (TargetPlatform->AllowsEditorObjects() != bTargetPlatformsAllowEditorObjects)
+		if (TargetPlatform->AllowsDevelopmentObjects() != bTargetPlatformsAllowDevelopmentObjects)
 		{
 			const ITargetPlatform* PlatformThatDoesNotAllow =
-				bTargetPlatformsAllowEditorObjects ? TargetPlatform : TargetPlatforms[0];
+				bTargetPlatformsAllowDevelopmentObjects ? TargetPlatform : TargetPlatforms[0];
 			UE_LOG(LogAssetManager, Error,
-				TEXT("Cooking platform %s and %s in a single cook is not supported, because they have different values for AllowsEditorObjects. ")
-				TEXT("This cook session will use AllowsEditorObjects = true, which will add packages to platform % s that should not be present."),
+				TEXT("Cooking platform %s and %s in a single cook is not supported, because they have different values for AllowsDevelopmentObjects. ")
+				TEXT("This cook session will use AllowsDevelopmentObjects = true, which will add packages to platform % s that should not be present."),
 				*TargetPlatforms[0]->PlatformName(), *TargetPlatform->PlatformName(),
 				*PlatformThatDoesNotAllow->PlatformName());
-			bTargetPlatformsAllowEditorObjects = true;
+			bTargetPlatformsAllowDevelopmentObjects = true;
 			break;
 		}
 	}
@@ -4188,7 +4188,7 @@ bool UAssetManager::VerifyCanCookPackage(FName PackageName, bool bLogError) cons
 		bRetVal = false;
 	}
 	else if ((CookRule == EPrimaryAssetCookRule::DevelopmentCook || CookRule == EPrimaryAssetCookRule::DevelopmentAlwaysCook)
-		&& bOnlyCookProductionAssets && !bTargetPlatformsAllowEditorObjects)
+		&& bOnlyCookProductionAssets && !bTargetPlatformsAllowDevelopmentObjects)
 	{
 		if (bLogError)
 		{
