Fix issue where names were not sent instead of handles for dynamic array properties if the name serialize property type option was enabled.

diff --git a/Engine/Source/Runtime/Engine/Private/RepLayout.cpp b/Engine/Source/Runtime/Engine/Private/RepLayout.cpp
--- a/Engine/Source/Runtime/Engine/Private/RepLayout.cpp
+++ b/Engine/Source/Runtime/Engine/Private/RepLayout.cpp
@@ -2738,7 +2738,14 @@ void FRepLayout::SendProperties_r(
 
 		if (Cmd.Type == ERepLayoutCmdType::DynamicArray)
 		{
-			WritePropertyHandle(Writer, HandleIterator.Handle, bDoChecksum);
+			if (SerializePropertyType == ESerializePropertyType::Handle)
+			{
+				WritePropertyHandle(Writer, HandleIterator.Handle, bDoChecksum);
+			}
+			else if (SerializePropertyType == ESerializePropertyType::Name)
+			{
+				WritePropertyName(Writer, Cmd.Property->GetFName(), bDoChecksum);
+			}
 
 			UE_NET_TRACE_DYNAMIC_NAME_SCOPE(Cmd.Property->GetFName(), Writer, GetTraceCollector(Writer), ENetTraceVerbosity::Trace);
 
