Fixed case where SkeletalMesh does not apply "Quality MinLOD" if it has SkeletalMeshLODSettings

diff --git a/Engine/Source/Runtime/Engine/Private/SkeletalMeshLODSettings.cpp b/Engine/Source/Runtime/Engine/Private/SkeletalMeshLODSettings.cpp
--- a/Engine/Source/Runtime/Engine/Private/SkeletalMeshLODSettings.cpp
+++ b/Engine/Source/Runtime/Engine/Private/SkeletalMeshLODSettings.cpp
@@ -15,6 +15,8 @@
 
 DEFINE_LOG_CATEGORY_STATIC(LogSkeletalMeshLODSettings, Warning, All)
 
+extern const TCHAR* GSkeletalMeshMinLodQualityLevelCVarName;
+extern const TCHAR* GSkeletalMeshMinLodQualityLevelScalabilitySection;
 
 USkeletalMeshLODSettings::USkeletalMeshLODSettings(const FObjectInitializer& ObjectInitializer)
 	: Super(ObjectInitializer)
@@ -23,6 +25,8 @@ USkeletalMeshLODSettings::USkeletalMeshLODSettings(const FObjectInitializer& Obj
 	MaxNumStreamedLODs.Default = 0;
 	// TODO: support saving some but not all optional LODs
 	MaxNumOptionalLODs.Default = 0;
+
+	MinQualityLevelLod.Init(GSkeletalMeshMinLodQualityLevelCVarName, GSkeletalMeshMinLodQualityLevelScalabilitySection);
 }
 
 const FSkeletalMeshLODGroupSettings& USkeletalMeshLODSettings::GetSettingsForLODLevel(const int32 LODIndex) const
