Defer computation like streaming source to streaming cell angle contribution to when it is really needed.
Most of the intersecting cells are already in the correct target state and don't require the information needed to prioritize cell requests.

diff --git a/Engine/Source/Runtime/Engine/Private/WorldPartition/WorldPartitionRuntimeSpatialHash.cpp b/Engine/Source/Runtime/Engine/Private/WorldPartition/WorldPartitionRuntimeSpatialHash.cpp
--- a/Engine/Source/Runtime/Engine/Private/WorldPartition/WorldPartitionRuntimeSpatialHash.cpp
+++ b/Engine/Source/Runtime/Engine/Private/WorldPartition/WorldPartitionRuntimeSpatialHash.cpp
@@ -1783,9 +1783,9 @@ EWorldPartitionStreamingPerformance UWorldPartitionRuntimeSpatialHash::GetStream
 	{
 		const float LoadingRange = StreamingGrid->LoadingRange;
 
-		if (StreamingCell->CachedIsBlockingSource)
+		if (StreamingCell->IsBlockingSource())
 		{
-			const double Distance = FMath::Sqrt(StreamingCell->CachedMinSquareDistanceToBlockingSource) - ((double)StreamingGrid->GetCellSize(StreamingCell->Level) / 2);
+			const double Distance = FMath::Sqrt(StreamingCell->GetMinSquareDistanceToBlockingSource()) - ((double)StreamingGrid->GetCellSize(StreamingCell->Level) / 2);
 
 			const double Ratio = Distance / LoadingRange;
 
diff --git a/Engine/Source/Runtime/Engine/Private/WorldPartition/WorldPartitionRuntimeSpatialHashCell.cpp b/Engine/Source/Runtime/Engine/Private/WorldPartition/WorldPartitionRuntimeSpatialHashCell.cpp
--- a/Engine/Source/Runtime/Engine/Private/WorldPartition/WorldPartitionRuntimeSpatialHashCell.cpp
+++ b/Engine/Source/Runtime/Engine/Private/WorldPartition/WorldPartitionRuntimeSpatialHashCell.cpp
@@ -55,15 +55,13 @@ void UWorldPartitionRuntimeSpatialHashCell::ResetStreamingSourceInfo() const
 	CachedIsBlockingSource = false;
 	CachedMinSquareDistanceToBlockingSource = MAX_dbl;
 	CachedMinSquareDistanceToSource = MAX_dbl;
-	CachedSourceModulatedDistances.Reset();
+	CachedSourceSquaredDistances.Reset();
+	CachedInstersectingShapes.Reset();
 }
 
-void UWorldPartitionRuntimeSpatialHashCell::AppendStreamingSourceInfo(const FWorldPartitionStreamingSource& Source, const FSphericalSector& SourceShape) const
+float UWorldPartitionRuntimeSpatialHashCell::ComputeSourceToCellAngleFactor(const FSphericalSector& SourceShape) const
 {
-	Super::AppendStreamingSourceInfo(Source, SourceShape);
-
 	float AngleContribution = FMath::Clamp(GRuntimeSpatialHashCellToSourceAngleContributionToCellImportance, 0.f, 1.f);
-	const double SquareDistance = FVector::DistSquared2D(SourceShape.GetCenter(), Position);
 	float AngleFactor = 1.f;
 	if (!FMath::IsNearlyZero(AngleContribution))
 	{
@@ -94,9 +92,17 @@ void UWorldPartitionRuntimeSpatialHashCell::AppendStreamingSourceInfo(const FWor
 			}
 			Angle = FMath::Abs(FMath::Acos(MaxDot) / UE_PI);
 		}
-		const float NormalizedAngle = FMath::Clamp(Angle, UE_PI/180.f, 1.f);
+		const float NormalizedAngle = FMath::Clamp(Angle, UE_PI / 180.f, 1.f);
 		AngleFactor = FMath::Pow(NormalizedAngle, AngleContribution);
 	}
+	return AngleFactor;
+}
+
+void UWorldPartitionRuntimeSpatialHashCell::AppendStreamingSourceInfo(const FWorldPartitionStreamingSource& Source, const FSphericalSector& SourceShape) const
+{
+	Super::AppendStreamingSourceInfo(Source, SourceShape);
+
+	const double SquareDistance = FVector::DistSquared2D(SourceShape.GetCenter(), Position);
 
 	// Only consider blocking sources
 	if (Source.bBlockOnSlowLoading)
@@ -104,13 +110,8 @@ void UWorldPartitionRuntimeSpatialHashCell::AppendStreamingSourceInfo(const FWor
 		CachedIsBlockingSource = true;
 		CachedMinSquareDistanceToBlockingSource = FMath::Min(SquareDistance, CachedMinSquareDistanceToBlockingSource);
 	}
-
-	CachedMinSquareDistanceToSource = FMath::Min(SquareDistance, CachedMinSquareDistanceToSource);
-
-	// Modulate distance to cell by angle relative to source forward vector (to prioritize cells in front)
-	const double SquareAngleFactor = AngleFactor * AngleFactor;
-	const double ModulatedSquareDistance = SquareDistance * SquareAngleFactor;
-	CachedSourceModulatedDistances.Add(ModulatedSquareDistance);
+	CachedSourceSquaredDistances.Add(SquareDistance);
+	CachedInstersectingShapes.Add(SourceShape);
 }
 
 FBox UWorldPartitionRuntimeSpatialHashCell::GetCellBounds() const
@@ -126,8 +127,11 @@ void UWorldPartitionRuntimeSpatialHashCell::MergeStreamingSourceInfo() const
 {
 	Super::MergeStreamingSourceInfo();
 
-	int32 Count = CachedSourceModulatedDistances.Num();
+	int32 Count = CachedSourceSquaredDistances.Num();
 	check(Count == CachedSourcePriorityWeights.Num());
+	check(Count == CachedInstersectingShapes.Num());
+
+	CachedMinSquareDistanceToSource = Count ? FMath::Min(CachedSourceSquaredDistances) : MAX_dbl;
 
 	if (Count)
 	{
@@ -136,24 +140,28 @@ void UWorldPartitionRuntimeSpatialHashCell::MergeStreamingSourceInfo() const
 		{
 			TotalSourcePriorityWeight += CachedSourcePriorityWeights[i];
 		}
-		int32 HighestPrioMinDistIndex = 0;
-		double WeightedModulatedDistance = 0.f;
+
+		double HighestPrioWeight = 0;
+		double HighestPrioMinModulatedDistance = 0;
+		double WeightedModulatedDistance = 0;
 		for (int32 i = 0; i < Count; ++i)
 		{
-			WeightedModulatedDistance += CachedSourceModulatedDistances[i] * CachedSourcePriorityWeights[i] / TotalSourcePriorityWeight;
-			
+			const float AngleFactor = ComputeSourceToCellAngleFactor(CachedInstersectingShapes[i]);
+			const double CurrentModulatedDistance = CachedSourceSquaredDistances[i] * AngleFactor * AngleFactor;
+			const double CurrentWeight = CachedSourcePriorityWeights[i];
+			WeightedModulatedDistance += CurrentModulatedDistance * CurrentWeight / TotalSourcePriorityWeight;
 			// Find highest priority source with the minimum modulated distance
-			if ((i != 0) &&
-				(CachedSourceModulatedDistances[i] < CachedSourceModulatedDistances[HighestPrioMinDistIndex]) &&
-				(CachedSourcePriorityWeights[i] >= CachedSourcePriorityWeights[HighestPrioMinDistIndex]))
+			if ((i == 0) || ((CurrentModulatedDistance < HighestPrioMinModulatedDistance) && (CurrentWeight >= HighestPrioWeight)))
 			{
-				HighestPrioMinDistIndex = i;
+				HighestPrioMinModulatedDistance = CurrentModulatedDistance;
+				HighestPrioWeight = CurrentWeight;
 			}
 		}
+
 		// Sorting distance is the minimum between these:
 		// - the highest priority source with the minimum modulated distance 
 		// - the weighted modulated distance
-		CachedSourceSortingDistance = FMath::Min(CachedSourceModulatedDistances[HighestPrioMinDistIndex], WeightedModulatedDistance);
+		CachedSourceSortingDistance = FMath::Min(HighestPrioMinModulatedDistance, WeightedModulatedDistance);
 	}
 }
 
diff --git a/Engine/Source/Runtime/Engine/Public/WorldPartition/WorldPartitionRuntimeSpatialHashCell.h b/Engine/Source/Runtime/Engine/Public/WorldPartition/WorldPartitionRuntimeSpatialHashCell.h
--- a/Engine/Source/Runtime/Engine/Public/WorldPartition/WorldPartitionRuntimeSpatialHashCell.h
+++ b/Engine/Source/Runtime/Engine/Public/WorldPartition/WorldPartitionRuntimeSpatialHashCell.h
@@ -15,14 +15,20 @@ class UWorldPartitionRuntimeSpatialHashCell : public UWorldPartitionRuntimeCell
 #if WITH_EDITOR
 	virtual void PostDuplicate(bool bDuplicateForPIE) override;
 #endif
+
+	//~Begin UWorldPartitionRuntimeCell
 	virtual void ResetStreamingSourceInfo() const override;
 	virtual void AppendStreamingSourceInfo(const FWorldPartitionStreamingSource& Source, const FSphericalSector& SourceShape) const override;
 	virtual void MergeStreamingSourceInfo() const override;
+	virtual int32 SortCompare(const UWorldPartitionRuntimeCell* InOther) const override;
+	//~End UWorldPartitionRuntimeCell
+
 	//~ Begin IWorldPartitionCell Interface
 	virtual FBox GetCellBounds() const override;
 	//~ End IWorldPartitionCell Interface
 
-	virtual int32 SortCompare(const UWorldPartitionRuntimeCell* InOther) const override;
+	bool IsBlockingSource() const { return CachedIsBlockingSource; }
+	double GetMinSquareDistanceToBlockingSource() const { return CachedMinSquareDistanceToBlockingSource; }
 
 	UPROPERTY()
 	FVector Position;
@@ -33,13 +39,21 @@ class UWorldPartitionRuntimeSpatialHashCell : public UWorldPartitionRuntimeCell
 	UPROPERTY()
 	int32 Level;
 		
+#if WITH_EDITORONLY_DATA
+	UPROPERTY()
+	TObjectPtr<UActorContainer> UnsavedActorsContainer;
+#endif
+
+private:
+	float ComputeSourceToCellAngleFactor(const FSphericalSector& SourceShape) const;
+
 	// Used to determine if cell was requested by blocking source
 	mutable bool CachedIsBlockingSource;
 
-	// Represents the square distance from cell to the closest blocking streaming source
+	// Square distance from the cell to the closest blocking streaming source
 	mutable double CachedMinSquareDistanceToBlockingSource;
 
-	// Represents the square distance from cell to the closest streaming source
+	// Square distance from the cell to the closest streaming source
 	mutable double CachedMinSquareDistanceToSource;
 
 	// Modulated distance to the different streaming sources used to sort relative priority amongst streaming cells
@@ -50,10 +64,9 @@ class UWorldPartitionRuntimeSpatialHashCell : public UWorldPartitionRuntimeCell
 	// - The angle between the cell and each source orientation
 	mutable double CachedSourceSortingDistance;
 
-	mutable TArray<double> CachedSourceModulatedDistances;
+	// Square distance from the cell to  the intersecting streaming sources
+	mutable TArray<double> CachedSourceSquaredDistances;
 
-#if WITH_EDITORONLY_DATA
-	UPROPERTY()
-	TObjectPtr<UActorContainer> UnsavedActorsContainer;
-#endif
+	// Intersecting streaming source shapes
+	mutable TArray<FSphericalSector> CachedInstersectingShapes;
 };
\ No newline at end of file
