Optimization : maintain world partition activated cells that are still pending to add to world instead of rebuilding it at every update.

diff --git a/Engine/Source/Runtime/Engine/Private/WorldPartition/HLOD/HLODSubsystem.cpp b/Engine/Source/Runtime/Engine/Private/WorldPartition/HLOD/HLODSubsystem.cpp
--- a/Engine/Source/Runtime/Engine/Private/WorldPartition/HLOD/HLODSubsystem.cpp
+++ b/Engine/Source/Runtime/Engine/Private/WorldPartition/HLOD/HLODSubsystem.cpp
@@ -295,6 +295,12 @@ const TArray<AWorldPartitionHLOD*>& UHLODSubsystem::GetHLODActorsForCell(const U
 
 void UHLODSubsystem::OnCellShown(const UWorldPartitionRuntimeCell* InCell)
 {
+	// @todo_ow ContentBundles do not support Hlods
+	if (InCell->GetContentBundleID().IsValid())
+	{
+		return;
+	}
+
 	const UWorldPartition* WorldPartition = InCell->GetCellOwner()->GetOuterWorld()->GetWorldPartition();
 	check(WorldPartition && WorldPartition->IsStreamingEnabled());
 
@@ -323,6 +329,12 @@ void UHLODSubsystem::OnCellShown(const UWorldPartitionRuntimeCell* InCell)
 
 void UHLODSubsystem::OnCellHidden(const UWorldPartitionRuntimeCell* InCell)
 {
+	// @todo_ow ContentBundles do not support Hlods
+	if (InCell->GetContentBundleID().IsValid())
+	{
+		return;
+	}
+
 	const UWorldPartition* WorldPartition = InCell->GetCellOwner()->GetOuterWorld()->GetWorldPartition();
 	check(WorldPartition && WorldPartition->IsStreamingEnabled());
 
diff --git a/Engine/Source/Runtime/Engine/Private/WorldPartition/WorldPartition.cpp b/Engine/Source/Runtime/Engine/Private/WorldPartition/WorldPartition.cpp
--- a/Engine/Source/Runtime/Engine/Private/WorldPartition/WorldPartition.cpp
+++ b/Engine/Source/Runtime/Engine/Private/WorldPartition/WorldPartition.cpp
@@ -14,6 +14,7 @@
 #include "WorldPartition/WorldPartitionSubsystem.h"
 #include "WorldPartition/WorldPartitionStreamingSource.h"
 #include "WorldPartition/WorldPartitionReplay.h"
+#include "WorldPartition/HLOD/HLODSubsystem.h"
 #include "Algo/Accumulate.h"
 #include "Algo/Transform.h"
 #include "Engine/World.h"
@@ -1376,6 +1377,30 @@ bool UWorldPartition::IsStreamingCompleted(EWorldPartitionRuntimeCellState Query
 	return true;
 }
 
+void UWorldPartition::OnCellShown(const UWorldPartitionRuntimeCell* InCell)
+{
+	if (GetWorld()->IsGameWorld())
+	{
+		if (IsStreamingEnabled())
+		{
+			GetWorld()->GetSubsystem<UHLODSubsystem>()->OnCellShown(InCell);
+		}
+		StreamingPolicy->OnCellShown(InCell);
+	}
+}
+
+void UWorldPartition::OnCellHidden(const UWorldPartitionRuntimeCell* InCell)
+{
+	if (GetWorld()->IsGameWorld())
+	{
+		if (IsStreamingEnabled())
+		{
+			GetWorld()->GetSubsystem<UHLODSubsystem>()->OnCellHidden(InCell);
+		}
+		StreamingPolicy->OnCellHidden(InCell);
+	}
+}
+
 bool UWorldPartition::CanDebugDraw() const
 {
 #if WITH_EDITOR
diff --git a/Engine/Source/Runtime/Engine/Private/WorldPartition/WorldPartitionRuntimeLevelStreamingCell.cpp b/Engine/Source/Runtime/Engine/Private/WorldPartition/WorldPartitionRuntimeLevelStreamingCell.cpp
--- a/Engine/Source/Runtime/Engine/Private/WorldPartition/WorldPartitionRuntimeLevelStreamingCell.cpp
+++ b/Engine/Source/Runtime/Engine/Private/WorldPartition/WorldPartitionRuntimeLevelStreamingCell.cpp
@@ -335,9 +335,7 @@ UWorldPartitionLevelStreamingDynamic* UWorldPartitionRuntimeLevelStreamingCell::
 	}
 #endif
 
-	// @todo_ow ContentBundles do not support Hlods and events below are forwarding
-	// Events to the HLodSubsystem which assumes knowledge of all cells (not true with plugins)
-	if (LevelStreaming  && !GetContentBundleID().IsValid())
+	if (LevelStreaming)
 	{
 		LevelStreaming->OnLevelShown.AddUniqueDynamic(this, &UWorldPartitionRuntimeLevelStreamingCell::OnLevelShown);
 		LevelStreaming->OnLevelHidden.AddUniqueDynamic(this, &UWorldPartitionRuntimeLevelStreamingCell::OnLevelHidden);
@@ -446,18 +444,18 @@ void UWorldPartitionRuntimeLevelStreamingCell::Deactivate() const
 void UWorldPartitionRuntimeLevelStreamingCell::OnLevelShown()
 {
 	check(LevelStreaming);
-	if (const UWorldPartition* WorldPartition = LevelStreaming->GetWorld()->GetWorldPartition(); WorldPartition && WorldPartition->IsStreamingEnabled())
+	if (UWorldPartition* WorldPartition = LevelStreaming->GetWorld()->GetWorldPartition())
 	{
-		LevelStreaming->GetWorld()->GetSubsystem<UHLODSubsystem>()->OnCellShown(this);
+		WorldPartition->OnCellShown(this);
 	}
 }
 
 void UWorldPartitionRuntimeLevelStreamingCell::OnLevelHidden()
 {
 	check(LevelStreaming);
-	if (const UWorldPartition* WorldPartition = LevelStreaming->GetWorld()->GetWorldPartition(); WorldPartition && WorldPartition->IsStreamingEnabled())
+	if (UWorldPartition* WorldPartition = LevelStreaming->GetWorld()->GetWorldPartition())
 	{
-		LevelStreaming->GetWorld()->GetSubsystem<UHLODSubsystem>()->OnCellHidden(this);
+		WorldPartition->OnCellHidden(this);
 	}
 }
 
diff --git a/Engine/Source/Runtime/Engine/Private/WorldPartition/WorldPartitionStreamingPolicy.cpp b/Engine/Source/Runtime/Engine/Private/WorldPartition/WorldPartitionStreamingPolicy.cpp
--- a/Engine/Source/Runtime/Engine/Private/WorldPartition/WorldPartitionStreamingPolicy.cpp
+++ b/Engine/Source/Runtime/Engine/Private/WorldPartition/WorldPartitionStreamingPolicy.cpp
@@ -323,7 +323,7 @@ void UWorldPartitionStreamingPolicy::UpdateStreamingState()
 		else if (!bCanDeactivateOrUnloadCells)
 		{
 			// When server streaming-out is disabled, revisit existing loaded/activated cells and add them in the proper FrameLoadCells/FrameActivateCells
-			for (const UWorldPartitionRuntimeCell* Cell : ActivatedCells)
+			for (const UWorldPartitionRuntimeCell* Cell : ActivatedCells.GetCells())
 			{
 				AddServerFrameCells(Cell);
 			}
@@ -411,7 +411,7 @@ void UWorldPartitionStreamingPolicy::UpdateStreamingState()
 			}
 		};
 
-		BuildCellsToUnload(ActivatedCells);
+		BuildCellsToUnload(ActivatedCells.GetCells());
 		BuildCellsToUnload(LoadedCells);
 	}
 
@@ -447,14 +447,7 @@ void UWorldPartitionStreamingPolicy::UpdateStreamingState()
 	// Sort cells and update streaming priority 
 	{
 		TRACE_CPUPROFILER_EVENT_SCOPE(UWorldPartitionStreamingPolicy::SortCellsAndUpdateStreamingPriority);
-		SortedAddToWorldCells.Reset();
-		for (const UWorldPartitionRuntimeCell* ActivatedCell : ActivatedCells)
-		{
-			if (!ActivatedCell->IsAddedToWorld() && !ActivatedCell->IsAlwaysLoaded())
-			{
-				SortedAddToWorldCells.Add(ActivatedCell);
-			}
-		}
+		SortedAddToWorldCells = ActivatedCells.GetPendingAddToWorldCells().Array();
 		SortStreamingCellsByImportance(SortedAddToWorldCells);
 
 		// Update level streaming priority so that UWorld::UpdateLevelStreaming will naturally process the levels in the correct order
@@ -833,6 +826,47 @@ void UWorldPartitionStreamingPolicy::DrawRuntimeHash3D()
 	}
 }
 
+void UWorldPartitionStreamingPolicy::OnCellShown(const UWorldPartitionRuntimeCell* InCell)
+{
+	ActivatedCells.OnAddedToWorld(InCell);
+}
+
+void UWorldPartitionStreamingPolicy::OnCellHidden(const UWorldPartitionRuntimeCell* InCell)
+{
+	ActivatedCells.OnRemovedFromWorld(InCell);
+}
+
+void UWorldPartitionStreamingPolicy::FActivatedCells::Add(const UWorldPartitionRuntimeCell* InCell)
+{
+	Cells.Add(InCell);
+	if (!InCell->IsAlwaysLoaded())
+	{
+		PendingAddToWorldCells.Add(InCell);
+	}
+}
+
+void UWorldPartitionStreamingPolicy::FActivatedCells::Remove(const UWorldPartitionRuntimeCell* InCell)
+{
+	Cells.Remove(InCell);
+	PendingAddToWorldCells.Remove(InCell);
+}
+
+void UWorldPartitionStreamingPolicy::FActivatedCells::OnAddedToWorld(const UWorldPartitionRuntimeCell* InCell)
+{
+	PendingAddToWorldCells.Remove(InCell);
+}
+
+void UWorldPartitionStreamingPolicy::FActivatedCells::OnRemovedFromWorld(const UWorldPartitionRuntimeCell* InCell)
+{
+	if (Cells.Contains(InCell))
+	{
+		if (!InCell->IsAlwaysLoaded())
+		{
+			PendingAddToWorldCells.Add(InCell);
+		}
+	}
+}
+
 /*
  * FStreamingSourceVelocity Implementation
  */
diff --git a/Engine/Source/Runtime/Engine/Public/WorldPartition/WorldPartition.h b/Engine/Source/Runtime/Engine/Public/WorldPartition/WorldPartition.h
--- a/Engine/Source/Runtime/Engine/Public/WorldPartition/WorldPartition.h
+++ b/Engine/Source/Runtime/Engine/Public/WorldPartition/WorldPartition.h
@@ -249,6 +249,9 @@ public:
 	void DrawRuntimeCellsDetails(UCanvas* Canvas, FVector2D& Offset);
 	void DrawStreamingStatusLegend(UCanvas* Canvas, FVector2D& Offset);
 
+	void OnCellShown(const UWorldPartitionRuntimeCell* InCell);
+	void OnCellHidden(const UWorldPartitionRuntimeCell* InCell);
+
 	EWorldPartitionStreamingPerformance GetStreamingPerformance() const;
 
 #if WITH_EDITORONLY_DATA
diff --git a/Engine/Source/Runtime/Engine/Public/WorldPartition/WorldPartitionStreamingPolicy.h b/Engine/Source/Runtime/Engine/Public/WorldPartition/WorldPartitionStreamingPolicy.h
--- a/Engine/Source/Runtime/Engine/Public/WorldPartition/WorldPartitionStreamingPolicy.h
+++ b/Engine/Source/Runtime/Engine/Public/WorldPartition/WorldPartitionStreamingPolicy.h
@@ -55,6 +55,9 @@ public:
 	virtual bool IsStreamingCompleted(const TArray<FWorldPartitionStreamingSource>* InStreamingSources) const;
 	virtual bool IsStreamingCompleted(EWorldPartitionRuntimeCellState QueryState, const TArray<FWorldPartitionStreamingQuerySource>& QuerySources, bool bExactState = true) const;
 
+	virtual void OnCellShown(const UWorldPartitionRuntimeCell* InCell);
+	virtual void OnCellHidden(const UWorldPartitionRuntimeCell* InCell);
+
 #if WITH_EDITOR
 	virtual TSubclassOf<class UWorldPartitionRuntimeCell> GetRuntimeCellClass() const PURE_VIRTUAL(UWorldPartitionStreamingPolicy::GetRuntimeCellClass, return UWorldPartitionRuntimeCell::StaticClass(); );
 
@@ -86,7 +89,26 @@ protected:
 
 	const UWorldPartition* WorldPartition;
 	TSet<const UWorldPartitionRuntimeCell*> LoadedCells;
-	TSet<const UWorldPartitionRuntimeCell*> ActivatedCells;
+
+	struct FActivatedCells
+	{
+		void Add(const UWorldPartitionRuntimeCell* InCell);
+		void Remove(const UWorldPartitionRuntimeCell* InCell);
+		bool Contains(const UWorldPartitionRuntimeCell* InCell) const { return Cells.Contains(InCell); }
+		void OnAddedToWorld(const UWorldPartitionRuntimeCell* InCell);
+		void OnRemovedFromWorld(const UWorldPartitionRuntimeCell* InCell);
+
+		const TSet<const UWorldPartitionRuntimeCell*>& GetCells() const { return Cells; }
+		const TSet<const UWorldPartitionRuntimeCell*>& GetPendingAddToWorldCells() const { return PendingAddToWorldCells; }
+
+	private:
+
+		TSet<const UWorldPartitionRuntimeCell*> Cells;
+		TSet<const UWorldPartitionRuntimeCell*> PendingAddToWorldCells;
+	};
+
+	FActivatedCells ActivatedCells;
+	mutable TArray<const UWorldPartitionRuntimeCell*> SortedAddToWorldCells;
 
 	// Streaming Sources
 	TArray<FWorldPartitionStreamingSource> StreamingSources;
@@ -97,7 +119,6 @@ protected:
 	
 	bool bCriticalPerformanceRequestedBlockTillOnWorld;
 	int32 CriticalPerformanceBlockTillLevelStreamingCompletedEpoch;
-	mutable TArray<const UWorldPartitionRuntimeCell*> SortedAddToWorldCells;
 
 	int32 DataLayersStatesServerEpoch;
 	int32 ContentBundleServerEpoch;
