[shaders]
- fix collisions in debug dumps for material shaders - name the folder based on the material instance name (if building a material instance) instead of just the base material name, also hash the asset path and static parameters to differentiate materials with the same name and different path and different static permutations
- add a debug dump option for writing single-job ShaderCompileWorker input files to allow debugging SCW in normal operation mode (including a separate txt file containing commandline args for SCW to do so)
- use path concatenation operator instead of + in shadercompileworker when constructing input/output file paths

diff --git a/Engine/Source/Programs/ShaderCompileWorker/Private/ShaderCompileWorker.cpp b/Engine/Source/Programs/ShaderCompileWorker/Private/ShaderCompileWorker.cpp
--- a/Engine/Source/Programs/ShaderCompileWorker/Private/ShaderCompileWorker.cpp
+++ b/Engine/Source/Programs/ShaderCompileWorker/Private/ShaderCompileWorker.cpp
@@ -307,8 +307,8 @@ public:
 	,	WorkingDirectory(InWorkingDirectory)
 	,	InputFilename(InInputFilename)
 	,	OutputFilename(InOutputFilename)
-	,	InputFilePath(FString(InWorkingDirectory) + InInputFilename)
-	,	OutputFilePath(FString(InWorkingDirectory) + InOutputFilename)
+	,	InputFilePath(FString(InWorkingDirectory) / InInputFilename)
+	,	OutputFilePath(FString(InWorkingDirectory) / InOutputFilename)
 	,	FormatVersionMap(InFormatVersionMap)
 	{
 		TArray<FString> Tokens, Switches;
diff --git a/Engine/Source/Runtime/Engine/Private/Materials/MaterialShader.cpp b/Engine/Source/Runtime/Engine/Private/Materials/MaterialShader.cpp
--- a/Engine/Source/Runtime/Engine/Private/Materials/MaterialShader.cpp
+++ b/Engine/Source/Runtime/Engine/Private/Materials/MaterialShader.cpp
@@ -1148,7 +1148,7 @@ static void PrepareMaterialShaderCompileJob(EShaderPlatform Platform,
 
 	// Compile the shader environment passed in with the shader type's source code.
 	::GlobalBeginCompileShader(
-		Material->GetFriendlyName() / LexToString(Material->GetQualityLevel()),
+		Material->GetUniqueAssetName(Platform) / LexToString(Material->GetQualityLevel()),
 		nullptr,
 		ShaderType,
 		ShaderPipeline,
diff --git a/Engine/Source/Runtime/Engine/Private/Materials/MaterialShared.cpp b/Engine/Source/Runtime/Engine/Private/Materials/MaterialShared.cpp
--- a/Engine/Source/Runtime/Engine/Private/Materials/MaterialShared.cpp
+++ b/Engine/Source/Runtime/Engine/Private/Materials/MaterialShared.cpp
@@ -2678,6 +2678,21 @@ bool FMaterial::CacheShaders(const FMaterialShaderMapId& ShaderMapId, EShaderPla
 #endif
 }
 
+#if WITH_EDITOR
+FString FMaterial::GetUniqueAssetName(EShaderPlatform Platform) const
+{
+	TArray<uint8> Buf;
+	FString PathStr = GetAssetPath().GetPlainNameString();
+	FMemoryWriter BufWriter(Buf);
+	BufWriter << PathStr;
+	FStaticParameterSet StaticParams;
+	GetStaticParameterSet(Platform, StaticParams);
+	FStaticParameterSet::StaticStruct()->SerializeBin(BufWriter, &StaticParams);
+	uint64 Hash = CityHash64((const char*)Buf.GetData(), Buf.Num());
+	return FString::Printf(TEXT("%s_%llu"), *GetAssetName(), Hash);
+}
+#endif
+
 /**
  * Helper task used to release the strong object reference to the material interface on the game thread
  * The release has to happen on the gamethread and the material interface can't be GCd while the PSO
diff --git a/Engine/Source/Runtime/Engine/Private/Materials/MeshMaterialShader.cpp b/Engine/Source/Runtime/Engine/Private/Materials/MeshMaterialShader.cpp
--- a/Engine/Source/Runtime/Engine/Private/Materials/MeshMaterialShader.cpp
+++ b/Engine/Source/Runtime/Engine/Private/Materials/MeshMaterialShader.cpp
@@ -64,7 +64,7 @@ static void PrepareMeshMaterialShaderCompileJob(EShaderPlatform Platform,
 
 	// Compile the shader environment passed in with the shader type's source code.
 	::GlobalBeginCompileShader(
-		Material->GetFriendlyName() / LexToString(Material->GetQualityLevel()),
+		Material->GetUniqueAssetName(Platform) / LexToString(Material->GetQualityLevel()),
 		VertexFactoryType,
 		ShaderType,
 		ShaderPipeline,
diff --git a/Engine/Source/Runtime/Engine/Private/ShaderCompiler/ShaderCompiler.cpp b/Engine/Source/Runtime/Engine/Private/ShaderCompiler/ShaderCompiler.cpp
--- a/Engine/Source/Runtime/Engine/Private/ShaderCompiler/ShaderCompiler.cpp
+++ b/Engine/Source/Runtime/Engine/Private/ShaderCompiler/ShaderCompiler.cpp
@@ -131,6 +131,12 @@ static TAutoConsoleVariable<bool> CVarJobCacheDDCPolicy(
 	TEXT("Whether to cache shaders in the job cache to your local machine or remotely to the network.\n"),
 	ECVF_ReadOnly);
 
+static TAutoConsoleVariable<bool> CVarDebugDumpWorkerInputs(
+	TEXT("r.ShaderCompiler.DebugDumpWorkerInputs"),
+	false,
+	TEXT("If true, worker input files will be saved for each individual compile job alongside other debug data (note that r.DumpShaderDebugInfo must also be enabled for this to function)"),
+	ECVF_ReadOnly);
+
 static bool IsShaderJobCacheDDCRemotePolicyEnabled()
 {
 	return CVarJobCacheDDCPolicy.GetValueOnAnyThread();
@@ -1344,8 +1350,7 @@ static void SplitJobsByType(const TArray<FShaderCommonCompileJobPtr>& QueuedJobs
 	}
 }
 
-// Serialize Queued Job information
-bool FShaderCompileUtilities::DoWriteTasks(const TArray<FShaderCommonCompileJobPtr>& QueuedJobs, FArchive& InTransferFile, IDistributedBuildController* BuildDistributionController, bool bUseRelativePaths, bool bCompressTaskFile)
+bool DoWriteTasksInner(const TArray<FShaderCommonCompileJobPtr>& QueuedJobs, FArchive& InTransferFile, IDistributedBuildController* BuildDistributionController, bool bUseRelativePaths, bool bCompressTaskFile)
 {
 	int32 InputVersion = ShaderCompileWorkerInputVersion;
 	InTransferFile << InputVersion;
@@ -1579,6 +1584,66 @@ bool FShaderCompileUtilities::DoWriteTasks(const TArray<FShaderCommonCompileJobP
 	return InTransferFile.Close();
 }
 
+const TCHAR* DebugWorkerInputFileName = TEXT("DebugSCW.in");
+const TCHAR* DebugWorkerOutputFileName = TEXT("DebugSCW.out");
+
+FString CreateShaderCompilerWorkerDebugCommandLine(const FShaderCompileJob* Job)
+{
+	return FString::Printf(TEXT("\"%s\" 0 \"%s\" %s %s -TimeToLive 0.5f"), // 0 is parent PID, pass short TTL to prevent SCW from hanging around waiting for input after processing the single job
+		*Job->Input.DumpDebugInfoPath, // working directory for SCW
+		*Job->Input.DebugGroupName, // console window title
+		DebugWorkerInputFileName,
+		DebugWorkerOutputFileName);
+}
+
+// Serialize Queued Job information
+bool FShaderCompileUtilities::DoWriteTasks(const TArray<FShaderCommonCompileJobPtr>& QueuedJobs, FArchive& InTransferFile, IDistributedBuildController* BuildDistributionController, bool bUseRelativePaths, bool bCompressTaskFile)
+{
+	if (CVarDebugDumpWorkerInputs.GetValueOnAnyThread())
+	{
+		for (const FShaderCommonCompileJobPtr& CommonJob : QueuedJobs)
+		{
+			auto WriteSingleJobFile = [](FShaderCompileJob* Job)
+			{
+				if (Job->Input.DumpDebugInfoEnabled())
+				{
+					TArray<FShaderCommonCompileJobPtr> SingleJobArray;
+					SingleJobArray.Add(Job);
+
+					FString DebugWorkerInputFilePath = Job->Input.DumpDebugInfoPath / DebugWorkerInputFileName;
+					FArchive* DebugWorkerInputFileWriter = IFileManager::Get().CreateFileWriter(*DebugWorkerInputFilePath, FILEWRITE_NoFail);
+					DoWriteTasksInner(
+						SingleJobArray, 
+						*DebugWorkerInputFileWriter, 
+						nullptr,	// Don't pass a IDistributedBuildController, this is only used for conversion to relative paths which we do not want for debug files
+						false,		// As above, use absolute paths not relative
+						true);		// Always compress the debug files; they are rather large so this saves some disk space
+					DebugWorkerInputFileWriter->Close();
+					delete DebugWorkerInputFileWriter;
+
+					FFileHelper::SaveStringToFile(
+						CreateShaderCompilerWorkerDebugCommandLine(Job), 
+						*(Job->Input.DumpDebugInfoPath / TEXT("DebugCompileArgs.txt")));
+				}
+			};
+			FShaderPipelineCompileJob* PipelineJob = CommonJob->GetShaderPipelineJob();
+			if (PipelineJob)
+			{
+				for (FShaderCompileJob* StageJob : PipelineJob->StageJobs)
+				{
+					WriteSingleJobFile(StageJob);
+				}
+			}
+			else
+			{
+				WriteSingleJobFile(CommonJob->GetSingleShaderJob());
+			}
+		}
+	}
+
+	return DoWriteTasksInner(QueuedJobs, InTransferFile, BuildDistributionController, bUseRelativePaths, bCompressTaskFile);
+}
+
 static void ProcessErrors(const FShaderCompileJob& CurrentJob, TArray<FString>& UniqueErrors, FString& ErrorString, int32& NumDumpedShaderSources)
 {
 	bool bReportedDebugInfo = false;
diff --git a/Engine/Source/Runtime/Engine/Public/MaterialShared.h b/Engine/Source/Runtime/Engine/Public/MaterialShared.h
--- a/Engine/Source/Runtime/Engine/Public/MaterialShared.h
+++ b/Engine/Source/Runtime/Engine/Public/MaterialShared.h
@@ -1860,6 +1860,15 @@ public:
 	virtual FString GetFriendlyName() const = 0;
 	/** Similar to GetFriendlyName, but but avoids historical behavior of the former, returning the exact asset name for material instances instead of just the material. */
 	virtual FString GetAssetName() const { return GetFriendlyName(); }
+
+#if WITH_EDITOR
+	/** Similar to GetAssetName, but appends a hash of the asset path and static switch params (if editoronly data is available)
+	to differentiate materials that would collide when considering name alone. */
+	ENGINE_API FString GetUniqueAssetName(EShaderPlatform Platform) const;
+#else
+	/** GetUniqueAssetName just returns GetAssetName in non-editor builds */
+	FString GetUniqueAssetName() const { return GetAssetName(); }
+#endif
 	virtual bool HasVertexPositionOffsetConnected() const { return false; }
 	virtual bool HasPixelDepthOffsetConnected() const { return false; }
 	virtual uint32 GetMaterialDecalResponse() const { return 0; }
