Fix encountered instances of TObjectPtr<T> used with T that is not derived from UObject.  A subsequent change will be disallowing this at compile time.

diff --git a/Engine/Plugins/Editor/ObjectMixer/ObjectMixer/Source/ObjectMixer/Public/Views/List/SObjectMixerEditorList.h b/Engine/Plugins/Editor/ObjectMixer/ObjectMixer/Source/ObjectMixer/Public/Views/List/SObjectMixerEditorList.h
--- a/Engine/Plugins/Editor/ObjectMixer/ObjectMixer/Source/ObjectMixer/Public/Views/List/SObjectMixerEditorList.h
+++ b/Engine/Plugins/Editor/ObjectMixer/ObjectMixer/Source/ObjectMixer/Public/Views/List/SObjectMixerEditorList.h
@@ -33,7 +33,7 @@ struct FListViewColumnInfo
 {
 	GENERATED_BODY()
 	
-	TObjectPtr<FProperty> PropertyRef = nullptr;
+	FProperty* PropertyRef = nullptr;
 	
 	UPROPERTY()
 	FName PropertyName = NAME_None;
diff --git a/Engine/Plugins/EnhancedInput/Source/InputEditor/Private/InputCustomizations.cpp b/Engine/Plugins/EnhancedInput/Source/InputEditor/Private/InputCustomizations.cpp
--- a/Engine/Plugins/EnhancedInput/Source/InputEditor/Private/InputCustomizations.cpp
+++ b/Engine/Plugins/EnhancedInput/Source/InputEditor/Private/InputCustomizations.cpp
@@ -70,7 +70,7 @@ void AddInputActionProperties(TSharedRef<IPropertyHandle> PropertyHandle, IDetai
 	void* Data = nullptr;
     if (PropertyHandle->GetValueData(Data) != FPropertyAccess::Fail)
     {
-    	TObjectPtr<FEnhancedActionKeyMapping> ActionKeyMapping = static_cast<FEnhancedActionKeyMapping*>(Data);
+    	FEnhancedActionKeyMapping* ActionKeyMapping = static_cast<FEnhancedActionKeyMapping*>(Data);
     	
     	// Make sure ActionKeyMapping action is valid. If triggers and modifiers are empty we can just back out - nothing to add
     	if (ActionKeyMapping->Action && (!ActionKeyMapping->Action->Triggers.IsEmpty() || !ActionKeyMapping->Action->Modifiers.IsEmpty()))
diff --git a/Engine/Source/Runtime/Engine/Private/Elements/Framework/EngineElementsLibrary.cpp b/Engine/Source/Runtime/Engine/Private/Elements/Framework/EngineElementsLibrary.cpp
--- a/Engine/Source/Runtime/Engine/Private/Elements/Framework/EngineElementsLibrary.cpp
+++ b/Engine/Source/Runtime/Engine/Private/Elements/Framework/EngineElementsLibrary.cpp
@@ -367,7 +367,7 @@ TTypedElementOwner<FActorElementData> UEngineElementsLibrary::CreateActorElement
 	}
 	
 	UTypedElementRegistry* Registry = UTypedElementRegistry::GetInstance();
-	TObjectPtr<ITypedElementDataStorageCompatibilityInterface> Storage = Registry->GetMutableDataStorageCompatibility();
+	ITypedElementDataStorageCompatibilityInterface* Storage = Registry->GetMutableDataStorageCompatibility();
 	if (Storage)
 	{
 		Storage->AddCompatibleObject(const_cast<AActor*>(InActor));
@@ -391,7 +391,7 @@ void UEngineElementsLibrary::DestroyActorElement(const AActor* InActor, TTypedEl
 void UEngineElementsLibrary::UnregisterActorElement(const AActor* InActor)
 {
 	UTypedElementRegistry* Registry = UTypedElementRegistry::GetInstance();
-	TObjectPtr<ITypedElementDataStorageCompatibilityInterface> Storage = Registry->GetMutableDataStorageCompatibility();
+	ITypedElementDataStorageCompatibilityInterface* Storage = Registry->GetMutableDataStorageCompatibility();
 	if (Storage)
 	{
 		Storage->RemoveCompatibleObject(const_cast<AActor*>(InActor));
