Partially revert 22920594 -- make GetUniqueAssetName() fall back to GetAssetName() to avoid code paths that require game thread

diff --git a/Engine/Source/Runtime/Engine/Private/Materials/MaterialShared.cpp b/Engine/Source/Runtime/Engine/Private/Materials/MaterialShared.cpp
--- a/Engine/Source/Runtime/Engine/Private/Materials/MaterialShared.cpp
+++ b/Engine/Source/Runtime/Engine/Private/Materials/MaterialShared.cpp
@@ -2681,6 +2681,9 @@ bool FMaterial::CacheShaders(const FMaterialShaderMapId& ShaderMapId, EShaderPla
 #if WITH_EDITOR
 FString FMaterial::GetUniqueAssetName(EShaderPlatform Platform) const
 {
+	// Temporary disabled due to UE-169098:
+	// This funciton may be called from worker threads, but GetStaticParameterSet() requires game thread.
+#if 0 
 	TArray<uint8> Buf;
 	FString PathStr = GetAssetPath().GetPlainNameString();
 	FMemoryWriter BufWriter(Buf);
@@ -2690,8 +2693,11 @@ FString FMaterial::GetUniqueAssetName(EShaderPlatform Platform) const
 	FStaticParameterSet::StaticStruct()->SerializeBin(BufWriter, &StaticParams);
 	uint64 Hash = CityHash64((const char*)Buf.GetData(), Buf.Num());
 	return FString::Printf(TEXT("%s_%llu"), *GetAssetName(), Hash);
-}
+#else
+	return GetAssetName();
 #endif
+}
+#endif // WITH_EDITOR
 
 /**
  * Helper task used to release the strong object reference to the material interface on the game thread
diff --git a/Engine/Source/Runtime/Engine/Public/MaterialShared.h b/Engine/Source/Runtime/Engine/Public/MaterialShared.h
--- a/Engine/Source/Runtime/Engine/Public/MaterialShared.h
+++ b/Engine/Source/Runtime/Engine/Public/MaterialShared.h
@@ -1867,7 +1867,7 @@ public:
 	ENGINE_API FString GetUniqueAssetName(EShaderPlatform Platform) const;
 #else
 	/** GetUniqueAssetName just returns GetAssetName in non-editor builds */
-	FString GetUniqueAssetName() const { return GetAssetName(); }
+	FString GetUniqueAssetName(EShaderPlatform Platform) const { return GetAssetName(); }
 #endif
 	virtual bool HasVertexPositionOffsetConnected() const { return false; }
 	virtual bool HasPixelDepthOffsetConnected() const { return false; }
