Update Kismet Execution function to skip over pending garbage and unreachable objects.

diff --git a/Engine/Source/Runtime/Engine/Private/UnrealEngine.cpp b/Engine/Source/Runtime/Engine/Private/UnrealEngine.cpp
--- a/Engine/Source/Runtime/Engine/Private/UnrealEngine.cpp
+++ b/Engine/Source/Runtime/Engine/Private/UnrealEngine.cpp
@@ -5915,7 +5915,7 @@ bool UEngine::HandleKismetEventCommand(UWorld* InWorld, const TCHAR* Cmd, FOutpu
 
 		// Send the command to everything in the world we're dealing with...
 		int32 NumInstanceCallsSucceeded = 0;
-		for (TObjectIterator<UObject> It; It; ++It)
+		for (TObjectIterator<UObject> It(EObjectFlags::RF_ClassDefaultObject, true, EInternalObjectFlags::Garbage); It; ++It)
 		{
 			UObject* const Obj = *It;
 
