- Fixed content bundle to inject its non-spatially loaded content into its corresponding world partition grid top level.
- Fixed non-spatially loaded to also include non-spatially loaded injected content bundle cells.
- Fixed world partition streaming queries to include non-spatially loaded injected content bundle cells.

diff --git a/Engine/Source/Runtime/Engine/Private/WorldPartition/WorldPartitionRuntimeSpatialHash.cpp b/Engine/Source/Runtime/Engine/Private/WorldPartition/WorldPartitionRuntimeSpatialHash.cpp
--- a/Engine/Source/Runtime/Engine/Private/WorldPartition/WorldPartitionRuntimeSpatialHash.cpp
+++ b/Engine/Source/Runtime/Engine/Private/WorldPartition/WorldPartitionRuntimeSpatialHash.cpp
@@ -202,17 +202,26 @@ void FSpatialHashStreamingGrid::GetCells(const FWorldPartitionStreamingQuerySour
 	}
 
 	// Non Spatial (always included)
-	const int32 TopLevel = GridLevels.Num() - 1;
-	for (const FSpatialHashStreamingGridLayerCell& LayerCell : GridLevels[TopLevel].LayerCells)
+	auto ForEachGridLevelCells = [&QuerySource, &OutCells, ShouldAddCell](const FSpatialHashStreamingGridLevel& InGridLevel)
 	{
-		for (const UWorldPartitionRuntimeCell* Cell : LayerCell.GridCells)
+		for (const FSpatialHashStreamingGridLayerCell& LayerCell : InGridLevel.LayerCells)
 		{
-			if (ShouldAddCell(Cell, QuerySource))
+			for (const UWorldPartitionRuntimeCell* Cell : LayerCell.GridCells)
 			{
-				OutCells.Add(Cell);
+				if (ShouldAddCell(Cell, QuerySource))
+				{
+					OutCells.Add(Cell);
+				}
 			}
 		}
-	}	
+	};
+
+	const int32 TopLevel = GridLevels.Num() - 1;
+	ForEachGridLevelCells(GridLevels[TopLevel]);
+	if (InjectedGridLevels.IsValidIndex(TopLevel))
+	{
+		ForEachGridLevelCells(InjectedGridLevels[TopLevel]);
+	}
 }
 
 void FSpatialHashStreamingGrid::GetCells(const TArray<FWorldPartitionStreamingSource>& Sources, const UDataLayerSubsystem* DataLayerSubsystem, UWorldPartitionRuntimeHash::FStreamingSourceCells& OutActivateCells, UWorldPartitionRuntimeHash::FStreamingSourceCells& OutLoadCells, bool bEnableZCulling) const
@@ -361,10 +370,14 @@ void FSpatialHashStreamingGrid::InjectExternalStreamingObjectGrid(const FSpatial
 	for (int32 SourceGridLevel = 0; SourceGridLevel < InExternalObjectStreamingGrid.GridLevels.Num(); ++SourceGridLevel)
 	{
 		const FSpatialHashStreamingGridLevel& ExternalObjectGridLevel = InExternalObjectStreamingGrid.GridLevels[SourceGridLevel];
-		int32 TargetGridLevel = FMath::Min<int32>(SourceGridLevel, GridLevels.Num()-1);
+		const int32 TargetGridTopLevel = GridLevels.Num() - 1;
+		const int32 SourceGridTopLevel = InExternalObjectStreamingGrid.GridLevels.Num() - 1;
+		const bool bIsSourceTopLevel = (SourceGridLevel == SourceGridTopLevel);
+		// Make sure to use target top level for source top level to make sure that the content won't be spatially loaded
+		int32 TargetGridLevel = bIsSourceTopLevel ? TargetGridTopLevel : FMath::Min<int32>(SourceGridLevel, TargetGridTopLevel);
 		if (ensure(InjectedGridLevels.IsValidIndex(TargetGridLevel)))
 		{
-			check(InExternalObjectStreamingGrid.GetGridHelper().Levels[TargetGridLevel].CellSize == GetGridHelper().Levels[TargetGridLevel].CellSize);
+			check(bIsSourceTopLevel || InExternalObjectStreamingGrid.GetGridHelper().Levels[TargetGridLevel].CellSize == GetGridHelper().Levels[TargetGridLevel].CellSize);
 
 			for (const FSpatialHashStreamingGridLayerCell& ExternalObjectLayerCell : ExternalObjectGridLevel.LayerCells)
 			{
@@ -411,7 +424,11 @@ void FSpatialHashStreamingGrid::RemoveExternalStreamingObjectGrid(const FSpatial
 	for (int SourceGridLevel = 0; SourceGridLevel < InExternalObjectStreamingGrid.GridLevels.Num(); ++SourceGridLevel)
 	{
 		const FSpatialHashStreamingGridLevel& ExternalObjectGridLevel = InExternalObjectStreamingGrid.GridLevels[SourceGridLevel];
-		int32 TargetGridLevel = FMath::Min<int32>(SourceGridLevel, GridLevels.Num() - 1);
+		const int32 TargetGridTopLevel = GridLevels.Num() - 1;
+		const int32 SourceGridTopLevel = InExternalObjectStreamingGrid.GridLevels.Num() - 1;
+		const bool bIsSourceTopLevel = (SourceGridLevel == SourceGridTopLevel);
+		// Make sure to use target top level for source top level to make sure that the content won't be spatially loaded
+		int32 TargetGridLevel = bIsSourceTopLevel ? TargetGridTopLevel : FMath::Min<int32>(SourceGridLevel, TargetGridTopLevel);
 		if (ensure(InjectedGridLevels.IsValidIndex(TargetGridLevel)))
 		{
 			TArray<FSpatialHashStreamingGridLayerCell>& LayerCells = InjectedGridLevels[TargetGridLevel].LayerCells;
@@ -522,22 +539,31 @@ void FSpatialHashStreamingGrid::GetNonSpatiallyLoadedCells(const UDataLayerSubsy
 {
 	if (GridLevels.Num() > 0)
 	{
-		const int32 TopLevel = GridLevels.Num() - 1;
-		for (const FSpatialHashStreamingGridLayerCell& LayerCell : GridLevels[TopLevel].LayerCells)
+		auto ForEachGridLevelCells = [DataLayerSubsystem, &OutActivateCells, &OutLoadCells](const FSpatialHashStreamingGridLevel& InGridLevel)
 		{
-			for (const UWorldPartitionRuntimeCell* Cell : LayerCell.GridCells)
+			for (const FSpatialHashStreamingGridLayerCell& LayerCell : InGridLevel.LayerCells)
 			{
-				if (!Cell->HasDataLayers() || (DataLayerSubsystem && DataLayerSubsystem->IsAnyDataLayerInEffectiveRuntimeState(Cell->GetDataLayers(), EDataLayerRuntimeState::Activated)))
-				{
-					check(Cell->IsAlwaysLoaded() || Cell->HasDataLayers() || Cell->GetContentBundleID().IsValid());
-					OutActivateCells.Add(Cell);
-				}
-				else if(DataLayerSubsystem && DataLayerSubsystem->IsAnyDataLayerInEffectiveRuntimeState(Cell->GetDataLayers(), EDataLayerRuntimeState::Loaded))
+				for (const UWorldPartitionRuntimeCell* Cell : LayerCell.GridCells)
 				{
-					check(Cell->HasDataLayers());
-					OutLoadCells.Add(Cell);
+					if (!Cell->HasDataLayers() || (DataLayerSubsystem && DataLayerSubsystem->IsAnyDataLayerInEffectiveRuntimeState(Cell->GetDataLayers(), EDataLayerRuntimeState::Activated)))
+					{
+						check(Cell->IsAlwaysLoaded() || Cell->HasDataLayers() || Cell->GetContentBundleID().IsValid());
+						OutActivateCells.Add(Cell);
+					}
+					else if (DataLayerSubsystem && DataLayerSubsystem->IsAnyDataLayerInEffectiveRuntimeState(Cell->GetDataLayers(), EDataLayerRuntimeState::Loaded))
+					{
+						check(Cell->HasDataLayers());
+						OutLoadCells.Add(Cell);
+					}
 				}
 			}
+		};
+
+		const int32 TopLevel = GridLevels.Num() - 1;
+		ForEachGridLevelCells(GridLevels[TopLevel]);
+		if (InjectedGridLevels.IsValidIndex(TopLevel))
+		{
+			ForEachGridLevelCells(InjectedGridLevels[TopLevel]);
 		}
 	}
 }
