Chaos -  fix welded box collision scale

Fixed the issue when you have an actor containing a hierarchy of static mesh components that are welded. If the child object has a scale relative to the parent, and we then scale the actor in the scene, the child's offset was not scaled by the actor scale so the collision would be in the wrong place.

diff --git a/Engine/Source/Runtime/Engine/Private/PhysicsEngine/BodySetup.cpp b/Engine/Source/Runtime/Engine/Private/PhysicsEngine/BodySetup.cpp
--- a/Engine/Source/Runtime/Engine/Private/PhysicsEngine/BodySetup.cpp
+++ b/Engine/Source/Runtime/Engine/Private/PhysicsEngine/BodySetup.cpp
@@ -1897,7 +1897,7 @@ FKBoxElem FKBoxElem::GetFinalScaled(const FVector& Scale3D, const FTransform& Re
 	ScaledBox.Z *= Scale3DAbs.Z;
 
 	FTransform ScaleTransform(FQuat::Identity, FVector::ZeroVector, Scale3D); 
-	FTransform BoxTransform = GetTransform() * ScaleTransform * RelativeTM;
+	FTransform BoxTransform = GetTransform() * RelativeTM * ScaleTransform;
 	ScaledBox.SetTransform(BoxTransform);
 
 	return ScaledBox;
