Iris
- Temporarily disabled usage of bIsInitializingBaseFromDefault when initializing FastArraySerializer as it causes horde errors due to invalid data stored in archetype.

diff --git a/Engine/Source/Runtime/Engine/Private/DataReplication.cpp b/Engine/Source/Runtime/Engine/Private/DataReplication.cpp
--- a/Engine/Source/Runtime/Engine/Private/DataReplication.cpp
+++ b/Engine/Source/Runtime/Engine/Private/DataReplication.cpp
@@ -26,6 +26,10 @@
 #include "Net/Core/PushModel/Types/PushModelPerNetDriverState.h"
 #include "Net/RPCDoSDetection.h"
 
+#if UE_WITH_IRIS
+#include "Iris/IrisConfig.h"
+#endif
+
 DECLARE_LLM_MEMORY_STAT(TEXT("NetObjReplicator"), STAT_NetObjReplicatorLLM, STATGROUP_LLMFULL);
 LLM_DEFINE_TAG(NetObjReplicator, NAME_None, TEXT("Networking"), GET_STATFNAME(STAT_NetObjReplicatorLLM), GET_STATFNAME(STAT_NetworkingSummaryLLM));
 
@@ -375,8 +379,11 @@ bool FObjectReplicator::SendCustomDeltaProperty(UObject* InObject, uint16 Custom
 	Parms.bInternalAck = Connection->IsInternalAck();
 
 #if UE_WITH_IRIS
-	// When initializing baselines we should not modify the source data if it originates from the CDO or archetype
-	Parms.bIsInitializingBaseFromDefault = Parms.Object && Parms.Object->HasAnyFlags(RF_ClassDefaultObject | RF_ArchetypeObject);
+	if (UE::Net::ShouldUseIrisReplication())
+	{
+		// When initializing baselines we should not modify the source data if it originates from the CDO or archetype
+		Parms.bIsInitializingBaseFromDefault = Parms.Object && Parms.Object->HasAnyFlags(RF_ClassDefaultObject | RF_ArchetypeObject);
+	}
 #endif
 
 	return FNetSerializeCB::SendCustomDeltaProperty(*RepLayout, Parms, CustomDeltaIndex);
