Iris - Fix for assert in FReplicationSystemUtil::GetReplicationSystem during shutdown when trying to access NetDriver

diff --git a/Engine/Source/Runtime/Engine/Private/Net/Iris/ReplicationSystem/ReplicationSystemUtil.cpp b/Engine/Source/Runtime/Engine/Private/Net/Iris/ReplicationSystem/ReplicationSystemUtil.cpp
--- a/Engine/Source/Runtime/Engine/Private/Net/Iris/ReplicationSystem/ReplicationSystemUtil.cpp
+++ b/Engine/Source/Runtime/Engine/Private/Net/Iris/ReplicationSystem/ReplicationSystemUtil.cpp
@@ -25,7 +25,7 @@ namespace UE::Net
 
 UReplicationSystem* FReplicationSystemUtil::GetReplicationSystem(const AActor* Actor)
 {
-	UNetDriver* NetDriver = Actor ? Actor->GetNetDriver() : nullptr;
+	UNetDriver* NetDriver = Actor && Actor->GetWorld() ? Actor->GetNetDriver() : nullptr;
 	return NetDriver ? NetDriver->GetReplicationSystem() : nullptr;
 }
 
@@ -226,7 +226,7 @@ void FReplicationSystemUtil::EndReplication(AActor* Actor, EEndPlayReason::Type
 
 void FReplicationSystemUtil::EndReplicationForActorComponent(UActorComponent* SubObject)
 {
-	if (UActorReplicationBridge* Bridge = GetActorReplicationBridge(SubObject->GetOwner()))
+	if (UActorReplicationBridge* Bridge = IsValid(SubObject) ? GetActorReplicationBridge(SubObject->GetOwner()) : nullptr)
 	{
 		Bridge->EndReplicationForActorComponent(SubObject);
 	}
