Fix serialization for older MaterialInstance Assets

diff --git a/Engine/Source/Runtime/Engine/Private/Materials/MaterialInstance.cpp b/Engine/Source/Runtime/Engine/Private/Materials/MaterialInstance.cpp
--- a/Engine/Source/Runtime/Engine/Private/Materials/MaterialInstance.cpp
+++ b/Engine/Source/Runtime/Engine/Private/Materials/MaterialInstance.cpp
@@ -2634,19 +2634,11 @@ void UMaterialInstance::Serialize(FArchive& Ar)
 		StaticParameters_DEPRECATED.UpdateLegacyTerrainLayerWeightData();
 	}
 
-	if (Ar.IsLoading())
+	if (Ar.IsLoading() && !StaticParameters_DEPRECATED.IsEmpty())
 	{
-		if (!StaticParameters_DEPRECATED.IsEmpty())
-		{
-			StaticParametersRuntime = MoveTemp(StaticParameters_DEPRECATED.GetRuntime());
-			GetEditorOnlyData()->StaticParameters = MoveTemp(StaticParameters_DEPRECATED.EditorOnly);
-			StaticParameters_DEPRECATED.Empty();
-		}
-
-		if (!GetEditorOnlyData()->StaticParameters.StaticSwitchParameters_DEPRECATED.IsEmpty())
-		{
-			StaticParametersRuntime.StaticSwitchParameters = MoveTemp(GetEditorOnlyData()->StaticParameters.StaticSwitchParameters_DEPRECATED);
-		}
+		StaticParametersRuntime = MoveTemp(StaticParameters_DEPRECATED.GetRuntime());
+		GetEditorOnlyData()->StaticParameters = MoveTemp(StaticParameters_DEPRECATED.EditorOnly);
+		StaticParameters_DEPRECATED.Empty();
 	}
 
 #endif // WITH_EDITOR
@@ -2812,6 +2804,11 @@ void UMaterialInstance::Serialize(FArchive& Ar)
 	{
 		ValidateTextureOverrides(GMaxRHIFeatureLevel);
 	}
+
+	if (Ar.IsLoading() && !GetEditorOnlyData()->StaticParameters.StaticSwitchParameters_DEPRECATED.IsEmpty())
+	{
+		StaticParametersRuntime.StaticSwitchParameters = MoveTemp(GetEditorOnlyData()->StaticParameters.StaticSwitchParameters_DEPRECATED);
+	}
 }
 
 void UMaterialInstance::PostLoad()
