fix editor only data

diff --git a/Engine/Source/Runtime/Engine/Private/Materials/MaterialInstance.cpp b/Engine/Source/Runtime/Engine/Private/Materials/MaterialInstance.cpp
--- a/Engine/Source/Runtime/Engine/Private/Materials/MaterialInstance.cpp
+++ b/Engine/Source/Runtime/Engine/Private/Materials/MaterialInstance.cpp
@@ -2805,10 +2805,12 @@ void UMaterialInstance::Serialize(FArchive& Ar)
 		ValidateTextureOverrides(GMaxRHIFeatureLevel);
 	}
 
+#if WITH_EDITORONLY_DATA
 	if (Ar.IsLoading() && !GetEditorOnlyData()->StaticParameters.StaticSwitchParameters_DEPRECATED.IsEmpty())
 	{
 		StaticParametersRuntime.StaticSwitchParameters = MoveTemp(GetEditorOnlyData()->StaticParameters.StaticSwitchParameters_DEPRECATED);
 	}
+#endif
 }
 
 void UMaterialInstance::PostLoad()
