Adds support for dynamic resolution

diff --git a/Engine/Source/Runtime/Engine/Classes/GameFramework/GameUserSettings.h b/Engine/Source/Runtime/Engine/Classes/GameFramework/GameUserSettings.h
--- a/Engine/Source/Runtime/Engine/Classes/GameFramework/GameUserSettings.h
+++ b/Engine/Source/Runtime/Engine/Classes/GameFramework/GameUserSettings.h
@@ -100,7 +100,7 @@ public:
 
 	/** Returns the user setting for dynamic resolution. */
 	UFUNCTION(BlueprintPure, Category = Settings)
-	bool IsDynamicResolutionEnabled() const;
+	virtual bool IsDynamicResolutionEnabled() const;
 
 	/** Checks if the Screen Resolution user setting is different from current */
 	UFUNCTION(BlueprintPure, Category=Settings)
diff --git a/Engine/Source/Runtime/Engine/Private/DynamicResolution.cpp b/Engine/Source/Runtime/Engine/Private/DynamicResolution.cpp
--- a/Engine/Source/Runtime/Engine/Private/DynamicResolution.cpp
+++ b/Engine/Source/Runtime/Engine/Private/DynamicResolution.cpp
@@ -40,6 +40,19 @@ static TAutoConsoleVariable<float> CVarTargetedGPUHeadRoomPercentage(
 	TEXT("Targeted GPU headroom (in percent from r.DynamicRes.FrameTimeBudget)."),
 	ECVF_RenderThreadSafe | ECVF_Default);
 
+/** On desktop, the swap chain doesn't allow tear amount configuration, so an overbudget frame can be droped with r.VSync=1.
+ * So need to lower the heuristic's target budget to lower chances to go overbudget.
+ *
+ * Moreover the GPU is a shared ressource with other process which may or may not be included in our GPU timings,
+ * and need to leave some GPU capacity to these application to not get preempted by OS scheduler.
+ * Given we can measure other application's GPU cost, need to leave enough headroom for them all the time.
+ */
+static TAutoConsoleVariable<float> CVarOverBudgetGPUHeadRoomPercentage(
+	TEXT("r.DynamicRes.OverBudgetGPUHeadRoomPercentage"),
+	0.0f,
+	TEXT("Amount of GPU headroom needed from which the frame is considered over budget. This is for platform not supporting controllable tearing with VSync (in percent from r.DynamicRes.FrameTimeBudget)."),
+	ECVF_RenderThreadSafe | ECVF_Default);
+
 static TAutoConsoleVariable<int32> CVarHistorySize(
 	TEXT("r.DynamicRes.HistorySize"),
 	16,
@@ -101,7 +114,7 @@ DynamicRenderScaling::FHeuristicSettings GetPrimaryDynamicResolutionSettings()
 	BudgetSetting.MinResolutionFraction      = DynamicRenderScaling::GetPercentageCVarToFraction(CVarDynamicResMinSP);
 	BudgetSetting.MaxResolutionFraction      = DynamicRenderScaling::GetPercentageCVarToFraction(CVarDynamicResMaxSP);
 	BudgetSetting.UpperBoundQuantization     = CVarUpperBoundQuantization.GetValueOnAnyThread();
-	BudgetSetting.BudgetMs                   = CVarFrameTimeBudget.GetValueOnAnyThread();
+	BudgetSetting.BudgetMs                   = CVarFrameTimeBudget.GetValueOnAnyThread() * (1.0f - DynamicRenderScaling::GetPercentageCVarToFraction(CVarOverBudgetGPUHeadRoomPercentage));
 	BudgetSetting.ChangeThreshold            = DynamicRenderScaling::GetPercentageCVarToFraction(CVarChangeThreshold);
 	BudgetSetting.TargetedHeadRoom           = DynamicRenderScaling::GetPercentageCVarToFraction(CVarTargetedGPUHeadRoomPercentage);
 	BudgetSetting.IncreaseAmortizationFactor = CVarIncreaseAmortizationFactor.GetValueOnAnyThread();
diff --git a/Engine/Source/Runtime/Engine/Private/Scalability.cpp b/Engine/Source/Runtime/Engine/Private/Scalability.cpp
--- a/Engine/Source/Runtime/Engine/Private/Scalability.cpp
+++ b/Engine/Source/Runtime/Engine/Private/Scalability.cpp
@@ -567,6 +567,11 @@ static float GetRenderScaleLevelFromQualityLevel(int32 InQualityLevel, EQualityL
 	return FCString::Atof(*ResolutionValueStrings[InQualityLevel]);
 }
 
+ENGINE_API float GetResolutionQualityFromGPUPerfIndex(float GPUPerfIndex)
+{
+	 return GetRenderScaleLevelFromQualityLevel(ComputeOptionFromPerfIndex(TEXT("ResolutionQuality"), 0, GPUPerfIndex));
+}
+
 FQualityLevels BenchmarkQualityLevels(uint32 WorkScale, float CPUMultiplier, float GPUMultiplier)
 {
 	ensure((CPUMultiplier > 0.0f) && (GPUMultiplier > 0.0f));
@@ -582,7 +587,7 @@ FQualityLevels BenchmarkQualityLevels(uint32 WorkScale, float CPUMultiplier, flo
 	const float GPUPerfIndex = SynthBenchmark.ComputeGPUPerfIndex(/*out*/ &Results.GPUBenchmarkSteps) * GPUMultiplier;
 
 	// decide on the actual quality needed
-	Results.ResolutionQuality = GetRenderScaleLevelFromQualityLevel(ComputeOptionFromPerfIndex(TEXT("ResolutionQuality"), CPUPerfIndex, GPUPerfIndex));
+	Results.ResolutionQuality = GetResolutionQualityFromGPUPerfIndex(GPUPerfIndex);
 	Results.ViewDistanceQuality = ComputeOptionFromPerfIndex(TEXT("ViewDistanceQuality"), CPUPerfIndex, GPUPerfIndex);
 	Results.AntiAliasingQuality = ComputeOptionFromPerfIndex(TEXT("AntiAliasingQuality"), CPUPerfIndex, GPUPerfIndex);
 	Results.ShadowQuality = ComputeOptionFromPerfIndex(TEXT("ShadowQuality"), CPUPerfIndex, GPUPerfIndex);
