Sky component accessors

diff --git a/Engine/Source/Runtime/Engine/Classes/Components/SkyAtmosphereComponent.h b/Engine/Source/Runtime/Engine/Classes/Components/SkyAtmosphereComponent.h
--- a/Engine/Source/Runtime/Engine/Classes/Components/SkyAtmosphereComponent.h
+++ b/Engine/Source/Runtime/Engine/Classes/Components/SkyAtmosphereComponent.h
@@ -158,9 +158,12 @@ class USkyAtmosphereComponent : public USceneComponent
 	float AerialPerspectiveStartDepth;
 
 
-
 	UFUNCTION(BlueprintCallable, Category = "Rendering")
 	ENGINE_API void OverrideAtmosphereLightDirection(int32 AtmosphereLightIndex, const FVector& LightDirection);
+	UFUNCTION(BlueprintCallable, Category = "Rendering")
+	ENGINE_API bool IsAtmosphereLightDirectionOverriden(int32 AtmosphereLightIndex);
+	UFUNCTION(BlueprintCallable, Category = "Rendering")
+	ENGINE_API FVector GetOverridenAtmosphereLightDirection(int32 AtmosphereLightIndex);
 
 	UFUNCTION(BlueprintCallable, Category = "Rendering")
 	ENGINE_API void SetAtmosphereHeight(float NewValue);
diff --git a/Engine/Source/Runtime/Engine/Private/Components/SkyAtmosphereComponent.cpp b/Engine/Source/Runtime/Engine/Private/Components/SkyAtmosphereComponent.cpp
--- a/Engine/Source/Runtime/Engine/Private/Components/SkyAtmosphereComponent.cpp
+++ b/Engine/Source/Runtime/Engine/Private/Components/SkyAtmosphereComponent.cpp
@@ -309,6 +309,26 @@ void USkyAtmosphereComponent::OverrideAtmosphereLightDirection(int32 AtmosphereL
 	}
 }
 
+bool USkyAtmosphereComponent::IsAtmosphereLightDirectionOverriden(int32 AtmosphereLightIndex)
+{
+	check(AtmosphereLightIndex >= 0 && AtmosphereLightIndex < NUM_ATMOSPHERE_LIGHTS);
+	if (AtmosphereLightIndex >= 0 && AtmosphereLightIndex < NUM_ATMOSPHERE_LIGHTS)
+	{
+		return OverrideAtmosphericLight[AtmosphereLightIndex];
+	}
+	return false;
+}
+
+FVector USkyAtmosphereComponent::GetOverridenAtmosphereLightDirection(int32 AtmosphereLightIndex)
+{
+	check(AtmosphereLightIndex >= 0 && AtmosphereLightIndex < NUM_ATMOSPHERE_LIGHTS);
+	if (AtmosphereLightIndex >= 0 && AtmosphereLightIndex < NUM_ATMOSPHERE_LIGHTS)
+	{
+		return OverrideAtmosphericLightDirection[AtmosphereLightIndex];
+	}
+	return FVector::ZeroVector;
+}
+
 void USkyAtmosphereComponent::GetOverrideLightStatus(bool* OutOverrideAtmosphericLight, FVector* OutOverrideAtmosphericLightDirection) const
 {
 	memcpy(OutOverrideAtmosphericLight, OverrideAtmosphericLight, sizeof(OverrideAtmosphericLight));
