GameFeatureDataAssetDependencyGatherer: Add Game references to content bundle external actors to properly setup their management resources

GameFeatureData assets with ContentBundles need to be resaved.

diff --git a/Engine/Plugins/Experimental/GameFeatures/Source/GameFeatures/Private/GameFeatureData.cpp b/Engine/Plugins/Experimental/GameFeatures/Source/GameFeatures/Private/GameFeatureData.cpp
--- a/Engine/Plugins/Experimental/GameFeatures/Source/GameFeatures/Private/GameFeatureData.cpp
+++ b/Engine/Plugins/Experimental/GameFeatures/Source/GameFeatures/Private/GameFeatureData.cpp
@@ -8,6 +8,8 @@
 #include "Misc/ConfigContext.h"
 #include "UObject/UObjectHash.h"
 #include "UObject/CoreRedirects.h"
+#include "GameFeatureAction_AddWPContent.h"
+#include "WorldPartition/ContentBundle/ContentBundleDescriptor.h"
 
 #include UE_INLINE_GENERATED_CPP_BY_NAME(GameFeatureData)
 
@@ -208,5 +210,55 @@ void UGameFeatureData::ReloadConfigs(FConfigFile& PluginConfig) const
 	}
 }
 
+#if WITH_EDITOR
+FName UGameFeatureData::GetContentBundleGuidsAssetRegistryTag()
+{
+	static const FName ContentBundlesTag("ContentBundleGuids");
+	return ContentBundlesTag;
+}
+
+void UGameFeatureData::GetContentBundleGuidsFromAsset(const FAssetData& Asset, TArray<FGuid>& OutContentBundleGuids)
+{
+	FString ContentBundleGuidsStr;
+	if (Asset.GetTagValue(GetContentBundleGuidsAssetRegistryTag(), ContentBundleGuidsStr))
+	{
+		TArray<FString> ContentBundleGuidsStrArray;
+		ContentBundleGuidsStr.ParseIntoArray(ContentBundleGuidsStrArray, TEXT(","));
+		for (const FString& GuidStr : ContentBundleGuidsStrArray)
+		{
+			FGuid ContentBundleGuid;
+			if (FGuid::Parse(GuidStr, ContentBundleGuid))
+			{
+				OutContentBundleGuids.Add(ContentBundleGuid);
+			}
+		}
+	}
+}
+
+void UGameFeatureData::GetAssetRegistryTags(TArray<FAssetRegistryTag>& OutTags) const
+{
+	Super::GetAssetRegistryTags(OutTags);
+
+	TArray<FGuid> ContentBundleGuids;
+
+	for (UGameFeatureAction* Action : Actions)
+	{
+		if (UGameFeatureAction_AddWPContent* WPAction = Cast<UGameFeatureAction_AddWPContent>(Action))
+		{
+			if (const UContentBundleDescriptor* ContentBundleDescriptor = WPAction->GetContentBundleDescriptor())
+			{
+				ContentBundleGuids.Add(ContentBundleDescriptor->GetGuid());
+			}
+		}
+	}
+
+	if (ContentBundleGuids.Num() > 0)
+	{
+		FString ContentBundleGuidsStr = FString::JoinBy(ContentBundleGuids, TEXT(","), [&](const FGuid& Guid) { return Guid.ToString(); });
+		OutTags.Add(FAssetRegistryTag(GetContentBundleGuidsAssetRegistryTag(), ContentBundleGuidsStr, FAssetRegistryTag::TT_Hidden));
+	}
+}
+#endif
+
 #undef LOCTEXT_NAMESPACE
 
diff --git a/Engine/Plugins/Experimental/GameFeatures/Source/GameFeatures/Private/GameFeatureDataAssetDependencyGatherer.cpp b/Engine/Plugins/Experimental/GameFeatures/Source/GameFeatures/Private/GameFeatureDataAssetDependencyGatherer.cpp
new file mode 100644
--- /dev/null
+++ b/Engine/Plugins/Experimental/GameFeatures/Source/GameFeatures/Private/GameFeatureDataAssetDependencyGatherer.cpp
@@ -0,0 +1,50 @@
+// Copyright Epic Games, Inc. All Rights Reserved.
+#include "GameFeatureDataAssetDependencyGatherer.h"
+
+#if WITH_EDITOR
+
+#include "GameFeatureData.h"
+#include "HAL/FileManager.h" 
+#include "WorldPartition/ContentBundle/ContentBundlePaths.h"
+#include "AssetRegistry/AssetData.h"
+
+// Register FGameFeatureDataAssetDependencyGatherer for UGameFeatureData class
+REGISTER_ASSETDEPENDENCY_GATHERER(FGameFeatureDataAssetDependencyGatherer, UGameFeatureData);
+
+void FGameFeatureDataAssetDependencyGatherer::GatherDependencies(const FAssetData& AssetData, const FAssetRegistryState& AssetRegistryState, TFunctionRef<FARCompiledFilter(const FARFilter&)> CompileFilterFunc, TArray<IAssetDependencyGatherer::FGathereredDependency>& OutDependencies, TArray<FString>& OutDependencyDirectories) const
+{
+	TArray<FGuid> ContentBundleGuids;
+	UGameFeatureData::GetContentBundleGuidsFromAsset(AssetData, ContentBundleGuids);
+	if (ContentBundleGuids.Num() > 0)
+	{
+		FARFilter Filter;
+		Filter.bRecursivePaths = true;
+		Filter.bIncludeOnlyOnDiskAssets = true;
+
+		FName MountPoint = FPackageName::GetPackageMountPoint(AssetData.PackagePath.ToString());
+		for (const FGuid& ContentBundleGuid : ContentBundleGuids)
+		{
+			FString ContentBundleRootPath;
+			if (ContentBundlePaths::BuildContentBundleRootPath(MountPoint.ToString(), ContentBundleGuid, ContentBundleRootPath))
+			{
+				const FString ExternalActorsPath = ULevel::GetExternalActorsPath(ContentBundleRootPath);
+
+				OutDependencyDirectories.Add(ExternalActorsPath);
+				Filter.PackagePaths.Add(*ExternalActorsPath);
+			}
+		}
+
+		if (Filter.PackagePaths.Num() > 0)
+		{
+			TArray<FAssetData> FilteredAssets;
+			AssetRegistryState.GetAssets(CompileFilterFunc(Filter), {}, FilteredAssets, true);
+
+			for (const FAssetData& FilteredAsset : FilteredAssets)
+			{
+				OutDependencies.Emplace(IAssetDependencyGatherer::FGathereredDependency{ FilteredAsset.PackageName, UE::AssetRegistry::EDependencyProperty::Game });
+			}
+		}
+	}
+}
+
+#endif
\ No newline at end of file
diff --git a/Engine/Plugins/Experimental/GameFeatures/Source/GameFeatures/Private/GameFeatureDataAssetDependencyGatherer.h b/Engine/Plugins/Experimental/GameFeatures/Source/GameFeatures/Private/GameFeatureDataAssetDependencyGatherer.h
new file mode 100644
--- /dev/null
+++ b/Engine/Plugins/Experimental/GameFeatures/Source/GameFeatures/Private/GameFeatureDataAssetDependencyGatherer.h
@@ -0,0 +1,21 @@
+// Copyright Epic Games, Inc. All Rights Reserved.
+
+#pragma once
+
+#if WITH_EDITOR
+
+#include "AssetRegistry/AssetDependencyGatherer.h"
+#include "AssetRegistry/AssetData.h"
+#include "GameFeatureData.h"
+
+class FGameFeatureDataAssetDependencyGatherer : public IAssetDependencyGatherer
+{
+public:
+	
+	FGameFeatureDataAssetDependencyGatherer() = default;
+	virtual ~FGameFeatureDataAssetDependencyGatherer() = default;
+
+	virtual void GatherDependencies(const FAssetData& AssetData, const FAssetRegistryState& AssetRegistryState, TFunctionRef<FARCompiledFilter(const FARFilter&)> CompileFilterFunc, TArray<IAssetDependencyGatherer::FGathereredDependency>& OutDependencies, TArray<FString>& OutDependencyDirectories) const override;
+};
+
+#endif
\ No newline at end of file
diff --git a/Engine/Plugins/Experimental/GameFeatures/Source/GameFeatures/Public/GameFeatureData.h b/Engine/Plugins/Experimental/GameFeatures/Source/GameFeatures/Public/GameFeatureData.h
--- a/Engine/Plugins/Experimental/GameFeatures/Source/GameFeatures/Public/GameFeatureData.h
+++ b/Engine/Plugins/Experimental/GameFeatures/Source/GameFeatures/Public/GameFeatureData.h
@@ -6,9 +6,12 @@
 #include "Engine/DataAsset.h"
 #include "Engine/AssetManagerTypes.h"
 #include "GameFeatureAction.h"
+#include "Misc/Guid.h"
 
 #include "GameFeatureData.generated.h"
 
+struct FAssetData;
+
 /** Data related to a game feature, a collection of code and content that adds a separable discrete feature to the game */
 UCLASS()
 class GAMEFEATURES_API UGameFeatureData : public UPrimaryDataAsset
@@ -29,6 +32,12 @@ public:
 	/** Method to process ini files for the plugin during activation */
 	void InitializeHierarchicalPluginIniFiles(const FString& PluginInstalledFilename) const;
 
+#if WITH_EDITOR
+	static FName GetContentBundleGuidsAssetRegistryTag();
+	static void GetContentBundleGuidsFromAsset(const FAssetData& Asset, TArray<FGuid>& OutContentBundleGuids);
+	virtual void GetAssetRegistryTags(TArray<FAssetRegistryTag>& OutTags) const override;
+#endif
+
 public:
 	//~UPrimaryDataAsset interface
 #if WITH_EDITORONLY_DATA
diff --git a/Engine/Source/Runtime/Engine/Private/WorldPartition/ContentBundle/ContentBundlePaths.cpp b/Engine/Source/Runtime/Engine/Private/WorldPartition/ContentBundle/ContentBundlePaths.cpp
--- a/Engine/Source/Runtime/Engine/Private/WorldPartition/ContentBundle/ContentBundlePaths.cpp
+++ b/Engine/Source/Runtime/Engine/Private/WorldPartition/ContentBundle/ContentBundlePaths.cpp
@@ -122,22 +122,38 @@ FStringView ContentBundlePaths::GetActorPathRelativeToExternalActors(FStringView
 	return FStringView();
 }
 
-bool ContentBundlePaths::BuildActorDescContainerPackgePath(const FString& InContenBundleMountPoint, const FGuid& InContentBundleGuid, const FString& InLevelPackagePath, FString& OutContainerPackagePath)
+bool ContentBundlePaths::BuildContentBundleRootPath(const FString& InContenBundleMountPoint, const FGuid& InContentBundleGuid, FString& OutContentBundleRootPath)
 {
 	if (InContenBundleMountPoint.IsEmpty() || !InContentBundleGuid.IsValid())
 	{
 		return false;
 	}
 
+	TStringBuilderWithBuffer<TCHAR, NAME_SIZE> ContentBundleRootPathBuilder;
+	ContentBundleRootPathBuilder += TEXT("/");
+	ContentBundleRootPathBuilder += InContenBundleMountPoint;
+	ContentBundleRootPathBuilder += GetContentBundleFolder();
+	ContentBundleRootPathBuilder += InContentBundleGuid.ToString();
+	ContentBundleRootPathBuilder += TEXT("/");
+	
+	OutContentBundleRootPath = *ContentBundleRootPathBuilder;
+
+	return true;
+}
+
+bool ContentBundlePaths::BuildActorDescContainerPackgePath(const FString& InContenBundleMountPoint, const FGuid& InContentBundleGuid, const FString& InLevelPackagePath, FString& OutContainerPackagePath)
+{
+	FString ContentBundleRootPath;
+	if (!BuildContentBundleRootPath(InContenBundleMountPoint, InContentBundleGuid, ContentBundleRootPath))
+	{
+		return false;
+	}
+
 	FString LevelRoot, LevelPackagePath, LevelName;
 	if (FPackageName::SplitLongPackageName(InLevelPackagePath, LevelRoot, LevelPackagePath, LevelName))
 	{
 		TStringBuilderWithBuffer<TCHAR, NAME_SIZE> ContentBundleActorDescContainerPackage;
-		ContentBundleActorDescContainerPackage += TEXT("/");
-		ContentBundleActorDescContainerPackage += InContenBundleMountPoint;
-		ContentBundleActorDescContainerPackage += GetContentBundleFolder();
-		ContentBundleActorDescContainerPackage += InContentBundleGuid.ToString();
-		ContentBundleActorDescContainerPackage += TEXT("/");
+		ContentBundleActorDescContainerPackage += ContentBundleRootPath;
 		ContentBundleActorDescContainerPackage += LevelPackagePath;
 		ContentBundleActorDescContainerPackage += LevelName;
 
diff --git a/Engine/Source/Runtime/Engine/Public/WorldPartition/ContentBundle/ContentBundlePaths.h b/Engine/Source/Runtime/Engine/Public/WorldPartition/ContentBundle/ContentBundlePaths.h
--- a/Engine/Source/Runtime/Engine/Public/WorldPartition/ContentBundle/ContentBundlePaths.h
+++ b/Engine/Source/Runtime/Engine/Public/WorldPartition/ContentBundle/ContentBundlePaths.h
@@ -39,6 +39,9 @@ namespace ContentBundlePaths
 	FString GetRelativeLevelFolder(const FContentBundleBase& ContentBundle);
 
 #if WITH_EDITOR
+	// return format is /{MountPoint}/{ExternalActorFolder}/{ContentBundleFolder}/{ContentBundleUID}/
+	ENGINE_API bool BuildContentBundleRootPath(const FString& InContenBundleMountPoint, const FGuid& InContentBundleGuid, FString& OutContentBundleRootPath);
+
 	// return format is /{MountPoint}/{ExternalActorFolder}/{ContentBundleFolder}/{ContentBundleUID}/{LevelPath}/
 	ENGINE_API bool BuildActorDescContainerPackgePath(const FString& InContenBundleMountPoint, const FGuid& InContentBundleGuid, const FString& InLevelPackagePath, FString& OutContainerPackagePath);
 
