Replicated SubObjects:

* Added API's that allow the authority to tear off or delete a replicated subobject on remote connections.
* Call DestroyReplicatedSubObjectOnRemotePeers or TearOffReplicatedSubObjectOnRemotePeers via the Actor or ActorComponent to do so.
* Added in the NetDriver's network object info a list of subobjects replicated to individual channels and the number of channels they are replicated to.

diff --git a/Engine/Source/Runtime/Engine/Classes/Components/ActorComponent.h b/Engine/Source/Runtime/Engine/Classes/Components/ActorComponent.h
--- a/Engine/Source/Runtime/Engine/Classes/Components/ActorComponent.h
+++ b/Engine/Source/Runtime/Engine/Classes/Components/ActorComponent.h
@@ -523,11 +523,33 @@ public:
 	void AddReplicatedSubObject(UObject* SubObject, ELifetimeCondition NetCondition = COND_None);
 
 	/**
-	* Unregister a SubObject so it stops being replicated.
-	* @param SubObject The subobject to remove
+	* Unregister a SubObject to stop replicating it's properties to clients.
+	* This does not remove or delete it from connections where it was already replicated.
+	* By default a replicated subobject gets deleted on clients when the original pointer on the authority gets invalidated.
+	* If you want to immediately remove it from client use the DeleteReplicatedSubObject or TearOffReplicatedSubObject functions instead of this one.
+	* @param SubObject The SubObject to remove
 	*/
 	void RemoveReplicatedSubObject(UObject* SubObject);
 
+	/**
+	* Stop replicating a subobject and tell actor channels to delete the replica of this subobject next time the Actor gets replicated
+	* Note it is up to the caller to delete the local object on the authority.
+	* If you are using the legacy subobject replication method (ReplicateSubObjects() aka bReplicateUsingRegisteredSubObjectList=false) make sure the
+	* subobject doesn't get replicated there either.
+	* @param SubObject THe SubObject to delete
+	*/
+	void DestroyReplicatedSubObjectOnRemotePeers(UObject* SubObject);
+
+	/**
+	* Stop replicating a subobject and tell actor channels who spawned a replica of this subobject to release ownership over it.
+	* This means that on the remote connection the network engine will stop holding a reference to the subobject and it's up to other systems
+	* to keep that reference active or the subobject will get garbage collected.
+	* If you are using the legacy subobject replication method (ReplicateSubObjects() aka bReplicateUsingRegisteredSubObjectList=false) make sure the
+	* subobject doesn't get replicated there either.
+	* @param SubObject The SubObject to tear off
+	*/
+	void TearOffReplicatedSubObjectOnRemotePeers(UObject* SubObject);
+
 	/**
 	* Tells if the object is registered to be replicated by this actor component.
 	*/
diff --git a/Engine/Source/Runtime/Engine/Classes/Engine/ActorChannel.h b/Engine/Source/Runtime/Engine/Classes/Engine/ActorChannel.h
--- a/Engine/Source/Runtime/Engine/Classes/Engine/ActorChannel.h
+++ b/Engine/Source/Runtime/Engine/Classes/Engine/ActorChannel.h
@@ -112,6 +112,8 @@ private:
 	/** Set to true if SerializeActor is called due to an RPC forcing the channel open */
 	uint32 bIsForcedSerializeFromRPC:1;
 
+    uint16 ChannelSubObjectDirtyCount = 0;
+
 public:
 	bool GetSkipRoleSwap() const { return !!bSkipRoleSwap; }
 	void SetSkipRoleSwap(const bool bShouldSkip) { bSkipRoleSwap = bShouldSkip; }
@@ -219,7 +221,8 @@ public:
 	void WriteContentBlockHeader( UObject* Obj, FNetBitWriter &Bunch, const bool bHasRepLayout );
 
 	/** Writes the header for a content block specifically for deleting sub-objects */
-	void WriteContentBlockForSubObjectDelete( FOutBunch & Bunch, FNetworkGUID & GuidToDelete );
+	UE_DEPRECATED(5.2, "This function will be made private in the future." )
+	void WriteContentBlockForSubObjectDelete( FOutBunch & Bunch, FNetworkGUID & GuidToDelete);
 
 	/** Writes header and payload of content block */
 	int32 WriteContentBlockPayload( UObject* Obj, FNetBitWriter &Bunch, const bool bHasRepLayout, FNetBitWriter& Payload );
@@ -406,6 +409,18 @@ protected:
 	/** Handle the replication of subobjects for this actor. Returns true if data was written into the Bunch. */
 	bool DoSubObjectReplication(FOutBunch& Bunch, FReplicationFlags& OutRepFlags);
 
+	enum class ESubObjectDeleteFlag : uint8
+	{
+		Destroyed, // Delete operation that occurs when we detect that the original object on the authority became invalid.
+		TearOff,      // The client's actor channel will remove references to this subobject
+		ForceDelete,  // The subobject needs to be deleted on the client even if the original on the authority was not.
+	};
+
+	/** Writes the header for a content block specifically for deleting sub-objects */
+	void WriteContentBlockForSubObjectDelete(FOutBunch& Bunch, FNetworkGUID& GuidToDelete, ESubObjectDeleteFlag DeleteFlag);
+
+	const TCHAR* ToString(UActorChannel::ESubObjectDeleteFlag DeleteFlag);
+
 private:
 
 	/** Replicate Subobjects using the actor's registered list and its replicated actor component list */
@@ -427,6 +442,8 @@ private:
 	void TestLegacyReplicateSubObjects(UActorComponent* ReplicatedComponent, FOutBunch& Bunch, FReplicationFlags RepFlags);
 	void TestLegacyReplicateSubObjects(FOutBunch& Bunch, FReplicationFlags RepFlags);
 
+	bool UpdateDeletedSubObjects(FOutBunch& Bunch);
+
 	inline TArray< TObjectPtr<UObject> >& GetCreatedSubObjects()
 	{
 		PRAGMA_DISABLE_DEPRECATION_WARNINGS
diff --git a/Engine/Source/Runtime/Engine/Classes/Engine/NetDriver.h b/Engine/Source/Runtime/Engine/Classes/Engine/NetDriver.h
--- a/Engine/Source/Runtime/Engine/Classes/Engine/NetDriver.h
+++ b/Engine/Source/Runtime/Engine/Classes/Engine/NetDriver.h
@@ -337,11 +337,11 @@ class StatelessConnectHandlerComponent;
 class UNetConnection;
 class UReplicationDriver;
 struct FNetworkObjectInfo;
-struct FReplicatedStaticActorDestructionInfo;
 class UChannel;
 class IAnalyticsProvider;
 class FNetAnalyticsAggregator;
 class UNetDriver;
+class UActorChannel;
 class PacketHandler;
 struct FReplicatedStaticActorDestructionInfo;
 
@@ -596,7 +596,7 @@ struct FActorPriority
 	int32						Priority;	// Update priority, higher = more important.
 	
 	FNetworkObjectInfo*			ActorInfo;	// Actor info.
-	class UActorChannel*		Channel;	// Actor channel.
+	UActorChannel*		        Channel;	// Actor channel.
 
 	FActorDestructionInfo *	DestructionInfo;	// Destroy an actor
 
@@ -604,7 +604,7 @@ struct FActorPriority
 		Priority(0), ActorInfo(NULL), Channel(NULL), DestructionInfo(NULL)
 	{}
 
-	FActorPriority(class UNetConnection* InConnection, class UActorChannel* InChannel, FNetworkObjectInfo* InActorInfo, const TArray<struct FNetViewer>& Viewers, bool bLowBandwidth);
+	FActorPriority(class UNetConnection* InConnection, UActorChannel* InChannel, FNetworkObjectInfo* InActorInfo, const TArray<struct FNetViewer>& Viewers, bool bLowBandwidth);
 	FActorPriority(class UNetConnection* InConnection, FActorDestructionInfo * DestructInfo, const TArray<struct FNetViewer>& Viewers );
 };
 
@@ -1543,6 +1543,12 @@ public:
 
 	ENGINE_API void RemoveNetworkActor(AActor* Actor);
 
+	/** Called when an authoritative actor wants to delete a replicated subobject on the clients it was already replicated to */
+	void DeleteSubObjectOnClients(AActor* Actor, UObject* SubObject);
+
+	/** Called when an authoritative actor wants to tear off a subobject on the clients it was already replicated to */
+	void TearOffSubObjectOnClients(AActor* Actor, UObject* SubObject);
+
 	ENGINE_API virtual void NotifyActorLevelUnloaded( AActor* Actor );
 
 	ENGINE_API virtual void NotifyActorTearOff(AActor* Actor);
diff --git a/Engine/Source/Runtime/Engine/Classes/Engine/NetworkObjectList.h b/Engine/Source/Runtime/Engine/Classes/Engine/NetworkObjectList.h
--- a/Engine/Source/Runtime/Engine/Classes/Engine/NetworkObjectList.h
+++ b/Engine/Source/Runtime/Engine/Classes/Engine/NetworkObjectList.h
@@ -8,6 +8,25 @@
 class AActor;
 class FArchive;
 
+#ifndef UE_REPLICATED_OBJECT_REFCOUNTING
+	// Allows every network actor to keep track of the number of channels an individual subobject was replicated. 
+	// Not needed on clients so its compiled out to reduce memory usage.
+	#define UE_REPLICATED_OBJECT_REFCOUNTING WITH_SERVER_CODE
+#endif
+
+#ifndef DO_REPLICATED_OBJECT_CHANNELREF_CHECKS
+	// Run additional code to validate any errors in the SubObject channel refcounting at the cost of extra memory usage and cpu overhead.
+	#define DO_REPLICATED_OBJECT_CHANNELREF_CHECKS (UE_REPLICATED_OBJECT_REFCOUNTING && WITH_SERVER_CODE && !UE_BUILD_SHIPPING) // will be disabled in Test builds after a few days.
+#endif
+
+/** Indicates the status of a replicated subobject replicated by actor channels. */
+enum class ENetSubObjectStatus : uint8
+{
+	Active,  // The subobject is active and replicated
+	TearOff, // The subobject should be torn off on the clients
+	Delete,  // The subobject should be deleted on the clients
+};
+
 /**
  * Struct to store an actor pointer and any internal metadata for that actor used
  * internally by a UNetDriver.
@@ -42,6 +61,80 @@ struct FNetworkObjectInfo
 	*/
 	TSet<TWeakObjectPtr<UNetConnection>> RecentlyDormantConnections;
 
+	/** 
+	* Keeps track of the number of channels that have replicated a subobject.
+	* When the status changes from Active its expected of existing references to be gradually removed as the actor replicates itself to each connection.
+	*/
+	struct FSubObjectChannelReference
+	{
+		/** The replicated subobject */
+		TWeakObjectPtr<UObject> SubObjectPtr;
+		
+		/** Number of channels that replicated the subobject */
+		uint16 ChannelRefCount = 0;
+
+		/** Current status of the subobject */
+		ENetSubObjectStatus Status = ENetSubObjectStatus::Active;
+
+#if DO_REPLICATED_OBJECT_CHANNELREF_CHECKS
+		/** Debug array to keep track of every individual channel */
+		TArray<UObject*> RegisteredChannels;
+#endif
+		inline bool operator==(const FSubObjectChannelReference& rhs) const { return SubObjectPtr.HasSameIndexAndSerialNumber(rhs.SubObjectPtr); }
+		inline bool operator==(const TWeakObjectPtr<UObject>& rhs) const { return SubObjectPtr.HasSameIndexAndSerialNumber(rhs); }
+
+		friend uint32 GetTypeHash(const FSubObjectChannelReference& SubObjChannelRef)
+		{
+			return GetTypeHash(SubObjChannelRef.SubObjectPtr);
+		}
+
+		inline bool IsTearOff() const { return Status == ENetSubObjectStatus::TearOff; }
+		inline bool IsDelete() const  { return Status == ENetSubObjectStatus::Delete; }
+		inline bool IsActive() const  { return Status == ENetSubObjectStatus::Active; }
+
+		FSubObjectChannelReference() = default;
+		explicit FSubObjectChannelReference(const TWeakObjectPtr<UObject>& InSubObject)
+			: SubObjectPtr(InSubObject)
+			, ChannelRefCount(1)
+			, Status(ENetSubObjectStatus::Active)
+		{
+		}
+	};
+
+	struct FSubObjectChannelRefKeyFuncs : BaseKeyFuncs<FSubObjectChannelReference, const TWeakObjectPtr<UObject>&, false>
+	{
+		static KeyInitType GetSetKey(ElementInitType Element)
+		{
+			return Element.SubObjectPtr;
+		}
+
+		static bool Matches(KeyInitType A, KeyInitType B)
+		{
+			return A.HasSameIndexAndSerialNumber(B);
+		}
+
+		static uint32 GetKeyHash(KeyInitType Key)
+		{
+			// Probably should use a better has like FObjectKey but can't access the internals of the WeakPtr.
+			return GetTypeHash(Key);
+		}
+	};
+
+
+#if UE_REPLICATED_OBJECT_REFCOUNTING
+	/** The set of active replicated subobjects for an actor */
+	TSet<FSubObjectChannelReference, FSubObjectChannelRefKeyFuncs> ActiveSubObjectChannelReferences;
+
+	/** The list of replicated subobjects that need to be torn off or deleted by the channels still referencing it */
+	TArray<FSubObjectChannelReference> InvalidSubObjectChannelReferences;
+#endif
+
+	/** 
+	* This variable gets increased when a replicated subobject gets added to the invalid list.
+	* Actor channels will test against this count and when it differs will check if they have any replicated subobjects to remove.
+	*/
+	uint16 InvalidSubObjectDirtyCount = 0;
+
 	/** Is this object still pending a full net update due to clients that weren't able to replicate the actor at the time of LastNetUpdateTime */
 	uint8 bPendingNetUpdate : 1;
 
@@ -192,9 +285,39 @@ public:
 	/** Marks any actors in the given package/level active if they were fully dormant or dormant for the passed in connection */
 	void FlushDormantActors(UNetConnection* const Connection, const FName& PackageName);
 
+#if UE_REPLICATED_OBJECT_REFCOUNTING
+	/** Set the subobject to be flagged for deletion */
+	void SetSubObjectForDeletion(AActor* Actor, UObject* SubObject);
+
+	/** Set the subobject to be flagged for tear off */
+	void SetSubObjectForTearOff(AActor* Actor, UObject* SubObject);
+
+	/**
+	* Called when a channel starts replicating a subobject for the first time.
+	* Used to keep track of the number of channels having an active reference to a specific subobject
+	*/
+	void AddSubObjectChannelReference(AActor* OwnerActor, UObject* ReplicatedSubObject, UActorChannel* ReferencingChannel);
+
+	/** Called when a channel stops replicating a subobject. */
+	void RemoveSubObjectChannelReference(AActor* OwnerActor, const TWeakObjectPtr<UObject>& ReplicatedSubObject, UActorChannel* ReferencingChannel);
+
+	/** Called when multiple subobjects that were flagged torn off or delete have removed their channel reference */
+	void RemoveMultipleInvalidSubObjectChannelReference(FNetworkObjectInfo* ActorNetInfo, const TArray< TWeakObjectPtr<UObject>, TInlineAllocator<16> >& SubObjectsToRemove, UActorChannel* ReferencingChannel);
+
+	/** Called when multiple subobjects that were still considered active have removed their channel reference */
+	void RemoveMultipleActiveSubObjectChannelReference(FNetworkObjectInfo* ActorNetInfo, const TArray< TWeakObjectPtr<UObject>, TInlineAllocator<16> >& SubObjectsToRemove, UActorChannel* ReferencingChannel);
+#endif //#if UE_REPLICATED_OBJECT_REFCOUNTING
+
 private:
 	bool MarkActiveInternal(const TSharedPtr<FNetworkObjectInfo>& ObjectInfo, UNetConnection* const Connection, UNetDriver* NetDriver);
 
+#if UE_REPLICATED_OBJECT_REFCOUNTING
+	void InvalidateSubObject(AActor* Actor, UObject* SubObject, ENetSubObjectStatus InvalidStatus);
+
+	bool HandleRemoveActiveSubObjectRef(FNetworkObjectInfo* ActorNetInfo, AActor* OwnerActor, const TWeakObjectPtr<UObject>& ReplicatedSubObject, UActorChannel* ReferencingChannel);
+	bool HandleRemoveInvalidSubObjectRef(FNetworkObjectInfo* ActorNetInfo, AActor* OwnerActor, const TWeakObjectPtr<UObject>& ReplicatedSubObject, UActorChannel* ReferencingChannel);
+#endif
+
 	FNetworkObjectSet AllNetworkObjects;
 	FNetworkObjectSet ActiveNetworkObjects;
 	FNetworkObjectSet ObjectsDormantOnAllConnections;
diff --git a/Engine/Source/Runtime/Engine/Classes/GameFramework/Actor.h b/Engine/Source/Runtime/Engine/Classes/GameFramework/Actor.h
--- a/Engine/Source/Runtime/Engine/Classes/GameFramework/Actor.h
+++ b/Engine/Source/Runtime/Engine/Classes/GameFramework/Actor.h
@@ -3719,11 +3719,33 @@ public:
 	void AddReplicatedSubObject(UObject* SubObject, ELifetimeCondition NetCondition = COND_None);
 
 	/**
-	* Unregister a SubObject so it stops being replicated.
+	* Unregister a SubObject to stop replicating it's properties to clients.
+	* This does not remove or delete it from connections where it was already replicated.
+	* By default a replicated subobject gets deleted on clients when the original pointer on the authority becomes invalid.
+	* If you want to immediately remove it from client use the DeleteReplicatedSubObject or TearOffReplicatedSubObject functions instead of this one.
 	* @param SubObject The SubObject to remove
 	*/
 	void RemoveReplicatedSubObject(UObject* SubObject);
 
+	/**
+	* Stop replicating a subobject and tell actor channels to delete the replica of this subobject next time the Actor gets replicated.
+	* Note it is up to the caller to delete the local object on the authority.
+	* If you are using the legacy subobject replication method (ReplicateSubObjects() aka bReplicateUsingRegisteredSubObjectList=false) make sure the
+	* subobject doesn't get replicated there either.
+	* @param SubObject THe SubObject to delete
+	*/
+	void DestroyReplicatedSubObjectOnRemotePeers(UObject* SubObject);
+
+	/**
+	* Stop replicating a subobject and tell actor channels who spawned a replica of this subobject to release ownership over it.
+	* This means that on the remote connection the network engine will stop holding a reference to the subobject and it's up to other systems 
+	* to keep that reference active or the subobject will get garbage collected.
+	* If you are using the legacy subobject replication method (ReplicateSubObjects() aka bReplicateUsingRegisteredSubObjectList=false) make sure the
+	* subobject doesn't get replicated there either.
+	* @param SubObject The SubObject to tear off
+	*/
+	void TearOffReplicatedSubObjectOnRemotePeers(UObject* SubObject);
+
 	/**
 	* Register a SubObject that will get replicated along with the actor component owning it.
 	* The subobject needs to be manually removed from the list before it gets deleted.
diff --git a/Engine/Source/Runtime/Engine/Private/ActorReplication.cpp b/Engine/Source/Runtime/Engine/Private/ActorReplication.cpp
--- a/Engine/Source/Runtime/Engine/Private/ActorReplication.cpp
+++ b/Engine/Source/Runtime/Engine/Private/ActorReplication.cpp
@@ -14,6 +14,7 @@
 #include "Engine/BlueprintGeneratedClass.h"
 #include "Net/UnrealNetwork.h"
 #include "Engine/ActorChannel.h"
+#include "Engine/NetDriver.h"
 #include "GameFramework/GameNetworkManager.h"
 #include "NetworkingDistanceConstants.h"
 #include "PhysicsReplication.h"
@@ -629,6 +630,67 @@ void AActor::RemoveReplicatedSubObject(UObject* SubObject)
 #endif // UE_WITH_IRIS
 }
 
+
+void AActor::DestroyReplicatedSubObjectOnRemotePeers(UObject* SubObject)
+{
+	check(SubObject);
+
+	if (!HasAuthority())
+	{
+		// Only the authority can call this.
+		return;
+	}
+
+	UE_LOG(LogNetSubObject, Verbose, TEXT("%s (0x%p) requested to Delete replicated subobject on clients %s (0x%p)"), *GetName(), this, *SubObject->GetName(), SubObject);
+
+	if (FWorldContext* const Context = GEngine->GetWorldContextFromWorld(GetWorld()))
+	{
+		for (const FNamedNetDriver& Driver : Context->ActiveNetDrivers)
+		{
+			if (Driver.NetDriver)
+			{
+				Driver.NetDriver->DeleteSubObjectOnClients(this, SubObject);
+			}
+		}
+	}
+
+	if (IsUsingRegisteredSubObjectList())
+	{
+		RemoveReplicatedSubObject(SubObject);
+	}
+}
+
+
+void AActor::TearOffReplicatedSubObjectOnRemotePeers(UObject* SubObject)
+{
+	check(SubObject);
+
+	if (!HasAuthority())
+	{
+		// Only the authority can call this.
+		return;
+	}
+
+	UE_LOG(LogNetSubObject, Verbose, TEXT("%s (0x%p) requested to TearOff replicated subobject on clients %s (0x%p)"), *GetName(), this, *SubObject->GetName(), SubObject);
+
+	if (FWorldContext* const Context = GEngine->GetWorldContextFromWorld(GetWorld()))
+	{
+		for (const FNamedNetDriver& Driver : Context->ActiveNetDrivers)
+		{
+			if (Driver.NetDriver)
+			{
+				Driver.NetDriver->TearOffSubObjectOnClients(this, SubObject);
+			}
+		}
+	}
+
+	if (IsUsingRegisteredSubObjectList())
+	{
+		RemoveReplicatedSubObject(SubObject);
+	}
+
+}
+
 void AActor::AddActorComponentReplicatedSubObject(UActorComponent* OwnerComponent, UObject* SubObject, ELifetimeCondition NetCondition)
 {
 	check(IsValid(OwnerComponent));
diff --git a/Engine/Source/Runtime/Engine/Private/Components/ActorComponent.cpp b/Engine/Source/Runtime/Engine/Private/Components/ActorComponent.cpp
--- a/Engine/Source/Runtime/Engine/Private/Components/ActorComponent.cpp
+++ b/Engine/Source/Runtime/Engine/Private/Components/ActorComponent.cpp
@@ -2127,6 +2127,67 @@ void UActorComponent::RemoveReplicatedSubObject(UObject* SubObject)
 	}
 }
 
+void UActorComponent::DestroyReplicatedSubObjectOnRemotePeers(UObject* SubObject)
+{
+	if (AActor* MyOwner = GetOwner())
+	{
+		if (!MyOwner->HasAuthority())
+		{
+			// Only the authority can call this.
+			return;
+		}
+
+		UE_LOG(LogNetSubObject, Verbose, TEXT("%s::%s (0x%p) requested to Delete replicated subobject on clients %s (0x%p)"), *MyOwner->GetName(), *GetName(), this, *SubObject->GetName(), SubObject);
+
+		if (FWorldContext* const Context = GEngine->GetWorldContextFromWorld(GetWorld()))
+		{
+			for (const FNamedNetDriver& Driver : Context->ActiveNetDrivers)
+			{
+				if (Driver.NetDriver)
+				{
+					Driver.NetDriver->DeleteSubObjectOnClients(MyOwner, SubObject);
+				}
+			}
+		}
+
+
+		if (IsUsingRegisteredSubObjectList())
+		{
+			MyOwner->RemoveActorComponentReplicatedSubObject(this, SubObject);
+		}
+	}
+}
+
+void UActorComponent::TearOffReplicatedSubObjectOnRemotePeers(UObject* SubObject)
+{
+	if (AActor* MyOwner=GetOwner())
+	{
+		if (!MyOwner->HasAuthority())
+		{
+			// Only the authority can call this.
+			return;
+		}
+
+		UE_LOG(LogNetSubObject, Verbose, TEXT("%s::%s (0x%p) requested to TearOff replicated subobject on clients %s (0x%p)"), *MyOwner->GetName(), *GetName(), this, *SubObject->GetName(), SubObject);
+	
+		if (FWorldContext* const Context = GEngine->GetWorldContextFromWorld(GetWorld()))
+		{
+			for (const FNamedNetDriver& Driver : Context->ActiveNetDrivers)
+			{
+				if (Driver.NetDriver)
+				{
+					Driver.NetDriver->TearOffSubObjectOnClients(MyOwner, SubObject);
+				}
+			}
+		}
+
+		if (IsUsingRegisteredSubObjectList())
+		{
+			MyOwner->RemoveActorComponentReplicatedSubObject(this, SubObject);
+		}
+	}
+}
+
 bool UActorComponent::IsReplicatedSubObjectRegistered(const UObject* SubObject) const
 {
 	if (AActor* MyOwner=GetOwner())
diff --git a/Engine/Source/Runtime/Engine/Private/DataChannel.cpp b/Engine/Source/Runtime/Engine/Private/DataChannel.cpp
--- a/Engine/Source/Runtime/Engine/Private/DataChannel.cpp
+++ b/Engine/Source/Runtime/Engine/Private/DataChannel.cpp
@@ -12,6 +12,7 @@
 #include "Misc/ScopeExit.h"
 #include "UObject/UObjectIterator.h"
 #include "UObject/ObjectKey.h"
+#include "UObject/CoreNetTypes.h"
 #include "EngineStats.h"
 #include "EngineGlobals.h"
 #include "Engine/Engine.h"
@@ -2179,9 +2180,9 @@ void UActorChannel::SetClosingFlag()
 
 void UActorChannel::ReleaseReferences(bool bKeepReplicators)
 {
-	Actor = nullptr;
 	CleanupReplicators(bKeepReplicators);
 	GetCreatedSubObjects().Empty();
+	Actor = nullptr;
 }
 
 void UActorChannel::BreakAndReleaseReferences()
@@ -2264,7 +2265,13 @@ void UActorChannel::CleanupReplicators(const bool bKeepReplicators)
 		}
 		else
 		{
-
+#if UE_REPLICATED_OBJECT_REFCOUNTING
+			// Ignore the main actor's replicator
+			if (Actor != ObjectReplicatorRef->ObjectPtr)
+			{
+				Connection->Driver->GetNetworkObjectList().RemoveSubObjectChannelReference(Actor, ObjectReplicatorRef->GetWeakObjectPtr(), this);
+			}
+#endif
 			ObjectReplicatorRef->CleanUp();
 		}
 	}
@@ -3108,6 +3115,9 @@ void UActorChannel::ProcessBunch( FInBunch & Bunch )
 		{
 			if (!Connection->Driver->IsServer())
 			{
+#if UE_REPLICATED_OBJECT_REFCOUNTING
+				Connection->Driver->GetNetworkObjectList().RemoveSubObjectChannelReference(Actor, ObjectReplicator.Get().GetWeakObjectPtr(), this);
+#endif
 				RepComp.RemoveCurrent(); // This should cause the replicator to be cleaned up as there should be no outstandings refs. 
 			}
 			continue;
@@ -3440,44 +3450,7 @@ int64 UActorChannel::ReplicateActor()
 		}
 
 		{
-			QUICK_SCOPE_CYCLE_COUNTER(STAT_NetDeletedSubObjects);
-
-			// Look for deleted subobjects
-			FObjectReplicator* LocalActorReplicator = ActorReplicator.Get();
-			for (auto RepComp = ReplicationMap.CreateIterator(); RepComp; ++RepComp)
-			{
-				TSharedRef<FObjectReplicator>& LocalReplicator = RepComp.Value();
-
-				if (!LocalReplicator->GetWeakObjectPtr().IsValid())
-				{
-					if (LocalReplicator->ObjectNetGUID.IsValid())
-					{
-						UE_NET_TRACE_SCOPE(ContentBlockForSubObjectDelete, Bunch, GetTraceCollector(Bunch), ENetTraceVerbosity::Trace);
-						UE_NET_TRACE_OBJECT_SCOPE(LocalReplicator->ObjectNetGUID, Bunch, GetTraceCollector(Bunch), ENetTraceVerbosity::Trace);
-
-						// Write a deletion content header:
-						WriteContentBlockForSubObjectDelete(Bunch, LocalReplicator->ObjectNetGUID);
-
-						bWroteSomethingImportant = true;
-						Bunch.bReliable = true;
-					}
-					else
-					{
-						UE_LOG(LogNetTraffic, Error, TEXT("Unable to write subobject delete for (%s), object replicator has invalid NetGUID"), *GetPathNameSafe(Actor));
-					}
-
-					// The only way this case would be possible is if someone tried destroying the Actor as a part of
-					// a Subobject's Pre / Post replication, during Replicate Subobjects, or OnSerializeNewActor.
-					// All of those are bad.
-					if (!ensureMsgf(LocalActorReplicator != &LocalReplicator.Get(), TEXT("UActorChannel::ReplicateActor: Actor was deleting during replication: %s"), *Describe()))
-					{
-						ActorReplicator.Reset();
-					}
-
-					LocalReplicator->CleanUp();
-					RepComp.RemoveCurrent();
-				}
-			}
+			bWroteSomethingImportant |= UpdateDeletedSubObjects(Bunch);
 		}
 	}
 
@@ -3586,6 +3559,109 @@ bool UActorChannel::DoSubObjectReplication(FOutBunch& Bunch, FReplicationFlags&
 	return bWroteSomethingImportant;
 }
 
+bool UActorChannel::UpdateDeletedSubObjects(FOutBunch& Bunch)
+{
+	QUICK_SCOPE_CYCLE_COUNTER(STAT_NetDeletedSubObjects);
+
+	bool bWroteSomethingImportant = false;
+
+	auto DeleteSubObject = [this, &bWroteSomethingImportant, &Bunch](const TSharedRef<FObjectReplicator>& SubObjectReplicator, const TWeakObjectPtr<UObject>& ObjectPtrToRemove, ESubObjectDeleteFlag DeleteFlag)
+	{
+		if (SubObjectReplicator->ObjectNetGUID.IsValid())
+		{
+			UE_NET_TRACE_SCOPE(ContentBlockForSubObjectDelete, Bunch, GetTraceCollector(Bunch), ENetTraceVerbosity::Trace);
+			UE_NET_TRACE_OBJECT_SCOPE(SubObjectReplicator->ObjectNetGUID, Bunch, GetTraceCollector(Bunch), ENetTraceVerbosity::Trace);
+
+			UE_LOG(LogNetSubObject, Verbose, TEXT("NetSubObject: Sending request to %s %s::%s (0x%p) NetGUID %s"), ToString(DeleteFlag), *Actor->GetName(), *GetNameSafe(ObjectPtrToRemove.GetEvenIfUnreachable()), ObjectPtrToRemove.GetEvenIfUnreachable(), *SubObjectReplicator->ObjectNetGUID.ToString());
+
+			// Write a deletion content header:
+			WriteContentBlockForSubObjectDelete(Bunch, SubObjectReplicator->ObjectNetGUID, DeleteFlag);
+
+			bWroteSomethingImportant = true;
+			Bunch.bReliable = true;
+		}
+		else
+		{
+			UE_LOG(LogNetTraffic, Error, TEXT("Unable to write subobject delete for %s::%s (0x%p), object replicator has invalid NetGUID"), *GetPathNameSafe(Actor), *GetNameSafe(ObjectPtrToRemove.GetEvenIfUnreachable()), ObjectPtrToRemove.GetEvenIfUnreachable());
+		}
+
+		SubObjectReplicator->CleanUp();
+	};
+
+#if UE_REPLICATED_OBJECT_REFCOUNTING
+	TArray< TWeakObjectPtr<UObject>, TInlineAllocator<16> > SubObjectsRemoved;
+
+	FNetworkObjectInfo* ActorInfo = Connection->Driver->FindNetworkObjectInfo(Actor);
+
+	// Check if the dirty count is different from our last update
+	if (ActorInfo)
+	{
+		if (ChannelSubObjectDirtyCount != ActorInfo->InvalidSubObjectDirtyCount)
+		{
+			ChannelSubObjectDirtyCount = ActorInfo->InvalidSubObjectDirtyCount;
+
+			for (const FNetworkObjectInfo::FSubObjectChannelReference& SubObjectRef : ActorInfo->InvalidSubObjectChannelReferences)
+			{
+				ensure(SubObjectRef.Status != ENetSubObjectStatus::Active);
+
+				// The TearOff won't be sent if the Object pointer is fully deleted once we get here.
+				// This will be fixed by converting the ReplicationMap to stop using raw pointers as the key
+				if (UObject* ObjectToRemove = SubObjectRef.SubObjectPtr.GetEvenIfUnreachable())
+				{
+					if (TSharedRef<FObjectReplicator>* SubObjectReplicator = ReplicationMap.Find(ObjectToRemove))
+					{
+						const ESubObjectDeleteFlag DeleteFlag = SubObjectRef.IsTearOff() ? ESubObjectDeleteFlag::TearOff : ESubObjectDeleteFlag::ForceDelete;
+						
+						DeleteSubObject(*SubObjectReplicator, SubObjectRef.SubObjectPtr, DeleteFlag);
+						ReplicationMap.Remove(ObjectToRemove);
+					}
+				}
+
+				SubObjectsRemoved.Add(SubObjectRef.SubObjectPtr);
+			}
+
+			Connection->Driver->GetNetworkObjectList().RemoveMultipleInvalidSubObjectChannelReference(ActorInfo, SubObjectsRemoved, this);
+		}
+
+		SubObjectsRemoved.Reset();
+	}
+#endif //#if UE_REPLICATED_OBJECT_REFCOUNTING
+
+	// Look for deleted subobjects
+	for (auto RepComp = ReplicationMap.CreateIterator(); RepComp; ++RepComp)
+	{
+		TSharedRef<FObjectReplicator>& SubObjectReplicator = RepComp.Value();
+
+		const TWeakObjectPtr<UObject> WeakObjPtr = SubObjectReplicator->GetWeakObjectPtr();
+		if (!WeakObjPtr.IsValid())
+		{
+			// The only way this case would be possible is if someone tried destroying the Actor as a part of
+			// a Subobject's Pre / Post replication, during Replicate Subobjects, or OnSerializeNewActor.
+			// All of those are bad.
+			if (!ensureMsgf(ActorReplicator.Get() != &SubObjectReplicator.Get(), TEXT("UActorChannel::ReplicateActor: Actor was deleting during replication: %s"), *Describe()))
+			{
+				ActorReplicator.Reset();
+			}
+
+#if UE_REPLICATED_OBJECT_REFCOUNTING
+			SubObjectsRemoved.Add(WeakObjPtr);
+#endif
+
+			DeleteSubObject(SubObjectReplicator, WeakObjPtr, ESubObjectDeleteFlag::Destroyed);
+			RepComp.RemoveCurrent();
+		}
+	}
+
+#if UE_REPLICATED_OBJECT_REFCOUNTING
+	if (SubObjectsRemoved.Num() > 0)
+	{
+		Connection->Driver->GetNetworkObjectList().RemoveMultipleActiveSubObjectChannelReference(ActorInfo, SubObjectsRemoved, this);
+	}
+#endif
+
+	return bWroteSomethingImportant;
+}
+
 bool UActorChannel::CanSubObjectReplicateToClient(ELifetimeCondition NetCondition, FObjectKey SubObjectKey, const TStaticBitArray<COND_Max>& ConditionMap) const
 {
 	bool bCanReplicateToClient = ConditionMap[NetCondition];
@@ -4133,7 +4209,12 @@ void UActorChannel::WriteContentBlockHeader( UObject* Obj, FNetBitWriter &Bunch,
 	NETWORK_PROFILER(GNetworkProfiler.TrackBeginContentBlock(Obj, Bunch.GetNumBits() - NumStartingBits, Connection));
 }
 
-void UActorChannel::WriteContentBlockForSubObjectDelete( FOutBunch & Bunch, FNetworkGUID & GuidToDelete )
+void UActorChannel::WriteContentBlockForSubObjectDelete(FOutBunch& Bunch, FNetworkGUID& GuidToDelete)
+{
+	WriteContentBlockForSubObjectDelete(Bunch, GuidToDelete, ESubObjectDeleteFlag::Destroyed);
+}
+
+void UActorChannel::WriteContentBlockForSubObjectDelete( FOutBunch& Bunch, FNetworkGUID& GuidToDelete, ESubObjectDeleteFlag Flag)
 {
 	check( Connection->Driver->IsServer() );
 
@@ -4158,8 +4239,8 @@ void UActorChannel::WriteContentBlockForSubObjectDelete( FOutBunch & Bunch, FNet
 	// Send that this is a delete message
 	Bunch.WriteBit( 1 );
 
-	uint8 DeleteFlags = 0; // flags will be defined later
-	Bunch << DeleteFlags;
+	uint8 DeleteFlag = static_cast<uint8>(Flag);
+	Bunch << DeleteFlag;
 
 	NET_CHECKSUM(Bunch); // Matches checksum in UPackageMapClient::ReadContentBlockHeader - DeleteSubObject message
 
@@ -4336,6 +4417,7 @@ UObject* UActorChannel::ReadContentBlockHeader(FInBunch& Bunch, bool& bObjectDel
 
 	bool bDeleteSubObject = false;
 	bool bSerializeClass = true;
+	ESubObjectDeleteFlag DeleteFlag = ESubObjectDeleteFlag::Destroyed;
 
 	if (Bunch.EngineNetVer() >= HISTORY_SUBOBJECT_DESTROY_FLAG)
 	{
@@ -4346,8 +4428,9 @@ UObject* UActorChannel::ReadContentBlockHeader(FInBunch& Bunch, bool& bObjectDel
 			bDeleteSubObject = true;
 			bSerializeClass = false;
 
-			uint8 DestroyFlags = 0;
-			Bunch << DestroyFlags;  // Destroy flags to be defined later
+			uint8 ReceivedDestroyFlag;
+			Bunch << ReceivedDestroyFlag;
+			DeleteFlag = static_cast<ESubObjectDeleteFlag>(ReceivedDestroyFlag);
 
 			NET_CHECKSUM(Bunch);
 		}
@@ -4386,11 +4469,17 @@ UObject* UActorChannel::ReadContentBlockHeader(FInBunch& Bunch, bool& bObjectDel
 				Connection->Driver->NotifySubObjectDestroyed(SubObj);
 			}
 
-			Actor->OnSubobjectDestroyFromReplication(SubObj);
+			UE_LOG(LogNetSubObject, Verbose, TEXT("NetSubObject: Client received request to %s %s::%s (0x%p) NetGuid %s"), ToString(DeleteFlag), *Actor->GetName(), *GetNameSafe(SubObj), SubObj, *NetGUID.ToString());
+
+			if (DeleteFlag != ESubObjectDeleteFlag::TearOff)
+			{
+				Actor->OnSubobjectDestroyFromReplication(SubObj);
 
-			SubObj->PreDestroyFromReplication();
-			SubObj->MarkAsGarbage();
+				SubObj->PreDestroyFromReplication();
+				SubObj->MarkAsGarbage();
+			}
 		}
+
 		bObjectDeleted = true;
 		return nullptr;
 	}
@@ -4812,10 +4901,14 @@ TSharedRef<FObjectReplicator>* UActorChannel::FindReplicator(UObject* Obj, bool*
 
 	// First, try to find it on the channel replication map
 	TSharedRef<FObjectReplicator>* ReplicatorRefPtr = ReplicationMap.Find( Obj );
-	if (ReplicatorRefPtr != nullptr)
+	if (ReplicatorRefPtr)
 	{
 		if (!ReplicatorRefPtr->Get().GetWeakObjectPtr().IsValid())
 		{
+#if UE_REPLICATED_OBJECT_REFCOUNTING
+			Connection->Driver->GetNetworkObjectList().RemoveSubObjectChannelReference(Actor, ReplicatorRefPtr->Get().GetWeakObjectPtr(), this);
+#endif
+
 			ReplicatorRefPtr = nullptr;
 			ReplicationMap.Remove(Obj);
 
@@ -4859,6 +4952,17 @@ TSharedRef<FObjectReplicator>& UActorChannel::CreateReplicator(UObject* Obj, boo
 	// Add to the replication map
 	TSharedRef<FObjectReplicator>& NewRef = ReplicationMap.Add(Obj, NewReplicator.ToSharedRef());
 
+#if UE_REPLICATED_OBJECT_REFCOUNTING
+	// If this is a subobject (and not the main actor) add a reference
+	if (Obj != Actor)
+	{
+		if (UNetDriver* NetDriver = Connection->GetDriver())
+		{
+			NetDriver->GetNetworkObjectList().AddSubObjectChannelReference(Actor, Obj, this);
+		}
+	}
+#endif
+
 	// Start replicating with this replicator
 	NewRef->StartReplicating(this);
 	return NewRef;
@@ -4935,6 +5039,7 @@ void UActorChannel::AddedToChannelPool()
 	bActorIsPendingKill = false;
 	bSkipRoleSwap = false;
 	bClearRecentActorRefs = true;
+	ChannelSubObjectDirtyCount = 0;
 	QueuedBunchStartTime = 0;
 	bSuppressQueuedBunchWarningsDueToHitches = false;
 #if !UE_BUILD_SHIPPING
@@ -5032,6 +5137,17 @@ bool UActorChannel::WriteSubObjectInBunch(UObject* Obj, FOutBunch& Bunch, FRepli
 	return bWroteSomething;
 }
 
+const TCHAR* UActorChannel::ToString(UActorChannel::ESubObjectDeleteFlag DeleteFlag)
+{
+	switch (DeleteFlag)
+	{
+	case UActorChannel::ESubObjectDeleteFlag::Destroyed:	return TEXT("Destroyed"); break;
+	case UActorChannel::ESubObjectDeleteFlag::TearOff:		return TEXT("TearOff"); break;
+	case UActorChannel::ESubObjectDeleteFlag::ForceDelete:	return TEXT("ForceDelete"); break;
+	}
+	return TEXT("Missing");
+}
+
 //------------------------------------------------------
 
 static void	DebugNetGUIDs( UWorld* InWorld )
diff --git a/Engine/Source/Runtime/Engine/Private/NetDriver.cpp b/Engine/Source/Runtime/Engine/Private/NetDriver.cpp
--- a/Engine/Source/Runtime/Engine/Private/NetDriver.cpp
+++ b/Engine/Source/Runtime/Engine/Private/NetDriver.cpp
@@ -3367,6 +3367,20 @@ void UNetDriver::RemoveNetworkActor(AActor* Actor)
 	}
 }
 
+void UNetDriver::DeleteSubObjectOnClients(AActor* Actor, UObject* SubObject)
+{
+#if UE_REPLICATED_OBJECT_REFCOUNTING
+	NetworkObjects->SetSubObjectForDeletion(Actor, SubObject);
+#endif
+}
+
+void UNetDriver::TearOffSubObjectOnClients(AActor* Actor, UObject* SubObject)
+{
+#if UE_REPLICATED_OBJECT_REFCOUNTING
+	NetworkObjects->SetSubObjectForTearOff(Actor, SubObject);
+#endif
+}
+
 void UNetDriver::NotifyActorRenamed(AActor* ThisActor, FName PreviousName)
 {
 	LLM_SCOPE_BYTAG(NetDriver);
diff --git a/Engine/Source/Runtime/Engine/Private/NetworkObjectList.cpp b/Engine/Source/Runtime/Engine/Private/NetworkObjectList.cpp
--- a/Engine/Source/Runtime/Engine/Private/NetworkObjectList.cpp
+++ b/Engine/Source/Runtime/Engine/Private/NetworkObjectList.cpp
@@ -6,6 +6,7 @@
 #include "EngineUtils.h"
 #include "Serialization/Archive.h"
 #include "Net/NetworkGranularMemoryLogging.h"
+#include "Engine/ActorChannel.h"
 #include "UObject/Package.h"
 
 namespace UE::Net::Private
@@ -16,6 +17,17 @@ namespace UE::Net::Private
 		bTrackDormantObjectsByLevel,
 		TEXT("When true, network object list will maintain a set of dormant actors per connnection per level."),
 		ECVF_Default);
+
+	static const TCHAR* LexToString(ENetSubObjectStatus SubObjectStatus)
+	{
+		switch (SubObjectStatus)
+		{
+		case ENetSubObjectStatus::Active: return TEXT("Active");
+		case ENetSubObjectStatus::TearOff: return TEXT("TearOff");
+		case ENetSubObjectStatus::Delete: return TEXT("Delete");
+		}
+		return TEXT("Missing");
+	}
 }
 
 void FNetworkObjectList::AddInitialObjects(UWorld* const World, UNetDriver* NetDriver)
@@ -467,6 +479,10 @@ void FNetworkObjectInfo::CountBytes(FArchive& Ar) const
 
 	GRANULAR_NETWORK_MEMORY_TRACKING_TRACK("DormantConnections", DormantConnections.CountBytes(Ar));
 	GRANULAR_NETWORK_MEMORY_TRACKING_TRACK("RecentlyDormantConnections", RecentlyDormantConnections.CountBytes(Ar));
+#if UE_REPLICATED_OBJECT_REFCOUNTING
+	GRANULAR_NETWORK_MEMORY_TRACKING_TRACK("ActiveSubObjectChannelReferences", ActiveSubObjectChannelReferences.CountBytes(Ar));
+	GRANULAR_NETWORK_MEMORY_TRACKING_TRACK("InvalidSubObjectChannelReferences", InvalidSubObjectChannelReferences.CountBytes(Ar));
+#endif
 }
 
 void FNetworkObjectList::CountBytes(FArchive& Ar) const
@@ -537,4 +553,188 @@ void FNetworkObjectList::FlushDormantActors(UNetConnection* const Connection, co
 	{
 		MarkActiveInternal(ActorInfo, Connection, Connection->Driver);
 	}
-}
\ No newline at end of file
+}
+#if UE_REPLICATED_OBJECT_REFCOUNTING
+void FNetworkObjectList::SetSubObjectForDeletion(AActor* Actor, UObject* SubObject)
+{
+	InvalidateSubObject(Actor, SubObject, ENetSubObjectStatus::Delete);
+}
+
+void FNetworkObjectList::SetSubObjectForTearOff(AActor* Actor, UObject* SubObject)
+{
+	InvalidateSubObject(Actor, SubObject, ENetSubObjectStatus::TearOff);
+}
+
+void FNetworkObjectList::InvalidateSubObject(AActor* Actor, UObject* SubObject, ENetSubObjectStatus InvalidStatus)
+{
+	ensure(InvalidStatus != ENetSubObjectStatus::Active);
+
+	if (TSharedPtr<FNetworkObjectInfo>* InfoPtr = AllNetworkObjects.Find(Actor))
+	{
+		FNetworkObjectInfo* ActorNetInfo = InfoPtr->Get();
+
+		const TWeakObjectPtr<UObject> SubObjectPtr = SubObject;
+		const FSetElementId FoundId = ActorNetInfo->ActiveSubObjectChannelReferences.FindId(SubObjectPtr);
+
+		if (FoundId.IsValidId())
+		{
+			// Flag its new state
+			ActorNetInfo->ActiveSubObjectChannelReferences[FoundId].Status = InvalidStatus;
+
+			// Move the reference to the destroyed list.
+			ActorNetInfo->InvalidSubObjectChannelReferences.Emplace(MoveTemp(ActorNetInfo->ActiveSubObjectChannelReferences[FoundId]));
+			ActorNetInfo->ActiveSubObjectChannelReferences.Remove(FoundId);
+
+			// Increase dirty count so channels can refresh this list
+			ActorNetInfo->InvalidSubObjectDirtyCount = FMath::Max(ActorNetInfo->InvalidSubObjectDirtyCount + 1, 1); // Wrap around to 1 because 0 is the default actor channel value
+		}
+	}
+}
+
+void FNetworkObjectList::AddSubObjectChannelReference(AActor* OwnerActor, UObject* ReplicatedSubObject, UActorChannel* ReferencingChannel)
+{
+	TSharedPtr<FNetworkObjectInfo>* InfoPtr = AllNetworkObjects.Find(OwnerActor);
+
+	if (ensureMsgf(InfoPtr, TEXT("The channel %s replicated %s (0x%p) without having any network objectinfo for %s"), *ReferencingChannel->Describe(), *GetNameSafe(ReplicatedSubObject), ReplicatedSubObject, *GetNameSafe(OwnerActor)))
+	{
+		FNetworkObjectInfo* ActorNetInfo = InfoPtr->Get();
+
+		const TWeakObjectPtr<UObject> SubObjectPtr = ReplicatedSubObject;
+
+#if DO_REPLICATED_OBJECT_CHANNELREF_CHECKS
+		// Make sure the object wasn't previously put in the invalid list  
+		if (FNetworkObjectInfo::FSubObjectChannelReference* DestroyedObjectRef = ActorNetInfo->InvalidSubObjectChannelReferences.FindByKey(SubObjectPtr))
+		{
+			ensureMsgf(DestroyedObjectRef == nullptr, TEXT("SubObject %s (0x%p) owned by %s was replicated by channel (0x%p) %s after it was set to %s"), *GetNameSafe(ReplicatedSubObject), ReplicatedSubObject, *GetNameSafe(OwnerActor), ReferencingChannel , *ReferencingChannel->Describe(), UE::Net::Private::LexToString(DestroyedObjectRef->Status));
+		}
+#endif
+
+		if (FNetworkObjectInfo::FSubObjectChannelReference* ChannelReferences = ActorNetInfo->ActiveSubObjectChannelReferences.Find(SubObjectPtr))
+		{
+#if DO_REPLICATED_OBJECT_CHANNELREF_CHECKS
+			checkf(!ChannelReferences->RegisteredChannels.Contains(ReferencingChannel), TEXT("SubObject %s (0x%p) owned by %s was already referenced by channel %s"), *GetNameSafe(ReplicatedSubObject), ReplicatedSubObject, *GetNameSafe(OwnerActor), *ReferencingChannel->Describe());
+			ChannelReferences->RegisteredChannels.Add(ReferencingChannel);
+#endif
+
+			ChannelReferences->ChannelRefCount++;
+			check(ChannelReferences->ChannelRefCount < MAX_uint16);
+
+			UE_LOG(LogNetSubObject, Verbose, TEXT("Adding ChannelRef (%d) to %s (0x%p) owned by %s for channel %s"), ChannelReferences->ChannelRefCount, *GetNameSafe(ReplicatedSubObject), ReplicatedSubObject,  *GetNameSafe(OwnerActor), *ReferencingChannel->Describe());
+		}
+		else
+		{
+			bool bWasSet = false;
+			FSetElementId ElementId = ActorNetInfo->ActiveSubObjectChannelReferences.Emplace(FNetworkObjectInfo::FSubObjectChannelReference(SubObjectPtr), &bWasSet);
+
+			UE_LOG(LogNetSubObject, Verbose, TEXT("Adding ChannelRef (new) to %s (0x%p) owned by %s for channel %s"), *GetNameSafe(ReplicatedSubObject), ReplicatedSubObject, *GetNameSafe(OwnerActor), *ReferencingChannel->Describe());
+
+#if DO_REPLICATED_OBJECT_CHANNELREF_CHECKS
+			ActorNetInfo->ActiveSubObjectChannelReferences[ElementId].RegisteredChannels.Add(ReferencingChannel);
+#endif
+		}
+	}
+}
+
+void FNetworkObjectList::RemoveMultipleInvalidSubObjectChannelReference(FNetworkObjectInfo* ActorNetInfo, const TArray< TWeakObjectPtr<UObject>, TInlineAllocator<16> >& SubObjectsToRemove, UActorChannel* ReferencingChannel)
+{
+	for (const TWeakObjectPtr<UObject>& ReplicatedSubObject : SubObjectsToRemove)
+	{
+		HandleRemoveInvalidSubObjectRef(ActorNetInfo, ActorNetInfo->Actor, ReplicatedSubObject, ReferencingChannel);
+	}
+}
+
+void FNetworkObjectList::RemoveMultipleActiveSubObjectChannelReference(FNetworkObjectInfo* ActorNetInfo, const TArray< TWeakObjectPtr<UObject>, TInlineAllocator<16> >& SubObjectsToRemove, UActorChannel* ReferencingChannel)
+{
+	for (const TWeakObjectPtr<UObject>& ReplicatedSubObject : SubObjectsToRemove)
+	{
+		bool bWasRemoved = HandleRemoveActiveSubObjectRef(ActorNetInfo, ActorNetInfo->Actor, ReplicatedSubObject, ReferencingChannel);
+
+		// If somehow the object wasn't in the active list, check if it was in the invalid list just to be sure.
+		// It's possible for an object set to be deleted to have it's uobject ptr become innacessible before we replicate it's owner.
+		// Calling this here would act as a safeguard to ensure its reference is removed even if the tear off/force delete command was never sent.
+		if (!bWasRemoved)
+		{
+			HandleRemoveInvalidSubObjectRef(ActorNetInfo, ActorNetInfo->Actor, ReplicatedSubObject, ReferencingChannel);
+		}
+	}
+}
+
+void FNetworkObjectList::RemoveSubObjectChannelReference(AActor* OwnerActor, const TWeakObjectPtr<UObject>& ReplicatedSubObject, UActorChannel* ReferencingChannel)
+{
+	check(ReplicatedSubObject.IsExplicitlyNull() == false);
+
+	TSharedPtr<FNetworkObjectInfo>* InfoPtr = AllNetworkObjects.Find(OwnerActor);
+	if (InfoPtr)
+	{
+		FNetworkObjectInfo* ActorNetInfo = InfoPtr->Get();
+
+		// Look in the active set
+		bool bWasRemoved = HandleRemoveActiveSubObjectRef(ActorNetInfo, OwnerActor, ReplicatedSubObject, ReferencingChannel);
+
+		// Look in the destroyed array
+		if (!bWasRemoved)
+		{
+			bWasRemoved = HandleRemoveInvalidSubObjectRef(ActorNetInfo, OwnerActor, ReplicatedSubObject, ReferencingChannel);
+		}
+	}
+}
+
+bool FNetworkObjectList::HandleRemoveActiveSubObjectRef(FNetworkObjectInfo* ActorNetInfo, AActor* OwnerActor, const TWeakObjectPtr<UObject>& SubObjectPtr, UActorChannel* ReferencingChannel)
+{
+	FSetElementId FoundId = ActorNetInfo->ActiveSubObjectChannelReferences.FindId(SubObjectPtr);
+	if (FoundId.IsValidId())
+	{
+		FNetworkObjectInfo::FSubObjectChannelReference& ActiveReference = ActorNetInfo->ActiveSubObjectChannelReferences[FoundId];
+
+		check(ActiveReference.SubObjectPtr == SubObjectPtr);
+
+		ActiveReference.ChannelRefCount--;
+
+		UE_LOG(LogNetSubObject, Verbose, TEXT("Removed ACTIVE ChannelRef (%d) for %s (0x%p) owned by %s for channel %s"), ActiveReference.ChannelRefCount, *GetNameSafe(SubObjectPtr.GetEvenIfUnreachable()), SubObjectPtr.GetEvenIfUnreachable(), *GetNameSafe(OwnerActor), *ReferencingChannel->Describe());
+
+#if DO_REPLICATED_OBJECT_CHANNELREF_CHECKS
+		int32 Removed = ActiveReference.RegisteredChannels.RemoveSingleSwap(ReferencingChannel);
+		checkf(Removed > 0, TEXT("Removed ACTIVE ref for Subobject %s (0x%p) owned by %s but it was never referenced by the channel %s"), *GetNameSafe(SubObjectPtr.GetEvenIfUnreachable()), SubObjectPtr.GetEvenIfUnreachable(), *GetNameSafe(OwnerActor), *ReferencingChannel->Describe());
+#endif
+
+		if (ActiveReference.ChannelRefCount == 0)
+		{
+			ActorNetInfo->ActiveSubObjectChannelReferences.Remove(FoundId);
+		}
+
+		return true;
+	}
+	else
+	{
+		check(false);
+	}
+
+	return false;
+}
+
+bool FNetworkObjectList::HandleRemoveInvalidSubObjectRef(FNetworkObjectInfo* ActorNetInfo, AActor* OwnerActor, const TWeakObjectPtr<UObject>& SubObjectPtr, UActorChannel* ReferencingChannel)
+{
+	const int32 Index = ActorNetInfo->InvalidSubObjectChannelReferences.IndexOfByKey(SubObjectPtr);
+	if (Index != INDEX_NONE)
+	{
+		FNetworkObjectInfo::FSubObjectChannelReference& InvalidReference = ActorNetInfo->InvalidSubObjectChannelReferences[Index];
+		InvalidReference.ChannelRefCount--;
+
+		UE_LOG(LogNetSubObject, Verbose, TEXT("Removed INVALID ChannelRef (%d) for %s (0x%p) owned by %s for channel %s"), InvalidReference.ChannelRefCount,*GetNameSafe(SubObjectPtr.GetEvenIfUnreachable()), SubObjectPtr.GetEvenIfUnreachable(), *GetNameSafe(OwnerActor), *ReferencingChannel->Describe());
+
+#if DO_REPLICATED_OBJECT_CHANNELREF_CHECKS
+		int32 Removed = InvalidReference.RegisteredChannels.RemoveSingleSwap(ReferencingChannel);
+		checkf(Removed > 0, TEXT("Removed INVALID ref for Subobject %s (0x%p) owned by %s but it was never referenced by the channel %s"), *GetNameSafe(SubObjectPtr.GetEvenIfUnreachable()), SubObjectPtr.GetEvenIfUnreachable(), *GetNameSafe(OwnerActor), *ReferencingChannel->Describe());
+#endif
+		if (InvalidReference.ChannelRefCount == 0)
+		{
+			ActorNetInfo->InvalidSubObjectChannelReferences.RemoveAtSwap(Index);
+		}
+
+		return true;
+	}
+
+	return false;
+}
+
+#endif //#if UE_REPLICATED_OBJECT_REFCOUNTING
\ No newline at end of file
