Make Input Device Properties a blueprint type

#rnx

diff --git a/Engine/Source/Runtime/Engine/Classes/GameFramework/InputDeviceProperties.h b/Engine/Source/Runtime/Engine/Classes/GameFramework/InputDeviceProperties.h
--- a/Engine/Source/Runtime/Engine/Classes/GameFramework/InputDeviceProperties.h
+++ b/Engine/Source/Runtime/Engine/Classes/GameFramework/InputDeviceProperties.h
@@ -4,11 +4,13 @@
 
 #include "CoreMinimal.h"
 #include "UObject/Object.h"
-#include "Curves/CurveFloat.h"
 #include "GenericPlatform/IInputInterface.h"
 
 #include "InputDeviceProperties.generated.h"
 
+class UCurveLinearColor;
+class UCurveFloat;
+
 /**
 * Base class that represents a single Input Device Property. An Input Device Property
 * represents a feature that can be set on an input device. Things like what color a
@@ -21,7 +23,7 @@
 * support certain device properties. An older gamepad may not have any advanced trigger haptics for 
 * example. 
 */
-UCLASS(Abstract, Blueprintable, EditInlineNew)
+UCLASS(Abstract, Blueprintable, BlueprintType)
 class ENGINE_API UInputDeviceProperty : public UObject
 {
 	GENERATED_BODY()
@@ -80,7 +82,7 @@ public:
 	* The duration that this device property should last. Override this if your property has any dynamic curves 
 	* to be the max time range.
 	*/
-	float GetDuration() const { return PropertyDuration; }
+	virtual float GetDuration() { return PropertyDuration; }
 
 protected:
 
@@ -96,14 +98,17 @@ protected:
 	* The duration that this device property should last. Override this if your property has any dynamic curves 
 	* to be the max time range.
 	*/
-	UPROPERTY(EditAnywhere, BlueprintReadWrite, Category = "Info")
+	UPROPERTY(VisibleAnywhere, BlueprintReadWrite, Category = "Info")
 	float PropertyDuration = 0.1f;
 };
 
-class UCurveLinearColor;
-
-/** A property that can be used to change the color of an input device's light */
-UCLASS(Blueprintable, EditInlineNew)
+/** 
+* A property that can be used to change the color of an input device's light 
+* 
+* NOTE: This property has platform specific implementations and may behavior differently per platform.
+* See the docs for more details on each platform.
+*/
+UCLASS(Blueprintable, BlueprintType)
 class UColorInputDeviceProperty : public UInputDeviceProperty
 {
 	GENERATED_BODY()
@@ -115,6 +120,7 @@ public:
 	virtual void EvaluateDeviceProperty_Implementation(const FPlatformUserId PlatformUser, const float DeltaTime, const float Duration) override;
 	virtual void ResetDeviceProperty_Implementation(const FPlatformUserId PlatformUser) override;
 	virtual FInputDeviceProperty* GetInternalDeviceProperty() override;
+	virtual float GetDuration() override;
 
 protected:
 
@@ -130,4 +136,4 @@ private:
 
 	/** The internal light color property that this represents; */
 	FInputDeviceLightColorProperty InternalProperty;
-};
\ No newline at end of file
+};
diff --git a/Engine/Source/Runtime/Engine/Private/GameFramework/InputDeviceProperties.cpp b/Engine/Source/Runtime/Engine/Private/GameFramework/InputDeviceProperties.cpp
--- a/Engine/Source/Runtime/Engine/Private/GameFramework/InputDeviceProperties.cpp
+++ b/Engine/Source/Runtime/Engine/Private/GameFramework/InputDeviceProperties.cpp
@@ -4,6 +4,7 @@
 #include "GenericPlatform/GenericPlatformInputDeviceMapper.h"
 #include "Framework/Application/SlateApplication.h"
 #include "Curves/CurveLinearColor.h"
+#include "Curves/CurveFloat.h"
 
 ///////////////////////////////////////////////////////////////////////
 // UInputDeviceProperty
@@ -38,26 +39,21 @@ void UInputDeviceProperty::ResetDeviceProperty_Implementation(const FPlatformUse
 // UColorInputDeviceProperty
 
 UColorInputDeviceProperty::UColorInputDeviceProperty()
-: UInputDeviceProperty()
+	: UInputDeviceProperty()
 {
 	PropertyName = FInputDeviceLightColorProperty::PropertyName();
-	float MinTime, MaxTime;
-	if (DeviceColorCurve)
-	{
-		DeviceColorCurve->GetTimeRange(MinTime, MaxTime);	
-	}
-	else
-	{
-		PropertyDuration = 1.0f;	
-	}
+	GetDuration();
 }
 
 void UColorInputDeviceProperty::EvaluateDeviceProperty_Implementation(const FPlatformUserId PlatformUser, const float DeltaTime, const float Duration)
 {
 	InternalProperty.bEnable = bEnable;
-	
-	FLinearColor CurveColor = DeviceColorCurve->GetLinearColorValue(Duration);
-	InternalProperty.Color = CurveColor.ToFColorSRGB();
+
+	if (ensure(DeviceColorCurve))
+	{
+		FLinearColor CurveColor = DeviceColorCurve->GetLinearColorValue(Duration);
+		InternalProperty.Color = CurveColor.ToFColorSRGB();
+	}
 }
 
 void UColorInputDeviceProperty::ResetDeviceProperty_Implementation(const FPlatformUserId PlatformUser)
@@ -70,4 +66,18 @@ void UColorInputDeviceProperty::ResetDeviceProperty_Implementation(const FPlatfo
 FInputDeviceProperty* UColorInputDeviceProperty::GetInternalDeviceProperty()
 {
 	return &InternalProperty;
-}
\ No newline at end of file
+}
+
+float UColorInputDeviceProperty::GetDuration()
+{
+	float MinTime, MaxTime;
+	if (DeviceColorCurve)
+	{
+		DeviceColorCurve->GetTimeRange(MinTime, MaxTime);
+	}
+	else
+	{
+		PropertyDuration = 1.0f;
+	}
+	return PropertyDuration;
+}
