Disable channel refcounting while investigating source of possible crash.

diff --git a/Engine/Source/Runtime/Engine/Classes/Engine/NetworkObjectList.h b/Engine/Source/Runtime/Engine/Classes/Engine/NetworkObjectList.h
--- a/Engine/Source/Runtime/Engine/Classes/Engine/NetworkObjectList.h
+++ b/Engine/Source/Runtime/Engine/Classes/Engine/NetworkObjectList.h
@@ -11,7 +11,7 @@ class FArchive;
 #ifndef UE_REPLICATED_OBJECT_REFCOUNTING
 	// Allows every network actor to keep track of the number of channels an individual subobject was replicated. 
 	// Not needed on clients so its compiled out to reduce memory usage.
-	#define UE_REPLICATED_OBJECT_REFCOUNTING WITH_SERVER_CODE
+	#define UE_REPLICATED_OBJECT_REFCOUNTING 0
 #endif
 
 #ifndef DO_REPLICATED_OBJECT_CHANNELREF_CHECKS
