- Make index buffer SRV usage match vertex buffer SRV to allow static mesh sampling to run without having to enable CPU access on platforms which support skin cache / manual fetch

#rnx

diff --git a/Engine/Source/Runtime/Engine/Private/RawIndexBuffer.cpp b/Engine/Source/Runtime/Engine/Private/RawIndexBuffer.cpp
--- a/Engine/Source/Runtime/Engine/Private/RawIndexBuffer.cpp
+++ b/Engine/Source/Runtime/Engine/Private/RawIndexBuffer.cpp
@@ -324,7 +324,7 @@ FBufferRHIRef FRawStaticIndexBuffer::CreateRHIBuffer_Internal()
 	if (GetNumIndices() > 0)
 	{
 		extern ENGINE_API bool DoSkeletalMeshIndexBuffersNeedSRV();
-		bool bSRV = DoSkeletalMeshIndexBuffersNeedSRV();
+		bool bSRV = RHISupportsManualVertexFetch(GMaxRHIShaderPlatform) || DoSkeletalMeshIndexBuffersNeedSRV();
 
 		// When bAllowCPUAccess is true, the meshes is likely going to be used for Niagara to spawn particles on mesh surface.
 		// And it can be the case for CPU *and* GPU access: no differenciation today. That is why we create a SRV in this case.
