Fixed cooked textures returning a size of 0,0 when queried by an editor

diff --git a/Engine/Source/Runtime/Engine/Classes/Engine/Texture2D.h b/Engine/Source/Runtime/Engine/Classes/Engine/Texture2D.h
--- a/Engine/Source/Runtime/Engine/Classes/Engine/Texture2D.h
+++ b/Engine/Source/Runtime/Engine/Classes/Engine/Texture2D.h
@@ -33,14 +33,14 @@ public:
 	/**
 	 * Retrieves the size of the source image from which the texture was created.
 	 */
-	FORCEINLINE FIntPoint GetImportedSize() const
-	{
 #if WITH_EDITOR
-		return Source.GetLogicalSize();
+	ENGINE_API FIntPoint GetImportedSize() const;
 #else // #if WITH_EDITOR
+	FORCEINLINE FIntPoint GetImportedSize() const
+	{
 		return ImportedSize;
-#endif // #if WITH_EDITOR
 	}
+#endif // #if WITH_EDITOR
 
 private:
 	/** True if streaming is temporarily disabled so we can update subregions of this texture's resource 
diff --git a/Engine/Source/Runtime/Engine/Private/Texture2D.cpp b/Engine/Source/Runtime/Engine/Private/Texture2D.cpp
--- a/Engine/Source/Runtime/Engine/Private/Texture2D.cpp
+++ b/Engine/Source/Runtime/Engine/Private/Texture2D.cpp
@@ -604,10 +604,24 @@ bool UTexture2D::IsReadyForAsyncPostLoad() const
 	return !PrivatePlatformData || PrivatePlatformData->IsReadyForAsyncPostLoad();
 }
 
+#if WITH_EDITOR
+FIntPoint UTexture2D::GetImportedSize() const
+{
+	if (!GetPackage()->HasAnyPackageFlags(PKG_Cooked))
+	{
+		return Source.GetLogicalSize();
+	}
+	return ImportedSize;
+}
+#endif // #if WITH_EDITOR
+
 void UTexture2D::PostLoad()
 {
 #if WITH_EDITOR
-	ImportedSize = Source.GetLogicalSize();
+	if (!GetPackage()->HasAnyPackageFlags(PKG_Cooked))
+	{
+		ImportedSize = Source.GetLogicalSize();
+	}
 
 	if (FApp::CanEverRender())
 	{
@@ -657,10 +671,7 @@ void UTexture2D::PreSave(FObjectPreSaveContext ObjectSaveContext)
 
 void UTexture2D::GetAssetRegistryTags(TArray<FAssetRegistryTag>& OutTags) const
 {
-	FIntPoint SourceSize(0, 0);
-#if WITH_EDITOR
-	SourceSize = Source.GetLogicalSize();
-#endif
+	FIntPoint SourceSize = GetImportedSize();
 
 	const FString DimensionsStr = FString::Printf(TEXT("%dx%d"), SourceSize.X, SourceSize.Y);
 	OutTags.Add( FAssetRegistryTag("Dimensions", DimensionsStr, FAssetRegistryTag::TT_Dimensional) );
