Force property compare when handling torn off actor in replay recording

diff --git a/Engine/Source/Runtime/Engine/Private/DemoNetDriver.cpp b/Engine/Source/Runtime/Engine/Private/DemoNetDriver.cpp
--- a/Engine/Source/Runtime/Engine/Private/DemoNetDriver.cpp
+++ b/Engine/Source/Runtime/Engine/Private/DemoNetDriver.cpp
@@ -5117,6 +5117,8 @@ void UDemoNetDriver::NotifyActorTornOff(AActor* Actor)
 {
 	if (IsRecording() && IsValid(Actor))
 	{
+		ForcePropertyCompare(Actor);
+
 		// Replicate one last time to the replay stream
 		ReplayHelper.ReplicateActor(Actor, ClientConnections[0], true);
 
