Remove 2gb check - texture build process supports these sizes, and hardware support is increasing as well.

diff --git a/Engine/Source/Runtime/Engine/Private/TextureDerivedDataTask.cpp b/Engine/Source/Runtime/Engine/Private/TextureDerivedDataTask.cpp
--- a/Engine/Source/Runtime/Engine/Private/TextureDerivedDataTask.cpp
+++ b/Engine/Source/Runtime/Engine/Private/TextureDerivedDataTask.cpp
@@ -551,10 +551,6 @@ static void DDC1_StoreClassicTextureInDerivedData(
 		check(CompressedImage.RawData.GetTypeSize() == 1);
 		int64 CompressedDataSize = CompressedImage.RawData.Num();
 
-		// check CompressedDataSize against int32_max , except if VT
-		// bForVirtualTextureStreamingBuild is false in this branch
-		check(CompressedDataSize < MAX_int32);
-
 		// CompressedImage sizes were padded up to multiple of 4 for d3d, no longer
 		UE_LOG(LogTextureUpload, Verbose, TEXT("Compressed Mip %d PF=%d : %dx%dx%d : %d ; up4 %dx%d=%d"),
 			MipIndex, (int)CompressedImage.PixelFormat,
