Chaos Cloth: Add Levelset visualization. Visualization uses the same code as that used by the Physics Asset Editor via FKLevelSetElem. Moved code from FKLevelSetElem into Chaos::FLevelSet so it could also be used by the ChaosClothVisualization.

I opted to keep the generic UE::Math vector structs rather than Chaos::TVector types for return types since this more closely matches how the code is actually used.

I used the same material that the Physics Asset Editor uses.

diff --git a/Engine/Plugins/Experimental/ChaosCloth/Source/Chaos/Private/ChaosCloth/ChaosClothVisualization.cpp b/Engine/Plugins/Experimental/ChaosCloth/Source/Chaos/Private/ChaosCloth/ChaosClothVisualization.cpp
--- a/Engine/Plugins/Experimental/ChaosCloth/Source/Chaos/Private/ChaosCloth/ChaosClothVisualization.cpp
+++ b/Engine/Plugins/Experimental/ChaosCloth/Source/Chaos/Private/ChaosCloth/ChaosClothVisualization.cpp
@@ -9,6 +9,7 @@
 #include "Chaos/DebugDrawQueue.h"
 #include "Chaos/Capsule.h"
 #include "Chaos/Convex.h"
+#include "Chaos/Levelset.h"
 #include "Chaos/Sphere.h"
 #include "Chaos/TaperedCapsule.h"
 #include "Chaos/TaperedCylinder.h"
@@ -40,6 +41,7 @@ namespace Chaos
 #if WITH_EDITOR
 		ClothMaterial = LoadObject<UMaterial>(nullptr, TEXT("/Engine/EditorMaterials/Cloth/CameraLitDoubleSided.CameraLitDoubleSided"), nullptr, LOAD_None, nullptr);  // LOAD_EditorOnly
 		ClothMaterialVertex = LoadObject<UMaterial>(nullptr, TEXT("/Engine/EditorMaterials/WidgetVertexColorMaterial"), nullptr, LOAD_None, nullptr);  // LOAD_EditorOnly
+		CollisionMaterial = LoadObject<UMaterial>(nullptr, TEXT("/Engine/EditorMaterials/PhAT_UnselectedMaterial"), nullptr, LOAD_None, nullptr);
 #endif  // #if WITH_EDITOR
 	}
 
@@ -55,6 +57,7 @@ namespace Chaos
 	{
 		Collector.AddReferencedObject(ClothMaterial);
 		Collector.AddReferencedObject(ClothMaterialVertex);
+		Collector.AddReferencedObject(CollisionMaterial);
 	}
 
 	void FClothVisualization::DrawPhysMeshShaded(FPrimitiveDrawInterface* PDI) const
@@ -430,6 +433,34 @@ namespace Chaos
 		DrawLine(PDI, Position, Position + Z, FLinearColor::Blue);
 	}
 
+	static void DrawLevelSet(FPrimitiveDrawInterface* PDI, const FTransform& Transform, const FMaterialRenderProxy* MaterialRenderProxy, const FLevelSet& LevelSet)
+	{
+#if WITH_EDITOR
+		if (PDI && MaterialRenderProxy)
+		{
+			TArray<FVector3f> Vertices;
+			TArray<FIntVector> Tris;
+			LevelSet.GetZeroIsosurfaceGridCellFaces(Vertices, Tris);
+
+			FDynamicMeshBuilder MeshBuilder(PDI->View->GetFeatureLevel());
+			for (const FVector3f& V : Vertices)
+			{
+				MeshBuilder.AddVertex(FDynamicMeshVertex(V));
+			}
+			for (const FIntVector& T : Tris)
+			{
+				MeshBuilder.AddTriangle(T[0], T[1], T[2]);
+			}
+
+			MeshBuilder.Draw(PDI, Transform.ToMatrixWithScale(), MaterialRenderProxy, SDPG_World, false, false);
+		}
+		else
+#endif
+		{
+			DrawCoordinateSystem(PDI, Transform.GetRotation(), Transform.GetTranslation());
+		}
+	}
+
 	void FClothVisualization::DrawBounds(FPrimitiveDrawInterface* PDI) const
 	{
 		if (!Solver)
@@ -826,6 +857,22 @@ namespace Chaos
 							DrawConvex(PDI, Object->GetObjectChecked<FConvex>(), Rotation, Position, Color);
 							break;
 
+						case ImplicitObjectType::Transformed: // Transformed only used for levelsets
+							if (Object->GetObjectChecked<TImplicitObjectTransformed<FReal, 3>>().Object()->GetType() == ImplicitObjectType::LevelSet)
+							{
+								const TRigidTransform<FReal, 3>& Transform = Object->GetObjectChecked<TImplicitObjectTransformed<FReal, 3>>().GetTransform();
+								const FTransform CombinedTransform = Transform * FTransform(Rotation, Position);
+								const FLevelSet& LevelSet = Object->GetObjectChecked<TImplicitObjectTransformed<FReal, 3>>().Object()->GetObjectChecked<FLevelSet>();
+								const FMaterialRenderProxy* MaterialRenderProxy =
+#if WITH_EDITOR
+									CollisionMaterial->GetRenderProxy();
+#else
+									nullptr;
+#endif
+								DrawLevelSet(PDI, CombinedTransform, MaterialRenderProxy, LevelSet);
+							}
+							break;
+
 						default:
 							DrawCoordinateSystem(PDI, Rotation, Position);  // Draw everything else as a coordinate for now
 							break;
diff --git a/Engine/Plugins/Experimental/ChaosCloth/Source/Chaos/Public/ChaosCloth/ChaosClothVisualization.h b/Engine/Plugins/Experimental/ChaosCloth/Source/Chaos/Public/ChaosCloth/ChaosClothVisualization.h
--- a/Engine/Plugins/Experimental/ChaosCloth/Source/Chaos/Public/ChaosCloth/ChaosClothVisualization.h
+++ b/Engine/Plugins/Experimental/ChaosCloth/Source/Chaos/Public/ChaosCloth/ChaosClothVisualization.h
@@ -105,6 +105,7 @@ namespace Chaos
 		// Visualization material
 		const UMaterial* ClothMaterial;
 		const UMaterial* ClothMaterialVertex;
+		const UMaterial* CollisionMaterial;
 #endif  // #if WITH_EDITOR
 #endif  // #if CHAOS_DEBUG_DRAW
 	};
diff --git a/Engine/Source/Runtime/Engine/Private/Collision/KAggregateGeom.cpp b/Engine/Source/Runtime/Engine/Private/Collision/KAggregateGeom.cpp
--- a/Engine/Source/Runtime/Engine/Private/Collision/KAggregateGeom.cpp
+++ b/Engine/Source/Runtime/Engine/Private/Collision/KAggregateGeom.cpp
@@ -591,83 +591,7 @@ void FKLevelSetElem::GetZeroIsosurfaceGridCellFaces(TArray<FVector3f>& Vertices,
 {
 	if (LevelSet.IsValid())
 	{
-		const Chaos::TUniformGrid<Chaos::FReal, 3>& Grid = LevelSet->GetGrid();
-		const Chaos::TVector<int32, 3> Cells = Grid.Counts();
-		for (int i = 0; i < Cells.X-1; ++i)
-		{
-			for (int j = 0; j < Cells.Y-1; ++j)
-			{
-				for (int k = 0; k < Cells.Z-1; ++k)
-				{
-					const double Sign = FMath::Sign(LevelSet->GetPhiArray()(i, j, k));
-					const double SignNextI = FMath::Sign(LevelSet->GetPhiArray()(i+1, j, k));
-					const double SignNextJ = FMath::Sign(LevelSet->GetPhiArray()(i, j+1, k));
-					const double SignNextK = FMath::Sign(LevelSet->GetPhiArray()(i, j, k+1));
-					const FVector3d Dx(Grid.Dx());
-
-					const FVector3d CellMin = Grid.MinCorner() + Dx * FVector3d(i, j, k);
-
-					if (Sign > SignNextI)
-					{
-						const int32 V0 = Vertices.Emplace( CellMin + Dx * FVector3d(1, 0, 0) );
-						const int32 V1 = Vertices.Emplace( CellMin + Dx * FVector3d(1, 1, 0) );
-						const int32 V2 = Vertices.Emplace( CellMin + Dx * FVector3d(1, 1, 1) );
-						const int32 V3 = Vertices.Emplace( CellMin + Dx * FVector3d(1, 0, 1) );
-						Tris.Emplace(FIntVector(V0, V1, V2));
-						Tris.Emplace(FIntVector(V2, V3, V0));
-					}
-					else if (Sign < SignNextI)
-					{
-						const int32 V0 = Vertices.Emplace(CellMin + Dx * FVector3d(1, 0, 0));
-						const int32 V1 = Vertices.Emplace(CellMin + Dx * FVector3d(1, 1, 0));
-						const int32 V2 = Vertices.Emplace(CellMin + Dx * FVector3d(1, 1, 1));
-						const int32 V3 = Vertices.Emplace(CellMin + Dx * FVector3d(1, 0, 1));
-						Tris.Emplace(FIntVector(V0, V2, V1));
-						Tris.Emplace(FIntVector(V2, V0, V3));
-					}
-
-					
-					if (Sign > SignNextJ)
-					{
-						const int32 V0 = Vertices.Emplace(CellMin + Dx * FVector3d(0, 1, 0));
-						const int32 V1 = Vertices.Emplace(CellMin + Dx * FVector3d(1, 1, 0));
-						const int32 V2 = Vertices.Emplace(CellMin + Dx * FVector3d(1, 1, 1));
-						const int32 V3 = Vertices.Emplace(CellMin + Dx * FVector3d(0, 1, 1));
-						Tris.Emplace(FIntVector(V0, V2, V1));
-						Tris.Emplace(FIntVector(V2, V0, V3));
-					}
-					else if (Sign < SignNextJ)
-					{
-						const int32 V0 = Vertices.Emplace(CellMin + Dx * FVector3d(0, 1, 0));
-						const int32 V1 = Vertices.Emplace(CellMin + Dx * FVector3d(1, 1, 0));
-						const int32 V2 = Vertices.Emplace(CellMin + Dx * FVector3d(1, 1, 1));
-						const int32 V3 = Vertices.Emplace(CellMin + Dx * FVector3d(0, 1, 1));
-						Tris.Emplace(FIntVector(V0, V1, V2));
-						Tris.Emplace(FIntVector(V2, V3, V0));
-					}
-
-					if (Sign > SignNextK)
-					{
-						const int32 V0 = Vertices.Emplace(CellMin + Dx * FVector3d(0, 0, 1));
-						const int32 V1 = Vertices.Emplace(CellMin + Dx * FVector3d(1, 0, 1));
-						const int32 V2 = Vertices.Emplace(CellMin + Dx * FVector3d(1, 1, 1));
-						const int32 V3 = Vertices.Emplace(CellMin + Dx * FVector3d(0, 1, 1));
-						Tris.Emplace(FIntVector(V0, V1, V2));
-						Tris.Emplace(FIntVector(V2, V3, V0));
-					}
-					else if (Sign < SignNextK)
-					{
-						const int32 V0 = Vertices.Emplace(CellMin + Dx * FVector3d(0, 0, 1));
-						const int32 V1 = Vertices.Emplace(CellMin + Dx * FVector3d(1, 0, 1));
-						const int32 V2 = Vertices.Emplace(CellMin + Dx * FVector3d(1, 1, 1));
-						const int32 V3 = Vertices.Emplace(CellMin + Dx * FVector3d(0, 1, 1));
-						Tris.Emplace(FIntVector(V0, V2, V1));
-						Tris.Emplace(FIntVector(V2, V0, V3));
-					}
-
-				}
-			}
-		}
+		LevelSet->GetZeroIsosurfaceGridCellFaces(Vertices, Tris);
 	}
 }
 
diff --git a/Engine/Source/Runtime/Experimental/Chaos/Private/Chaos/Levelset.cpp b/Engine/Source/Runtime/Experimental/Chaos/Private/Chaos/Levelset.cpp
--- a/Engine/Source/Runtime/Experimental/Chaos/Private/Chaos/Levelset.cpp
+++ b/Engine/Source/Runtime/Experimental/Chaos/Private/Chaos/Levelset.cpp
@@ -614,6 +614,86 @@ bool FLevelSet::CheckData(FErrorReporter& ErrorReporter, const FParticles& InPar
 	return true;
 }
 
+void FLevelSet::GetZeroIsosurfaceGridCellFaces(TArray<FVector3f>& Vertices, TArray<FIntVector>& Tris) const
+{
+	const TVector<int32, 3> Cells = MGrid.Counts();
+	const FVector3d Dx(MGrid.Dx());
+
+	for (int i = 0; i < Cells.X - 1; ++i)
+	{
+		for (int j = 0; j < Cells.Y - 1; ++j)
+		{
+			for (int k = 0; k < Cells.Z - 1; ++k)
+			{
+				const double Sign = FMath::Sign(MPhi(i, j, k));
+				const double SignNextI = FMath::Sign(MPhi(i + 1, j, k));
+				const double SignNextJ = FMath::Sign(MPhi(i, j + 1, k));
+				const double SignNextK = FMath::Sign(MPhi(i, j, k + 1));
+
+				const FVector3d CellMin = MGrid.MinCorner() + Dx * FVector3d(i, j, k);
+
+				if (Sign > SignNextI)
+				{
+					const int32 V0 = Vertices.Emplace(CellMin + Dx * FVector3d(1, 0, 0));
+					const int32 V1 = Vertices.Emplace(CellMin + Dx * FVector3d(1, 1, 0));
+					const int32 V2 = Vertices.Emplace(CellMin + Dx * FVector3d(1, 1, 1));
+					const int32 V3 = Vertices.Emplace(CellMin + Dx * FVector3d(1, 0, 1));
+					Tris.Emplace(FIntVector(V0, V1, V2));
+					Tris.Emplace(FIntVector(V2, V3, V0));
+				}
+				else if (Sign < SignNextI)
+				{
+					const int32 V0 = Vertices.Emplace(CellMin + Dx * FVector3d(1, 0, 0));
+					const int32 V1 = Vertices.Emplace(CellMin + Dx * FVector3d(1, 1, 0));
+					const int32 V2 = Vertices.Emplace(CellMin + Dx * FVector3d(1, 1, 1));
+					const int32 V3 = Vertices.Emplace(CellMin + Dx * FVector3d(1, 0, 1));
+					Tris.Emplace(FIntVector(V0, V2, V1));
+					Tris.Emplace(FIntVector(V2, V0, V3));
+				}
+
+
+				if (Sign > SignNextJ)
+				{
+					const int32 V0 = Vertices.Emplace(CellMin + Dx * FVector3d(0, 1, 0));
+					const int32 V1 = Vertices.Emplace(CellMin + Dx * FVector3d(1, 1, 0));
+					const int32 V2 = Vertices.Emplace(CellMin + Dx * FVector3d(1, 1, 1));
+					const int32 V3 = Vertices.Emplace(CellMin + Dx * FVector3d(0, 1, 1));
+					Tris.Emplace(FIntVector(V0, V2, V1));
+					Tris.Emplace(FIntVector(V2, V0, V3));
+				}
+				else if (Sign < SignNextJ)
+				{
+					const int32 V0 = Vertices.Emplace(CellMin + Dx * FVector3d(0, 1, 0));
+					const int32 V1 = Vertices.Emplace(CellMin + Dx * FVector3d(1, 1, 0));
+					const int32 V2 = Vertices.Emplace(CellMin + Dx * FVector3d(1, 1, 1));
+					const int32 V3 = Vertices.Emplace(CellMin + Dx * FVector3d(0, 1, 1));
+					Tris.Emplace(FIntVector(V0, V1, V2));
+					Tris.Emplace(FIntVector(V2, V3, V0));
+				}
+
+				if (Sign > SignNextK)
+				{
+					const int32 V0 = Vertices.Emplace(CellMin + Dx * FVector3d(0, 0, 1));
+					const int32 V1 = Vertices.Emplace(CellMin + Dx * FVector3d(1, 0, 1));
+					const int32 V2 = Vertices.Emplace(CellMin + Dx * FVector3d(1, 1, 1));
+					const int32 V3 = Vertices.Emplace(CellMin + Dx * FVector3d(0, 1, 1));
+					Tris.Emplace(FIntVector(V0, V1, V2));
+					Tris.Emplace(FIntVector(V2, V3, V0));
+				}
+				else if (Sign < SignNextK)
+				{
+					const int32 V0 = Vertices.Emplace(CellMin + Dx * FVector3d(0, 0, 1));
+					const int32 V1 = Vertices.Emplace(CellMin + Dx * FVector3d(1, 0, 1));
+					const int32 V2 = Vertices.Emplace(CellMin + Dx * FVector3d(1, 1, 1));
+					const int32 V3 = Vertices.Emplace(CellMin + Dx * FVector3d(0, 1, 1));
+					Tris.Emplace(FIntVector(V0, V2, V1));
+					Tris.Emplace(FIntVector(V2, V0, V3));
+				}
+			}
+		}
+	}
+}
+
 void FLevelSet::ComputeConvexity(const TArray<TVec3<int32>>& InterfaceIndices)
 {
 	this->bIsConvex = true;
diff --git a/Engine/Source/Runtime/Experimental/Chaos/Public/Chaos/Levelset.h b/Engine/Source/Runtime/Experimental/Chaos/Public/Chaos/Levelset.h
--- a/Engine/Source/Runtime/Experimental/Chaos/Public/Chaos/Levelset.h
+++ b/Engine/Source/Runtime/Experimental/Chaos/Public/Chaos/Levelset.h
@@ -138,6 +138,9 @@ class CHAOS_API FLevelSet final : public FImplicitObject
 
 	bool CheckData(FErrorReporter& ErrorReporter, const FParticles& InParticles, const FTriangleMesh& Mesh, const TArray<FVec3> &Normals);
 
+	// Used to generate a simple debug surface
+	void GetZeroIsosurfaceGridCellFaces(TArray<FVector3f>& Vertices, TArray<FIntVector>& Tris) const;
+
 	virtual uint32 GetTypeHash() const override
 	{
 		uint32 Result = 0;
