Added FindOrLoadLocalIniFile helper function that looks in the cache before trying to load the ini file.

CreateProfile:
Before: 2387.2ms
After: 9.68ms

LoadConfig:
Before: 2595ms
After: 541ms

U******RuntimeSettings::PostInitProperties:
Before: 27ms
After: 0.47ms

InitializeSerializationOptionsFromIni:
Before: 66ms
After: 0.86ms

diff --git a/Engine/Plugins/FX/Niagara/Source/Niagara/Private/NiagaraPlatformSet.cpp b/Engine/Plugins/FX/Niagara/Source/Niagara/Private/NiagaraPlatformSet.cpp
--- a/Engine/Plugins/FX/Niagara/Source/Niagara/Private/NiagaraPlatformSet.cpp
+++ b/Engine/Plugins/FX/Niagara/Source/Niagara/Private/NiagaraPlatformSet.cpp
@@ -1242,23 +1242,23 @@ FNiagaraPlatformSet::FPlatformIniSettings& FNiagaraPlatformSet::GetPlatformIniSe
 	}
 
 	//Load config files in which we can reasonable expect to find fx.Niagara.QualityLevel and may be set.
-	FConfigFile EngineSettings;
-	FConfigCacheIni::LoadLocalIniFile(EngineSettings, TEXT("Engine"), true, *PlatformName);//Should use BaseProfileName? Are either of these ensured to be correct? I worry this is brittle.
-
-	FConfigFile GameSettings;
-	FConfigCacheIni::LoadLocalIniFile(GameSettings, TEXT("Game"), true, *PlatformName);//Should use BaseProfileName? Are either of these ensured to be correct? I worry this is brittle.
+	FConfigFile LocalEngineIni;
+	FConfigFile* LoadedEngineIni = FConfigCacheIni::FindOrLoadPlatformConfig(LocalEngineIni, TEXT("Engine"), *PlatformName); //Should use BaseProfileName? Are either of these ensured to be correct? I worry this is brittle.
 
-	FConfigFile ScalabilitySettings;
-	FConfigCacheIni::LoadLocalIniFile(ScalabilitySettings, TEXT("Scalability"), true, *PlatformName);//Should use BaseProfileName? Are either of these ensured to be correct? I worry this is brittle.
+	FConfigFile LocalGameIni;
+	FConfigFile* LoadedGameIni = FConfigCacheIni::FindOrLoadPlatformConfig(LocalGameIni, TEXT("Game"), *PlatformName); //Should use BaseProfileName? Are either of these ensured to be correct? I worry this is brittle.
+	
+	FConfigFile LocalScalabilityIni;
+	FConfigFile* LoadedScalabilityIni = FConfigCacheIni::FindOrLoadPlatformConfig(LocalScalabilityIni, TEXT("Scalability"), *PlatformName); //Should use BaseProfileName? Are either of these ensured to be correct? I worry this is brittle.
 
 	auto FindCVarValue = [&](const TCHAR* Section, const TCHAR* CVarName, int32& OutVal)
 	{
 		bool bFound = true;
-		if (!ScalabilitySettings.GetInt(Section, CVarName, OutVal))
+		if (!LoadedScalabilityIni->GetInt(Section, CVarName, OutVal))
 		{
-			if (!GameSettings.GetInt(Section, CVarName, OutVal))
+			if (!LoadedGameIni->GetInt(Section, CVarName, OutVal))
 			{
-				if (!EngineSettings.GetInt(Section, CVarName, OutVal))
+				if (!LoadedEngineIni->GetInt(Section, CVarName, OutVal))
 				{
 					bFound = false;
 				}
@@ -1302,7 +1302,7 @@ FNiagaraPlatformSet::FPlatformIniSettings& FNiagaraPlatformSet::GetPlatformIniSe
 	{
 		FString SectionName = Scalability::GetScalabilitySectionString(TEXT("EffectsQuality"), EQ, NumEffectsQualities);
 		int32 NiagaraQualityLevelForEQ = DefaultQualityLevel;
-		ScalabilitySettings.GetInt(*SectionName, NiagaraQualityLevelName, NiagaraQualityLevelForEQ);
+		LoadedScalabilityIni->GetInt(*SectionName, NiagaraQualityLevelName, NiagaraQualityLevelForEQ);
 		NewSetting.QualityLevelsPerEffectsQuality.Add(NiagaraQualityLevelForEQ);
 		EQStr += FString::Printf(TEXT("EQ:%d = NQL:%d\n"), EQ, NiagaraQualityLevelForEQ);
 	}
diff --git a/Engine/Source/Runtime/AssetRegistry/Private/AssetRegistry.cpp b/Engine/Source/Runtime/AssetRegistry/Private/AssetRegistry.cpp
--- a/Engine/Source/Runtime/AssetRegistry/Private/AssetRegistry.cpp
+++ b/Engine/Source/Runtime/AssetRegistry/Private/AssetRegistry.cpp
@@ -1180,16 +1180,9 @@ static TSet<FName> MakeNameSet(const TArray<FString>& Strings)
 
 void InitializeSerializationOptionsFromIni(FAssetRegistrySerializationOptions& Options, const FString& PlatformIniName, UE::AssetRegistry::ESerializationTarget Target)
 {
-	FConfigFile* EngineIni = nullptr;
-#if WITH_EDITOR
 	// Use passed in platform, or current platform if empty
-	FConfigFile PlatformEngineIni;
-	FConfigCacheIni::LoadLocalIniFile(PlatformEngineIni, TEXT("Engine"), true, (!PlatformIniName.IsEmpty() ? *PlatformIniName : ANSI_TO_TCHAR(FPlatformProperties::IniPlatformName())));
-	EngineIni = &PlatformEngineIni;
-#else
-	// In cooked builds, always use the normal engine INI
-	EngineIni = GConfig->FindConfigFile(GEngineIni);
-#endif
+	FConfigFile LocalEngineIni;
+	FConfigFile* EngineIni = FConfigCacheIni::FindOrLoadPlatformConfig(LocalEngineIni, TEXT("Engine"), (!PlatformIniName.IsEmpty() ? *PlatformIniName : ANSI_TO_TCHAR(FPlatformProperties::IniPlatformName())));
 
 	Options = FAssetRegistrySerializationOptions(Target);
 	// For DevelopmentAssetRegistry, all non-tag options are overridden in the constructor
diff --git a/Engine/Source/Runtime/Core/Private/Misc/ConfigCacheIni.cpp b/Engine/Source/Runtime/Core/Private/Misc/ConfigCacheIni.cpp
--- a/Engine/Source/Runtime/Core/Private/Misc/ConfigCacheIni.cpp
+++ b/Engine/Source/Runtime/Core/Private/Misc/ConfigCacheIni.cpp
@@ -3926,6 +3926,38 @@ bool FConfigCacheIni::LoadExternalIniFile(FConfigFile & ConfigFile, const TCHAR
 	return Context.Load(IniName);
 }
 
+FConfigFile* FConfigCacheIni::FindPlatformConfig(const TCHAR* IniName, const TCHAR* Platform)
+{
+	if (Platform != nullptr && FCString::Stricmp(Platform, ANSI_TO_TCHAR(FPlatformProperties::IniPlatformName())) != 0)
+	{
+#if ALLOW_OTHER_PLATFORM_CONFIG
+		return FConfigCacheIni::ForPlatform(Platform)->FindConfigFile(IniName);
+#else
+		return nullptr;
+#endif
+	}
+
+	if (GConfig != nullptr)
+	{
+		return GConfig->FindConfigFile(IniName);
+	}
+
+	return nullptr;
+}
+
+FConfigFile* FConfigCacheIni::FindOrLoadPlatformConfig(FConfigFile& LocalFile, const TCHAR* IniName, const TCHAR* Platform)
+{
+	FConfigFile* File = FindPlatformConfig(IniName, Platform);
+	if (File == nullptr)
+	{
+		FConfigContext Context = FConfigContext::ReadIntoLocalFile(LocalFile, Platform);
+		Context.Load(IniName);
+		File = &LocalFile;
+	}
+
+	return File;
+}
+
 void FConfigCacheIni::LoadConsoleVariablesFromINI()
 {
 	FString ConsoleVariablesPath = FPaths::EngineDir() + TEXT("Config/ConsoleVariables.ini");
@@ -4370,7 +4402,7 @@ void FConfigCacheIni::AsyncInitializeConfigForPlatforms()
 FConfigCacheIni* FConfigCacheIni::ForPlatform(FName PlatformName)
 {
 #if ALLOW_OTHER_PLATFORM_CONFIG
-	check(GConfig->bIsReadyForUse);
+	check(GConfig != nullptr && GConfig->bIsReadyForUse);
 
 	// use GConfig when no platform is specified
 	if (PlatformName == NAME_None)
diff --git a/Engine/Source/Runtime/Core/Public/Misc/ConfigCacheIni.h b/Engine/Source/Runtime/Core/Public/Misc/ConfigCacheIni.h
--- a/Engine/Source/Runtime/Core/Public/Misc/ConfigCacheIni.h
+++ b/Engine/Source/Runtime/Core/Public/Misc/ConfigCacheIni.h
@@ -1206,6 +1206,26 @@ public:
 	 *	@return A normalized version of the path (may be the same as the input).
 	 */
 	static FString NormalizeConfigIniPath(const FString& NonNormalizedPath);
+
+	/**
+	 * This helper function searches the cache before trying to load the ini file using LoadLocalIniFile. 
+	 * Note that the returned FConfigFile pointer must have the same lifetime as the passed in LocalFile.
+	 *
+	 * @param LocalFile The output object to fill. If the FConfigFile is found in the cache, this won't be used.
+	 * @param IniName Either a Base ini name (Engine) or a full ini name (WrangleContent). NO PATH OR EXTENSION SHOULD BE USED!
+	 * @param Platform The platform to use for Base ini names, NULL means to use the current platform
+	 * @return FConfigFile Found or loaded FConfigFile
+	 */
+	static FConfigFile* FindOrLoadPlatformConfig(FConfigFile& LocalFile, const TCHAR* IniName, const TCHAR* Platform = NULL);
+
+	/**
+	 * Attempts to find the platform config in the cache.
+	 *
+	 * @param IniName Either a Base ini name (Engine) or a full ini name (WrangleContent). NO PATH OR EXTENSION SHOULD BE USED!
+	 * @param Platform The platform to use for Base ini names, NULL means to use the current platform
+	 * @return FConfigFile Found FConfigFile
+	 */
+	static FConfigFile* FindPlatformConfig(const TCHAR* IniName, const TCHAR* Platform);
 	
 	/**
 	 * Save the current config cache state into a file for bootstrapping other processes.
diff --git a/Engine/Source/Runtime/CoreUObject/Private/UObject/Obj.cpp b/Engine/Source/Runtime/CoreUObject/Private/UObject/Obj.cpp
--- a/Engine/Source/Runtime/CoreUObject/Private/UObject/Obj.cpp
+++ b/Engine/Source/Runtime/CoreUObject/Private/UObject/Obj.cpp
@@ -2476,14 +2476,13 @@ void UObject::LoadConfig( UClass* ConfigClass/*=NULL*/, const TCHAR* InFilename/
 	// and if the name isn't the current running platform (no need to load extra files if already in GConfig)
 	bool bUseConfigOverride = InFilename == nullptr && GetConfigOverridePlatform() != nullptr &&
 		FCString::Stricmp(GetConfigOverridePlatform(), ANSI_TO_TCHAR(FPlatformProperties::IniPlatformName())) != 0;
-	FConfigFile OverrideConfig;
+	FConfigFile* OverrideConfigFile = nullptr;
+	FConfigFile LocalOverrideConfig;
 	if (bUseConfigOverride)
 	{
-		// load into a local ini file
-		FConfigCacheIni::LoadLocalIniFile(OverrideConfig, *GetClass()->ClassConfigName.ToString(), true, GetConfigOverridePlatform());
+		OverrideConfigFile = FConfigCacheIni::FindOrLoadPlatformConfig(LocalOverrideConfig, *GetClass()->ClassConfigName.ToString(), GetConfigOverridePlatform());
 	}
 
-
 	FString ClassSection;
 	FString ClassPathSection;
 	FName LongCommitName;
@@ -2535,11 +2534,11 @@ void UObject::LoadConfig( UClass* ConfigClass/*=NULL*/, const TCHAR* InFilename/
 		UE_LOG(LogConfig, Verbose, TEXT("(%s) '%s' loading configuration for property %s from %s"), *ConfigClass->GetName(), *GetName(), *PropertyToLoad->GetName(), *Filename);
 	}
 
-	auto GetConfigValue = [&OverrideConfig, &bUseConfigOverride](const TCHAR* ClassSection, const TCHAR* Key, const TCHAR* ConfigName, FString& OutValue)
+	auto GetConfigValue = [&OverrideConfigFile, &bUseConfigOverride](const TCHAR* ClassSection, const TCHAR* Key, const TCHAR* ConfigName, FString& OutValue)
 	{
 		if (bUseConfigOverride)
 		{
-			return OverrideConfig.GetString(ClassSection, Key, OutValue);
+			return OverrideConfigFile->GetString(ClassSection, Key, OutValue);
 		}
 		else
 		{
@@ -2547,11 +2546,11 @@ void UObject::LoadConfig( UClass* ConfigClass/*=NULL*/, const TCHAR* InFilename/
 		}
 	};
 
-	auto GetConfigSection = [&bUseConfigOverride, &OverrideConfig](const TCHAR* SectionName, const TCHAR* ConfigFilename)
+	auto GetConfigSection = [&OverrideConfigFile, &bUseConfigOverride](const TCHAR* SectionName, const TCHAR* ConfigFilename)
 	{
 		if (bUseConfigOverride)
 		{
-			return OverrideConfig.Find(SectionName);
+			return OverrideConfigFile->Find(SectionName);
 		}
 		else
 		{
diff --git a/Engine/Source/Runtime/Engine/Private/DeviceProfiles/DeviceProfileManager.cpp b/Engine/Source/Runtime/Engine/Private/DeviceProfiles/DeviceProfileManager.cpp
--- a/Engine/Source/Runtime/Engine/Private/DeviceProfiles/DeviceProfileManager.cpp
+++ b/Engine/Source/Runtime/Engine/Private/DeviceProfiles/DeviceProfileManager.cpp
@@ -634,7 +634,7 @@ UDeviceProfile* UDeviceProfileManager::CreateProfile(const FString& ProfileName,
 		// @todo config: we could likely cache local ini files to speed this up,
 		// along with the ones we load in LoadConfig
 		// NOTE: This happens at runtime, so maybe only do this if !RequiresCookedData()?
-		FConfigFile* PlatformConfigFile;
+		const FConfigFile* PlatformConfigFile = nullptr;
 		FConfigFile LocalConfigFile;
 		if (FPlatformProperties::RequiresCookedData())
 		{
@@ -642,8 +642,7 @@ UDeviceProfile* UDeviceProfileManager::CreateProfile(const FString& ProfileName,
 		}
 		else
 		{
-			FConfigCacheIni::LoadLocalIniFile(LocalConfigFile, TEXT("DeviceProfiles"), true, ConfigPlatform);
-			PlatformConfigFile = &LocalConfigFile;
+			PlatformConfigFile = FConfigCacheIni::FindOrLoadPlatformConfig(LocalConfigFile, TEXT("DeviceProfiles"), ConfigPlatform);
 		}
 
 		// Build Parent objects first. Important for setup
@@ -812,13 +811,12 @@ void UDeviceProfileManager::LoadProfiles()
 			FString ConfigLoadPlatform = PlatformIndex == 0 ? FString(FPlatformProperties::IniPlatformName()) : ConfidentialPlatforms[PlatformIndex - 1].ToString();
 
 			// load the DP.ini files (from current platform and then by the extra confidential platforms)
-			FConfigFile PlatformConfigFile;
-			FConfigCacheIni::LoadLocalIniFile(PlatformConfigFile, TEXT("DeviceProfiles"), true, *ConfigLoadPlatform);
+			FConfigFile LocalPlatformConfigFile;
+			const FConfigFile* PlatformConfigFile = FConfigCacheIni::FindOrLoadPlatformConfig(LocalPlatformConfigFile, TEXT("DeviceProfiles"), *ConfigLoadPlatform);
 
 			// load all of the DeviceProfiles
 			TArray<FString> ProfileDescriptions;
-			PlatformConfigFile.GetArray(TEXT("DeviceProfiles"), TEXT("DeviceProfileNameAndTypes"), ProfileDescriptions);
-
+			PlatformConfigFile->GetArray(TEXT("DeviceProfiles"), TEXT("DeviceProfileNameAndTypes"), ProfileDescriptions);
 
 			// add them to our collection of profiles by platform
 			for (const FString& Desc : ProfileDescriptions)
@@ -1044,9 +1042,9 @@ void UDeviceProfileManager::HandleDeviceProfileOverrideChange()
 		FString PlatformName = ANSI_TO_TCHAR(FPlatformProperties::IniPlatformName());
 		
 		TArray<FString> DeviceProfileNameAndTypes;
-		FConfigFile PlatformConfigFile;
-		FConfigCacheIni::LoadLocalIniFile(PlatformConfigFile, TEXT("DeviceProfiles"), true, *PlatformName);
-		PlatformConfigFile.GetArray(TEXT("DeviceProfiles"), TEXT("DeviceProfileNameAndTypes"), DeviceProfileNameAndTypes);
+		FConfigFile LocalConfigFile;
+		const FConfigFile* PlatformConfigFile = FConfigCacheIni::FindOrLoadPlatformConfig(LocalConfigFile, TEXT("DeviceProfiles"), *PlatformName);
+		PlatformConfigFile->GetArray(TEXT("DeviceProfiles"), TEXT("DeviceProfileNameAndTypes"), DeviceProfileNameAndTypes);
 			
 		bool bCreateIfMissing = false;
 		for (const FString& Desc: DeviceProfileNameAndTypes)
