Fix for Anim Notify States being ended when CustomTimeDilation is set to zero

diff --git a/Engine/Source/Runtime/Engine/Private/Animation/AnimCompositeBase.cpp b/Engine/Source/Runtime/Engine/Private/Animation/AnimCompositeBase.cpp
--- a/Engine/Source/Runtime/Engine/Private/Animation/AnimCompositeBase.cpp
+++ b/Engine/Source/Runtime/Engine/Private/Animation/AnimCompositeBase.cpp
@@ -102,14 +102,10 @@ void FAnimSegment::GetAnimNotifiesFromTrackPositions(const float& PreviousTrackP
 
 void FAnimSegment::GetAnimNotifiesFromTrackPositions(const float& PreviousTrackPosition, const float& CurrentTrackPosition, FAnimNotifyContext& NotifyContext) const
 {
-	if( PreviousTrackPosition == CurrentTrackPosition )
-	{
-		return;
-	}
-
 	const bool bTrackPlayingBackwards = (PreviousTrackPosition > CurrentTrackPosition);
 	const float SegmentStartPos = StartPos;
 	const float SegmentEndPos = StartPos + GetLength();
+	const bool bZeroTrackPositionDelta = CurrentTrackPosition == PreviousTrackPosition;
 
 	// if track range overlaps segment
 	if( bTrackPlayingBackwards 
@@ -139,12 +135,13 @@ void FAnimSegment::GetAnimNotifiesFromTrackPositions(const float& PreviousTrackP
 			// Abstract out end point since animation can be playing forward or backward.
 			const float AnimEndPoint = bAnimPlayingBackwards ? AnimStartTime : AnimEndTime;
 
-			for(int32 IterationsLeft=FMath::Max(LoopingCount, 1); ((IterationsLeft > 0) && (TrackTimeToGo > 0.f)); --IterationsLeft)
+			for(int32 IterationsLeft=FMath::Max(LoopingCount, 1); ((IterationsLeft > 0) && (TrackTimeToGo > 0.f || bZeroTrackPositionDelta)); --IterationsLeft)
 			{
 				// Track time left to reach end point of animation.
 				const float TrackTimeToAnimEndPoint = (AnimEndPoint - AnimStartPosition) / AbsValidPlayRate;
 
 				// If our time left is shorter than time to end point, no problem. End there.
+				// This will also run if we arrive with bZeroTrackPositionDelta == true, as TrackTimeToGo == 0.f
 				if( FMath::Abs(TrackTimeToGo) < FMath::Abs(TrackTimeToAnimEndPoint) )
 				{
 					const float PlayRate = ValidPlayRate * (bTrackPlayingBackwards ? -1.f : 1.f);
diff --git a/Engine/Source/Runtime/Engine/Private/Animation/AnimMontage.cpp b/Engine/Source/Runtime/Engine/Private/Animation/AnimMontage.cpp
--- a/Engine/Source/Runtime/Engine/Private/Animation/AnimMontage.cpp
+++ b/Engine/Source/Runtime/Engine/Private/Animation/AnimMontage.cpp
@@ -2360,6 +2360,12 @@ void FAnimMontageInstance::Advance(float DeltaTime, struct FRootMotionMovementPa
 				MontageSubStepper.AddEvaluationTime(DeltaTime);
 			}
 
+			// Gather active anim state notifies if DeltaTime == 0 (happens when TimeDilation is 0.f), so these are not prematurely ended
+			if (DeltaTime == 0.f)
+			{
+				HandleEvents(Position, Position, nullptr);
+			}
+
 			while (bPlaying && MontageSubStepper.HasTimeRemaining() && (++NumIterations < MaxIterations))
 			{
 				SCOPE_CYCLE_COUNTER(STAT_AnimMontageInstance_Advance_Iteration);
@@ -2558,7 +2564,6 @@ void FAnimMontageInstance::HandleEvents(float PreviousTrackPos, float CurrentTra
 	// now get active Notifies based on how it advanced
 	if (AnimInstance.IsValid())
 	{
-		TMap<FName, TArray<FAnimNotifyEventReference>> NotifyMap;
 		FAnimTickRecord TickRecord;
 
 		// Add instance ID to context to differentiate notifies between different instances of the same montage
@@ -2575,15 +2580,11 @@ void FAnimMontageInstance::HandleEvents(float PreviousTrackPos, float CurrentTra
 		// we'll do this for all slots for now
 		for (auto SlotTrack = Montage->SlotAnimTracks.CreateIterator(); SlotTrack; ++SlotTrack)
 		{
-			TArray<FAnimNotifyEventReference>& MapNotifies = NotifyMap.FindOrAdd(SlotTrack->SlotName);
-			PRAGMA_DISABLE_DEPRECATION_WARNINGS
-			SlotTrack->AnimTrack.GetAnimNotifiesFromTrackPositions(PreviousTrackPos, CurrentTrackPos, MapNotifies);
-			PRAGMA_ENABLE_DEPRECATION_WARNINGS
+			SlotTrack->AnimTrack.GetAnimNotifiesFromTrackPositions(PreviousTrackPos, CurrentTrackPos, NotifyContext);
 		}
 
 		// Queue all these notifies.
 		AnimInstance->NotifyQueue.AddAnimNotifies(NotifyContext.ActiveNotifies, NotifyWeight);
-		AnimInstance->NotifyQueue.AddAnimNotifies(NotifyMap, NotifyWeight);
 	}
 
 	// Update active state branching points, before we handle the immediate tick marker.
diff --git a/Engine/Source/Runtime/Engine/Private/Animation/AnimSequenceBase.cpp b/Engine/Source/Runtime/Engine/Private/Animation/AnimSequenceBase.cpp
--- a/Engine/Source/Runtime/Engine/Private/Animation/AnimSequenceBase.cpp
+++ b/Engine/Source/Runtime/Engine/Private/Animation/AnimSequenceBase.cpp
@@ -381,11 +381,6 @@ void UAnimSequenceBase::GetAnimNotifies(const float& StartTime, const float& Del
 
 void UAnimSequenceBase::GetAnimNotifies(const float& StartTime, const float& DeltaTime, FAnimNotifyContext& NotifyContext) const
 {
-	if(DeltaTime == 0.f)
-	{
-		return;
-	}
-
 	// Early out if we have no notifies
 	if (!IsNotifyAvailable())
 	{
@@ -452,7 +447,7 @@ void UAnimSequenceBase::GetAnimNotifiesFromDeltaPositions(const float& PreviousP
 void UAnimSequenceBase::GetAnimNotifiesFromDeltaPositions(const float& PreviousPosition, const float& CurrentPosition,  FAnimNotifyContext& NotifyContext) const
 {
 	// Early out if we have no notifies
-	if( (Notifies.Num() == 0) || (PreviousPosition == CurrentPosition) )
+	if (Notifies.Num() == 0)
 	{
 		return;
 	}
@@ -496,6 +491,7 @@ void UAnimSequenceBase::GetAnimNotifiesFromDeltaPositions(const float& PreviousP
 			const float NotifyStartTime = AnimNotifyEvent.GetTriggerTime();
 			const float NotifyEndTime = AnimNotifyEvent.GetEndTriggerTime();
 
+			// Note that if you arrive with zero delta time (CurrentPosition == PreviousPosition), only Notify States will be extracted
 			if( (NotifyStartTime <= CurrentPosition) && (NotifyEndTime > PreviousPosition) )
 			{
 				if (NotifyContext.TickRecord)
diff --git a/Engine/Source/Runtime/Engine/Private/Animation/BlendSpace.cpp b/Engine/Source/Runtime/Engine/Private/Animation/BlendSpace.cpp
--- a/Engine/Source/Runtime/Engine/Private/Animation/BlendSpace.cpp
+++ b/Engine/Source/Runtime/Engine/Private/Animation/BlendSpace.cpp
@@ -553,7 +553,8 @@ void UBlendSpace::TickAssetPlayer(FAnimTickRecord& Instance, struct FAnimNotifyQ
 					ClampedNormalizedCurrentTime = FMath::Clamp<float>(NormalizedPreviousTime, 0.f, 1.f);
 				}
 
-				const bool bGenerateNotifies = (NormalizedCurrentTime != NormalizedPreviousTime) && NotifyTriggerMode != ENotifyTriggerMode::None;
+				const bool bHasDeltaTime = (NormalizedCurrentTime != NormalizedPreviousTime);
+				const bool bGenerateNotifies = NotifyTriggerMode != ENotifyTriggerMode::None;
 
 				// Get the index of the highest weight, assuming that the first is the highest until we find otherwise
 				const bool bTriggerNotifyHighestWeightedAnim = NotifyTriggerMode == ENotifyTriggerMode::HighestWeightedAnimation && SampleDataList.Num() > 0;
@@ -603,9 +604,12 @@ void UBlendSpace::TickAssetPlayer(FAnimTickRecord& Instance, struct FAnimNotifyQ
 								Sample.Animation->GetAnimNotifies(PrevSampleDataTime, DeltaTimePosition, NotifyContext);
 							}
 
-							if (Context.RootMotionMode == ERootMotionMode::RootMotionFromEverything && Sample.Animation->bEnableRootMotion)
+							if (bHasDeltaTime)
 							{
-								Context.RootMotionMovementParams.AccumulateWithBlend(Sample.Animation->ExtractRootMotion(PrevSampleDataTime, DeltaTimePosition, Instance.bLooping), SampleEntry.GetClampedWeight());
+								if (Context.RootMotionMode == ERootMotionMode::RootMotionFromEverything && Sample.Animation->bEnableRootMotion)
+								{
+									Context.RootMotionMovementParams.AccumulateWithBlend(Sample.Animation->ExtractRootMotion(PrevSampleDataTime, DeltaTimePosition, Instance.bLooping), SampleEntry.GetClampedWeight());
+								}
 							}
 
 							// Capture the final adjusted delta time and previous frame time as an asset player record
