World Partition Location Volume: fix some edge cases when undoing moving a location volume or changing its properties, etc.

#rnx

diff --git a/Engine/Source/Runtime/Engine/Private/LocationVolume.cpp b/Engine/Source/Runtime/Engine/Private/LocationVolume.cpp
--- a/Engine/Source/Runtime/Engine/Private/LocationVolume.cpp
+++ b/Engine/Source/Runtime/Engine/Private/LocationVolume.cpp
@@ -42,9 +42,12 @@ void ALocationVolume::PostRegisterAllComponents()
 {
 	Super::PostRegisterAllComponents();
 
-	if (!GetWorld()->IsGameWorld() && GetWorld()->IsPartitionedWorld() && !WorldPartitionActorLoader)
+	if (!GetWorld()->IsGameWorld() && GetWorld()->IsPartitionedWorld())
 	{
-		WorldPartitionActorLoader = new FLoaderAdapterLocationVolumeActor(this);
+		if (!WorldPartitionActorLoader)
+		{
+			WorldPartitionActorLoader = new FLoaderAdapterLocationVolumeActor(this);
+		}
 
 		if (bIsAutoLoad)
 		{
@@ -58,8 +61,9 @@ void ALocationVolume::UnregisterAllComponents(bool bForReregister)
 {
 	Super::UnregisterAllComponents(bForReregister);
 
-	if (!bForReregister && WorldPartitionActorLoader)
+	if (HasActorRegisteredAllComponents() && !bForReregister && WorldPartitionActorLoader)
 	{
+		bIsAutoLoad = WorldPartitionActorLoader->IsLoaded();
 		WorldPartitionActorLoader->Unload();
 	}
 }
