World Partition Editor Loader Interface: don't reset transactions during a transaction to avoid an assert, this is when undoing moving a location volume.

#rnx

diff --git a/Engine/Source/Runtime/Engine/Private/WorldPartition/WorldPartitionActorLoaderInterface.cpp b/Engine/Source/Runtime/Engine/Private/WorldPartition/WorldPartitionActorLoaderInterface.cpp
--- a/Engine/Source/Runtime/Engine/Private/WorldPartition/WorldPartitionActorLoaderInterface.cpp
+++ b/Engine/Source/Runtime/Engine/Private/WorldPartition/WorldPartitionActorLoaderInterface.cpp
@@ -273,7 +273,10 @@ void IWorldPartitionActorLoaderInterface::ILoaderAdapter::PostLoadedStateChanged
 		GEngine->BroadcastLevelActorListChanged();
 		GEditor->NoteSelectionChange();
 
-		GEditor->ResetTransaction(LOCTEXT("LoadingEditorActorResetTrans", "Editor Actors Loading State Changed"));
+		if (!GIsTransacting)
+		{
+			GEditor->ResetTransaction(LOCTEXT("LoadingEditorActorResetTrans", "Editor Actors Loading State Changed"));
+		}
 
 		if (NumUnloads)
 		{
