Create the input device engine subsystem. This subsystem will allow you to set device properties on specific platform users without needing to use any force feedback effects. This can tick in the Editor to be used with editor utiltiy widgets and in the future, setting device properties right from the editor without needing to actually PIE.

diff --git a/Engine/Source/Runtime/Engine/Classes/GameFramework/InputDeviceSubsystem.h b/Engine/Source/Runtime/Engine/Classes/GameFramework/InputDeviceSubsystem.h
new file mode 100644
--- /dev/null
+++ b/Engine/Source/Runtime/Engine/Classes/GameFramework/InputDeviceSubsystem.h
@@ -0,0 +1,105 @@
+// Copyright Epic Games, Inc. All Rights Reserved.
+
+#pragma once
+
+#include "Subsystems/EngineSubsystem.h"
+#include "Tickable.h"
+#include "Templates/SubclassOf.h"
+#include "InputDeviceSubsystem.generated.h"
+
+DECLARE_LOG_CATEGORY_EXTERN(LogInputDeviceProperties, Log, All);
+
+class UInputDeviceProperty;
+
+/** Contains a pointer to an active device property and keeps track of how long it has been evaluated for */
+USTRUCT()
+struct FActiveDeviceProperty
+{
+	GENERATED_BODY()
+
+	/** The active device property */
+	UPROPERTY(Transient)
+	TObjectPtr<UInputDeviceProperty> Property = nullptr;
+
+	/** How long this property has been evaluated for. DeltaTime is added to this on tick */
+	double EvaluatedDuration = 0.0;
+
+	/** The platform user that is actively receiving this device property */
+	FPlatformUserId PlatformUser = PLATFORMUSERID_NONE;
+
+	/**
+	* If true, then when the InputDeviceProperty is done being evaluated and has fulfilled its duration, then 
+	* it will have it's Reset function called. Set this to false if you want your device property to stay 
+	* in it's ending state. 
+	*/
+	bool bResetUponCompletion = true;
+};
+
+/** Parameters for the UInputDeviceSubsystem::SetDeviceProperty function */
+USTRUCT(BlueprintType)
+struct FSetDevicePropertyParams
+{
+	GENERATED_BODY()
+
+	FSetDevicePropertyParams();
+	
+	/** The device Property to set */
+	UPROPERTY(EditAnywhere, BlueprintReadWrite, Category = "Input Devices")
+	TSubclassOf<UInputDeviceProperty> DevicePropertyClass;
+	
+	/** The Platform User whose device's should receive the device property */
+	UPROPERTY(EditAnywhere, BlueprintReadWrite, Category = "Input Devices")
+	FPlatformUserId UserId;
+
+	/**
+	* If true, then when the InputDeviceProperty is done being evaluated and has fulfilled its duration, then
+	* it will have it's Reset function called. Set this to false if you want your device property to stay
+	* in it's ending state.
+	*/
+	UPROPERTY(EditAnywhere, BlueprintReadWrite, Category = "Input Devices")
+	bool bResetUponCompletion = true;
+};
+
+/**
+* The input device subsystem provides an interface to allow users to set Input Device Properties
+* on any Platform User. 
+*/
+UCLASS(BlueprintType)
+class UInputDeviceSubsystem : public UEngineSubsystem, public FTickableGameObject
+{
+	GENERATED_BODY()
+
+public:
+
+	//~ Begin FTickableGameObject interface	
+	virtual UWorld* GetTickableGameObjectWorld() const override;
+	virtual ETickableTickType GetTickableTickType() const override;
+	virtual bool IsAllowedToTick() const override;
+	virtual bool IsTickableInEditor() const override;
+	virtual TStatId GetStatId() const override;
+	virtual void Tick(float DeltaTime) override;
+	//~ End FTickableGameObject interface
+
+	/** Set the given device property, which will start the evaluation and application of it to the given platform user. */
+	UFUNCTION(BlueprintCallable, Category = "Input Devices", meta = (AutoCreateRefTerm = "Params"))
+	void SetDeviceProperty(const FSetDevicePropertyParams& Params);
+
+	/** 
+	* Remove any active device properties that have the same class as the one given.
+	* 
+	* @param UserId					The PlatformUser that would have the device property applied to them
+	* @param DevicePropertyClass	The Device Property type to remove
+	* 
+	* @return						The number of removed device properties. 
+	*/
+	UFUNCTION(BlueprintCallable, Category = "Input Devices", meta = (ReturnDisplayName = "Num Removed"))
+	int32 RemoveDeviceProperty(const FPlatformUserId UserId, TSubclassOf<UInputDeviceProperty> DevicePropertyClass);
+
+protected:
+
+	/**
+	* Array of the active device properties that are currently being evaluated on Tick.
+	*/
+	UPROPERTY(Transient)
+	TArray<FActiveDeviceProperty> ActiveProperties;
+};
\ No newline at end of file
diff --git a/Engine/Source/Runtime/Engine/Private/GameFramework/InputDeviceSubsystem.cpp b/Engine/Source/Runtime/Engine/Private/GameFramework/InputDeviceSubsystem.cpp
new file mode 100644
--- /dev/null
+++ b/Engine/Source/Runtime/Engine/Private/GameFramework/InputDeviceSubsystem.cpp
@@ -0,0 +1,150 @@
+// Copyright Epic Games, Inc. All Rights Reserved.
+
+#include "GameFramework/InputDeviceSubsystem.h"
+#include "GameFramework/InputDeviceProperties.h"
+#include "UnrealEngine.h"	// For GEngine
+#include "Engine/Engine.h"	// For FWorldContext
+#include "Application/SlateApplicationBase.h"
+
+DEFINE_LOG_CATEGORY(LogInputDeviceProperties);
+
+FSetDevicePropertyParams::FSetDevicePropertyParams()
+	: UserId(FSlateApplicationBase::SlateAppPrimaryPlatformUser)
+{
+
+}
+
+UWorld* UInputDeviceSubsystem::GetTickableGameObjectWorld() const
+{
+	// Use the default world by default...
+	UWorld* World = GetWorld();
+
+	// ...but if we don't have one (i.e. we are in the editor and not PIE'ing)
+	// then we need to get the editor world. This will let us preview
+	// device properties without needing to actually PIE every time
+	if (!World && GEngine)
+	{
+		for (const FWorldContext& Context : GEngine->GetWorldContexts())
+		{
+			UWorld* ThisWorld = Context.World();
+			if (!ThisWorld)
+			{
+				continue;
+			}
+			// Prefer new PIE window worlds
+			else if (Context.WorldType == EWorldType::PIE)
+			{
+				World = ThisWorld;
+				break;
+			}
+			// Fallback to the editor world, which is still valid for previewing device properties
+			else if (Context.WorldType == EWorldType::Editor)
+			{
+				World = ThisWorld;
+			}
+		}
+	}
+
+	return World;
+}
+
+ETickableTickType UInputDeviceSubsystem::GetTickableTickType() const
+{
+	return ETickableTickType::Conditional;
+}
+
+bool UInputDeviceSubsystem::IsAllowedToTick() const
+{
+	// Only tick when there are active device properties
+	return !ActiveProperties.IsEmpty();
+}
+
+bool UInputDeviceSubsystem::IsTickableInEditor() const
+{
+	// We want to tick in editor to allow previewing of device properties
+	return true;
+}
+
+TStatId UInputDeviceSubsystem::GetStatId() const
+{
+	RETURN_QUICK_DECLARE_CYCLE_STAT(UInputDeviceSubsystem, STATGROUP_Tickables);
+}
+
+void UInputDeviceSubsystem::Tick(float DeltaTime)
+{
+	for (int32 Index = ActiveProperties.Num() - 1; Index >= 0; --Index)
+	{
+		if (!ActiveProperties[Index].Property)
+		{
+			// Something has gone wrong if we get here... maybe the Property has been GC'd ?
+			// This is really just an emergency handling case
+			ensure(false);
+			continue;
+		}
+
+		// Increase the evaluated time of this property
+		ActiveProperties[Index].EvaluatedDuration += DeltaTime;
+
+		// If the property has run past it's duration, reset it and remove it from our active properties
+		if (ActiveProperties[Index].EvaluatedDuration > ActiveProperties[Index].Property->GetDuration())
+		{
+			// Evaluation complete, mark this one for removal
+			if (ActiveProperties[Index].bResetUponCompletion)
+			{
+				ActiveProperties[Index].Property->ResetDeviceProperty(ActiveProperties[Index].PlatformUser);
+			}			
+			ActiveProperties.RemoveAtSwap(Index);
+		}
+		// Otherwise, we can evaluate and apply it as normal
+		else
+		{
+			ActiveProperties[Index].Property->EvaluateDeviceProperty(ActiveProperties[Index].PlatformUser, DeltaTime, ActiveProperties[Index].EvaluatedDuration);
+			ActiveProperties[Index].Property->ApplyDeviceProperty(ActiveProperties[Index].PlatformUser);
+		}
+	}
+}
+
+void UInputDeviceSubsystem::SetDeviceProperty(const FSetDevicePropertyParams& Params)
+{
+	if (!Params.DevicePropertyClass)
+	{
+		UE_LOG(LogInputDeviceProperties, Error, TEXT("Invalid DevicePropertyClass passed into SetDeviceProperty! Nothing will happen."));
+		return;
+	}
+	
+	FActiveDeviceProperty ActiveProp = {};
+
+	// Spawn an instance of this device property
+	ActiveProp.Property = NewObject<UInputDeviceProperty>(/* Outer = */ this, /* Class */ Params.DevicePropertyClass);
+	ensure(ActiveProp.Property);
+
+	ActiveProp.PlatformUser = Params.UserId;
+	ActiveProp.bResetUponCompletion = Params.bResetUponCompletion;
+	
+	ActiveProperties.Emplace(ActiveProp);
+}
+
+int32 UInputDeviceSubsystem::RemoveDeviceProperty(const FPlatformUserId UserId, TSubclassOf<UInputDeviceProperty> DevicePropertyClass)
+{
+	int32 NumRemoved = 0;
+
+	if (DevicePropertyClass)
+	{
+		// Remove all active properties that are of the same class type
+		for (int32 Index = ActiveProperties.Num() - 1; Index >= 0; --Index)
+		{
+			if (ActiveProperties[Index].PlatformUser == UserId && ActiveProperties[Index].Property && ActiveProperties[Index].Property->GetClass() == DevicePropertyClass)
+			{
+				ActiveProperties[Index].Property->ResetDeviceProperty(ActiveProperties[Index].PlatformUser);
+				ActiveProperties.RemoveAtSwap(Index);
+				++NumRemoved;
+			}
+		}
+	}	
+	else
+	{
+		UE_LOG(LogInputDeviceProperties, Error, TEXT("Invalid DevicePropertyClass passed into RemoveDeviceProperty! Nothing will happen."));
+	}
+
+	return NumRemoved;
+}
\ No newline at end of file
