WorldPartitionBlueprintLibrary: add support for pin/unpin actors.
Note: also fixed UnloadActors parameter, existing scripts will need to be updated.

#rnx

diff --git a/Engine/Source/Runtime/Engine/Private/WorldPartitionBlueprintLibrary.cpp b/Engine/Source/Runtime/Engine/Private/WorldPartitionBlueprintLibrary.cpp
--- a/Engine/Source/Runtime/Engine/Private/WorldPartitionBlueprintLibrary.cpp
+++ b/Engine/Source/Runtime/Engine/Private/WorldPartitionBlueprintLibrary.cpp
@@ -207,19 +207,39 @@ void UWorldPartitionBlueprintLibrary::LoadActors(const TArray<FGuid>& InActorsTo
 #endif
 }
 
-void UWorldPartitionBlueprintLibrary::UnloadActors(const TArray<FGuid>& InActorsToLoad)
+void UWorldPartitionBlueprintLibrary::UnloadActors(const TArray<FGuid>& InActorsToUnload)
 {
 #if WITH_EDITOR
 	if (UWorldPartition* WorldPartition = GetWorldPartition())
 	{
 		if (TUniquePtr<FLoaderAdapterActorList>* LoaderAdapterActorList = LoaderAdapterActorListMap.Find(WorldPartition))
 		{
-			(*LoaderAdapterActorList)->RemoveActors(InActorsToLoad);
+			(*LoaderAdapterActorList)->RemoveActors(InActorsToUnload);
 		}
 	}
 #endif
 }
 
+void UWorldPartitionBlueprintLibrary::PinActors(const TArray<FGuid>& InActorsToPin)
+{
+#if WITH_EDITOR
+	if (UWorldPartition* WorldPartition = GetWorldPartition())
+	{
+		WorldPartition->PinActors(InActorsToPin);
+	}
+#endif
+}
+
+void UWorldPartitionBlueprintLibrary::UnpinActors(const TArray<FGuid>& InActorsToUnpin)
+{
+#if WITH_EDITOR
+	if (UWorldPartition* WorldPartition = GetWorldPartition())
+	{
+		WorldPartition->UnpinActors(InActorsToUnpin);
+	}
+#endif
+}
+
 bool UWorldPartitionBlueprintLibrary::GetActorDescs(TArray<FActorDesc>& OutActorDescs)
 {
 #if WITH_EDITOR
diff --git a/Engine/Source/Runtime/Engine/Public/WorldPartition/WorldPartitionBlueprintLibrary.h b/Engine/Source/Runtime/Engine/Public/WorldPartition/WorldPartitionBlueprintLibrary.h
--- a/Engine/Source/Runtime/Engine/Public/WorldPartition/WorldPartitionBlueprintLibrary.h
+++ b/Engine/Source/Runtime/Engine/Public/WorldPartition/WorldPartitionBlueprintLibrary.h
@@ -112,7 +112,19 @@ public:
 	 * Unload actors
 	 */
 	UFUNCTION(BlueprintCallable, Category="World Partition")
-	static void UnloadActors(const TArray<FGuid>& InActorsToLoad);
+	static void UnloadActors(const TArray<FGuid>& InActorsToUnload);
+
+	/**
+	 * Pin actors
+	 */
+	UFUNCTION(BlueprintCallable, Category="World Partition")
+	static void PinActors(const TArray<FGuid>& InActorsToPin);
+
+	/**
+	 * Unpin actors
+	 */
+	UFUNCTION(BlueprintCallable, Category="World Partition")
+	static void UnpinActors(const TArray<FGuid>& InActorsToUnpin);
 
 	/**
 	 * Gets all the actor descriptors into the provided array, recursing into actor containers.
