WorldPartitionStreamingPolicy::UpdateStreamingState optim : Detect if nothing relevant to world partition streaming changed and early-out.
- Added wp.Runtime.UpdateStreamingOptim.Enable to enable it (off by default).
- Added advanced options to control how aggressif is the optimization (wp.Runtime.UpdateStreamingOptim.LocationQuantization and wp.Runtime.UpdateStreamingOptim.RotationQuantization).

diff --git a/Engine/Source/Runtime/Engine/Private/WorldPartition/DataLayer/WorldDataLayers.cpp b/Engine/Source/Runtime/Engine/Private/WorldPartition/DataLayer/WorldDataLayers.cpp
--- a/Engine/Source/Runtime/Engine/Private/WorldPartition/DataLayer/WorldDataLayers.cpp
+++ b/Engine/Source/Runtime/Engine/Private/WorldPartition/DataLayer/WorldDataLayers.cpp
@@ -222,12 +222,14 @@ void AWorldDataLayers::OnDataLayerRuntimeStateChanged_Implementation(const UData
 
 void AWorldDataLayers::OnRep_ActiveDataLayerNames()
 {
+	++DataLayersStateEpoch;
 	ActiveDataLayerNames.Reset();
 	ActiveDataLayerNames.Append(RepActiveDataLayerNames);
 }
 
 void AWorldDataLayers::OnRep_LoadedDataLayerNames()
 {
+	++DataLayersStateEpoch;
 	LoadedDataLayerNames.Reset();
 	LoadedDataLayerNames.Append(RepLoadedDataLayerNames);
 }
@@ -250,6 +252,7 @@ EDataLayerRuntimeState AWorldDataLayers::GetDataLayerRuntimeStateByName(FName In
 
 void AWorldDataLayers::OnRep_EffectiveActiveDataLayerNames()
 {
+	++DataLayersStateEpoch;
 	EffectiveActiveDataLayerNames.Reset();
 	EffectiveActiveDataLayerNames.Append(RepEffectiveActiveDataLayerNames);
 	if (UDataLayerSubsystem* DataLayerSubsystem = UWorld::GetSubsystem<UDataLayerSubsystem>(GetWorld()))
@@ -260,6 +263,7 @@ void AWorldDataLayers::OnRep_EffectiveActiveDataLayerNames()
 
 void AWorldDataLayers::OnRep_EffectiveLoadedDataLayerNames()
 {
+	++DataLayersStateEpoch;
 	EffectiveLoadedDataLayerNames.Reset();
 	EffectiveLoadedDataLayerNames.Append(RepEffectiveLoadedDataLayerNames);
 	if (UDataLayerSubsystem* DataLayerSubsystem = UWorld::GetSubsystem<UDataLayerSubsystem>(GetWorld()))
diff --git a/Engine/Source/Runtime/Engine/Private/WorldPartition/WorldPartition.cpp b/Engine/Source/Runtime/Engine/Private/WorldPartition/WorldPartition.cpp
--- a/Engine/Source/Runtime/Engine/Private/WorldPartition/WorldPartition.cpp
+++ b/Engine/Source/Runtime/Engine/Private/WorldPartition/WorldPartition.cpp
@@ -1341,8 +1341,12 @@ void UWorldPartition::Tick(float DeltaSeconds)
 #endif
 }
 
+DECLARE_CYCLE_STAT(TEXT("World Partition Update Streaming"), STAT_WorldPartitionUpdateStreaming, STATGROUP_Engine);
+
 void UWorldPartition::UpdateStreamingState()
 {
+	SCOPE_CYCLE_COUNTER(STAT_WorldPartitionUpdateStreaming);
+
 	if (GetWorld()->IsGameWorld())
 	{
 		StreamingPolicy->UpdateStreamingState();
diff --git a/Engine/Source/Runtime/Engine/Private/WorldPartition/WorldPartitionRuntimeSpatialHash.cpp b/Engine/Source/Runtime/Engine/Private/WorldPartition/WorldPartitionRuntimeSpatialHash.cpp
--- a/Engine/Source/Runtime/Engine/Private/WorldPartition/WorldPartitionRuntimeSpatialHash.cpp
+++ b/Engine/Source/Runtime/Engine/Private/WorldPartition/WorldPartitionRuntimeSpatialHash.cpp
@@ -1947,6 +1947,12 @@ bool UWorldPartitionRuntimeSpatialHash::ContainsRuntimeHash(const FString& Name)
 	return bResult;
 }
 
+bool UWorldPartitionRuntimeSpatialHash::IsStreaming3D() const
+{
+	// If using Z culling, return true.
+	return GetEffectiveEnableZCulling(bEnableZCulling);
+}
+
 TArray<const FSpatialHashStreamingGrid*> UWorldPartitionRuntimeSpatialHash::GetFilteredStreamingGrids() const
 {
 	TArray<const FSpatialHashStreamingGrid*> FilteredStreamingGrids;
diff --git a/Engine/Source/Runtime/Engine/Private/WorldPartition/WorldPartitionStreamingPolicy.cpp b/Engine/Source/Runtime/Engine/Private/WorldPartition/WorldPartitionStreamingPolicy.cpp
--- a/Engine/Source/Runtime/Engine/Private/WorldPartition/WorldPartitionStreamingPolicy.cpp
+++ b/Engine/Source/Runtime/Engine/Private/WorldPartition/WorldPartitionStreamingPolicy.cpp
@@ -16,6 +16,7 @@
 #include "WorldPartition/HLOD/HLODActor.h"
 #include "WorldPartition/HLOD/HLODSubsystem.h"
 #include "WorldPartition/ContentBundle/ContentBundle.h"
+#include "WorldPartition/ContentBundle/ContentBundleWorldSubsystem.h"
 #include "WorldPartition/WorldPartitionDebugHelper.h"
 #include "Engine/World.h"
 #include "Engine/Level.h"
@@ -55,6 +56,35 @@ static FAutoConsoleVariableRef CVarBlockOnSlowStreaming(
 	GBlockOnSlowStreaming,
 	TEXT("Set if streaming needs to block when to slow to catchup."));
 
+bool UWorldPartitionStreamingPolicy::IsUpdateOptimEnabled = false;
+FAutoConsoleVariableRef UWorldPartitionStreamingPolicy::CVarUpdateOptimEnabled(
+	TEXT("wp.Runtime.UpdateStreaming.EnableOptimization"),
+	UWorldPartitionStreamingPolicy::IsUpdateOptimEnabled,
+	TEXT("Set to 1 to enable an optimization that skips world partition streaming update\n")
+	TEXT("if nothing relevant changed since last update."),
+	ECVF_Default);
+
+int32 UWorldPartitionStreamingPolicy::LocationQuantization = 400;
+FAutoConsoleVariableRef UWorldPartitionStreamingPolicy::CVarLocationQuantization(
+	TEXT("wp.Runtime.UpdateStreaming.LocationQuantization"),
+	UWorldPartitionStreamingPolicy::LocationQuantization,
+	TEXT("Distance (in Unreal units) used to quantize the streaming sources location to determine if a world partition streaming update is necessary."),
+	ECVF_Default);
+
+int32 UWorldPartitionStreamingPolicy::RotationQuantization = 10;
+FAutoConsoleVariableRef UWorldPartitionStreamingPolicy::CVarRotationQuantization(
+	TEXT("wp.Runtime.UpdateStreaming.RotationQuantization"),
+	UWorldPartitionStreamingPolicy::RotationQuantization,
+	TEXT("Angle (in degrees) used to quantize the streaming sources rotation to determine if a world partition streaming update is necessary."),
+	ECVF_Default);
+
+int32 UWorldPartitionStreamingPolicy::ForceUpdateFrameCount = 30;
+FAutoConsoleVariableRef UWorldPartitionStreamingPolicy::CVarForceUpdateFrameCount(
+	TEXT("wp.Runtime.UpdateStreaming.ForceUpdateFrameCount"),
+	UWorldPartitionStreamingPolicy::ForceUpdateFrameCount,
+	TEXT("Frequency (in frames) at which world partition streaming update will be executed regardless if no changes are detected."),
+	ECVF_Default);
+
 static void SortStreamingCellsByImportance(TArray<const UWorldPartitionRuntimeCell*>& InOutCells)
 {
 	if (InOutCells.Num() > 1)
@@ -66,11 +96,14 @@ static void SortStreamingCellsByImportance(TArray<const UWorldPartitionRuntimeCe
 
 UWorldPartitionStreamingPolicy::UWorldPartitionStreamingPolicy(const FObjectInitializer& ObjectInitializer)
 	: Super(ObjectInitializer) 
+	, bLastUpdateCompletedLoadingAndActivation(false)
 	, bCriticalPerformanceRequestedBlockTillOnWorld(false)
 	, CriticalPerformanceBlockTillLevelStreamingCompletedEpoch(0)
 	, DataLayersStatesServerEpoch(INT_MIN)
 	, ContentBundleServerEpoch(INT_MIN)
 	, ServerStreamingEnabledEpoch(INT_MIN)
+	, UpdateStreamingHash(0)
+	, UpdateStreamingStateCalls(0)
 	, StreamingPerformance(EWorldPartitionStreamingPerformance::Good)
 #if !UE_BUILD_SHIPPING
 	, OnScreenMessageStartTime(0.0)
@@ -197,10 +230,93 @@ bool UWorldPartitionStreamingPolicy::IsInBlockTillLevelStreamingCompleted(bool b
 	return bIsInBlockTillLevelStreamingCompleted;
 }
 
+int32 UWorldPartitionStreamingPolicy::ComputeServerStreamingEnabledEpoch() const
+{
+	return WorldPartition->IsServer() ? (WorldPartition->IsServerStreamingEnabled() ? 1 : 0) : INT_MIN;
+}
+
+bool UWorldPartitionStreamingPolicy::IsUpdateStreamingOptimEnabled()
+{
+	return UWorldPartitionStreamingPolicy::IsUpdateOptimEnabled &&
+		(UWorldPartitionStreamingPolicy::LocationQuantization > 0) &&
+		(UWorldPartitionStreamingPolicy::RotationQuantization > 0);
+};
+
+uint32 UWorldPartitionStreamingPolicy::ComputeUpdateStreamingHash() const
+{
+	const UWorldPartitionRuntimeHash* RuntimeHash = WorldPartition->RuntimeHash;
+	const bool bForceFrameUpdate = (UWorldPartitionStreamingPolicy::ForceUpdateFrameCount > 0) ? ((UpdateStreamingStateCalls % UWorldPartitionStreamingPolicy::ForceUpdateFrameCount) == 0) : false;
+
+	const bool bCanOptimize = 
+		RuntimeHash &&
+		IsUpdateStreamingOptimEnabled() &&							// Check CVars to see if optimization is enabled
+		!WorldPartition->IsServer() &&								// Don't optimize on the server
+		bLastUpdateCompletedLoadingAndActivation &&					// Don't optimize if last frame didn't process all cells to load/activate
+		!IsInBlockTillLevelStreamingCompleted() &&					// Don't optimize when inside UWorld::BlockTillLevelStreamingCompleted
+		ActivatedCells.GetPendingAddToWorldCells().IsEmpty() &&		// Don't optimize when remaining cells to add to world
+		!bForceFrameUpdate;											// We garantee to update every N frame to force some internal updates like UpdateStreamingPerformance
+
+	if (bCanOptimize)
+	{
+		// Build hash that will be used to detect relevant changes
+		uint32 NewHash = GetTypeHash(ComputeServerStreamingEnabledEpoch());
+		NewHash = HashCombine(NewHash, GetTypeHash(FContentBundle::GetContentBundleEpoch()));
+		NewHash = HashCombine(NewHash, GetTypeHash(AWorldDataLayers::GetDataLayersStateEpoch()));
+		NewHash = HashCombine(NewHash, GetTypeHash(RuntimeHash->IsStreaming3D()));
+		for (const FWorldPartitionStreamingSource& Source : StreamingSources)
+		{
+			NewHash = HashCombine(NewHash, ComputeStreamingSourceHash(Source));
+		}
+		return NewHash;
+	}
+
+	return 0;
+};
+
+uint32 UWorldPartitionStreamingPolicy::ComputeStreamingSourceHash(const FWorldPartitionStreamingSource& Source) const
+{
+	uint32 Hash = (GetTypeHash(Source.Name));
+	Hash = HashCombine(Hash, GetTypeHash(FMath::FloorToInt(Source.Location.X / UWorldPartitionStreamingPolicy::LocationQuantization)));
+	Hash = HashCombine(Hash, GetTypeHash(FMath::FloorToInt(Source.Location.Y / UWorldPartitionStreamingPolicy::LocationQuantization)));
+	Hash = HashCombine(Hash, GetTypeHash(FMath::FloorToInt(Source.Rotation.Yaw / UWorldPartitionStreamingPolicy::RotationQuantization)));
+	// Only consider Z position and pitch/roll rotations when hash is streaming in 3D
+	if (WorldPartition->RuntimeHash && WorldPartition->RuntimeHash->IsStreaming3D())
+	{
+		Hash = HashCombine(Hash, GetTypeHash(FMath::FloorToInt(Source.Location.Z / UWorldPartitionStreamingPolicy::LocationQuantization)));
+		Hash = HashCombine(Hash, GetTypeHash(FMath::FloorToInt(Source.Rotation.Pitch / UWorldPartitionStreamingPolicy::RotationQuantization)));
+		Hash = HashCombine(Hash, GetTypeHash(FMath::FloorToInt(Source.Rotation.Roll / UWorldPartitionStreamingPolicy::RotationQuantization)));
+	}
+	Hash = HashCombine(Hash, GetTypeHash(Source.TargetState));
+	Hash = HashCombine(Hash, GetTypeHash(Source.bBlockOnSlowLoading));
+	Hash = HashCombine(Hash, GetTypeHash(Source.bReplay));
+	Hash = HashCombine(Hash, GetTypeHash(Source.bRemote));
+	Hash = HashCombine(Hash, GetTypeHash(Source.Priority));
+	Hash = HashCombine(Hash, GetTypeHash(Source.TargetBehavior));
+	for (FName TargetGrid : Source.TargetGrids)
+	{
+		Hash = HashCombine(Hash, GetTypeHash(TargetGrid));
+	}
+	for (const FSoftObjectPath& TargetHLODLayer : Source.TargetHLODLayers)
+	{
+		Hash = HashCombine(Hash, GetTypeHash(TargetHLODLayer));
+	}
+	PRAGMA_DISABLE_DEPRECATION_WARNINGS
+	Hash = HashCombine(Hash, GetTypeHash(Source.TargetGrid));
+	Hash = HashCombine(Hash, GetTypeHash(Source.TargetHLODLayer));
+	PRAGMA_ENABLE_DEPRECATION_WARNINGS
+	for (const FStreamingSourceShape& Shape : Source.Shapes)
+	{
+		Hash = HashCombine(Hash, GetTypeHash(Shape));
+	}
+	return Hash;
+};
+
 void UWorldPartitionStreamingPolicy::UpdateStreamingState()
 {
 	TRACE_CPUPROFILER_EVENT_SCOPE(UWorldPartitionStreamingPolicy::UpdateStreamingState);
 
+	++UpdateStreamingStateCalls;
+
 	UWorld* World = GetWorld();
 	check(World);
 	check(World->IsGameWorld());
@@ -214,7 +330,18 @@ void UWorldPartitionStreamingPolicy::UpdateStreamingState()
 	
 	// Update streaming sources
 	UpdateStreamingSources();
-	
+
+	// Detect if nothing relevant changed and early out
+	const uint32 NewUpdateStreamingHash = ComputeUpdateStreamingHash();
+	const bool bShouldSkipUpdate = NewUpdateStreamingHash && (UpdateStreamingHash == NewUpdateStreamingHash);
+	if (bShouldSkipUpdate)
+	{
+		return;
+	}
+
+	// Update new streaming sources hash
+	UpdateStreamingHash = NewUpdateStreamingHash;
+
 	TSet<FName> ClientVisibleLevelNames;
 	bool bUpdateEpoch = false;
 
@@ -231,7 +358,7 @@ void UWorldPartitionStreamingPolicy::UpdateStreamingState()
 	const bool bIsServer = WorldPartition->IsServer();
 	const bool bIsServerStreamingEnabled = WorldPartition->IsServerStreamingEnabled();
 	const bool bCanDeactivateOrUnloadCells = !bIsServer || WorldPartition->IsServerStreamingOutEnabled();
-	const int32 NewServerStreamingEnabledEpoch = bIsServer ? (bIsServerStreamingEnabled ? 1 : 0) : INT_MIN;
+	const int32 NewServerStreamingEnabledEpoch = ComputeServerStreamingEnabledEpoch();
 
 	if (!bIsServer || bIsServerStreamingEnabled || AWorldPartitionReplay::IsPlaybackEnabled(World))
 	{
@@ -434,17 +561,21 @@ void UWorldPartitionStreamingPolicy::UpdateStreamingState()
 		SetCellsStateToUnloaded(ToUnloadCells);
 	}
 
+	int32 ToLoadAndActivateCount = 0;
+
 	// Do Activation State first as it is higher prio than Load State (if we have a limited number of loading cells per frame)
 	if (ToActivateCells.Num() > 0)
 	{
-		SetCellsStateToActivated(ToActivateCells);
+		ToLoadAndActivateCount += SetCellsStateToActivated(ToActivateCells);
 	}
 
 	if (ToLoadCells.Num() > 0)
 	{
-		SetCellsStateToLoaded(ToLoadCells);
+		ToLoadAndActivateCount += SetCellsStateToLoaded(ToLoadCells);
 	}
 
+	bLastUpdateCompletedLoadingAndActivation = (ToLoadAndActivateCount == (ToActivateCells.Num() + ToLoadCells.Num()));
+
 	// Sort cells and update streaming priority 
 	{
 		TRACE_CPUPROFILER_EVENT_SCOPE(UWorldPartitionStreamingPolicy::SortCellsAndUpdateStreamingPriority);
@@ -599,12 +730,14 @@ int32 UWorldPartitionStreamingPolicy::GetMaxCellsToLoad() const
 	return MAX_int32;
 }
 
-void UWorldPartitionStreamingPolicy::SetCellsStateToLoaded(const TArray<const UWorldPartitionRuntimeCell*>& ToLoadCells)
+int32 UWorldPartitionStreamingPolicy::SetCellsStateToLoaded(const TArray<const UWorldPartitionRuntimeCell*>& ToLoadCells)
 {
 	TRACE_CPUPROFILER_EVENT_SCOPE(UWorldPartitionStreamingPolicy::SetCellsStateToLoaded);
 
 	int32 MaxCellsToLoad = GetMaxCellsToLoad();
 
+	int32 LoadedCount = 0;
+
 	// Trigger cell loading. Depending on actual state of cell limit loading.
 	for (const UWorldPartitionRuntimeCell* Cell : ToLoadCells)
 	{
@@ -615,25 +748,30 @@ void UWorldPartitionStreamingPolicy::SetCellsStateToLoaded(const TArray<const UW
 			Cell->Deactivate();
 			ActivatedCells.Remove(Cell);
 			LoadedCells.Add(Cell);
+			++LoadedCount;
 		}
 		else if (MaxCellsToLoad > 0)
 		{
 			Cell->Load();
 			LoadedCells.Add(Cell);
+			++LoadedCount;
 			if (!Cell->IsAlwaysLoaded())
 			{
 				--MaxCellsToLoad;
 			}
 		}
 	}
+
+	return LoadedCount;
 }
 
-void UWorldPartitionStreamingPolicy::SetCellsStateToActivated(const TArray<const UWorldPartitionRuntimeCell*>& ToActivateCells)
+int32 UWorldPartitionStreamingPolicy::SetCellsStateToActivated(const TArray<const UWorldPartitionRuntimeCell*>& ToActivateCells)
 {
 	TRACE_CPUPROFILER_EVENT_SCOPE(UWorldPartitionStreamingPolicy::SetCellsStateToActivated);
 
 	int32 MaxCellsToLoad = GetMaxCellsToLoad();
-
+	
+	int32 ActivatedCount = 0;
 	// Trigger cell activation. Depending on actual state of cell limit loading.
 	for (const UWorldPartitionRuntimeCell* Cell : ToActivateCells)
 	{
@@ -644,6 +782,7 @@ void UWorldPartitionStreamingPolicy::SetCellsStateToActivated(const TArray<const
 			LoadedCells.Remove(Cell);
 			ActivatedCells.Add(Cell);
 			Cell->Activate();
+			++ActivatedCount;
 		}
 		else if (MaxCellsToLoad > 0)
 		{
@@ -653,8 +792,10 @@ void UWorldPartitionStreamingPolicy::SetCellsStateToActivated(const TArray<const
 			}
 			ActivatedCells.Add(Cell);
 			Cell->Activate();
+			++ActivatedCount;
 		}
 	}
+	return ActivatedCount;
 }
 
 void UWorldPartitionStreamingPolicy::SetCellsStateToUnloaded(const TArray<const UWorldPartitionRuntimeCell*>& ToUnloadCells)
diff --git a/Engine/Source/Runtime/Engine/Public/WorldPartition/WorldPartitionRuntimeHash.h b/Engine/Source/Runtime/Engine/Public/WorldPartition/WorldPartitionRuntimeHash.h
--- a/Engine/Source/Runtime/Engine/Public/WorldPartition/WorldPartitionRuntimeHash.h
+++ b/Engine/Source/Runtime/Engine/Public/WorldPartition/WorldPartitionRuntimeHash.h
@@ -136,6 +136,7 @@ public:
 	virtual bool Draw2D(class UCanvas* Canvas, const TArray<FWorldPartitionStreamingSource>& Sources, const FVector2D& PartitionCanvasSize, const FVector2D& Offset, FVector2D& OutUsedCanvasSize) const { return false; }
 	virtual void Draw3D(const TArray<FWorldPartitionStreamingSource>& Sources) const {}
 	virtual bool ContainsRuntimeHash(const FString& Name) const { return false; }
+	virtual bool IsStreaming3D() const { return true; }
 
 protected:
 	virtual EWorldPartitionStreamingPerformance GetStreamingPerformanceForCell(const UWorldPartitionRuntimeCell* Cell) const { return EWorldPartitionStreamingPerformance::Good; }
diff --git a/Engine/Source/Runtime/Engine/Public/WorldPartition/WorldPartitionRuntimeSpatialHash.h b/Engine/Source/Runtime/Engine/Public/WorldPartition/WorldPartitionRuntimeSpatialHash.h
--- a/Engine/Source/Runtime/Engine/Public/WorldPartition/WorldPartitionRuntimeSpatialHash.h
+++ b/Engine/Source/Runtime/Engine/Public/WorldPartition/WorldPartitionRuntimeSpatialHash.h
@@ -347,6 +347,7 @@ private:
 	virtual bool Draw2D(class UCanvas* Canvas, const TArray<FWorldPartitionStreamingSource>& Sources, const FVector2D& PartitionCanvasSize, const FVector2D& Offset, FVector2D& OutUsedCanvasSize) const override;
 	virtual void Draw3D(const TArray<FWorldPartitionStreamingSource>& Sources) const override;
 	virtual bool ContainsRuntimeHash(const FString& Name) const override;
+	virtual bool IsStreaming3D() const override;
 
 private:
 	void GetAlwaysLoadedStreamingCells(const FSpatialHashStreamingGrid& StreamingGrid, TSet<const UWorldPartitionRuntimeCell*>& Cells) const;
diff --git a/Engine/Source/Runtime/Engine/Public/WorldPartition/WorldPartitionStreamingPolicy.h b/Engine/Source/Runtime/Engine/Public/WorldPartition/WorldPartitionStreamingPolicy.h
--- a/Engine/Source/Runtime/Engine/Public/WorldPartition/WorldPartitionStreamingPolicy.h
+++ b/Engine/Source/Runtime/Engine/Public/WorldPartition/WorldPartitionStreamingPolicy.h
@@ -77,13 +77,13 @@ public:
 	EWorldPartitionStreamingPerformance GetStreamingPerformance() const { return StreamingPerformance; }
 
 protected:
-	virtual void SetCellsStateToLoaded(const TArray<const UWorldPartitionRuntimeCell*>& ToLoadCells);
-	virtual void SetCellsStateToActivated(const TArray<const UWorldPartitionRuntimeCell*>& ToActivateCells);
+	virtual int32 SetCellsStateToLoaded(const TArray<const UWorldPartitionRuntimeCell*>& ToLoadCells);
+	virtual int32 SetCellsStateToActivated(const TArray<const UWorldPartitionRuntimeCell*>& ToActivateCells);
 	virtual void SetCellsStateToUnloaded(const TArray<const UWorldPartitionRuntimeCell*>& ToUnloadCells);
 	virtual int32 GetCellLoadingCount() const { return 0; }
 	virtual int32 GetMaxCellsToLoad() const;
 	virtual void UpdateStreamingSources();
-	void UpdateStreamingPerformance(const TSet<const UWorldPartitionRuntimeCell*>& CellsToActivate);
+	void UpdateStreamingPerformance(const TSet<const UWorldPartitionRuntimeCell*>& InCells);
 	bool ShouldSkipCellForPerformance(const UWorldPartitionRuntimeCell* Cell) const;
 	bool IsInBlockTillLevelStreamingCompleted(bool bIsCausedByBadStreamingPerformance = false) const;
 
@@ -117,12 +117,32 @@ protected:
 	TSet<const UWorldPartitionRuntimeCell*> FrameActivateCells;
 	TSet<const UWorldPartitionRuntimeCell*> FrameLoadCells;
 	
+private:
+	// Update optimization
+	uint32 ComputeUpdateStreamingHash() const;
+	uint32 ComputeStreamingSourceHash(const FWorldPartitionStreamingSource& Source) const;
+	int32 ComputeServerStreamingEnabledEpoch() const;
+
+	static bool IsUpdateStreamingOptimEnabled();
+
+	// CVars to control update optimization
+	static bool IsUpdateOptimEnabled;
+	static int32 LocationQuantization;
+	static int32 RotationQuantization;
+	static int32 ForceUpdateFrameCount;
+	static FAutoConsoleVariableRef CVarUpdateOptimEnabled;
+	static FAutoConsoleVariableRef CVarLocationQuantization;
+	static FAutoConsoleVariableRef CVarRotationQuantization;
+	static FAutoConsoleVariableRef CVarForceUpdateFrameCount;
+
+	bool bLastUpdateCompletedLoadingAndActivation;
 	bool bCriticalPerformanceRequestedBlockTillOnWorld;
 	int32 CriticalPerformanceBlockTillLevelStreamingCompletedEpoch;
-
 	int32 DataLayersStatesServerEpoch;
 	int32 ContentBundleServerEpoch;
 	int32 ServerStreamingEnabledEpoch;
+	uint32 UpdateStreamingHash;
+	uint32 UpdateStreamingStateCalls;
 
 	EWorldPartitionStreamingPerformance StreamingPerformance;
 #if !UE_BUILD_SHIPPING
diff --git a/Engine/Source/Runtime/Engine/Public/WorldPartition/WorldPartitionStreamingSource.h b/Engine/Source/Runtime/Engine/Public/WorldPartition/WorldPartitionStreamingSource.h
--- a/Engine/Source/Runtime/Engine/Public/WorldPartition/WorldPartitionStreamingSource.h
+++ b/Engine/Source/Runtime/Engine/Public/WorldPartition/WorldPartitionStreamingSource.h
@@ -124,6 +124,20 @@ struct FStreamingSourceShape
 	/* Streaming source shape rotation (local to streaming source). */
 	UPROPERTY(BlueprintReadWrite, EditAnywhere, Category = Streaming)
 	FRotator Rotation;
+
+	friend uint32 GetTypeHash(const FStreamingSourceShape& InShape)
+	{
+		uint32 Hash = GetTypeHash(InShape.bUseGridLoadingRange);
+		Hash = HashCombine(Hash, GetTypeHash(InShape.LoadingRangeScale));
+		Hash = HashCombine(Hash, GetTypeHash(InShape.Radius));
+		Hash = HashCombine(Hash, GetTypeHash(InShape.bIsSector));
+		Hash = HashCombine(Hash, GetTypeHash(InShape.SectorAngle));
+		Hash = HashCombine(Hash, GetTypeHash(InShape.Location));
+		Hash = HashCombine(Hash, GetTypeHash(InShape.Rotation.Pitch));
+		Hash = HashCombine(Hash, GetTypeHash(InShape.Rotation.Yaw));
+		Hash = HashCombine(Hash, GetTypeHash(InShape.Rotation.Roll));
+		return Hash;
+	}
 };
 
 UENUM(BlueprintType)
