Content bundle debug display now also shows content bundle hashed ID next to its name : This facilitates debugging as the ID is used to name the streaming cell level and the cell debug name.

diff --git a/Engine/Source/Runtime/Engine/Private/WorldPartition/ContentBundle/ContentBundleDescriptor.cpp b/Engine/Source/Runtime/Engine/Private/WorldPartition/ContentBundle/ContentBundleDescriptor.cpp
--- a/Engine/Source/Runtime/Engine/Private/WorldPartition/ContentBundle/ContentBundleDescriptor.cpp
+++ b/Engine/Source/Runtime/Engine/Private/WorldPartition/ContentBundle/ContentBundleDescriptor.cpp
@@ -1,5 +1,6 @@
 // Copyright Epic Games, Inc. All Rights Reserved.
 #include "WorldPartition/ContentBundle/ContentBundleDescriptor.h"
+#include "WorldPartition/ContentBundle/ContentBundleLog.h"
 #include "Math/Color.h"
 
 #include UE_INLINE_GENERATED_CPP_BY_NAME(ContentBundleDescriptor)
@@ -16,6 +17,12 @@ bool UContentBundleDescriptor::IsValid() const
 		&& !PackageRoot.IsEmpty();
 }
 
+FString UContentBundleDescriptor::GetContentBundleCompactString(const FGuid& InContentBundleID)
+{
+	UE_CLOG(!InContentBundleID.IsValid(), LogContentBundle, Log, TEXT("Called UContentBundleDescriptor::GetContentBundleCompactString with an invalid Content Bundle guid."))
+	return FString::Printf(TEXT("%X"), GetTypeHash(InContentBundleID));
+}
+
 #if WITH_EDITOR
 void UContentBundleDescriptor::InitializeObject(const FString& InContentBundleName, const FString& InPackageRoot)
 {
diff --git a/Engine/Source/Runtime/Engine/Private/WorldPartition/ContentBundle/ContentBundleWorldSubsystem.cpp b/Engine/Source/Runtime/Engine/Private/WorldPartition/ContentBundle/ContentBundleWorldSubsystem.cpp
--- a/Engine/Source/Runtime/Engine/Private/WorldPartition/ContentBundle/ContentBundleWorldSubsystem.cpp
+++ b/Engine/Source/Runtime/Engine/Private/WorldPartition/ContentBundle/ContentBundleWorldSubsystem.cpp
@@ -94,7 +94,7 @@ void UContentBundleManager::DrawContentBundlesStatus(const UWorld* InWorld, UCan
 
 				(*ContentBundleContainer)->ForEachContentBundle([&Offset, &Canvas, &Pos, &MaxTextWidth](FContentBundleBase* ContentBundle)
 				{
-					FWorldPartitionDebugHelper::DrawLegendItem(Canvas, *ContentBundle->GetDisplayName(), GEngine->GetSmallFont(), ContentBundle->GetDebugColor(), FColor::White, Pos, &MaxTextWidth);
+					FWorldPartitionDebugHelper::DrawLegendItem(Canvas, *FString::Printf(TEXT("%s [%s]"), *ContentBundle->GetDisplayName(), *UContentBundleDescriptor::GetContentBundleCompactString(ContentBundle->GetDescriptor()->GetGuid())), GEngine->GetSmallFont(), ContentBundle->GetDebugColor(), FColor::White, Pos, &MaxTextWidth);
 				});
 			}
 
diff --git a/Engine/Source/Runtime/Engine/Private/WorldPartition/WorldPartitionRuntimeCell.cpp b/Engine/Source/Runtime/Engine/Private/WorldPartition/WorldPartitionRuntimeCell.cpp
--- a/Engine/Source/Runtime/Engine/Private/WorldPartition/WorldPartitionRuntimeCell.cpp
+++ b/Engine/Source/Runtime/Engine/Private/WorldPartition/WorldPartitionRuntimeCell.cpp
@@ -8,6 +8,7 @@
 #include "WorldPartition/WorldPartitionLevelStreamingPolicy.h"
 #include "WorldPartition/DataLayer/DataLayerSubsystem.h"
 #include "WorldPartition/DataLayer/DataLayersID.h"
+#include "WorldPartition/ContentBundle/ContentBundleDescriptor.h"
 #include "Engine/World.h"
 
 #include UE_INLINE_GENERATED_CPP_BY_NAME(WorldPartitionRuntimeCell)
@@ -84,7 +85,7 @@ void UWorldPartitionRuntimeCell::UpdateDebugName()
 
 	if (ContentBundleID.IsValid())
 	{
-		Builder += FString::Printf(TEXT(" CB[%X]"), GetTypeHash(ContentBundleID));
+		Builder += FString::Printf(TEXT(" CB[%s]"), *UContentBundleDescriptor::GetContentBundleCompactString(ContentBundleID));
 	}
 
 	DebugInfo.Name = Builder.ToString();
diff --git a/Engine/Source/Runtime/Engine/Private/WorldPartition/WorldPartitionRuntimeSpatialHash.cpp b/Engine/Source/Runtime/Engine/Private/WorldPartition/WorldPartitionRuntimeSpatialHash.cpp
--- a/Engine/Source/Runtime/Engine/Private/WorldPartition/WorldPartitionRuntimeSpatialHash.cpp
+++ b/Engine/Source/Runtime/Engine/Private/WorldPartition/WorldPartitionRuntimeSpatialHash.cpp
@@ -1297,7 +1297,7 @@ FString UWorldPartitionRuntimeSpatialHash::GetCellNameString(FName InGridName, c
 
 	if (InContentBundleID.IsValid())
 	{
-		CellName += FString::Printf(TEXT("CB%X"), GetTypeHash(InContentBundleID));
+		CellName += FString::Printf(TEXT("CB%s"), *UContentBundleDescriptor::GetContentBundleCompactString(InContentBundleID));
 	}
 
 	return CellName;
diff --git a/Engine/Source/Runtime/Engine/Public/WorldPartition/ContentBundle/ContentBundleDescriptor.h b/Engine/Source/Runtime/Engine/Public/WorldPartition/ContentBundle/ContentBundleDescriptor.h
--- a/Engine/Source/Runtime/Engine/Public/WorldPartition/ContentBundle/ContentBundleDescriptor.h
+++ b/Engine/Source/Runtime/Engine/Public/WorldPartition/ContentBundle/ContentBundleDescriptor.h
@@ -35,6 +35,9 @@ public:
 
 	bool IsValid() const;
 
+	// Helper method that returns a compact string for a given content bundle ID
+	static FString GetContentBundleCompactString(const FGuid& InContentBundleID);
+
 private:
 	void InitDebugColor();
 
