#ue5 - Truncation fixes for misc engine headers and modules:
- GameplayMediaEncoder
- GameplayTasks

diff --git a/Engine/Source/Runtime/Engine/Classes/Components/PrimitiveComponent.h b/Engine/Source/Runtime/Engine/Classes/Components/PrimitiveComponent.h
--- a/Engine/Source/Runtime/Engine/Classes/Components/PrimitiveComponent.h
+++ b/Engine/Source/Runtime/Engine/Classes/Components/PrimitiveComponent.h
@@ -831,7 +831,7 @@ public:
 	 * @return Whether this actor was recently rendered.
 	 */
 	UFUNCTION(Category = "Rendering", BlueprintCallable, meta=(DisplayName="Was Component Recently Rendered", Keywords="scene visible"))
-	bool WasRecentlyRendered(float Tolerance = 0.2) const;
+	bool WasRecentlyRendered(float Tolerance = 0.2f) const;
 
 	void SetLastRenderTime(float InLastRenderTime);
 	float GetLastRenderTime() const { return LastRenderTime; }
diff --git a/Engine/Source/Runtime/Engine/Classes/Engine/EngineTypes.h b/Engine/Source/Runtime/Engine/Classes/Engine/EngineTypes.h
--- a/Engine/Source/Runtime/Engine/Classes/Engine/EngineTypes.h
+++ b/Engine/Source/Runtime/Engine/Classes/Engine/EngineTypes.h
@@ -2645,9 +2645,9 @@ struct FSkeletalMeshBuildSettings
 		, bUseHighPrecisionSkinWeights(false)
 		, bUseFullPrecisionUVs(false)
 		, bUseBackwardsCompatibleF16TruncUVs(false)
-		, ThresholdPosition(0.00002)
-		, ThresholdTangentNormal(0.00002)
-		, ThresholdUV(0.0009765625)
+		, ThresholdPosition(0.00002f)
+		, ThresholdTangentNormal(0.00002f)
+		, ThresholdUV(0.0009765625f)
 		, MorphThresholdPosition(0.015f)
 	{}
 
diff --git a/Engine/Source/Runtime/Engine/Classes/Engine/MeshMerging.h b/Engine/Source/Runtime/Engine/Classes/Engine/MeshMerging.h
--- a/Engine/Source/Runtime/Engine/Classes/Engine/MeshMerging.h
+++ b/Engine/Source/Runtime/Engine/Classes/Engine/MeshMerging.h
@@ -778,7 +778,7 @@ struct FMeshApproximationSettings
 
 	/** Approximation Accuracy in Meters, will determine (eg) voxel resolution */
 	UPROPERTY(BlueprintReadWrite, EditAnywhere, Category = ShapeSettings, meta = (DisplayName = "Approximation Accuracy (meters)", ClampMin = "0.001"))
-	float ApproximationAccuracy = 1.0;
+	float ApproximationAccuracy = 1.0f;
 
 	/** Maximum allowable voxel count along main directions. This is a limit on ApproximationAccuracy. Max of 1290 (1290^3 is the last integer < 2^31, using a bigger number results in failures in TArray code & probably elsewhere) */
 	UPROPERTY(BlueprintReadWrite, EditAnywhere, AdvancedDisplay, Category = ShapeSettings, meta = (ClampMin = "64", ClampMax = "1290"))
@@ -790,7 +790,7 @@ struct FMeshApproximationSettings
 
 	/** Multiplier on Approximation Accuracy used for auto-thickening */
 	UPROPERTY(BlueprintReadWrite, EditAnywhere, Category = ShapeSettings, meta = (ClampMin = "0.001", EditCondition = "bAttemptAutoThickening"))
-	float TargetMinThicknessMultiplier = 1.5;
+	float TargetMinThicknessMultiplier = 1.5f;
 
 	/** If enabled, tiny parts will be excluded from the mesh merging, which can improve performance */
 	UPROPERTY(BlueprintReadWrite, EditAnywhere, Category = ShapeSettings)
@@ -798,7 +798,7 @@ struct FMeshApproximationSettings
 
 	/** Multiplier on Approximation Accuracy used to define tiny-part threshold, using maximum bounding-box dimension */
 	UPROPERTY(BlueprintReadWrite, EditAnywhere, Category = ShapeSettings, meta = (ClampMin = "0.001", EditCondition = "bIgnoreTinyParts"))
-	float TinyPartSizeMultiplier = 0.05;
+	float TinyPartSizeMultiplier = 0.05f;
 
 
 	/** Optional methods to attempt to close off the bottom of open meshes */
@@ -808,7 +808,7 @@ struct FMeshApproximationSettings
 
 	/** Winding Threshold controls hole filling at open mesh borders. Smaller value means "more/rounder" filling */
 	UPROPERTY(BlueprintReadWrite, EditAnywhere, AdvancedDisplay, Category = ShapeSettings, meta = (ClampMin = "0.01", ClampMax = "0.99"))
-	float WindingThreshold = 0.5;
+	float WindingThreshold = 0.5f;
 
 	/** If true, topological expand/contract is used to try to fill small gaps between objects. */
 	UPROPERTY(BlueprintReadWrite, EditAnywhere, Category = ShapeSettings)
@@ -816,7 +816,7 @@ struct FMeshApproximationSettings
 
 	/** Distance in Meters to expand/contract to fill gaps */
 	UPROPERTY(BlueprintReadWrite, EditAnywhere, Category = ShapeSettings, meta = (DisplayName = "Gap Filling Distance (meters)", ClampMin = "0.001", EditCondition = "bFillGaps"))
-	float GapDistance = 0.1;
+	float GapDistance = 0.1f;
 
 
 	//
diff --git a/Engine/Source/Runtime/Engine/Classes/GameFramework/Actor.h b/Engine/Source/Runtime/Engine/Classes/GameFramework/Actor.h
--- a/Engine/Source/Runtime/Engine/Classes/GameFramework/Actor.h
+++ b/Engine/Source/Runtime/Engine/Classes/GameFramework/Actor.h
@@ -3305,7 +3305,7 @@ public:
 	 * @return Whether this actor was recently rendered.
 	 */
 	UFUNCTION(Category="Rendering", BlueprintCallable, meta=(DisplayName="Was Actor Recently Rendered", Keywords="scene visible"))
-	bool WasRecentlyRendered(float Tolerance = 0.2) const;
+	bool WasRecentlyRendered(float Tolerance = 0.2f) const;
 
 	/** Returns the most recent time any of this actor's components were rendered */
 	virtual float GetLastRenderTime() const;
diff --git a/Engine/Source/Runtime/GameplayMediaEncoder/Private/GameplayMediaEncoder.cpp b/Engine/Source/Runtime/GameplayMediaEncoder/Private/GameplayMediaEncoder.cpp
--- a/Engine/Source/Runtime/GameplayMediaEncoder/Private/GameplayMediaEncoder.cpp
+++ b/Engine/Source/Runtime/GameplayMediaEncoder/Private/GameplayMediaEncoder.cpp
@@ -420,7 +420,7 @@ void FGameplayMediaEncoder::FloatToPCM16(float const* floatSamples, int32 numSam
 	float const* ptr = floatSamples;
 	for(auto&& sample : out)
 	{
-		int32 N = *ptr >= 0 ? *ptr * int32(MAX_int16) : *ptr * (int32(MAX_int16) + 1);
+		int32 N = *ptr >= 0 ? FMath::TruncToInt32(*ptr * int32(MAX_int16)) : FMath::TruncToInt32(*ptr * (int32(MAX_int16) + 1));
 		sample = static_cast<int16>(FMath::Clamp(N, int32(MIN_int16), int32(MAX_int16)));
 		ptr++;
 	}
@@ -703,7 +703,7 @@ void FGameplayMediaEncoder::CopyTexture(const FTexture2DRHIRef& SourceTexture, F
 		RHICmdList.BeginRenderPass(RPInfo, TEXT("CopyBackbuffer"));
 
 		{
-			RHICmdList.SetViewport(0, 0, 0.0f, DestinationTexture->GetSizeX(), DestinationTexture->GetSizeY(), 1.0f);
+			RHICmdList.SetViewport(0, 0, 0.0f, (float)DestinationTexture->GetSizeX(), (float)DestinationTexture->GetSizeY(), 1.0f);
 
 			FGraphicsPipelineStateInitializer GraphicsPSOInit;
 			RHICmdList.ApplyCachedRenderTargets(GraphicsPSOInit);
@@ -734,8 +734,8 @@ void FGameplayMediaEncoder::CopyTexture(const FTexture2DRHIRef& SourceTexture, F
 			}
 
 			RendererModule->DrawRectangle(RHICmdList, 0, 0,                // Dest X, Y
-			                              DestinationTexture->GetSizeX(),  // Dest Width
-			                              DestinationTexture->GetSizeY(),  // Dest Height
+			                              (float)DestinationTexture->GetSizeX(),  // Dest Width
+			                              (float)DestinationTexture->GetSizeY(),  // Dest Height
 			                              0, 0,                            // Source U, V
 			                              1, 1,                            // Source USize, VSize
 			                              DestinationTexture->GetSizeXY(), // Target buffer size
diff --git a/Engine/Source/Runtime/GameplayTasks/Classes/Tasks/GameplayTask_TimeLimitedExecution.h b/Engine/Source/Runtime/GameplayTasks/Classes/Tasks/GameplayTask_TimeLimitedExecution.h
--- a/Engine/Source/Runtime/GameplayTasks/Classes/Tasks/GameplayTask_TimeLimitedExecution.h
+++ b/Engine/Source/Runtime/GameplayTasks/Classes/Tasks/GameplayTask_TimeLimitedExecution.h
@@ -45,8 +45,8 @@ private:
 
 	void OnTimer();
 	
-	float Time;
-	float TimeStarted;
+	double Time;
+	double TimeStarted;
 
 	uint32 bTimeExpired : 1;
 	uint32 bChildTaskFinished : 1;
diff --git a/Engine/Source/Runtime/GameplayTasks/Classes/Tasks/GameplayTask_WaitDelay.h b/Engine/Source/Runtime/GameplayTasks/Classes/Tasks/GameplayTask_WaitDelay.h
--- a/Engine/Source/Runtime/GameplayTasks/Classes/Tasks/GameplayTask_WaitDelay.h
+++ b/Engine/Source/Runtime/GameplayTasks/Classes/Tasks/GameplayTask_WaitDelay.h
@@ -36,6 +36,6 @@ private:
 
 	void OnTimeFinish();
 	
-	float Time;
-	float TimeStarted;
+	double Time;
+	double TimeStarted;
 };
diff --git a/Engine/Source/Runtime/GameplayTasks/Private/GameplayTaskResource.cpp b/Engine/Source/Runtime/GameplayTasks/Private/GameplayTaskResource.cpp
--- a/Engine/Source/Runtime/GameplayTasks/Private/GameplayTaskResource.cpp
+++ b/Engine/Source/Runtime/GameplayTasks/Private/GameplayTaskResource.cpp
@@ -71,7 +71,7 @@ void UGameplayTaskResource::UpdateAutoResourceID()
 
 	if (AutoResourceID == INDEX_NONE)
 	{
-		AutoResourceID = NextAutoResID++;
+		AutoResourceID = static_cast<int8>(NextAutoResID++);
 		if (AutoResourceID >= FGameplayResourceSet::MaxResources)
 		{
 			UE_LOG(LogGameplayTasks, Error, TEXT("AutoResourceID out of bounds (probably too much GameplayTaskResource classes, consider manually assigning values if you can split all classes into non-overlapping sets"));
diff --git a/Engine/Source/Runtime/GameplayTasks/Private/GameplayTasksComponent.cpp b/Engine/Source/Runtime/GameplayTasks/Private/GameplayTasksComponent.cpp
--- a/Engine/Source/Runtime/GameplayTasks/Private/GameplayTasksComponent.cpp
+++ b/Engine/Source/Runtime/GameplayTasks/Private/GameplayTasksComponent.cpp
@@ -928,12 +928,13 @@ FString FGameplayResourceSet::GetDebugDescription() const
 	int32 FlagIndex = 0;
 
 #if WITH_GAMEPLAYTASK_DEBUG
+	static_assert(FlagsCount < TNumericLimits<uint8>::Max());
 	FString Description;
 	for (; FlagIndex < FlagsCount && FlagsCopy != 0; ++FlagIndex)
 	{
 		if (FlagsCopy & (1 << FlagIndex))
 		{
-			Description += UGameplayTaskResource::GetDebugDescription(FlagIndex);
+			Description += UGameplayTaskResource::GetDebugDescription( IntCastChecked<uint8>(FlagIndex) );
 			Description += TEXT(' ');
 		}
 
diff --git a/Engine/Source/Runtime/GameplayTasks/Private/Tasks/GameplayTask_TimeLimitedExecution.cpp b/Engine/Source/Runtime/GameplayTasks/Private/Tasks/GameplayTask_TimeLimitedExecution.cpp
--- a/Engine/Source/Runtime/GameplayTasks/Private/Tasks/GameplayTask_TimeLimitedExecution.cpp
+++ b/Engine/Source/Runtime/GameplayTasks/Private/Tasks/GameplayTask_TimeLimitedExecution.cpp
@@ -50,7 +50,7 @@ void UGameplayTask_TimeLimitedExecution::Activate()
 
 	// Use a dummy timer handle as we don't need to store it for later but we don't need to look for something to clear
 	FTimerHandle TimerHandle;
-	World->GetTimerManager().SetTimer(TimerHandle, this, &UGameplayTask_TimeLimitedExecution::OnTimer, Time, false);
+	World->GetTimerManager().SetTimer(TimerHandle, this, &UGameplayTask_TimeLimitedExecution::OnTimer, (float)Time, false);
 	UE_VLOG(GetGameplayTasksComponent(), LogGameplayTasks, Verbose, TEXT("%s> started timeout: %.2fs for task:%s"), *GetName(), Time, *ChildTask->GetName());
 
 	if (!ChildTask->IsActive())
@@ -97,7 +97,7 @@ void UGameplayTask_TimeLimitedExecution::OnTimer()
 
 FString UGameplayTask_TimeLimitedExecution::GetDebugString() const
 {
-	const float TimeLeft = Time - GetWorld()->TimeSince(TimeStarted);
+	const double TimeLeft = Time - GetWorld()->TimeSince(TimeStarted);
 	return FString::Printf(TEXT("TimeLimit for %s. Time: %.2f. TimeLeft: %.2f"), *GetNameSafe(GetChildTask()), Time, TimeLeft);
 }
 
diff --git a/Engine/Source/Runtime/GameplayTasks/Private/Tasks/GameplayTask_WaitDelay.cpp b/Engine/Source/Runtime/GameplayTasks/Private/Tasks/GameplayTask_WaitDelay.cpp
--- a/Engine/Source/Runtime/GameplayTasks/Private/Tasks/GameplayTask_WaitDelay.cpp
+++ b/Engine/Source/Runtime/GameplayTasks/Private/Tasks/GameplayTask_WaitDelay.cpp
@@ -48,7 +48,7 @@ void UGameplayTask_WaitDelay::Activate()
 
 	// Use a dummy timer handle as we don't need to store it for later but we don't need to look for something to clear
 	FTimerHandle TimerHandle;
-	World->GetTimerManager().SetTimer(TimerHandle, this, &UGameplayTask_WaitDelay::OnTimeFinish, Time, false);
+	World->GetTimerManager().SetTimer(TimerHandle, this, &UGameplayTask_WaitDelay::OnTimeFinish, (float)Time, false);
 }
 
 void UGameplayTask_WaitDelay::OnTimeFinish()
@@ -59,7 +59,7 @@ void UGameplayTask_WaitDelay::OnTimeFinish()
 
 FString UGameplayTask_WaitDelay::GetDebugString() const
 {
-	float TimeLeft = Time - GetWorld()->TimeSince(TimeStarted);
+	double TimeLeft = Time - GetWorld()->TimeSince(TimeStarted);
 	return FString::Printf(TEXT("WaitDelay. Time: %.2f. TimeLeft: %.2f"), Time, TimeLeft);
 }
 
diff --git a/Engine/Source/Runtime/Net/Core/Public/Net/Core/Misc/DDoSDetection.h b/Engine/Source/Runtime/Net/Core/Public/Net/Core/Misc/DDoSDetection.h
--- a/Engine/Source/Runtime/Net/Core/Public/Net/Core/Misc/DDoSDetection.h
+++ b/Engine/Source/Runtime/Net/Core/Public/Net/Core/Misc/DDoSDetection.h
@@ -158,9 +158,9 @@ struct NETCORE_API FDDoSStateConfig : public FDDoSState
 	{
 		// Exclude escalation triggers from this
 		//Target.EscalateTimeQuotaMSPerFrame		= (EscalateTimeQuotaMSPerFrame == -1 ? -1 : EscalateTimeQuotaMSPerFrame * FrameAdjustment);
-		Target.PacketLimitPerFrame				= (PacketLimitPerFrame == -1 ? -1 : (int32)(PacketLimitPerFrame * FrameAdjustment));
-		Target.PacketTimeLimitMSPerFrame		= (PacketTimeLimitMSPerFrame == -1 ? -1 : (int32)(PacketTimeLimitMSPerFrame * FrameAdjustment));
-		Target.NetConnPacketTimeLimitMSPerFrame	= (NetConnPacketTimeLimitMSPerFrame == -1 ? -1 : (int32)(NetConnPacketTimeLimitMSPerFrame * FrameAdjustment));
+		Target.PacketLimitPerFrame				= (PacketLimitPerFrame == -1 ? -1 : (int32)((float)PacketLimitPerFrame * FrameAdjustment));
+		Target.PacketTimeLimitMSPerFrame		= (PacketTimeLimitMSPerFrame == -1 ? -1 : (int32)((float)PacketTimeLimitMSPerFrame * FrameAdjustment));
+		Target.NetConnPacketTimeLimitMSPerFrame	= (NetConnPacketTimeLimitMSPerFrame == -1 ? -1 : (int32)((float)NetConnPacketTimeLimitMSPerFrame * FrameAdjustment));
 	}
 };
 
diff --git a/Engine/Source/Runtime/RenderCore/Public/RenderTransform.h b/Engine/Source/Runtime/RenderCore/Public/RenderTransform.h
--- a/Engine/Source/Runtime/RenderCore/Public/RenderTransform.h
+++ b/Engine/Source/Runtime/RenderCore/Public/RenderTransform.h
@@ -428,7 +428,7 @@ struct FCompressedTransform
 		{
 			Axis[i] = In.TransformRows[i];
 			Scale[i] = Axis[i].Size();
-			Axis[i] /= Scale[i] ? Scale[i] : 1.f;
+			Axis[i] /= (Scale[i] != 0.f) ? Scale[i] : 1.f;
 		}
 
 		// Rotation
@@ -526,9 +526,9 @@ struct FCompressedTransform
 		{
 			FVector2f OctZ;
 			float Spin0;
-			OctZ.X = ( (int32)Rotation[0] - 32768 ) * (1.0f / 32767.0f);
-			OctZ.Y = ( (int32)Rotation[1] - 32768 ) * (1.0f / 32767.0f);
-			Spin0  = ( (int32)( Rotation[2] & 0x7fff ) - 16384 ) * (0.70710678f / 16383.0f);	// rsqrt(2)
+			OctZ.X = float( (int32)Rotation[0] - 32768 ) * (1.0f / 32767.0f);
+			OctZ.Y = float( (int32)Rotation[1] - 32768 ) * (1.0f / 32767.0f);
+			Spin0  = float( (int32)( Rotation[2] & 0x7fff ) - 16384 ) * (0.70710678f / 16383.0f);	// rsqrt(2)
 			bool bSpinIsX = Rotation[2] > 0x7fff;
 
 			Out.TransformRows[2] = HemiOctahedronToUnitVector( OctZ );
diff --git a/Engine/Source/Runtime/SignalProcessing/Public/DSP/Dsp.h b/Engine/Source/Runtime/SignalProcessing/Public/DSP/Dsp.h
--- a/Engine/Source/Runtime/SignalProcessing/Public/DSP/Dsp.h
+++ b/Engine/Source/Runtime/SignalProcessing/Public/DSP/Dsp.h
@@ -237,7 +237,7 @@ namespace Audio
 				}
 
 				// Advance phase, range reduce, and store.
-				float PhaseInRadians = LastPhase + BlockSampleCount * PhasePerSample;
+				float PhaseInRadians = LastPhase + (float)BlockSampleCount * PhasePerSample;
 				PhaseInRadians -= FMath::FloorToFloat(PhaseInRadians / (2 * UE_PI)) * (2 * UE_PI);
 				LastPhase = PhaseInRadians;
 
@@ -673,7 +673,7 @@ namespace Audio
 				return CurrentValue;
 			}
 
-			CurrentValue = DeltaValue * (float) CurrentTick / DurationTicks + StartValue;
+			CurrentValue = DeltaValue * (float)CurrentTick / (float)DurationTicks + StartValue;
 
 			++CurrentTick;
 			return CurrentValue;
@@ -689,7 +689,7 @@ namespace Audio
 			}
 
 			CurrentTick = FMath::Min(CurrentTick + NumTicksToJumpAhead, DurationTicks);
-			CurrentValue = DeltaValue * (float)CurrentTick / DurationTicks + StartValue;
+			CurrentValue = DeltaValue * (float)CurrentTick / (float)DurationTicks + StartValue;
 
 			return CurrentValue;
 		}
diff --git a/Engine/Source/Runtime/SignalProcessing/Public/DSP/ParamInterpolator.h b/Engine/Source/Runtime/SignalProcessing/Public/DSP/ParamInterpolator.h
--- a/Engine/Source/Runtime/SignalProcessing/Public/DSP/ParamInterpolator.h
+++ b/Engine/Source/Runtime/SignalProcessing/Public/DSP/ParamInterpolator.h
@@ -30,7 +30,7 @@ namespace Audio
 			}
 			else
 			{
-				DeltaValue = (InValue - CurrentValue) / InNumInterpFrames;
+				DeltaValue = (InValue - CurrentValue) / (float)InNumInterpFrames;
 				StartingValue = CurrentValue;
 			}
 		}
diff --git a/Engine/Source/Runtime/Slate/Public/Widgets/Views/SListView.h b/Engine/Source/Runtime/Slate/Public/Widgets/Views/SListView.h
--- a/Engine/Source/Runtime/Slate/Public/Widgets/Views/SListView.h
+++ b/Engine/Source/Runtime/Slate/Public/Widgets/Views/SListView.h
@@ -2322,7 +2322,7 @@ protected:
 						}
 						else
 						{
-							NewScrollOffset = (int32)NewScrollOffset + ((PrecedingDistance - AbsScrollByAmount) / WidgetDimensions.ScrollAxis);
+							NewScrollOffset = float(FMath::TruncToInt32(NewScrollOffset)) + ((PrecedingDistance - AbsScrollByAmount) / WidgetDimensions.ScrollAxis);
 							break;
 						}
 					}
diff --git a/Engine/Source/Runtime/SlateCore/Public/Widgets/SWindow.h b/Engine/Source/Runtime/SlateCore/Public/Widgets/SWindow.h
--- a/Engine/Source/Runtime/SlateCore/Public/Widgets/SWindow.h
+++ b/Engine/Source/Runtime/SlateCore/Public/Widgets/SWindow.h
@@ -1026,7 +1026,7 @@ public:
 	 */
 	inline FVector2D GetViewportSize() const
 	{
-		return (ViewportSize.X) ? ViewportSize : Size;
+		return (ViewportSize.X != 0) ? ViewportSize : Size;
 	}
 	
 	/**
