fix GetGammaSpace in DDC2TextureBuilds

diff --git a/Engine/Source/Runtime/Engine/Private/TextureDerivedDataBuildUtils.cpp b/Engine/Source/Runtime/Engine/Private/TextureDerivedDataBuildUtils.cpp
--- a/Engine/Source/Runtime/Engine/Private/TextureDerivedDataBuildUtils.cpp
+++ b/Engine/Source/Runtime/Engine/Private/TextureDerivedDataBuildUtils.cpp
@@ -215,15 +215,11 @@ static void WriteSource(FCbWriter& Writer, const UTexture& Texture, int32 LayerI
 {
 	const FTextureSource& Source = Texture.Source;
 
-	FTextureFormatSettings TextureFormatSettings;
-	Texture.GetLayerFormatSettings(LayerIndex, TextureFormatSettings);
-	EGammaSpace GammaSpace = TextureFormatSettings.SRGB ? (Texture.bUseLegacyGamma ? EGammaSpace::Pow22 : EGammaSpace::sRGB) : EGammaSpace::Linear;
-
 	Writer.BeginObject();
 
 	Writer.AddInteger("CompressionFormat", Source.GetSourceCompression());
 	Writer.AddInteger("SourceFormat", Source.GetFormat(LayerIndex));
-	Writer.AddInteger("GammaSpace", static_cast<uint8>(GammaSpace));
+	Writer.AddInteger("GammaSpace", static_cast<uint8>(Source.GetGammaSpace(LayerIndex)));
 	Writer.AddInteger("NumSlices", (BuildSettings.bCubemap || BuildSettings.bTextureArray || BuildSettings.bVolume) ? Source.GetNumSlices() : 1);
 	Writer.AddInteger("SizeX", Source.GetSizeX());
 	Writer.AddInteger("SizeY", Source.GetSizeY());
