- Update Cascade PerfStats as we update the dynamic data to ensure they do not get out of sync

#rnx

diff --git a/Engine/Source/Runtime/Engine/Private/Particles/ParticleComponents.cpp b/Engine/Source/Runtime/Engine/Private/Particles/ParticleComponents.cpp
--- a/Engine/Source/Runtime/Engine/Private/Particles/ParticleComponents.cpp
+++ b/Engine/Source/Runtime/Engine/Private/Particles/ParticleComponents.cpp
@@ -4382,6 +4382,10 @@ FParticleDynamicData* UParticleSystemComponent::CreateDynamicData(ERHIFeatureLev
 		ParticleDynamicData->SystemRadiusForMacroUVs = Template->MacroUVRadius;
 	}
 
+#if WITH_PARTICLE_PERF_STATS
+	ParticleDynamicData->PerfStatContext = GetPerfStatsContext();
+#endif
+
 	if( ReplayState == PRS_Replaying )
 	{
 		SCOPE_CYCLE_COUNTER(STAT_ParticleSystemComponent_CreateDynamicData_Replay);
diff --git a/Engine/Source/Runtime/Engine/Private/Particles/ParticleSystemRender.cpp b/Engine/Source/Runtime/Engine/Private/Particles/ParticleSystemRender.cpp
--- a/Engine/Source/Runtime/Engine/Private/Particles/ParticleSystemRender.cpp
+++ b/Engine/Source/Runtime/Engine/Private/Particles/ParticleSystemRender.cpp
@@ -6519,9 +6519,6 @@ FParticleSystemSceneProxy::FParticleSystemSceneProxy(UParticleSystemComponent* C
 	, VisualizeLODIndex(Component->GetCurrentLODIndex())
 	, LastFramePreRendered(-1)
 	, FirstFreeMeshBatch(0)
-#if WITH_PARTICLE_PERF_STATS
-	, PerfStatContext(Component->GetPerfStatsContext())
-#endif
 {
 	SetWireframeColor(FLinearColor(3.0f, 0.0f, 0.0f));
 	SetLevelColor(FLinearColor(1.0f, 1.0f, 0.0f));
@@ -6694,6 +6691,9 @@ void FParticleSystemSceneProxy::UpdateData(FParticleDynamicData* NewDynamicData)
 	ENQUEUE_RENDER_COMMAND(ParticleUpdateDataCommand)(
 		[Proxy, NewDynamicData](FRHICommandListImmediate& RHICmdList)
 		{
+		#if WITH_PARTICLE_PERF_STATS
+			Proxy->PerfStatContext = NewDynamicData ? NewDynamicData->PerfStatContext : FParticlePerfStatsContext();
+		#endif
 			CSV_SCOPED_TIMING_STAT_EXCLUSIVE(ParticleUpdate);
 			SCOPE_CYCLE_COUNTER(STAT_ParticleUpdateRTTime);
 			STAT(FScopeCycleCounter Context(Proxy->GetStatId());)
@@ -6701,7 +6701,7 @@ void FParticleSystemSceneProxy::UpdateData(FParticleDynamicData* NewDynamicData)
 
 			Proxy->UpdateData_RenderThread(NewDynamicData);
 		}
-		);
+	);
 }
 
 void FParticleSystemSceneProxy::UpdateData_RenderThread(FParticleDynamicData* NewDynamicData)
diff --git a/Engine/Source/Runtime/Engine/Public/ParticleHelper.h b/Engine/Source/Runtime/Engine/Public/ParticleHelper.h
--- a/Engine/Source/Runtime/Engine/Public/ParticleHelper.h
+++ b/Engine/Source/Runtime/Engine/Public/ParticleHelper.h
@@ -2353,6 +2353,10 @@ public:
 
 	/** World space radius that UVs generated with the ParticleMacroUV material node will tile based on. */
 	float SystemRadiusForMacroUVs;
+
+#if WITH_PARTICLE_PERF_STATS
+	FParticlePerfStatsContext PerfStatContext;
+#endif
 };
 
 //
