Make device properties EditInlineNew, which I mistakenly removed with 22953412.

#br trivial
#rnx

diff --git a/Engine/Source/Runtime/Engine/Classes/GameFramework/InputDeviceProperties.h b/Engine/Source/Runtime/Engine/Classes/GameFramework/InputDeviceProperties.h
--- a/Engine/Source/Runtime/Engine/Classes/GameFramework/InputDeviceProperties.h
+++ b/Engine/Source/Runtime/Engine/Classes/GameFramework/InputDeviceProperties.h
@@ -23,7 +23,7 @@ class UCurveFloat;
 * support certain device properties. An older gamepad may not have any advanced trigger haptics for 
 * example. 
 */
-UCLASS(Abstract, Blueprintable, BlueprintType)
+UCLASS(Abstract, Blueprintable, BlueprintType, EditInlineNew)
 class ENGINE_API UInputDeviceProperty : public UObject
 {
 	GENERATED_BODY()
