Add templated version of GetComponentByClass.

Many people don't know about the templated FindComponentByClass so hopefully this will replace unnecessary cases of Cast<T>(GetComponentByClass()).

diff --git a/Engine/Source/Runtime/Engine/Classes/GameFramework/Actor.h b/Engine/Source/Runtime/Engine/Classes/GameFramework/Actor.h
--- a/Engine/Source/Runtime/Engine/Classes/GameFramework/Actor.h
+++ b/Engine/Source/Runtime/Engine/Classes/GameFramework/Actor.h
@@ -3386,6 +3386,13 @@ public:
 	UFUNCTION(BlueprintCallable, Category = "Actor", meta = (ComponentClass = "/Script/Engine.ActorComponent"), meta = (DeterminesOutputType = "ComponentClass"))
 	UActorComponent* GetComponentByClass(TSubclassOf<UActorComponent> ComponentClass) const;
 
+	/** Templated version of GetComponentByClass */
+	template<class T>
+	T* GetComponentByClass() const
+	{
+		return FindComponentByClass<T>();
+	}
+
 	/**
 	 * Gets all the components that inherit from the given class.
 	 * Currently returns an array of UActorComponent which must be cast to the correct type.
