[Engine]
Removal of includes to reduce transitive includes in high traffic headers to reduce compile times
Highlights
* Skeleton and BonePose not leaked out of commonly used animation headers
* AudioComponent leaking out less
* Brush not leaking Level.h (Brush is included indirectly a lot)
* VertexStreamComponent moved to its own file so Components.h can include that instead of VertexFactory (which would leak out all of RHI and others)

diff --git a/Engine/Plugins/Experimental/Water/Source/Runtime/Public/WaterBodyComponent.h b/Engine/Plugins/Experimental/Water/Source/Runtime/Public/WaterBodyComponent.h
--- a/Engine/Plugins/Experimental/Water/Source/Runtime/Public/WaterBodyComponent.h
+++ b/Engine/Plugins/Experimental/Water/Source/Runtime/Public/WaterBodyComponent.h
@@ -19,11 +19,10 @@
 
 #if UE_ENABLE_INCLUDE_ORDER_DEPRECATED_IN_5_2
 #include "DynamicMeshBuilder.h"
-#endif
-
 #if WITH_EDITOR
 #include "MeshDescription.h"
 #endif
+#endif
 
 #include "WaterBodyComponent.generated.h"
 
@@ -37,6 +36,7 @@ class UMaterialInstanceDynamic;
 class FTokenizedMessage;
 namespace UE::Geometry { class FDynamicMesh3; }
 struct FDynamicMeshVertex;
+struct FMeshDescription;
 
 // ----------------------------------------------------------------------------------
 
diff --git a/Engine/Plugins/Runtime/ReplicationGraph/Source/Private/ReplicationGraphTypes.cpp b/Engine/Plugins/Runtime/ReplicationGraph/Source/Private/ReplicationGraphTypes.cpp
--- a/Engine/Plugins/Runtime/ReplicationGraph/Source/Private/ReplicationGraphTypes.cpp
+++ b/Engine/Plugins/Runtime/ReplicationGraph/Source/Private/ReplicationGraphTypes.cpp
@@ -7,6 +7,7 @@
 
 #include "EngineUtils.h"
 #include "Engine/Engine.h"
+#include "Misc/ConfigCacheIni.h"
 #include "Net/DataReplication.h"
 #include "Engine/ActorChannel.h"
 #include "Engine/NetworkObjectList.h"
@@ -312,6 +313,17 @@ void ForEachClientPIEWorld(TFunction<void(UWorld*)> Func)
 }
 #endif
 
+FReplicationGraphCSVTracker::FReplicationGraphCSVTracker()
+	: EverythingElse(TEXT("Other"))
+	, EverythingElse_FastPath(TEXT("OtherFastPath"))
+	, ActorDiscovery(TEXT("ActorDiscovery"))
+{
+	ResetTrackedClasses();
+
+	GConfig->GetBool(TEXT("ReplicationGraphCSVTracker"), TEXT("bReportUntrackedClasses"), bReportUntrackedClasses, GEngineIni);
+}
+
+
 void LogListDetails(FActorRepList& RepList, FOutputDevice& Ar)
 {
 	FString ListContentString;
@@ -489,6 +501,18 @@ void FGlobalActorReplicationInfoMap::AddDependentActor(AActor* Parent, AActor* C
 	}
 }
 
+FName FNewReplicatedActorInfo::GetStreamingLevelNameOfActor(const AActor* Actor)
+{
+	ULevel* Level = Actor ? Cast<ULevel>(Actor->GetOuter()) : nullptr;
+	return (Level && Level->IsPersistentLevel() == false) ? Level->GetOutermost()->GetFName() : NAME_None;
+}
+
+FActorConnectionPair::FActorConnectionPair(AActor* InActor, UNetConnection* InConnection)
+	: Actor(InActor)
+	, Connection(InConnection)
+{
+}
+
 bool FActorRepListStatCollector::WasNodeVisited(const UReplicationGraphNode* NodeToVisit)
 {
 	const bool* Visited = VisitedNodes.Find(NodeToVisit);
diff --git a/Engine/Plugins/Runtime/ReplicationGraph/Source/Public/ReplicationGraphTypes.h b/Engine/Plugins/Runtime/ReplicationGraph/Source/Public/ReplicationGraphTypes.h
--- a/Engine/Plugins/Runtime/ReplicationGraph/Source/Public/ReplicationGraphTypes.h
+++ b/Engine/Plugins/Runtime/ReplicationGraph/Source/Public/ReplicationGraphTypes.h
@@ -3,24 +3,28 @@
 #pragma once
 
 #include "Containers/StaticArray.h"
-#include "GameFramework/WorldSettings.h"
-#include "UObject/Package.h"
 #include "EngineGlobals.h"
-#include "Engine/ActorChannel.h"
-#include "Engine/LocalPlayer.h"
 #include "GameFramework/Actor.h"
-#include "Misc/ConfigCacheIni.h"
+#include "GameFramework/WorldSettings.h"
 #include "Misc/EnumClassFlags.h"
 #include "Net/DataBunch.h"
 #include "UObject/ObjectKey.h"
 #include "UObject/UObjectHash.h"
 #include "ProfilingDebugging/CsvProfiler.h"
-#include "Engine/NetConnection.h"
 #include "HAL/LowLevelMemTracker.h"
+#include "UObject/Package.h"
 #include "ReplicationGraphTypes.generated.h"
 
+#if UE_ENABLE_INCLUDE_ORDER_DEPRECATED_IN_5_2
+#include "Engine/ActorChannel.h"
+#include "Engine/LocalPlayer.h"
+#include "Engine/NetConnection.h"
+#include "Misc/ConfigCacheIni.h"
+#endif
+
 class AActor;
 class AController;
+class UActorChannel;
 class UNetConnection;
 class UNetReplicationGraphConnection;
 class UReplicationGraph;
@@ -1352,12 +1356,7 @@ struct FNewReplicatedActorInfo
 
 	AActor* GetActor() const { return Actor; }
 
-	static FORCEINLINE FName GetStreamingLevelNameOfActor(const AActor* Actor)
-	{
-		ULevel* Level = Actor ? Cast<ULevel>(Actor->GetOuter()) : nullptr;
-		return (Level && Level->IsPersistentLevel() == false) ? Level->GetOutermost()->GetFName() : NAME_None;
-	}
-
+	REPLICATIONGRAPH_API static FName GetStreamingLevelNameOfActor(const AActor* Actor);
 
 	FActorRepListType Actor;
 	FName StreamingLevelName;
@@ -1496,15 +1495,7 @@ struct FReplicationGraphCSVTracker
 
 	FRIEND_ENUM_CLASS_FLAGS(EActorFlags)
 
-	FReplicationGraphCSVTracker()
-		: EverythingElse(TEXT("Other"))
-		, EverythingElse_FastPath(TEXT("OtherFastPath"))
-		, ActorDiscovery(TEXT("ActorDiscovery"))
-	{
-		ResetTrackedClasses();
-
-		GConfig->GetBool(TEXT("ReplicationGraphCSVTracker"), TEXT("bReportUntrackedClasses"), bReportUntrackedClasses, GEngineIni);
-	}
+	REPLICATIONGRAPH_API FReplicationGraphCSVTracker();
 
 	void TearDown()
 	{
@@ -1909,7 +1900,7 @@ ENUM_CLASS_FLAGS(FReplicationGraphCSVTracker::EActorFlags);
 struct FActorConnectionPair
 {
 	FActorConnectionPair() { }
-	FActorConnectionPair(AActor* InActor, UNetConnection* InConnection) : Actor(InActor), Connection(InConnection) { }
+	FActorConnectionPair(AActor* InActor, UNetConnection* InConnection);
 
 	TWeakObjectPtr<AActor> Actor;
 	TWeakObjectPtr<UNetConnection> Connection;
diff --git a/Engine/Source/Runtime/Engine/Classes/AI/Navigation/NavAgentInterface.h b/Engine/Source/Runtime/Engine/Classes/AI/Navigation/NavAgentInterface.h
--- a/Engine/Source/Runtime/Engine/Classes/AI/Navigation/NavAgentInterface.h
+++ b/Engine/Source/Runtime/Engine/Classes/AI/Navigation/NavAgentInterface.h
@@ -5,11 +5,15 @@
 #include "CoreMinimal.h"
 #include "UObject/ObjectMacros.h"
 #include "UObject/Interface.h"
-#include "AI/Navigation/NavigationTypes.h"
 #include "NavAgentInterface.generated.h"
 
+#if UE_ENABLE_INCLUDE_ORDER_DEPRECATED_IN_5_2
+#include "AI/Navigation/NavigationTypes.h"
+#endif
+
 class AActor;
 class IPathFollowingAgentInterface;
+struct FNavAgentProperties;
 
 UINTERFACE(MinimalAPI, meta=(CannotImplementInterfaceInBlueprint))
 class UNavAgentInterface : public UInterface
@@ -26,7 +30,7 @@ class INavAgentInterface
 	 *	@NOTE the function will be renamed to GetNavAgentProperties in 4.8. Current name was introduced
 	 *		to help with deprecating old GetNavAgentProperties function
 	 */
-	virtual const FNavAgentProperties& GetNavAgentPropertiesRef() const { return FNavAgentProperties::DefaultProperties; }
+	ENGINE_API virtual const FNavAgentProperties& GetNavAgentPropertiesRef() const;
 
 	/**
 	 *	Retrieves Agent's location
diff --git a/Engine/Source/Runtime/Engine/Classes/Animation/AnimMontage.h b/Engine/Source/Runtime/Engine/Classes/Animation/AnimMontage.h
--- a/Engine/Source/Runtime/Engine/Classes/Animation/AnimMontage.h
+++ b/Engine/Source/Runtime/Engine/Classes/Animation/AnimMontage.h
@@ -92,9 +92,7 @@ struct FSlotAnimationTrack
 	UPROPERTY(EditAnywhere, Category=Slot)
 	FAnimTrack AnimTrack;
 
-	FSlotAnimationTrack()
-		: SlotName(FAnimSlotGroup::DefaultSlotName)
-	{}
+	ENGINE_API FSlotAnimationTrack();
 };
 
 /** 
diff --git a/Engine/Source/Runtime/Engine/Classes/Animation/AnimSequence.h b/Engine/Source/Runtime/Engine/Classes/Animation/AnimSequence.h
--- a/Engine/Source/Runtime/Engine/Classes/Animation/AnimSequence.h
+++ b/Engine/Source/Runtime/Engine/Classes/Animation/AnimSequence.h
@@ -601,7 +601,7 @@ public:
 	/**
 	* Return true if compressed data is out of date / missing and so animation needs to use raw data
 	*/
-	bool DoesNeedRecompress() const { return GetSkeleton() && (bUseRawDataOnly || (GetSkeletonVirtualBoneGuid() != GetSkeleton()->GetVirtualBoneGuid())); }
+	bool DoesNeedRecompress() const;
 
 	/**
 	 * Create Animation Sequence from Reference Pose of the Mesh
diff --git a/Engine/Source/Runtime/Engine/Classes/Animation/AnimSequenceBase.h b/Engine/Source/Runtime/Engine/Classes/Animation/AnimSequenceBase.h
--- a/Engine/Source/Runtime/Engine/Classes/Animation/AnimSequenceBase.h
+++ b/Engine/Source/Runtime/Engine/Classes/Animation/AnimSequenceBase.h
@@ -257,15 +257,7 @@ public:
 
 	// return true if anim notify is available 
 	virtual bool IsNotifyAvailable() const;
-protected:
-	template <typename DataType>
-	void VerifyCurveNames(USkeleton& Skeleton, const FName& NameContainer, TArray<DataType>& CurveList)
-	{
-		for (DataType& Curve : CurveList)
-		{
-			Skeleton.VerifySmartName(NameContainer, Curve.Name);
-		}
-	}
+
 #if WITH_EDITOR
 	void OnEndLoadPackage(const FEndLoadPackageContext& Context);
 	virtual void OnAnimModelLoaded();
diff --git a/Engine/Source/Runtime/Engine/Classes/Animation/Skeleton.h b/Engine/Source/Runtime/Engine/Classes/Animation/Skeleton.h
--- a/Engine/Source/Runtime/Engine/Classes/Animation/Skeleton.h
+++ b/Engine/Source/Runtime/Engine/Classes/Animation/Skeleton.h
@@ -34,6 +34,7 @@ class FPackageReloadedEvent;
 struct FAssetData;
 enum class EPackageReloadPhase : uint8;
 class USkeleton;
+typedef SmartName::UID_Type SkeletonAnimCurveUID;
 
 // Delegate used to control global skeleton compatibility
 DECLARE_DELEGATE_RetVal(bool, FAreAllSkeletonsCompatible);
@@ -367,7 +368,7 @@ public:
 	TMap< FName, FReferencePose > AnimRetargetSources;
 
 	// Typedefs for greater smartname UID readability, add one for each smartname category 
-	typedef SmartName::UID_Type AnimCurveUID;
+	typedef SkeletonAnimCurveUID AnimCurveUID;
 
 	// Names for smartname mappings, if you're adding a new category of smartnames add a new name here
 	static ENGINE_API const FName AnimCurveMappingName;
diff --git a/Engine/Source/Runtime/Engine/Classes/Engine/Brush.h b/Engine/Source/Runtime/Engine/Classes/Engine/Brush.h
--- a/Engine/Source/Runtime/Engine/Classes/Engine/Brush.h
+++ b/Engine/Source/Runtime/Engine/Classes/Engine/Brush.h
@@ -4,11 +4,11 @@
 
 #include "CoreMinimal.h"
 #include "UObject/ObjectMacros.h"
-#include "Engine/Level.h"
 #include "GameFramework/Actor.h"
 #include "Brush.generated.h"
 
 class UBrushBuilder;
+class ULevel;
 
 //-----------------------------------------------------------------------------
 // Variables.
@@ -207,17 +207,7 @@ public:
 	 * @param	OutLevels if specified, provides a copy of the levels array
 	 * @return	true if the csg needs to be rebuilt on the next editor tick.	
 	 */
-	static bool NeedsRebuild(TArray< TWeakObjectPtr< ULevel > >* OutLevels = nullptr)
-	{
-		LevelsToRebuild.RemoveAllSwap([](const TWeakObjectPtr<ULevel>& Level) { return !Level.IsValid(); });
-
-		if (OutLevels)
-		{
-			*OutLevels = LevelsToRebuild;
-		}
-		
-		return(LevelsToRebuild.Num() > 0);
-	}
+	static bool NeedsRebuild(TArray< TWeakObjectPtr< ULevel > >* OutLevels = nullptr);
 
 	/**
 	 * Called upon finishing the csg rebuild to clear the rebuild bool.
@@ -232,7 +222,7 @@ public:
 	 *
 	 * @param	InLevel The level that needs rebuilding
 	 */
-	static void SetNeedRebuild(ULevel* InLevel){if(InLevel){LevelsToRebuild.AddUnique(InLevel);}}
+	static void SetNeedRebuild(ULevel* InLevel);
 #endif//WITH_EDITOR
 
 	/** @return true if this is a static brush */
diff --git a/Engine/Source/Runtime/Engine/Private/AI/Navigation/NavigationInterfaces.cpp b/Engine/Source/Runtime/Engine/Private/AI/Navigation/NavigationInterfaces.cpp
--- a/Engine/Source/Runtime/Engine/Private/AI/Navigation/NavigationInterfaces.cpp
+++ b/Engine/Source/Runtime/Engine/Private/AI/Navigation/NavigationInterfaces.cpp
@@ -21,6 +21,11 @@ UNavAgentInterface::UNavAgentInterface(const FObjectInitializer& ObjectInitializ
 {
 }
 
+const FNavAgentProperties& INavAgentInterface::GetNavAgentPropertiesRef() const
+{
+	return FNavAgentProperties::DefaultProperties;
+}
+
 UNavRelevantInterface::UNavRelevantInterface(const FObjectInitializer& ObjectInitializer)
 	: Super(ObjectInitializer)
 {
diff --git a/Engine/Source/Runtime/Engine/Private/Animation/AnimCurveTypes.cpp b/Engine/Source/Runtime/Engine/Private/Animation/AnimCurveTypes.cpp
--- a/Engine/Source/Runtime/Engine/Private/Animation/AnimCurveTypes.cpp
+++ b/Engine/Source/Runtime/Engine/Private/Animation/AnimCurveTypes.cpp
@@ -1,6 +1,7 @@
 // Copyright Epic Games, Inc. All Rights Reserved.
 
 #include "Animation/AnimCurveTypes.h"
+#include "Animation/Skeleton.h"
 #include "UObject/FrameworkObjectVersion.h"
 #include "UObject/AnimObjectVersion.h"
 #include "Math/RandomStream.h"
@@ -370,7 +371,7 @@ float FCachedFloatCurve::GetValueAtPosition(const UAnimSequenceBase* InAnimSeque
 	return InAnimSequence->EvaluateCurveData(GetAnimCurveUID(InAnimSequence), InPosition);
 }
 
-USkeleton::AnimCurveUID FCachedFloatCurve::GetAnimCurveUID(const UAnimSequenceBase* InAnimSequence) const
+SkeletonAnimCurveUID FCachedFloatCurve::GetAnimCurveUID(const UAnimSequenceBase* InAnimSequence) const
 {
 	if (CurveName != CachedCurveName && InAnimSequence)
 	{
@@ -392,7 +393,7 @@ const FFloatCurve* FCachedFloatCurve::GetFloatCurve(const UAnimSequenceBase* InA
 {
 	if (InAnimSequence)
 	{
-		USkeleton::AnimCurveUID DistanceCurveUID = GetAnimCurveUID(InAnimSequence);
+		SkeletonAnimCurveUID DistanceCurveUID = GetAnimCurveUID(InAnimSequence);
 		if (DistanceCurveUID != SmartName::MaxUID)
 		{
 			const FAnimationCurveIdentifier CurveId(DistanceCurveUID, ERawCurveTrackTypes::RCT_Float);
@@ -452,7 +453,7 @@ void FRawCurveTracks::EvaluateTransformCurveData(USkeleton * Skeleton, TMap<FNam
 	}
 }
 #endif
-FAnimCurveBase * FRawCurveTracks::GetCurveData(USkeleton::AnimCurveUID Uid, ERawCurveTrackTypes SupportedCurveType /*= FloatType*/)
+FAnimCurveBase * FRawCurveTracks::GetCurveData(SkeletonAnimCurveUID Uid, ERawCurveTrackTypes SupportedCurveType /*= FloatType*/)
 {
 	switch (SupportedCurveType)
 	{
@@ -468,7 +469,7 @@ FAnimCurveBase * FRawCurveTracks::GetCurveData(USkeleton::AnimCurveUID Uid, ERaw
 	}
 }
 
-const FAnimCurveBase * FRawCurveTracks::GetCurveData(USkeleton::AnimCurveUID Uid, ERawCurveTrackTypes SupportedCurveType /*= FloatType*/) const
+const FAnimCurveBase * FRawCurveTracks::GetCurveData(SkeletonAnimCurveUID Uid, ERawCurveTrackTypes SupportedCurveType /*= FloatType*/) const
 {
 	switch (SupportedCurveType)
 	{
@@ -650,7 +651,7 @@ bool FRawCurveTracks::DuplicateCurveData(const FSmartName& CurveToCopy, const FS
 // but this has to be refactored once we'd like to move onto serialize
 ///////////////////////////////////
 template <typename DataType>
-DataType * FRawCurveTracks::GetCurveDataImpl(TArray<DataType> & Curves, USkeleton::AnimCurveUID Uid)
+DataType * FRawCurveTracks::GetCurveDataImpl(TArray<DataType> & Curves, SkeletonAnimCurveUID Uid)
 {
 	for(DataType& Curve : Curves)
 	{
@@ -664,7 +665,7 @@ DataType * FRawCurveTracks::GetCurveDataImpl(TArray<DataType> & Curves, USkeleto
 }
 
 template <typename DataType>
-const DataType * FRawCurveTracks::GetCurveDataImpl(const TArray<DataType> & Curves, USkeleton::AnimCurveUID Uid) const
+const DataType * FRawCurveTracks::GetCurveDataImpl(const TArray<DataType> & Curves, SkeletonAnimCurveUID Uid) const
 {
 	for (const DataType& Curve : Curves)
 	{
diff --git a/Engine/Source/Runtime/Engine/Private/Animation/AnimMontage.cpp b/Engine/Source/Runtime/Engine/Private/Animation/AnimMontage.cpp
--- a/Engine/Source/Runtime/Engine/Private/Animation/AnimMontage.cpp
+++ b/Engine/Source/Runtime/Engine/Private/Animation/AnimMontage.cpp
@@ -1953,6 +1953,10 @@ bool FAnimMontageInstance::SimulateAdvance(float DeltaTime, float& InOutPosition
 	return true;
 }
 
+FSlotAnimationTrack::FSlotAnimationTrack()
+	: SlotName(FAnimSlotGroup::DefaultSlotName)
+{}
+
 void FMontageSubStepper::Initialize(const struct FAnimMontageInstance& InAnimInstance)
 {
 	MontageInstance = &InAnimInstance;
diff --git a/Engine/Source/Runtime/Engine/Private/Animation/AnimSequence.cpp b/Engine/Source/Runtime/Engine/Private/Animation/AnimSequence.cpp
--- a/Engine/Source/Runtime/Engine/Private/Animation/AnimSequence.cpp
+++ b/Engine/Source/Runtime/Engine/Private/Animation/AnimSequence.cpp
@@ -3867,6 +3867,11 @@ uint8* UAnimSequence::FindSyncMarkerPropertyData(int32 SyncMarkerIndex, FArrayPr
 	return nullptr;
 }
 
+bool UAnimSequence::DoesNeedRecompress() const
+{
+	return GetSkeleton() && (bUseRawDataOnly || (GetSkeletonVirtualBoneGuid() != GetSkeleton()->GetVirtualBoneGuid()));
+}
+
 bool UAnimSequence::CreateAnimation(USkeletalMesh* Mesh)
 {
 	// create animation from Mesh's ref pose
diff --git a/Engine/Source/Runtime/Engine/Private/Animation/AnimSequenceBase.cpp b/Engine/Source/Runtime/Engine/Private/Animation/AnimSequenceBase.cpp
--- a/Engine/Source/Runtime/Engine/Private/Animation/AnimSequenceBase.cpp
+++ b/Engine/Source/Runtime/Engine/Private/Animation/AnimSequenceBase.cpp
@@ -84,6 +84,15 @@ void UAnimSequenceBase::DeclareConstructClasses(TArray<FTopLevelAssetPath>& OutC
 }
 #endif
 
+template <typename DataType>
+void VerifyCurveNames(USkeleton& Skeleton, const FName& NameContainer, TArray<DataType>& CurveList)
+{
+	for (DataType& Curve : CurveList)
+	{
+		Skeleton.VerifySmartName(NameContainer, Curve.Name);
+	}
+}
+
 void UAnimSequenceBase::PostLoad()
 {
 #if WITH_EDITORONLY_DATA
diff --git a/Engine/Source/Runtime/Engine/Private/AudioCompressionSettingsUtils.cpp b/Engine/Source/Runtime/Engine/Private/AudioCompressionSettingsUtils.cpp
--- a/Engine/Source/Runtime/Engine/Private/AudioCompressionSettingsUtils.cpp
+++ b/Engine/Source/Runtime/Engine/Private/AudioCompressionSettingsUtils.cpp
@@ -2,6 +2,7 @@
 
 #include "AudioCompressionSettingsUtils.h"
 #include "AudioCompressionSettings.h"
+#include "AudioStreamingCache.h"
 #include "Misc/DataDrivenPlatformInfoRegistry.h"
 #include "HAL/PlatformMisc.h"
 
diff --git a/Engine/Source/Runtime/Engine/Private/Brush.cpp b/Engine/Source/Runtime/Engine/Private/Brush.cpp
--- a/Engine/Source/Runtime/Engine/Private/Brush.cpp
+++ b/Engine/Source/Runtime/Engine/Private/Brush.cpp
@@ -8,6 +8,7 @@
 #include "EngineGlobals.h"
 #include "Engine/Polys.h"
 #include "Engine/Engine.h"
+#include "Engine/Level.h"
 #include "Model.h"
 #include "Materials/Material.h"
 #include "Engine/BrushBuilder.h"
@@ -143,6 +144,26 @@ void ABrush::CopyPosRotScaleFrom( ABrush* Other )
 	ReregisterAllComponents();
 }
 
+bool ABrush::NeedsRebuild(TArray< TWeakObjectPtr< ULevel > >* OutLevels)
+{
+	LevelsToRebuild.RemoveAllSwap([](const TWeakObjectPtr<ULevel>& Level) { return !Level.IsValid(); });
+
+	if (OutLevels)
+	{
+		*OutLevels = LevelsToRebuild;
+	}
+
+	return(LevelsToRebuild.Num() > 0);
+}
+
+void ABrush::SetNeedRebuild(ULevel* InLevel)
+{
+	if (InLevel)
+	{
+		LevelsToRebuild.AddUnique(InLevel);
+	}
+}
+
 void ABrush::InitPosRotScale()
 {
 	check(BrushComponent);
diff --git a/Engine/Source/Runtime/Engine/Public/Animation/AnimCompressionTypes.h b/Engine/Source/Runtime/Engine/Public/Animation/AnimCompressionTypes.h
--- a/Engine/Source/Runtime/Engine/Public/Animation/AnimCompressionTypes.h
+++ b/Engine/Source/Runtime/Engine/Public/Animation/AnimCompressionTypes.h
@@ -6,7 +6,9 @@
 #include "UObject/ObjectMacros.h"
 #include "Misc/Guid.h"
 #include "Misc/MemStack.h"
+#if UE_ENABLE_INCLUDE_ORDER_DEPRECATED_IN_5_2
 #include "BonePose.h"
+#endif
 #include "Animation/AnimTypes.h"
 #include "Animation/AnimCurveTypes.h"
 #include "Animation/AnimationAsset.h"
@@ -45,6 +47,7 @@ class UAnimCurveCompressionSettings;
 class UAnimBoneCompressionSettings;
 class UAnimBoneCompressionCodec;
 class USkeleton;
+struct FCompactPose;
 
 template<typename ArrayClass>
 struct ENGINE_API FCompressedOffsetDataBase
diff --git a/Engine/Source/Runtime/Engine/Public/Animation/AnimCurveTypes.h b/Engine/Source/Runtime/Engine/Public/Animation/AnimCurveTypes.h
--- a/Engine/Source/Runtime/Engine/Public/Animation/AnimCurveTypes.h
+++ b/Engine/Source/Runtime/Engine/Public/Animation/AnimCurveTypes.h
@@ -6,11 +6,17 @@
 #include "UObject/ObjectMacros.h"
 #include "Animation/AnimTypes.h"
 #include "Animation/SmartName.h"
-#include "Animation/Skeleton.h"
 #include "Curves/RichCurve.h"
 #include "Misc/EnumRange.h"
+
+#if UE_ENABLE_INCLUDE_ORDER_DEPRECATED_IN_5_2
+#include "Animation/Skeleton.h"
+#endif
+
 #include "AnimCurveTypes.generated.h"
 
+typedef SmartName::UID_Type SkeletonAnimCurveUID;
+
 UENUM()
 enum class EAnimCurveType : uint8 
 {
@@ -275,7 +281,7 @@ public:
 	FName CurveName;
 
 private:
-	mutable USkeleton::AnimCurveUID CachedUID;
+	mutable SkeletonAnimCurveUID CachedUID;
 	mutable FName CachedCurveName;
 
 public:
@@ -288,7 +294,7 @@ public:
 	const FFloatCurve* GetFloatCurve(const UAnimSequenceBase* InAnimSequence) const;
 
 protected:
-	USkeleton::AnimCurveUID GetAnimCurveUID(const UAnimSequenceBase* InAnimSequence) const;
+	SkeletonAnimCurveUID GetAnimCurveUID(const UAnimSequenceBase* InAnimSequence) const;
 };
 
 /**
@@ -416,7 +422,7 @@ struct FBaseBlendedCurve
 	}
 
 	/** Invalidate value of InUID */
-	void InvalidateCurveWeight(USkeleton::AnimCurveUID InUid)
+	void InvalidateCurveWeight(SkeletonAnimCurveUID InUid)
 	{
 		check(bInitialized);
 
@@ -429,7 +435,7 @@ struct FBaseBlendedCurve
 	}
 
 	/** Set value of InUID to InValue */
-	void Set(USkeleton::AnimCurveUID InUid, float InValue)
+	void Set(SkeletonAnimCurveUID InUid, float InValue)
 	{
 		check(bInitialized);
 
@@ -442,7 +448,7 @@ struct FBaseBlendedCurve
 	}
 
 	/** Get Value of InUID - @todo : add validation check here and make sure caller also knows it's not valid*/
-	float Get(USkeleton::AnimCurveUID InUid) const
+	float Get(SkeletonAnimCurveUID InUid) const
 	{
 		check(bInitialized);
 
@@ -456,7 +462,7 @@ struct FBaseBlendedCurve
 	}
 
 	/** Get Value of InUID with validation and default value */
-	float Get(USkeleton::AnimCurveUID InUid, bool& OutIsValid, float InDefaultValue=0.f) const
+	float Get(SkeletonAnimCurveUID InUid, bool& OutIsValid, float InDefaultValue=0.f) const
 	{
 		check(bInitialized);
 
@@ -472,7 +478,7 @@ struct FBaseBlendedCurve
 	}
 
 	/** Get Array Index by UID */
-	int32 GetArrayIndexByUID(USkeleton::AnimCurveUID InUid) const
+	int32 GetArrayIndexByUID(SkeletonAnimCurveUID InUid) const
 	{
 		int32 ArrayIndex = (*UIDToArrayIndexLUT).IsValidIndex(InUid) ? (*UIDToArrayIndexLUT)[InUid] : MAX_uint16;
 		if (ArrayIndex != MAX_uint16)
@@ -483,7 +489,7 @@ struct FBaseBlendedCurve
 	}
 
 	/** return true if enabled. return false otherwise. */
-	bool IsEnabled(USkeleton::AnimCurveUID InUid) const
+	bool IsEnabled(SkeletonAnimCurveUID InUid) const
 	{
 		check(bInitialized);
 
@@ -867,12 +873,12 @@ struct FRawCurveTracks
 	/**
 	 * Find curve data based on the curve UID
 	 */
-	ENGINE_API FAnimCurveBase * GetCurveData(USkeleton::AnimCurveUID Uid, ERawCurveTrackTypes SupportedCurveType = ERawCurveTrackTypes::RCT_Float);
+	ENGINE_API FAnimCurveBase * GetCurveData(SkeletonAnimCurveUID Uid, ERawCurveTrackTypes SupportedCurveType = ERawCurveTrackTypes::RCT_Float);
 
 	/**
 	* Find curve data based on the curve UID
 	*/
-	ENGINE_API const FAnimCurveBase * GetCurveData(USkeleton::AnimCurveUID Uid, ERawCurveTrackTypes SupportedCurveType = ERawCurveTrackTypes::RCT_Float) const;
+	ENGINE_API const FAnimCurveBase * GetCurveData(SkeletonAnimCurveUID Uid, ERawCurveTrackTypes SupportedCurveType = ERawCurveTrackTypes::RCT_Float) const;
 
 	/**
 	 * Add new curve from the provided UID and return true if success
@@ -931,13 +937,13 @@ private:
 	 * Find curve data based on the curve UID
 	 */
 	template <typename DataType>
-	DataType * GetCurveDataImpl(TArray<DataType>& Curves, USkeleton::AnimCurveUID Uid);
+	DataType * GetCurveDataImpl(TArray<DataType>& Curves, SkeletonAnimCurveUID Uid);
 
 	/**
 	* Find curve data based on the curve UID
 	*/
 	template <typename DataType>
-	const DataType * GetCurveDataImpl(const TArray<DataType>& Curves, USkeleton::AnimCurveUID Uid) const;
+	const DataType * GetCurveDataImpl(const TArray<DataType>& Curves, SkeletonAnimCurveUID Uid) const;
 
 	/**
 	 * Add new curve from the provided UID and return true if success
diff --git a/Engine/Source/Runtime/Engine/Public/AudioCompressionSettingsUtils.h b/Engine/Source/Runtime/Engine/Public/AudioCompressionSettingsUtils.h
--- a/Engine/Source/Runtime/Engine/Public/AudioCompressionSettingsUtils.h
+++ b/Engine/Source/Runtime/Engine/Public/AudioCompressionSettingsUtils.h
@@ -5,7 +5,12 @@
 #include "IAudioExtensionPlugin.h"
 #include "AudioPluginUtilities.h"
 #include "AudioCompressionSettings.h"
+
+#if UE_ENABLE_INCLUDE_ORDER_DEPRECATED_IN_5_2
 #include "AudioStreamingCache.h"
+#endif
+
+struct FCachedAudioStreamingManagerParams;
 
 class ENGINE_API FPlatformCompressionUtilities
 {
diff --git a/Engine/Source/Runtime/Engine/Public/Components.h b/Engine/Source/Runtime/Engine/Public/Components.h
--- a/Engine/Source/Runtime/Engine/Public/Components.h
+++ b/Engine/Source/Runtime/Engine/Public/Components.h
@@ -5,13 +5,20 @@
 #include "CoreMinimal.h"
 #include "UObject/ObjectMacros.h"
 #include "MeshUVChannelInfo.h"
+#include "VertexStreamComponent.h"
+
+#if UE_ENABLE_INCLUDE_ORDER_DEPRECATED_IN_5_2
 #include "RenderResource.h"
 #include "VertexFactory.h"
+#endif
 
 /*=============================================================================
 	Components.h: Forward declarations of object components of actors
 =============================================================================*/
 
+class FRHIShaderResourceView;
+
+
 // Constants.
 enum { MAX_STATIC_TEXCOORDS = 8 };
 
diff --git a/Engine/Source/Runtime/Landscape/Classes/LandscapeProxy.h b/Engine/Source/Runtime/Landscape/Classes/LandscapeProxy.h
--- a/Engine/Source/Runtime/Landscape/Classes/LandscapeProxy.h
+++ b/Engine/Source/Runtime/Landscape/Classes/LandscapeProxy.h
@@ -11,7 +11,6 @@
 #include "Async/AsyncWork.h"
 #include "Engine/Texture.h"
 #include "PerPlatformProperties.h"
-#include "LandscapeComponent.h"
 #include "LandscapeNaniteComponent.h"
 #include "LandscapeWeightmapUsage.h"
 #include "LandscapeHeightfieldCollisionComponent.h"
@@ -19,6 +18,10 @@
 #include "ActorPartition/PartitionActor.h"
 #include "ILandscapeSplineInterface.h"
 
+#if UE_ENABLE_INCLUDE_ORDER_DEPRECATED_IN_5_2
+#include "LandscapeComponent.h"
+#endif
+
 #if WITH_EDITOR
 #include "WorldPartition/WorldPartitionHandle.h"
 #endif
@@ -42,8 +45,10 @@ class UPhysicalMaterial;
 class USplineComponent;
 class UTexture2D;
 class FLandscapeEditLayerReadback;
+class FMaterialUpdateContext;
 struct FAsyncGrassBuilder;
 struct FLandscapeInfoLayerSettings;
+struct FLandscapePerLODMaterialOverride;
 class FLandscapeProxyComponentDataChangedParams;
 struct FMeshDescription;
 enum class ENavDataGatheringMode : uint8;
@@ -845,7 +850,7 @@ public:
 	virtual const ALandscape* GetLandscapeActor() const PURE_VIRTUAL(GetLandscapeActor, return nullptr;)
 
 	const TArray<FLandscapePerLODMaterialOverride>& GetPerLODOverrideMaterials() const { return PerLODOverrideMaterials; }
-	void SetPerLODOverrideMaterials(const TArray<FLandscapePerLODMaterialOverride>& InValue) { PerLODOverrideMaterials = InValue; }
+	LANDSCAPE_API void SetPerLODOverrideMaterials(const TArray<FLandscapePerLODMaterialOverride>& InValue);
 
 	static void SetGrassUpdateInterval(int32 Interval) { GrassUpdateInterval = Interval; }
 
diff --git a/Engine/Source/Runtime/Landscape/Private/LandscapeProxy.cpp b/Engine/Source/Runtime/Landscape/Private/LandscapeProxy.cpp
--- a/Engine/Source/Runtime/Landscape/Private/LandscapeProxy.cpp
+++ b/Engine/Source/Runtime/Landscape/Private/LandscapeProxy.cpp
@@ -26,3 +26,7 @@ LANDSCAPE_API FLandscapeImportLayerInfo::FLandscapeImportLayerInfo(const FLandsc
 
 #endif // WITH_EDITOR
 
+void ALandscapeProxy::SetPerLODOverrideMaterials(const TArray<FLandscapePerLODMaterialOverride>& InValue)
+{
+	PerLODOverrideMaterials = InValue;
+}
diff --git a/Engine/Source/Runtime/RenderCore/Public/VertexFactory.h b/Engine/Source/Runtime/RenderCore/Public/VertexFactory.h
--- a/Engine/Source/Runtime/RenderCore/Public/VertexFactory.h
+++ b/Engine/Source/Runtime/RenderCore/Public/VertexFactory.h
@@ -30,6 +30,7 @@
 #include "ShaderCore.h"
 #include "Templates/UnrealTemplate.h"
 #include "UObject/NameTypes.h"
+#include "VertexStreamComponent.h"
 
 #include <atomic>
 
@@ -90,13 +91,6 @@ struct FVertexInputStream
  */
 typedef TArray<FVertexInputStream, TInlineAllocator<4>> FVertexInputStreamArray;
 
-enum class EVertexStreamUsage : uint8
-{
-	Default			= 0 << 0,
-	Instancing		= 1 << 0,
-	Overridden		= 1 << 1,
-	ManualFetch		= 1 << 2
-};
 ENUM_CLASS_FLAGS(EVertexStreamUsage);
 
 enum class EVertexInputStreamType : uint8
@@ -130,62 +124,6 @@ enum class EVertexFactoryFlags : uint32
 };
 ENUM_CLASS_FLAGS(EVertexFactoryFlags);
 
-/**
- * A typed data source for a vertex factory which streams data from a vertex buffer.
- */
-struct FVertexStreamComponent
-{
-	/** The vertex buffer to stream data from.  If null, no data can be read from this stream. */
-	const FVertexBuffer* VertexBuffer = nullptr;
-
-	/** The offset to the start of the vertex buffer fetch. */
-	uint32 StreamOffset = 0;
-
-	/** The offset of the data, relative to the beginning of each element in the vertex buffer. */
-	uint8 Offset = 0;
-
-	/** The stride of the data. */
-	uint8 Stride = 0;
-
-	/** The type of the data read from this stream. */
-	TEnumAsByte<EVertexElementType> Type = VET_None;
-
-	EVertexStreamUsage VertexStreamUsage = EVertexStreamUsage::Default;
-
-	/**
-	 * Initializes the data stream to null.
-	 */
-	FVertexStreamComponent()
-	{}
-
-	/**
-	 * Minimal initialization constructor.
-	 */
-	FVertexStreamComponent(const FVertexBuffer* InVertexBuffer, uint32 InOffset, uint32 InStride, EVertexElementType InType, EVertexStreamUsage Usage = EVertexStreamUsage::Default) :
-		VertexBuffer(InVertexBuffer),
-		StreamOffset(0),
-		Offset((uint8)InOffset),
-		Stride((uint8)InStride),
-		Type(InType),
-		VertexStreamUsage(Usage)
-	{
-		check(InStride <= 0xFF);
-		check(InOffset <= 0xFF);
-	}
-
-	FVertexStreamComponent(const FVertexBuffer* InVertexBuffer, uint32 InStreamOffset, uint32 InOffset, uint32 InStride, EVertexElementType InType, EVertexStreamUsage Usage = EVertexStreamUsage::Default) :
-		VertexBuffer(InVertexBuffer),
-		StreamOffset(InStreamOffset),
-		Offset((uint8)InOffset),
-		Stride((uint8)InStride),
-		Type(InType),
-		VertexStreamUsage(Usage)
-	{
-		check(InStride <= 0xFF);
-		check(InOffset <= 0xFF);
-	}
-};
-
 /**
  * A macro which initializes a FVertexStreamComponent to read a member from a struct.
  */
diff --git a/Engine/Source/Runtime/RenderCore/Public/VertexStreamComponent.h b/Engine/Source/Runtime/RenderCore/Public/VertexStreamComponent.h
new file mode 100644
--- /dev/null
+++ b/Engine/Source/Runtime/RenderCore/Public/VertexStreamComponent.h
@@ -0,0 +1,71 @@
+// Copyright Epic Games, Inc. All Rights Reserved.
+
+#pragma once
+
+#include "RHIDefinitions.h"
+
+class FVertexBuffer;
+
+enum class EVertexStreamUsage : uint8
+{
+	Default = 0 << 0,
+	Instancing = 1 << 0,
+	Overridden = 1 << 1,
+	ManualFetch = 1 << 2
+};
+
+/**
+ * A typed data source for a vertex factory which streams data from a vertex buffer.
+ */
+struct FVertexStreamComponent
+{
+	/** The vertex buffer to stream data from.  If null, no data can be read from this stream. */
+	const FVertexBuffer* VertexBuffer = nullptr;
+
+	/** The offset to the start of the vertex buffer fetch. */
+	uint32 StreamOffset = 0;
+
+	/** The offset of the data, relative to the beginning of each element in the vertex buffer. */
+	uint8 Offset = 0;
+
+	/** The stride of the data. */
+	uint8 Stride = 0;
+
+	/** The type of the data read from this stream. */
+	TEnumAsByte<EVertexElementType> Type = VET_None;
+
+	EVertexStreamUsage VertexStreamUsage = EVertexStreamUsage::Default;
+
+	/**
+	 * Initializes the data stream to null.
+	 */
+	FVertexStreamComponent()
+	{}
+
+	/**
+	 * Minimal initialization constructor.
+	 */
+	FVertexStreamComponent(const FVertexBuffer* InVertexBuffer, uint32 InOffset, uint32 InStride, EVertexElementType InType, EVertexStreamUsage Usage = EVertexStreamUsage::Default) :
+		VertexBuffer(InVertexBuffer),
+		StreamOffset(0),
+		Offset((uint8)InOffset),
+		Stride((uint8)InStride),
+		Type(InType),
+		VertexStreamUsage(Usage)
+	{
+		check(InStride <= 0xFF);
+		check(InOffset <= 0xFF);
+	}
+
+	FVertexStreamComponent(const FVertexBuffer* InVertexBuffer, uint32 InStreamOffset, uint32 InOffset, uint32 InStride, EVertexElementType InType, EVertexStreamUsage Usage = EVertexStreamUsage::Default) :
+		VertexBuffer(InVertexBuffer),
+		StreamOffset(InStreamOffset),
+		Offset((uint8)InOffset),
+		Stride((uint8)InStride),
+		Type(InType),
+		VertexStreamUsage(Usage)
+	{
+		check(InStride <= 0xFF);
+		check(InOffset <= 0xFF);
+	}
+};
diff --git a/Engine/Source/Runtime/UMG/Public/Blueprint/WidgetBlueprintLibrary.h b/Engine/Source/Runtime/UMG/Public/Blueprint/WidgetBlueprintLibrary.h
--- a/Engine/Source/Runtime/UMG/Public/Blueprint/WidgetBlueprintLibrary.h
+++ b/Engine/Source/Runtime/UMG/Public/Blueprint/WidgetBlueprintLibrary.h
@@ -12,9 +12,11 @@
 #include "Styling/SlateBrush.h"
 #include "Components/SlateWrapperTypes.h"
 #include "Blueprint/DragDropOperation.h"
+#if UE_ENABLE_INCLUDE_ORDER_DEPRECATED_IN_5_2
 #include "Blueprint/UserWidget.h"
-#include "Kismet/BlueprintFunctionLibrary.h"
 #include "Slate/SGameLayerManager.h"
+#endif
+#include "Kismet/BlueprintFunctionLibrary.h"
 #include "WidgetBlueprintLibrary.generated.h"
 
 class UFont;
@@ -22,6 +24,7 @@ class UMaterialInstanceDynamic;
 class UMaterialInterface;
 class USlateBrushAsset;
 class UTexture2D;
+class UUserWidget;
 
 UCLASS(meta=(ScriptName="WidgetLibrary"))
 class UMG_API UWidgetBlueprintLibrary : public UBlueprintFunctionLibrary
