Added forward declaration in StaticMeshVertexBuffer to fix compile error

diff --git a/Engine/Source/Runtime/Engine/Public/Rendering/StaticMeshVertexBuffer.h b/Engine/Source/Runtime/Engine/Public/Rendering/StaticMeshVertexBuffer.h
--- a/Engine/Source/Runtime/Engine/Public/Rendering/StaticMeshVertexBuffer.h
+++ b/Engine/Source/Runtime/Engine/Public/Rendering/StaticMeshVertexBuffer.h
@@ -9,6 +9,7 @@
 #include "Components.h"
 #include "Math/Vector2DHalf.h"
 
+class FVertexFactory;
 class FStaticMeshVertexDataInterface;
 
 template<typename TangentTypeT>
