Added forward declarations in StaticMeshVertexbuffer

diff --git a/Engine/Source/Runtime/Engine/Public/Rendering/StaticMeshVertexBuffer.h b/Engine/Source/Runtime/Engine/Public/Rendering/StaticMeshVertexBuffer.h
--- a/Engine/Source/Runtime/Engine/Public/Rendering/StaticMeshVertexBuffer.h
+++ b/Engine/Source/Runtime/Engine/Public/Rendering/StaticMeshVertexBuffer.h
@@ -11,6 +11,8 @@
 
 class FVertexFactory;
 class FStaticMeshVertexDataInterface;
+class FVertexFactory;
+struct FStaticMeshDataType;
 
 template<typename TangentTypeT>
 struct TStaticMeshVertexTangentDatum
