- Fix Import/ExportTextItem
- Add testing for Import/ExportTextItem, and test case that covers OSSEOSPlus usage.

diff --git a/Engine/Source/Runtime/Engine/Private/OnlineReplStructs.cpp b/Engine/Source/Runtime/Engine/Private/OnlineReplStructs.cpp
--- a/Engine/Source/Runtime/Engine/Private/OnlineReplStructs.cpp
+++ b/Engine/Source/Runtime/Engine/Private/OnlineReplStructs.cpp
@@ -13,6 +13,7 @@
 #include "EngineLogs.h"
 #include "Net/OnlineEngineInterface.h"
 #include "Misc/ConfigCacheIni.h"
+#include "Misc/OutputDeviceNull.h"
 
 #include UE_INLINE_GENERATED_CPP_BY_NAME(OnlineReplStructs)
 
@@ -474,11 +475,11 @@ bool FUniqueNetIdRepl::ExportTextItem(FString& ValueStr, FUniqueNetIdRepl const&
 			FName Type = GetType();
 			if (Type == UOnlineEngineInterface::Get()->GetDefaultOnlineSubsystemName())
 			{
-				ValueStr += ToString();
+				ValueStr += FString::Printf(TEXT("\"%s\""), *ToString());
 			}
 			else
 			{
-				ValueStr += FString::Printf(TEXT("%s:%s"), *Type.ToString(), *ToString());
+				ValueStr += FString::Printf(TEXT("%s:\"%s\""), *Type.ToString(), *ToString());
 			}
 		}
 		else
@@ -651,15 +652,18 @@ void TestUniqueIdRepl(UWorld* InWorld)
 	FUniqueNetIdRepl CustomOSSPlusPrefixIn(UOnlineEngineInterface::Get()->CreateUniquePlayerIdWrapper(TEXT("+123456"), NAME_CustomOSS));
 	FUniqueNetIdRepl CustomOSSOddIntegerStringIn(UOnlineEngineInterface::Get()->CreateUniquePlayerIdWrapper(TEXT("123456789"), NAME_CustomOSS));
 	FUniqueNetIdRepl CustomOSSEvenIntegerStringIn(UOnlineEngineInterface::Get()->CreateUniquePlayerIdWrapper(TEXT("1234567890"), NAME_CustomOSS));
+	FUniqueNetIdRepl CustomOSSSeparatorsRequiringQuotedStringIn(UOnlineEngineInterface::Get()->CreateUniquePlayerIdWrapper(TEXT("1234_+_567|890"), NAME_CustomOSS));
 
 	CHECK_REPL_VALIDITY(CustomOSSIdIn, bSetupSuccess);
 	CHECK_REPL_VALIDITY(CustomOSSEncodedIdIn, bSetupSuccess);
 	CHECK_REPL_VALIDITY(CustomOSSPlusPrefixIn, bSetupSuccess);
 	CHECK_REPL_VALIDITY(CustomOSSOddIntegerStringIn, bSetupSuccess);
 	CHECK_REPL_VALIDITY(CustomOSSEvenIntegerStringIn, bSetupSuccess);
+	CHECK_REPL_VALIDITY(CustomOSSSeparatorsRequiringQuotedStringIn, bSetupSuccess);
 
 	bool bRegularSerializationSuccess = true;
 	bool bNetworkSerializationSuccess = true;
+	bool bTextItemSerializationSuccess = true;
 	if (bSetupSuccess)
 	{
 		// Regular Serialization (persistent/disk based using FString)
@@ -682,6 +686,7 @@ void TestUniqueIdRepl(UWorld* InWorld)
 				TestUniqueIdWriter << CustomOSSPlusPrefixIn;
 				TestUniqueIdWriter << CustomOSSOddIntegerStringIn;
 				TestUniqueIdWriter << CustomOSSEvenIntegerStringIn;
+				TestUniqueIdWriter << CustomOSSSeparatorsRequiringQuotedStringIn;
 			}
 
 			FUniqueNetIdRepl EmptyIdOut;
@@ -695,6 +700,7 @@ void TestUniqueIdRepl(UWorld* InWorld)
 			FUniqueNetIdRepl CustomOSSPlusPrefixOut;
 			FUniqueNetIdRepl CustomOSSOddIntegerStringOut;
 			FUniqueNetIdRepl CustomOSSEvenIntegerStringOut;
+			FUniqueNetIdRepl CustomOSSSeparatorsRequiringQuotedStringOut;
 
 			// Serialize Out
 			{
@@ -710,6 +716,7 @@ void TestUniqueIdRepl(UWorld* InWorld)
 				TestUniqueIdReader << CustomOSSPlusPrefixOut;
 				TestUniqueIdReader << CustomOSSOddIntegerStringOut;
 				TestUniqueIdReader << CustomOSSEvenIntegerStringOut;
+				TestUniqueIdReader << CustomOSSSeparatorsRequiringQuotedStringOut;
 			}
 
 			if (EmptyIdOut.IsValid())
@@ -734,6 +741,7 @@ void TestUniqueIdRepl(UWorld* InWorld)
 			CHECK_REPL_EQUALITY(CustomOSSPlusPrefixIn, CustomOSSPlusPrefixOut, bRegularSerializationSuccess);
 			CHECK_REPL_EQUALITY(CustomOSSOddIntegerStringIn, CustomOSSOddIntegerStringOut, bRegularSerializationSuccess);
 			CHECK_REPL_EQUALITY(CustomOSSEvenIntegerStringIn, CustomOSSEvenIntegerStringOut, bRegularSerializationSuccess);
+			CHECK_REPL_EQUALITY(CustomOSSSeparatorsRequiringQuotedStringIn, CustomOSSSeparatorsRequiringQuotedStringOut, bRegularSerializationSuccess);
 		}
 
 		// Network serialization (network/transient using MakeReplicationData)
@@ -766,6 +774,8 @@ void TestUniqueIdRepl(UWorld* InWorld)
 				EncodingFailures += bOutSuccess ? 0 : 1;
 				CustomOSSEvenIntegerStringIn.NetSerialize(TestUniqueIdWriter, nullptr, bOutSuccess);
 				EncodingFailures += bOutSuccess ? 0 : 1;
+				CustomOSSSeparatorsRequiringQuotedStringIn.NetSerialize(TestUniqueIdWriter, nullptr, bOutSuccess);
+				EncodingFailures += bOutSuccess ? 0 : 1;
 			}
 
 			if (EncodingFailures > 0)
@@ -787,6 +797,7 @@ void TestUniqueIdRepl(UWorld* InWorld)
 				FUniqueNetIdRepl CustomOSSPlusPrefixOut;
 				FUniqueNetIdRepl CustomOSSOddIntegerStringOut;
 				FUniqueNetIdRepl CustomOSSEvenIntegerStringOut;
+				FUniqueNetIdRepl CustomOSSSeparatorsRequiringQuotedStringOut;
 
 				// Serialize Out
 				uint8 DecodingFailures = 0;
@@ -815,6 +826,8 @@ void TestUniqueIdRepl(UWorld* InWorld)
 					DecodingFailures += bOutSuccess ? 0 : 1;
 					CustomOSSEvenIntegerStringOut.NetSerialize(TestUniqueIdReader, nullptr, bOutSuccess);
 					DecodingFailures += bOutSuccess ? 0 : 1;
+					CustomOSSSeparatorsRequiringQuotedStringOut.NetSerialize(TestUniqueIdReader, nullptr, bOutSuccess);
+					DecodingFailures += bOutSuccess ? 0 : 1;
 				}
 
 				if (DecodingFailures > 0)
@@ -832,7 +845,7 @@ void TestUniqueIdRepl(UWorld* InWorld)
 				if (EmptyIdIn != EmptyIdOut)
 				{
 					UE_LOG(LogNet, Warning, TEXT("EmptyId In/Out mismatch"));
-					bRegularSerializationSuccess = false;
+					bNetworkSerializationSuccess = false;
 				}
 
 				CHECK_REPL_EQUALITY(ValidIdIn, ValidIdOut, bNetworkSerializationSuccess);
@@ -845,6 +858,161 @@ void TestUniqueIdRepl(UWorld* InWorld)
 				CHECK_REPL_EQUALITY(CustomOSSPlusPrefixIn, CustomOSSPlusPrefixOut, bRegularSerializationSuccess);
 				CHECK_REPL_EQUALITY(CustomOSSOddIntegerStringIn, CustomOSSOddIntegerStringOut, bRegularSerializationSuccess);
 				CHECK_REPL_EQUALITY(CustomOSSEvenIntegerStringIn, CustomOSSEvenIntegerStringOut, bRegularSerializationSuccess);
+				CHECK_REPL_EQUALITY(CustomOSSSeparatorsRequiringQuotedStringIn, CustomOSSSeparatorsRequiringQuotedStringOut, bRegularSerializationSuccess);
+			}
+		}
+
+		// TextItem serialization
+		{
+			bool bOutSuccess = false;
+			const FUniqueNetIdRepl DefaultValue;
+
+			// Serialize In
+			FString EmptyIdTextItem;
+			FString ValidIdTextItem;
+			FString OddStringIdTextItem;
+			FString NonHexStringIdTextItem;
+			FString UpperCaseStringIdTextItem;
+			FString WayTooLongForHexEncodingIdTextItem;
+			FString CustomOSSIdTextItem;
+			FString CustomOSSEncodedIdTextItem;
+			FString CustomOSSPlusPrefixTextItem;
+			FString CustomOSSOddIntegerStringTextItem;
+			FString CustomOSSEvenIntegerStringTextItem;
+			FString CustomOSSSeparatorsRequiringQuotedStringTextItem;
+			uint8 ExportFailures = 0;
+			{
+				bOutSuccess = EmptyIdIn.ExportTextItem(EmptyIdTextItem, DefaultValue, nullptr, 0, nullptr);
+				ExportFailures += bOutSuccess ? 0 : 1;
+				bOutSuccess = ValidIdIn.ExportTextItem(ValidIdTextItem, DefaultValue, nullptr, 0, nullptr);
+				ExportFailures += bOutSuccess ? 0 : 1;
+				bOutSuccess = OddStringIdIn.ExportTextItem(OddStringIdTextItem, DefaultValue, nullptr, 0, nullptr);
+				ExportFailures += bOutSuccess ? 0 : 1;
+				bOutSuccess = NonHexStringIdIn.ExportTextItem(NonHexStringIdTextItem, DefaultValue, nullptr, 0, nullptr);
+				ExportFailures += bOutSuccess ? 0 : 1;
+				bOutSuccess = UpperCaseStringIdIn.ExportTextItem(UpperCaseStringIdTextItem, DefaultValue, nullptr, 0, nullptr);
+				ExportFailures += bOutSuccess ? 0 : 1;
+				bOutSuccess = WayTooLongForHexEncodingIdIn.ExportTextItem(WayTooLongForHexEncodingIdTextItem, DefaultValue, nullptr, 0, nullptr);
+				ExportFailures += bOutSuccess ? 0 : 1;
+				bOutSuccess = CustomOSSIdIn.ExportTextItem(CustomOSSIdTextItem, DefaultValue, nullptr, 0, nullptr);
+				ExportFailures += bOutSuccess ? 0 : 1;
+				bOutSuccess = CustomOSSEncodedIdIn.ExportTextItem(CustomOSSEncodedIdTextItem, DefaultValue, nullptr, 0, nullptr);
+				ExportFailures += bOutSuccess ? 0 : 1;
+				bOutSuccess = CustomOSSPlusPrefixIn.ExportTextItem(CustomOSSPlusPrefixTextItem, DefaultValue, nullptr, 0, nullptr);
+				ExportFailures += bOutSuccess ? 0 : 1;
+				bOutSuccess = CustomOSSOddIntegerStringIn.ExportTextItem(CustomOSSOddIntegerStringTextItem, DefaultValue, nullptr, 0, nullptr);
+				ExportFailures += bOutSuccess ? 0 : 1;
+				bOutSuccess = CustomOSSEvenIntegerStringIn.ExportTextItem(CustomOSSEvenIntegerStringTextItem, DefaultValue, nullptr, 0, nullptr);
+				ExportFailures += bOutSuccess ? 0 : 1;
+				bOutSuccess = CustomOSSSeparatorsRequiringQuotedStringIn.ExportTextItem(CustomOSSSeparatorsRequiringQuotedStringTextItem, DefaultValue, nullptr, 0, nullptr);
+				ExportFailures += bOutSuccess ? 0 : 1;
+			}
+
+			if (ExportFailures > 0)
+			{
+				UE_LOG(LogNet, Warning, TEXT("There were %d export failures"), ExportFailures);
+				bTextItemSerializationSuccess = false;
+			}
+
+			if (bTextItemSerializationSuccess)
+			{
+				FOutputDeviceNull ErrorText;
+
+				FUniqueNetIdRepl EmptyIdOut;
+				FUniqueNetIdRepl ValidIdOut;
+				FUniqueNetIdRepl OddStringIdOut;
+				FUniqueNetIdRepl NonHexStringIdOut;
+				FUniqueNetIdRepl UpperCaseStringIdOut;
+				FUniqueNetIdRepl WayTooLongForHexEncodingIdOut;
+				FUniqueNetIdRepl CustomOSSIdOut;
+				FUniqueNetIdRepl CustomOSSEncodedIdOut;
+				FUniqueNetIdRepl CustomOSSPlusPrefixOut;
+				FUniqueNetIdRepl CustomOSSOddIntegerStringOut;
+				FUniqueNetIdRepl CustomOSSEvenIntegerStringOut;
+				FUniqueNetIdRepl CustomOSSSeparatorsRequiringQuotedStringOut;
+
+				// Serialize Out
+				uint8 ImportFailures = 0;
+				const TCHAR* Buffer = nullptr;
+				{
+					Buffer = *EmptyIdTextItem;
+					bOutSuccess = EmptyIdOut.ImportTextItem(Buffer, 0, nullptr, &ErrorText);
+					ImportFailures += bOutSuccess ? 0 : 1;
+
+					Buffer = *ValidIdTextItem;
+					bOutSuccess = ValidIdOut.ImportTextItem(Buffer, 0, nullptr, &ErrorText);
+					ImportFailures += bOutSuccess ? 0 : 1;
+
+					Buffer = *OddStringIdTextItem;
+					bOutSuccess = OddStringIdOut.ImportTextItem(Buffer, 0, nullptr, &ErrorText);
+					ImportFailures += bOutSuccess ? 0 : 1;
+
+					Buffer = *NonHexStringIdTextItem;
+					bOutSuccess = NonHexStringIdOut.ImportTextItem(Buffer, 0, nullptr, &ErrorText);
+					ImportFailures += bOutSuccess ? 0 : 1;
+
+					Buffer = *UpperCaseStringIdTextItem;
+					bOutSuccess = UpperCaseStringIdOut.ImportTextItem(Buffer, 0, nullptr, &ErrorText);
+					ImportFailures += bOutSuccess ? 0 : 1;
+
+					Buffer = *WayTooLongForHexEncodingIdTextItem;
+					bOutSuccess = WayTooLongForHexEncodingIdOut.ImportTextItem(Buffer, 0, nullptr, &ErrorText);
+					ImportFailures += bOutSuccess ? 0 : 1;
+
+					Buffer = *CustomOSSIdTextItem;
+					bOutSuccess = CustomOSSIdOut.ImportTextItem(Buffer, 0, nullptr, &ErrorText);
+					ImportFailures += bOutSuccess ? 0 : 1;
+
+					Buffer = *CustomOSSEncodedIdTextItem;
+					bOutSuccess = CustomOSSEncodedIdOut.ImportTextItem(Buffer, 0, nullptr, &ErrorText);
+					ImportFailures += bOutSuccess ? 0 : 1;
+
+					Buffer = *CustomOSSPlusPrefixTextItem;
+					bOutSuccess = CustomOSSPlusPrefixOut.ImportTextItem(Buffer, 0, nullptr, &ErrorText);
+					ImportFailures += bOutSuccess ? 0 : 1;
+
+					Buffer = *CustomOSSOddIntegerStringTextItem;
+					bOutSuccess = CustomOSSOddIntegerStringOut.ImportTextItem(Buffer, 0, nullptr, &ErrorText);
+					ImportFailures += bOutSuccess ? 0 : 1;
+
+					Buffer = *CustomOSSEvenIntegerStringTextItem;
+					bOutSuccess = CustomOSSEvenIntegerStringOut.ImportTextItem(Buffer, 0, nullptr, &ErrorText);
+					ImportFailures += bOutSuccess ? 0 : 1;
+
+					Buffer = *CustomOSSSeparatorsRequiringQuotedStringTextItem;
+					bOutSuccess = CustomOSSSeparatorsRequiringQuotedStringOut.ImportTextItem(Buffer, 0, nullptr, &ErrorText);
+					ImportFailures += bOutSuccess ? 0 : 1;
+				}
+
+				if (ImportFailures > 0)
+				{
+					UE_LOG(LogNet, Warning, TEXT("There were %d import failures"), ImportFailures);
+					bTextItemSerializationSuccess = false;
+				}
+
+				if (EmptyIdOut.IsValid())
+				{
+					UE_LOG(LogNet, Warning, TEXT("EmptyId %s should have been invalid"), *EmptyIdOut->ToDebugString());
+					bTextItemSerializationSuccess = false;
+				}
+
+				if (EmptyIdIn != EmptyIdOut)
+				{
+					UE_LOG(LogNet, Warning, TEXT("EmptyId In/Out mismatch"));
+					bTextItemSerializationSuccess = false;
+				}
+
+				CHECK_REPL_EQUALITY(ValidIdIn, ValidIdOut, bTextItemSerializationSuccess);
+				CHECK_REPL_EQUALITY(OddStringIdIn, OddStringIdOut, bTextItemSerializationSuccess);
+				CHECK_REPL_EQUALITY(NonHexStringIdIn, NonHexStringIdOut, bTextItemSerializationSuccess);
+				CHECK_REPL_EQUALITY(UpperCaseStringIdIn, UpperCaseStringIdOut, bTextItemSerializationSuccess);
+				CHECK_REPL_EQUALITY(WayTooLongForHexEncodingIdIn, WayTooLongForHexEncodingIdOut, bTextItemSerializationSuccess);
+				CHECK_REPL_EQUALITY(CustomOSSIdIn, CustomOSSIdOut, bTextItemSerializationSuccess);
+				CHECK_REPL_EQUALITY(CustomOSSEncodedIdIn, CustomOSSEncodedIdOut, bTextItemSerializationSuccess);
+				CHECK_REPL_EQUALITY(CustomOSSPlusPrefixIn, CustomOSSPlusPrefixOut, bTextItemSerializationSuccess);
+				CHECK_REPL_EQUALITY(CustomOSSOddIntegerStringIn, CustomOSSOddIntegerStringOut, bTextItemSerializationSuccess);
+				CHECK_REPL_EQUALITY(CustomOSSEvenIntegerStringIn, CustomOSSEvenIntegerStringOut, bTextItemSerializationSuccess);
+				CHECK_REPL_EQUALITY(CustomOSSSeparatorsRequiringQuotedStringIn, CustomOSSSeparatorsRequiringQuotedStringOut, bTextItemSerializationSuccess);
 			}
 		}
 	}
@@ -926,6 +1094,7 @@ void TestUniqueIdRepl(UWorld* InWorld)
 	UE_LOG(LogNet, Log, TEXT("	Setup: %s"), bSetupSuccess ? TEXT("PASS") : TEXT("FAIL"));
 	UE_LOG(LogNet, Log, TEXT("	Normal: %s"), bRegularSerializationSuccess ? (bSetupSuccess ? TEXT("PASS") : TEXT("SKIPPED")) : TEXT("FAIL"));
 	UE_LOG(LogNet, Log, TEXT("	Network: %s"), bNetworkSerializationSuccess ? (bSetupSuccess ? TEXT("PASS") : TEXT("SKIPPED")) : TEXT("FAIL"));
+	UE_LOG(LogNet, Log, TEXT("	TextItem: %s"), bTextItemSerializationSuccess ? (bSetupSuccess ? TEXT("PASS") : TEXT("SKIPPED")) : TEXT("FAIL"));
 	UE_LOG(LogNet, Log, TEXT("	Platform: %s"), bPlatformSerializationSuccess ? (bSetupSuccess ? TEXT("PASS") : TEXT("SKIPPED")) : TEXT("FAIL"));
 	UE_LOG(LogNet, Log, TEXT("	JSON: %s"), bJSONSerializationSuccess ? (bSetupSuccess ? TEXT("PASS") : TEXT("SKIPPED")) : TEXT("FAIL"));
 
