Path Tracer: Expose a blueprint command to toggle the path tracer show flag on or off

This gives end users a way to enable the path tracer without requiring the console (which is not enabled in shipping builds out of the box). Note that only the game viewport is targetted by this method.

diff --git a/Engine/Source/Runtime/Engine/Classes/Kismet/KismetRenderingLibrary.h b/Engine/Source/Runtime/Engine/Classes/Kismet/KismetRenderingLibrary.h
--- a/Engine/Source/Runtime/Engine/Classes/Kismet/KismetRenderingLibrary.h
+++ b/Engine/Source/Runtime/Engine/Classes/Kismet/KismetRenderingLibrary.h
@@ -205,4 +205,10 @@ class UKismetRenderingLibrary : public UBlueprintFunctionLibrary
 	/** Calculates the projection matrix using this view info's aspect ratio (regardless of bConstrainAspectRatio) */
 	UFUNCTION(BlueprintPure, Category="Rendering", meta = (DisplayName = "Calculate Projection Matrix (Minimal View Info)", ScriptMethod = "CalculateProjectionMatrix"))
 	static ENGINE_API FMatrix CalculateProjectionMatrix(const FMinimalViewInfo& MinimalViewInfo);
+
+	/** Enables or disables the path tracer for all viewports simulatenously.
+	 * This command is equivalent to setting ShowFlag.PathTracing, but is accessible even from shipping builds.
+	 */
+	UFUNCTION(BlueprintCallable, Category = "Rendering", meta = (Keywords="EnablePathTracing"))
+	static ENGINE_API void EnablePathTracing(bool bEnablePathTracer);
 };
diff --git a/Engine/Source/Runtime/Engine/Private/KismetRenderingLibrary.cpp b/Engine/Source/Runtime/Engine/Private/KismetRenderingLibrary.cpp
--- a/Engine/Source/Runtime/Engine/Private/KismetRenderingLibrary.cpp
+++ b/Engine/Source/Runtime/Engine/Private/KismetRenderingLibrary.cpp
@@ -14,6 +14,7 @@
 #include "TextureResource.h"
 #include "SceneUtils.h"
 #include "Logging/MessageLog.h"
+#include "Engine/GameViewportClient.h"
 #include "Engine/TextureRenderTarget2D.h"
 #include "Engine/TextureRenderTarget2DArray.h"
 #include "Engine/TextureRenderTargetVolume.h"
@@ -885,5 +886,17 @@ ENGINE_API FMatrix UKismetRenderingLibrary::CalculateProjectionMatrix(const FMin
 	return MinimalViewInfo.CalculateProjectionMatrix();
 }
 
+ENGINE_API void UKismetRenderingLibrary::EnablePathTracing(bool bEnablePathTracer)
+{
+	if (GEngine != nullptr && GEngine->GameViewport != nullptr)
+	{
+		FEngineShowFlags* EngineShowFlags = GEngine->GameViewport->GetEngineShowFlags();
+		if (EngineShowFlags != nullptr)
+		{
+			EngineShowFlags->SetPathTracing(bEnablePathTracer);
+		}
+	}
+}
+
 #undef LOCTEXT_NAMESPACE
 
