Make sure the FSkeletalMeshImportData::ApplyRigToGeo increase the bounding box to include both the target and the source mesh so we found match for skinning transfert for all target vertex.

#rnx

diff --git a/Engine/Source/Runtime/Engine/Private/Rendering/SkeletalMeshLODImporterData.cpp b/Engine/Source/Runtime/Engine/Private/Rendering/SkeletalMeshLODImporterData.cpp
--- a/Engine/Source/Runtime/Engine/Private/Rendering/SkeletalMeshLODImporterData.cpp
+++ b/Engine/Source/Runtime/Engine/Private/Rendering/SkeletalMeshLODImporterData.cpp
@@ -222,7 +222,7 @@ bool FSkeletalMeshImportData::ApplyRigToGeo(FSkeletalMeshImportData& Other)
 	Influences.Reset();
 
 	FWedgePosition OldGeoOverlappingPosition;
-	FWedgePosition::FillWedgePosition(OldGeoOverlappingPosition, Other.Points, Other.Wedges, UE_THRESH_POINTS_ARE_SAME);
+	FWedgePosition::FillWedgePosition(OldGeoOverlappingPosition, Other.Points, Other.Wedges, Points, UE_THRESH_POINTS_ARE_SAME);
 	FOctreeQueryHelper OctreeQueryHelper(OldGeoOverlappingPosition.GetOctree());
 
 	//////////////////////////////////////////////////////////////////////////
@@ -1144,6 +1144,7 @@ void FWedgePosition::FillWedgePosition(
 	FWedgePosition& OutOverlappingPosition,
 	const TArray<FVector3f>& Points,
 	const TArray<SkeletalMeshImportData::FVertex> Wedges,
+	const TArray<FVector3f>& TargetPositions,
 	float ComparisonThreshold)
 {
 	OutOverlappingPosition.Points= Points;
@@ -1161,6 +1162,8 @@ void FWedgePosition::FillWedgePosition(
 
 
 	FBox3f OldBounds(OutOverlappingPosition.Points);
+	//Make sure the bounds include the target bounds so we find a match for every target vertex
+	OldBounds += FBox3f(TargetPositions);
 	OutOverlappingPosition.WedgePosOctree = new TWedgeInfoPosOctree((FVector)OldBounds.GetCenter(), OldBounds.GetExtent().GetMax());
 
 	// Add each old vertex to the octree
diff --git a/Engine/Source/Runtime/Engine/Public/Rendering/SkeletalMeshLODImporterData.h b/Engine/Source/Runtime/Engine/Public/Rendering/SkeletalMeshLODImporterData.h
--- a/Engine/Source/Runtime/Engine/Public/Rendering/SkeletalMeshLODImporterData.h
+++ b/Engine/Source/Runtime/Engine/Public/Rendering/SkeletalMeshLODImporterData.h
@@ -743,10 +743,12 @@ struct FWedgePosition
 	// Fill the data:
 	// Create the SortedPosition use to find exact match (position)
 	// Create the wedge position octree to find the closest position, we use this when there is no exact match
+	// The targetPositions is use to to max out the octree bounding box to both the source and target geometry
 	static void FillWedgePosition(
 		FWedgePosition& OutOverlappingPosition,
 		const TArray<FVector3f>& Positions,
 		const TArray<SkeletalMeshImportData::FVertex> Wedges,
+		const TArray<FVector3f>& TargetPositions,
 		float ComparisonThreshold);
 
 private:
