Prevent blending values from being stomped when a layer is unlinked twice in the same frame - When resetting to self layer, don't blend unless the target instance has changed.

#tests PIE

diff --git a/Engine/Source/Runtime/Engine/Private/Animation/AnimNode_LinkedAnimLayer.cpp b/Engine/Source/Runtime/Engine/Private/Animation/AnimNode_LinkedAnimLayer.cpp
--- a/Engine/Source/Runtime/Engine/Private/Animation/AnimNode_LinkedAnimLayer.cpp
+++ b/Engine/Source/Runtime/Engine/Private/Animation/AnimNode_LinkedAnimLayer.cpp
@@ -111,7 +111,11 @@ void FAnimNode_LinkedAnimLayer::InitializeSelfLayer(const UAnimInstance* SelfAni
 
 	IAnimClassInterface* NewAnimBPClass = IAnimClassInterface::GetFromClass(SelfClass);
 
-	RequestBlend(PriorAnimBPClass, NewAnimBPClass);
+	// No need for blending if the instance hasn't changed (this was causing issues when a layer was unlinked more than once in a single frame due to faulty blueprints, causing the good blend values to be stomped before being processed)
+	if (CurrentTarget != GetTargetInstance<UAnimInstance>())
+	{
+		RequestBlend(PriorAnimBPClass, NewAnimBPClass);
+	}
 }
 
 void FAnimNode_LinkedAnimLayer::SetLinkedLayerInstance(const UAnimInstance* InOwningAnimInstance, UAnimInstance* InNewLinkedInstance)
