World Partition: fix handles Reset not properly decrementing intrernal ref counts + updated handles unit tests.

#rnx

diff --git a/Engine/Source/Runtime/Engine/Private/Tests/WorldPartition/WorldPartitionHandleTests.cpp b/Engine/Source/Runtime/Engine/Private/Tests/WorldPartition/WorldPartitionHandleTests.cpp
--- a/Engine/Source/Runtime/Engine/Private/Tests/WorldPartition/WorldPartitionHandleTests.cpp
+++ b/Engine/Source/Runtime/Engine/Private/Tests/WorldPartition/WorldPartitionHandleTests.cpp
@@ -213,24 +213,37 @@ namespace WorldPartitionTests
 
 				Reference = MoveTemp(HandleFromReference);
 				TestTrue(TEXT("Handle move src not valid"), !HandleFromReference.IsValid());
-				TestTrue(TEXT("Handle move dst valid"), Reference.IsValid());
-				TestTrue(TEXT("Handle soft refcount"), Reference->GetSoftRefCount() == 0);
-				TestTrue(TEXT("Handle hard refcount"), Reference->GetHardRefCount() == 1);
+				TestTrue(TEXT("Reference move dst valid"), Reference.IsValid());
+				TestTrue(TEXT("Reference soft refcount"), Reference->GetSoftRefCount() == 0);
+				TestTrue(TEXT("Reference hard refcount"), Reference->GetHardRefCount() == 1);
 			}
 
 			// Reference handle move
 			{
 				FWorldPartitionReference ReferenceFromHandle(MoveTemp(Handle));
-				TestTrue(TEXT("Reference move src not valid"), !Handle.IsValid());
+				TestTrue(TEXT("Handle move src not valid"), !Handle.IsValid());
 				TestTrue(TEXT("Reference move dst valid"), ReferenceFromHandle.IsValid());
 				TestTrue(TEXT("Reference soft refcount"), ReferenceFromHandle->GetSoftRefCount() == 0);
 				TestTrue(TEXT("Reference hard refcount"), ReferenceFromHandle->GetHardRefCount() == 1);
 
 				Handle = MoveTemp(ReferenceFromHandle);
 				TestTrue(TEXT("Reference move src not valid"), !ReferenceFromHandle.IsValid());
-				TestTrue(TEXT("Reference move dst valid"), Handle.IsValid());
-				TestTrue(TEXT("Reference soft refcount"), Handle->GetSoftRefCount() == 1);
-				TestTrue(TEXT("Reference hard refcount"), Handle->GetHardRefCount() == 0);
+				TestTrue(TEXT("Handle move dst valid"), Handle.IsValid());
+				TestTrue(TEXT("Handle soft refcount"), Handle->GetSoftRefCount() == 1);
+				TestTrue(TEXT("Handle hard refcount"), Handle->GetHardRefCount() == 0);
+			}
+
+			// Reset
+			{
+				FWorldPartitionReference ReferenceFromHandle(Handle.ToReference());
+				TestTrue(TEXT("Reference valid"), ReferenceFromHandle.IsValid());
+				TestTrue(TEXT("Reference soft refcount"), ReferenceFromHandle->GetSoftRefCount() == 1);
+				TestTrue(TEXT("Reference hard refcount"), ReferenceFromHandle->GetHardRefCount() == 1);
+
+				ReferenceFromHandle.Reset();
+				TestTrue(TEXT("Reference not valid"), !ReferenceFromHandle.IsValid());
+				TestTrue(TEXT("Handle soft refcount"), Handle->GetSoftRefCount() == 1);
+				TestTrue(TEXT("Handle hard refcount"), Handle->GetHardRefCount() == 0);
 			}
 		}
 
@@ -295,4 +308,4 @@ namespace WorldPartitionTests
 
 #undef TEST_NAME_ROOT
 
-#endif 
+#endif 
\ No newline at end of file
diff --git a/Engine/Source/Runtime/Engine/Public/WorldPartition/WorldPartitionHandle.h b/Engine/Source/Runtime/Engine/Public/WorldPartition/WorldPartitionHandle.h
--- a/Engine/Source/Runtime/Engine/Public/WorldPartition/WorldPartitionHandle.h
+++ b/Engine/Source/Runtime/Engine/Public/WorldPartition/WorldPartitionHandle.h
@@ -271,6 +271,11 @@ public:
 
 	FORCEINLINE void Reset()
 	{
+		if (IsValid())
+		{
+			DecRefCount();
+		}
+
 		Container = nullptr;
 		ActorDesc = nullptr;
 	}
