Level Instance: fix crash when committing an level instance after its actor was unloaded.

#rnx

diff --git a/Engine/Source/Runtime/Engine/Private/LevelInstance/LevelInstanceSubsystem.cpp b/Engine/Source/Runtime/Engine/Private/LevelInstance/LevelInstanceSubsystem.cpp
--- a/Engine/Source/Runtime/Engine/Private/LevelInstance/LevelInstanceSubsystem.cpp
+++ b/Engine/Source/Runtime/Engine/Private/LevelInstance/LevelInstanceSubsystem.cpp
@@ -1941,6 +1941,9 @@ bool ULevelInstanceSubsystem::EditLevelInstanceInternal(ILevelInstanceInterface*
 		// Only support one level of recursion to commit current edit
 		check(!bRecursive);
 		FLevelInstanceID PendingEditId = LevelInstance->GetLevelInstanceID();
+
+		// Make sure to keep the top level instance actor loaded when we commit the current one
+		FWorldPartitionReference CurrentEditLevelInstanceActorRef = CurrentEditLevelInstanceActor;
 		
 		check(!IsLevelInstanceEditDirty(LevelInstanceEdit.Get()));
 		CommitLevelInstanceInternal(LevelInstanceEdit);
@@ -2006,6 +2009,21 @@ bool ULevelInstanceSubsystem::EditLevelInstanceInternal(ILevelInstanceInterface*
 
 	ResetLoadersForWorldAsset(LevelInstance->GetWorldAsset().GetLongPackageName());
 
+	if (UWorldPartition* WorldPartition = LevelInstanceActor->GetWorld()->GetWorldPartition(); WorldPartition && WorldPartition->IsMainWorldPartition())
+	{
+		AActor* TopLevelInstanceActor = LevelInstanceActor;
+		while (AActor* CurrentTopLevelInstanceActor = Cast<AActor>(GetParentLevelInstance(TopLevelInstanceActor)))
+		{
+			TopLevelInstanceActor = CurrentTopLevelInstanceActor;
+		}
+
+		if (FWorldPartitionActorDesc* TopLevelInstanceActorActorDesc = WorldPartition->GetActorDesc(TopLevelInstanceActor->GetActorGuid()))
+		{
+			check(!CurrentEditLevelInstanceActor.IsValid());
+			CurrentEditLevelInstanceActor = FWorldPartitionReference(TopLevelInstanceActorActorDesc->GetContainer(), TopLevelInstanceActorActorDesc->GetGuid());
+		}
+	}
+
 	return true;
 }
 
@@ -2183,6 +2201,8 @@ bool ULevelInstanceSubsystem::CommitLevelInstanceInternal(TUniquePtr<FLevelInsta
 	{
 		OnCommitChild(AncestorID, bChangesCommitted);
 	}
+
+	CurrentEditLevelInstanceActor.Reset();
 		
 	if (ILevelInstanceInterface* LevelInstanceToSelect = GetLevelInstance(LevelInstanceToSelectID))
 	{
diff --git a/Engine/Source/Runtime/Engine/Public/LevelInstance/LevelInstanceSubsystem.h b/Engine/Source/Runtime/Engine/Public/LevelInstance/LevelInstanceSubsystem.h
--- a/Engine/Source/Runtime/Engine/Public/LevelInstance/LevelInstanceSubsystem.h
+++ b/Engine/Source/Runtime/Engine/Public/LevelInstance/LevelInstanceSubsystem.h
@@ -8,6 +8,7 @@
 #include "UObject/GCObject.h"
 #include "Folder.h"
 #include "LevelInstance/LevelInstanceTypes.h"
+#include "WorldPartition/WorldPartitionHandle.h"
 
 #include "LevelInstanceSubsystem.generated.h"
 
@@ -257,5 +258,7 @@ private:
 	TMap<FLevelInstanceID, int32> ChildEdits;
 
 	FActorDescContainerInstanceManager ActorDescContainerInstanceManager;
+
+	FWorldPartitionReference CurrentEditLevelInstanceActor;
 #endif
 };
\ No newline at end of file
