fix non-unity build

#rnx

diff --git a/Engine/Source/Runtime/Engine/Public/WorldPartition/WorldPartitionSubsystem.h b/Engine/Source/Runtime/Engine/Public/WorldPartition/WorldPartitionSubsystem.h
--- a/Engine/Source/Runtime/Engine/Public/WorldPartition/WorldPartitionSubsystem.h
+++ b/Engine/Source/Runtime/Engine/Public/WorldPartition/WorldPartitionSubsystem.h
@@ -7,6 +7,7 @@
 #include "WorldPartitionSubsystem.generated.h"
 
 class UWorldPartition;
+class FWorldPartitionActorDesc;
 
 enum class EWorldPartitionRuntimeCellState : uint8;
 
