FTexture : set debug shader bGreyScaleFormat ReplicateR property from CompressionSettings, not PixelFormat, so it matches what MaterialExpressions wil do
bGreyScaleFormat is only used by debug/editor shaders, game runtime uses SamplerType from MaterialExpressions
bGreyScaleFormat bool is needed because FTexture does not have access to CompressionSettings
behavior is observably different with Grayscale G16 and cases like TC_Alpha non-multiple-of-4

diff --git a/Engine/Source/Runtime/Core/Public/PixelFormat.h b/Engine/Source/Runtime/Core/Public/PixelFormat.h
--- a/Engine/Source/Runtime/Core/Public/PixelFormat.h
+++ b/Engine/Source/Runtime/Core/Public/PixelFormat.h
@@ -17,8 +17,8 @@ enum EPixelFormat : uint8
 	PF_Unknown              =0,
 	PF_A32B32G32R32F        =1,
 	PF_B8G8R8A8             =2,
-	PF_G8                   =3,
-	PF_G16                  =4,
+	PF_G8                   =3, // G8  means Gray/Grey , not Green , typically actually uses a red format with replication of R to RGB
+	PF_G16                  =4, // G16 means Gray/Grey like G8
 	PF_DXT1                 =5,
 	PF_DXT3                 =6,
 	PF_DXT5                 =7,
diff --git a/Engine/Source/Runtime/Engine/Private/Rendering/StreamableTextureResource.cpp b/Engine/Source/Runtime/Engine/Private/Rendering/StreamableTextureResource.cpp
--- a/Engine/Source/Runtime/Engine/Private/Rendering/StreamableTextureResource.cpp
+++ b/Engine/Source/Runtime/Engine/Private/Rendering/StreamableTextureResource.cpp
@@ -97,7 +97,7 @@ FStreamableTextureResource::FStreamableTextureResource(UTexture* InOwner, const
 	// HDR images are stored in linear but still require gamma correction to display correctly.
 	bIgnoreGammaConversions = !InOwner->SRGB && !IsHDR(PixelFormat);
 	bSRGB = InOwner->SRGB;
-	bGreyScaleFormat = (PixelFormat == PF_G8) || (PixelFormat == PF_BC4);
+	bGreyScaleFormat = ( InOwner->CompressionSettings == TC_Grayscale || InOwner->CompressionSettings == TC_Alpha );
 
 	Filter = (ESamplerFilter)UDeviceProfileManager::Get().GetActiveProfile()->GetTextureLODSettings()->GetSamplerFilter(InOwner);
 	AddressU = InOwner->GetTextureAddressX() == TA_Wrap ? AM_Wrap : (InOwner->GetTextureAddressX() == TA_Clamp ? AM_Clamp : AM_Mirror);
diff --git a/Engine/Source/Runtime/Engine/Private/Texture2D.cpp b/Engine/Source/Runtime/Engine/Private/Texture2D.cpp
--- a/Engine/Source/Runtime/Engine/Private/Texture2D.cpp
+++ b/Engine/Source/Runtime/Engine/Private/Texture2D.cpp
@@ -1272,8 +1272,7 @@ FVirtualTexture2DResource::FVirtualTexture2DResource(const UTexture2D* InOwner,
 	FirstMipToUse = FMath::Min((int32)MaxMip, InFirstMipToUse);
 
 	bSRGB = InOwner->SRGB;
-	const EPixelFormat PixelFormat = VTData->LayerTypes[0];
-	bGreyScaleFormat = (PixelFormat == PF_G8) || (PixelFormat == PF_BC4);
+	bGreyScaleFormat = ( InOwner->CompressionSettings == TC_Grayscale || InOwner->CompressionSettings == TC_Alpha );
 
 	// Initialize this resource FeatureLevel, so it gets re-created on FeatureLevel changes
 	SetFeatureLevel(GMaxRHIFeatureLevel);
diff --git a/Engine/Source/Runtime/Engine/Private/TextureCube.cpp b/Engine/Source/Runtime/Engine/Private/TextureCube.cpp
--- a/Engine/Source/Runtime/Engine/Private/TextureCube.cpp
+++ b/Engine/Source/Runtime/Engine/Private/TextureCube.cpp
@@ -489,8 +489,7 @@ public:
 		SamplerStateRHI = GetOrCreateSamplerState(SamplerStateInitializer);
 
 		// Set the greyscale format flag appropriately.
-		EPixelFormat PixelFormat = Owner->GetPixelFormat();
-		bGreyScaleFormat = (PixelFormat == PF_G8) || (PixelFormat == PF_BC4);
+		bGreyScaleFormat = ( Owner->CompressionSettings == TC_Grayscale || Owner->CompressionSettings == TC_Alpha );
 	}
 
 	virtual void ReleaseRHI() override
diff --git a/Engine/Source/Runtime/Engine/Private/TextureCubeArray.cpp b/Engine/Source/Runtime/Engine/Private/TextureCubeArray.cpp
--- a/Engine/Source/Runtime/Engine/Private/TextureCubeArray.cpp
+++ b/Engine/Source/Runtime/Engine/Private/TextureCubeArray.cpp
@@ -169,8 +169,7 @@ public:
 		SamplerStateRHI = GetOrCreateSamplerState(SamplerStateInitializer);
 
 		// Set the greyscale format flag appropriately.
-		EPixelFormat PixelFormat = Owner->GetPixelFormat();
-		bGreyScaleFormat = (PixelFormat == PF_G8) || (PixelFormat == PF_BC4);
+		bGreyScaleFormat = ( Owner->CompressionSettings == TC_Grayscale || Owner->CompressionSettings == TC_Alpha );
 	}
 
 	virtual void ReleaseRHI() override
diff --git a/Engine/Source/Runtime/RenderCore/Public/RenderResource.h b/Engine/Source/Runtime/RenderCore/Public/RenderResource.h
--- a/Engine/Source/Runtime/RenderCore/Public/RenderResource.h
+++ b/Engine/Source/Runtime/RenderCore/Public/RenderResource.h
@@ -420,7 +420,10 @@ public:
 	/** Base values for fading in/out mip-levels. */
 	FMipBiasFade		MipBiasFade;
 
-	/** true if the texture is in a greyscale texture format. */
+	/** bGreyScaleFormat indicates the texture is actually in R channel but should be read as Grey (replicate R to RGBA)
+	 *  this is set from CompressionSettings, not PixelFormat
+	 *  this is only used by Editor/Debug shaders, not real game materials, which use SamplerType from MaterialExpressions
+	 */
 	bool				bGreyScaleFormat;
 
 	/**
