Use proper rounding when determining viewport size in order to avoid scaling artifacts.

diff --git a/Engine/Source/Runtime/Engine/Private/Slate/SceneViewport.cpp b/Engine/Source/Runtime/Engine/Private/Slate/SceneViewport.cpp
--- a/Engine/Source/Runtime/Engine/Private/Slate/SceneViewport.cpp
+++ b/Engine/Source/Runtime/Engine/Private/Slate/SceneViewport.cpp
@@ -389,7 +389,12 @@ void FSceneViewport::OnDrawViewport( const FGeometry& AllottedGeometry, const FS
 	/** Check to see if the viewport should be resized */
 	if (!bForceViewportSize)
 	{
-		FIntPoint DrawSize = FIntPoint( FMath::RoundToInt( AllottedGeometry.GetDrawSize().X ), FMath::RoundToInt( AllottedGeometry.GetDrawSize().Y ) );
+		// When a viewport element is created, its coordinates are rounded after applying the render transform.
+		// Considering that both absolute position and the size of the viewport can be non-integers,
+		// the correct draw size of the viewport can only be determined by rounding its absolute coordinates.
+		FVector2D TopLeft = AllottedGeometry.GetAbsolutePosition();
+		FVector2D BottomRight = TopLeft + AllottedGeometry.GetDrawSize();
+		FIntPoint DrawSize = FIntPoint(FMath::RoundToInt(BottomRight.X) - FMath::RoundToInt(TopLeft.X), FMath::RoundToInt(BottomRight.Y) - FMath::RoundToInt(TopLeft.Y));
 		bool bIsHDREnabled = IsHDREnabled();
 
 		SWindow* PaintWindow = OutDrawElements.GetPaintWindow();
