Strata - fix for starta node to be index by the material path instead of the expression memory address. This enables:
- acyclic graph instead of tress (a node can be linked to multiple pin in the graph).
- multiple material functions used in a graph using the same physical slab expression instance behind the scene.

diff --git a/Engine/Source/Runtime/Engine/Private/Materials/HLSLMaterialTranslator.cpp b/Engine/Source/Runtime/Engine/Private/Materials/HLSLMaterialTranslator.cpp
--- a/Engine/Source/Runtime/Engine/Private/Materials/HLSLMaterialTranslator.cpp
+++ b/Engine/Source/Runtime/Engine/Private/Materials/HLSLMaterialTranslator.cpp
@@ -5,6 +5,7 @@
 #include "HLSLMaterialTranslator.h"
 #include "VT/VirtualTextureScalability.h"
 #include "Materials/MaterialExpressionCustom.h"
+#include "Materials/MaterialExpressionStrata.h"
 #include "MaterialExpressionSettings.h"
 #include "Engine/RendererSettings.h"
 #include <functional>
@@ -74,6 +75,8 @@ static inline bool IsDebugTextureSampleEnabled()
 
 bool Engine_IsStrataEnabled();
 
+#define DEBUG_STRATA_TREE_STACK 0
+
 static uint32 GetStrataBytePerPixel(EShaderPlatform InPlatform)
 {
 	// We enforce at least 20 bytes per pixel because this is the minimal Strata GBuffer footprint of the simplest material.
@@ -342,6 +345,9 @@ FHLSLMaterialTranslator::FHLSLMaterialTranslator(FMaterial* InMaterial,
 	bStrataMaterialIsUnlitNode = false;
 	bStrataUsesConversionFromLegacy = false;
 	bStrataOutputsOpaqueRoughRefractions = false;
+
+	// Default value used as the root of the tree for the first path (when a node parent==nullptr).
+	StrataNodeIdentifierStack.Push(FGuid(0x7AEE, 0xBAD, 0xDEAD, 0xBEEF));
 }
 
 FHLSLMaterialTranslator::~FHLSLMaterialTranslator()
@@ -730,6 +736,9 @@ bool FHLSLMaterialTranslator::Translate()
 			TArray<FShaderCodeChunk> TempChunks;
 			AssignTempScope(TempChunks);
 
+		#if DEBUG_STRATA_TREE_STACK
+			UE_LOG(LogMaterial, Display, TEXT(" StrataTreeStack: StrataGenerateMaterialTopologyTree"));
+		#endif
 			FrontMaterialExpr->StrataGenerateMaterialTopologyTree(this, nullptr, 0);
 			if (!StrataGenerateDerivedMaterialOperatorData())
 			{
@@ -855,8 +864,12 @@ bool FHLSLMaterialTranslator::Translate()
 
 		// Make sure to compile this property before using ShadingModelsFromCompilation
 		Chunk[MP_ShadingModel]					= Material->CompilePropertyAndSetMaterialProperty(MP_ShadingModel			,this);
-			
-			Chunk[MP_FrontMaterial]					= Material->CompilePropertyAndSetMaterialProperty(MP_FrontMaterial			,this);
+		
+		
+	#if DEBUG_STRATA_TREE_STACK
+		UE_LOG(LogMaterial, Display, TEXT(" StrataTreeStack: Material->CompilePropertyAndSetMaterialProperty(MP_FrontMaterial)"));
+	#endif
+		Chunk[MP_FrontMaterial]					= Material->CompilePropertyAndSetMaterialProperty(MP_FrontMaterial			,this);
 		}
 
 		// Get shading models from compilation (or material).
@@ -3711,13 +3724,19 @@ int32 FHLSLMaterialTranslator::CallExpression(FMaterialExpressionKey ExpressionK
 		ExpressionKey.OutputIndex = INDEX_NONE;
 	}
 
+	// Strata BSDF expression should not be de-duplicated using expression output hash. 
+	// This is automatically handled via the compiler StrataTreeStack.
+	// It means that a node can be blended at multiple point of the graph (allowing acyclic graph instead of tree, e.g. a Slab can be used into multiple input).
+	// It is worth noting that only strata BSDF can be duplicated today according to StrataTreeStack.
+	const bool bExpressionIsStrataBSDF = ExpressionKey.Expression->IsA<UMaterialExpressionStrataSlabBSDF>();
+
 	// Check if this expression has already been translated.
 	check(ShaderFrequency < SF_NumFrequencies);
 	auto& CurrentFunctionStack = FunctionStacks[ShaderFrequency];
 	FMaterialFunctionCompileState* CurrentFunctionState = CurrentFunctionStack.Last();
 
 	static bool sDebugCacheDuplicateCode = true;
-	int32* ExistingCodeIndex = sDebugCacheDuplicateCode ? CurrentFunctionState->ExpressionCodeMap.Find(ExpressionKey) : nullptr;
+	int32* ExistingCodeIndex = sDebugCacheDuplicateCode && !bExpressionIsStrataBSDF ? CurrentFunctionState->ExpressionCodeMap.Find(ExpressionKey) : nullptr;
 	int32 Result = INDEX_NONE;
 	if (ExistingCodeIndex)
 	{
@@ -9747,7 +9766,7 @@ int32 FHLSLMaterialTranslator::GetLocal(const FName& LocalName)
 	return AddInlinedCodeChunk(MCT_Float1, TEXT("%s"), *Entry->Name);
 }
 
-FStrataOperator& FHLSLMaterialTranslator::StrataCompilationRegisterOperator(int32 OperatorType, UMaterialExpression* Expression, UMaterialExpression* Parent, bool bUseParameterBlending)
+FStrataOperator& FHLSLMaterialTranslator::StrataCompilationRegisterOperator(int32 OperatorType, FGuid StrataExpressionGuid, UMaterialExpression* Parent, FGuid StrataParentExpressionGuid, bool bUseParameterBlending)
 {
 	if (OperatorType == STRATA_OPERATOR_BSDF_LEGACY)
 	{
@@ -9758,10 +9777,10 @@ FStrataOperator& FHLSLMaterialTranslator::StrataCompilationRegisterOperator(int3
 
 	static FStrataOperator DefaultOperatorOnError = FStrataOperator();
 
-	if (StrataMaterialExpressionToOperatorIndex.Find(Expression))
+	if (StrataMaterialExpressionToOperatorIndex.Find(StrataExpressionGuid))
 	{
-		// It is not possible to register/use a Strata BSDF multiple times when creating a topology.
-		Errorf(TEXT("Material %s: It is not possible to uses a Strata BSDF (or any ouput of type StrataData) multiple times within a Strata material topology (asset: %s).\r\n"), *Material->GetDebugName(), *Material->GetAssetPath().ToString());
+		// It is not possible to register/use a Strata BSDF multiple times with this same exact graph path. (that would break the strata tree code generation)
+		Errorf(TEXT("Material %s: It is not possible to uses a Strata BSDF (or any ouput of type StrataData) multiple times within a Strata material topology with the same graph path GUID (asset: %s).\r\n"), *Material->GetDebugName(), *Material->GetAssetPath().ToString());
 		return DefaultOperatorOnError;
 	}
 
@@ -9772,14 +9791,14 @@ FStrataOperator& FHLSLMaterialTranslator::StrataCompilationRegisterOperator(int3
 		return DefaultOperatorOnError;
 	}
 
-	int32* ParentOperatorIndex = StrataMaterialExpressionToOperatorIndex.Find(Parent);
+	int32* ParentOperatorIndex = StrataMaterialExpressionToOperatorIndex.Find(StrataParentExpressionGuid);
 	if (Parent!=nullptr && ParentOperatorIndex == nullptr)
 	{
 		Errorf(TEXT("Material %s tries to register unknown operator parents (asset: %s).\r\n"), *Material->GetDebugName(), *Material->GetAssetPath().ToString());
 		return DefaultOperatorOnError;
 	}
 
-	StrataMaterialExpressionToOperatorIndex.Add(Expression, NewOperatorIndex);
+	StrataMaterialExpressionToOperatorIndex.Add(StrataExpressionGuid, NewOperatorIndex);
 
 	FStrataOperator& NewOperator = StrataMaterialExpressionRegisteredOperators.AddDefaulted_GetRef();
 
@@ -9799,9 +9818,9 @@ FStrataOperator& FHLSLMaterialTranslator::StrataCompilationRegisterOperator(int3
 	return NewOperator;
 }
 
-FStrataOperator& FHLSLMaterialTranslator::StrataCompilationGetOperator(UMaterialExpression* Expression)
+FStrataOperator& FHLSLMaterialTranslator::StrataCompilationGetOperator(FGuid StrataExpressionGuid)
 {
-	auto* OperatorIndex = StrataMaterialExpressionToOperatorIndex.Find(Expression);
+	auto* OperatorIndex = StrataMaterialExpressionToOperatorIndex.Find(StrataExpressionGuid);
 	if (!(OperatorIndex && *OperatorIndex >= 0 && *OperatorIndex < STRATA_MAX_OPERATOR_COUNT))
 	{
 		static FStrataOperator DefaultOperatorOnError = FStrataOperator();
@@ -11251,6 +11270,75 @@ bool FHLSLMaterialTranslator::StrataSkipsOpacityEvaluation()
 		&& !Material->GetShadingModels().HasShadingModel(MSM_PreintegratedSkin);
 }
 
+FGuid FHLSLMaterialTranslator::StrataTreeStackPush(UMaterialExpression* Expression, uint32 InputIndex)
+{
+	// Create an md5 hash for the parent, its input pin index and current node to represent the path.
+	uint32 IntputHashBuffer[9];
+	FGuid PreviousNodeGuid = StrataTreeStackGetPathUniqueId();
+	FGuid NodeGuid = Expression->MaterialExpressionGuid;
+	IntputHashBuffer[0] = PreviousNodeGuid.A;
+	IntputHashBuffer[1] = PreviousNodeGuid.B;
+	IntputHashBuffer[2] = PreviousNodeGuid.C;
+	IntputHashBuffer[3] = PreviousNodeGuid.D;
+	IntputHashBuffer[4] = InputIndex;
+	IntputHashBuffer[5] = NodeGuid.A;
+	IntputHashBuffer[6] = NodeGuid.B;
+	IntputHashBuffer[7] = NodeGuid.C;
+	IntputHashBuffer[8] = NodeGuid.D;
+
+	uint32 OutputHashBuffer[]{ 0, 0, 0, 0 };
+	FMD5 IdentifierStringHash;
+	IdentifierStringHash.Update((uint8*)IntputHashBuffer, sizeof(IntputHashBuffer));
+	IdentifierStringHash.Final((uint8*)&OutputHashBuffer);
+
+	StrataNodeIdentifierStack.Push(FGuid(OutputHashBuffer[0], OutputHashBuffer[1], OutputHashBuffer[2], OutputHashBuffer[3]));
+
+#if DEBUG_STRATA_TREE_STACK
+	UE_LOG(LogMaterial, Display, TEXT(" StrataTreeStack: Push (input %i of %s)"), InputIndex , *Expression->GetName());
+	TStringBuilder<2048> GuidStack;
+	for (auto& Entry : StrataNodeIdentifierStack)
+	{
+		GuidStack.Append(*Entry.ToString());
+		GuidStack.Append(TEXT("  "));
+	}
+	UE_LOG(LogMaterial, Display, TEXT(" StrataTreeStack: %s."), *GuidStack);
+#endif
+
+	return StrataNodeIdentifierStack.Top();
+}
+
+FGuid FHLSLMaterialTranslator::StrataTreeStackGetPathUniqueId()
+{
+	return StrataNodeIdentifierStack.Top();
+}
+
+FGuid FHLSLMaterialTranslator::StrataTreeStackGetParentPathUniqueId()
+{
+	if (StrataNodeIdentifierStack.Num() < 2)
+	{
+		// return some default when strata tree stack unique guid cannot be found
+		FGuid NullParent;
+		return NullParent;
+	}
+	return StrataNodeIdentifierStack.Last(1);
+}
+
+void FHLSLMaterialTranslator::StrataTreeStackPop()
+{
+	check(StrataNodeIdentifierStack.Num() >= 2);// 2 because there must always be the root remaining.
+	StrataNodeIdentifierStack.Pop();
+
+#if DEBUG_STRATA_TREE_STACK
+	TStringBuilder<2048> GuidStack;
+	for (auto& Entry : StrataNodeIdentifierStack)
+	{
+		GuidStack.Append(*Entry.ToString());
+		GuidStack.Append(TEXT("  "));
+	}
+	UE_LOG(LogMaterial, Display, TEXT(" StrataTreeStack: Pop %s."), *GuidStack);
+#endif
+}
+
 int32 FHLSLMaterialTranslator::MapARPassthroughCameraUV(int32 UV)
 {
 	if (UV == INDEX_NONE)
diff --git a/Engine/Source/Runtime/Engine/Private/Materials/HLSLMaterialTranslator.h b/Engine/Source/Runtime/Engine/Private/Materials/HLSLMaterialTranslator.h
--- a/Engine/Source/Runtime/Engine/Private/Materials/HLSLMaterialTranslator.h
+++ b/Engine/Source/Runtime/Engine/Private/Materials/HLSLMaterialTranslator.h
@@ -431,7 +431,7 @@ protected:
 	 * A normal code chunk hash can point to multiple shared info in case it is paired with different tangents. */
 	TMultiMap<uint64, FStrataSharedLocalBasesInfo> CodeChunkToStrataSharedLocalBasis;
 
-	TMap<UMaterialExpression*, int32> StrataMaterialExpressionToOperatorIndex;
+	TMap<FGuid, int32> StrataMaterialExpressionToOperatorIndex;
 	TArray<FStrataOperator> StrataMaterialExpressionRegisteredOperators;
 	FStrataOperator* StrataMaterialRootOperator;
 	uint32 StrataMaterialBSDFCount;
@@ -442,6 +442,14 @@ protected:
 	bool bStrataUsesConversionFromLegacy;
 	bool bStrataOutputsOpaqueRoughRefractions;
 
+	/** Stack of unique id for each node of the strata tree 
+	* This is transient and updated on the fly in the exact same way when parsing node for 
+	*  1- StrataGenerateMaterialTopologyTree: generating a picture of the strata material tree for code generation and simplifications.
+	*  2- CompilePropertyAndSetMaterialProperty(MP_ShadingModel): compiling the code with some features enabled/disabled and accounting for tree simplification decided beforehand.
+	* It is not valid to uise that information outside of those functions.
+	*/
+	TArray<FGuid> StrataNodeIdentifierStack;
+
 	struct FStrataSimplificationStatus
 	{
 		bool bMaterialFitsInMemoryBudget = false;	// Track whether or not the material fits.
@@ -1099,9 +1107,14 @@ protected:
 	virtual int32 StrataHazinessToSecondaryRoughness(int32 BaseRoughness, int32 Haziness, int32 OutputIndex) override;
 	virtual int32 StrataCompilePreview(int32 StrataDataCodeChunk) override;
 	virtual bool StrataSkipsOpacityEvaluation() override;
+
+	virtual FGuid StrataTreeStackPush(UMaterialExpression* Expression, uint32 InputIndex) override;
+	virtual FGuid StrataTreeStackGetPathUniqueId() override;
+	virtual FGuid StrataTreeStackGetParentPathUniqueId() override;
+	virtual void StrataTreeStackPop() override;
 	
-	virtual FStrataOperator& StrataCompilationRegisterOperator(int32 OperatorType, UMaterialExpression* Expression, UMaterialExpression* Parent, bool bUseParameterBlending = false) override;
-	virtual FStrataOperator& StrataCompilationGetOperator(UMaterialExpression* Expression) override;
+	virtual FStrataOperator& StrataCompilationRegisterOperator(int32 OperatorType, FGuid StrataExpressionGuid, UMaterialExpression* Parent, FGuid StrataParentExpressionGuid, bool bUseParameterBlending = false) override;
+	virtual FStrataOperator& StrataCompilationGetOperator(FGuid StrataExpressionGuid) override;
 	virtual FStrataOperator* StrataCompilationGetOperatorFromIndex(int32 OperatorIndex) override;
 
 	virtual FStrataRegisteredSharedLocalBasis StrataCompilationInfoRegisterSharedLocalBasis(int32 NormalCodeChunk) override;
diff --git a/Engine/Source/Runtime/Engine/Private/Materials/MaterialExpressions.cpp b/Engine/Source/Runtime/Engine/Private/Materials/MaterialExpressions.cpp
--- a/Engine/Source/Runtime/Engine/Private/Materials/MaterialExpressions.cpp
+++ b/Engine/Source/Runtime/Engine/Private/Materials/MaterialExpressions.cpp
@@ -7499,15 +7499,31 @@ int32 UMaterialExpressionBlendMaterialAttributes::Compile(class FMaterialCompile
 
 	switch (BlendType)
 	{
-	case EMaterialAttributeBlend::UseA: return A.CompileWithDefault(Compiler, AttributeID);
-	case EMaterialAttributeBlend::UseB: return B.CompileWithDefault(Compiler, AttributeID);
+	case EMaterialAttributeBlend::UseA:
+	{
+		Compiler->StrataTreeStackPush(this, 0);
+		int32 CodeChunk = A.CompileWithDefault(Compiler, AttributeID);
+		Compiler->StrataTreeStackPop();
+		return CodeChunk;
+	}
+	case EMaterialAttributeBlend::UseB:
+	{
+		Compiler->StrataTreeStackPush(this, 1);
+		int32 CodeChunk = B.CompileWithDefault(Compiler, AttributeID);
+		Compiler->StrataTreeStackPop();
+		return CodeChunk;
+	}
 	default:
 		check(BlendType == EMaterialAttributeBlend::Blend);
 	}
 
 	// Allow custom blends or fallback to standard interpolation
+	Compiler->StrataTreeStackPush(this, 0);
 	int32 ResultA = A.CompileWithDefault(Compiler, AttributeID);
+	Compiler->StrataTreeStackPop();
+	Compiler->StrataTreeStackPush(this, 1);
 	int32 ResultB = B.CompileWithDefault(Compiler, AttributeID);
+	Compiler->StrataTreeStackPop();
 	int32 ResultAlpha = Alpha.Compile(Compiler);
 
 	MaterialAttributeBlendFunction BlendFunction = FMaterialAttributeDefinitionMap::GetBlendFunction(AttributeID);
@@ -7596,7 +7612,7 @@ FStrataOperator* UMaterialExpressionBlendMaterialAttributes::StrataGenerateMater
 {
 	// STRATA_TODO: this likely no longer work. We would need to stop parsing and always do parameter blending at this stage.
 	const bool bUseParameterBlending = false;
-	FStrataOperator& StrataOperator = Compiler->StrataCompilationRegisterOperator(STRATA_OPERATOR_HORIZONTAL, this, Parent, bUseParameterBlending);
+	FStrataOperator& StrataOperator = Compiler->StrataCompilationRegisterOperator(STRATA_OPERATOR_HORIZONTAL, Compiler->StrataTreeStackGetPathUniqueId(), Parent, Compiler->StrataTreeStackGetParentPathUniqueId(), bUseParameterBlending);
 
 	UMaterialExpression* ChildAExpression = A.GetTracedInput().Expression;
 	UMaterialExpression* ChildBExpression = B.GetTracedInput().Expression;
@@ -7604,12 +7620,16 @@ FStrataOperator* UMaterialExpressionBlendMaterialAttributes::StrataGenerateMater
 	FStrataOperator* OpB = nullptr;
 	if (ChildAExpression)
 	{
+		Compiler->StrataTreeStackPush(this, 0);
 		OpA = ChildAExpression->StrataGenerateMaterialTopologyTree(Compiler, this, A.OutputIndex);
+		Compiler->StrataTreeStackPop();
 		AssignOperatorIndexIfNotNull(StrataOperator.LeftIndex, OpA);
 	}
 	if (ChildBExpression)
 	{
+		Compiler->StrataTreeStackPush(this, 1);
 		OpB = ChildBExpression->StrataGenerateMaterialTopologyTree(Compiler, this, B.OutputIndex);
+		Compiler->StrataTreeStackPop();
 		AssignOperatorIndexIfNotNull(StrataOperator.RightIndex, OpB);
 	}
 	CombineFlagForParameterBlending(StrataOperator, OpA, OpB);
@@ -22066,7 +22086,7 @@ int32 UMaterialExpressionStrataLegacyConversion::Compile(class FMaterialCompiler
 		SSSProfileCodeChunk = Compiler->ForceCast(Compiler->ScalarParameter(GetSubsurfaceProfileParameterName(), 1.0f), MCT_Float1);
 	}
 
-	FStrataOperator& StrataOperator = Compiler->StrataCompilationGetOperator(this);
+	FStrataOperator& StrataOperator = Compiler->StrataCompilationGetOperator(Compiler->StrataTreeStackGetPathUniqueId());
 	StrataOperator.BSDFRegisteredSharedLocalBasis = NewRegisteredSharedLocalBasis;
 
 	int32 ShadingModelCount = ConvertedStrataMaterialInfo.CountShadingModels();
@@ -22266,7 +22286,7 @@ FStrataOperator* UMaterialExpressionStrataLegacyConversion::StrataGenerateMateri
 {
 	auto AddDefaultWorstCase = [&](bool bSSS)
 	{
-		FStrataOperator& SlabOperator = Compiler->StrataCompilationRegisterOperator(STRATA_OPERATOR_BSDF_LEGACY, this, Parent);
+		FStrataOperator& SlabOperator = Compiler->StrataCompilationRegisterOperator(STRATA_OPERATOR_BSDF_LEGACY, Compiler->StrataTreeStackGetPathUniqueId(), Parent, Compiler->StrataTreeStackGetParentPathUniqueId());
 		SlabOperator.BSDFType = STRATA_BSDF_TYPE_SLAB;
 		SlabOperator.bBSDFHasSSS = true;
 		SlabOperator.bBSDFHasMFPPluggedIn = true;
@@ -22287,7 +22307,7 @@ FStrataOperator* UMaterialExpressionStrataLegacyConversion::StrataGenerateMateri
 
 		if (ConvertedStrataMaterialInfo.HasShadingModel(SSM_Unlit))
 		{
-			FStrataOperator& Operator = Compiler->StrataCompilationRegisterOperator(STRATA_OPERATOR_BSDF_LEGACY, this, Parent);
+			FStrataOperator& Operator = Compiler->StrataCompilationRegisterOperator(STRATA_OPERATOR_BSDF_LEGACY, Compiler->StrataTreeStackGetPathUniqueId(), Parent, Compiler->StrataTreeStackGetParentPathUniqueId());
 			Operator.BSDFType = STRATA_BSDF_TYPE_UNLIT;
 			return &Operator;
 		}
@@ -22301,19 +22321,19 @@ FStrataOperator* UMaterialExpressionStrataLegacyConversion::StrataGenerateMateri
 		}
 		else if (ConvertedStrataMaterialInfo.HasShadingModel(SSM_Hair))
 		{
-			FStrataOperator& Operator = Compiler->StrataCompilationRegisterOperator(STRATA_OPERATOR_BSDF_LEGACY, this, Parent);
+			FStrataOperator& Operator = Compiler->StrataCompilationRegisterOperator(STRATA_OPERATOR_BSDF_LEGACY, Compiler->StrataTreeStackGetPathUniqueId(), Parent, Compiler->StrataTreeStackGetParentPathUniqueId());
 			Operator.BSDFType = STRATA_BSDF_TYPE_HAIR;
 			return &Operator;
 		}
 		else if (ConvertedStrataMaterialInfo.HasShadingModel(SSM_Eye))
 		{
-			FStrataOperator& Operator = Compiler->StrataCompilationRegisterOperator(STRATA_OPERATOR_BSDF_LEGACY, this, Parent);
+			FStrataOperator& Operator = Compiler->StrataCompilationRegisterOperator(STRATA_OPERATOR_BSDF_LEGACY, Compiler->StrataTreeStackGetPathUniqueId(), Parent, Compiler->StrataTreeStackGetParentPathUniqueId());
 			Operator.BSDFType = STRATA_BSDF_TYPE_EYE;
 			return &Operator;
 		}
 		else if (ConvertedStrataMaterialInfo.HasShadingModel(SSM_SingleLayerWater))
 		{
-			FStrataOperator& Operator = Compiler->StrataCompilationRegisterOperator(STRATA_OPERATOR_BSDF_LEGACY, this, Parent);
+			FStrataOperator& Operator = Compiler->StrataCompilationRegisterOperator(STRATA_OPERATOR_BSDF_LEGACY, Compiler->StrataTreeStackGetPathUniqueId(), Parent, Compiler->StrataTreeStackGetParentPathUniqueId());
 			Operator.BSDFType = STRATA_BSDF_TYPE_SINGLELAYERWATER;
 			return &Operator;
 		}
@@ -22398,7 +22418,7 @@ int32 UMaterialExpressionStrataSlabBSDF::Compile(class FMaterialCompiler* Compil
 	int32 TangentCodeChunk = bHasAnisotropy ? CompileWithDefaultTangentWS(Compiler, Tangent) : INDEX_NONE;
 	const FStrataRegisteredSharedLocalBasis NewRegisteredSharedLocalBasis = StrataCompilationInfoCreateSharedLocalBasis(Compiler, NormalCodeChunk, TangentCodeChunk);
 
-	FStrataOperator& StrataOperator = Compiler->StrataCompilationGetOperator(this);
+	FStrataOperator& StrataOperator = Compiler->StrataCompilationGetOperator(Compiler->StrataTreeStackGetPathUniqueId());
 	StrataOperator.BSDFRegisteredSharedLocalBasis = NewRegisteredSharedLocalBasis;
 
 	int32 OutputCodeChunk = Compiler->StrataSlabBSDF(
@@ -22865,7 +22885,7 @@ void UMaterialExpressionStrataSlabBSDF::GatherStrataMaterialInfo(FStrataMaterial
 
 FStrataOperator* UMaterialExpressionStrataSlabBSDF::StrataGenerateMaterialTopologyTree(class FMaterialCompiler* Compiler, class UMaterialExpression* Parent, int32 OutputIndex)
 {
-	FStrataOperator& StrataOperator = Compiler->StrataCompilationRegisterOperator(STRATA_OPERATOR_BSDF, this, Parent);
+	FStrataOperator& StrataOperator = Compiler->StrataCompilationRegisterOperator(STRATA_OPERATOR_BSDF, Compiler->StrataTreeStackGetPathUniqueId(), Parent, Compiler->StrataTreeStackGetParentPathUniqueId());
 	StrataOperator.BSDFType = STRATA_BSDF_TYPE_SLAB;
 	StrataOperator.bBSDFHasEdgeColor = HasEdgeColor();
 	StrataOperator.bBSDFHasFuzz = HasFuzz();
@@ -22939,7 +22959,7 @@ int32 UMaterialExpressionStrataSimpleClearCoatBSDF::Compile(class FMaterialCompi
 	const int32 NullTangentCodeChunk = INDEX_NONE;
 	const FStrataRegisteredSharedLocalBasis NewRegisteredSharedLocalBasis = StrataCompilationInfoCreateSharedLocalBasis(Compiler, NormalCodeChunk, NullTangentCodeChunk);
 
-	FStrataOperator& StrataOperator = Compiler->StrataCompilationGetOperator(this);
+	FStrataOperator& StrataOperator = Compiler->StrataCompilationGetOperator(Compiler->StrataTreeStackGetPathUniqueId());
 	StrataOperator.BSDFRegisteredSharedLocalBasis = NewRegisteredSharedLocalBasis;
 
 	int32 OutputCodeChunk = Compiler->StrataSlabBSDF(
@@ -23100,7 +23120,7 @@ void UMaterialExpressionStrataSimpleClearCoatBSDF::GatherStrataMaterialInfo(FStr
 
 FStrataOperator* UMaterialExpressionStrataSimpleClearCoatBSDF::StrataGenerateMaterialTopologyTree(class FMaterialCompiler* Compiler, class UMaterialExpression* Parent, int32 OutputIndex)
 {
-	FStrataOperator& StrataOperator = Compiler->StrataCompilationRegisterOperator(STRATA_OPERATOR_BSDF, this, Parent);
+	FStrataOperator& StrataOperator = Compiler->StrataCompilationRegisterOperator(STRATA_OPERATOR_BSDF, Compiler->StrataTreeStackGetPathUniqueId(), Parent, Compiler->StrataTreeStackGetParentPathUniqueId());
 	StrataOperator.BSDFType = STRATA_BSDF_TYPE_SLAB;
 	StrataOperator.bBSDFHasEdgeColor = false;
 	StrataOperator.bBSDFHasFuzz = false;
@@ -23184,7 +23204,7 @@ void UMaterialExpressionStrataVolumetricFogCloudBSDF::GatherStrataMaterialInfo(F
 
 FStrataOperator* UMaterialExpressionStrataVolumetricFogCloudBSDF::StrataGenerateMaterialTopologyTree(class FMaterialCompiler* Compiler, class UMaterialExpression* Parent, int32 OutputIndex)
 {
-	FStrataOperator& StrataOperator = Compiler->StrataCompilationRegisterOperator(STRATA_OPERATOR_BSDF, this, Parent);
+	FStrataOperator& StrataOperator = Compiler->StrataCompilationRegisterOperator(STRATA_OPERATOR_BSDF, Compiler->StrataTreeStackGetPathUniqueId(), Parent, Compiler->StrataTreeStackGetParentPathUniqueId());
 	StrataOperator.BSDFType = STRATA_BSDF_TYPE_VOLUMETRICFOGCLOUD;
 	return &StrataOperator;
 }
@@ -23251,7 +23271,7 @@ void UMaterialExpressionStrataLightFunction::GatherStrataMaterialInfo(FStrataMat
 
 FStrataOperator* UMaterialExpressionStrataLightFunction::StrataGenerateMaterialTopologyTree(class FMaterialCompiler* Compiler, class UMaterialExpression* Parent, int32 OutputIndex)
 {
-	FStrataOperator& StrataOperator = Compiler->StrataCompilationRegisterOperator(STRATA_OPERATOR_BSDF, this, Parent);
+	FStrataOperator& StrataOperator = Compiler->StrataCompilationRegisterOperator(STRATA_OPERATOR_BSDF, Compiler->StrataTreeStackGetPathUniqueId(), Parent, Compiler->StrataTreeStackGetParentPathUniqueId());
 	StrataOperator.BSDFType = STRATA_BSDF_TYPE_UNLIT;
 	return &StrataOperator;
 }
@@ -23324,7 +23344,7 @@ void UMaterialExpressionStrataPostProcess::GatherStrataMaterialInfo(FStrataMater
 
 FStrataOperator* UMaterialExpressionStrataPostProcess::StrataGenerateMaterialTopologyTree(class FMaterialCompiler* Compiler, class UMaterialExpression* Parent, int32 OutputIndex)
 {
-	FStrataOperator& StrataOperator = Compiler->StrataCompilationRegisterOperator(STRATA_OPERATOR_BSDF, this, Parent);
+	FStrataOperator& StrataOperator = Compiler->StrataCompilationRegisterOperator(STRATA_OPERATOR_BSDF, Compiler->StrataTreeStackGetPathUniqueId(), Parent, Compiler->StrataTreeStackGetParentPathUniqueId());
 	StrataOperator.BSDFType = STRATA_BSDF_TYPE_UNLIT;
 	return &StrataOperator;
 }
@@ -23355,10 +23375,12 @@ int32 UMaterialExpressionStrataConvertToDecal::Compile(class FMaterialCompiler*
 	}
 
 	int32 CoverageCodeChunk = Coverage.GetTracedInput().Expression ? Coverage.Compile(Compiler) : Compiler->Constant(1.0f);
+	Compiler->StrataTreeStackPush(this, 0);
 	int32 DecalMaterialCodeChunk = DecalMaterial.Compile(Compiler);
+	Compiler->StrataTreeStackPop();
 
 	int32 OutputCodeChunk = INDEX_NONE;
-	FStrataOperator& StrataOperator = Compiler->StrataCompilationGetOperator(this);
+	FStrataOperator& StrataOperator = Compiler->StrataCompilationGetOperator(Compiler->StrataTreeStackGetPathUniqueId());
 	if (!StrataOperator.bUseParameterBlending)
 	{
 		return Compiler->Errorf(TEXT("Strata Convert To Decal node must recveive StrataData a parameter blended strata material sub tree."));
@@ -23422,13 +23444,15 @@ void UMaterialExpressionStrataConvertToDecal::GatherStrataMaterialInfo(FStrataMa
 FStrataOperator* UMaterialExpressionStrataConvertToDecal::StrataGenerateMaterialTopologyTree(class FMaterialCompiler* Compiler, class UMaterialExpression* Parent, int32 OutputIndex)
 {
 	const bool bUseParameterBlending = true;
-	FStrataOperator& StrataOperator = Compiler->StrataCompilationRegisterOperator(STRATA_OPERATOR_WEIGHT, this, Parent, bUseParameterBlending);
+	FStrataOperator& StrataOperator = Compiler->StrataCompilationRegisterOperator(STRATA_OPERATOR_WEIGHT, Compiler->StrataTreeStackGetPathUniqueId(), Parent, Compiler->StrataTreeStackGetParentPathUniqueId(), bUseParameterBlending);
 
 	UMaterialExpression* ChildDecalMaterialExpression = DecalMaterial.GetTracedInput().Expression;
 	FStrataOperator* OpA = nullptr;
 	if (ChildDecalMaterialExpression)
 	{
+		Compiler->StrataTreeStackPush(this, 0);
 		OpA = ChildDecalMaterialExpression->StrataGenerateMaterialTopologyTree(Compiler, this, DecalMaterial.OutputIndex);
+		Compiler->StrataTreeStackPop();
 		AssignOperatorIndexIfNotNull(StrataOperator.LeftIndex, OpA);
 	}
 	CombineFlagForParameterBlending(StrataOperator, OpA);
@@ -23502,7 +23526,7 @@ void UMaterialExpressionStrataUnlitBSDF::GatherStrataMaterialInfo(FStrataMateria
 
 FStrataOperator* UMaterialExpressionStrataUnlitBSDF::StrataGenerateMaterialTopologyTree(class FMaterialCompiler* Compiler, class UMaterialExpression* Parent, int32 OutputIndex)
 {
-	FStrataOperator& StrataOperator = Compiler->StrataCompilationRegisterOperator(STRATA_OPERATOR_BSDF, this, Parent);
+	FStrataOperator& StrataOperator = Compiler->StrataCompilationRegisterOperator(STRATA_OPERATOR_BSDF, Compiler->StrataTreeStackGetPathUniqueId(), Parent, Compiler->StrataTreeStackGetParentPathUniqueId());
 	StrataOperator.BSDFType = STRATA_BSDF_TYPE_UNLIT;
 	return &StrataOperator;
 }
@@ -23531,7 +23555,7 @@ int32 UMaterialExpressionStrataHairBSDF::Compile(class FMaterialCompiler* Compil
 	int32 TangentCodeChunk = CompileWithDefaultTangentWS(Compiler, Tangent);
 	const FStrataRegisteredSharedLocalBasis NewRegisteredSharedLocalBasis = StrataCompilationInfoCreateSharedLocalBasis(Compiler, TangentCodeChunk);
 
-	FStrataOperator& StrataOperator = Compiler->StrataCompilationGetOperator(this);
+	FStrataOperator& StrataOperator = Compiler->StrataCompilationGetOperator(Compiler->StrataTreeStackGetPathUniqueId());
 	StrataOperator.BSDFRegisteredSharedLocalBasis = NewRegisteredSharedLocalBasis;
 
 	if (StrataOperator.bUseParameterBlending)
@@ -23617,7 +23641,7 @@ void UMaterialExpressionStrataHairBSDF::GatherStrataMaterialInfo(FStrataMaterial
 
 FStrataOperator* UMaterialExpressionStrataHairBSDF::StrataGenerateMaterialTopologyTree(class FMaterialCompiler* Compiler, class UMaterialExpression* Parent, int32 OutputIndex)
 {
-	FStrataOperator& StrataOperator = Compiler->StrataCompilationRegisterOperator(STRATA_OPERATOR_BSDF, this, Parent);
+	FStrataOperator& StrataOperator = Compiler->StrataCompilationRegisterOperator(STRATA_OPERATOR_BSDF, Compiler->StrataTreeStackGetPathUniqueId(), Parent, Compiler->StrataTreeStackGetParentPathUniqueId());
 	StrataOperator.BSDFType = STRATA_BSDF_TYPE_HAIR;
 	return &StrataOperator;
 }
@@ -23649,7 +23673,7 @@ int32 UMaterialExpressionStrataEyeBSDF::Compile(class FMaterialCompiler* Compile
 		SSSProfileCodeChunk = Compiler->ForceCast(Compiler->ScalarParameter(GetSubsurfaceProfileParameterName(), 1.0f), MCT_Float1);
 	}
 
-	FStrataOperator& StrataOperator = Compiler->StrataCompilationGetOperator(this);
+	FStrataOperator& StrataOperator = Compiler->StrataCompilationGetOperator(Compiler->StrataTreeStackGetPathUniqueId());
 	StrataOperator.BSDFRegisteredSharedLocalBasis = NewRegisteredSharedLocalBasis;
 
 	if (StrataOperator.bUseParameterBlending)
@@ -23730,7 +23754,7 @@ void UMaterialExpressionStrataEyeBSDF::GatherStrataMaterialInfo(FStrataMaterialI
 
 FStrataOperator* UMaterialExpressionStrataEyeBSDF::StrataGenerateMaterialTopologyTree(class FMaterialCompiler* Compiler, class UMaterialExpression* Parent, int32 OutputIndex)
 {
-	FStrataOperator& StrataOperator = Compiler->StrataCompilationRegisterOperator(STRATA_OPERATOR_BSDF, this, Parent);
+	FStrataOperator& StrataOperator = Compiler->StrataCompilationRegisterOperator(STRATA_OPERATOR_BSDF, Compiler->StrataTreeStackGetPathUniqueId(), Parent, Compiler->StrataTreeStackGetParentPathUniqueId());
 	StrataOperator.BSDFType = STRATA_BSDF_TYPE_EYE;
 	return &StrataOperator;
 }
@@ -23757,7 +23781,7 @@ int32 UMaterialExpressionStrataSingleLayerWaterBSDF::Compile(class FMaterialComp
 	int32 NormalCodeChunk = CompileWithDefaultNormalWS(Compiler, Normal);
 	const FStrataRegisteredSharedLocalBasis NewRegisteredSharedLocalBasis = StrataCompilationInfoCreateSharedLocalBasis(Compiler, NormalCodeChunk);
 
-	FStrataOperator& StrataOperator = Compiler->StrataCompilationGetOperator(this);
+	FStrataOperator& StrataOperator = Compiler->StrataCompilationGetOperator(Compiler->StrataTreeStackGetPathUniqueId());
 	StrataOperator.BSDFRegisteredSharedLocalBasis = NewRegisteredSharedLocalBasis;
 
 	if (StrataOperator.bUseParameterBlending)
@@ -23860,7 +23884,7 @@ void UMaterialExpressionStrataSingleLayerWaterBSDF::GatherStrataMaterialInfo(FSt
 
 FStrataOperator* UMaterialExpressionStrataSingleLayerWaterBSDF::StrataGenerateMaterialTopologyTree(class FMaterialCompiler* Compiler, class UMaterialExpression* Parent, int32 OutputIndex)
 {
-	FStrataOperator& StrataOperator = Compiler->StrataCompilationRegisterOperator(STRATA_OPERATOR_BSDF, this, Parent);
+	FStrataOperator& StrataOperator = Compiler->StrataCompilationRegisterOperator(STRATA_OPERATOR_BSDF, Compiler->StrataTreeStackGetPathUniqueId(), Parent, Compiler->StrataTreeStackGetParentPathUniqueId());
 	StrataOperator.BSDFType = STRATA_BSDF_TYPE_SINGLELAYERWATER;
 	return &StrataOperator;
 }
@@ -23895,13 +23919,17 @@ int32 UMaterialExpressionStrataHorizontalMixing::Compile(class FMaterialCompiler
 		return Compiler->Errorf(TEXT("Missing Background input"));
 	}
 
-	int32 ForegroundCodeChunk = Foreground.Compile(Compiler);
+	Compiler->StrataTreeStackPush(this, 0);
 	int32 BackgroundCodeChunk = Background.Compile(Compiler);
+	Compiler->StrataTreeStackPop();
+	Compiler->StrataTreeStackPush(this, 1);
+	int32 ForegroundCodeChunk = Foreground.Compile(Compiler);
+	Compiler->StrataTreeStackPop();
 
 	const int32 HorizontalMixCodeChunk = CompileWithDefaultFloat1(Compiler, Mix, 0.5f);
 
 	int32 OutputCodeChunk = INDEX_NONE;
-	FStrataOperator& StrataOperator = Compiler->StrataCompilationGetOperator(this);
+	FStrataOperator& StrataOperator = Compiler->StrataCompilationGetOperator(Compiler->StrataTreeStackGetPathUniqueId());
 	if (StrataOperator.bUseParameterBlending)
 	{
 		if (ForegroundCodeChunk == INDEX_NONE)
@@ -24000,7 +24028,7 @@ void UMaterialExpressionStrataHorizontalMixing::GatherStrataMaterialInfo(FStrata
 
 FStrataOperator* UMaterialExpressionStrataHorizontalMixing::StrataGenerateMaterialTopologyTree(class FMaterialCompiler* Compiler, class UMaterialExpression* Parent, int32 OutputIndex)
 {
-	FStrataOperator& StrataOperator = Compiler->StrataCompilationRegisterOperator(STRATA_OPERATOR_HORIZONTAL, this, Parent, bUseParameterBlending);
+	FStrataOperator& StrataOperator = Compiler->StrataCompilationRegisterOperator(STRATA_OPERATOR_HORIZONTAL, Compiler->StrataTreeStackGetPathUniqueId(), Parent, Compiler->StrataTreeStackGetParentPathUniqueId(), bUseParameterBlending);
 
 	UMaterialExpression* ChildAExpression = Background.GetTracedInput().Expression;
 	UMaterialExpression* ChildBExpression = Foreground.GetTracedInput().Expression;
@@ -24008,12 +24036,16 @@ FStrataOperator* UMaterialExpressionStrataHorizontalMixing::StrataGenerateMateri
 	FStrataOperator* OpB = nullptr;
 	if (ChildAExpression)
 	{
+		Compiler->StrataTreeStackPush(this, 0);
 		OpA = ChildAExpression->StrataGenerateMaterialTopologyTree(Compiler, this, Background.OutputIndex);
+		Compiler->StrataTreeStackPop();
 		AssignOperatorIndexIfNotNull(StrataOperator.LeftIndex, OpA);
 	}
 	if (ChildBExpression)
 	{
+		Compiler->StrataTreeStackPush(this, 1);
 		OpB = ChildBExpression->StrataGenerateMaterialTopologyTree(Compiler, this, Foreground.OutputIndex);
+		Compiler->StrataTreeStackPop();
 		AssignOperatorIndexIfNotNull(StrataOperator.RightIndex, OpB);
 	}
 	CombineFlagForParameterBlending(StrataOperator, OpA, OpB);
@@ -24051,11 +24083,15 @@ int32 UMaterialExpressionStrataVerticalLayering::Compile(class FMaterialCompiler
 		return Compiler->Errorf(TEXT("Missing Base input"));
 	}
 
+	Compiler->StrataTreeStackPush(this, 0);
 	int32 TopCodeChunk = Top.Compile(Compiler);
+	Compiler->StrataTreeStackPop();
+	Compiler->StrataTreeStackPush(this, 1);
 	int32 BaseCodeChunk = Base.Compile(Compiler);
+	Compiler->StrataTreeStackPop();
 
 	int32 OutputCodeChunk = INDEX_NONE;
-	FStrataOperator& StrataOperator = Compiler->StrataCompilationGetOperator(this);
+	FStrataOperator& StrataOperator = Compiler->StrataCompilationGetOperator(Compiler->StrataTreeStackGetPathUniqueId());
 	if (StrataOperator.bUseParameterBlending)
 	{
 		FStrataOperator* TopBSDFOperator = Compiler->StrataCompilationGetOperatorFromIndex(StrataOperator.LeftIndex);
@@ -24150,7 +24186,7 @@ void UMaterialExpressionStrataVerticalLayering::GatherStrataMaterialInfo(FStrata
 
 FStrataOperator* UMaterialExpressionStrataVerticalLayering::StrataGenerateMaterialTopologyTree(class FMaterialCompiler* Compiler, class UMaterialExpression* Parent, int32 OutputIndex)
 {
-	FStrataOperator& StrataOperator = Compiler->StrataCompilationRegisterOperator(STRATA_OPERATOR_VERTICAL, this, Parent, bUseParameterBlending);
+	FStrataOperator& StrataOperator = Compiler->StrataCompilationRegisterOperator(STRATA_OPERATOR_VERTICAL, Compiler->StrataTreeStackGetPathUniqueId(), Parent, Compiler->StrataTreeStackGetParentPathUniqueId(), bUseParameterBlending);
 
 	UMaterialExpression* ChildAExpression = Top.GetTracedInput().Expression;
 	UMaterialExpression* ChildBExpression = Base.GetTracedInput().Expression;
@@ -24158,12 +24194,16 @@ FStrataOperator* UMaterialExpressionStrataVerticalLayering::StrataGenerateMateri
 	FStrataOperator* OpB = nullptr;
 	if (ChildAExpression)
 	{
+		Compiler->StrataTreeStackPush(this, 0);
 		OpA = ChildAExpression->StrataGenerateMaterialTopologyTree(Compiler, this, Top.OutputIndex);
+		Compiler->StrataTreeStackPop();
 		AssignOperatorIndexIfNotNull(StrataOperator.LeftIndex, OpA);
 	}
 	if (ChildBExpression)
 	{
+		Compiler->StrataTreeStackPush(this, 1);
 		OpB = ChildBExpression->StrataGenerateMaterialTopologyTree(Compiler, this, Base.OutputIndex);
+		Compiler->StrataTreeStackPop();
 		AssignOperatorIndexIfNotNull(StrataOperator.RightIndex, OpB);
 	}
 	CombineFlagForParameterBlending(StrataOperator, OpA, OpB);
@@ -24201,11 +24241,15 @@ int32 UMaterialExpressionStrataAdd::Compile(class FMaterialCompiler* Compiler, i
 		return Compiler->Errorf(TEXT("Missing B input"));
 	}
 
+	Compiler->StrataTreeStackPush(this, 0);
 	int32 ACodeChunk = A.Compile(Compiler);
+	Compiler->StrataTreeStackPop();
+	Compiler->StrataTreeStackPush(this, 1);
 	int32 BCodeChunk = B.Compile(Compiler);
+	Compiler->StrataTreeStackPop();
 
 	int32 OutputCodeChunk = INDEX_NONE;
-	FStrataOperator& StrataOperator = Compiler->StrataCompilationGetOperator(this);
+	FStrataOperator& StrataOperator = Compiler->StrataCompilationGetOperator(Compiler->StrataTreeStackGetPathUniqueId());
 	if (StrataOperator.bUseParameterBlending)
 	{
 		FStrataOperator* ABSDFOperator = Compiler->StrataCompilationGetOperatorFromIndex(StrataOperator.LeftIndex);
@@ -24299,7 +24343,7 @@ void UMaterialExpressionStrataAdd::GatherStrataMaterialInfo(FStrataMaterialInfo&
 
 FStrataOperator* UMaterialExpressionStrataAdd::StrataGenerateMaterialTopologyTree(class FMaterialCompiler* Compiler, class UMaterialExpression* Parent, int32 OutputIndex)
 {
-	FStrataOperator& StrataOperator = Compiler->StrataCompilationRegisterOperator(STRATA_OPERATOR_ADD, this, Parent, bUseParameterBlending);
+	FStrataOperator& StrataOperator = Compiler->StrataCompilationRegisterOperator(STRATA_OPERATOR_ADD, Compiler->StrataTreeStackGetPathUniqueId(), Parent, Compiler->StrataTreeStackGetParentPathUniqueId(), bUseParameterBlending);
 
 	UMaterialExpression* ChildAExpression = A.GetTracedInput().Expression;
 	UMaterialExpression* ChildBExpression = B.GetTracedInput().Expression;
@@ -24307,12 +24351,16 @@ FStrataOperator* UMaterialExpressionStrataAdd::StrataGenerateMaterialTopologyTre
 	FStrataOperator* OpB = nullptr;
 	if (ChildAExpression)
 	{
+		Compiler->StrataTreeStackPush(this, 0);
 		OpA = ChildAExpression->StrataGenerateMaterialTopologyTree(Compiler, this, A.OutputIndex);
+		Compiler->StrataTreeStackPop();
 		AssignOperatorIndexIfNotNull(StrataOperator.LeftIndex, OpA);
 	}
 	if (ChildBExpression)
 	{
+		Compiler->StrataTreeStackPush(this, 1);
 		OpB = ChildBExpression->StrataGenerateMaterialTopologyTree(Compiler, this, B.OutputIndex);
+		Compiler->StrataTreeStackPop();
 		AssignOperatorIndexIfNotNull(StrataOperator.RightIndex, OpB);
 	}
 	CombineFlagForParameterBlending(StrataOperator, OpA, OpB);
@@ -24345,11 +24393,13 @@ int32 UMaterialExpressionStrataWeight::Compile(class FMaterialCompiler* Compiler
 		return Compiler->Errorf(TEXT("Missing A input"));
 	}
 
+	Compiler->StrataTreeStackPush(this, 0);
 	int32 ACodeChunk = A.Compile(Compiler);
+	Compiler->StrataTreeStackPop();
 	int32 WeightCodeChunk = Weight.GetTracedInput().Expression ? Weight.Compile(Compiler) : Compiler->Constant(1.0f);
 
 	int32 OutputCodeChunk = INDEX_NONE;
-	FStrataOperator& StrataOperator = Compiler->StrataCompilationGetOperator(this);
+	FStrataOperator& StrataOperator = Compiler->StrataCompilationGetOperator(Compiler->StrataTreeStackGetPathUniqueId());
 	if (StrataOperator.bUseParameterBlending)
 	{
 		// Propagate the parameter blended normal
@@ -24415,13 +24465,15 @@ void UMaterialExpressionStrataWeight::GatherStrataMaterialInfo(FStrataMaterialIn
 
 FStrataOperator* UMaterialExpressionStrataWeight::StrataGenerateMaterialTopologyTree(class FMaterialCompiler* Compiler, class UMaterialExpression* Parent, int32 OutputIndex)
 {
-	FStrataOperator& StrataOperator = Compiler->StrataCompilationRegisterOperator(STRATA_OPERATOR_WEIGHT, this, Parent);
+	FStrataOperator& StrataOperator = Compiler->StrataCompilationRegisterOperator(STRATA_OPERATOR_WEIGHT, Compiler->StrataTreeStackGetPathUniqueId(), Parent, Compiler->StrataTreeStackGetParentPathUniqueId());
 
 	UMaterialExpression* ChildAExpression = A.GetTracedInput().Expression;
 	FStrataOperator* OpA = nullptr;
 	if (ChildAExpression)
 	{
+		Compiler->StrataTreeStackPush(this, 0);
 		OpA = ChildAExpression->StrataGenerateMaterialTopologyTree(Compiler, this, A.OutputIndex);
+		Compiler->StrataTreeStackPop();
 		AssignOperatorIndexIfNotNull(StrataOperator.LeftIndex, OpA);
 	}
 	CombineFlagForParameterBlending(StrataOperator, OpA);
@@ -24454,12 +24506,14 @@ int32 UMaterialExpressionStrataThinFilm::Compile(class FMaterialCompiler* Compil
 		return Compiler->Errorf(TEXT("Missing A input"));
 	}
 
+	Compiler->StrataTreeStackPush(this, 0);
 	int32 ACodeChunk = A.Compile(Compiler);
+	Compiler->StrataTreeStackPop();
 	int32 ThicknessCodeChunk = Thickness.GetTracedInput().Expression ? Thickness.Compile(Compiler) : Compiler->Constant(1.0f);
 	int32 IORCodeChunk = IOR.GetTracedInput().Expression ? IOR.Compile(Compiler) : Compiler->Constant(1.44f);
 
 	int32 OutputCodeChunk = INDEX_NONE;
-	FStrataOperator& StrataOperator = Compiler->StrataCompilationGetOperator(this);
+	FStrataOperator& StrataOperator = Compiler->StrataCompilationGetOperator(Compiler->StrataTreeStackGetPathUniqueId());
 	if (StrataOperator.bUseParameterBlending)
 	{
 		// Propagate the parameter blended normal
@@ -24514,13 +24568,15 @@ void UMaterialExpressionStrataThinFilm::GatherStrataMaterialInfo(FStrataMaterial
 
 FStrataOperator* UMaterialExpressionStrataThinFilm::StrataGenerateMaterialTopologyTree(class FMaterialCompiler* Compiler, class UMaterialExpression* Parent, int32 OutputIndex)
 {
-	FStrataOperator& StrataOperator = Compiler->StrataCompilationRegisterOperator(STRATA_OPERATOR_THINFILM, this, Parent);
+	FStrataOperator& StrataOperator = Compiler->StrataCompilationRegisterOperator(STRATA_OPERATOR_THINFILM, Compiler->StrataTreeStackGetPathUniqueId(), Parent, Compiler->StrataTreeStackGetParentPathUniqueId());
 
 	UMaterialExpression* ChildAExpression = A.GetTracedInput().Expression;
 	FStrataOperator* OpA = nullptr;
 	if (ChildAExpression)
 	{
+		Compiler->StrataTreeStackPush(this, 0);
 		OpA = ChildAExpression->StrataGenerateMaterialTopologyTree(Compiler, this, A.OutputIndex);
+		Compiler->StrataTreeStackPop();
 		AssignOperatorIndexIfNotNull(StrataOperator.LeftIndex, OpA);
 	}
 	CombineFlagForParameterBlending(StrataOperator, OpA);
diff --git a/Engine/Source/Runtime/Engine/Public/MaterialCompiler.h b/Engine/Source/Runtime/Engine/Public/MaterialCompiler.h
--- a/Engine/Source/Runtime/Engine/Public/MaterialCompiler.h
+++ b/Engine/Source/Runtime/Engine/Public/MaterialCompiler.h
@@ -681,15 +681,33 @@ public:
 	 * This is dedicated to skip evaluating any opacity input when a material instance toggles the translucent blend mode to opaque.
 	 */
 	virtual bool StrataSkipsOpacityEvaluation() = 0;
-	
+
+	/**
+	 * Pushes a node of the Strata tree being walked. A node is caracterised by an expression and its input path taken.
+	 */
+	virtual FGuid StrataTreeStackPush(UMaterialExpression* Expression, uint32 InputIndex) = 0;
+	/**
+	 * Returns the unique id of the strata tree path, identifying the current node we have reached. 
+	 * That is used as a unique identifier for each BSDF, recover their strata tree operator and apply some simplifications if required.
+	 */
+	virtual FGuid StrataTreeStackGetPathUniqueId() = 0;
+	/**
+	 * Returns the unique id of the strata tree path for the parent node of the current node position.
+	 */
+	virtual FGuid StrataTreeStackGetParentPathUniqueId() = 0;
+	/**
+	 * Pops a node node of the Strata tree being walked. Used when walking up the tree back to its root level.
+	 */
+	virtual void StrataTreeStackPop() = 0;
+
 	/**
 	 * Register an operator of the tree representation the Strata material and its topology.
 	 */
-	virtual FStrataOperator& StrataCompilationRegisterOperator(int32 OperatorType, UMaterialExpression* Expression, UMaterialExpression* Parent, bool bUseParameterBlending = false) = 0;
+	virtual FStrataOperator& StrataCompilationRegisterOperator(int32 OperatorType, FGuid StrataExpressionGuid, UMaterialExpression* Parent, FGuid StrataParentExpressionGuid, bool bUseParameterBlending = false) = 0;
 	/**
 	 * Return the operator information for a given expression.
 	 */
-	virtual FStrataOperator& StrataCompilationGetOperator(UMaterialExpression* Expression) = 0;
+	virtual FStrataOperator& StrataCompilationGetOperator(FGuid StrataExpressionGuid) = 0;
 	/**
 	 * Return the operator information for a given index.
 	 */
@@ -1365,14 +1383,31 @@ public:
 		return Compiler->StrataSkipsOpacityEvaluation();
 	}
 
-	virtual FStrataOperator& StrataCompilationRegisterOperator(int32 OperatorType, UMaterialExpression* Expression, UMaterialExpression* Parent, bool bUseParameterBlending = false) override
+	virtual FGuid StrataTreeStackPush(UMaterialExpression* Expression, uint32 InputIndex) override
+	{
+		return Compiler->StrataTreeStackPush(Expression, InputIndex);
+	}
+	virtual FGuid StrataTreeStackGetPathUniqueId() override
+	{
+		return Compiler->StrataTreeStackGetPathUniqueId();
+	}
+	virtual FGuid StrataTreeStackGetParentPathUniqueId() override
+	{
+		return Compiler->StrataTreeStackGetParentPathUniqueId();
+	}
+	virtual void StrataTreeStackPop() override
+	{
+		Compiler->StrataTreeStackPop();
+	}
+
+	virtual FStrataOperator& StrataCompilationRegisterOperator(int32 OperatorType, FGuid StrataExpressionGuid, UMaterialExpression* Parent, FGuid StrataParentExpressionGuid, bool bUseParameterBlending = false) override
 	{
-		return Compiler->StrataCompilationRegisterOperator(OperatorType, Expression, Parent, bUseParameterBlending);
+		return Compiler->StrataCompilationRegisterOperator(OperatorType, StrataExpressionGuid, Parent, StrataParentExpressionGuid, bUseParameterBlending);
 	}
 
-	virtual FStrataOperator& StrataCompilationGetOperator(UMaterialExpression* Expression) override
+	virtual FStrataOperator& StrataCompilationGetOperator(FGuid StrataExpressionGuid) override
 	{
-		return Compiler->StrataCompilationGetOperator(Expression);
+		return Compiler->StrataCompilationGetOperator(StrataExpressionGuid);
 	}
 
 	virtual FStrataOperator* StrataCompilationGetOperatorFromIndex(int32 OperatorIndex) override
