Merging CL#23014870

Added approximate visual hash. Fix for crash on log message.

diff --git a/Engine/Source/Runtime/Engine/Private/SkeletalMesh.cpp b/Engine/Source/Runtime/Engine/Private/SkeletalMesh.cpp
--- a/Engine/Source/Runtime/Engine/Private/SkeletalMesh.cpp
+++ b/Engine/Source/Runtime/Engine/Private/SkeletalMesh.cpp
@@ -4178,6 +4178,97 @@ bool USkeletalMesh::AreAllFlagsIdentical( const TArray<bool>& BoolArray ) const
 	return true;
 }
 
+bool operator== ( const FSkeletalMaterial& LHS, const FSkeletalMaterial& RHS )
+{
+	return ( LHS.MaterialInterface == RHS.MaterialInterface );
+}
+
+bool operator== ( const FSkeletalMaterial& LHS, const UMaterialInterface& RHS )
+{
+	return ( LHS.MaterialInterface == &RHS );
+}
+
+bool operator== ( const UMaterialInterface& LHS, const FSkeletalMaterial& RHS )
+{
+	return ( RHS.MaterialInterface == &LHS );
+}
+
+ENGINE_API FArchive& operator<<(FArchive& Ar, FMeshUVChannelInfo& ChannelData)
+{
+	Ar << ChannelData.bInitialized;
+	Ar << ChannelData.bOverrideDensities;
+
+	for (int32 CoordIndex = 0; CoordIndex < TEXSTREAM_MAX_NUM_UVCHANNELS; ++CoordIndex)
+	{
+		Ar << ChannelData.LocalUVDensities[CoordIndex];
+	}
+
+	return Ar;
+}
+
+#if WITH_EDITORONLY_DATA
+void FSkeletalMaterial::DeclareCustomVersions(FArchive& Ar)
+{
+	Ar.UsingCustomVersion(FEditorObjectVersion::GUID);
+	Ar.UsingCustomVersion(FCoreObjectVersion::GUID);
+}
+#endif
+
+FArchive& operator<<(FArchive& Ar, FSkeletalMaterial& Elem)
+{
+	Ar.UsingCustomVersion(FEditorObjectVersion::GUID);
+	Ar.UsingCustomVersion(FCoreObjectVersion::GUID);
+
+	Ar << Elem.MaterialInterface;
+
+	//Use the automatic serialization instead of this custom operator
+	if (Ar.CustomVer(FEditorObjectVersion::GUID) >= FEditorObjectVersion::RefactorMeshEditorMaterials)
+	{
+		Ar << Elem.MaterialSlotName;
+
+		bool bSerializeImportedMaterialSlotName = !Ar.IsCooking() || Ar.CookingTarget()->HasEditorOnlyData();
+		if (Ar.CustomVer(FCoreObjectVersion::GUID) >= FCoreObjectVersion::SkeletalMaterialEditorDataStripping)
+		{
+			Ar << bSerializeImportedMaterialSlotName;
+		}
+		else if (!FPlatformProperties::HasEditorOnlyData())
+		{
+			bSerializeImportedMaterialSlotName = false;
+		}
+		if (bSerializeImportedMaterialSlotName)
+		{
+#if WITH_EDITORONLY_DATA
+			Ar << Elem.ImportedMaterialSlotName;
+#else
+			FName UnusedImportedMaterialSlotName;
+			Ar << UnusedImportedMaterialSlotName;
+#endif
+		}
+	}
+#if WITH_EDITORONLY_DATA
+	else
+	{
+		if (Ar.UEVer() >= VER_UE4_MOVE_SKELETALMESH_SHADOWCASTING)
+		{
+			Ar << Elem.bEnableShadowCasting_DEPRECATED;
+		}
+
+		Ar.UsingCustomVersion(FRecomputeTangentCustomVersion::GUID);
+		if (Ar.CustomVer(FRecomputeTangentCustomVersion::GUID) >= FRecomputeTangentCustomVersion::RuntimeRecomputeTangent)
+		{
+			Ar << Elem.bRecomputeTangent_DEPRECATED;
+		}
+	}
+#endif
+	
+	if (!Ar.IsLoading() || Ar.CustomVer(FRenderingObjectVersion::GUID) >= FRenderingObjectVersion::TextureStreamingMeshUVChannelData)
+	{
+		Ar << Elem.UVChannelData;
+	}
+
+	return Ar;
+}
+
 TArray<USkeletalMeshSocket*> USkeletalMesh::GetActiveSocketList() const
 {
 	TArray<USkeletalMeshSocket*> ActiveSockets = Sockets;
diff --git a/Engine/Source/Runtime/Engine/Public/MeshUVChannelInfo.h b/Engine/Source/Runtime/Engine/Public/MeshUVChannelInfo.h
--- a/Engine/Source/Runtime/Engine/Public/MeshUVChannelInfo.h
+++ b/Engine/Source/Runtime/Engine/Public/MeshUVChannelInfo.h
@@ -57,5 +57,5 @@ struct FMeshUVChannelInfo
 	UPROPERTY(EditAnywhere, AdvancedDisplay, Category = Basic, meta = (EditCondition = "bOverrideDensities"))
 	float LocalUVDensities[MAX_TEXCOORDS];
 
-	friend FArchive& operator<<(FArchive& Ar, FMeshUVChannelInfo& Info);
+	friend ENGINE_API FArchive& operator<<(FArchive& Ar, FMeshUVChannelInfo& Info);
 };
