Fixed actor world partition properties not properly set to readonly.

diff --git a/Engine/Source/Runtime/Engine/Private/ActorEditor.cpp b/Engine/Source/Runtime/Engine/Private/ActorEditor.cpp
--- a/Engine/Source/Runtime/Engine/Private/ActorEditor.cpp
+++ b/Engine/Source/Runtime/Engine/Private/ActorEditor.cpp
@@ -88,14 +88,12 @@ bool AActor::CanEditChange(const FProperty* PropertyThatWillChange) const
 	{
 		if (!IsTemplate())
 		{
-			if (UWorld* World = GetTypedOuter<UWorld>())
+			UWorld* World = GetTypedOuter<UWorld>();
+			UWorldPartition* WorldPartition = World ? World->GetWorldPartition() : nullptr;
+			if (!WorldPartition || (!WorldPartition->IsStreamingEnabled() && !bIsDataLayersProperty))
 			{
-				if (UWorldPartition* WorldPartition = World->GetWorldPartition())
-				{
-					return bIsDataLayersProperty || WorldPartition->IsStreamingEnabled();
-				}
+				return false;
 			}
-			return false;
 		}
 	}
 
