Content Bundle: fix some crashes encountered while playing with the feature.

#rnx

diff --git a/Engine/Source/Runtime/Engine/Private/WorldPartition/ContentBundle/ContentBundleEditor.cpp b/Engine/Source/Runtime/Engine/Private/WorldPartition/ContentBundle/ContentBundleEditor.cpp
--- a/Engine/Source/Runtime/Engine/Private/WorldPartition/ContentBundle/ContentBundleEditor.cpp
+++ b/Engine/Source/Runtime/Engine/Private/WorldPartition/ContentBundle/ContentBundleEditor.cpp
@@ -189,22 +189,28 @@ bool FContentBundleEditor::GetActors(TArray<AActor*>& Actors)
 {
 	Actors.Reserve(GetActorCount());
 
-	for (FActorDescList::TIterator<> It(ActorDescContainer.Get()); It; ++It)
+	if (ActorDescContainer.IsValid())
 	{
-		if (AActor* Actor = It->GetActor())
+		for (FActorDescList::TIterator<> It(ActorDescContainer.Get()); It; ++It)
 		{
-			if (Actor != WorldDataLayersActorReference.Get())
+			if (AActor* Actor = It->GetActor())
 			{
-				Actors.Add(Actor);
+				if (Actor != WorldDataLayersActorReference.Get())
+				{
+					Actors.Add(Actor);
+				}
 			}
 		}
 	}
 
-	for (auto UnsavedActor : UnsavedActorMonitor->GetUnsavedActors())
+	if (UnsavedActorMonitor)
 	{
-		if (UnsavedActor.IsValid())
+		for (auto UnsavedActor : UnsavedActorMonitor->GetUnsavedActors())
 		{
-			Actors.Add(UnsavedActor.Get());
+			if (UnsavedActor.IsValid())
+			{
+				Actors.Add(UnsavedActor.Get());
+			}
 		}
 	}
 
