WorldPartitionMinimapVolume: Set to NoCollision preset

#rnx

diff --git a/Engine/Source/Runtime/Engine/Private/WorldPartition/WorldPartitionMiniMapVolume.cpp b/Engine/Source/Runtime/Engine/Private/WorldPartition/WorldPartitionMiniMapVolume.cpp
--- a/Engine/Source/Runtime/Engine/Private/WorldPartition/WorldPartitionMiniMapVolume.cpp
+++ b/Engine/Source/Runtime/Engine/Private/WorldPartition/WorldPartitionMiniMapVolume.cpp
@@ -1,6 +1,8 @@
 // Copyright Epic Games, Inc. All Rights Reserved.
 
 #include "WorldPartition/WorldPartitionMiniMapVolume.h"
+#include "Engine/CollisionProfile.h"
+#include "Components/BrushComponent.h"
 
 #include UE_INLINE_GENERATED_CPP_BY_NAME(WorldPartitionMiniMapVolume)
 
@@ -10,4 +12,6 @@ AWorldPartitionMiniMapVolume::AWorldPartitionMiniMapVolume(const FObjectInitiali
 #if WITH_EDITORONLY_DATA
 	bIsSpatiallyLoaded = false;
 #endif
+
+	GetBrushComponent()->SetCollisionProfileName(UCollisionProfile::NoCollision_ProfileName);
 }
