Fix player going into A pose after exiting riding or vehicles in cooked build
Add missing layer node initialization code when linking to self-instance

diff --git a/Engine/Source/Runtime/Engine/Private/Animation/AnimInstance.cpp b/Engine/Source/Runtime/Engine/Private/Animation/AnimInstance.cpp
--- a/Engine/Source/Runtime/Engine/Private/Animation/AnimInstance.cpp
+++ b/Engine/Source/Runtime/Engine/Private/Animation/AnimInstance.cpp
@@ -3076,6 +3076,17 @@ void UAnimInstance::PerformLinkedLayerOverlayOperation(TSubclassOf<UAnimInstance
 						else
 						{
 							LayerNode->SetLinkedLayerInstance(this, nullptr);
+							
+							if (!bInDeferSubGraphInitialization)
+							{
+								UAnimInstance* LinkedInstance = LayerNode->GetTargetInstance<UAnimInstance>();
+								if (LayerNode->LinkedRoot && LinkedInstance)
+								{
+									FAnimInstanceProxy& ThisProxy = GetProxyOnAnyThread<FAnimInstanceProxy>();
+									FAnimInstanceProxy& LinkedProxy = LinkedInstance->GetProxyOnAnyThread<FAnimInstanceProxy>();
+									InitializeAndCacheBonesForLinkedRoot(LayerNode, ThisProxy, LinkedInstance, LinkedProxy);
+								}
+							}
 						}
 					}
 				}
