Fix NEEDS_INSTANCE_LOCAL_TO_WORLD_PS being overwritten by a vertex property after a pixel property enabled it. This makes it consistent with other defines, e.g. NEEDS_PARTICLE_COLOR.

diff --git a/Engine/Source/Runtime/Engine/Private/Materials/HLSLMaterialTranslator.cpp b/Engine/Source/Runtime/Engine/Private/Materials/HLSLMaterialTranslator.cpp
--- a/Engine/Source/Runtime/Engine/Private/Materials/HLSLMaterialTranslator.cpp
+++ b/Engine/Source/Runtime/Engine/Private/Materials/HLSLMaterialTranslator.cpp
@@ -8545,7 +8545,7 @@ int32 FHLSLMaterialTranslator::TransformBase(EMaterialCommonBasis SourceCoordBas
 			{
 				CodeStr = LWCMultiplyMatrix(TEXT("<A>"), TEXT("GetWorldToInstance(Parameters)"), AWComponent);
 				CodeDerivStr = LWCMultiplyMatrix(TEXT("<A>"), TEXT("GetWorldToInstance(Parameters)"), 0);
-				bUsesInstanceWorldToLocalPS = ShaderFrequency == SF_Pixel;
+				bUsesInstanceWorldToLocalPS |= ShaderFrequency == SF_Pixel;
 			}
 
 			// else use MCB_TranslatedWorld as intermediary basis
@@ -8591,7 +8591,7 @@ int32 FHLSLMaterialTranslator::TransformBase(EMaterialCommonBasis SourceCoordBas
 			{
 				CodeStr = LWCMultiplyMatrix(TEXT("<A>"), TEXT("GetInstanceToWorld(Parameters)"), AWComponent);
 				CodeDerivStr = LWCMultiplyMatrix(TEXT("<A>"), TEXT("GetInstanceToWorld(Parameters)"), 0);
-				bUsesInstanceLocalToWorldPS = ShaderFrequency == SF_Pixel;
+				bUsesInstanceLocalToWorldPS |= ShaderFrequency == SF_Pixel;
 			}
 			// use World as an intermediary base
 			break;
