Added check for GraphNode not to be null in UMaterialExpression::PostEditChangeProperty() which prevents a rare crash from occurring.

diff --git a/Engine/Source/Runtime/Engine/Private/Materials/MaterialExpressions.cpp b/Engine/Source/Runtime/Engine/Private/Materials/MaterialExpressions.cpp
--- a/Engine/Source/Runtime/Engine/Private/Materials/MaterialExpressions.cpp
+++ b/Engine/Source/Runtime/Engine/Private/Materials/MaterialExpressions.cpp
@@ -1532,7 +1532,7 @@ void UMaterialExpression::PostEditChangeProperty(FPropertyChangedEvent& Property
 				for (UMaterialExpressionParameter* ExpressionParameter : ParameterExpressions)
 				{
 					// If the other expression type and parameter name are the same as this expression's...
-					if (ExpressionParameter->GetArchetype() == GetArchetype() && ExpressionParameter->ParameterName == ParameterName)
+					if (ExpressionParameter->GraphNode && ExpressionParameter->GetArchetype() == GetArchetype() && ExpressionParameter->ParameterName == ParameterName)
 					{
 						// ...modify the pin on other parameter expression node with the new value.
 						UEdGraphPin* OtherPin = ExpressionParameter->GraphNode->GetPinAt(PinIndex);
