[shaders] re-fix shader debug info collisions in a way that doesn't rely on being run on the game thread (hashing the shader map portions of the DDC key instead, which also contain static parameter values)

diff --git a/Engine/Source/Runtime/Engine/Private/Materials/MaterialShader.cpp b/Engine/Source/Runtime/Engine/Private/Materials/MaterialShader.cpp
--- a/Engine/Source/Runtime/Engine/Private/Materials/MaterialShader.cpp
+++ b/Engine/Source/Runtime/Engine/Private/Materials/MaterialShader.cpp
@@ -1123,6 +1123,7 @@ void FMaterialShaderMapId::SetShaderDependencies(const TArray<FShaderType*>& Sha
 static void PrepareMaterialShaderCompileJob(EShaderPlatform Platform,
 	EShaderPermutationFlags PermutationFlags,
 	const FMaterial* Material,
+	const FMaterialShaderMapId& ShaderMapId,
 	FSharedShaderCompilerEnvironment* MaterialEnvironment,
 	const FShaderPipelineType* ShaderPipeline,
 	FString DebugDescription,
@@ -1148,7 +1149,7 @@ static void PrepareMaterialShaderCompileJob(EShaderPlatform Platform,
 
 	// Compile the shader environment passed in with the shader type's source code.
 	::GlobalBeginCompileShader(
-		Material->GetUniqueAssetName(Platform) / LexToString(Material->GetQualityLevel()),
+		Material->GetUniqueAssetName(Platform, ShaderMapId) / LexToString(Material->GetQualityLevel()),
 		nullptr,
 		ShaderType,
 		ShaderPipeline,
@@ -1169,9 +1170,10 @@ static void PrepareMaterialShaderCompileJob(EShaderPlatform Platform,
  */
 void FMaterialShaderType::BeginCompileShader(
 	EShaderCompileJobPriority Priority,
-	uint32 ShaderMapId,
+	uint32 ShaderMapJobId,
 	int32 PermutationId,
 	const FMaterial* Material,
+	const FMaterialShaderMapId& ShaderMapId,
 	FSharedShaderCompilerEnvironment* MaterialEnvironment,
 	EShaderPlatform Platform,
 	EShaderPermutationFlags PermutationFlags,
@@ -1180,20 +1182,21 @@ void FMaterialShaderType::BeginCompileShader(
 	const TCHAR* DebugExtension
 	) const
 {
-	FShaderCompileJob* NewJob = GShaderCompilingManager->PrepareShaderCompileJob(ShaderMapId, FShaderCompileJobKey(this, nullptr, PermutationId), Priority);
+	FShaderCompileJob* NewJob = GShaderCompilingManager->PrepareShaderCompileJob(ShaderMapJobId, FShaderCompileJobKey(this, nullptr, PermutationId), Priority);
 	if (NewJob)
 	{
-		PrepareMaterialShaderCompileJob(Platform, PermutationFlags, Material, MaterialEnvironment, nullptr, DebugDescription, DebugExtension, NewJob);
+		PrepareMaterialShaderCompileJob(Platform, PermutationFlags, Material, ShaderMapId, MaterialEnvironment, nullptr, DebugDescription, DebugExtension, NewJob);
 		NewJobs.Add(FShaderCommonCompileJobPtr(NewJob));
 	}
 }
 
 void FMaterialShaderType::BeginCompileShaderPipeline(
 	EShaderCompileJobPriority Priority,
-	uint32 ShaderMapId,
+	uint32 ShaderMapJobId,
 	EShaderPlatform Platform,
 	EShaderPermutationFlags PermutationFlags,
 	const FMaterial* Material,
+	const FMaterialShaderMapId& ShaderMapId,
 	FSharedShaderCompilerEnvironment* MaterialEnvironment,
 	const FShaderPipelineType* ShaderPipeline,
 	TArray<FShaderCommonCompileJobPtr>& NewJobs,
@@ -1204,12 +1207,12 @@ void FMaterialShaderType::BeginCompileShaderPipeline(
 	UE_LOG(LogShaders, Verbose, TEXT("	Pipeline: %s"), ShaderPipeline->GetName());
 
 	// Add all the jobs as individual first, then add the dependencies into a pipeline job
-	auto* NewPipelineJob = GShaderCompilingManager->PreparePipelineCompileJob(ShaderMapId, FShaderPipelineCompileJobKey(ShaderPipeline, nullptr, kUniqueShaderPermutationId), Priority);
+	auto* NewPipelineJob = GShaderCompilingManager->PreparePipelineCompileJob(ShaderMapJobId, FShaderPipelineCompileJobKey(ShaderPipeline, nullptr, kUniqueShaderPermutationId), Priority);
 	if (NewPipelineJob)
 	{
 		for (FShaderCompileJob* StageJob : NewPipelineJob->StageJobs)
 		{
-			PrepareMaterialShaderCompileJob(Platform, PermutationFlags, Material, MaterialEnvironment, ShaderPipeline, DebugDescription, DebugExtension, StageJob);
+			PrepareMaterialShaderCompileJob(Platform, PermutationFlags, Material, ShaderMapId, MaterialEnvironment, ShaderPipeline, DebugDescription, DebugExtension, StageJob);
 		}
 		NewJobs.Add(FShaderCommonCompileJobPtr(NewPipelineJob));
 	}
@@ -1814,6 +1817,7 @@ void FMaterialShaderMap::SubmitCompileJobs(uint32 CompilingShaderMapId,
 					ShaderPlatform,
 					LocalPermutationFlags,
 					Material,
+					ShaderMapId,
 					MaterialEnvironment,
 					MeshLayout.VertexFactoryType,
 					CompileJobs,
@@ -1857,6 +1861,7 @@ void FMaterialShaderMap::SubmitCompileJobs(uint32 CompilingShaderMapId,
 					ShaderPlatform,
 					LocalPermutationFlags,
 					Material,
+					ShaderMapId,
 					MaterialEnvironment,
 					MeshLayout.VertexFactoryType,
 					Pipeline,
@@ -1915,6 +1920,7 @@ void FMaterialShaderMap::SubmitCompileJobs(uint32 CompilingShaderMapId,
 				CompilingShaderMapId,
 				Shader.PermutationId,
 				Material,
+				ShaderMapId,
 				MaterialEnvironment,
 				ShaderPlatform,
 				LocalPermutationFlags,
@@ -1952,6 +1958,7 @@ void FMaterialShaderMap::SubmitCompileJobs(uint32 CompilingShaderMapId,
 					ShaderPlatform,
 					LocalPermutationFlags,
 					Material,
+					ShaderMapId,
 					MaterialEnvironment,
 					Pipeline,
 					CompileJobs,
diff --git a/Engine/Source/Runtime/Engine/Private/Materials/MaterialShared.cpp b/Engine/Source/Runtime/Engine/Private/Materials/MaterialShared.cpp
--- a/Engine/Source/Runtime/Engine/Private/Materials/MaterialShared.cpp
+++ b/Engine/Source/Runtime/Engine/Private/Materials/MaterialShared.cpp
@@ -2680,23 +2680,12 @@ bool FMaterial::CacheShaders(const FMaterialShaderMapId& ShaderMapId, EShaderPla
 }
 
 #if WITH_EDITOR
-FString FMaterial::GetUniqueAssetName(EShaderPlatform Platform) const
-{
-	// Temporary disabled due to UE-169098:
-	// This funciton may be called from worker threads, but GetStaticParameterSet() requires game thread.
-#if 0 
-	TArray<uint8> Buf;
-	FString PathStr = GetAssetPath().GetPlainNameString();
-	FMemoryWriter BufWriter(Buf);
-	BufWriter << PathStr;
-	FStaticParameterSet StaticParams;
-	GetStaticParameterSet(Platform, StaticParams);
-	FStaticParameterSet::StaticStruct()->SerializeBin(BufWriter, &StaticParams);
-	uint64 Hash = CityHash64((const char*)Buf.GetData(), Buf.Num());
-	return FString::Printf(TEXT("%s_%llu"), *GetAssetName(), Hash);
-#else
-	return GetAssetName();
-#endif
+FString FMaterial::GetUniqueAssetName(EShaderPlatform Platform, const FMaterialShaderMapId& ShaderMapId) const
+{
+	FString ToHash = GetAssetPath().GetPlainNameString();
+	ShaderMapId.AppendKeyString(ToHash);
+	uint64 Hash = CityHash64((const char*)*ToHash, ToHash.Len());
+	return FString::Printf(TEXT("%s_%llx"), *GetAssetName(), Hash);
 }
 #endif // WITH_EDITOR
 
@@ -2941,6 +2930,7 @@ void FMaterial::CacheGivenTypes(EShaderPlatform Platform, const TArray<const FVe
 					Platform,
 					GameThreadShaderMap->GetPermutationFlags(),
 					this,
+					GameThreadShaderMap->GetShaderMapId(),
 					GameThreadPendingCompilerEnvironment,
 					VFType,
 					PipelineType,
@@ -2955,6 +2945,7 @@ void FMaterial::CacheGivenTypes(EShaderPlatform Platform, const TArray<const FVe
 					GetGameThreadCompilingShaderMapId(),
 					0,
 					this,
+					GameThreadShaderMap->GetShaderMapId(),
 					GameThreadPendingCompilerEnvironment,
 					Platform,
 					GameThreadShaderMap->GetPermutationFlags(),
@@ -2971,6 +2962,7 @@ void FMaterial::CacheGivenTypes(EShaderPlatform Platform, const TArray<const FVe
 					Platform,
 					GameThreadShaderMap->GetPermutationFlags(),
 					this,
+					GameThreadShaderMap->GetShaderMapId(),
 					GameThreadPendingCompilerEnvironment,
 					VFType,
 					CompileJobs,
@@ -3362,11 +3354,35 @@ bool FMaterial::TryGetShaders(const FMaterialShaderTypes& InTypes, const FVertex
 						{
 							if (InVertexFactoryType)
 							{
-								FMeshMaterialShaderType::BeginCompileShaderPipeline(EShaderCompileJobPriority::ForceLocal, CompilingShaderMapId, kUniqueShaderPermutationId, ShaderPlatform, PermutationFlags, this, RenderingThreadPendingCompilerEnvironment, InVertexFactoryType, InTypes.PipelineType, CompileJobs, nullptr, nullptr);
+								FMeshMaterialShaderType::BeginCompileShaderPipeline(
+									EShaderCompileJobPriority::ForceLocal, 
+									CompilingShaderMapId, 
+									kUniqueShaderPermutationId, 
+									ShaderPlatform, 
+									PermutationFlags, 
+									this,
+									ShaderMap->GetShaderMapId(),
+									RenderingThreadPendingCompilerEnvironment, 
+									InVertexFactoryType, 
+									InTypes.PipelineType, 
+									CompileJobs, 
+									nullptr, 
+									nullptr);
 							}
 							else
 							{
-								FMaterialShaderType::BeginCompileShaderPipeline(EShaderCompileJobPriority::ForceLocal, CompilingShaderMapId, ShaderPlatform, PermutationFlags, this, RenderingThreadPendingCompilerEnvironment, InTypes.PipelineType, CompileJobs, nullptr, nullptr);
+								FMaterialShaderType::BeginCompileShaderPipeline(
+									EShaderCompileJobPriority::ForceLocal, 
+									CompilingShaderMapId, 
+									ShaderPlatform, 
+									PermutationFlags, 
+									this,
+									ShaderMap->GetShaderMapId(),
+									RenderingThreadPendingCompilerEnvironment, 
+									InTypes.PipelineType, 
+									CompileJobs, 
+									nullptr, 
+									nullptr);
 							}
 						}
 					}
@@ -3417,11 +3433,34 @@ bool FMaterial::TryGetShaders(const FMaterialShaderTypes& InTypes, const FVertex
 							{
 								if (InVertexFactoryType)
 								{
-									ShaderType->AsMeshMaterialShaderType()->BeginCompileShader(EShaderCompileJobPriority::ForceLocal, CompilingShaderMapId, PermutationId, ShaderPlatform, PermutationFlags, this, RenderingThreadPendingCompilerEnvironment, InVertexFactoryType, CompileJobs, nullptr, nullptr);
+									ShaderType->AsMeshMaterialShaderType()->BeginCompileShader(
+										EShaderCompileJobPriority::ForceLocal, 
+										CompilingShaderMapId, 
+										PermutationId, 
+										ShaderPlatform, 
+										PermutationFlags, 
+										this,
+										ShaderMap->GetShaderMapId(),
+										RenderingThreadPendingCompilerEnvironment, 
+										InVertexFactoryType, 
+										CompileJobs, 
+										nullptr, 
+										nullptr);
 								}
 								else
 								{
-									ShaderType->AsMaterialShaderType()->BeginCompileShader(EShaderCompileJobPriority::ForceLocal, CompilingShaderMapId, PermutationId, this, RenderingThreadPendingCompilerEnvironment, ShaderPlatform, PermutationFlags, CompileJobs, nullptr, nullptr);
+									ShaderType->AsMaterialShaderType()->BeginCompileShader(
+										EShaderCompileJobPriority::ForceLocal, 
+										CompilingShaderMapId, 
+										PermutationId, 
+										this,
+										ShaderMap->GetShaderMapId(),
+										RenderingThreadPendingCompilerEnvironment, 
+										ShaderPlatform, 
+										PermutationFlags, 
+										CompileJobs, 
+										nullptr, 
+										nullptr);
 								}
 							}
 						}
diff --git a/Engine/Source/Runtime/Engine/Private/Materials/MeshMaterialShader.cpp b/Engine/Source/Runtime/Engine/Private/Materials/MeshMaterialShader.cpp
--- a/Engine/Source/Runtime/Engine/Private/Materials/MeshMaterialShader.cpp
+++ b/Engine/Source/Runtime/Engine/Private/Materials/MeshMaterialShader.cpp
@@ -27,6 +27,7 @@ namespace MaterialMeshCookStats
 static void PrepareMeshMaterialShaderCompileJob(EShaderPlatform Platform,
 	EShaderPermutationFlags PermutationFlags,
 	const FMaterial* Material,
+	const FMaterialShaderMapId& MaterialShaderMapId,
 	FSharedShaderCompilerEnvironment* MaterialEnvironment,
 	const FShaderPipelineType* ShaderPipeline,
 	const TCHAR* DebugDescription,
@@ -64,7 +65,7 @@ static void PrepareMeshMaterialShaderCompileJob(EShaderPlatform Platform,
 
 	// Compile the shader environment passed in with the shader type's source code.
 	::GlobalBeginCompileShader(
-		Material->GetUniqueAssetName(Platform) / LexToString(Material->GetQualityLevel()),
+		Material->GetUniqueAssetName(Platform, MaterialShaderMapId) / LexToString(Material->GetQualityLevel()),
 		VertexFactoryType,
 		ShaderType,
 		ShaderPipeline,
@@ -88,32 +89,34 @@ static void PrepareMeshMaterialShaderCompileJob(EShaderPlatform Platform,
  */
 void FMeshMaterialShaderType::BeginCompileShader(
 	EShaderCompileJobPriority Priority,
-	uint32 ShaderMapId,
+	uint32 ShaderMapJobId,
 	int32 PermutationId,
 	EShaderPlatform Platform,
 	EShaderPermutationFlags PermutationFlags,
 	const FMaterial* Material,
+	const FMaterialShaderMapId& ShaderMapId,
 	FSharedShaderCompilerEnvironment* MaterialEnvironment,
 	const FVertexFactoryType* VertexFactoryType,
 	TArray<FShaderCommonCompileJobPtr>& NewJobs,
 	const TCHAR* DebugDescription,
 	const TCHAR* DebugExtension) const
 {
-	FShaderCompileJob* NewJob = GShaderCompilingManager->PrepareShaderCompileJob(ShaderMapId, FShaderCompileJobKey(this, VertexFactoryType, PermutationId), Priority);
+	FShaderCompileJob* NewJob = GShaderCompilingManager->PrepareShaderCompileJob(ShaderMapJobId, FShaderCompileJobKey(this, VertexFactoryType, PermutationId), Priority);
 	if (NewJob)
 	{
-		PrepareMeshMaterialShaderCompileJob(Platform, PermutationFlags, Material, MaterialEnvironment, nullptr, DebugDescription, DebugExtension, NewJob);
+		PrepareMeshMaterialShaderCompileJob(Platform, PermutationFlags, Material, ShaderMapId, MaterialEnvironment, nullptr, DebugDescription, DebugExtension, NewJob);
 		NewJobs.Add(FShaderCommonCompileJobPtr(NewJob));
 	}
 }
 
 void FMeshMaterialShaderType::BeginCompileShaderPipeline(
 	EShaderCompileJobPriority Priority,
-	uint32 ShaderMapId,
+	uint32 ShaderMapJobId,
 	int32 PermutationId,
 	EShaderPlatform Platform,
 	EShaderPermutationFlags PermutationFlags,
 	const FMaterial* Material,
+	const FMaterialShaderMapId& ShaderMapId,
 	FSharedShaderCompilerEnvironment* MaterialEnvironment,
 	const FVertexFactoryType* VertexFactoryType,
 	const FShaderPipelineType* ShaderPipeline,
@@ -125,12 +128,12 @@ void FMeshMaterialShaderType::BeginCompileShaderPipeline(
 	UE_LOG(LogShaders, Verbose, TEXT("	Pipeline: %s"), ShaderPipeline->GetName());
 
 	// Add all the jobs as individual first, then add the dependencies into a pipeline job
-	auto* NewPipelineJob = GShaderCompilingManager->PreparePipelineCompileJob(ShaderMapId, FShaderPipelineCompileJobKey(ShaderPipeline, VertexFactoryType, PermutationId), Priority);
+	auto* NewPipelineJob = GShaderCompilingManager->PreparePipelineCompileJob(ShaderMapJobId, FShaderPipelineCompileJobKey(ShaderPipeline, VertexFactoryType, PermutationId), Priority);
 	if (NewPipelineJob)
 	{
 		for (FShaderCompileJob* StageJob : NewPipelineJob->StageJobs)
 		{
-			PrepareMeshMaterialShaderCompileJob(Platform, PermutationFlags, Material, MaterialEnvironment, ShaderPipeline, DebugDescription, DebugExtension, StageJob);
+			PrepareMeshMaterialShaderCompileJob(Platform, PermutationFlags, Material, ShaderMapId, MaterialEnvironment, ShaderPipeline, DebugDescription, DebugExtension, StageJob);
 		}
 		NewJobs.Add(FShaderCommonCompileJobPtr(NewPipelineJob));
 	}
diff --git a/Engine/Source/Runtime/Engine/Public/MaterialShaderType.h b/Engine/Source/Runtime/Engine/Public/MaterialShaderType.h
--- a/Engine/Source/Runtime/Engine/Public/MaterialShaderType.h
+++ b/Engine/Source/Runtime/Engine/Public/MaterialShaderType.h
@@ -23,6 +23,7 @@
 
 class FMaterial;
 class FMaterialShaderMap;
+class FMaterialShaderMapId;
 class FShaderCommonCompileJob;
 class FShaderCompileJob;
 class FUniformExpressionSet;
@@ -132,9 +133,10 @@ public:
 	 */
 	void BeginCompileShader(
 		EShaderCompileJobPriority Priority,
-		uint32 ShaderMapId,
+		uint32 ShaderMapJobId,
 		int32 PermutationId,
 		const FMaterial* Material,
+		const FMaterialShaderMapId& ShaderMapId,
 		FSharedShaderCompilerEnvironment* MaterialEnvironment,
 		EShaderPlatform Platform,
 		EShaderPermutationFlags PermutationFlags,
@@ -145,10 +147,11 @@ public:
 
 	static void BeginCompileShaderPipeline(
 		EShaderCompileJobPriority Priority,
-		uint32 ShaderMapId,
+		uint32 ShaderMapJobId,
 		EShaderPlatform Platform,
 		EShaderPermutationFlags PermutationFlags,
 		const FMaterial* Material,
+		const FMaterialShaderMapId& ShaderMapId,
 		FSharedShaderCompilerEnvironment* MaterialEnvironment,
 		const FShaderPipelineType* ShaderPipeline,
 		TArray<TRefCountPtr<FShaderCommonCompileJob>>& NewJobs,
diff --git a/Engine/Source/Runtime/Engine/Public/MaterialShared.h b/Engine/Source/Runtime/Engine/Public/MaterialShared.h
--- a/Engine/Source/Runtime/Engine/Public/MaterialShared.h
+++ b/Engine/Source/Runtime/Engine/Public/MaterialShared.h
@@ -1864,10 +1864,10 @@ public:
 #if WITH_EDITOR
 	/** Similar to GetAssetName, but appends a hash of the asset path and static switch params (if editoronly data is available)
 	to differentiate materials that would collide when considering name alone. */
-	ENGINE_API FString GetUniqueAssetName(EShaderPlatform Platform) const;
+	ENGINE_API FString GetUniqueAssetName(EShaderPlatform Platform, const FMaterialShaderMapId& ShaderMapId) const;
 #else
 	/** GetUniqueAssetName just returns GetAssetName in non-editor builds */
-	FString GetUniqueAssetName(EShaderPlatform Platform) const { return GetAssetName(); }
+	FString GetUniqueAssetName(EShaderPlatform Platform, const FMaterialShaderMapId& ShaderMapId) const { return GetAssetName(); }
 #endif
 	virtual bool HasVertexPositionOffsetConnected() const { return false; }
 	virtual bool HasPixelDepthOffsetConnected() const { return false; }
diff --git a/Engine/Source/Runtime/Engine/Public/MeshMaterialShaderType.h b/Engine/Source/Runtime/Engine/Public/MeshMaterialShaderType.h
--- a/Engine/Source/Runtime/Engine/Public/MeshMaterialShaderType.h
+++ b/Engine/Source/Runtime/Engine/Public/MeshMaterialShaderType.h
@@ -11,6 +11,7 @@
 #include "MaterialShaderType.h"
 
 class FMaterial;
+class FMaterialShaderMapId;
 class FShaderCommonCompileJob;
 class FShaderCompileJob;
 class FUniformExpressionSet;
@@ -87,11 +88,12 @@ public:
 	 */
 	void BeginCompileShader(
 		EShaderCompileJobPriority Priority,
-		uint32 ShaderMapId,
+		uint32 ShaderMapJobId,
 		int32 PermutationId,
 		EShaderPlatform Platform,
 		EShaderPermutationFlags PermutationFlags,
 		const FMaterial* Material,
+		const FMaterialShaderMapId& ShaderMapId,
 		FSharedShaderCompilerEnvironment* MaterialEnvironment,
 		const FVertexFactoryType* VertexFactoryType,
 		TArray<TRefCountPtr<FShaderCommonCompileJob>>& NewJobs,
@@ -101,11 +103,12 @@ public:
 
 	static void BeginCompileShaderPipeline(
 		EShaderCompileJobPriority Priority,
-		uint32 ShaderMapId,
+		uint32 ShaderMapJobId,
 		int32 PermutationId,
 		EShaderPlatform Platform,
 		EShaderPermutationFlags PermutationFlags,
 		const FMaterial* Material,
+		const FMaterialShaderMapId& ShaderMapId,
 		FSharedShaderCompilerEnvironment* MaterialEnvironment,
 		const FVertexFactoryType* VertexFactoryType,
 		const FShaderPipelineType* ShaderPipeline,
