Hot-fix crash when user force deletes an asset (e.g. a material) used by a mesh with an instance in the active level.

This is a workaround to alleviate the crash. Meanwhile we are looking for a proper fix.

diff --git a/Engine/Source/Runtime/Engine/Private/Components/PrimitiveComponent.cpp b/Engine/Source/Runtime/Engine/Private/Components/PrimitiveComponent.cpp
--- a/Engine/Source/Runtime/Engine/Private/Components/PrimitiveComponent.cpp
+++ b/Engine/Source/Runtime/Engine/Private/Components/PrimitiveComponent.cpp
@@ -591,7 +591,13 @@ void UPrimitiveComponent::CreateRenderState_Concurrent(FRegisterComponentContext
 	UpdateBounds();
 
 	// If the primitive isn't hidden and the detail mode setting allows it, add it to the scene.
-	if (ShouldComponentAddToScene())
+	if (ShouldComponentAddToScene()
+#ifdef WITH_EDITOR
+		// [HOTFIX] When force deleting an asset, a SceneProxy is set to null from a different thread unsafely, causing the old stale value of SceneProxy being read here from the cache.
+		// We need to better investigate why this happens, but for now this prevents a crash from occurring.
+		&& SceneProxy == nullptr
+#endif
+	)
 	{
 		if (Context != nullptr)
 		{
