Wrap scope transactions around WorldParition EnableStreaming and CanBeUsedByLevelInstance properties so that they can be captured by the transactions system.

diff --git a/Engine/Source/Runtime/Engine/Private/WorldPartition/WorldPartition.cpp b/Engine/Source/Runtime/Engine/Private/WorldPartition/WorldPartition.cpp
--- a/Engine/Source/Runtime/Engine/Private/WorldPartition/WorldPartition.cpp
+++ b/Engine/Source/Runtime/Engine/Private/WorldPartition/WorldPartition.cpp
@@ -1096,6 +1096,8 @@ void UWorldPartition::SetEnableStreaming(bool bInEnableStreaming)
 {
 	if (bEnableStreaming != bInEnableStreaming)
 	{
+		const FScopedTransaction Transaction(LOCTEXT("EditorWorldPartitionSetEnableStreaming", "Set WorldPartition EnableStreaming"));
+
 		Modify();
 		bEnableStreaming = bInEnableStreaming;
 		OnEnableStreamingChanged();
@@ -1114,6 +1116,8 @@ void UWorldPartition::SetCanBeUsedByLevelInstance(bool bInCanBeUsedByLevelInstan
 
 	if (bCanBeUsedByLevelInstance != bInCanBeUsedByLevelInstance)
 	{
+		const FScopedTransaction Transaction(LOCTEXT("EditorWorldPartitionCanBeUsedByLevelInstance", "Set WorldPartition CanBeUsedByLevelInstance"));
+
 		Modify();
 		bCanBeUsedByLevelInstance = bInCanBeUsedByLevelInstance;
 		if (bCanBeUsedByLevelInstance)
