World Partition: DisableStreamingIn()/EnableStreamingIn() allow temporarly disabling new cells from being loaded/activated

diff --git a/Engine/Source/Runtime/Engine/Private/WorldPartition/WorldPartition.cpp b/Engine/Source/Runtime/Engine/Private/WorldPartition/WorldPartition.cpp
--- a/Engine/Source/Runtime/Engine/Private/WorldPartition/WorldPartition.cpp
+++ b/Engine/Source/Runtime/Engine/Private/WorldPartition/WorldPartition.cpp
@@ -292,6 +292,7 @@ UWorldPartition::UWorldPartition(const FObjectInitializer& ObjectInitializer)
 	, NumUserCreatedLoadedRegions(0)
 #endif
 	, InitState(EWorldPartitionInitState::Uninitialized)
+	, bStreamingInEnabled(true)
 	, StreamingPolicy(nullptr)
 	, Replay(nullptr)
 {
@@ -1457,6 +1458,23 @@ EWorldPartitionStreamingPerformance UWorldPartition::GetStreamingPerformance() c
 	return StreamingPolicy->GetStreamingPerformance();
 }
 
+bool UWorldPartition::IsStreamingInEnabled() const
+{
+	return bStreamingInEnabled;
+}
+
+void UWorldPartition::DisableStreamingIn()
+{
+	check(bStreamingInEnabled);
+	bStreamingInEnabled = false;
+}
+
+void UWorldPartition::EnableStreamingIn()
+{
+	check(!bStreamingInEnabled);
+	bStreamingInEnabled = true;
+}
+
 #if WITH_EDITOR
 void UWorldPartition::DrawRuntimeHashPreview()
 {
diff --git a/Engine/Source/Runtime/Engine/Private/WorldPartition/WorldPartitionStreamingPolicy.cpp b/Engine/Source/Runtime/Engine/Private/WorldPartition/WorldPartitionStreamingPolicy.cpp
--- a/Engine/Source/Runtime/Engine/Private/WorldPartition/WorldPartitionStreamingPolicy.cpp
+++ b/Engine/Source/Runtime/Engine/Private/WorldPartition/WorldPartitionStreamingPolicy.cpp
@@ -360,6 +360,8 @@ void UWorldPartitionStreamingPolicy::UpdateStreamingState()
 	const bool bCanDeactivateOrUnloadCells = !bIsServer || WorldPartition->IsServerStreamingOutEnabled();
 	const int32 NewServerStreamingEnabledEpoch = ComputeServerStreamingEnabledEpoch();
 
+	const bool bIsStreamingInEnabled = WorldPartition->IsStreamingInEnabled();
+
 	if (!bIsServer || bIsServerStreamingEnabled || AWorldPartitionReplay::IsPlaybackEnabled(World))
 	{
 		// When world partition can't stream, all cells must be unloaded
@@ -563,15 +565,18 @@ void UWorldPartitionStreamingPolicy::UpdateStreamingState()
 
 	int32 ToLoadAndActivateCount = 0;
 
-	// Do Activation State first as it is higher prio than Load State (if we have a limited number of loading cells per frame)
-	if (ToActivateCells.Num() > 0)
+	if (bIsStreamingInEnabled)
 	{
-		ToLoadAndActivateCount += SetCellsStateToActivated(ToActivateCells);
-	}
+		// Do Activation State first as it is higher prio than Load State (if we have a limited number of loading cells per frame)
+		if (ToActivateCells.Num() > 0)
+		{
+			ToLoadAndActivateCount += SetCellsStateToActivated(ToActivateCells);
+		}
 
-	if (ToLoadCells.Num() > 0)
-	{
-		ToLoadAndActivateCount += SetCellsStateToLoaded(ToLoadCells);
+		if (ToLoadCells.Num() > 0)
+		{
+			ToLoadAndActivateCount += SetCellsStateToLoaded(ToLoadCells);
+		}
 	}
 
 	bLastUpdateCompletedLoadingAndActivation = (ToLoadAndActivateCount == (ToActivateCells.Num() + ToLoadCells.Num()));
diff --git a/Engine/Source/Runtime/Engine/Public/WorldPartition/WorldPartition.h b/Engine/Source/Runtime/Engine/Public/WorldPartition/WorldPartition.h
--- a/Engine/Source/Runtime/Engine/Public/WorldPartition/WorldPartition.h
+++ b/Engine/Source/Runtime/Engine/Public/WorldPartition/WorldPartition.h
@@ -256,6 +256,10 @@ public:
 
 	EWorldPartitionStreamingPerformance GetStreamingPerformance() const;
 
+	bool IsStreamingInEnabled() const;
+	void DisableStreamingIn();
+	void EnableStreamingIn();
+
 #if WITH_EDITORONLY_DATA
 	UPROPERTY(DuplicateTransient)
 	TObjectPtr<UWorldPartitionEditorHash> EditorHash;
@@ -322,6 +326,9 @@ private:
 	EWorldPartitionInitState InitState;
 	TOptional<FTransform> InstanceTransform;
 
+	// Defaults to true, can be set to false to temporarly disable Streaming in of new cells.
+	bool bStreamingInEnabled;
+
 	mutable TOptional<bool> bCachedUseMakingInvisibleTransactionRequests;
 	mutable TOptional<bool> bCachedUseMakingVisibleTransactionRequests;
 	mutable TOptional<bool> bCachedIsServerStreamingEnabled;
