- Fix PVS warning due to ImportTextItem always returning true, even when import failed.

diff --git a/Engine/Source/Runtime/Engine/Private/OnlineReplStructs.cpp b/Engine/Source/Runtime/Engine/Private/OnlineReplStructs.cpp
--- a/Engine/Source/Runtime/Engine/Private/OnlineReplStructs.cpp
+++ b/Engine/Source/Runtime/Engine/Private/OnlineReplStructs.cpp
@@ -550,6 +550,7 @@ bool FUniqueNetIdRepl::ImportTextItem(const TCHAR*& Buffer, int32 PortFlags, UOb
 #if !NO_LOGGING
 		ErrorText->CategorizedLogf(LogNet.GetCategoryName(), ELogVerbosity::Warning, TEXT("Failed to import text to FUniqueNetIdRepl Parent:%s"), *GetPathNameSafe(Parent));
 #endif
+		return false;
 	}
 
 	return true;
