Texture: Integrated the texture compiling manager with the derived data build scheduler

This allows for pausing of async texture compilation and limiting concurrent texture builds based on the same memory limits that the asset compiling manager uses.

diff --git a/Engine/Source/Developer/TextureBuild/Private/TextureBuildFunction.cpp b/Engine/Source/Developer/TextureBuild/Private/TextureBuildFunction.cpp
--- a/Engine/Source/Developer/TextureBuild/Private/TextureBuildFunction.cpp
+++ b/Engine/Source/Developer/TextureBuild/Private/TextureBuildFunction.cpp
@@ -296,6 +296,7 @@ FGuid FTextureBuildFunction::GetVersion() const
 
 void FTextureBuildFunction::Configure(UE::DerivedData::FBuildConfigContext& Context) const
 {
+	Context.SetTypeName(UTF8TEXTVIEW("Texture"));
 	Context.SetCacheBucket(UE::DerivedData::FCacheBucket(ANSITEXTVIEW("Texture")));
 
 	const FCbObject Settings = Context.FindConstant(UTF8TEXTVIEW("Settings"));
diff --git a/Engine/Source/Runtime/Engine/Private/TextureCompiler.cpp b/Engine/Source/Runtime/Engine/Private/TextureCompiler.cpp
--- a/Engine/Source/Runtime/Engine/Private/TextureCompiler.cpp
+++ b/Engine/Source/Runtime/Engine/Private/TextureCompiler.cpp
@@ -4,7 +4,10 @@
 
 #if WITH_EDITOR
 
+#include "DerivedDataSharedString.h"
 #include "Engine/Texture.h"
+#include "Experimental/DerivedDataBuildSchedulerThreadPoolProvider.h"
+#include "Features/IModularFeatures.h"
 #include "ObjectCacheContext.h"
 #include "Settings/EditorExperimentalSettings.h"
 #include "Misc/QueuedThreadPoolWrapper.h"
@@ -78,10 +81,42 @@ namespace TextureCompilingManagerImpl
 	}
 }
 
+class UE::Private::FTextureThreadPoolProvider final : public UE::DerivedData::IBuildSchedulerThreadPoolProvider
+{
+public:
+	explicit FTextureThreadPoolProvider(FTextureCompilingManager& InManager)
+		: Manager(InManager)
+	{
+		IModularFeatures::Get().RegisterModularFeature(FeatureName, this);
+	}
+
+	~FTextureThreadPoolProvider() final
+	{
+		IModularFeatures::Get().UnregisterModularFeature(FeatureName, this);
+	}
+
+private:
+	const UE::DerivedData::FUtf8SharedString& GetTypeName() const final
+	{
+		static UE::DerivedData::FUtf8SharedString TypeName(UTF8TEXTVIEW("Texture"));
+		return TypeName;
+	}
+
+	FQueuedThreadPool* GetThreadPool() const final
+	{
+		// TODO: This is not thread-safe right now because the thread pool is created on first use.
+		//       That first use is expected to occur on the main thread, which we do not guarantee.
+		return Manager.GetThreadPool();
+	}
+
+	const FTextureCompilingManager& Manager;
+};
+
 FTextureCompilingManager::FTextureCompilingManager()
 	: Notification(GetAssetNameFormat())
 {
 	TextureCompilingManagerImpl::EnsureInitializedCVars();
+	ThreadPoolProvider = MakePimpl<UE::Private::FTextureThreadPoolProvider>(*this);
 }
 
 FName FTextureCompilingManager::GetStaticAssetTypeName()
diff --git a/Engine/Source/Runtime/Engine/Private/TextureDerivedDataTask.cpp b/Engine/Source/Runtime/Engine/Private/TextureDerivedDataTask.cpp
--- a/Engine/Source/Runtime/Engine/Private/TextureDerivedDataTask.cpp
+++ b/Engine/Source/Runtime/Engine/Private/TextureDerivedDataTask.cpp
@@ -2307,10 +2307,9 @@ public:
 		const FTextureBuildSettings& InSettingsFetchOrBuild,
 		const FTexturePlatformData::FTextureEncodeResultMetadata* InFetchFirstMetadata, // can be nullptr
 		const FTexturePlatformData::FTextureEncodeResultMetadata* InFetchOrBuildMetadata, // can be nullptr
-		EQueuedWorkPriority InPriority,
+		EQueuedWorkPriority Priority,
 		ETextureCacheFlags Flags)
 		: BuildResults(InDerivedData)
-		, Priority(InPriority)
 		, InputResolver(Texture)
 	{
 		static bool bLoadedModules = LoadModules();
@@ -2328,7 +2327,7 @@ public:
 
 		using namespace UE::DerivedData;
 
-		EPriority OwnerPriority = EnumHasAnyFlags(Flags, ETextureCacheFlags::Async) ? ConvertPriority(Priority) : UE::DerivedData::EPriority::Blocking;
+		EPriority OwnerPriority = EnumHasAnyFlags(Flags, ETextureCacheFlags::Async) ? ConvertFromQueuedWorkPriority(Priority) : EPriority::Blocking;
 		Owner.Emplace(OwnerPriority);
 
 		bool bUseCompositeTexture;
@@ -2586,13 +2585,14 @@ public:
 
 	EQueuedWorkPriority GetPriority() const final
 	{
-		return Priority;
+		using namespace UE::DerivedData;
+		return ConvertToQueuedWorkPriority(Owner->GetPriority());
 	}
 
 	bool SetPriority(EQueuedWorkPriority QueuedWorkPriority) final
 	{
-		Priority = QueuedWorkPriority;
-		Owner->SetPriority(ConvertPriority(QueuedWorkPriority));
+		using namespace UE::DerivedData;
+		Owner->SetPriority(ConvertFromQueuedWorkPriority(QueuedWorkPriority));
 		return true;
 	}
 
@@ -2794,35 +2794,6 @@ public:
 		return DefinitionBuilder.Build();
 	}
 
-	static UE::DerivedData::EPriority ConvertPriority(EQueuedWorkPriority SourcePriority)
-	{
-		using namespace UE::DerivedData;
-		switch (SourcePriority)
-		{
-		case EQueuedWorkPriority::Lowest:  return EPriority::Lowest;
-		case EQueuedWorkPriority::Low:     return EPriority::Low;
-		case EQueuedWorkPriority::Normal:  return EPriority::Normal;
-		case EQueuedWorkPriority::High:    return EPriority::High;
-		case EQueuedWorkPriority::Highest: return EPriority::Highest;
-		default:                           return EPriority::Normal;
-		}
-	}
-
-	static EQueuedWorkPriority ConvertPriority(UE::DerivedData::EPriority SourcePriority)
-	{
-		using namespace UE::DerivedData;
-		switch (SourcePriority)
-		{
-		case EPriority::Lowest:   return EQueuedWorkPriority::Lowest;
-		case EPriority::Low:      return EQueuedWorkPriority::Low;
-		case EPriority::Normal:   return EQueuedWorkPriority::Normal;
-		case EPriority::High:     return EQueuedWorkPriority::High;
-		case EPriority::Highest:  return EQueuedWorkPriority::Highest;
-		case EPriority::Blocking: return EQueuedWorkPriority::Blocking;
-		default:                  return EQueuedWorkPriority::Normal;
-		}
-	}
-
 	static ITextureCompressorModule* TextureCompressorModule;
 	static bool LoadModules()
 	{
@@ -2839,7 +2810,6 @@ public:
 
 	// Build bureaucracy
 	TOptional<UE::DerivedData::FRequestOwner> Owner;
-	EQueuedWorkPriority Priority;
 
 	UE::DerivedData::FOptionalBuildSession BuildSession;
 	UE::TextureDerivedData::FTextureBuildInputResolver InputResolver;
diff --git a/Engine/Source/Runtime/Engine/Public/TextureCompiler.h b/Engine/Source/Runtime/Engine/Public/TextureCompiler.h
--- a/Engine/Source/Runtime/Engine/Public/TextureCompiler.h
+++ b/Engine/Source/Runtime/Engine/Public/TextureCompiler.h
@@ -6,6 +6,7 @@
 #include "Containers/Set.h"
 #include "AssetCompilingManager.h"
 #include "AsyncCompilationHelpers.h"
+#include "Templates/PimplPtr.h"
 
 #if WITH_EDITOR
 
@@ -13,6 +14,8 @@ class UTexture;
 class FQueuedThreadPool;
 enum class EQueuedWorkPriority : uint8;
 
+namespace UE::Private { class FTextureThreadPoolProvider; }
+
 DECLARE_MULTICAST_DELEGATE_OneParam(FTexturePostCompileEvent, const TArrayView<UTexture* const>&);
 
 class FTextureCompilingManager : IAssetCompilingManager
@@ -104,6 +107,8 @@ private:
 
 	/** Event issued at the end of the compile process */
 	FTexturePostCompileEvent TexturePostCompileEvent;
+
+	TPimplPtr<UE::Private::FTextureThreadPoolProvider> ThreadPoolProvider;
 };
 
 #endif
\ No newline at end of file
