Added an optional world partition filter function that, if registered, can filter out world partition streaming source providers.
Added GetStreamingSourceOwner to IWorldPartitionStreamingSourceProvider interface to return associated UObject (if any).

diff --git a/Engine/Source/Runtime/Engine/Classes/Components/WorldPartitionStreamingSourceComponent.h b/Engine/Source/Runtime/Engine/Classes/Components/WorldPartitionStreamingSourceComponent.h
--- a/Engine/Source/Runtime/Engine/Classes/Components/WorldPartitionStreamingSourceComponent.h
+++ b/Engine/Source/Runtime/Engine/Classes/Components/WorldPartitionStreamingSourceComponent.h
@@ -37,6 +37,7 @@ class ENGINE_API UWorldPartitionStreamingSourceComponent : public UActorComponen
 
 	// IWorldPartitionStreamingSourceProvider interface
 	virtual bool GetStreamingSource(FWorldPartitionStreamingSource& OutStreamingSource) const override;
+	virtual const UObject* GetStreamingSourceOwner() const override { return this; }
 
 	/** Returns true if streaming is completed for this streaming source component. */
 	UFUNCTION(BlueprintCallable, Category = "Streaming")
diff --git a/Engine/Source/Runtime/Engine/Classes/GameFramework/PlayerController.h b/Engine/Source/Runtime/Engine/Classes/GameFramework/PlayerController.h
--- a/Engine/Source/Runtime/Engine/Classes/GameFramework/PlayerController.h
+++ b/Engine/Source/Runtime/Engine/Classes/GameFramework/PlayerController.h
@@ -817,6 +817,7 @@ public:
 protected:
 	virtual bool GetStreamingSourcesInternal(TArray<FWorldPartitionStreamingSource>& OutStreamingSources) const;
 	virtual bool GetStreamingSource(FWorldPartitionStreamingSource& OutStreamingSource) const final;
+	virtual const UObject* GetStreamingSourceOwner() const override final { return this; }
 	
 	/** Pawn has been possessed, so changing state to NAME_Playing. Start it walking and begin playing with it. */
 	virtual void BeginPlayingState();
diff --git a/Engine/Source/Runtime/Engine/Private/WorldPartition/WorldPartitionSubsystem.cpp b/Engine/Source/Runtime/Engine/Private/WorldPartition/WorldPartitionSubsystem.cpp
--- a/Engine/Source/Runtime/Engine/Private/WorldPartition/WorldPartitionSubsystem.cpp
+++ b/Engine/Source/Runtime/Engine/Private/WorldPartition/WorldPartitionSubsystem.cpp
@@ -232,6 +232,21 @@ bool UWorldPartitionSubsystem::UnregisterStreamingSourceProvider(IWorldPartition
 	return !!StreamingSourceProviders.Remove(StreamingSource);
 }
 
+TSet<IWorldPartitionStreamingSourceProvider*> UWorldPartitionSubsystem::GetStreamingSourceProviders() const
+{
+	TSet<IWorldPartitionStreamingSourceProvider*> Result = StreamingSourceProviders;
+	if (!Result.IsEmpty() && IsStreamingSourceProviderFiltered.IsBound())
+	{
+		for (auto It = Result.CreateIterator(); It; ++It)
+		{
+			if (IsStreamingSourceProviderFiltered.Execute(*It))
+			{
+				It.RemoveCurrent();
+			}
+		}
+	}
+	return Result;
+}
 
 void UWorldPartitionSubsystem::Tick(float DeltaSeconds)
 {
diff --git a/Engine/Source/Runtime/Engine/Public/WorldPartition/WorldPartitionStreamingSource.h b/Engine/Source/Runtime/Engine/Public/WorldPartition/WorldPartitionStreamingSource.h
--- a/Engine/Source/Runtime/Engine/Public/WorldPartition/WorldPartitionStreamingSource.h
+++ b/Engine/Source/Runtime/Engine/Public/WorldPartition/WorldPartitionStreamingSource.h
@@ -551,4 +551,6 @@ struct ENGINE_API IWorldPartitionStreamingSourceProvider
 		}
 		return false;
 	}
+
+	virtual const UObject* GetStreamingSourceOwner() const { return nullptr; }
 };
\ No newline at end of file
diff --git a/Engine/Source/Runtime/Engine/Public/WorldPartition/WorldPartitionSubsystem.h b/Engine/Source/Runtime/Engine/Public/WorldPartition/WorldPartitionSubsystem.h
--- a/Engine/Source/Runtime/Engine/Public/WorldPartition/WorldPartitionSubsystem.h
+++ b/Engine/Source/Runtime/Engine/Public/WorldPartition/WorldPartitionSubsystem.h
@@ -54,10 +54,13 @@ public:
 
 	void DumpStreamingSources(FOutputDevice& OutputDevice) const;
 
-	TSet<IWorldPartitionStreamingSourceProvider*> GetStreamingSourceProviders() const { return StreamingSourceProviders; }
+	TSet<IWorldPartitionStreamingSourceProvider*> GetStreamingSourceProviders() const;
 	void RegisterStreamingSourceProvider(IWorldPartitionStreamingSourceProvider* StreamingSource);
 	bool UnregisterStreamingSourceProvider(IWorldPartitionStreamingSourceProvider* StreamingSource);
 
+	DECLARE_DELEGATE_RetVal_OneParam(bool, FWorldPartitionStreamingSourceProviderFilter, const IWorldPartitionStreamingSourceProvider*);
+	FWorldPartitionStreamingSourceProviderFilter& OnIsStreamingSourceProviderFiltered() { return IsStreamingSourceProviderFiltered; }
+
 #if WITH_EDITOR
 	void ForEachWorldPartition(TFunctionRef<bool(UWorldPartition*)> Func);
 
@@ -80,6 +83,8 @@ private:
 
 	TSet<IWorldPartitionStreamingSourceProvider*> StreamingSourceProviders;
 
+	FWorldPartitionStreamingSourceProviderFilter IsStreamingSourceProviderFiltered;
+
 	FDelegateHandle	DrawHandle;
 
 	// GC backup values
