Fix a bug where force feedback effects would still vibrate past their last keyframe if they had a device property with a higher duration. I fixed this by separating the "GetDuration" function to exclude device property duration, and then manually check if we need to get the force feedback values or the device property values based on their own duration.

#rnx

diff --git a/Engine/Source/Runtime/Engine/Classes/GameFramework/ForceFeedbackEffect.h b/Engine/Source/Runtime/Engine/Classes/GameFramework/ForceFeedbackEffect.h
--- a/Engine/Source/Runtime/Engine/Classes/GameFramework/ForceFeedbackEffect.h
+++ b/Engine/Source/Runtime/Engine/Classes/GameFramework/ForceFeedbackEffect.h
@@ -105,6 +105,9 @@ class UForceFeedbackEffect : public UObject
 
 	float GetDuration();
 
+	/** Returns the longest duration of any active UInputDeviceProperty's that this effect has on it. */
+	float GetTotalDevicePropertyDuration();
+
 	void GetValues(const float EvalTime, FForceFeedbackValues& Values, float ValueMultiplier = 1.f) const;
 
 	void SetDeviceProperties(const FPlatformUserId PlatformUser, const float DeltaTime, const float EvalTime);
diff --git a/Engine/Source/Runtime/Engine/Private/GameFramework/ForceFeedbackEffect.cpp b/Engine/Source/Runtime/Engine/Private/GameFramework/ForceFeedbackEffect.cpp
--- a/Engine/Source/Runtime/Engine/Private/GameFramework/ForceFeedbackEffect.cpp
+++ b/Engine/Source/Runtime/Engine/Private/GameFramework/ForceFeedbackEffect.cpp
@@ -20,8 +20,8 @@ void UForceFeedbackEffect::PostEditChangeChainProperty(struct FPropertyChangedCh
 {
 	// After any edit (really we only care about the curve, but easier this way) update the cached duration value
 	GetDuration();
+	GetTotalDevicePropertyDuration();
 }
-
 #endif
 
 float UForceFeedbackEffect::GetDuration()
@@ -41,22 +41,29 @@ float UForceFeedbackEffect::GetDuration()
 				Duration = MaxTime;
 			}
 		}
+	}
+
+	return Duration;
+}
+
+float UForceFeedbackEffect::GetTotalDevicePropertyDuration()
+{
+	float LongestDuration = 0.0f;
 
-		// Check the device properties for any longer durations
-		for (TObjectPtr<UInputDeviceProperty> DeviceProperty : DeviceProperties)
+	// Check the device properties for any longer durations
+	for (TObjectPtr<UInputDeviceProperty> DeviceProperty : DeviceProperties)
+	{
+		if (DeviceProperty)
 		{
-			if (DeviceProperty)
+			const float PropertyDuration = DeviceProperty->GetDuration();
+			if (PropertyDuration > LongestDuration)
 			{
-				const float PropertyDuration = DeviceProperty->GetDuration();
-				if (PropertyDuration > Duration)
-				{
-					Duration = PropertyDuration;
-				}	
+				LongestDuration = PropertyDuration;
 			}
 		}
 	}
 
-	return Duration;
+	return LongestDuration;
 }
 
 void UForceFeedbackEffect::GetValues(const float EvalTime, FForceFeedbackValues& Values, const float ValueMultiplier) const
@@ -91,8 +98,15 @@ void UForceFeedbackEffect::SetDeviceProperties(const FPlatformUserId PlatformUse
 	{
 		if (DeviceProp)
 		{
-			DeviceProp->EvaluateDeviceProperty(PlatformUser, DeltaTime, EvalTime);
-			DeviceProp->ApplyDeviceProperty(PlatformUser);
+			if (EvalTime > DeviceProp->GetDuration())
+			{
+				DeviceProp->ResetDeviceProperty(PlatformUser);
+			}
+			else
+			{
+				DeviceProp->EvaluateDeviceProperty(PlatformUser, DeltaTime, EvalTime);
+				DeviceProp->ApplyDeviceProperty(PlatformUser);
+			}			
 		}
 	}
 }
@@ -129,20 +143,30 @@ bool FActiveForceFeedbackEffect::Update(const float DeltaTime, FForceFeedbackVal
 		return false;
 	}
 
-	const float Duration = ForceFeedbackEffect->GetDuration();
+	const float EffectDuration = ForceFeedbackEffect->GetDuration();
+	const float DevicePropDuration = ForceFeedbackEffect->GetTotalDevicePropertyDuration();
 
 	PlayTime += (Parameters.bIgnoreTimeDilation ? FApp::GetDeltaTime() : DeltaTime);
 
-	if (PlayTime > Duration && (!Parameters.bLooping || Duration == 0.f) )
+	// If the play time is longer then the force feedback effect curve's last key value, 
+	// or if there are still device properties that need to be evaluated
+	if (PlayTime > EffectDuration && PlayTime > DevicePropDuration && (!Parameters.bLooping || (EffectDuration == 0.0f && DevicePropDuration == 0.0f)))
 	{
 		return false;
 	}
+	// Update the effect values if we can
+	if (PlayTime <= EffectDuration)
+	{
+		GetValues(Values);
+	}	
 	
-	GetValues(Values);
-
-	// Set any input device properties associated with this effect
-	const float EvalTime = PlayTime - Duration * FMath::FloorToFloat(PlayTime / Duration);
-	ForceFeedbackEffect->SetDeviceProperties(PlatformUser, DeltaTime, EvalTime);
+	// Update device properties if we can
+	if (PlayTime <= DevicePropDuration)
+	{
+		// Set any input device properties associated with this effect
+		const float EvalTime = PlayTime - DevicePropDuration * FMath::FloorToFloat(PlayTime / DevicePropDuration);
+		ForceFeedbackEffect->SetDeviceProperties(PlatformUser, DeltaTime, EvalTime);
+	}	
 
 	return true;
 }
