Fix for VT cook tags not accounting for FirstMipToSerialize. FVirtualTextureBuiltData::Serialize doesn't seem to check for the shift going to zero, but in case some assumptions change in the future, I did.

diff --git a/Engine/Source/Runtime/Engine/Private/TextureDerivedData.cpp b/Engine/Source/Runtime/Engine/Private/TextureDerivedData.cpp
--- a/Engine/Source/Runtime/Engine/Private/TextureDerivedData.cpp
+++ b/Engine/Source/Runtime/Engine/Private/TextureDerivedData.cpp
@@ -2526,7 +2526,7 @@ static void SerializePlatformData(
 		if (bIsVirtual)
 		{
 			FVirtualTextureBuiltData* VTData = PlatformData->VTData;
-			CookTags->Add(Texture, "Size", FString::Printf(TEXT("%dx%d"), VTData->Width, VTData->Height));
+			CookTags->Add(Texture, "Size", FString::Printf(TEXT("%dx%d"), FMath::Max(VTData->Width >> FirstMipToSerialize, 1U), FMath::Max(VTData->Height >> FirstMipToSerialize, 1U)));
 		}
 		else if (PlatformData->Mips.Num() > 0) // PlatformData->Mips is empty if BuildTexture failed
 		{
