Rename Device Property's GetDuration function to something more explicit so that it's clear that it will be recalcuated.

#rnx

diff --git a/Engine/Source/Runtime/Engine/Classes/GameFramework/InputDeviceProperties.h b/Engine/Source/Runtime/Engine/Classes/GameFramework/InputDeviceProperties.h
--- a/Engine/Source/Runtime/Engine/Classes/GameFramework/InputDeviceProperties.h
+++ b/Engine/Source/Runtime/Engine/Classes/GameFramework/InputDeviceProperties.h
@@ -11,6 +11,10 @@
 class UCurveLinearColor;
 class UCurveFloat;
 
+#if WITH_EDITOR
+	struct FPropertyChangedChainEvent;
+#endif	// WITH_EDITOR
+
 /**
 * Base class that represents a single Input Device Property. An Input Device Property
 * represents a feature that can be set on an input device. Things like what color a
@@ -29,6 +33,8 @@ class ENGINE_API UInputDeviceProperty : public UObject
 	GENERATED_BODY()
 public:
 
+	UInputDeviceProperty();
+
 	/**
 	* Evaluate this device property for a given duration. 
 	* If overriding in Blueprints, make sure to call the parent function!
@@ -82,17 +88,20 @@ public:
 	* The duration that this device property should last. Override this if your property has any dynamic curves 
 	* to be the max time range.
 	*/
-	virtual float GetDuration() { return PropertyDuration; }
+	float GetDuration() const;
+	
+	/**
+	 * Recalculates this device property's duration. This should be called whenever there are changes made
+	 * to things like curves, or other time sensative properties.
+	 */
+	virtual float RecalculateDuration();
 
-protected:
+	// Post edit change property to update the duration if there are any dynamic options like for curves
+#if WITH_EDITOR
+	virtual void PostEditChangeChainProperty(FPropertyChangedChainEvent& PropertyChangedEvent) override;
+#endif	// WITH_EDITOR
 
-	/** 
-	* The name of this unique device property. 
-	* This is used internally by the platform to determine which device property this represents. 
-	* This needs to correspond to a FInputDeviceProperty name!
-	*/
-	UPROPERTY(VisibleAnywhere, Category = "Info")
-	FName PropertyName;
+protected:
 
 	/**
 	* The duration that this device property should last. Override this if your property has any dynamic curves 
@@ -115,12 +124,10 @@ class UColorInputDeviceProperty : public UInputDeviceProperty
 
 public:
 
-	UColorInputDeviceProperty();
-
 	virtual void EvaluateDeviceProperty_Implementation(const FPlatformUserId PlatformUser, const float DeltaTime, const float Duration) override;
 	virtual void ResetDeviceProperty_Implementation(const FPlatformUserId PlatformUser) override;
 	virtual FInputDeviceProperty* GetInternalDeviceProperty() override;
-	virtual float GetDuration() override;
+	virtual float RecalculateDuration() override;
 
 protected:
 
diff --git a/Engine/Source/Runtime/Engine/Private/GameFramework/ForceFeedbackEffect.cpp b/Engine/Source/Runtime/Engine/Private/GameFramework/ForceFeedbackEffect.cpp
--- a/Engine/Source/Runtime/Engine/Private/GameFramework/ForceFeedbackEffect.cpp
+++ b/Engine/Source/Runtime/Engine/Private/GameFramework/ForceFeedbackEffect.cpp
@@ -55,7 +55,7 @@ float UForceFeedbackEffect::GetTotalDevicePropertyDuration()
 	{
 		if (DeviceProperty)
 		{
-			const float PropertyDuration = DeviceProperty->GetDuration();
+			const float PropertyDuration = DeviceProperty->RecalculateDuration();
 			if (PropertyDuration > LongestDuration)
 			{
 				LongestDuration = PropertyDuration;
diff --git a/Engine/Source/Runtime/Engine/Private/GameFramework/InputDeviceProperties.cpp b/Engine/Source/Runtime/Engine/Private/GameFramework/InputDeviceProperties.cpp
--- a/Engine/Source/Runtime/Engine/Private/GameFramework/InputDeviceProperties.cpp
+++ b/Engine/Source/Runtime/Engine/Private/GameFramework/InputDeviceProperties.cpp
@@ -9,6 +9,11 @@
 ///////////////////////////////////////////////////////////////////////
 // UInputDeviceProperty
 
+UInputDeviceProperty::UInputDeviceProperty()
+{
+	RecalculateDuration();
+}
+
 void UInputDeviceProperty::ApplyDeviceProperty(const FPlatformUserId UserId)
 {
 	if (FInputDeviceProperty* RawProp = GetInternalDeviceProperty())
@@ -25,6 +30,24 @@ void UInputDeviceProperty::ApplyDeviceProperty(const FPlatformUserId UserId)
 	}
 }
 
+float UInputDeviceProperty::GetDuration() const
+{
+	return PropertyDuration;
+}
+
+float UInputDeviceProperty::RecalculateDuration()
+{
+	return PropertyDuration;
+}
+
+#if WITH_EDITOR
+void UInputDeviceProperty::PostEditChangeChainProperty(FPropertyChangedChainEvent& PropertyChangedEvent)
+{
+	Super::PostEditChangeChainProperty(PropertyChangedEvent);
+	RecalculateDuration();
+}
+#endif	// WITH_EDITOR
+
 void UInputDeviceProperty::EvaluateDeviceProperty_Implementation(const FPlatformUserId PlatformUser, const float DeltaTime, const float Duration)
 {
 
@@ -38,13 +61,6 @@ void UInputDeviceProperty::ResetDeviceProperty_Implementation(const FPlatformUse
 ///////////////////////////////////////////////////////////////////////
 // UColorInputDeviceProperty
 
-UColorInputDeviceProperty::UColorInputDeviceProperty()
-	: UInputDeviceProperty()
-{
-	PropertyName = FInputDeviceLightColorProperty::PropertyName();
-	GetDuration();
-}
-
 void UColorInputDeviceProperty::EvaluateDeviceProperty_Implementation(const FPlatformUserId PlatformUser, const float DeltaTime, const float Duration)
 {
 	InternalProperty.bEnable = bEnable;
@@ -68,16 +84,17 @@ FInputDeviceProperty* UColorInputDeviceProperty::GetInternalDeviceProperty()
 	return &InternalProperty;
 }
 
-float UColorInputDeviceProperty::GetDuration()
+float UColorInputDeviceProperty::RecalculateDuration()
 {
 	float MinTime, MaxTime;
 	if (DeviceColorCurve)
 	{
 		DeviceColorCurve->GetTimeRange(MinTime, MaxTime);
+		PropertyDuration = MaxTime;
 	}
 	else
 	{
 		PropertyDuration = 1.0f;
 	}
 	return PropertyDuration;
-}
+}
\ No newline at end of file
