[Engine]
* World - Moved UWorld::InitializationValues outside world to FWorldInitializationValues in order to be able to forward declare. Replaced existing typename with using of new type

diff --git a/Engine/Source/Runtime/Engine/Classes/Engine/World.h b/Engine/Source/Runtime/Engine/Classes/Engine/World.h
--- a/Engine/Source/Runtime/Engine/Classes/Engine/World.h
+++ b/Engine/Source/Runtime/Engine/Classes/Engine/World.h
@@ -858,6 +858,79 @@ private:
 	static void BroadcastWorldPartitionUninitialized(UWorld* InWorld, UWorldPartition* InWorldPartition);
 };
 
+/** Struct containing a collection of optional parameters for initialization of a World. */
+struct FWorldInitializationValues
+{
+	FWorldInitializationValues()
+		: bInitializeScenes(true)
+		, bAllowAudioPlayback(true)
+		, bRequiresHitProxies(true)
+		, bCreatePhysicsScene(true)
+		, bCreateNavigation(true)
+		, bCreateAISystem(true)
+		, bShouldSimulatePhysics(true)
+		, bEnableTraceCollision(false)
+		, bForceUseMovementComponentInNonGameWorld(false)
+		, bTransactional(true)
+		, bCreateFXSystem(true)
+		, bCreateWorldPartition(false)
+	{
+	}
+
+	/** Should the scenes (physics, rendering) be created. */
+	uint32 bInitializeScenes:1;
+
+	/** Are sounds allowed to be generated from this world. */
+	uint32 bAllowAudioPlayback:1;
+
+	/** Should the render scene create hit proxies. */
+	uint32 bRequiresHitProxies:1;
+
+	/** Should the physics scene be created. bInitializeScenes must be true for this to be considered. */
+	uint32 bCreatePhysicsScene:1;
+
+	/** Should the navigation system be created for this world. */
+	uint32 bCreateNavigation:1;
+
+	/** Should the AI system be created for this world. */
+	uint32 bCreateAISystem:1;
+
+	/** Should physics be simulated in this world. */
+	uint32 bShouldSimulatePhysics:1;
+
+	/** Are collision trace calls valid within this world. */
+	uint32 bEnableTraceCollision:1;
+
+	/** Special flag to enable movement component in non game worlds (see UMovementComponent::OnRegister) */
+	uint32 bForceUseMovementComponentInNonGameWorld:1;
+
+	/** Should actions performed to objects in this world be saved to the transaction buffer. */
+	uint32 bTransactional:1;
+
+	/** Should the FX system be created for this world. */
+	uint32 bCreateFXSystem:1;
+
+	/** Should the world be partitioned */
+	uint32 bCreateWorldPartition:1;
+
+	/** The default game mode for this world (if any) */
+	TSubclassOf<class AGameModeBase> DefaultGameMode;
+
+	FWorldInitializationValues& InitializeScenes(const bool bInitialize) { bInitializeScenes = bInitialize; return *this; }
+	FWorldInitializationValues& AllowAudioPlayback(const bool bAllow) { bAllowAudioPlayback = bAllow; return *this; }
+	FWorldInitializationValues& RequiresHitProxies(const bool bRequires) { bRequiresHitProxies = bRequires; return *this; }
+	FWorldInitializationValues& CreatePhysicsScene(const bool bCreate) { bCreatePhysicsScene = bCreate; return *this; }
+	FWorldInitializationValues& CreateNavigation(const bool bCreate) { bCreateNavigation = bCreate; return *this; }
+	FWorldInitializationValues& CreateAISystem(const bool bCreate) { bCreateAISystem = bCreate; return *this; }
+	FWorldInitializationValues& ShouldSimulatePhysics(const bool bInShouldSimulatePhysics) { bShouldSimulatePhysics = bInShouldSimulatePhysics; return *this; }
+	FWorldInitializationValues& EnableTraceCollision(const bool bInEnableTraceCollision) { bEnableTraceCollision = bInEnableTraceCollision; return *this; }
+	FWorldInitializationValues& ForceUseMovementComponentInNonGameWorld(const bool bInForceUseMovementComponentInNonGameWorld) { bForceUseMovementComponentInNonGameWorld = bInForceUseMovementComponentInNonGameWorld; return *this; }
+	FWorldInitializationValues& SetTransactional(const bool bInTransactional) { bTransactional = bInTransactional; return *this; }
+	FWorldInitializationValues& CreateFXSystem(const bool bCreate) { bCreateFXSystem = bCreate; return *this; }
+	FWorldInitializationValues& CreateWorldPartition(const bool bCreate) { bCreateWorldPartition = bCreate; return *this; }
+	FWorldInitializationValues& SetDefaultGameMode(TSubclassOf<class AGameModeBase> GameMode) { DefaultGameMode = GameMode; return *this; }
+};
+
 /** 
  * The World is the top level object representing a map or a sandbox in which Actors and Components will exist and be rendered.  
  *
@@ -2884,83 +2957,13 @@ public:
 	UCanvas* GetCanvasForRenderingToTarget();
 	UCanvas* GetCanvasForDrawMaterialToRenderTarget();
 
-	/** Struct containing a collection of optional parameters for initialization of a World. */
-	struct InitializationValues
-	{
-		InitializationValues()
-			: bInitializeScenes(true)
-			, bAllowAudioPlayback(true)
-			, bRequiresHitProxies(true)
-			, bCreatePhysicsScene(true)
-			, bCreateNavigation(true)
-			, bCreateAISystem(true)
-			, bShouldSimulatePhysics(true)
-			, bEnableTraceCollision(false)
-			, bForceUseMovementComponentInNonGameWorld(false)
-			, bTransactional(true)
-			, bCreateFXSystem(true)
-			, bCreateWorldPartition(false)
-		{
-		}
-
-		/** Should the scenes (physics, rendering) be created. */
-		uint32 bInitializeScenes:1;
-
-		/** Are sounds allowed to be generated from this world. */
-		uint32 bAllowAudioPlayback:1;
-
-		/** Should the render scene create hit proxies. */
-		uint32 bRequiresHitProxies:1;
-
-		/** Should the physics scene be created. bInitializeScenes must be true for this to be considered. */
-		uint32 bCreatePhysicsScene:1;
-
-		/** Should the navigation system be created for this world. */
-		uint32 bCreateNavigation:1;
-
-		/** Should the AI system be created for this world. */
-		uint32 bCreateAISystem:1;
-
-		/** Should physics be simulated in this world. */
-		uint32 bShouldSimulatePhysics:1;
-
-		/** Are collision trace calls valid within this world. */
-		uint32 bEnableTraceCollision:1;
-
-		/** Special flag to enable movement component in non game worlds (see UMovementComponent::OnRegister) */
-		uint32 bForceUseMovementComponentInNonGameWorld:1;
-
-		/** Should actions performed to objects in this world be saved to the transaction buffer. */
-		uint32 bTransactional:1;
-
-		/** Should the FX system be created for this world. */
-		uint32 bCreateFXSystem:1;
-
-		/** Should the world be partitioned */
-		uint32 bCreateWorldPartition:1;
-
-		/** The default game mode for this world (if any) */
-		TSubclassOf<class AGameModeBase> DefaultGameMode;
-
-		InitializationValues& InitializeScenes(const bool bInitialize) { bInitializeScenes = bInitialize; return *this; }
-		InitializationValues& AllowAudioPlayback(const bool bAllow) { bAllowAudioPlayback = bAllow; return *this; }
-		InitializationValues& RequiresHitProxies(const bool bRequires) { bRequiresHitProxies = bRequires; return *this; }
-		InitializationValues& CreatePhysicsScene(const bool bCreate) { bCreatePhysicsScene = bCreate; return *this; }
-		InitializationValues& CreateNavigation(const bool bCreate) { bCreateNavigation = bCreate; return *this; }
-		InitializationValues& CreateAISystem(const bool bCreate) { bCreateAISystem = bCreate; return *this; }
-		InitializationValues& ShouldSimulatePhysics(const bool bInShouldSimulatePhysics) { bShouldSimulatePhysics = bInShouldSimulatePhysics; return *this; }
-		InitializationValues& EnableTraceCollision(const bool bInEnableTraceCollision) { bEnableTraceCollision = bInEnableTraceCollision; return *this; }
-		InitializationValues& ForceUseMovementComponentInNonGameWorld(const bool bInForceUseMovementComponentInNonGameWorld) { bForceUseMovementComponentInNonGameWorld = bInForceUseMovementComponentInNonGameWorld; return *this; }
-		InitializationValues& SetTransactional(const bool bInTransactional) { bTransactional = bInTransactional; return *this; }
-		InitializationValues& CreateFXSystem(const bool bCreate) { bCreateFXSystem = bCreate; return *this; }
-		InitializationValues& CreateWorldPartition(const bool bCreate) { bCreateWorldPartition = bCreate; return *this; }
-		InitializationValues& SetDefaultGameMode(TSubclassOf<class AGameModeBase> GameMode) { DefaultGameMode = GameMode; return *this; }
-	};
+	// Legacy for backwards compatibility
+	using InitializationValues = FWorldInitializationValues;
 
 	/**
 	 * Initializes the world, associates the persistent level and sets the proper zones.
 	 */
-	void InitWorld(const InitializationValues IVS = InitializationValues());
+	void InitWorld(const FWorldInitializationValues IVS = FWorldInitializationValues());
 #if WITH_EDITOR
 	/**
 	 * InitWorld usually has to be balanced with CleanupWorld. If the KeepInitializedDuringLoadTag LinkerInstancingContext tag is present,
