ULevel::RemoveLoadedActors - Temporarily change a verify to ensure while an issue is fixed
#rnx

diff --git a/Engine/Source/Runtime/Engine/Private/Level.cpp b/Engine/Source/Runtime/Engine/Private/Level.cpp
--- a/Engine/Source/Runtime/Engine/Private/Level.cpp
+++ b/Engine/Source/Runtime/Engine/Private/Level.cpp
@@ -993,12 +993,14 @@ void ULevel::RemoveLoadedActors(const TArray<AActor*>& ActorList, const FTransfo
 		});
 
 		int32 ActorIndex;
-		verify(Actors.Find(Actor, ActorIndex));
-
-		Actors[ActorIndex] = nullptr;
-		ActorsForGC.Remove(Actor);
+		// temporarily downgraded to ensure while an issue is fixed
+		if (ensure(Actors.Find(Actor, ActorIndex))) 
+		{
+			Actors[ActorIndex] = nullptr;
+			ActorsForGC.Remove(Actor);
 
-		ActorsQueue.Add(Actor);
+			ActorsQueue.Add(Actor);
+		}
 	};
 
 	for (AActor* Actor : ActorList)
