Add flag for enabling feedback for a VT sample. Defaults to true.
It should be very rare to want to disable feedback. But it might be useful for things like landscape grass which can rely on the underlying landscape providing the feedback. In that case disabling feedback and requesting mip level 0 will just return the mip to match landscape color, and it will reduce cost of mip calculation and feedback writing from the grass base pass shader.

diff --git a/Engine/Source/Runtime/Engine/Classes/Materials/MaterialExpressionRuntimeVirtualTextureSample.h b/Engine/Source/Runtime/Engine/Classes/Materials/MaterialExpressionRuntimeVirtualTextureSample.h
--- a/Engine/Source/Runtime/Engine/Classes/Materials/MaterialExpressionRuntimeVirtualTextureSample.h
+++ b/Engine/Source/Runtime/Engine/Classes/Materials/MaterialExpressionRuntimeVirtualTextureSample.h
@@ -90,6 +90,10 @@ class ENGINE_API UMaterialExpressionRuntimeVirtualTextureSample : public UMateri
 	UPROPERTY(EditAnywhere, Category = VirtualTexture, meta = (DisplayName = "Enable adaptive page table"))
 	bool bAdaptive = false;
 
+	/** Enable virtual texture feedback. */
+	UPROPERTY(EditAnywhere, Category = VirtualTexture)
+	bool bEnableFeedback = true;
+
 	/** Defines how the mip level is calculated for the virtual texture lookup. */
 	UPROPERTY(EditAnywhere, Category = TextureSample)
 	TEnumAsByte<enum ERuntimeVirtualTextureMipValueMode> MipValueMode = RVTMVM_None;
diff --git a/Engine/Source/Runtime/Engine/Private/Materials/HLSLMaterialTranslator.cpp b/Engine/Source/Runtime/Engine/Private/Materials/HLSLMaterialTranslator.cpp
--- a/Engine/Source/Runtime/Engine/Private/Materials/HLSLMaterialTranslator.cpp
+++ b/Engine/Source/Runtime/Engine/Private/Materials/HLSLMaterialTranslator.cpp
@@ -5855,7 +5855,8 @@ int32 FHLSLMaterialTranslator::TextureSample(
 	ESamplerSourceMode SamplerSource,
 	int32 TextureReferenceIndex,
 	bool AutomaticViewMipBias,
-	bool AdaptiveVirtualTexture
+	bool AdaptiveVirtualTexture,
+	bool EnableFeedback
 )
 {
 	if (TextureIndex == INDEX_NONE || CoordinateIndex == INDEX_NONE)
@@ -6282,7 +6283,7 @@ int32 FHLSLMaterialTranslator::TextureSample(
 
 		// Only support GPU feedback from pixel shader
 		//todo[vt]: Support feedback from other shader types
-		const bool bGenerateFeedback = ShaderFrequency == SF_Pixel;
+		const bool bGenerateFeedback = EnableFeedback && ShaderFrequency == SF_Pixel;
 
 		VTLayerIndex = UniformTextureExpressions[(uint32)EMaterialTextureParameterType::Virtual][VirtualTextureIndex]->GetTextureLayerIndex();
 		if (VTLayerIndex != INDEX_NONE)
diff --git a/Engine/Source/Runtime/Engine/Private/Materials/HLSLMaterialTranslator.h b/Engine/Source/Runtime/Engine/Private/Materials/HLSLMaterialTranslator.h
--- a/Engine/Source/Runtime/Engine/Private/Materials/HLSLMaterialTranslator.h
+++ b/Engine/Source/Runtime/Engine/Private/Materials/HLSLMaterialTranslator.h
@@ -852,7 +852,8 @@ protected:
 		ESamplerSourceMode SamplerSource = SSM_FromTextureAsset,
 		int32 TextureReferenceIndex = INDEX_NONE,
 		bool AutomaticViewMipBias = false,
-		bool AdaptiveVirtualTexture = false
+		bool AdaptiveVirtualTexture = false,
+		bool EnableFeedback = true
 	) override;
 
 	virtual int32 TextureProperty(int32 TextureIndex, EMaterialExposedTextureProperty Property) override;
diff --git a/Engine/Source/Runtime/Engine/Private/Materials/MaterialExpressions.cpp b/Engine/Source/Runtime/Engine/Private/Materials/MaterialExpressions.cpp
--- a/Engine/Source/Runtime/Engine/Private/Materials/MaterialExpressions.cpp
+++ b/Engine/Source/Runtime/Engine/Private/Materials/MaterialExpressions.cpp
@@ -3329,7 +3329,7 @@ int32 UMaterialExpressionRuntimeVirtualTextureSample::Compile(class FMaterialCom
 			SAMPLERTYPE_VirtualMasks,
 			MipValue0Index, MipValue1Index, TextureMipLevelMode, SamplerSourceMode,
 			TextureReferenceIndex[TexureLayerIndex],
-			bAutomaticMipViewBias, bAdaptive);
+			bAutomaticMipViewBias, bAdaptive, bEnableFeedback);
 	}
 
 	// Compile any unpacking code
diff --git a/Engine/Source/Runtime/Engine/Public/MaterialCompiler.h b/Engine/Source/Runtime/Engine/Public/MaterialCompiler.h
--- a/Engine/Source/Runtime/Engine/Public/MaterialCompiler.h
+++ b/Engine/Source/Runtime/Engine/Public/MaterialCompiler.h
@@ -347,7 +347,7 @@ public:
 	virtual int32 If(int32 A,int32 B,int32 AGreaterThanB,int32 AEqualsB,int32 ALessThanB,int32 Threshold) = 0;
 
 	virtual int32 TextureCoordinate(uint32 CoordinateIndex, bool UnMirrorU, bool UnMirrorV) = 0;
-	virtual int32 TextureSample(int32 Texture,int32 Coordinate,enum EMaterialSamplerType SamplerType,int32 MipValue0Index=INDEX_NONE,int32 MipValue1Index=INDEX_NONE,ETextureMipValueMode MipValueMode=TMVM_None,ESamplerSourceMode SamplerSource=SSM_FromTextureAsset,int32 TextureReferenceIndex=INDEX_NONE, bool AutomaticViewMipBias=false, bool AdaptiveVirtualTexture=false) = 0;
+	virtual int32 TextureSample(int32 Texture,int32 Coordinate,enum EMaterialSamplerType SamplerType,int32 MipValue0Index=INDEX_NONE,int32 MipValue1Index=INDEX_NONE,ETextureMipValueMode MipValueMode=TMVM_None,ESamplerSourceMode SamplerSource=SSM_FromTextureAsset,int32 TextureReferenceIndex=INDEX_NONE, bool AutomaticViewMipBias=false, bool AdaptiveVirtualTexture=false, bool EnableFeedback = true) = 0;
 	virtual int32 TextureProperty(int32 InTexture, EMaterialExposedTextureProperty Property) = 0;
 
 	virtual int32 TextureDecalMipmapLevel(int32 TextureSizeInput) = 0;
@@ -853,7 +853,7 @@ public:
 
 	virtual int32 If(int32 A, int32 B, int32 AGreaterThanB, int32 AEqualsB, int32 ALessThanB, int32 Threshold) override { return Compiler->If(A, B, AGreaterThanB, AEqualsB, ALessThanB, Threshold); }
 
-	virtual int32 TextureSample(int32 InTexture, int32 Coordinate, enum EMaterialSamplerType SamplerType, int32 MipValue0Index, int32 MipValue1Index, ETextureMipValueMode MipValueMode, ESamplerSourceMode SamplerSource, int32 TextureReferenceIndex, bool AutomaticViewMipBias, bool AdaptiveVirtualTexture) override
+	virtual int32 TextureSample(int32 InTexture, int32 Coordinate, enum EMaterialSamplerType SamplerType, int32 MipValue0Index, int32 MipValue1Index, ETextureMipValueMode MipValueMode, ESamplerSourceMode SamplerSource, int32 TextureReferenceIndex, bool AutomaticViewMipBias, bool AdaptiveVirtualTexture, bool EnableFeedback) override
 	{
 		return Compiler->TextureSample(InTexture, Coordinate, SamplerType, MipValue0Index, MipValue1Index, MipValueMode, SamplerSource, TextureReferenceIndex, AutomaticViewMipBias, AdaptiveVirtualTexture);
 	}
