Engine: Fix include path for UnrealEd Selection.h to not root from module parent directory

#rnx

diff --git a/Engine/Source/Runtime/Engine/Classes/Engine/Selection.h b/Engine/Source/Runtime/Engine/Classes/Engine/Selection.h
--- a/Engine/Source/Runtime/Engine/Classes/Engine/Selection.h
+++ b/Engine/Source/Runtime/Engine/Classes/Engine/Selection.h
@@ -4,5 +4,5 @@
 
 #if WITH_EDITOR
 // TODO: Deprecate this adapter include in favor of using the UnrealEd version directly
-#include "UnrealEd/Public/Selection.h"
+#include "Editor/UnrealEd/Public/Selection.h"
 #endif
