Fixed static analysis error.

diff --git a/Engine/Source/Runtime/Engine/Private/Materials/HLSLMaterialTranslator.cpp b/Engine/Source/Runtime/Engine/Private/Materials/HLSLMaterialTranslator.cpp
--- a/Engine/Source/Runtime/Engine/Private/Materials/HLSLMaterialTranslator.cpp
+++ b/Engine/Source/Runtime/Engine/Private/Materials/HLSLMaterialTranslator.cpp
@@ -3728,7 +3728,7 @@ int32 FHLSLMaterialTranslator::CallExpression(FMaterialExpressionKey ExpressionK
 	// This is automatically handled via the compiler StrataTreeStack.
 	// It means that a node can be blended at multiple point of the graph (allowing acyclic graph instead of tree, e.g. a Slab can be used into multiple input).
 	// It is worth noting that only strata BSDF can be duplicated today according to StrataTreeStack.
-	const bool bExpressionIsStrataBSDF = ExpressionKey.Expression->IsA<UMaterialExpressionStrataSlabBSDF>();
+	const bool bExpressionIsStrataBSDF = ExpressionKey.Expression && ExpressionKey.Expression->IsA<UMaterialExpressionStrataSlabBSDF>();
 
 	// Check if this expression has already been translated.
 	check(ShaderFrequency < SF_NumFrequencies);
