Fixed thread safety of Engine.Sup[p]ressWarningsInOnScreenDisplay

diff --git a/Engine/Source/Runtime/Core/Private/Misc/BufferedOutputDevice.cpp b/Engine/Source/Runtime/Core/Private/Misc/BufferedOutputDevice.cpp
--- a/Engine/Source/Runtime/Core/Private/Misc/BufferedOutputDevice.cpp
+++ b/Engine/Source/Runtime/Core/Private/Misc/BufferedOutputDevice.cpp
@@ -2,17 +2,18 @@
 
 #include "Misc/BufferedOutputDevice.h"
 
-#include "CoreTypes.h"
 #include "Misc/ScopeLock.h"
 #include "Templates/UnrealTemplate.h"
 
-void FBufferedOutputDevice::Serialize(const TCHAR* InData, ELogVerbosity::Type Verbosity, const class FName& Category)
+void FBufferedOutputDevice::Serialize(const TCHAR* InData, ELogVerbosity::Type Verbosity, const FName& Category)
 {
 	if (Verbosity > FilterLevel)
+	{
 		return;
+	}
 
 	FScopeLock ScopeLock(&SynchronizationObject);
-	new(BufferedLines)FBufferedLine(InData, Category, Verbosity);
+	new(BufferedLines) FBufferedLine(InData, Category, Verbosity);
 }
 
 void FBufferedOutputDevice::GetContents(TArray<FBufferedLine>& DestBuffer)
@@ -20,3 +21,12 @@ void FBufferedOutputDevice::GetContents(TArray<FBufferedLine>& DestBuffer)
 	FScopeLock ScopeLock(&SynchronizationObject);
 	DestBuffer = MoveTemp(BufferedLines);
 }
+
+void FBufferedOutputDevice::RedirectTo(FOutputDevice& Ar)
+{
+	FScopeLock ScopeLock(&SynchronizationObject);
+	for (const FBufferedLine& BufferedLine : BufferedLines)
+	{
+		Ar.Serialize(BufferedLine.Data, BufferedLine.Verbosity, BufferedLine.Category);
+	}
+}
diff --git a/Engine/Source/Runtime/Core/Public/Misc/BufferedOutputDevice.h b/Engine/Source/Runtime/Core/Public/Misc/BufferedOutputDevice.h
--- a/Engine/Source/Runtime/Core/Public/Misc/BufferedOutputDevice.h
+++ b/Engine/Source/Runtime/Core/Public/Misc/BufferedOutputDevice.h
@@ -8,33 +8,20 @@
 #include "Logging/LogVerbosity.h"
 #include "Misc/OutputDevice.h"
 #include "Misc/OutputDeviceRedirector.h"
-#include "UObject/NameTypes.h"
 
 /** Buffered output device. */
 class CORE_API FBufferedOutputDevice : public FOutputDevice
 {
 protected:
 	TArray<FBufferedLine>		BufferedLines;
-	ELogVerbosity::Type			FilterLevel;
 	FCriticalSection			SynchronizationObject;
+	ELogVerbosity::Type			FilterLevel = ELogVerbosity::All;
 
 public:
-
-	FBufferedOutputDevice() : FilterLevel(ELogVerbosity::All)
-	{
-	}
-
 	void	SetVerbosity(ELogVerbosity::Type Verbosity) { FilterLevel = Verbosity; }
-	void	Serialize(const TCHAR* InData, ELogVerbosity::Type Verbosity, const class FName& Category) override;
+	void	Serialize(const TCHAR* InData, ELogVerbosity::Type Verbosity, const FName& Category) override;
 	void	GetContents(TArray<FBufferedLine>& DestBuffer);
 
 	/** Pushes buffered lines into the specified output device. */
-	void RedirectTo(class FOutputDevice& Ar)
-	{
-		for (const FBufferedLine& BufferedLine : BufferedLines)
-		{
-			Ar.Serialize(BufferedLine.Data, BufferedLine.Verbosity, BufferedLine.Category);
-		}
-	}
+	void	RedirectTo(FOutputDevice& Ar);
 };
-
diff --git a/Engine/Source/Runtime/Engine/Classes/Engine/Engine.h b/Engine/Source/Runtime/Engine/Classes/Engine/Engine.h
--- a/Engine/Source/Runtime/Engine/Classes/Engine/Engine.h
+++ b/Engine/Source/Runtime/Engine/Classes/Engine/Engine.h
@@ -12,7 +12,9 @@
 #include "Engine/EngineBaseTypes.h"
 #include "UObject/SoftObjectPath.h"
 #include "Engine/World.h"
+#if UE_ENABLE_INCLUDE_ORDER_DEPRECATED_IN_5_2
 #include "Misc/BufferedOutputDevice.h"
+#endif
 #include "Misc/FrameRate.h"
 #include "Subsystems/SubsystemCollection.h"
 #include "Subsystems/EngineSubsystem.h"
@@ -21,6 +23,7 @@
 #include "RHI.h"
 #include "AudioDeviceManager.h"
 #endif
+#include "Templates/PimplPtr.h"
 #include "Templates/UniqueObj.h"
 #include "Containers/Ticker.h"
 #include "DynamicRenderScaling.h"
@@ -3671,20 +3674,8 @@ private:
 
 	// Helper struct that registers itself with the output redirector and copies off warnings
 	// and errors that we'll overlay on the client viewport
-	struct FErrorsAndWarningsCollector : public FBufferedOutputDevice
-	{
-		FErrorsAndWarningsCollector();
-		~FErrorsAndWarningsCollector();
-
-		void Initialize();
-		bool Tick(float Seconds);
-
-		TMap<uint32, uint32>	MessagesToCountMap;
-		FTSTicker::FDelegateHandle			TickerHandle;
-		float					DisplayTime;
-	};
-
-	FErrorsAndWarningsCollector	ErrorsAndWarningsCollector;
+	struct FErrorsAndWarningsCollector;
+	TPimplPtr<FErrorsAndWarningsCollector> ErrorsAndWarningsCollector;
 
 private:
 
diff --git a/Engine/Source/Runtime/Engine/Private/UnrealEngine.cpp b/Engine/Source/Runtime/Engine/Private/UnrealEngine.cpp
--- a/Engine/Source/Runtime/Engine/Private/UnrealEngine.cpp
+++ b/Engine/Source/Runtime/Engine/Private/UnrealEngine.cpp
@@ -27,6 +27,7 @@ UnrealEngine.cpp: Implements the UEngine class and helpers.
 #include "AI/NavigationSystemConfig.h"
 #include "Misc/MessageDialog.h"
 #include "HAL/FileManager.h"
+#include "Misc/BufferedOutputDevice.h"
 #include "Misc/CommandLine.h"
 #include "Misc/FileHelper.h"
 #include "Misc/Paths.h"
@@ -1905,7 +1906,7 @@ void UEngine::Init(IEngineLoop* InEngineLoop)
 
 	// Start capturing errors and warnings
 #if !(UE_BUILD_SHIPPING || UE_BUILD_TEST)
-	ErrorsAndWarningsCollector.Initialize();
+	ErrorsAndWarningsCollector = MakePimpl<FErrorsAndWarningsCollector>();
 #endif
 
 #if WITH_EDITOR
@@ -10957,11 +10958,20 @@ void UEngine::RemoveOnScreenDebugMessage(uint64 Key)
 #endif
 }
 
-UEngine::FErrorsAndWarningsCollector::FErrorsAndWarningsCollector()
+struct UEngine::FErrorsAndWarningsCollector : public FBufferedOutputDevice
 {
-}
+	FErrorsAndWarningsCollector();
+	~FErrorsAndWarningsCollector();
+
+	bool Tick(float Seconds);
+
+	TArray<FBufferedLine>		ActiveLines;
+	TMap<uint32, uint32>		MessagesToCountMap;
+	FTSTicker::FDelegateHandle	TickerHandle;
+	float						DisplayTime;
+};
 
-void UEngine::FErrorsAndWarningsCollector::Initialize()
+UEngine::FErrorsAndWarningsCollector::FErrorsAndWarningsCollector()
 {
 	DisplayTime = 0.0f;
 	GConfig->GetFloat(TEXT("/Script/Engine.Engine"), TEXT("DurationOfErrorsAndWarningsOnHUD"), DisplayTime, GEngineIni);
@@ -10985,12 +10995,18 @@ UEngine::FErrorsAndWarningsCollector::~FErrorsAndWarningsCollector()
 
 bool UEngine::FErrorsAndWarningsCollector::Tick(float Seconds)
 {
-    QUICK_SCOPE_CYCLE_COUNTER(STAT_FErrorsAndWarningsCollector_Tick);
+	QUICK_SCOPE_CYCLE_COUNTER(STAT_FErrorsAndWarningsCollector_Tick);
 
 	// Set this each tick, in case the cvar is changed at runtime
 	SetVerbosity((GSupressWarningsInOnScreenDisplay != 0) ? ELogVerbosity::Error : ELogVerbosity::Warning);
 
-	if (BufferedLines.Num())
+	{
+		TArray<FBufferedLine> NewLines;
+		GetContents(NewLines);
+		ActiveLines.Append(MoveTemp(NewLines));
+	}
+
+	if (ActiveLines.Num())
 	{
 		int DupeCount = 0;
 		int CurrentHash = 0;
@@ -10998,7 +11014,7 @@ bool UEngine::FErrorsAndWarningsCollector::Tick(float Seconds)
 		// Remove any dupes and count them
 		do 
 		{
-			uint32 ThisHash = FCrc::StrCrc32(BufferedLines[DupeCount].Data);
+			uint32 ThisHash = FCrc::StrCrc32(ActiveLines[DupeCount].Data);
 
 			if (CurrentHash && ThisHash != CurrentHash)
 			{
@@ -11008,14 +11024,14 @@ bool UEngine::FErrorsAndWarningsCollector::Tick(float Seconds)
 			CurrentHash = ThisHash;
 			DupeCount++;
 
-		} while (DupeCount < BufferedLines.Num());
+		} while (DupeCount < ActiveLines.Num());
 
 		// Save off properties
-		FString Msg = BufferedLines[0].Data;
-		ELogVerbosity::Type Verbosity = BufferedLines[0].Verbosity;
+		FString Msg = ActiveLines[0].Data;
+		ELogVerbosity::Type Verbosity = ActiveLines[0].Verbosity;
 
 		// Remove any lines we condensed
-		BufferedLines.RemoveAt(0, DupeCount);
+		ActiveLines.RemoveAt(0, DupeCount);
 
 		uint32* pCount = MessagesToCountMap.Find(CurrentHash);
 
@@ -11038,11 +11054,10 @@ bool UEngine::FErrorsAndWarningsCollector::Tick(float Seconds)
 		const float MaxTimeToKeep = 30.f;
 		const int32 MaxNumToKeep = FMath::TruncToInt(MaxTimeToKeep / DisplayTime);
 
-		if (BufferedLines.Num() > MaxNumToKeep)
+		if (ActiveLines.Num() > MaxNumToKeep)
 		{
-			BufferedLines.RemoveAt(0, BufferedLines.Num() - MaxNumToKeep);
+			ActiveLines.RemoveAt(0, ActiveLines.Num() - MaxNumToKeep);
 		}
-
 	}
 
 	return true;
