PR #9027: Fix navmesh generation for hierarchical instanced static meshes (Contributed by [REDACTED_PERSON])

diff --git a/Engine/Source/Runtime/Engine/Private/HierarchicalInstancedStaticMesh.cpp b/Engine/Source/Runtime/Engine/Private/HierarchicalInstancedStaticMesh.cpp
--- a/Engine/Source/Runtime/Engine/Private/HierarchicalInstancedStaticMesh.cpp
+++ b/Engine/Source/Runtime/Engine/Private/HierarchicalInstancedStaticMesh.cpp
@@ -2464,7 +2464,7 @@ int32 UHierarchicalInstancedStaticMeshComponent::AddInstance(const FTransform& I
 	int32 InstanceIndex = UInstancedStaticMeshComponent::AddInstance(InstanceTransform, bWorldSpace);
 
 	// The tree will be fully rebuilt once the static mesh compilation is finished, no need for incremental update in that case.
-	if (InstanceIndex != INDEX_NONE && GetStaticMesh() && !GetStaticMesh()->IsCompiling() && GetStaticMesh()->HasValidRenderData())
+	if (InstanceIndex != INDEX_NONE && GetStaticMesh() && !GetStaticMesh()->IsCompiling() && GetStaticMesh()->HasValidRenderData(false))
 	{	
 		check(InstanceIndex == InstanceReorderTable.Num());
 
@@ -2503,7 +2503,7 @@ TArray<int32> UHierarchicalInstancedStaticMeshComponent::AddInstances(const TArr
 	TArray<int32> InstanceIndices = UInstancedStaticMeshComponent::AddInstances(InstanceTransforms, true, bWorldSpace);
 
 	// The tree will be fully rebuilt once the static mesh compilation is finished, no need for incremental update in that case.
-	if (InstanceIndices.Num() > 0 && GetStaticMesh() && !GetStaticMesh()->IsCompiling() && GetStaticMesh()->HasValidRenderData())
+	if (InstanceIndices.Num() > 0 && GetStaticMesh() && !GetStaticMesh()->IsCompiling() && GetStaticMesh()->HasValidRenderData(false))
 	{
 		bIsOutOfDate = true;
 		bConcurrentChanges |= IsAsyncBuilding();
@@ -2689,7 +2689,7 @@ void UHierarchicalInstancedStaticMeshComponent::BuildTree()
 	}
 
 	// The tree will be fully rebuilt once the static mesh compilation is finished, no need for incremental update in that case.
-	if (PerInstanceSMData.Num() > 0 && GetStaticMesh() && !GetStaticMesh()->IsCompiling() && GetStaticMesh()->HasValidRenderData())
+	if (PerInstanceSMData.Num() > 0 && GetStaticMesh() && !GetStaticMesh()->IsCompiling() && GetStaticMesh()->HasValidRenderData(false))
 	{
 		// Build the tree in translated space to maintain precision.
 		TranslatedInstanceSpaceOrigin = CalcTranslatedInstanceSpaceOrigin();
@@ -2829,7 +2829,7 @@ void UHierarchicalInstancedStaticMeshComponent::ApplyEmpty()
 	SortedInstances.Empty();
 	UnbuiltInstanceBoundsList.Empty();
 	BuiltInstanceBounds.Init();
-	CacheMeshExtendedBounds = (GetStaticMesh() && (GetStaticMesh()->IsCompiling() || GetStaticMesh()->HasValidRenderData())) ? GetStaticMesh()->GetBounds() : FBoxSphereBounds(ForceInitToZero);
+	CacheMeshExtendedBounds = (GetStaticMesh() && (GetStaticMesh()->IsCompiling() || GetStaticMesh()->HasValidRenderData(false))) ? GetStaticMesh()->GetBounds() : FBoxSphereBounds(ForceInitToZero);
 	InstanceUpdateCmdBuffer.Reset();
 	if (PerInstanceRenderData.IsValid())
 	{
@@ -3014,7 +3014,7 @@ void UHierarchicalInstancedStaticMeshComponent::BuildTreeAsync()
 
 	// Verify that the mesh is valid before using it.
 	// The tree will be fully rebuilt once the static mesh compilation is finished, no need to do it now.
-	if (PerInstanceSMData.Num() > 0 && GetStaticMesh() && !GetStaticMesh()->IsCompiling() && GetStaticMesh()->HasValidRenderData())
+	if (PerInstanceSMData.Num() > 0 && GetStaticMesh() && !GetStaticMesh()->IsCompiling() && GetStaticMesh()->HasValidRenderData(false))
 	{
 		double StartTime = FPlatformTime::Seconds();
 		
