World Partition Validations: changed how we handle invalid actor references to detect references to external containers, such as content bundles. This will also report these invalid setups when validating changelists, which was silently falling through the cracks before.

#rnx

diff --git a/Engine/Plugins/Editor/DataValidation/Source/DataValidation/Private/WorldPartitionChangelistValidator.cpp b/Engine/Plugins/Editor/DataValidation/Source/DataValidation/Private/WorldPartitionChangelistValidator.cpp
--- a/Engine/Plugins/Editor/DataValidation/Source/DataValidation/Private/WorldPartitionChangelistValidator.cpp
+++ b/Engine/Plugins/Editor/DataValidation/Source/DataValidation/Private/WorldPartitionChangelistValidator.cpp
@@ -51,26 +51,6 @@ EDataValidationResult UWorldPartitionChangelistValidator::ValidateLoadedAsset_Im
 	return Result;
 }
 
-FString GetPrettyPackageName(const FWorldPartitionActorDescView& Desc)
-{
-	FString AssetPath = Desc.GetActorSoftPath().ToString();
-	
-	int32 LastDot = -1;
-	if (AssetPath.FindLastChar('.', LastDot))
-	{
-		AssetPath.LeftInline(LastDot);
-	}
-
-	FString AssetName = Desc.GetActorLabel().ToString();
-	
-	if (AssetName.Len() == 0)
-	{
-		AssetName = Desc.GetActorName().ToString();
-	}
-	
-	return AssetPath + TEXT(".") + AssetName;
-}
-
 // Extract all Actors/Map from Changelist (in OFPA this should be one Actor per Package, and we'll discard all Actors from non WorldPartition maps)
 // and add them to a Map of World->Files[] so that we can do one validation per world. Once Worlds are identified, we either the UActorDescContainer 
 // from memory (if loaded) or request it to be loaded, we then build a Set of objects that interest us from the Actors in the CL 
@@ -237,12 +217,23 @@ EDataValidationResult UWorldPartitionChangelistValidator::ValidateActorsAndDataL
 		}
 
 		// Invoke static WorldPartition Validation from the ActorDescContainer
-		const bool bIsStreamingDisabled = ULevel::GetIsStreamingDisabledFromPackage(MapPath.GetPackageName());
-		const bool bIsChangelistValidation = true;
+		UWorldPartition::FCheckForErrorsParams Params;
+		Params.ErrorHandler = this;
+		Params.bEnableStreaming = ULevel::GetIsStreamingDisabledFromPackage(MapPath.GetPackageName());;
+
+		ContainersToValidate.ForEachActorDescContainer([&Params](const UActorDescContainer* ActorDescContainer)
+		{
+			for (FActorDescList::TConstIterator<> ActorDescIt(ActorDescContainer); ActorDescIt; ++ActorDescIt)
+			{
+				check(!Params.ActorGuidsToContainerMap.Contains(ActorDescIt->GetGuid()));
+				Params.ActorGuidsToContainerMap.Add(ActorDescIt->GetGuid(), ActorDescContainer);
+			}
+		});
 
-		ContainersToValidate.ForEachActorDescContainer([this, bIsStreamingDisabled, bIsChangelistValidation](const UActorDescContainer* ActorDescContainer)
+		ContainersToValidate.ForEachActorDescContainer([&Params](const UActorDescContainer* ActorDescContainer)
 		{
-			UWorldPartition::CheckForErrors(this, ActorDescContainer, !bIsStreamingDisabled, bIsChangelistValidation);
+			Params.ActorDescContainer = ActorDescContainer;
+			UWorldPartition::CheckForErrors(Params);
 		});
 	}
 
@@ -277,13 +268,13 @@ bool UWorldPartitionChangelistValidator::Filter(const UDataLayerInstance* InData
 	return DataLayerWithAsset != nullptr && DataLayerWithAsset->GetAsset() != nullptr && RelevantDataLayerAssets.Contains(DataLayerWithAsset->GetAsset()->GetPathName());
 }
 
-void UWorldPartitionChangelistValidator::OnInvalidReference(const FWorldPartitionActorDescView& ActorDescView, const FGuid& ReferenceGuid)
+void UWorldPartitionChangelistValidator::OnInvalidReference(const FWorldPartitionActorDescView& ActorDescView, const FGuid& ReferenceGuid, FWorldPartitionActorDescView* ReferenceActorDescView)
 {
 	if (Filter(ActorDescView))
 	{
-		FText CurrentError = FText::Format(LOCTEXT("DataValidation.Changelist.WorldPartition.InvalidReference", "Actor {0} has a missing reference to {1}"),
-											FText::FromString(GetPrettyPackageName(ActorDescView)), 
-											FText::FromString(ReferenceGuid.ToString()));
+		FText CurrentError = FText::Format(LOCTEXT("DataValidation.Changelist.WorldPartition.InvalidReference", "Actor {0} has an invalid reference to {1}"),
+											FText::FromString(GetFullActorName(ActorDescView)), 
+											FText::FromString(ReferenceActorDescView ? GetFullActorName(*ReferenceActorDescView) : ReferenceGuid.ToString()));
 
 		Errors->Add(CurrentError);
 	}
@@ -302,9 +293,9 @@ void UWorldPartitionChangelistValidator::OnInvalidReferenceGridPlacement(const F
 
 			FText CurrentError = FText::Format(LOCTEXT("DataValidation.Changelist.WorldPartition.InvalidReferenceGridPlacement", "{0} {1} is referencing {2} {3}."),
 												FText::FromString(ActorDescView.GetIsSpatiallyLoaded() ? *SpatiallyLoadedActor : *NonSpatiallyLoadedActor),
-												FText::FromString(GetPrettyPackageName(ActorDescView)),
+												FText::FromString(GetFullActorName(ActorDescView)),
 												FText::FromString(ReferenceActorDescView.GetIsSpatiallyLoaded() ? *SpatiallyLoadedActor : *NonSpatiallyLoadedActor),
-												FText::FromString(GetPrettyPackageName(ReferenceActorDescView)));
+												FText::FromString(GetFullActorName(ReferenceActorDescView)));
 
 			Errors->Add(CurrentError);
 		}
@@ -316,8 +307,8 @@ void UWorldPartitionChangelistValidator::OnInvalidReferenceDataLayers(const FWor
 	if (Filter(ActorDescView) || Filter(ReferenceActorDescView))
 	{
 		FText CurrentError = FText::Format(LOCTEXT("DataValidation.Changelist.WorldPartition.DataLayerError", "{0} is referencing {1} but both actors are using a different set of runtime data layers."),
-											FText::FromString(GetPrettyPackageName(ActorDescView)),
-											FText::FromString(GetPrettyPackageName(ReferenceActorDescView)));
+											FText::FromString(GetFullActorName(ActorDescView)),
+											FText::FromString(GetFullActorName(ReferenceActorDescView)));
 
 		Errors->Add(CurrentError);
 	}
@@ -328,8 +319,8 @@ void UWorldPartitionChangelistValidator::OnInvalidReferenceRuntimeGrid(const FWo
 	if (Filter(ActorDescView) || Filter(ReferenceActorDescView))
 	{
 		FText CurrentError = FText::Format(LOCTEXT("DataValidation.Changelist.WorldPartition.RuntimeGridError", "{0} is referencing {1} but both actors are using a different runtime grid."),
-			FText::FromString(GetPrettyPackageName(ActorDescView)),
-			FText::FromString(GetPrettyPackageName(ReferenceActorDescView)));
+			FText::FromString(GetFullActorName(ActorDescView)),
+			FText::FromString(GetFullActorName(ReferenceActorDescView)));
 
 		Errors->Add(CurrentError);
 	}
@@ -340,7 +331,7 @@ void UWorldPartitionChangelistValidator::OnInvalidReferenceLevelScriptStreamed(c
 	if (Filter(ActorDescView))
 	{		
 		FText CurrentError = FText::Format(LOCTEXT("DataValidation.Changelist.WorldPartition.InvalidReferenceLevelScriptStreamed", "Level script blueprint references streamed actor {0}."),
-											FText::FromString(GetPrettyPackageName(ActorDescView)));
+											FText::FromString(GetFullActorName(ActorDescView)));
 		
 		Errors->Add(CurrentError);
 	}
@@ -351,7 +342,7 @@ void UWorldPartitionChangelistValidator::OnInvalidReferenceLevelScriptDataLayers
 	if (Filter(ActorDescView))
 	{
 		FText CurrentError = FText::Format(LOCTEXT("DataValidation.Changelist.WorldPartition.InvalidReferenceLevelScriptDataLayers", "Level script blueprint references streamed actor {0} with a non empty set of data layers."),
-											FText::FromString(GetPrettyPackageName(ActorDescView)));
+											FText::FromString(GetFullActorName(ActorDescView)));
 
 		Errors->Add(CurrentError);
 	}
@@ -408,14 +399,27 @@ void UWorldPartitionChangelistValidator::OnLevelInstanceInvalidWorldAsset(const
 {
 	if (Filter(ActorDescView))
 	{
-		if (Reason == ELevelInstanceInvalidReason::WorldAssetNotFound)
+		FText CurrentError;
+
+		switch (Reason)
 		{
-			FText CurrentError = FText::Format(LOCTEXT("DataValidation.Changelist.WorldPartition.LevelInstanceInvalidWorldAsset", "Level instance {0} has an invalid world asset {1}."),
-				FText::FromString(GetPrettyPackageName(ActorDescView)), 
+		case ELevelInstanceInvalidReason::WorldAssetNotFound:
+			CurrentError = FText::Format(LOCTEXT("DataValidation.Changelist.WorldPartition.LevelInstanceInvalidWorldAsset", "Level instance {0} has an invalid world asset {1}."),
+				FText::FromString(GetFullActorName(ActorDescView)), 
 				FText::FromName(WorldAsset));
+			break;
+		case ELevelInstanceInvalidReason::WorldAssetNotUsingExternalActors:
+			// Not a validation error
+			break;
+		case ELevelInstanceInvalidReason::WorldAssetImcompatiblePartitioned:
+			// Not a validation error
+			break;
+		case ELevelInstanceInvalidReason::WorldAssetHasInvalidContainer:
+			// We cannot treat that error as a validation error as it's possible to validate changelists without loading the world
+			break;
+		};
 
-			Errors->Add(CurrentError);
-		}
+		Errors->Add(CurrentError);
 	}
 }
 
diff --git a/Engine/Plugins/Editor/DataValidation/Source/DataValidation/Private/WorldPartitionChangelistValidator.h b/Engine/Plugins/Editor/DataValidation/Source/DataValidation/Private/WorldPartitionChangelistValidator.h
--- a/Engine/Plugins/Editor/DataValidation/Source/DataValidation/Private/WorldPartitionChangelistValidator.h
+++ b/Engine/Plugins/Editor/DataValidation/Source/DataValidation/Private/WorldPartitionChangelistValidator.h
@@ -15,13 +15,12 @@ class DATAVALIDATION_API UWorldPartitionChangelistValidator : public UEditorVali
 {
 	GENERATED_BODY()
 
-protected:		
-	
-	TArray<FText>*							Errors = nullptr;
-	FTopLevelAssetPath						RelevantMap;
-	TSet<FGuid>								RelevantActorGuids;
-	TSet<FString>							RelevantDataLayerAssets;
-	bool									SubmittingWorldDataLayers = false;
+protected:	
+	TArray<FText>* Errors = nullptr;
+	FTopLevelAssetPath RelevantMap;
+	TSet<FGuid> RelevantActorGuids;
+	TSet<FString> RelevantDataLayerAssets;
+	bool SubmittingWorldDataLayers = false;
 
 	virtual bool CanValidateAsset_Implementation(UObject* InAsset) const override;
 	virtual EDataValidationResult ValidateLoadedAsset_Implementation(UObject* InAsset, TArray<FText>& ValidationErrors) override;
@@ -35,7 +34,7 @@ protected:
 	bool Filter(const UDataLayerInstance* DataLayerInstance);
 
 	// IStreamingGenerationErrorHandler Interface methods
-	virtual void OnInvalidReference(const FWorldPartitionActorDescView& ActorDescView, const FGuid& ReferenceGuid) override;
+	virtual void OnInvalidReference(const FWorldPartitionActorDescView& ActorDescView, const FGuid& ReferenceGuid, FWorldPartitionActorDescView* ReferenceActorDescView) override;
 	virtual void OnInvalidReferenceGridPlacement(const FWorldPartitionActorDescView& ActorDescView, const FWorldPartitionActorDescView& ReferenceActorDescView) override;
 	virtual void OnInvalidReferenceDataLayers(const FWorldPartitionActorDescView& ActorDescView, const FWorldPartitionActorDescView& ReferenceActorDescView) override;
 	virtual void OnInvalidReferenceLevelScriptStreamed(const FWorldPartitionActorDescView& ActorDescView) override;
@@ -47,4 +46,3 @@ protected:
 	virtual void OnActorNeedsResave(const FWorldPartitionActorDescView& ActorDescView) override;
 	virtual void OnLevelInstanceInvalidWorldAsset(const FWorldPartitionActorDescView& ActorDescView, FName WorldAsset, ELevelInstanceInvalidReason Reason) override;
 };
-
diff --git a/Engine/Source/Runtime/Engine/Private/WorldPartition/ErrorHandling/WorldPartitionStreamingGenerationErrorHandler.cpp b/Engine/Source/Runtime/Engine/Private/WorldPartition/ErrorHandling/WorldPartitionStreamingGenerationErrorHandler.cpp
--- a/Engine/Source/Runtime/Engine/Private/WorldPartition/ErrorHandling/WorldPartitionStreamingGenerationErrorHandler.cpp
+++ b/Engine/Source/Runtime/Engine/Private/WorldPartition/ErrorHandling/WorldPartitionStreamingGenerationErrorHandler.cpp
@@ -7,7 +7,9 @@
 // Helpers
 FString IStreamingGenerationErrorHandler::GetFullActorName(const FWorldPartitionActorDescView& ActorDescView)
 {
-	const FSoftObjectPath ActorPath = ActorDescView.GetActorSoftPath();
-	return ActorDescView.GetActorLabelOrName().ToString() + TEXT(" (") + ActorPath.GetLongPackageName() + TEXT(")");
+	const FString ActorPath = ActorDescView.GetActorSoftPath().GetLongPackageName();
+	const FString ActorLabel = ActorDescView.GetActorLabelOrName().ToString();
+	const FString ActorPackage = ActorDescView.GetActorPackage().ToString();
+	return ActorPath + TEXT(".") + ActorLabel + TEXT(" (") + ActorPackage + TEXT(")");
 }
-#endif
+#endif
\ No newline at end of file
diff --git a/Engine/Source/Runtime/Engine/Private/WorldPartition/ErrorHandling/WorldPartitionStreamingGenerationLogErrorHandler.cpp b/Engine/Source/Runtime/Engine/Private/WorldPartition/ErrorHandling/WorldPartitionStreamingGenerationLogErrorHandler.cpp
--- a/Engine/Source/Runtime/Engine/Private/WorldPartition/ErrorHandling/WorldPartitionStreamingGenerationLogErrorHandler.cpp
+++ b/Engine/Source/Runtime/Engine/Private/WorldPartition/ErrorHandling/WorldPartitionStreamingGenerationLogErrorHandler.cpp
@@ -6,9 +6,9 @@
 #include "WorldPartition/DataLayer/DataLayerInstanceWithAsset.h"
 #include "GameFramework/Actor.h"
 
-void FStreamingGenerationLogErrorHandler::OnInvalidReference(const FWorldPartitionActorDescView& ActorDescView, const FGuid& ReferenceGuid)
+void FStreamingGenerationLogErrorHandler::OnInvalidReference(const FWorldPartitionActorDescView& ActorDescView, const FGuid& ReferenceGuid, FWorldPartitionActorDescView* ReferenceActorDescView)
 {
-	UE_LOG(LogWorldPartition, Log, TEXT("Actor %s have missing reference to %s"), *GetFullActorName(ActorDescView), *ReferenceGuid.ToString());
+	UE_LOG(LogWorldPartition, Log, TEXT("Actor %s has an invalid reference to %s"), *GetFullActorName(ActorDescView), ReferenceActorDescView ? *GetFullActorName(*ReferenceActorDescView) : *ReferenceGuid.ToString());
 }
 
 void FStreamingGenerationLogErrorHandler::OnInvalidReferenceGridPlacement(const FWorldPartitionActorDescView& ActorDescView, const FWorldPartitionActorDescView& ReferenceActorDescView)
@@ -74,6 +74,9 @@ void FStreamingGenerationLogErrorHandler::OnLevelInstanceInvalidWorldAsset(const
 	case ELevelInstanceInvalidReason::WorldAssetImcompatiblePartitioned:
 		UE_LOG(LogWorldPartition, Log, TEXT("Actor %s world asset '%s' is partitioned but not marked as compatible"), *ActorName, *WorldAsset.ToString());
 		break;
+	case ELevelInstanceInvalidReason::WorldAssetHasInvalidContainer:
+		UE_LOG(LogWorldPartition, Log, TEXT("Actor %s world asset '%s' has an invalid container"), *ActorName, *WorldAsset.ToString());
+		break;
 	};
 }
 #endif
diff --git a/Engine/Source/Runtime/Engine/Private/WorldPartition/ErrorHandling/WorldPartitionStreamingGenerationTokenizedMessageErrorHandler.cpp b/Engine/Source/Runtime/Engine/Private/WorldPartition/ErrorHandling/WorldPartitionStreamingGenerationTokenizedMessageErrorHandler.cpp
--- a/Engine/Source/Runtime/Engine/Private/WorldPartition/ErrorHandling/WorldPartitionStreamingGenerationTokenizedMessageErrorHandler.cpp
+++ b/Engine/Source/Runtime/Engine/Private/WorldPartition/ErrorHandling/WorldPartitionStreamingGenerationTokenizedMessageErrorHandler.cpp
@@ -11,13 +11,13 @@
 
 #define LOCTEXT_NAMESPACE "WorldPartition"
 
-void ITokenizedMessageErrorHandler::OnInvalidReference(const FWorldPartitionActorDescView& ActorDescView, const FGuid& ReferenceGuid)
+void ITokenizedMessageErrorHandler::OnInvalidReference(const FWorldPartitionActorDescView& ActorDescView, const FGuid& ReferenceGuid, FWorldPartitionActorDescView* ReferenceActorDescView)
 {
-	TSharedRef<FTokenizedMessage> Message = FTokenizedMessage::Create(EMessageSeverity::Warning);
+	TSharedRef<FTokenizedMessage> Message = FTokenizedMessage::Create(EMessageSeverity::Error);
 	Message->AddToken(FTextToken::Create(LOCTEXT("TokenMessage_WorldPartition_Actor", "Actor")))
 		->AddToken(FActorToken::Create(ActorDescView.GetActorSoftPath().ToString(), ActorDescView.GetGuid(), FText::FromString(GetFullActorName(ActorDescView))))
-		->AddToken(FTextToken::Create(LOCTEXT("TokenMessage_WorldPartition_HaveMissingRefsTo", "have missing references to")))
-		->AddToken(FTextToken::Create(FText::FromString(ReferenceGuid.ToString())));
+		->AddToken(FTextToken::Create(LOCTEXT("TokenMessage_WorldPartition_HaveMissingRefsTo", "has an invalid reference to")))
+		->AddToken(FTextToken::Create(ReferenceActorDescView ? FText::FromString(GetFullActorName(*ReferenceActorDescView)) : FText::FromString(ReferenceGuid.ToString())));
 
 	AddAdditionalNameToken(Message, FName(TEXT("WorldPartition_MissingActorReference_CheckForErrors")));
 
@@ -164,6 +164,10 @@ void ITokenizedMessageErrorHandler::OnLevelInstanceInvalidWorldAsset(const FWorl
 		MessageSeverity = EMessageSeverity::Error;
 		ReasonText = LOCTEXT("TokenMessage_WorldPartition_WorldAssetIsPartitionedIncompatible", "is partitioned but not marked as compatible");
 		break;
+	case ELevelInstanceInvalidReason::WorldAssetHasInvalidContainer:
+		MessageSeverity = EMessageSeverity::Error;
+		ReasonText = LOCTEXT("TokenMessage_WorldPartition_WorldAssetHasInvalidContainer", "has an invalid container");
+		break;
 	};
 
 	const FString WorldAssetStr = WorldAsset.ToString();
diff --git a/Engine/Source/Runtime/Engine/Private/WorldPartition/WorldPartitionActorDescView.cpp b/Engine/Source/Runtime/Engine/Private/WorldPartition/WorldPartitionActorDescView.cpp
--- a/Engine/Source/Runtime/Engine/Private/WorldPartition/WorldPartitionActorDescView.cpp
+++ b/Engine/Source/Runtime/Engine/Private/WorldPartition/WorldPartitionActorDescView.cpp
@@ -155,6 +155,11 @@ bool FWorldPartitionActorDescView::IsContainerInstance() const
 	return ActorDesc->IsContainerInstance();
 }
 
+FName FWorldPartitionActorDescView::GetLevelPackage() const
+{
+	return ActorDesc->GetLevelPackage();
+}
+
 bool FWorldPartitionActorDescView::GetContainerInstance(const UActorDescContainer*& OutLevelContainer, FTransform& OutLevelTransform, EContainerClusterMode& OutClusterMode) const
 {
 	return ActorDesc->GetContainerInstance(OutLevelContainer, OutLevelTransform, OutClusterMode);
diff --git a/Engine/Source/Runtime/Engine/Private/WorldPartition/WorldPartitionStreamingGeneration.cpp b/Engine/Source/Runtime/Engine/Private/WorldPartition/WorldPartitionStreamingGeneration.cpp
--- a/Engine/Source/Runtime/Engine/Private/WorldPartition/WorldPartitionStreamingGeneration.cpp
+++ b/Engine/Source/Runtime/Engine/Private/WorldPartition/WorldPartitionStreamingGeneration.cpp
@@ -66,6 +66,12 @@ FWorldPartitionActorDescView* FActorDescViewMap::Emplace(const FGuid& InGuid, co
 	return NewActorDescView;
 }
 
+UWorldPartition::FCheckForErrorsParams::FCheckForErrorsParams()
+	: ErrorHandler(nullptr)
+	, ActorDescContainer(nullptr)
+	, bEnableStreaming(false)
+{}
+
 class FWorldPartitionStreamingGenerator
 {
 	class FStreamingGenerationContext : public IStreamingGenerationContext
@@ -200,6 +206,7 @@ class FWorldPartitionStreamingGenerator
 
 		const UActorDescContainer* Container;
 		FActorDescViewMap ActorDescViewMap;
+		FActorDescViewMap FilteredActorDescViewMap;
 		TArray<FWorldPartitionActorDescView> ContainerInstanceViews;
 		TArray<TArray<FGuid>> Clusters;
 	};
@@ -241,27 +248,7 @@ class FWorldPartitionStreamingGenerator
 		}
 	}
 
-	void ResolveRuntimeReferences(FWorldPartitionActorDescView& ActorDescView, const FActorDescViewMap& ActorDescViewMap)
-	{
-		TArray<FGuid> RuntimeReferences;
-		RuntimeReferences.Reserve(ActorDescView.GetReferences().Num());
-
-		for (const FGuid& ReferenceGuid : ActorDescView.GetReferences())
-		{
-			if (const FWorldPartitionActorDescView* ReferenceDescView = ActorDescViewMap.FindByGuid(ReferenceGuid))
-			{
-				check(!ReferenceDescView->GetActorIsEditorOnly());
-				RuntimeReferences.Add(ReferenceGuid);
-			}
-		}
-
-		if (RuntimeReferences.Num() != ActorDescView.GetReferences().Num())
-		{
-			ActorDescView.SetRuntimeReferences(RuntimeReferences);
-		}
-	}
-
-	void CreateActorDescViewMap(const UActorDescContainer* InContainer, FActorDescViewMap& OutActorDescViewMap, const FActorContainerID& InContainerID, TArray<FWorldPartitionActorDescView>& OutContainerInstances)
+	void CreateActorDescViewMap(const UActorDescContainer* InContainer, FActorDescViewMap& OutActorDescViewMap, FActorDescViewMap& OutFilteredActorDescViewMap, const FActorContainerID& InContainerID, TArray<FWorldPartitionActorDescView>& OutContainerInstances)
 	{
 		// Should we handle unsaved or newly created actors?
 		const bool bHandleUnsavedActors = ModifiedActorsDescList && InContainerID.IsMainContainer();
@@ -351,6 +338,11 @@ class FWorldPartitionStreamingGenerator
 				FWorldPartitionActorDescView ActorDescView(*ActorDescIt);
 				RegisterActorDescView(ActorDescIt->GetGuid(), ActorDescView);
 			}
+			else
+			{
+				FWorldPartitionActorDescView ActorDescView(*ActorDescIt);
+				OutFilteredActorDescViewMap.Emplace(ActorDescIt->GetGuid(), ActorDescView);
+			}
 		}
 
 		// Append new unsaved actors for the persistent level
@@ -377,7 +369,7 @@ class FWorldPartitionStreamingGenerator
 			ContainerDescriptor.Container = InContainer;
 		
 			// Gather actor descriptor views for this container
-			CreateActorDescViewMap(InContainer, ContainerDescriptor.ActorDescViewMap, InContainerID, ContainerDescriptor.ContainerInstanceViews);
+			CreateActorDescViewMap(InContainer, ContainerDescriptor.ActorDescViewMap, ContainerDescriptor.FilteredActorDescViewMap, InContainerID, ContainerDescriptor.ContainerInstanceViews);
 
 			// Resolve actor descriptor views before validation
 			ResolveContainerDescriptor(ContainerDescriptor);
@@ -413,13 +405,7 @@ class FWorldPartitionStreamingGenerator
 
 			if (!ContainerInstanceView.GetContainerInstance(SubContainer, SubTransform, SubClusterMode))
 			{
-				// Don't generate an error when validating changelist because container instances
-				// won't be registered only because the world is not loaded.
-				if (const bool bGenerateError = !bIsChangelistValidation)
-				{
-					//@todo_ow: make a specific error for missing container instance sublevel?
-					ErrorHandler->OnInvalidReference(ContainerInstanceView, FGuid());
-				}
+				ErrorHandler->OnLevelInstanceInvalidWorldAsset(ContainerInstanceView, ContainerInstanceView.GetLevelPackage(), IStreamingGenerationErrorHandler::ELevelInstanceInvalidReason::WorldAssetHasInvalidContainer);
 				continue;
 			}
 
@@ -502,7 +488,6 @@ class FWorldPartitionStreamingGenerator
 		{
 			ResolveRuntimeSpatiallyLoaded(ActorDescView);
 			ResolveRuntimeDataLayers(ActorDescView, ContainerDescriptor.ActorDescViewMap);
-			ResolveRuntimeReferences(ActorDescView, ContainerDescriptor.ActorDescViewMap);
 		};
 
 		ContainerDescriptor.ActorDescViewMap.ForEachActorDescView([this, &ResolveActorDescView](FWorldPartitionActorDescView& ActorDescView)
@@ -633,12 +618,10 @@ class FWorldPartitionStreamingGenerator
 					{
 						// Filter out parent back references
 						FWorldPartitionActorDescView* ReferenceActorDesc = ContainerDescriptor.ActorDescViewMap.FindByGuid(ReferenceGuid);
-						if (ReferenceActorDesc && ReferenceActorDesc->GetParentActor() == ActorDescView.GetGuid())
+						if (!ReferenceActorDesc || (ReferenceActorDesc->GetParentActor() != ActorDescView.GetGuid()))
 						{
-							continue;
+							References.Emplace(FActorReferenceInfo { ActorDescView.GetGuid(), &ActorDescView, ReferenceGuid, ReferenceActorDesc });
 						}
-
-						References.Emplace(FActorReferenceInfo{ ActorDescView.GetGuid(), &ActorDescView, ReferenceGuid, ReferenceActorDesc });
 					}
 				}
 
@@ -665,13 +648,19 @@ class FWorldPartitionStreamingGenerator
 
 				if (TopParentDescView)
 				{
-					References.Emplace(FActorReferenceInfo{ TopParentDescView->GetGuid(), TopParentDescView, ActorDescView.GetGuid(), &ActorDescView });
+					References.Emplace(FActorReferenceInfo { TopParentDescView->GetGuid(), TopParentDescView, ActorDescView.GetGuid(), &ActorDescView });
 				}
 
 				if (ParentGuid.IsValid())
 				{
 					// In case of missing parent add a missing reference 
-					References.Emplace(FActorReferenceInfo{ ActorDescView.GetGuid(), &ActorDescView, ParentGuid, nullptr });
+					References.Emplace(FActorReferenceInfo { ActorDescView.GetGuid(), &ActorDescView, ParentGuid, nullptr });
+				}
+
+				TArray<FGuid> RuntimeReferences;
+				if (NbValidationPasses)
+				{
+					RuntimeReferences.Reserve(ActorDescView.GetReferences().Num());
 				}
 
 				for (FActorReferenceInfo& Info : References)
@@ -727,14 +716,42 @@ class FWorldPartitionStreamingGenerator
 							NbErrorsDetected++;
 						}
 
+						if (NbValidationPasses)
+						{
+							RuntimeReferences.Add(Info.ReferenceGuid);
+						}
 					}
 					else
 					{
-						if (!NbValidationPasses)
+						if (!ContainerDescriptor.FilteredActorDescViewMap.FindByGuid(Info.ReferenceGuid))
 						{
-							ErrorHandler->OnInvalidReference(*RefererActorDescView, Info.ReferenceGuid);
+							if (!NbValidationPasses)
+							{
+								FWorldPartitionActorDescView ReferendceActorDescView;
+								FWorldPartitionActorDescView* ReferendceActorDescViewPtr = nullptr;
+
+								if (const UActorDescContainer** ExistingReferenceContainerPtr = ActorGuidsToContainerMap.Find((Info.ReferenceGuid)))
+								{
+									if (const FWorldPartitionActorDesc* ReferenceActorDesc = (*ExistingReferenceContainerPtr)->GetActorDesc(Info.ReferenceGuid))
+									{
+										ReferendceActorDescView = FWorldPartitionActorDescView(ReferenceActorDesc);
+										ReferendceActorDescViewPtr = &ReferendceActorDescView;
+									}
+								}
+
+								ErrorHandler->OnInvalidReference(*RefererActorDescView, Info.ReferenceGuid, ReferendceActorDescViewPtr);
+							}
 						}
-						// Do not increment NbErrorsDetected since it won't be fixed and thus will always occur
+
+						NbErrorsDetected++;
+					}
+				}
+
+				if (NbValidationPasses)
+				{
+					if (RuntimeReferences.Num() != ActorDescView.GetReferences().Num())
+					{
+						ActorDescView.SetRuntimeReferences(RuntimeReferences);
 					}
 				}
 			});		
@@ -753,12 +770,29 @@ class FWorldPartitionStreamingGenerator
 	}
 
 public:
-	FWorldPartitionStreamingGenerator(FActorDescList* InModifiedActorsDescList, IStreamingGenerationErrorHandler* InErrorHandler, bool bInEnableStreaming, bool bInIsChangelistValidation = false, TArray<TSubclassOf<AActor>> InFilteredClasses = {})
-		: bEnableStreaming(bInEnableStreaming)
-		, bIsChangelistValidation(bInIsChangelistValidation)
-		, ModifiedActorsDescList(InModifiedActorsDescList)
-		, FilteredClasses(InFilteredClasses)
-		, ErrorHandler(InErrorHandler ? InErrorHandler : &NullErrorHandler)	
+	struct FWorldPartitionStreamingGeneratorParams
+	{
+		FWorldPartitionStreamingGeneratorParams()
+			: ModifiedActorsDescList(nullptr)
+			, ErrorHandler(&NullErrorHandler)
+			, bEnableStreaming(false)
+		{}
+
+		FActorDescList* ModifiedActorsDescList;
+		IStreamingGenerationErrorHandler* ErrorHandler;
+		bool bEnableStreaming;
+		TMap<FGuid, const UActorDescContainer*> ActorGuidsToContainerMap;
+		TArray<TSubclassOf<AActor>> FilteredClasses;
+
+		inline static FStreamingGenerationNullErrorHandler NullErrorHandler;
+	};
+
+	FWorldPartitionStreamingGenerator(const FWorldPartitionStreamingGeneratorParams& Params)
+		: bEnableStreaming(Params.bEnableStreaming)
+		, ModifiedActorsDescList(Params.ModifiedActorsDescList)
+		, FilteredClasses(Params.FilteredClasses)
+		, ErrorHandler(Params.ErrorHandler)
+		, ActorGuidsToContainerMap(Params.ActorGuidsToContainerMap)
 	{}
 
 	void PreparationPhase(const UActorDescContainer* Container)
@@ -863,11 +897,9 @@ public:
 
 private:
 	bool bEnableStreaming;
-	bool bIsChangelistValidation;
 	FActorDescList* ModifiedActorsDescList;
 	TArray<TSubclassOf<AActor>> FilteredClasses;
 	IStreamingGenerationErrorHandler* ErrorHandler;
-	FStreamingGenerationNullErrorHandler NullErrorHandler;
 
 	/** Maps containers to their container descriptor */
 	TMap<const UActorDescContainer*, FContainerDescriptor> ContainerDescriptorsMap;
@@ -880,6 +912,9 @@ private:
 
 	/** Data required for streaming generation interface */
 	TUniquePtr<FStreamingGenerationContext> StreamingGenerationContext;
+
+	/** List of containers participating in this streaming generation step */
+	TMap<FGuid, const UActorDescContainer*> ActorGuidsToContainerMap;
 };
 
 bool UWorldPartition::GenerateStreaming(TArray<FString>* OutPackagesToGenerate)
@@ -910,7 +945,12 @@ bool UWorldPartition::GenerateContainerStreaming(const UActorDescContainer* InAc
 	TUniquePtr<FArchive> LogFileAr = FWorldPartitionStreamingGenerator::CreateDumpStateLogArchive(StateLogSuffix);
 	FHierarchicalLogArchive HierarchicalLogAr(*LogFileAr);
 
-	FWorldPartitionStreamingGenerator StreamingGenerator(ModifiedActorsDescList, ErrorHandler, IsStreamingEnabled());
+	FWorldPartitionStreamingGenerator::FWorldPartitionStreamingGeneratorParams StreamingGeneratorParams;
+	StreamingGeneratorParams.ModifiedActorsDescList = ModifiedActorsDescList;
+	StreamingGeneratorParams.ErrorHandler = ErrorHandler;
+	StreamingGeneratorParams.bEnableStreaming = IsStreamingEnabled();
+
+	FWorldPartitionStreamingGenerator StreamingGenerator(StreamingGeneratorParams);
 
 	// Preparation Phase
 	StreamingGenerator.PreparationPhase(InActorDescContainer);
@@ -946,8 +986,13 @@ void UWorldPartition::FlushStreaming()
 void UWorldPartition::GenerateHLOD(ISourceControlHelper* SourceControlHelper, bool bCreateActorsOnly)
 {
 	FStreamingGenerationLogErrorHandler LogErrorHandler;
-	const bool bInIsChangelistValidation = false;
-	FWorldPartitionStreamingGenerator StreamingGenerator(nullptr, &LogErrorHandler, IsStreamingEnabled(), bInIsChangelistValidation, { AWorldPartitionHLOD::StaticClass() });
+
+	FWorldPartitionStreamingGenerator::FWorldPartitionStreamingGeneratorParams StreamingGeneratorParams;
+	StreamingGeneratorParams.ErrorHandler = &LogErrorHandler;
+	StreamingGeneratorParams.bEnableStreaming = IsStreamingEnabled();
+	StreamingGeneratorParams.FilteredClasses.Add(AWorldPartitionHLOD::StaticClass());
+
+	FWorldPartitionStreamingGenerator StreamingGenerator(StreamingGeneratorParams);
 	StreamingGenerator.PreparationPhase(ActorDescContainer);
 
 	TUniquePtr<FArchive> LogFileAr = FWorldPartitionStreamingGenerator::CreateDumpStateLogArchive(TEXT("HLOD"));
@@ -959,19 +1004,52 @@ void UWorldPartition::GenerateHLOD(ISourceControlHelper* SourceControlHelper, bo
 
 void UWorldPartition::CheckForErrors(IStreamingGenerationErrorHandler* ErrorHandler) const
 {
-	const bool bIsChangelistValidation = false;
-	ForEachActorDescContainer([this, ErrorHandler, bIsChangelistValidation](const UActorDescContainer* InActorDescContainer)
+	FCheckForErrorsParams Params;
+	Params.ErrorHandler = ErrorHandler;
+	Params.ActorDescContainer = ActorDescContainer;
+	Params.bEnableStreaming = IsStreamingEnabled();
+
+	ForEachActorDescContainer([&Params](const UActorDescContainer* InActorDescContainer)
 	{
-		CheckForErrors(ErrorHandler, InActorDescContainer, IsStreamingEnabled(), bIsChangelistValidation);
+		for (FActorDescList::TConstIterator<> ActorDescIt(InActorDescContainer); ActorDescIt; ++ActorDescIt)
+		{
+			check(!Params.ActorGuidsToContainerMap.Contains(ActorDescIt->GetGuid()));
+			Params.ActorGuidsToContainerMap.Add(ActorDescIt->GetGuid(), InActorDescContainer);
+		}
 	});
+
+	ForEachActorDescContainer([this, &Params](const UActorDescContainer* InActorDescContainer)
+	{
+		Params.ActorDescContainer = InActorDescContainer;
+		CheckForErrors(Params);
+	});
+}
+
+void UWorldPartition::CheckForErrors(IStreamingGenerationErrorHandler* ErrorHandler, const UActorDescContainer* ActorDescContainer, bool bEnableStreaming, bool)
+{
+	FCheckForErrorsParams Params;
+	Params.ErrorHandler = ErrorHandler;
+	Params.ActorDescContainer = ActorDescContainer;
+	Params.bEnableStreaming = bEnableStreaming;
+
+	CheckForErrors(Params);
 }
 
-void UWorldPartition::CheckForErrors(IStreamingGenerationErrorHandler* ErrorHandler, const UActorDescContainer* ActorDescContainer, bool bEnableStreaming, bool bIsChangelistValidation)
+void UWorldPartition::CheckForErrors(const FCheckForErrorsParams& Params)
 {
 	FActorDescList ModifiedActorDescList;
-	FWorldPartitionStreamingGenerator StreamingGenerator(ActorDescContainer->GetWorld() ? &ModifiedActorDescList : nullptr, ErrorHandler, bEnableStreaming, bIsChangelistValidation);
-	StreamingGenerator.PreparationPhase(ActorDescContainer);
+
+	FWorldPartitionStreamingGenerator::FWorldPartitionStreamingGeneratorParams StreamingGeneratorParams;
+	StreamingGeneratorParams.ModifiedActorsDescList = Params.ActorDescContainer->GetWorld() ? &ModifiedActorDescList : nullptr;
+	StreamingGeneratorParams.ErrorHandler = Params.ErrorHandler;
+	StreamingGeneratorParams.bEnableStreaming = Params.bEnableStreaming;
+	StreamingGeneratorParams.ActorGuidsToContainerMap = Params.ActorGuidsToContainerMap;
+
+	FWorldPartitionStreamingGenerator StreamingGenerator(StreamingGeneratorParams);
+
+	StreamingGenerator.PreparationPhase(Params.ActorDescContainer);
 }
+
 #endif // WITH_EDITOR
 
-#undef LOCTEXT_NAMESPACE
+#undef LOCTEXT_NAMESPACE
\ No newline at end of file
diff --git a/Engine/Source/Runtime/Engine/Public/WorldPartition/ErrorHandling/WorldPartitionStreamingGenerationErrorHandler.h b/Engine/Source/Runtime/Engine/Public/WorldPartition/ErrorHandling/WorldPartitionStreamingGenerationErrorHandler.h
--- a/Engine/Source/Runtime/Engine/Public/WorldPartition/ErrorHandling/WorldPartitionStreamingGenerationErrorHandler.h
+++ b/Engine/Source/Runtime/Engine/Public/WorldPartition/ErrorHandling/WorldPartitionStreamingGenerationErrorHandler.h
@@ -13,10 +13,13 @@ class ENGINE_API IStreamingGenerationErrorHandler
 public:
 	virtual ~IStreamingGenerationErrorHandler() {}
 
+	UE_DEPRECATED(5.2, "OnInvalidReference is deprecated, use the version which takes an optional actor descriptor view.")
+	virtual void OnInvalidReference(const FWorldPartitionActorDescView& ActorDescView, const FGuid& ReferenceGuid) {}
+
 	/** 
 	 * Called when an actor references an invalid actor.
 	 */
-	virtual void OnInvalidReference(const FWorldPartitionActorDescView& ActorDescView, const FGuid& ReferenceGuid) = 0;
+	virtual void OnInvalidReference(const FWorldPartitionActorDescView& ActorDescView, const FGuid& ReferenceGuid, FWorldPartitionActorDescView* ReferenceActorDescView) = 0;
 
 	/** 
 	 * Called when an actor references an actor using a different grid placement.
@@ -70,7 +73,8 @@ public:
 	{
 		WorldAssetNotFound,
 		WorldAssetNotUsingExternalActors,
-		WorldAssetImcompatiblePartitioned
+		WorldAssetImcompatiblePartitioned,
+		WorldAssetHasInvalidContainer
 	};
 
 	virtual void OnLevelInstanceInvalidWorldAsset(const FWorldPartitionActorDescView& ActorDescView, FName WorldAsset, ELevelInstanceInvalidReason Reason) = 0;
diff --git a/Engine/Source/Runtime/Engine/Public/WorldPartition/ErrorHandling/WorldPartitionStreamingGenerationLogErrorHandler.h b/Engine/Source/Runtime/Engine/Public/WorldPartition/ErrorHandling/WorldPartitionStreamingGenerationLogErrorHandler.h
--- a/Engine/Source/Runtime/Engine/Public/WorldPartition/ErrorHandling/WorldPartitionStreamingGenerationLogErrorHandler.h
+++ b/Engine/Source/Runtime/Engine/Public/WorldPartition/ErrorHandling/WorldPartitionStreamingGenerationLogErrorHandler.h
@@ -7,7 +7,7 @@
 class ENGINE_API FStreamingGenerationLogErrorHandler : public IStreamingGenerationErrorHandler
 {
 public:
-	virtual void OnInvalidReference(const FWorldPartitionActorDescView& ActorDescView, const FGuid& ReferenceGuid) override;
+	virtual void OnInvalidReference(const FWorldPartitionActorDescView& ActorDescView, const FGuid& ReferenceGuid, FWorldPartitionActorDescView* ReferenceActorDescView) override;
 	virtual void OnInvalidReferenceGridPlacement(const FWorldPartitionActorDescView& ActorDescView, const FWorldPartitionActorDescView& ReferenceActorDescView) override;
 	virtual void OnInvalidReferenceDataLayers(const FWorldPartitionActorDescView& ActorDescView, const FWorldPartitionActorDescView& ReferenceActorDescView) override;
 	virtual void OnInvalidReferenceLevelScriptStreamed(const FWorldPartitionActorDescView& ActorDescView) override;
diff --git a/Engine/Source/Runtime/Engine/Public/WorldPartition/ErrorHandling/WorldPartitionStreamingGenerationNullErrorHandler.h b/Engine/Source/Runtime/Engine/Public/WorldPartition/ErrorHandling/WorldPartitionStreamingGenerationNullErrorHandler.h
--- a/Engine/Source/Runtime/Engine/Public/WorldPartition/ErrorHandling/WorldPartitionStreamingGenerationNullErrorHandler.h
+++ b/Engine/Source/Runtime/Engine/Public/WorldPartition/ErrorHandling/WorldPartitionStreamingGenerationNullErrorHandler.h
@@ -7,7 +7,7 @@
 class ENGINE_API FStreamingGenerationNullErrorHandler : public IStreamingGenerationErrorHandler
 {
 public:
-	virtual void OnInvalidReference(const FWorldPartitionActorDescView& ActorDescView, const FGuid& ReferenceGuid) override {}
+	virtual void OnInvalidReference(const FWorldPartitionActorDescView& ActorDescView, const FGuid& ReferenceGuid, FWorldPartitionActorDescView* ReferenceActorDescView) override {}
 	virtual void OnInvalidReferenceGridPlacement(const FWorldPartitionActorDescView& ActorDescView, const FWorldPartitionActorDescView& ReferenceActorDescView) override {}
 	virtual void OnInvalidReferenceDataLayers(const FWorldPartitionActorDescView& ActorDescView, const FWorldPartitionActorDescView& ReferenceActorDescView) override {}
 	virtual void OnInvalidReferenceLevelScriptStreamed(const FWorldPartitionActorDescView& ActorDescView) override {}
diff --git a/Engine/Source/Runtime/Engine/Public/WorldPartition/ErrorHandling/WorldPartitionStreamingGenerationTokenizedMessageErrorHandler.h b/Engine/Source/Runtime/Engine/Public/WorldPartition/ErrorHandling/WorldPartitionStreamingGenerationTokenizedMessageErrorHandler.h
--- a/Engine/Source/Runtime/Engine/Public/WorldPartition/ErrorHandling/WorldPartitionStreamingGenerationTokenizedMessageErrorHandler.h
+++ b/Engine/Source/Runtime/Engine/Public/WorldPartition/ErrorHandling/WorldPartitionStreamingGenerationTokenizedMessageErrorHandler.h
@@ -8,7 +8,7 @@ class FTokenizedMessage;
 
 class ENGINE_API ITokenizedMessageErrorHandler : public IStreamingGenerationErrorHandler
 {
-	virtual void OnInvalidReference(const FWorldPartitionActorDescView& ActorDescView, const FGuid& ReferenceGuid) override;
+	virtual void OnInvalidReference(const FWorldPartitionActorDescView& ActorDescView, const FGuid& ReferenceGuid, FWorldPartitionActorDescView* ReferenceActorDescView) override;
 	virtual void OnInvalidReferenceGridPlacement(const FWorldPartitionActorDescView& ActorDescView, const FWorldPartitionActorDescView& ReferenceActorDescView) override;
 	virtual void OnInvalidReferenceDataLayers(const FWorldPartitionActorDescView& ActorDescView, const FWorldPartitionActorDescView& ReferenceActorDescView) override;
 	virtual void OnInvalidReferenceLevelScriptStreamed(const FWorldPartitionActorDescView& ActorDescView) override;
diff --git a/Engine/Source/Runtime/Engine/Public/WorldPartition/WorldPartition.h b/Engine/Source/Runtime/Engine/Public/WorldPartition/WorldPartition.h
--- a/Engine/Source/Runtime/Engine/Public/WorldPartition/WorldPartition.h
+++ b/Engine/Source/Runtime/Engine/Public/WorldPartition/WorldPartition.h
@@ -190,8 +190,23 @@ public:
 	void DumpActorDescs(const FString& Path);
 
 	void CheckForErrors(IStreamingGenerationErrorHandler* ErrorHandler) const;
+
+	/* Struct of optional parameters passed to check for errors function. */
+	struct ENGINE_API FCheckForErrorsParams
+	{
+		FCheckForErrorsParams();
+
+		IStreamingGenerationErrorHandler* ErrorHandler;
+		const UActorDescContainer* ActorDescContainer;
+		bool bEnableStreaming;
+		TMap<FGuid, const UActorDescContainer*> ActorGuidsToContainerMap;
+	};
+
+	UE_DEPRECATED(5.2, "CheckForErrors is deprecated, CheckForErrors with FCheckForErrorsParams should be used instead.")
 	static void CheckForErrors(IStreamingGenerationErrorHandler* ErrorHandler, const UActorDescContainer* ActorDescContainer, bool bEnableStreaming, bool bIsChangelistValidation);
 
+	static void CheckForErrors(const FCheckForErrorsParams& Params);
+
 	void AppendAssetRegistryTags(TArray<FAssetRegistryTag>& OutTags) const;
 
 	DECLARE_MULTICAST_DELEGATE_OneParam(FActorDescContainerRegistrationDelegate, UActorDescContainer*);
diff --git a/Engine/Source/Runtime/Engine/Public/WorldPartition/WorldPartitionActorDescView.h b/Engine/Source/Runtime/Engine/Public/WorldPartition/WorldPartitionActorDescView.h
--- a/Engine/Source/Runtime/Engine/Public/WorldPartition/WorldPartitionActorDescView.h
+++ b/Engine/Source/Runtime/Engine/Public/WorldPartition/WorldPartitionActorDescView.h
@@ -59,6 +59,7 @@ PRAGMA_ENABLE_DEPRECATION_WARNINGS
 	FGuid GetContentBundleGuid() const;
 
 	bool IsContainerInstance() const;
+	FName GetLevelPackage() const;
 	bool GetContainerInstance(const UActorDescContainer*& OutLevelContainer, FTransform& OutLevelTransform, EContainerClusterMode& OutClusterMode) const;
 
 	void CheckForErrors(IStreamingGenerationErrorHandler* ErrorHandler) const;
