Changed world partition streaming level names containing a content bundle from "<cellname>CB<id>" to "<cellname>_CB<id>for readability.

diff --git a/Engine/Source/Runtime/Engine/Private/WorldPartition/WorldPartitionRuntimeSpatialHash.cpp b/Engine/Source/Runtime/Engine/Private/WorldPartition/WorldPartitionRuntimeSpatialHash.cpp
--- a/Engine/Source/Runtime/Engine/Private/WorldPartition/WorldPartitionRuntimeSpatialHash.cpp
+++ b/Engine/Source/Runtime/Engine/Private/WorldPartition/WorldPartitionRuntimeSpatialHash.cpp
@@ -1297,7 +1297,7 @@ FString UWorldPartitionRuntimeSpatialHash::GetCellNameString(FName InGridName, c
 
 	if (InContentBundleID.IsValid())
 	{
-		CellName += FString::Printf(TEXT("CB%s"), *UContentBundleDescriptor::GetContentBundleCompactString(InContentBundleID));
+		CellName += FString::Printf(TEXT("_CB%s"), *UContentBundleDescriptor::GetContentBundleCompactString(InContentBundleID));
 	}
 
 	return CellName;
