Removed an invalid warning in ULevelStreamingProfilingSubsystem.

diff --git a/Engine/Source/Runtime/Engine/Private/ProfilingDebugging/LevelStreamingProfilingSubsystem.cpp b/Engine/Source/Runtime/Engine/Private/ProfilingDebugging/LevelStreamingProfilingSubsystem.cpp
--- a/Engine/Source/Runtime/Engine/Private/ProfilingDebugging/LevelStreamingProfilingSubsystem.cpp
+++ b/Engine/Source/Runtime/Engine/Private/ProfilingDebugging/LevelStreamingProfilingSubsystem.cpp
@@ -250,6 +250,8 @@ void ULevelStreamingProfilingSubsystem::OnLevelStreamingTargetStateChanged(UWorl
 			return OnLevelQueuedForRemoveFromWorld(World, StreamingLevel, LevelIfLoaded);
 		case ELevelStreamingState::Unloaded:
 			return OnLevelQueuedForLoading(World, StreamingLevel);
+		case ELevelStreamingState::MakingVisible:
+			return;
 		default:
 			UE_LOG(LogLevelStreamingProfiling, Warning, TEXT("Unexpected desired state LoadedNotVisible in current state %s"), ::EnumToString(CurrentState));
 			return;
