Fix a bug with Nanite static mesh caching where LOD0 could have a single material section and change materials, but the cache would not be invalidated.

#rnx

diff --git a/Engine/Source/Runtime/Engine/Private/StaticMesh.cpp b/Engine/Source/Runtime/Engine/Private/StaticMesh.cpp
--- a/Engine/Source/Runtime/Engine/Private/StaticMesh.cpp
+++ b/Engine/Source/Runtime/Engine/Private/StaticMesh.cpp
@@ -2701,16 +2701,13 @@ static FString BuildStaticMeshDerivedDataKeySuffix(const ITargetPlatform* Target
 		KeySuffix += TEXT("_MinMLOD");
 	}
 
-	// Append the section material slot mappings for LOD0, if there are two or more (and thus are permutable), as they are baked into the Nanite build.
+	// Append the section material slot mappings for LOD0, as they are baked into the Nanite build.	
 	const FMeshSectionInfoMap& SectionInfoMap = Mesh->GetSectionInfoMap();
 	int32 NumLOD0Sections = SectionInfoMap.GetSectionNumber(0);
-	if (NumLOD0Sections >= 2)
+	KeySuffix += TEXT("_");
+	for (int32 SectionIndex = 0; SectionIndex < NumLOD0Sections; SectionIndex++)
 	{
-		KeySuffix += TEXT("_");
-		for (int32 SectionIndex = 0; SectionIndex < NumLOD0Sections; SectionIndex++)
-		{
-			KeySuffix += LexToString(SectionInfoMap.Get(0, SectionIndex).MaterialIndex);
-		}
+		KeySuffix += LexToString(SectionInfoMap.Get(0, SectionIndex).MaterialIndex);
 	}
 
 #if PLATFORM_CPU_ARM_FAMILY
