Fix audio derived data invalidating old data.
This fix reverts the logic so that old sounds don't go through a pcm->float->pcm conversion.

diff --git a/Engine/Source/Developer/AudioFormatADPCM/Private/AudioFormatADPCM.cpp b/Engine/Source/Developer/AudioFormatADPCM/Private/AudioFormatADPCM.cpp
--- a/Engine/Source/Developer/AudioFormatADPCM/Private/AudioFormatADPCM.cpp
+++ b/Engine/Source/Developer/AudioFormatADPCM/Private/AudioFormatADPCM.cpp
@@ -416,7 +416,7 @@ class FAudioFormatADPCM : public IAudioFormat
 	enum
 	{
 		/** Version for ADPCM format, this becomes part of the DDC key. */
-		UE_AUDIO_ADPCM_VER = 3,
+		UE_AUDIO_ADPCM_VER = 4,
 	};
 
 	void InterleaveBuffers(const TArray<TArray<uint8> >& SrcBuffers, TArray<uint8> & InterleavedBuffer) const
diff --git a/Engine/Source/Developer/AudioFormatOgg/Private/AudioFormatOgg.cpp b/Engine/Source/Developer/AudioFormatOgg/Private/AudioFormatOgg.cpp
--- a/Engine/Source/Developer/AudioFormatOgg/Private/AudioFormatOgg.cpp
+++ b/Engine/Source/Developer/AudioFormatOgg/Private/AudioFormatOgg.cpp
@@ -32,7 +32,7 @@ class FAudioFormatOgg : public IAudioFormat
 	enum
 	{
 		/** Version for OGG format, this becomes part of the DDC key. */
-		UE_AUDIO_OGG_VER = 5,
+		UE_AUDIO_OGG_VER = 6,
 	};
 
 public:
diff --git a/Engine/Source/Developer/AudioFormatOpus/Private/AudioFormatOpus.cpp b/Engine/Source/Developer/AudioFormatOpus/Private/AudioFormatOpus.cpp
--- a/Engine/Source/Developer/AudioFormatOpus/Private/AudioFormatOpus.cpp
+++ b/Engine/Source/Developer/AudioFormatOpus/Private/AudioFormatOpus.cpp
@@ -31,7 +31,7 @@ class FAudioFormatOpus : public IAudioFormat
 	enum
 	{
 		/** Version for OPUS format, this becomes part of the DDC key. */
-		UE_AUDIO_OPUS_VER = 7,
+		UE_AUDIO_OPUS_VER = 8,
 	};
 
 public:
diff --git a/Engine/Source/Runtime/Engine/Private/AudioDerivedData.cpp b/Engine/Source/Runtime/Engine/Private/AudioDerivedData.cpp
--- a/Engine/Source/Runtime/Engine/Private/AudioDerivedData.cpp
+++ b/Engine/Source/Runtime/Engine/Private/AudioDerivedData.cpp
@@ -56,6 +56,43 @@ namespace AudioCookStats
 
 #if WITH_EDITORONLY_DATA
 
+namespace AudioDerivedDataPrivate
+{
+	constexpr float FloatToPcm16Scalar = 32767.f;
+
+	// This function for converting pcm16 to float is used so that existing 
+	// assets cook to the exact same bit-wise result. While it would be nice to 
+	// replace the divide operator with a multiply operator in the for loop, this 
+	// should not be done as it produces different results.
+	void ArrayPcm16ToFloat(TArrayView<const int16> InView, TArrayView<float> OutView)
+	{
+		check(InView.Num() == OutView.Num());
+
+		const int16* InData = InView.GetData();
+		float* OutData = OutView.GetData();
+		const uint32 Num = InView.Num();
+
+		for (uint32 i = 0; i < Num; i++)
+		{
+			OutData[i] = (float)InData[i] / FloatToPcm16Scalar;
+		}
+	}
+
+	void ArrayFloatToPcm16(TArrayView<const float> InView, TArrayView<int16> OutView)
+	{
+		check(InView.Num() == OutView.Num());
+
+		const float* InData = InView.GetData();
+		int16* OutData = OutView.GetData();
+		const uint32 Num = InView.Num();
+
+		for (uint32 i = 0; i < Num; i++)
+		{
+			OutData[i] = (int16)(InData[i] * FloatToPcm16Scalar);
+		}
+	}
+}
+
 // Any thread implicated in the streamed audio platform data build must have a valid scope to be granted access
 // to properties being modified by the build itself without triggering a FinishCache.
 // Any other thread that is a consumer of the FStreamedAudioPlatformData will trigger a FinishCache when accessing
@@ -136,7 +173,7 @@ Derived data key generation.
 
 // If you want to bump this version, generate a new guid using
 // VS->Tools->Create GUID and paste it here. https://www.guidgen.com works too.
-#define STREAMEDAUDIO_DERIVEDDATA_VER		TEXT("BC6E92FBBD314E3B9B9EC6778749EB5E")
+#define STREAMEDAUDIO_DERIVEDDATA_VER		TEXT("0E2C7011AFE845488C383F8A3531A3F1")
 
 /**
  * Computes the derived data key suffix for a SoundWave's Streamed Audio.
@@ -1301,6 +1338,10 @@ struct FAudioCookInputs
  */
 static void CookSimpleWave(const FAudioCookInputs& Inputs, TArray<uint8>& OutputBuffer)
 {
+	// Warning: Existing released assets should maintain bitwise exact encoded audio
+	// in order to minimize patch sizes. Changing anything in this function can 
+	// change the final encoded values and result in large unintended patches. 
+	
 	TRACE_CPUPROFILER_EVENT_SCOPE(CookSimpleWave);
 
 	FWaveModInfo WaveInfo;
@@ -1342,58 +1383,68 @@ static void CookSimpleWave(const FAudioCookInputs& Inputs, TArray<uint8>& Output
 	int32 NumBytes = Input.Num();
 	int32 NumSamples = NumBytes / sizeof(int16);
 
-	// To float for processing
-	Audio::FAlignedFloatBuffer InputFloatBuffer;
-	InputFloatBuffer.SetNumUninitialized(NumSamples);
-	
-	Audio::ArrayPcm16ToFloat(MakeArrayView((int16*)Input.GetData(), NumSamples), InputFloatBuffer);
+	const bool bNeedsResample = (Inputs.SampleRateOverride > 0 && Inputs.SampleRateOverride != (float)WaveSampleRate);
+	const bool bNeedsToApplyWaveTransformation = (Inputs.WaveTransformations.Num() > 0);
 
-	// Run any transformations
-	if(Inputs.WaveTransformations.Num())
-	{
-		Audio::FWaveformTransformationWaveInfo TransformationInfo;
+	// Only convert PCM wave data to float if needed. The conversion alters the sample
+	// values enough to produce different results in the final encoded data. 
+	const bool bNeedsFloatConversion = bNeedsResample || bNeedsToApplyWaveTransformation;
 
-		TransformationInfo.Audio = &InputFloatBuffer;
-		TransformationInfo.NumChannels = NumChannels;
-		TransformationInfo.SampleRate = WaveSampleRate;
+	if (bNeedsFloatConversion)
+	{
+		// To float for processing
+		Audio::FAlignedFloatBuffer InputFloatBuffer;
+		InputFloatBuffer.SetNumUninitialized(NumSamples);
 		
-		for(const Audio::FTransformationPtr& Transformation : Inputs.WaveTransformations)
+		AudioDerivedDataPrivate::ArrayPcm16ToFloat(MakeArrayView((int16*)Input.GetData(), NumSamples), InputFloatBuffer);
+
+		// Run any transformations
+		if(bNeedsToApplyWaveTransformation)
 		{
-			Transformation->ProcessAudio(TransformationInfo);
-		}
+			Audio::FWaveformTransformationWaveInfo TransformationInfo;
+
+			TransformationInfo.Audio = &InputFloatBuffer;
+			TransformationInfo.NumChannels = NumChannels;
+			TransformationInfo.SampleRate = WaveSampleRate;
+			
+			for(const Audio::FTransformationPtr& Transformation : Inputs.WaveTransformations)
+			{
+				Transformation->ProcessAudio(TransformationInfo);
+			}
 
-		UE_CLOG(WaveSampleRate != TransformationInfo.SampleRate, LogAudioDerivedData, Warning, TEXT("Wave transformations which alter the sample rate are not supported. Cooked audio for %s may be incorrect"), *Inputs.SoundFullName);
-		WaveSampleRate = TransformationInfo.SampleRate;
+			UE_CLOG(WaveSampleRate != TransformationInfo.SampleRate, LogAudioDerivedData, Warning, TEXT("Wave transformations which alter the sample rate are not supported. Cooked audio for %s may be incorrect"), *Inputs.SoundFullName);
+			WaveSampleRate = TransformationInfo.SampleRate;
 
-		UE_CLOG(NumChannels != TransformationInfo.NumChannels, LogAudioDerivedData, Error, TEXT("Wave transformations which alter number of channels are not supported. Cooked audio for %s may be incorrect"), *Inputs.SoundFullName);
-		NumChannels = TransformationInfo.NumChannels;
+			UE_CLOG(NumChannels != TransformationInfo.NumChannels, LogAudioDerivedData, Error, TEXT("Wave transformations which alter number of channels are not supported. Cooked audio for %s may be incorrect"), *Inputs.SoundFullName);
+			NumChannels = TransformationInfo.NumChannels;
 
-		NumSamples = InputFloatBuffer.Num();
-	}
-	
-	// Check for a platform resample override here and resample if necessary:
-	if (Inputs.SampleRateOverride > 0 && Inputs.SampleRateOverride != (float)WaveSampleRate)
-	{
-		ResampleWaveData(InputFloatBuffer, NumChannels, WaveSampleRate, Inputs.SampleRateOverride);
+			NumSamples = InputFloatBuffer.Num();
+		}
 		
-		WaveSampleRate = Inputs.SampleRateOverride;
-		NumSamples = InputFloatBuffer.Num();
-	}
+		// Resample if necessary
+		if (bNeedsResample)
+		{
+			ResampleWaveData(InputFloatBuffer, NumChannels, WaveSampleRate, Inputs.SampleRateOverride);
+			
+			WaveSampleRate = Inputs.SampleRateOverride;
+			NumSamples = InputFloatBuffer.Num();
+		}
 
-	// Clip Normalize
-	const float MaxValue = Audio::ArrayMaxAbsValue(InputFloatBuffer);
-	if (MaxValue > 1.0f)
-	{
-		UE_LOG(LogAudioDerivedData, Display, TEXT("Audio clipped during cook: This asset will be normalized by a factor of 1/%f. Consider attenuating the above asset."), MaxValue);
+		// Clip Normalize
+		const float MaxValue = Audio::ArrayMaxAbsValue(InputFloatBuffer);
+		if (MaxValue > 1.0f)
+		{
+			UE_LOG(LogAudioDerivedData, Display, TEXT("Audio clipped during cook: This asset will be normalized by a factor of 1/%f. Consider attenuating the above asset."), MaxValue);
 
-		Audio::ArrayMultiplyByConstantInPlace(InputFloatBuffer, 1.f / MaxValue);
-	}
+			Audio::ArrayMultiplyByConstantInPlace(InputFloatBuffer, 1.f / MaxValue);
+		}
 
-	// back to PCM
-	NumBytes = NumSamples * sizeof(int16);
-	Input.SetNumUninitialized(NumBytes);
-	
-	Audio::ArrayFloatToPcm16(InputFloatBuffer, MakeArrayView((int16*)Input.GetData(), NumSamples));
+		// back to PCM
+		NumBytes = NumSamples * sizeof(int16);
+		Input.SetNumUninitialized(NumBytes);
+		
+		AudioDerivedDataPrivate::ArrayFloatToPcm16(InputFloatBuffer, MakeArrayView((int16*)Input.GetData(), NumSamples));
+	}
 
 	// Compression Quality
 	FSoundQualityInfo QualityInfo = { 0 };
@@ -1435,6 +1486,10 @@ static void CookSimpleWave(const FAudioCookInputs& Inputs, TArray<uint8>& Output
  */
 static void CookSurroundWave(const FAudioCookInputs& Inputs,  TArray<uint8>& OutputBuffer)
 {
+	// Warning: Existing released assets should maintain bitwise exact encoded audio
+	// in order to minimize patch sizes. Changing anything in this function can 
+	// change the final encoded values and result in large unintended patches. 
+
 	TRACE_CPUPROFILER_EVENT_SCOPE(CookSurroundWave);
 
 	check(!OutputBuffer.Num());
@@ -1583,7 +1638,7 @@ static void CookSurroundWave(const FAudioCookInputs& Inputs,  TArray<uint8>& Out
 		// convert to float
 		for (int32 ChannelIndex = 0; ChannelIndex < ChannelCount; ChannelIndex++)
 		{
-			Audio::ArrayPcm16ToFloat(MakeArrayView((const int16*)(SourceBuffers[ChannelIndex].GetData()), NumFrames), InputMultichannelBuffer[ChannelIndex]);
+			AudioDerivedDataPrivate::ArrayPcm16ToFloat(MakeArrayView((const int16*)(SourceBuffers[ChannelIndex].GetData()), NumFrames), InputMultichannelBuffer[ChannelIndex]);
 		}
 	
 		Audio::ArrayInterleave(InputMultichannelBuffer, InterleavedFloatBuffer);
@@ -1640,7 +1695,7 @@ static void CookSurroundWave(const FAudioCookInputs& Inputs,  TArray<uint8>& Out
 				
 			PcmBuffer.SetNum(SampleDataSize);
 			
-			Audio::ArrayFloatToPcm16(InputMultichannelBuffer[ChannelIndex], MakeArrayView((int16*)PcmBuffer.GetData(), NumFrames));
+			AudioDerivedDataPrivate::ArrayFloatToPcm16(InputMultichannelBuffer[ChannelIndex], MakeArrayView((int16*)PcmBuffer.GetData(), NumFrames));
 		}
 	}
 
