Add CVar in the SGameLayerManager to remove alpha multiplication on disabled game-layer widgets.

diff --git a/Engine/Source/Runtime/Engine/Private/Slate/SGameLayerManager.cpp b/Engine/Source/Runtime/Engine/Private/Slate/SGameLayerManager.cpp
--- a/Engine/Source/Runtime/Engine/Private/Slate/SGameLayerManager.cpp
+++ b/Engine/Source/Runtime/Engine/Private/Slate/SGameLayerManager.cpp
@@ -40,6 +40,13 @@ static FAutoConsoleVariableRef CVarEnablePerUserHitTesting(
 	FConsoleVariableDelegate::CreateStatic(&HandlePerUserHitTestingToggled)
 );
 
+static int32 bApplyDisabledEffectOnWidgets = 1;
+
+static FAutoConsoleVariableRef CVarApplyDisabledEffectOnWidgets(
+	TEXT("ApplyDisabledEffectOnWidgets"),
+	bApplyDisabledEffectOnWidgets,
+	TEXT("If true, disabled game-layer widgets will have alpha multiplied by 0.45."));
+
 #if UE_SLATE_WITH_GAMELAYER_CANVAS_VISIBILITY_COMMANDS
 namespace CanvasVisibility
 {
@@ -384,7 +391,11 @@ void SGameLayerManager::Tick(const FGeometry& AllottedGeometry, const double InC
 int32 SGameLayerManager::OnPaint(const FPaintArgs& Args, const FGeometry& AllottedGeometry, const FSlateRect& MyCullingRect, FSlateWindowElementList& OutDrawElements, int32 LayerId, const FWidgetStyle& InWidgetStyle, bool bParentEnabled) const
 {
 	SCOPED_NAMED_EVENT_TEXT("Paint: Game UI", FColor::Green);
+	
+	// By default game layer state is used to ignore disabled effects, don't set this flag if user wants disabled draw effects in-game
+	OutDrawElements.SetIsInGameLayer(!bApplyDisabledEffectOnWidgets);
 	const int32 ResultLayer = SCompoundWidget::OnPaint(Args, AllottedGeometry, MyCullingRect, OutDrawElements, LayerId, InWidgetStyle, bParentEnabled);
+	OutDrawElements.SetIsInGameLayer(false);
 	return ResultLayer;
 }
 
diff --git a/Engine/Source/Runtime/SlateCore/Private/Rendering/DrawElements.cpp b/Engine/Source/Runtime/SlateCore/Private/Rendering/DrawElements.cpp
--- a/Engine/Source/Runtime/SlateCore/Private/Rendering/DrawElements.cpp
+++ b/Engine/Source/Runtime/SlateCore/Private/Rendering/DrawElements.cpp
@@ -146,6 +146,7 @@ FSlateWindowElementList::FSlateWindowElementList(const TSharedPtr<SWindow>& InPa
 	, bNeedsDeferredResolve(false)
 	, ResolveToDeferredIndex()
 	, WindowSize(FVector2f(0.0f, 0.0f))
+	, bIsInGameLayer(false)
 	//, bReportReferences(true)
 {
 	if (InPaintWindow.IsValid())
@@ -171,6 +172,11 @@ FSlateWindowElementList::~FSlateWindowElementList()
 
 void FSlateDrawElement::Init(FSlateWindowElementList& ElementList, EElementType InElementType, uint32 InLayer, const FPaintGeometry& PaintGeometry, ESlateDrawEffect InDrawEffects)
 {
+	if (ElementList.GetIsInGameLayer())
+	{
+		InDrawEffects &= ~ESlateDrawEffect::DisabledEffect;
+	}
+
 	RenderTransform = PaintGeometry.GetAccumulatedRenderTransform();
 	Position = PaintGeometry.DrawPosition;
 	Scale = PaintGeometry.DrawScale;
@@ -813,6 +819,15 @@ void FSlateDrawElement::AddReferencedObjects(FReferenceCollector& Collector)
 	}
 }
 
+void FSlateWindowElementList::SetIsInGameLayer(bool bInGameLayer)
+{
+	bIsInGameLayer = bInGameLayer;
+}
+
+bool FSlateWindowElementList::GetIsInGameLayer()
+{
+	return bIsInGameLayer;
+}
 
 FSlateDrawElement& FSlateWindowElementList::AddUninitialized()
 {
diff --git a/Engine/Source/Runtime/SlateCore/Public/Rendering/DrawElements.h b/Engine/Source/Runtime/SlateCore/Public/Rendering/DrawElements.h
--- a/Engine/Source/Runtime/SlateCore/Public/Rendering/DrawElements.h
+++ b/Engine/Source/Runtime/SlateCore/Public/Rendering/DrawElements.h
@@ -691,6 +691,9 @@ public:
 	SLATECORE_API void SetRenderTargetWindow(SWindow* InRenderTargetWindow);
 
 	void AddReferencedObjects(FReferenceCollector& Collector);
+	
+	SLATECORE_API void SetIsInGameLayer(bool bInGameLayer);
+	SLATECORE_API bool GetIsInGameLayer();
 
 private:
 	FSlateDrawElement& AddCachedElement();
@@ -760,4 +763,7 @@ private:
 
 	/** Store the size of the window being used to paint */
 	FVector2f WindowSize;
+
+	/** Store if currently drawing within Game Layer */
+	bool bIsInGameLayer;
 };
