Create the UInputDeviceProperty objects that will let you set trigger feedback, resistence, and vibration

diff --git a/Engine/Source/Runtime/Engine/Classes/GameFramework/InputDeviceProperties.h b/Engine/Source/Runtime/Engine/Classes/GameFramework/InputDeviceProperties.h
--- a/Engine/Source/Runtime/Engine/Classes/GameFramework/InputDeviceProperties.h
+++ b/Engine/Source/Runtime/Engine/Classes/GameFramework/InputDeviceProperties.h
@@ -73,13 +73,21 @@ public:
 	virtual void ResetDeviceProperty_Implementation(const FPlatformUserId PlatformUser);
 
 	/**
-	* Apply the given device property to the Input Interface. 
+	* Apply the device property from GetInternalDeviceProperty to the given platform user. 
 	* Note: To remove any applied affects of this device property, call ResetDeviceProperty.
 	* 
 	* @param UserId		The owning Platform User whose input device this property should be applied to.
 	*/
 	UFUNCTION(BlueprintCallable, Category = "InputDevice")
 	virtual void ApplyDeviceProperty(const FPlatformUserId UserId);
+
+	/** 
+	* Apply the given device property
+	* 
+	* @param UserId			The owning Platform User whose input device this property should be applied to.
+	* @param RawProperty	The internal input device property to apply.
+	*/
+	static void ApplyDeviceProperty(const FPlatformUserId UserId, FInputDeviceProperty* RawProperty);
 	
 	/** Gets a pointer to the current input device property that the IInputInterface can use. */
 	virtual FInputDeviceProperty* GetInternalDeviceProperty() { return nullptr; };
@@ -114,7 +122,7 @@ protected:
 /** 
 * A property that can be used to change the color of an input device's light 
 * 
-* NOTE: This property has platform specific implementations and may behavior differently per platform.
+* NOTE: This property has platform specific implementations and may behave differently per platform.
 * See the docs for more details on each platform.
 */
 UCLASS(Blueprintable, BlueprintType)
@@ -144,3 +152,145 @@ private:
 	/** The internal light color property that this represents; */
 	FInputDeviceLightColorProperty InternalProperty;
 };
+
+/** A property that effect the triggers on a gamepad */
+UCLASS(Abstract, Blueprintable)
+class ENGINE_API UInputDeviceTriggerEffect : public UInputDeviceProperty
+{
+	GENERATED_BODY()
+
+public:	
+
+	virtual void ResetDeviceProperty_Implementation(const FPlatformUserId PlatformUser) override;
+
+	/** Which trigger this property should effect */
+    UPROPERTY(EditAnywhere, BlueprintReadWrite, Category = "DeviceProperty")
+    EInputDeviceTriggerMask AffectedTriggers = EInputDeviceTriggerMask::None;
+
+	/** True if the triggers should be reset after the duration of this device property */
+	UPROPERTY(EditAnywhere, BlueprintReadWrite, Category = "DeviceProperty")
+	bool bResetUponCompletion = true;
+
+protected:
+
+	/** Internal property that can be used to reset a given trigger */
+	FInputDeviceTriggerResetProperty ResetProperty = {};
+};
+
+/** 
+* Sets simple trigger feedback
+* 
+* NOTE: This property has platform specific implementations and may behave differently per platform.
+* See the docs for more details on each platform.
+*/
+UCLASS(Blueprintable)
+class UInputDeviceTriggerFeedbackProperty : public UInputDeviceTriggerEffect
+{
+	GENERATED_BODY()
+
+public:
+	
+	UInputDeviceTriggerFeedbackProperty();
+	
+	virtual void EvaluateDeviceProperty_Implementation(const FPlatformUserId PlatformUser, const float DeltaTime, const float Duration) override;	
+	virtual FInputDeviceProperty* GetInternalDeviceProperty() override;
+	virtual float RecalculateDuration() override;
+
+	/** What position on the trigger that the feedback should be applied to over time (1-9) */
+	UPROPERTY(EditAnywhere, BlueprintReadWrite, Category = "DeviceProperty")
+	TObjectPtr<UCurveFloat> FeedbackPositionCurve;
+
+	/** How strong the feedback is over time (1-8) */
+	UPROPERTY(EditAnywhere, BlueprintReadWrite, Category = "DeviceProperty")
+	TObjectPtr<UCurveFloat> FeedbackStrenghCurve;
+
+private:
+
+	int32 GetPositionValue(const float Duration) const;
+	int32 GetStrengthValue(const float Duration) const;
+
+	/** The internal property that represents this trigger feedback. */
+	FInputDeviceTriggerFeedbackProperty InternalProperty;
+};
+
+/** 
+* Provides resistance to a trigger while it is being pressed between a start and end value
+* 
+* NOTE: This property has platform specific implementations and may behave differently per platform.
+* See the docs for more details on each platform.
+*/
+UCLASS(Blueprintable)
+class UInputDeviceTriggerResistanceProperty : public UInputDeviceTriggerEffect
+{
+	GENERATED_BODY()
+
+public:
+
+	UInputDeviceTriggerResistanceProperty();
+
+	virtual void EvaluateDeviceProperty_Implementation(const FPlatformUserId PlatformUser, const float DeltaTime, const float Duration) override;
+	virtual FInputDeviceProperty* GetInternalDeviceProperty() override;
+
+protected:
+
+	/** The position that the trigger should start providing resistance */
+	UPROPERTY(EditAnywhere, BlueprintReadWrite, Category = "DeviceProperty")
+	int32 StartPosition = 0;
+
+	/** How strong the resistance is */
+	UPROPERTY(EditAnywhere, BlueprintReadWrite, Category = "DeviceProperty")
+	int32 StartStrengh = 0;
+	
+	UPROPERTY(EditAnywhere, BlueprintReadWrite, Category = "DeviceProperty")
+	int32 EndPosition = 0;
+	
+	UPROPERTY(EditAnywhere, BlueprintReadWrite, Category = "DeviceProperty")
+	int32 EndStrengh = 0;
+
+private:
+
+	/** The internal property that represents this trigger resistance */
+	FInputDeviceTriggerResistanceProperty InternalProperty;
+};
+
+
+/**
+* Sets trigger vibration
+*
+* NOTE: This property has platform specific implementations and may behave differently per platform.
+* See the docs for more details on each platform.
+*/
+UCLASS(Blueprintable)
+class UInputDeviceTriggerVibrationProperty : public UInputDeviceTriggerEffect
+{
+	GENERATED_BODY()
+
+public:
+
+	UInputDeviceTriggerVibrationProperty();
+
+	virtual void EvaluateDeviceProperty_Implementation(const FPlatformUserId PlatformUser, const float DeltaTime, const float Duration) override;
+	virtual FInputDeviceProperty* GetInternalDeviceProperty() override;
+	virtual float RecalculateDuration() override;
+
+	/** What position on the trigger that the feedback should be applied to over time (1-9) */
+	UPROPERTY(EditAnywhere, BlueprintReadWrite, Category = "DeviceProperty")
+	TObjectPtr<UCurveFloat> TriggerPositionCurve;
+
+	/** The frequency of the vibration */
+	UPROPERTY(EditAnywhere, BlueprintReadWrite, Category = "DeviceProperty")
+	TObjectPtr<UCurveFloat> VibrationFrequencyCurve;
+
+	/** The amplitude of the vibration */
+	UPROPERTY(EditAnywhere, BlueprintReadWrite, Category = "DeviceProperty")
+	TObjectPtr<UCurveFloat> VibrationAmplitudeCurve;
+
+private:
+
+	int32 GetTriggerPositionValue(const float Duration) const;
+	int32 GetVibrationFrequencyValue(const float Duration) const;
+	int32 GetVibrationAmplitudeValue(const float Duration) const;
+
+	/** The internal property that represents this trigger feedback. */
+	FInputDeviceTriggerVibrationProperty InternalProperty;
+};
\ No newline at end of file
diff --git a/Engine/Source/Runtime/Engine/Private/GameFramework/InputDeviceProperties.cpp b/Engine/Source/Runtime/Engine/Private/GameFramework/InputDeviceProperties.cpp
--- a/Engine/Source/Runtime/Engine/Private/GameFramework/InputDeviceProperties.cpp
+++ b/Engine/Source/Runtime/Engine/Private/GameFramework/InputDeviceProperties.cpp
@@ -16,7 +16,12 @@ UInputDeviceProperty::UInputDeviceProperty()
 
 void UInputDeviceProperty::ApplyDeviceProperty(const FPlatformUserId UserId)
 {
-	if (FInputDeviceProperty* RawProp = GetInternalDeviceProperty())
+	UInputDeviceProperty::ApplyDeviceProperty(UserId, GetInternalDeviceProperty());
+}
+
+void UInputDeviceProperty::ApplyDeviceProperty(const FPlatformUserId UserId, FInputDeviceProperty* RawProperty)
+{
+	if (ensure(RawProperty))
 	{
 		IInputInterface* InputInterface = FSlateApplication::Get().IsInitialized() ? FSlateApplication::Get().GetInputInterface() : nullptr;
 		if (InputInterface)
@@ -25,7 +30,7 @@ void UInputDeviceProperty::ApplyDeviceProperty(const FPlatformUserId UserId)
 			IPlatformInputDeviceMapper::Get().RemapUserAndDeviceToControllerId(UserId, ControllerId);
 
 			// TODO_BH: Refactor input interface to take an FPlatformUserId directly (UE-158881)
-			InputInterface->SetDeviceProperty(ControllerId, RawProp);
+			InputInterface->SetDeviceProperty(ControllerId, RawProperty);
 		}
 	}
 }
@@ -97,4 +102,184 @@ float UColorInputDeviceProperty::RecalculateDuration()
 		PropertyDuration = 1.0f;
 	}
 	return PropertyDuration;
-}
\ No newline at end of file
+}
+
+///////////////////////////////////////////////////////////////////////
+// UInputDeviceTriggerEffect
+
+void UInputDeviceTriggerEffect::ResetDeviceProperty_Implementation(const FPlatformUserId PlatformUser)
+{
+	if (bResetUponCompletion)
+	{
+		// Pass in our reset property
+		ResetProperty.AffectedTriggers = AffectedTriggers;
+		ApplyDeviceProperty(PlatformUser, &ResetProperty);
+	}	
+}
+
+///////////////////////////////////////////////////////////////////////
+// UInputDeviceTriggerFeedbackProperty
+
+UInputDeviceTriggerFeedbackProperty::UInputDeviceTriggerFeedbackProperty()
+	: UInputDeviceTriggerEffect()
+{
+	InternalProperty.AffectedTriggers = AffectedTriggers;
+}
+
+int32 UInputDeviceTriggerFeedbackProperty::GetPositionValue(const float Duration) const
+{
+	if (ensure(FeedbackPositionCurve))
+	{
+		// TODO: Make the max position a cvar
+		int32 Position = FeedbackPositionCurve->GetFloatValue(Duration);
+		return FMath::Clamp(Position, 0, 9);
+	}
+	return 0;
+}
+
+int32 UInputDeviceTriggerFeedbackProperty::GetStrengthValue(const float Duration) const
+{
+	if (ensure(FeedbackStrenghCurve))
+	{
+		// TODO: Make the max Strength a cvar
+		int32 Strength = FeedbackPositionCurve->GetFloatValue(Duration);
+		return FMath::Clamp(Strength, 0, 8);
+	}
+	return 0;
+}
+
+void UInputDeviceTriggerFeedbackProperty::EvaluateDeviceProperty_Implementation(const FPlatformUserId PlatformUser, const float DeltaTime, const float Duration)
+{		
+	InternalProperty.AffectedTriggers = AffectedTriggers;
+
+	InternalProperty.Position = GetPositionValue(Duration);
+	InternalProperty.Strengh = GetStrengthValue(Duration);
+}
+
+FInputDeviceProperty* UInputDeviceTriggerFeedbackProperty::GetInternalDeviceProperty()
+{
+	return &InternalProperty;
+}
+
+float UInputDeviceTriggerFeedbackProperty::RecalculateDuration()
+{
+	// Get the max time from the two curves on this property
+	float MinTime, MaxTime = 0.0f;
+	if (FeedbackPositionCurve)
+	{
+		FeedbackPositionCurve->GetTimeRange(MinTime, MaxTime);
+	}
+	
+	if (FeedbackStrenghCurve)
+	{
+		FeedbackStrenghCurve->GetTimeRange(MinTime, MaxTime);
+	}
+	
+	PropertyDuration = MaxTime;
+	return PropertyDuration;
+}
+
+///////////////////////////////////////////////////////////////////////
+// UInputDeviceTriggerResistanceProperty
+
+UInputDeviceTriggerResistanceProperty::UInputDeviceTriggerResistanceProperty()
+	: UInputDeviceTriggerEffect()
+{
+	PropertyDuration = 1.0f;
+}
+
+void UInputDeviceTriggerResistanceProperty::EvaluateDeviceProperty_Implementation(const FPlatformUserId PlatformUser, const float DeltaTime, const float Duration)
+{
+	InternalProperty.AffectedTriggers = AffectedTriggers;
+	InternalProperty.StartPosition = StartPosition;
+	InternalProperty.StartStrengh = StartStrengh;
+	InternalProperty.EndPosition = EndPosition;
+	InternalProperty.EndStrengh = EndStrengh;
+}
+
+FInputDeviceProperty* UInputDeviceTriggerResistanceProperty::GetInternalDeviceProperty()
+{
+	return &InternalProperty;
+}
+
+
+///////////////////////////////////////////////////////////////////////
+// UInputDeviceTriggerVibrationProperty
+
+UInputDeviceTriggerVibrationProperty::UInputDeviceTriggerVibrationProperty()
+	: UInputDeviceTriggerEffect()
+{
+	PropertyDuration = 1.0f;
+}
+
+void UInputDeviceTriggerVibrationProperty::EvaluateDeviceProperty_Implementation(const FPlatformUserId PlatformUser, const float DeltaTime, const float Duration)
+{
+	InternalProperty.AffectedTriggers = AffectedTriggers;
+	InternalProperty.TriggerPosition = GetTriggerPositionValue(Duration);
+	InternalProperty.VibrationFrequency = GetVibrationFrequencyValue(Duration);
+	InternalProperty.VibrationAmplitude = GetVibrationAmplitudeValue(Duration);
+}
+
+FInputDeviceProperty* UInputDeviceTriggerVibrationProperty::GetInternalDeviceProperty()
+{
+	return &InternalProperty;
+}
+
+float UInputDeviceTriggerVibrationProperty::RecalculateDuration()
+{
+	// Get the max time from the curves on this property
+	float MaxTime = 0.0f;
+
+	auto EvaluateMaxTime = [&MaxTime](TObjectPtr<UCurveFloat> InCurve)
+	{
+		float MinCurveTime, MaxCurveTime = 0.0f;
+		if (InCurve)
+		{
+			InCurve->GetTimeRange(MinCurveTime, MaxCurveTime);
+			if (MaxCurveTime > MaxTime)
+			{
+				MaxTime = MaxCurveTime;
+			}
+		}
+	};
+
+	EvaluateMaxTime(TriggerPositionCurve);
+	EvaluateMaxTime(VibrationFrequencyCurve);
+	EvaluateMaxTime(VibrationAmplitudeCurve);
+
+	PropertyDuration = MaxTime;
+	return PropertyDuration;
+}
+
+int32 UInputDeviceTriggerVibrationProperty::GetTriggerPositionValue(const float Duration) const
+{
+	if (ensure(TriggerPositionCurve))
+	{
+		// TODO: Make the max Strength a cvar
+		int32 Strength = TriggerPositionCurve->GetFloatValue(Duration);
+		return FMath::Clamp(Strength, 0, 9);
+	}
+	return 0;
+}
+
+int32 UInputDeviceTriggerVibrationProperty::GetVibrationFrequencyValue(const float Duration) const
+{
+	if (ensure(VibrationFrequencyCurve))
+	{
+		// TODO: Make the max Frequency a cvar
+		int32 Strength = VibrationFrequencyCurve->GetFloatValue(Duration);
+		return FMath::Clamp(Strength, 0, 255);
+	}
+	return 0;
+}
+
+int32 UInputDeviceTriggerVibrationProperty::GetVibrationAmplitudeValue(const float Duration) const
+{
+	if (ensure(VibrationAmplitudeCurve))
+	{
+		// TODO: Make the max Amplitude a cvar
+		int32 Strength = VibrationAmplitudeCurve->GetFloatValue(Duration);
+		return FMath::Clamp(Strength, 0, 8);
+	}
+	return 0;
+}
