When an actor channel has been closed due to tear off, update the actor flag and tear off as normal on the client side even if the property update containing bTearOff was lost.

diff --git a/Engine/Source/Runtime/Engine/Classes/GameFramework/Actor.h b/Engine/Source/Runtime/Engine/Classes/GameFramework/Actor.h
--- a/Engine/Source/Runtime/Engine/Classes/GameFramework/Actor.h
+++ b/Engine/Source/Runtime/Engine/Classes/GameFramework/Actor.h
@@ -57,6 +57,11 @@ class UActorFolder;
 struct FActorDataLayer;
 struct FHitResult;
 
+namespace UE::Net
+{
+	class FTearOffSetter;
+}
+
 // By default, debug and development builds (even cooked) will keep actor labels. Manually define this if you want to make a local build
 // that keep actor labels for Test or Shipping builds.
 #define ACTOR_HAS_LABELS (UE_BUILD_DEBUG || UE_BUILD_DEVELOPMENT || WITH_PROFILEGPU)
@@ -262,6 +267,8 @@ private:
 	UPROPERTY(Replicated)
 	uint8 bTearOff:1;
 
+	friend class UE::Net::FTearOffSetter;
+
 	/** When set, indicates that external guarantees ensure that this actor's name is deterministic between server and client, and as such can be addressed by its full path */
 	UPROPERTY()
 	uint8 bForceNetAddressable:1;
diff --git a/Engine/Source/Runtime/Engine/Private/DataChannel.cpp b/Engine/Source/Runtime/Engine/Private/DataChannel.cpp
--- a/Engine/Source/Runtime/Engine/Private/DataChannel.cpp
+++ b/Engine/Source/Runtime/Engine/Private/DataChannel.cpp
@@ -162,6 +162,18 @@ namespace UE::Net
 	/** Whether or not the currently executing ReplicateActor code has met the randomized chance for enabling name debugging */
 	static bool GMetAsyncDemoNameDebugChance = false;
 #endif
+
+	// bTearOff is private but we still might need to adjust the flag on clients based on the channel close reason
+	class FTearOffSetter final
+	{
+		FTearOffSetter() = delete;
+
+	public:
+		static void SetTearOff(AActor* Actor)
+		{
+			Actor->bTearOff = true;
+		}
+	};
 }; // namespace UE::Net
 
 template<typename T>
@@ -2402,6 +2414,12 @@ bool UActorChannel::CleanUp(const bool bForDestroy, EChannelCloseReason CloseRea
 		checkSlow(Connection->Driver->IsValidLowLevel());
 		if (Actor != NULL)
 		{
+			if (CloseReason == EChannelCloseReason::TearOff && !Actor->GetTearOff())
+			{
+				UE_LOG(LogNetTraffic, Verbose, TEXT("UActorChannel::CleanUp: Closed for TearOff but actor flag not set, last property update may have been dropped: %s"), *Describe());
+				UE::Net::FTearOffSetter::SetTearOff(Actor);
+			}
+
 			if (Actor->GetTearOff() && !Connection->Driver->ShouldClientDestroyTearOffActors())
 			{
 				if (!bTornOff)
