Cherypick cook-time memory optimisation from UE5Main: clearing out resampled data, compressed bone/curve data when fully cooking a AnimSequence package

diff --git a/Engine/Source/Runtime/Engine/Private/Animation/AnimSequence.cpp b/Engine/Source/Runtime/Engine/Private/Animation/AnimSequence.cpp
--- a/Engine/Source/Runtime/Engine/Private/Animation/AnimSequence.cpp
+++ b/Engine/Source/Runtime/Engine/Private/Animation/AnimSequence.cpp
@@ -1119,6 +1119,31 @@ void UAnimSequence::PostEditChangeProperty(FPropertyChangedEvent& PropertyChange
 		RecompressAnimationData();
 	}
 }
+
+void UAnimSequence::WillNeverCacheCookedPlatformDataAgain()
+{
+	Super::WillNeverCacheCookedPlatformDataAgain();
+	ResampledAnimationTrackData.Empty();
+	ClearCompressedCurveData();
+	ClearCompressedBoneData();
+	bUseRawDataOnly = true;
+}
+
+void UAnimSequence::ClearCachedCookedPlatformData(const ITargetPlatform* TargetPlatform)
+{
+	Super::ClearCachedCookedPlatformData(TargetPlatform);
+
+	// For now just clear resampled data (as this will always be resampled when not found for compression)
+	ResampledAnimationTrackData.Empty();
+}
+
+void UAnimSequence::ClearAllCachedCookedPlatformData()
+{
+	Super::ClearAllCachedCookedPlatformData();
+	
+	// For now just clear resampled data (as this will always be resampled when not found for compression)
+	ResampledAnimationTrackData.Empty();
+}
 #endif // WITH_EDITOR
 
 #if WITH_EDITOR
