Added method in WorldPartitionSubsystem to know if streaming source is registered.

diff --git a/Engine/Source/Runtime/Engine/Private/WorldPartition/WorldPartitionSubsystem.cpp b/Engine/Source/Runtime/Engine/Private/WorldPartition/WorldPartitionSubsystem.cpp
--- a/Engine/Source/Runtime/Engine/Private/WorldPartition/WorldPartitionSubsystem.cpp
+++ b/Engine/Source/Runtime/Engine/Private/WorldPartition/WorldPartitionSubsystem.cpp
@@ -227,6 +227,11 @@ void UWorldPartitionSubsystem::RegisterStreamingSourceProvider(IWorldPartitionSt
 	UE_CLOG(bIsAlreadyInSet, LogWorldPartition, Warning, TEXT("Streaming source provider already registered."));
 }
 
+bool UWorldPartitionSubsystem::IsStreamingSourceProviderRegistered(IWorldPartitionStreamingSourceProvider* StreamingSource) const
+{
+	return StreamingSourceProviders.Contains(StreamingSource);
+}
+
 bool UWorldPartitionSubsystem::UnregisterStreamingSourceProvider(IWorldPartitionStreamingSourceProvider* StreamingSource)
 {
 	return !!StreamingSourceProviders.Remove(StreamingSource);
diff --git a/Engine/Source/Runtime/Engine/Public/WorldPartition/WorldPartitionSubsystem.h b/Engine/Source/Runtime/Engine/Public/WorldPartition/WorldPartitionSubsystem.h
--- a/Engine/Source/Runtime/Engine/Public/WorldPartition/WorldPartitionSubsystem.h
+++ b/Engine/Source/Runtime/Engine/Public/WorldPartition/WorldPartitionSubsystem.h
@@ -56,6 +56,7 @@ public:
 
 	TSet<IWorldPartitionStreamingSourceProvider*> GetStreamingSourceProviders() const;
 	void RegisterStreamingSourceProvider(IWorldPartitionStreamingSourceProvider* StreamingSource);
+	bool IsStreamingSourceProviderRegistered(IWorldPartitionStreamingSourceProvider* StreamingSource) const;
 	bool UnregisterStreamingSourceProvider(IWorldPartitionStreamingSourceProvider* StreamingSource);
 
 	DECLARE_DELEGATE_RetVal_OneParam(bool, FWorldPartitionStreamingSourceProviderFilter, const IWorldPartitionStreamingSourceProvider*);
