Prevent AGameModeBase::ReplicateStreamingStatus and APlayerController::ClientUpdateLevelStreamingStatus from modifying streaming levels handled by world partition.

diff --git a/Engine/Source/Runtime/Engine/Classes/Engine/LevelStreaming.h b/Engine/Source/Runtime/Engine/Classes/Engine/LevelStreaming.h
--- a/Engine/Source/Runtime/Engine/Classes/Engine/LevelStreaming.h
+++ b/Engine/Source/Runtime/Engine/Classes/Engine/LevelStreaming.h
@@ -404,6 +404,9 @@ public:
 	/** Returns whether level should start to render only when it will be fully added to the world or not. */
 	virtual bool ShouldRequireFullVisibilityToRender() const { return LODPackageNames.Num() > 0; }
 
+	/** Returns whether level status can be replicated from the server to the client */
+	virtual bool CanReplicateStreamingStatus() const { return true; }
+
 	/** 
 	 * Virtual that can be overriden to change whether a streaming level should be loaded.
 	 * Doesn't do anything at the base level as should be loaded defaults to true 
diff --git a/Engine/Source/Runtime/Engine/Private/GameModeBase.cpp b/Engine/Source/Runtime/Engine/Private/GameModeBase.cpp
--- a/Engine/Source/Runtime/Engine/Private/GameModeBase.cpp
+++ b/Engine/Source/Runtime/Engine/Private/GameModeBase.cpp
@@ -921,7 +921,7 @@ void AGameModeBase::ReplicateStreamingStatus(APlayerController* PC)
 			TArray<FUpdateLevelStreamingLevelStatus> LevelStatuses;
 			for (ULevelStreaming* TheLevel : MyWorld->GetStreamingLevels())
 			{
-				if (TheLevel != nullptr)
+				if (TheLevel && TheLevel->CanReplicateStreamingStatus())
 				{
 					const ULevel* LoadedLevel = TheLevel->GetLoadedLevel();
 
diff --git a/Engine/Source/Runtime/Engine/Private/PlayerController.cpp b/Engine/Source/Runtime/Engine/Private/PlayerController.cpp
--- a/Engine/Source/Runtime/Engine/Private/PlayerController.cpp
+++ b/Engine/Source/Runtime/Engine/Private/PlayerController.cpp
@@ -325,51 +325,37 @@ void APlayerController::ClientUpdateLevelStreamingStatus_Implementation(FName Pa
 		}
 	}
 
+	// Search for the streaming level object by name
+	ULevelStreaming* LevelStreamingObject = FLevelUtils::FindStreamingLevel(World, PackageName);
+
+	// Skip if streaming level object doesn't allow replicating the status
+	if (LevelStreamingObject && !LevelStreamingObject->CanReplicateStreamingStatus())
+	{
+		return;
+	}
+
 	// if we're about to commit a map change, we assume that the streaming update is based on the to be loaded map and so defer it until that is complete
 	if (GEngine->ShouldCommitPendingMapChange(World))
 	{
 		GEngine->AddNewPendingStreamingLevel(World, PackageName, bNewShouldBeLoaded, bNewShouldBeVisible, LODIndex);		
 	}
-	else
+	else if (LevelStreamingObject)
 	{
-		// search for the level object by name
-		ULevelStreaming* LevelStreamingObject = nullptr;
-		if (World && PackageName != NAME_None)
+		// If we're unloading any levels, we need to request a one frame delay of garbage collection to make sure it happens after the level is actually unloaded
+		if (LevelStreamingObject->ShouldBeLoaded() && !bNewShouldBeLoaded)
 		{
-			for (ULevelStreaming* CurrentLevelStreamingObject : World->GetStreamingLevels())
-			{
-				if (CurrentLevelStreamingObject && CurrentLevelStreamingObject->GetWorldAssetPackageFName() == PackageName)
-				{
-					LevelStreamingObject = CurrentLevelStreamingObject;
-					if (LevelStreamingObject)
-					{
-						// If we're unloading any levels, we need to request a one frame delay of garbage collection to make sure it happens after the level is actually unloaded
-						if (LevelStreamingObject->ShouldBeLoaded() && !bNewShouldBeLoaded)
-						{
-							GEngine->DelayGarbageCollection();
-						}
-
-						LevelStreamingObject->SetShouldBeLoaded(bNewShouldBeLoaded);
-						LevelStreamingObject->SetShouldBeVisible(bNewShouldBeVisible);
-						LevelStreamingObject->bShouldBlockOnLoad = bNewShouldBlockOnLoad;
-						LevelStreamingObject->SetLevelLODIndex(LODIndex);
-						LevelStreamingObject->UpdateNetVisibilityTransactionState(bNewShouldBeVisible, TransactionId);
-					}
-					else
-					{
-						UE_LOG(LogStreaming, Log, TEXT("Unable to handle streaming object %s"),*LevelStreamingObject->GetName() );
-					}
-
-					// break out of object iterator if we found a match
-					break;
-				}
-			}
+			GEngine->DelayGarbageCollection();
 		}
 
-		if (LevelStreamingObject == NULL)
-		{
-			UE_LOG(LogStreaming, Log, TEXT("Unable to find streaming object %s"), *PackageName.ToString() );
-		}
+		LevelStreamingObject->SetShouldBeLoaded(bNewShouldBeLoaded);
+		LevelStreamingObject->SetShouldBeVisible(bNewShouldBeVisible);
+		LevelStreamingObject->bShouldBlockOnLoad = bNewShouldBlockOnLoad;
+		LevelStreamingObject->SetLevelLODIndex(LODIndex);
+		LevelStreamingObject->UpdateNetVisibilityTransactionState(bNewShouldBeVisible, TransactionId);
+	}
+	else
+	{
+		UE_LOG(LogStreaming, Log, TEXT("Unable to find streaming object %s"), *PackageName.ToString() );
 	}
 }
 
diff --git a/Engine/Source/Runtime/Engine/Public/WorldPartition/WorldPartitionLevelStreamingDynamic.h b/Engine/Source/Runtime/Engine/Public/WorldPartition/WorldPartitionLevelStreamingDynamic.h
--- a/Engine/Source/Runtime/Engine/Public/WorldPartition/WorldPartitionLevelStreamingDynamic.h
+++ b/Engine/Source/Runtime/Engine/Public/WorldPartition/WorldPartitionLevelStreamingDynamic.h
@@ -41,6 +41,8 @@ class ENGINE_API UWorldPartitionLevelStreamingDynamic : public ULevelStreamingDy
 
 	void Initialize(const UWorldPartitionRuntimeLevelStreamingCell& InCell);
 
+	virtual bool CanReplicateStreamingStatus() const override { return false; }
+
 #if WITH_EDITOR
 	static UWorldPartitionLevelStreamingDynamic* LoadInEditor(UWorld* World, FName LevelStreamingName, const TArray<FWorldPartitionRuntimeCellObjectMapping>& InPackages);
 	static void UnloadFromEditor(UWorldPartitionLevelStreamingDynamic* InLevelStreaming);
