Fixed normal space not being repsected by material attribute.

diff --git a/Engine/Source/Runtime/Engine/Private/Materials/HLSLMaterialTranslator.cpp b/Engine/Source/Runtime/Engine/Private/Materials/HLSLMaterialTranslator.cpp
--- a/Engine/Source/Runtime/Engine/Private/Materials/HLSLMaterialTranslator.cpp
+++ b/Engine/Source/Runtime/Engine/Private/Materials/HLSLMaterialTranslator.cpp
@@ -3608,6 +3608,12 @@ EShaderFrequency FHLSLMaterialTranslator::GetCurrentShaderFrequency() const
 	return ShaderFrequency;
 }
 
+bool FHLSLMaterialTranslator::IsTangentSpaceNormal() const
+{
+	check(Material);
+	return Material->GetMaterialDomain() == MD_DeferredDecal || Material->IsTangentSpaceNormal();
+}
+
 FMaterialShadingModelField FHLSLMaterialTranslator::GetMaterialShadingModels() const
 {
 	check(Material);
@@ -8697,7 +8703,7 @@ int32 FHLSLMaterialTranslator::TransformPosition(EMaterialCommonBasis SourceCoor
 
 int32 FHLSLMaterialTranslator::TransformNormalFromRequestedBasisToWorld(int32 NormalCodeChunk)
 {
-	if (Material->GetMaterialDomain() == MD_DeferredDecal || Material->IsTangentSpaceNormal())
+	if (IsTangentSpaceNormal())
 	{
 		// See TransformTangentNormalToWorld definitions in MaterialTemplate.ush
 		return AddCodeChunk(MCT_Float3, TEXT("TransformTangentNormalToWorld(Parameters.TangentToWorld, %s)"), *GetParameterCode(NormalCodeChunk));
diff --git a/Engine/Source/Runtime/Engine/Private/Materials/HLSLMaterialTranslator.h b/Engine/Source/Runtime/Engine/Private/Materials/HLSLMaterialTranslator.h
--- a/Engine/Source/Runtime/Engine/Private/Materials/HLSLMaterialTranslator.h
+++ b/Engine/Source/Runtime/Engine/Private/Materials/HLSLMaterialTranslator.h
@@ -703,6 +703,8 @@ protected:
 
 	virtual EShaderFrequency GetCurrentShaderFrequency() const override;
 
+	bool IsTangentSpaceNormal() const;
+
 	virtual FMaterialShadingModelField GetMaterialShadingModels() const override;
 	virtual FMaterialShadingModelField GetCompiledShadingModels() const override;
 
diff --git a/Engine/Source/Runtime/Engine/Private/Materials/MaterialShared.cpp b/Engine/Source/Runtime/Engine/Private/Materials/MaterialShared.cpp
--- a/Engine/Source/Runtime/Engine/Private/Materials/MaterialShared.cpp
+++ b/Engine/Source/Runtime/Engine/Private/Materials/MaterialShared.cpp
@@ -5335,6 +5335,24 @@ int32 FMaterialAttributeDefintion::CompileDefaultValue(FMaterialCompiler* Compil
 		return Compiler->ShadingModel(Compiler->GetMaterialShadingModels().GetFirstShadingModel());
 	}
 
+	if (Property == MP_Normal)
+	{
+		if (Compiler->IsTangentSpaceNormal())
+		{
+			return Compiler->Constant3(0.0f, 0.0f, 1.0f);	// Tangent space normal
+		}
+		return Compiler->VertexNormal();					// World space normal
+	}
+
+	if (Property == MP_Tangent)
+	{
+		if (Compiler->IsTangentSpaceNormal())
+		{
+			return Compiler->Constant3(1.0f, 0.0f, 0.0f);	// Tangent space tangent
+		}
+		return Compiler->VertexTangent();					// World space tangent
+	}
+
 	if (Property == MP_FrontMaterial)
 	{
 		check(ValueType == MCT_Strata);
diff --git a/Engine/Source/Runtime/Engine/Public/MaterialCompiler.h b/Engine/Source/Runtime/Engine/Public/MaterialCompiler.h
--- a/Engine/Source/Runtime/Engine/Public/MaterialCompiler.h
+++ b/Engine/Source/Runtime/Engine/Public/MaterialCompiler.h
@@ -244,6 +244,8 @@ public:
 
 	virtual const ITargetPlatform* GetTargetPlatform() const = 0;
 
+	virtual bool IsTangentSpaceNormal() const = 0;
+
 	virtual FMaterialShadingModelField GetMaterialShadingModels() const = 0;
 	
 	/** Get the shading models that were encountered when compiling a material's Shading Model attribute graph.
@@ -764,6 +766,8 @@ public:
 	virtual const FGuid GetMaterialAttribute() override { return Compiler->GetMaterialAttribute(); }
 	virtual void SetBaseMaterialAttribute(const FGuid& InAttributeID) override { Compiler->SetBaseMaterialAttribute(InAttributeID); }
 
+	virtual bool IsTangentSpaceNormal() const override { return Compiler->IsTangentSpaceNormal(); }
+
 	virtual void PushParameterOwner(const FMaterialParameterInfo& InOwnerInfo) override { Compiler->PushParameterOwner(InOwnerInfo); }
 	virtual FMaterialParameterInfo PopParameterOwner() override { return Compiler->PopParameterOwner(); }
 
