Defer SkinnedMeshComponent Force/ClearMotionBlur updates until end-of-frame to be processed together with regular dynamic data update.
This reduces number of times bone buffer, morph buffer, cloth sim buffer are written into to just once per frame.

diff --git a/Engine/Plugins/Experimental/ChaosClothAsset/Source/ChaosClothAssetEngine/Private/ChaosClothAsset/ClothComponent.cpp b/Engine/Plugins/Experimental/ChaosClothAsset/Source/ChaosClothAssetEngine/Private/ChaosClothAsset/ClothComponent.cpp
--- a/Engine/Plugins/Experimental/ChaosClothAsset/Source/ChaosClothAssetEngine/Private/ChaosClothAsset/ClothComponent.cpp
+++ b/Engine/Plugins/Experimental/ChaosClothAsset/Source/ChaosClothAssetEngine/Private/ChaosClothAsset/ClothComponent.cpp
@@ -56,8 +56,6 @@ UChaosClothComponent::UChaosClothComponent(const FObjectInitializer& ObjectIniti
 	bSyncAttachParentLOD = true;
 	//bIgnoreMasterPoseComponentLOD = false;
 
-	CurrentBoneTransformRevisionNumber = 0;
-
 	//ExternalInterpolationAlpha = 0.0f;
 	//ExternalDeltaTime = 0.0f;
 	//ExternalTickRate = 1;
diff --git a/Engine/Source/Runtime/Engine/Classes/Components/SkinnedMeshComponent.h b/Engine/Source/Runtime/Engine/Classes/Components/SkinnedMeshComponent.h
--- a/Engine/Source/Runtime/Engine/Classes/Components/SkinnedMeshComponent.h
+++ b/Engine/Source/Runtime/Engine/Classes/Components/SkinnedMeshComponent.h
@@ -24,7 +24,7 @@
 #include "SkinnedMeshComponent.generated.h"
 
 enum class ESkinCacheUsage : uint8;
-
+enum class EPreviousBoneTransformUpdateMode;
 class FPrimitiveSceneProxy;
 class FColorVertexBuffer;
 class FSkinWeightVertexBuffer;
@@ -286,6 +286,16 @@ public:
 	void ForceMotionVector();
 
 private:
+	EPreviousBoneTransformUpdateMode UpdateBoneTransformRevisionNumber();
+
+	enum class EBoneTransformUpdateMethod
+	{
+		AnimationUpdate,
+		ClearMotionVector,
+		ForceMotionVector
+	};
+	TArray<EBoneTransformUpdateMethod> BoneTransformUpdateMethodQueue;
+
 	/** Temporary array of of component-space bone matrices, update each frame and used for rendering the mesh. */
 	TArray<FTransform> ComponentSpaceTransformsArray[2];
 
diff --git a/Engine/Source/Runtime/Engine/Private/Components/SkinnedMeshComponent.cpp b/Engine/Source/Runtime/Engine/Private/Components/SkinnedMeshComponent.cpp
--- a/Engine/Source/Runtime/Engine/Private/Components/SkinnedMeshComponent.cpp
+++ b/Engine/Source/Runtime/Engine/Private/Components/SkinnedMeshComponent.cpp
@@ -791,6 +791,54 @@ void USkinnedMeshComponent::RefreshExternalMorphTargetWeights()
 	}
 }
 
+EPreviousBoneTransformUpdateMode USkinnedMeshComponent::UpdateBoneTransformRevisionNumber()
+{
+	EPreviousBoneTransformUpdateMode PreviousUpdateMode = EPreviousBoneTransformUpdateMode::None;
+
+	if (BoneTransformUpdateMethodQueue.Num() > 0)
+	{
+		if (BoneTransformUpdateMethodQueue.Last() == EBoneTransformUpdateMethod::ClearMotionVector)
+		{
+			// Last entry is ClearMotionVector, increment revision number by 2 which allows current bone buffer to be bound to previous shader slot to cancel out velocity
+			// See FGPUBaseSkinVertexFactory::FShaderDataType::GetBoneBufferInternal() for how it's used.
+			CurrentBoneTransformRevisionNumber += 2;
+		}
+		else
+		{
+			CurrentBoneTransformRevisionNumber++;
+			if (BoneTransformUpdateMethodQueue.Num() > 1)
+			{
+				// When there are multiple entries in the queue, previous bone transforms may or may not need updating, so for simplicity make sure it updates.
+				// E.g. Queue: ForceMotionVector -> AnimationUpdate, previous data may be overriden by ForceMotionVector and should update
+				PreviousUpdateMode = EPreviousBoneTransformUpdateMode::UpdatePrevious;
+			}
+
+			// If we may not evaluate animation every frame and there's an animation update in the queue, copy current to previous
+			if (bExternalEvaluationRateLimited && !bExternalInterpolate)
+			{
+				// Is there an animation update in the queue?
+				int32 AnimationUpdateIndex = -1;
+				for (int32 i = 0; i < BoneTransformUpdateMethodQueue.Num(); ++i)
+				{
+					if (BoneTransformUpdateMethodQueue[i] == EBoneTransformUpdateMethod::AnimationUpdate)
+					{
+						AnimationUpdateIndex = i;
+						break;
+					}
+				}
+				if (AnimationUpdateIndex > -1)
+				{
+					PreviousUpdateMode = EPreviousBoneTransformUpdateMode::DuplicateCurrentToPrevious;
+				}
+			}
+		}
+
+		BoneTransformUpdateMethodQueue.Reset();
+	}
+
+	return PreviousUpdateMode;
+}
+
 void USkinnedMeshComponent::CreateRenderState_Concurrent(FRegisterComponentContext* Context)
 {
 	LLM_SCOPE(ELLMTag::SkeletalMesh);
@@ -909,6 +957,7 @@ void USkinnedMeshComponent::CreateRenderState_Concurrent(FRegisterComponentConte
 				}
 
 				RefreshExternalMorphTargetWeights();
+				UpdateBoneTransformRevisionNumber();
 				MeshObject->Update(ModifiedLODLevel, this, ActiveMorphTargets, MorphTargetWeights, EPreviousBoneTransformUpdateMode::UpdatePrevious, GetExternalMorphWeights(ModifiedLODLevel));  // send to rendering thread
 			}
 		}
@@ -1007,13 +1056,15 @@ void USkinnedMeshComponent::SendRenderDynamicData_Concurrent()
 			{
 				ActiveMorphTargets.Empty();
 			}
+			
+			EPreviousBoneTransformUpdateMode PreviousBoneTransformUpdateMode = UpdateBoneTransformRevisionNumber();
 
 			MeshObject->Update(
 				UseLOD,
 				this,
 				ActiveMorphTargets,
 				MorphTargetWeights,
-				bExternalEvaluationRateLimited && !bExternalInterpolate ? EPreviousBoneTransformUpdateMode::DuplicateCurrentToPrevious : EPreviousBoneTransformUpdateMode::None,
+				PreviousBoneTransformUpdateMode,
 				GetExternalMorphWeights(UseLOD)
 			);  // send to rendering thread
 
@@ -1051,18 +1102,12 @@ void USkinnedMeshComponent::ClearMotionVector()
 			UseLOD = FMath::Max<int32>(UseLOD, MeshObject->GetSkeletalMeshRenderData().PendingFirstLODIdx);
 		}
 
-		// rendering bone velocity is updated by revision number
 		// if you have situation where you want to clear the bone velocity (that causes temporal AA or motion blur)
 		// use this function to clear it
-		// this function updates renderer twice using increasing of revision number, so that renderer updates previous/new transform correctly
-		++CurrentBoneTransformRevisionNumber;
-		MeshObject->Update(UseLOD, this, ActiveMorphTargets, MorphTargetWeights, EPreviousBoneTransformUpdateMode::None, GetExternalMorphWeights(UseLOD));  // send to rendering thread
-
-		++CurrentBoneTransformRevisionNumber;
-		MeshObject->Update(UseLOD, this, ActiveMorphTargets, MorphTargetWeights, EPreviousBoneTransformUpdateMode::None, GetExternalMorphWeights(UseLOD));  // send to rendering thread
-
-		// Skin cache may need updating
-		MarkForNeededEndOfFrameUpdate();
+		BoneTransformUpdateMethodQueue.Add(EBoneTransformUpdateMethod::ClearMotionVector);
+		// Make sure an update happens
+		bForceMeshObjectUpdate = true;
+		MarkRenderDynamicDataDirty();
 	}
 }
 
@@ -1077,11 +1122,10 @@ void USkinnedMeshComponent::ForceMotionVector()
 			UseLOD = FMath::Max<int32>(UseLOD, MeshObject->GetSkeletalMeshRenderData().PendingFirstLODIdx);
 		}
 
-		++CurrentBoneTransformRevisionNumber;
-		MeshObject->Update(UseLOD, this, ActiveMorphTargets, MorphTargetWeights, EPreviousBoneTransformUpdateMode::None, GetExternalMorphWeights(UseLOD));
-
-		// Skin cache may need updating
-		MarkForNeededEndOfFrameUpdate();
+		BoneTransformUpdateMethodQueue.Add(EBoneTransformUpdateMethod::ForceMotionVector);
+		// Make sure an update happens
+		bForceMeshObjectUpdate = true;
+		MarkRenderDynamicDataDirty();
 	}
 }
 
@@ -3771,7 +3815,7 @@ void USkinnedMeshComponent::FlipEditableSpaceBases()
 			}
 		}
 
-		++CurrentBoneTransformRevisionNumber;
+		BoneTransformUpdateMethodQueue.Add(EBoneTransformUpdateMethod::AnimationUpdate);
 	}
 }
 
diff --git a/Engine/Source/Runtime/Engine/Private/GPUSkinVertexFactory.cpp b/Engine/Source/Runtime/Engine/Private/GPUSkinVertexFactory.cpp
--- a/Engine/Source/Runtime/Engine/Private/GPUSkinVertexFactory.cpp
+++ b/Engine/Source/Runtime/Engine/Private/GPUSkinVertexFactory.cpp
@@ -722,10 +722,10 @@ void TGPUSkinVertexFactory<BoneInfluenceType>::UpdateMorphVertexStream(const FMo
 }
 
 template <GPUSkinBoneInfluenceType BoneInfluenceType>
-const FMorphVertexBuffer* TGPUSkinVertexFactory<BoneInfluenceType>::GetMorphVertexBuffer(bool bPrevious, uint32 FrameNumber) const
+const FMorphVertexBuffer* TGPUSkinVertexFactory<BoneInfluenceType>::GetMorphVertexBuffer(bool bPrevious) const
 {
 	check(Data.IsValid());
-	return Data->MorphVertexBufferPool ? &Data->MorphVertexBufferPool->GetMorphVertexBufferForReading(bPrevious, FrameNumber) : nullptr;
+	return Data->MorphVertexBufferPool ? &Data->MorphVertexBufferPool->GetMorphVertexBufferForReading(bPrevious) : nullptr;
 }
 
 /*-----------------------------------------------------------------------------
@@ -836,7 +836,7 @@ public:
 		{
 			const FMorphVertexBuffer* MorphVertexBuffer = nullptr;
 			const auto* GPUSkinVertexFactory = (const FGPUBaseSkinVertexFactory*)VertexFactory;
-			MorphVertexBuffer = GPUSkinVertexFactory->GetMorphVertexBuffer(!View->Family->bWorldIsPaused, View->Family->FrameNumber);
+			MorphVertexBuffer = GPUSkinVertexFactory->GetMorphVertexBuffer(!View->Family->bWorldIsPaused);
 			ShaderBindings.Add(PreviousMorphBufferParameter, MorphVertexBuffer ? MorphVertexBuffer->GetSRV() : GNullVertexBuffer.VertexBufferSRV.GetReference());
 		}
 	}
@@ -1085,10 +1085,8 @@ public:
 
 		ShaderBindings.Add(Shader->GetUniformBufferParameter<FAPEXClothUniformShaderParameters>(), ClothShaderData.GetClothUniformBuffer());
 
-		uint32 FrameNumber = View->Family->FrameNumber;
-
-		ShaderBindings.Add(ClothSimulVertsPositionsNormalsParameter, ClothShaderData.GetClothBufferForReading(false, FrameNumber).VertexBufferSRV);
-		ShaderBindings.Add(ClothToLocalParameter, ClothShaderData.GetClothToLocalForReading(false, FrameNumber));
+		ShaderBindings.Add(ClothSimulVertsPositionsNormalsParameter, ClothShaderData.GetClothBufferForReading(false).VertexBufferSRV);
+		ShaderBindings.Add(ClothToLocalParameter, ClothShaderData.GetClothToLocalForReading(false));
 		ShaderBindings.Add(ClothBlendWeightParameter,ClothShaderData.ClothBlendWeight);
 		ShaderBindings.Add(ClothNumInfluencesPerVertexParameter, ClothShaderData.NumInfluencesPerVertex);
 
@@ -1097,8 +1095,8 @@ public:
 		const bool bIsMobile = IsMobilePlatform(ShaderPlatform);
 		if (!bIsMobile)
 		{
-			ShaderBindings.Add(PreviousClothSimulVertsPositionsNormalsParameter, ClothShaderData.GetClothBufferForReading(true, FrameNumber).VertexBufferSRV);
-			ShaderBindings.Add(PreviousClothToLocalParameter, ClothShaderData.GetClothToLocalForReading(true, FrameNumber));
+			ShaderBindings.Add(PreviousClothSimulVertsPositionsNormalsParameter, ClothShaderData.GetClothBufferForReading(true).VertexBufferSRV);
+			ShaderBindings.Add(PreviousClothToLocalParameter, ClothShaderData.GetClothToLocalForReading(true));
 		}
 
 		ShaderBindings.Add(GPUSkinApexClothParameter, ClothVertexFactory->GetClothBuffer());
@@ -1126,7 +1124,7 @@ IMPLEMENT_TYPE_LAYOUT(TGPUSkinAPEXClothVertexFactoryShaderParameters);
 -----------------------------------------------------------------------------*/
 
 bool FGPUBaseSkinAPEXClothVertexFactory::ClothShaderType::UpdateClothSimulData(FRHICommandListImmediate& RHICmdList, const TArray<FVector3f>& InSimulPositions,
-	const TArray<FVector3f>& InSimulNormals, uint32 FrameNumberToPrepare, ERHIFeatureLevel::Type FeatureLevel)
+	const TArray<FVector3f>& InSimulNormals, uint32 RevisionNumber, ERHIFeatureLevel::Type FeatureLevel)
 {
 	QUICK_SCOPE_CYCLE_COUNTER(STAT_FGPUBaseSkinAPEXClothVertexFactory_UpdateClothSimulData);
 
@@ -1134,7 +1132,8 @@ bool FGPUBaseSkinAPEXClothVertexFactory::ClothShaderType::UpdateClothSimulData(F
 
 	check(IsInRenderingThread());
 	
-	FVertexBufferAndSRV* CurrentClothBuffer = &GetClothBufferForWriting(FrameNumberToPrepare);
+	SetCurrentRevisionNumber(RevisionNumber);
+	FVertexBufferAndSRV* CurrentClothBuffer = &GetClothBufferForWriting();
 
 	NumSimulVerts = FMath::Min(NumSimulVerts, (uint32)MAX_APEXCLOTH_VERTICES_FOR_VB);
 
@@ -1204,6 +1203,90 @@ bool FGPUBaseSkinAPEXClothVertexFactory::ClothShaderType::UpdateClothSimulData(F
 	return false;
 }
 
+void FGPUBaseSkinAPEXClothVertexFactory::ClothShaderType::SetCurrentRevisionNumber(uint32 RevisionNumber)
+{
+	if (bDoubleBuffer)
+	{
+		// Flip revision number to previous if this is new, otherwise keep current version.
+		if (CurrentRevisionNumber != RevisionNumber)
+		{
+			PreviousRevisionNumber = CurrentRevisionNumber;
+			CurrentRevisionNumber = RevisionNumber;
+			CurrentBuffer = 1 - CurrentBuffer;
+		}
+	}
+}
+
+FVertexBufferAndSRV& FGPUBaseSkinAPEXClothVertexFactory::ClothShaderType::GetClothBufferForWriting()
+{
+	uint32 Index = GetClothBufferIndexForWriting();
+	return ClothSimulPositionNormalBuffer[Index];
+}
+
+bool FGPUBaseSkinAPEXClothVertexFactory::ClothShaderType::HasClothBufferForReading(bool bPrevious) const
+{
+	uint32 Index = GetClothBufferIndexForReading(bPrevious);
+	return ClothSimulPositionNormalBuffer[Index].VertexBufferRHI.IsValid();
+}
+
+const FVertexBufferAndSRV& FGPUBaseSkinAPEXClothVertexFactory::ClothShaderType::GetClothBufferForReading(bool bPrevious) const
+{
+	uint32 Index = GetClothBufferIndexForReading(bPrevious);
+	checkf(ClothSimulPositionNormalBuffer[Index].VertexBufferRHI.IsValid(), TEXT("Index: %i Buffer0: %s Buffer1: %s"), Index, ClothSimulPositionNormalBuffer[0].VertexBufferRHI.IsValid() ? TEXT("true") : TEXT("false"), ClothSimulPositionNormalBuffer[1].VertexBufferRHI.IsValid() ? TEXT("true") : TEXT("false"));
+	return ClothSimulPositionNormalBuffer[Index];
+}
+
+FMatrix44f& FGPUBaseSkinAPEXClothVertexFactory::ClothShaderType::GetClothToLocalForWriting()
+{
+	uint32 Index = GetClothBufferIndexForWriting();
+	return ClothToLocal[Index];
+}
+
+const FMatrix44f& FGPUBaseSkinAPEXClothVertexFactory::ClothShaderType::GetClothToLocalForReading(bool bPrevious) const
+{
+	uint32 Index = GetClothBufferIndexForReading(bPrevious);
+	return ClothToLocal[Index];
+}
+
+uint32 FGPUBaseSkinAPEXClothVertexFactory::ClothShaderType::GetClothBufferIndexInternal(bool bPrevious) const
+{
+	uint32 BufferIndex = 0;
+	if (bDoubleBuffer)
+	{
+		if ((CurrentRevisionNumber - PreviousRevisionNumber) > 1)
+		{
+			// If the revision number has incremented too much, ignore the request and use the current buffer.
+			// With ClearMotionVector calls, we intentionally increment revision number to retrieve current buffer for bPrevious true.
+			bPrevious = false;
+		}
+
+		BufferIndex = CurrentBuffer ^ (uint32)bPrevious;
+	}
+	return BufferIndex;
+}
+
+uint32 FGPUBaseSkinAPEXClothVertexFactory::ClothShaderType::GetClothBufferIndexForWriting() const
+{
+	return bDoubleBuffer ? GetClothBufferIndexInternal(false) : 0;
+}
+
+uint32 FGPUBaseSkinAPEXClothVertexFactory::ClothShaderType::GetClothBufferIndexForReading(bool bPrevious) const
+{
+	uint32 BufferIndex = 0;
+	if (bDoubleBuffer)
+	{
+		BufferIndex = GetClothBufferIndexInternal(bPrevious);
+		if (!ClothSimulPositionNormalBuffer[BufferIndex].VertexBufferRHI.IsValid())
+		{
+			// This only could happen first time updating when the previous data is not available
+			check(bPrevious);
+			// If no previous data available, use the current one
+			BufferIndex = GetClothBufferIndexInternal(false);
+		}
+	}
+	return BufferIndex;
+}
+
 /*-----------------------------------------------------------------------------
 	TGPUSkinAPEXClothVertexFactory
 -----------------------------------------------------------------------------*/
diff --git a/Engine/Source/Runtime/Engine/Private/SkeletalMeshDeformerHelpers.cpp b/Engine/Source/Runtime/Engine/Private/SkeletalMeshDeformerHelpers.cpp
--- a/Engine/Source/Runtime/Engine/Private/SkeletalMeshDeformerHelpers.cpp
+++ b/Engine/Source/Runtime/Engine/Private/SkeletalMeshDeformerHelpers.cpp
@@ -40,7 +40,7 @@ FRHIShaderResourceView* FSkeletalMeshDeformerHelpers::GetMorphTargetBufferForRea
 
 	FSkeletalMeshObjectGPUSkin const* MeshObjectGPU = static_cast<FSkeletalMeshObjectGPUSkin const*>(MeshObject);
 	FGPUBaseSkinVertexFactory const* BaseVertexFactory = MeshObjectGPU->GetBaseSkinVertexFactory(LODIndex, SectionIndex);
-	FMorphVertexBuffer const* MorphVertexBuffer = BaseVertexFactory != nullptr ? BaseVertexFactory->GetMorphVertexBuffer(bPreviousFrame, FrameNumber) : nullptr;
+	FMorphVertexBuffer const* MorphVertexBuffer = BaseVertexFactory != nullptr ? BaseVertexFactory->GetMorphVertexBuffer(bPreviousFrame) : nullptr;
 
 	return MorphVertexBuffer != nullptr ? MorphVertexBuffer->GetSRV() : nullptr;
 }
@@ -61,7 +61,7 @@ FSkeletalMeshDeformerHelpers::FClothBuffers FSkeletalMeshDeformerHelpers::GetClo
 	FGPUBaseSkinVertexFactory const* BaseVertexFactory = MeshObjectGPU->GetBaseSkinVertexFactory(LODIndex, SectionIndex);
 	FGPUBaseSkinAPEXClothVertexFactory const* ClothVertexFactory = BaseVertexFactory != nullptr ? BaseVertexFactory->GetClothVertexFactory() : nullptr;
 
-	if (ClothVertexFactory == nullptr || !ClothVertexFactory->GetClothShaderData().HasClothBufferForReading(bPreviousFrame, FrameNumber))
+	if (ClothVertexFactory == nullptr || !ClothVertexFactory->GetClothShaderData().HasClothBufferForReading(bPreviousFrame))
 	{
 		return FClothBuffers();
 	}
@@ -73,8 +73,8 @@ FSkeletalMeshDeformerHelpers::FClothBuffers FSkeletalMeshDeformerHelpers::GetClo
 	FClothBuffers Ret;
 	Ret.ClothInfluenceBuffer = ClothVertexFactory->GetClothBuffer();
 	Ret.ClothInfluenceBufferOffset = ClothVertexFactory->GetClothIndexOffset(RenderSection.BaseVertexIndex);
-	Ret.ClothSimulatedPositionAndNormalBuffer = ClothVertexFactory->GetClothShaderData().GetClothBufferForReading(bPreviousFrame, FrameNumber).VertexBufferSRV;
-	Ret.ClothToLocal = ClothVertexFactory->GetClothShaderData().GetClothToLocalForReading(bPreviousFrame, FrameNumber);
+	Ret.ClothSimulatedPositionAndNormalBuffer = ClothVertexFactory->GetClothShaderData().GetClothBufferForReading(bPreviousFrame).VertexBufferSRV;
+	Ret.ClothToLocal = ClothVertexFactory->GetClothShaderData().GetClothToLocalForReading(bPreviousFrame);
 	return Ret;
 }
 
diff --git a/Engine/Source/Runtime/Engine/Private/SkeletalRenderGPUSkin.cpp b/Engine/Source/Runtime/Engine/Private/SkeletalRenderGPUSkin.cpp
--- a/Engine/Source/Runtime/Engine/Private/SkeletalRenderGPUSkin.cpp
+++ b/Engine/Source/Runtime/Engine/Private/SkeletalRenderGPUSkin.cpp
@@ -173,78 +173,52 @@ void FMorphVertexBufferPool::EnableDoubleBuffer()
 	}
 }
 
-const FMorphVertexBuffer& FMorphVertexBufferPool::GetMorphVertexBufferForReading(bool bPrevious, uint32 FrameNumber) const
+void FMorphVertexBufferPool::SetCurrentRevisionNumber(uint32 RevisionNumber)
 {
-	uint32 Index = 0;
 	if (bDoubleBuffer)
 	{
-		// Find the buffer matching frame number or otherwise use the newer buffer
-		if (MorphVertexBuffers[0].FrameNumber == FrameNumber)
-		{
-			Index = 0;
-		}
-		else if (MorphVertexBuffers[1].FrameNumber == FrameNumber)
+		// Flip revision number to previous if this is new, otherwise keep current version.
+		if (CurrentRevisionNumber != RevisionNumber)
 		{
-			Index = 1;
-		}
-		else if (MorphVertexBuffers[0].FrameNumber == -1)
-		{
-			Index = 0;
-		}
-		else if (MorphVertexBuffers[1].FrameNumber == -1)
-		{
-			Index = 1;
-		}
-		else
-		{
-			// Get the newer frame
-			Index = (MorphVertexBuffers[0].FrameNumber < MorphVertexBuffers[1].FrameNumber) ? 1 : 0;
-		}
-
-		bool bPreviousValid = (MorphVertexBuffers[1 - Index].FrameNumber != -1);
-		if (bPrevious && bPreviousValid)
-		{
-			Index = 1 - Index;
+			PreviousRevisionNumber = CurrentRevisionNumber;
+			CurrentRevisionNumber = RevisionNumber;
+			CurrentBuffer = 1 - CurrentBuffer;
 		}
 	}
-
-	checkf(MorphVertexBuffers[Index].VertexBufferRHI.IsValid(), TEXT("Index: %i Buffer0: %s Frame0: %i Buffer1: %s Frame1: %i"), Index, MorphVertexBuffers[0].VertexBufferRHI.IsValid() ? TEXT("true") : TEXT("false"), MorphVertexBuffers[0].FrameNumber, MorphVertexBuffers[1].VertexBufferRHI.IsValid() ? TEXT("true") : TEXT("false"), MorphVertexBuffers[1].FrameNumber);
-	return MorphVertexBuffers[Index];
 }
 
-FMorphVertexBuffer& FMorphVertexBufferPool::GetMorphVertexBufferForWriting(uint32 FrameNumber)
+const FMorphVertexBuffer& FMorphVertexBufferPool::GetMorphVertexBufferForReading(bool bPrevious) const
 {
 	uint32 Index = 0;
 	if (bDoubleBuffer)
 	{
-		// Find the buffer matching frame number or otherwise use the older buffer
-		if (MorphVertexBuffers[0].FrameNumber == FrameNumber)
+		if ((CurrentRevisionNumber - PreviousRevisionNumber) > 1)
 		{
-			Index = 0;
+			// If the revision number has incremented too much, ignore the request and use the current buffer.
+			// With ClearMotionVector calls, we intentionally increment revision number to retrieve current buffer for bPrevious true.
+			bPrevious = false;
 		}
-		else if (MorphVertexBuffers[1].FrameNumber == FrameNumber)
-		{
-			Index = 1;
-		}
-		else if (MorphVertexBuffers[0].FrameNumber == -1)
+
+		Index = CurrentBuffer ^ (uint32)bPrevious;
+
+		if (!MorphVertexBuffers[Index].bHasBeenUpdated)
 		{
+			// this should only happen the first time updating, in which case the previous buffer hasn't been written into yet.
+			check(Index == 1);
+			check(MorphVertexBuffers[0].bHasBeenUpdated);
 			Index = 0;
 		}
-		else if (MorphVertexBuffers[1].FrameNumber == -1)
-		{
-			Index = 1;
-		}
-		else
-		{
-			// Get the older frame
-			Index = (MorphVertexBuffers[0].FrameNumber < MorphVertexBuffers[1].FrameNumber) ? 0 : 1;
-		}
 	}
 
-	MorphVertexBuffers[Index].FrameNumber = FrameNumber;
+	checkf(MorphVertexBuffers[Index].VertexBufferRHI.IsValid(), TEXT("Index: %i Buffer0: %s Buffer1: %s"), Index, MorphVertexBuffers[0].VertexBufferRHI.IsValid() ? TEXT("true") : TEXT("false"), MorphVertexBuffers[1].VertexBufferRHI.IsValid() ? TEXT("true") : TEXT("false"));
 	return MorphVertexBuffers[Index];
 }
 
+FMorphVertexBuffer& FMorphVertexBufferPool::GetMorphVertexBufferForWriting()
+{
+	return MorphVertexBuffers[CurrentBuffer];
+}
+
 /*-----------------------------------------------------------------------------
 FSkeletalMeshObjectGPUSkin
 -----------------------------------------------------------------------------*/
@@ -592,13 +566,6 @@ void FSkeletalMeshObjectGPUSkin::ProcessUpdatedDynamicData(EGPUSkinCacheEntryMod
 	bMorphNeedsUpdateDeferred = false;
 
 	FSkeletalMeshObjectLOD& LOD = LODs[LODIndex];
-	FMorphVertexBuffer& MorphVertexBuffer = LOD.MorphVertexBufferPool.GetMorphVertexBufferForWriting(FrameNumberToPrepare);
-
-	// if hasn't been updated, force update again
-	bMorphNeedsUpdate = MorphVertexBuffer.bHasBeenUpdated ? bMorphNeedsUpdate : true;
-	bMorphNeedsUpdate |= (GForceUpdateMorphTargets != 0);
-	bMorphNeedsUpdate |= bAlwaysUpdateMorphVertexBuffer;
-
 	const FSkeletalMeshLODRenderData& LODData = SkeletalMeshRenderData->LODRenderData[LODIndex];
 	const TArray<FSkelMeshRenderSection>& Sections = GetRenderSections(LODIndex);
 
@@ -608,6 +575,18 @@ void FSkeletalMeshObjectGPUSkin::ProcessUpdatedDynamicData(EGPUSkinCacheEntryMod
 		(DynamicData->NumWeightedActiveMorphTargets > 0 && LODData.GetNumVertices() > 0) || 
 		(DynamicData->ExternalMorphWeightData.HasActiveMorphs() && LODData.GetNumVertices() > 0);
 
+	bMorphNeedsUpdate |= (GForceUpdateMorphTargets != 0);
+	bMorphNeedsUpdate |= bAlwaysUpdateMorphVertexBuffer;
+	if (bMorph && bMorphNeedsUpdate)
+	{
+		// Morph vertex buffer needs updating so advance the revision number
+		LOD.MorphVertexBufferPool.SetCurrentRevisionNumber(RevisionNumber);
+	}
+
+	FMorphVertexBuffer& MorphVertexBuffer = LOD.MorphVertexBufferPool.GetMorphVertexBufferForWriting();
+	// if morph buffer hasn't been updated, force update it
+	bMorphNeedsUpdate = MorphVertexBuffer.bHasBeenUpdated ? bMorphNeedsUpdate : true;
+
 	// use correct vertex factories based on alternate weights usage
 	FVertexFactoryData& VertexFactoryData = LOD.GPUSkinVertexFactories;
 
@@ -722,15 +701,15 @@ void FSkeletalMeshObjectGPUSkin::ProcessUpdatedDynamicData(EGPUSkinCacheEntryMod
 				FGPUBaseSkinAPEXClothVertexFactory::ClothShaderType& ClothShaderData = VertexFactoryData.ClothVertexFactories[SectionIdx]->GetClothShaderData();
 				ClothShaderData.ClothBlendWeight = DynamicData->ClothBlendWeight;
 
-				bNeedFence = ClothShaderData.UpdateClothSimulData(RHICmdList, SimData->Positions, SimData->Normals, FrameNumberToPrepare, FeatureLevel) || bNeedFence;
+				bNeedFence = ClothShaderData.UpdateClothSimulData(RHICmdList, SimData->Positions, SimData->Normals, RevisionNumber, FeatureLevel) || bNeedFence;
 				// Transform from cloth space to local space. Cloth space is relative to cloth root bone, local space is component space.
-				ClothShaderData.GetClothToLocalForWriting(FrameNumberToPrepare) = FMatrix44f(SimData->ComponentRelativeTransform.ToMatrixWithScale());
+				ClothShaderData.GetClothToLocalForWriting() = FMatrix44f(SimData->ComponentRelativeTransform.ToMatrixWithScale());
 			}
 
 			// Try to use the GPU skinning cache if possible
 			if (bUseSkinCache)
 			{
-				FMatrix44f ClothToLocal = bClothFactory ? VertexFactoryData.ClothVertexFactories[SectionIdx]->GetClothShaderData().GetClothToLocalForWriting(FrameNumberToPrepare) : FMatrix44f::Identity;
+				FMatrix44f ClothToLocal = bClothFactory ? VertexFactoryData.ClothVertexFactories[SectionIdx]->GetClothShaderData().GetClothToLocalForWriting() : FMatrix44f::Identity;
 
 				// ProcessEntry returns false if not enough memory is left in skin cache to allocate for the mesh, if that happens don't try to process subsequent sections because they will also fail.
 				if (bSkinCacheResult)
@@ -1603,7 +1582,7 @@ void FSkeletalMeshObjectGPUSkin::RefreshClothingTransforms(const FMatrix& InNewL
 
 					if(FClothSimulData* SimData = DynamicData->ClothingSimData.Find(ActorIdx))
 					{
-						ClothShaderData.GetClothToLocalForWriting(FrameNumber) = FMatrix44f(SimData->ComponentRelativeTransform.ToMatrixWithScale());
+						ClothShaderData.GetClothToLocalForWriting() = FMatrix44f(SimData->ComponentRelativeTransform.ToMatrixWithScale());
 					}
 				}
 			}
@@ -1688,7 +1667,7 @@ void InitMorphVertexFactoryComponents(FGPUSkinDataType* VertexFactoryData,
 {
 	VertexFactoryData->bMorphTarget = true;
 	VertexFactoryData->MorphVertexBufferPool = VertexBuffers.MorphVertexBufferPool;
-	const FMorphVertexBuffer& MorphVB = VertexBuffers.MorphVertexBufferPool->GetMorphVertexBufferForReading(false, -1);
+	const FMorphVertexBuffer& MorphVB = VertexBuffers.MorphVertexBufferPool->GetMorphVertexBufferForReading(false);
 
 	// delta positions
 	VertexFactoryData->DeltaPositionComponent = FVertexStreamComponent(
diff --git a/Engine/Source/Runtime/Engine/Private/SkeletalRenderGPUSkin.h b/Engine/Source/Runtime/Engine/Private/SkeletalRenderGPUSkin.h
--- a/Engine/Source/Runtime/Engine/Private/SkeletalRenderGPUSkin.h
+++ b/Engine/Source/Runtime/Engine/Private/SkeletalRenderGPUSkin.h
@@ -290,9 +290,6 @@ private:
 	// parent mesh containing the source data, never 0
 	FSkeletalMeshRenderData* SkelMeshRenderData;
 
-	/** Latest updated frame number, -1 means invalid */
-	uint32 FrameNumber = -1;
-
 	friend class FMorphVertexBufferPool;
 };
 
@@ -313,14 +310,21 @@ public:
 	SIZE_T GetResourceSize() const;
 	void EnableDoubleBuffer();
 	bool IsDoubleBuffered() const	{ return bDoubleBuffer; }
-	const FMorphVertexBuffer& GetMorphVertexBufferForReading(bool bPrevious, uint32 FrameNumber) const;
-	FMorphVertexBuffer& GetMorphVertexBufferForWriting(uint32 FrameNumber);
+	void SetCurrentRevisionNumber(uint32 RevisionNumber);
+	const FMorphVertexBuffer& GetMorphVertexBufferForReading(bool bPrevious) const;
+	FMorphVertexBuffer& GetMorphVertexBufferForWriting();
 
 private:
 	/** Vertex buffer that stores the morph target vertex deltas. */
 	FMorphVertexBuffer MorphVertexBuffers[2];
 	/** whether to double buffer. If going through skin cache, then use single buffer; otherwise double buffer. */
 	bool bDoubleBuffer = false;
+
+	// 0 / 1 to index into MorphVertexBuffer
+	uint32 CurrentBuffer = 0;
+	// RevisionNumber Tracker
+	uint32 PreviousRevisionNumber = 0;
+	uint32 CurrentRevisionNumber = 0;
 };
 
 /**
diff --git a/Engine/Source/Runtime/Engine/Public/GPUSkinVertexFactory.h b/Engine/Source/Runtime/Engine/Public/GPUSkinVertexFactory.h
--- a/Engine/Source/Runtime/Engine/Public/GPUSkinVertexFactory.h
+++ b/Engine/Source/Runtime/Engine/Public/GPUSkinVertexFactory.h
@@ -324,6 +324,8 @@ public:
 
 			if ((CurrentRevisionNumber - PreviousRevisionNumber) > 1)
 			{
+				// If the revision number has incremented too much, ignore the request and use the current buffer.
+				// With ClearMotionVector calls, we intentionally increment revision number to retrieve current buffer for bPrevious true.
 				bPrevious = false;
 			}
 
@@ -378,7 +380,7 @@ public:
 
 	/** Morph vertex factory functions */
 	virtual void UpdateMorphVertexStream(const class FMorphVertexBuffer* MorphVertexBuffer) {}
-	virtual const class FMorphVertexBuffer* GetMorphVertexBuffer(bool bPrevious, uint32 FrameNumber) const { return nullptr; }
+	virtual const class FMorphVertexBuffer* GetMorphVertexBuffer(bool bPrevious) const { return nullptr; }
 	/** Cloth vertex factory access. */
 	virtual class FGPUBaseSkinAPEXClothVertexFactory* GetClothVertexFactory() { return nullptr; }
 	virtual class FGPUBaseSkinAPEXClothVertexFactory const* GetClothVertexFactory() const { return nullptr; }
@@ -453,7 +455,7 @@ public:
 	
 	/** FGPUBaseSkinVertexFactory overrides */
 	virtual void UpdateMorphVertexStream(const class FMorphVertexBuffer* MorphVertexBuffer) override;
-	virtual const class FMorphVertexBuffer* GetMorphVertexBuffer(bool bPrevious, uint32 FrameNumber) const override;
+	virtual const class FMorphVertexBuffer* GetMorphVertexBuffer(bool bPrevious) const override;
 
 
 	// FRenderResource interface.
@@ -609,7 +611,7 @@ public:
 			Reset();
 		}
 
-		bool UpdateClothSimulData(FRHICommandListImmediate& RHICmdList, const TArray<FVector3f>& InSimulPositions, const TArray<FVector3f>& InSimulNormals, uint32 FrameNumber, ERHIFeatureLevel::Type FeatureLevel);
+		bool UpdateClothSimulData(FRHICommandListImmediate& RHICmdList, const TArray<FVector3f>& InSimulPositions, const TArray<FVector3f>& InSimulNormals, uint32 RevisionNumber, ERHIFeatureLevel::Type FeatureLevel);
 
 		void ReleaseClothSimulData()
 		{
@@ -632,72 +634,15 @@ public:
 		{
 			return APEXClothUniformBuffer;
 		}
-		
-		// @param FrameNumber usually from View.Family->FrameNumber
-		// @return IsValid() can fail, then you have to create the buffers first (or if the size changes)
-		FVertexBufferAndSRV& GetClothBufferForWriting(uint32 FrameNumber)
-		{
-			uint32 Index = GetOldestIndex(FrameNumber);
-			Index = (BufferFrameNumber[0] == FrameNumber) ? 0 : Index;
-			Index = (BufferFrameNumber[1] == FrameNumber) ? 1 : Index;
-
-			// we don't write -1 as that is used to invalidate the entry
-			if(FrameNumber == -1)
-			{
-				// this could cause a 1 frame glitch on wraparound
-				FrameNumber = 0;
-			}
-
-			BufferFrameNumber[Index] = FrameNumber;
-
-			return ClothSimulPositionNormalBuffer[Index];
-		}
-
-		bool HasClothBufferForReading(bool bPrevious, uint32 FrameNumber) const
-		{
-			int32 Index = GetMostRecentIndex(FrameNumber);
-			if (bPrevious && DoWeHavePreviousData())
-			{
-				Index = 1 - Index;
-			}
-			return ClothSimulPositionNormalBuffer[Index].VertexBufferRHI.IsValid();
-		}
-
-		// @param bPrevious true:previous, false:current
-		// @param FrameNumber usually from View.Family->FrameNumber
-		const FVertexBufferAndSRV& GetClothBufferForReading(bool bPrevious, uint32 FrameNumber) const
-		{
-			int32 Index = GetMostRecentIndex(FrameNumber);
-
-			if(bPrevious && DoWeHavePreviousData())
-			{
-				Index = 1 - Index;
-			}
-
-			checkf(ClothSimulPositionNormalBuffer[Index].VertexBufferRHI.IsValid(), TEXT("Index: %i Buffer0: %s Frame0: %i Buffer1: %s Frame1: %i"), Index,  ClothSimulPositionNormalBuffer[0].VertexBufferRHI.IsValid() ? TEXT("true") : TEXT("false"), BufferFrameNumber[0], ClothSimulPositionNormalBuffer[1].VertexBufferRHI.IsValid() ? TEXT("true") : TEXT("false"), BufferFrameNumber[1]);
-			return ClothSimulPositionNormalBuffer[Index];
-		}
-		
-		FMatrix44f& GetClothToLocalForWriting(uint32 FrameNumber)
-		{
-			uint32 Index = GetOldestIndex(FrameNumber);
-			Index = (BufferFrameNumber[0] == FrameNumber) ? 0 : Index;
-			Index = (BufferFrameNumber[1] == FrameNumber) ? 1 : Index;
 
-			return ClothToLocal[Index];
-		}
-
-		const FMatrix44f& GetClothToLocalForReading(bool bPrevious, uint32 FrameNumber) const
-		{
-			int32 Index = GetMostRecentIndex(FrameNumber);
+		void SetCurrentRevisionNumber(uint32 RevisionNumber);
 
-			if(bPrevious && DoWeHavePreviousData())
-			{
-				Index = 1 - Index;
-			}
+		FVertexBufferAndSRV& GetClothBufferForWriting();
+		bool HasClothBufferForReading(bool bPrevious) const;
+		const FVertexBufferAndSRV& GetClothBufferForReading(bool bPrevious) const;
 
-			return ClothToLocal[Index];
-		}
+		FMatrix44f& GetClothToLocalForWriting();
+		const FMatrix44f& GetClothToLocalForReading(bool bPrevious) const;
 
 		/**
 		 * weight to blend between simulated positions and key-framed poses
@@ -707,12 +652,17 @@ public:
 		uint32 NumInfluencesPerVertex = 1;
 
 	private:
+		// Helper for GetClothBufferIndexForWriting and GetClothBufferIndexForReading
+		uint32 GetClothBufferIndexInternal(bool bPrevious) const;
+		// Helper for GetClothBufferForWriting and GetClothToLocalForWriting
+		uint32 GetClothBufferIndexForWriting() const;
+		// Helper for GetClothBufferForReading and GetClothToLocalForReading
+		uint32 GetClothBufferIndexForReading(bool bPrevious) const;
+
 		// fallback for ClothSimulPositionNormalBuffer if the shadermodel doesn't allow it
 		TUniformBufferRef<FAPEXClothUniformShaderParameters> APEXClothUniformBuffer;
 		// 
 		FVertexBufferAndSRV ClothSimulPositionNormalBuffer[2];
-		// from GFrameNumber, to detect pause and old data when an object was not rendered for some time
-		uint32 BufferFrameNumber[2];
 
 		/**
 		 * Matrix to apply to positions/normals
@@ -722,78 +672,17 @@ public:
 		/** Whether to double buffer. */
 		bool bDoubleBuffer = false;
 
-		// @return 0 / 1, index into ClothSimulPositionNormalBuffer[]
-		uint32 GetMostRecentIndex(uint32 FrameNumber) const
-		{
-			if (!bDoubleBuffer)
-			{
-				return 0;
-			}
-
-			if(BufferFrameNumber[0] == -1)
-			{
-				//ensure(BufferFrameNumber[1] != -1);
-
-				return 1;
-			}
-			else if(BufferFrameNumber[1] == -1)
-			{
-				//ensure(BufferFrameNumber[0] != -1);
-				return 0;
-			}
-
-			// should handle warp around correctly, did some basic testing
-			uint32 Age0 = FrameNumber - BufferFrameNumber[0];
-			uint32 Age1 = FrameNumber - BufferFrameNumber[1];
-
-			return (Age0 > Age1) ? 1 : 0;
-		}
-
-		// @return 0/1, index into ClothSimulPositionNormalBuffer[]
-		uint32 GetOldestIndex(uint32 FrameNumber) const
-		{
-			if (!bDoubleBuffer)
-			{
-				return 0;
-			}
-
-			if(BufferFrameNumber[0] == -1)
-			{
-				return 0;
-			}
-			else if(BufferFrameNumber[1] == -1)
-			{
-				return 1;
-			}
-
-			// should handle warp around correctly (todo: test)
-			uint32 Age0 = FrameNumber - BufferFrameNumber[0];
-			uint32 Age1 = FrameNumber - BufferFrameNumber[1];
-
-			return (Age0 > Age1) ? 0 : 1;
-		}
-
-		bool DoWeHavePreviousData() const
-		{
-			if(BufferFrameNumber[0] == -1 || BufferFrameNumber[1] == -1)
-			{
-				return false;
-			}
-			
-			int32 Diff = BufferFrameNumber[0] - BufferFrameNumber[1];
-
-			uint32 DiffAbs = FMath::Abs(Diff);
-
-			// threshold is >1 because there could be in between frames e.g. HitProxyRendering
-			// We should switch to TickNumber to solve this
-			return DiffAbs <= 2;
-		}
+		// 0 / 1 to index into BoneBuffer
+		uint32 CurrentBuffer = 0;
+		// RevisionNumber Tracker
+		uint32 PreviousRevisionNumber = 0;
+		uint32 CurrentRevisionNumber = 0;
 
 		void Reset()
 		{
-			// both are not valid
-			BufferFrameNumber[0] = -1;
-			BufferFrameNumber[1] = -1;
+			CurrentBuffer = 0;
+			PreviousRevisionNumber = 0;
+			CurrentRevisionNumber = 0;
 
 			ClothToLocal[0] = FMatrix44f::Identity;
 			ClothToLocal[1] = FMatrix44f::Identity;
