Exclude PIE pakages when looking for world assets when resolving URLs.

diff --git a/Engine/Source/Runtime/Engine/Private/URL.cpp b/Engine/Source/Runtime/Engine/Private/URL.cpp
--- a/Engine/Source/Runtime/Engine/Private/URL.cpp
+++ b/Engine/Source/Runtime/Engine/Private/URL.cpp
@@ -478,12 +478,17 @@ FURL::FURL( FURL* Base, const TCHAR* TextURL, ETravelType Type )
 					if (!AssetRegistry.IsLoadingAssets())
 					{
 						TArray<FAssetData> MapList;
-						if (AssetRegistry.GetAssetsByClass(UWorld::StaticClass()->GetClassPathName(), /*out*/ MapList))
+
+						FARFilter Filter;
+						Filter.ClassPaths.Add(UWorld::StaticClass()->GetClassPathName());
+						Filter.WithoutPackageFlags = PKG_CookGenerated | PKG_PlayInEditor;
+
+						if (AssetRegistry.GetAssets(Filter, /*out*/ MapList))
 						{
 							FName TargetTestName(*URLStr);
 							for (const FAssetData& MapAsset : MapList)
 							{
-								if (MapAsset.AssetName == TargetTestName && !MapAsset.HasAnyPackageFlags(PKG_CookGenerated))
+								if (MapAsset.AssetName == TargetTestName)
 								{
 									Map = MapAsset.PackageName.ToString();
 									bFoundMap = true;
