Fix crash when changing the default ISM NumCustomDataFloats on a BP and then loading a level where that BP was already saved.

Also fixes crash where NumCustomDataFloats was changed on a placed BP in a level.

This also adds the feature of being able to retain the per instance custom data on placed BP instances when changing the number of Custom Data Floats on the BP asset. This allows users to change the default BP and still have every placed instance work as expected when sampling the custom data from a material.

#rnx

diff --git a/Engine/Source/Runtime/Engine/Private/InstancedStaticMesh.cpp b/Engine/Source/Runtime/Engine/Private/InstancedStaticMesh.cpp
--- a/Engine/Source/Runtime/Engine/Private/InstancedStaticMesh.cpp
+++ b/Engine/Source/Runtime/Engine/Private/InstancedStaticMesh.cpp
@@ -4471,13 +4471,39 @@ void UInstancedStaticMeshComponent::PostLoad()
 	Super::PostLoad();
 
 	// Ensure we have the proper amount of per instance custom float data.
-	// We have instances, and we have custom float data per instance, but we don't have custom float data allocated.
-	if (PerInstanceSMData.Num() > 0 && 
-		PerInstanceSMCustomData.Num() == 0 &&
-		NumCustomDataFloats > 0)
+	// We have instances, and we have custom float data per instance, but we don't have the right amount of custom float data allocated.
+	if ((PerInstanceSMData.Num() * NumCustomDataFloats) != PerInstanceSMCustomData.Num() && PerInstanceSMData.Num() > 0)
 	{
-		UE_LOG(LogStaticMesh, Warning, TEXT("%s has %d instances, and %d NumCustomDataFloats, but no PerInstanceSMCustomData.  Allocating %d custom floats."), *GetFullName(), PerInstanceSMData.Num(), NumCustomDataFloats, PerInstanceSMData.Num() * NumCustomDataFloats);
-		PerInstanceSMCustomData.AddZeroed(PerInstanceSMData.Num() * NumCustomDataFloats);
+		// No custom data at all, add all zeroes
+		if (PerInstanceSMCustomData.Num() == 0)
+		{
+			UE_LOG(LogStaticMesh, Warning, TEXT("%s has %d instances, and %d NumCustomDataFloats, but no PerInstanceSMCustomData.  Allocating %d custom floats."), *GetFullName(), PerInstanceSMData.Num(), NumCustomDataFloats, PerInstanceSMData.Num() * NumCustomDataFloats);
+			PerInstanceSMCustomData.AddZeroed(PerInstanceSMData.Num() * NumCustomDataFloats);
+		}
+		else
+		{
+			// Custom data exists, so we preserve it in our new allocation
+			UE_LOG(LogStaticMesh, Warning, TEXT("%s has %d instances, and %d NumCustomDataFloats, but has %d PerInstanceSMCustomData.  Allocating %d custom floats to match."), *GetFullName(), PerInstanceSMData.Num(), NumCustomDataFloats, PerInstanceSMCustomData.Num(), PerInstanceSMData.Num() * NumCustomDataFloats);
+
+			const int32 NumCustomDataFloatsPrev = PerInstanceSMCustomData.Num() / PerInstanceSMData.Num();
+
+			TArray<float> OldPerInstanceSMCustomData;
+			Exchange(OldPerInstanceSMCustomData, PerInstanceSMCustomData);
+
+			// Allocate the proper amount and zero it
+			PerInstanceSMCustomData.AddZeroed(PerInstanceSMData.Num() * NumCustomDataFloats);
+
+			const int32 NumFloatsToCopy = FMath::Min(NumCustomDataFloatsPrev, NumCustomDataFloats);
+
+			// Copy over existing data 
+			for (int32 InstanceID = 0; InstanceID < PerInstanceSMData.Num(); InstanceID++)
+			{
+				for (int32 CopyID = 0; CopyID < NumFloatsToCopy; CopyID++)
+				{
+					PerInstanceSMCustomData[InstanceID * NumCustomDataFloats + CopyID] = OldPerInstanceSMCustomData[InstanceID * NumCustomDataFloatsPrev + CopyID];
+				}
+			}
+		}
 	}
 
 	// Has different implementation in HISMC
