Fix a crash caused by trying to use non resident skeletal mesh LODs when streaming is disabled

diff --git a/Engine/Source/Runtime/Engine/Private/Components/SkinnedMeshComponent.cpp b/Engine/Source/Runtime/Engine/Private/Components/SkinnedMeshComponent.cpp
--- a/Engine/Source/Runtime/Engine/Private/Components/SkinnedMeshComponent.cpp
+++ b/Engine/Source/Runtime/Engine/Private/Components/SkinnedMeshComponent.cpp
@@ -939,7 +939,7 @@ void USkinnedMeshComponent::CreateRenderState_Concurrent(FRegisterComponentConte
 			}
 
 			// Clamp to loaded streaming data if available
-			if (GetSkinnedAsset()->IsStreamable() && MeshObject)
+			if ((GetSkinnedAsset()->IsStreamable() || !IStreamingManager::Get().IsRenderAssetStreamingEnabled(EStreamableRenderAssetType::SkeletalMesh)) && MeshObject)
 			{
 				ModifiedLODLevel = FMath::Max<int32>(ModifiedLODLevel, MeshObject->GetSkeletalMeshRenderData().PendingFirstLODIdx);
 			}
@@ -1040,7 +1040,7 @@ void USkinnedMeshComponent::SendRenderDynamicData_Concurrent()
 
 		int32 UseLOD = GetPredictedLODLevel();
 		// Clamp to loaded streaming data if available
-		if (GetSkinnedAsset()->IsStreamable() && MeshObject)
+		if ((GetSkinnedAsset()->IsStreamable() || !IStreamingManager::Get().IsRenderAssetStreamingEnabled(EStreamableRenderAssetType::SkeletalMesh)) && MeshObject)
 		{
 			UseLOD = FMath::Max<int32>(UseLOD, MeshObject->GetSkeletalMeshRenderData().PendingFirstLODIdx);
 		}
