Add per-platform input settings that let you specify Hardware Input devices that are available on each platform. This will also be where per-platform variables go for force feedback default values.

diff --git a/Engine/Config/Windows/WindowsInput.ini b/Engine/Config/Windows/WindowsInput.ini
new file mode 100644
--- /dev/null
+++ b/Engine/Config/Windows/WindowsInput.ini
@@ -0,0 +1,2 @@
+[InputPlatformSettings_Windows InputPlatformSettings]
++HardwareDevices=(InputClassName="XInputInterface",HardwareDeviceIdentifier="XInputController")
\ No newline at end of file
diff --git a/Engine/Plugins/Experimental/RawInput/Config/Input.ini b/Engine/Plugins/Experimental/RawInput/Config/Input.ini
new file mode 100644
--- /dev/null
+++ b/Engine/Plugins/Experimental/RawInput/Config/Input.ini
@@ -0,0 +1,5 @@
+; Add hardware devices here to detect them at runtime using FInputDeviceScope::GetCurrent
+; @see FRawInputWindows::SendControllerEvents()
+
+; [InputPlatformSettings_Windows InputPlatformSettings]
+; +HardwareDevices=(InputClassName="RawInput",HardwareDeviceIdentifier="<YourDeviceName>")
\ No newline at end of file
diff --git a/Engine/Plugins/Runtime/Steam/SteamController/Config/Input.ini b/Engine/Plugins/Runtime/Steam/SteamController/Config/Input.ini
new file mode 100644
--- /dev/null
+++ b/Engine/Plugins/Runtime/Steam/SteamController/Config/Input.ini
@@ -0,0 +1,2 @@
+[InputPlatformSettings_Windows InputPlatformSettings]
++HardwareDevices=(InputClassName="SteamController",HardwareDeviceIdentifier="SteamController")
\ No newline at end of file
diff --git a/Engine/Plugins/Runtime/Windows/WinDualShock/Config/Input.ini b/Engine/Plugins/Runtime/Windows/WinDualShock/Config/Input.ini
new file mode 100644
--- /dev/null
+++ b/Engine/Plugins/Runtime/Windows/WinDualShock/Config/Input.ini
@@ -0,0 +1,3 @@
+[InputPlatformSettings_Windows InputPlatformSettings]
++HardwareDevices=(InputClassName="FWinDualShock",HardwareDeviceIdentifier="DualSense")
++HardwareDevices=(InputClassName="FWinDualShock",HardwareDeviceIdentifier="DualShock4")
diff --git a/Engine/Source/Runtime/Engine/Classes/GameFramework/InputSettings.h b/Engine/Source/Runtime/Engine/Classes/GameFramework/InputSettings.h
--- a/Engine/Source/Runtime/Engine/Classes/GameFramework/InputSettings.h
+++ b/Engine/Source/Runtime/Engine/Classes/GameFramework/InputSettings.h
@@ -9,6 +9,7 @@
 #include "UObject/SoftObjectPath.h"
 #include "GameFramework/PlayerInput.h"
 #include "Components/InputComponent.h"
+#include "Engine/PlatformSettings.h"
 
 #include "InputSettings.generated.h"
 
@@ -27,6 +28,13 @@ class ENGINE_API UInputSettings
 	UPROPERTY(config, EditAnywhere, EditFixedSize, Category="Bindings", meta=(ToolTip="List of Axis Properties"), AdvancedDisplay)
 	TArray<struct FInputAxisConfigEntry> AxisConfig;
 
+	/**
+	 * Platform specific settings for Input.
+	 * @see UInputPlatformSettings
+	 */
+	UPROPERTY(EditAnywhere, Category = "Enhanced Input")
+	FPerPlatformSettings PlatformSettings;
+
 	UPROPERTY(config, EditAnywhere, Category="Bindings", AdvancedDisplay)
 	uint8 bAltEnterTogglesFullscreen:1;
 
@@ -273,3 +281,77 @@ public:
 private:
 	void PopulateAxisConfigs();
 };
+
+/**
+* An identifier that can be used to determine what input devices are available based on the FInputDeviceScope.
+* These mappings should match a FInputDeviceScope that is used by an IInputDevice
+*/
+USTRUCT(BlueprintType)
+struct FHardwareDeviceIdentifier
+{
+	GENERATED_BODY()
+
+	FHardwareDeviceIdentifier() = default;
+	
+	/** 
+	* The name of the Input Class that uses this hardware device.
+	* This should correspond with a FInputDeviceScope that is used by an IInputDevice
+	*/
+	UPROPERTY(EditDefaultsOnly, BlueprintReadOnly, Category = "Hardware")
+	FName InputClassName;
+
+	/**
+	 * The name of this hardware device. 
+	 * This should correspond with a FInputDeviceScope that is used by an IInputDevice
+	 */
+	UPROPERTY(EditDefaultsOnly, BlueprintReadOnly, Category = "Hardware")
+	FString HardwareDeviceIdentifier;
+
+	bool IsValid() const;
+	
+	static FHardwareDeviceIdentifier Invalid;
+
+	bool operator==(const FHardwareDeviceIdentifier& Other) const
+	{
+		return Other.InputClassName == InputClassName && Other.HardwareDeviceIdentifier == HardwareDeviceIdentifier;
+	}
+};
+
+/** Per-Platform input options */
+UCLASS(config=Input, defaultconfig)
+class ENGINE_API UInputPlatformSettings : public UPlatformSettings
+{
+	GENERATED_BODY()
+
+public:
+
+	static UInputPlatformSettings* Get();
+
+#if WITH_EDITOR
+	/**
+	* Returns an array of Hardware device names from every registered platform ini.
+	* For use in the editor so that you can get a list of all known input devices and 
+	* make device-specific options. For example, you can map any data type to a specific input 
+	* device
+	* 
+	* UPROPERTY(EditAnywhere, BlueprintReadOnly, meta=(GetOptions="Engine.InputPlatformSettings.GetAllHardwareDeviceNames"))
+	* TMap<FString, UFooData> DeviceSpecificMap;
+	* 
+	* and the editor will make a nice drop down for you with all the current options that are in the settings.
+	*/
+	UFUNCTION()
+	static const TArray<FString>& GetAllHardwareDeviceNames();
+#endif	// WITH_EDITOR
+
+	/** Add the given hardware device identifier to this platform's settings. */
+	void AddHardwareDeviceIdentifier(const FHardwareDeviceIdentifier& InHardwareDevice);
+
+	/** Returns an array of all Hardware Device Identifiers known to this platform */
+	const TArray<FHardwareDeviceIdentifier>& GetHardwareDevices() const;
+	
+protected:
+
+	/** A list of identifiable hardware devices available on this platform */
+	UPROPERTY(config, EditAnywhere, Category = "Hardware")
+	TArray<FHardwareDeviceIdentifier> HardwareDevices;
+};
diff --git a/Engine/Source/Runtime/Engine/Private/UserInterface/InputSettings.cpp b/Engine/Source/Runtime/Engine/Private/UserInterface/InputSettings.cpp
--- a/Engine/Source/Runtime/Engine/Private/UserInterface/InputSettings.cpp
+++ b/Engine/Source/Runtime/Engine/Private/UserInterface/InputSettings.cpp
@@ -33,6 +33,7 @@ UInputSettings::UInputSettings(const FObjectInitializer& ObjectInitializer)
 	, DefaultPlayerInputClass(UPlayerInput::StaticClass())
 	, DefaultInputComponentClass(UInputComponent::StaticClass())
 {
+	PlatformSettings.Initialize(UInputPlatformSettings::StaticClass());
 }
 
 void UInputSettings::RemoveInvalidKeys()
@@ -491,4 +492,51 @@ void UInputSettings::SetDefaultInputComponentClass(TSubclassOf<UInputComponent>
 	}
 }
 
+FHardwareDeviceIdentifier FHardwareDeviceIdentifier::Invalid = { NAME_None, TEXT("Invalid") };
 
+bool FHardwareDeviceIdentifier::IsValid() const
+{
+	return InputClassName.IsValid() && !HardwareDeviceIdentifier.IsEmpty();
+}
+
+UInputPlatformSettings* UInputPlatformSettings::Get()
+{
+	return UPlatformSettingsManager::Get().GetSettingsForPlatform<UInputPlatformSettings>();
+}
+
+void UInputPlatformSettings::AddHardwareDeviceIdentifier(const FHardwareDeviceIdentifier& InHardwareDevice)
+{
+	if (ensure(InHardwareDevice.IsValid()))
+	{
+		HardwareDevices.AddUnique(InHardwareDevice);
+	}
+}
+
+const TArray<FHardwareDeviceIdentifier>& UInputPlatformSettings::GetHardwareDevices() const
+{
+	return HardwareDevices;
+}
+
+#if WITH_EDITOR
+const TArray<FString>& UInputPlatformSettings::GetAllHardwareDeviceNames()
+{
+	static TArray<FString> HardwareDevices;
+	HardwareDevices.Reset();
+
+	// Get every known platform's InputPlatformSettings and compile a list of them
+	TArray<UPlatformSettings*> AllInputSettings = UPlatformSettingsManager::Get().GetAllPlatformSettings<UInputPlatformSettings>();
+
+	for (const UPlatformSettings* Setting : AllInputSettings)
+	{
+		if (const UInputPlatformSettings* InputSetting = Cast<UInputPlatformSettings>(Setting))
+		{
+			for (const FHardwareDeviceIdentifier& Device : InputSetting->HardwareDevices)
+			{
+				HardwareDevices.AddUnique(Device.HardwareDeviceIdentifier);
+			}
+		}
+	}
+	
+	return HardwareDevices;
+}
+#endif	// WITH_EDITOR
\ No newline at end of file
