Report the max bone influences of each skeletal mesh to the asset registry

diff --git a/Engine/Source/Runtime/Engine/Private/SkeletalMesh.cpp b/Engine/Source/Runtime/Engine/Private/SkeletalMesh.cpp
--- a/Engine/Source/Runtime/Engine/Private/SkeletalMesh.cpp
+++ b/Engine/Source/Runtime/Engine/Private/SkeletalMesh.cpp
@@ -66,6 +66,7 @@
 #include "Animation/SkinWeightProfile.h"
 #include "Streaming/SkeletalMeshUpdate.h"
 #include "HAL/FileManager.h"
+#include "Algo/MaxElement.h"
 
 #if WITH_EDITOR
 #include "Async/ParallelFor.h"
@@ -3511,7 +3512,27 @@ void USkeletalMesh::GetAssetRegistryTags(TArray<FAssetRegistryTag>& OutTags) con
 		GetAssetImportData()->AppendAssetRegistryTags(OutTags);
 #endif
 	}
-#endif
+
+	FString MaxBoneInfluencesString;
+	if (GetImportedModel())
+	{
+		// Find the LOD with the highest maximum bone influences
+		//
+		// This will be nullptr if LODModels is empty
+		const FSkeletalMeshLODModel* MaxBoneInfluencesLODModel = Algo::MaxElementBy(GetImportedModel()->LODModels, &FSkeletalMeshLODModel::GetMaxBoneInfluences);
+
+		if (MaxBoneInfluencesLODModel)
+		{
+			// Note that this value is clamped to FGPUBaseSkinVertexFactory::GetMaxGPUSkinBones, so it's affected
+			// by project settings such as r.GPUSkin.UnlimitedBoneInfluences.
+			MaxBoneInfluencesString = FString::FromInt(MaxBoneInfluencesLODModel->GetMaxBoneInfluences());
+		}
+	}
+
+	// The tag must be added unconditionally, because some code calls this function on the CDO to find out what
+	// tags are available.
+	OutTags.Add(FAssetRegistryTag("MaxBoneInfluences", MaxBoneInfluencesString, FAssetRegistryTag::TT_Numerical));
+#endif // WITH_EDITORONLY_DATA
 	
 	Super::GetAssetRegistryTags(OutTags);
 }
