When looking for the cell that contains a subobject in the world partition level streaming policy, use the runtime hash in the world partition as the outer instead of the world partition itself.

diff --git a/Engine/Source/Runtime/Engine/Private/WorldPartition/WorldPartitionLevelStreamingPolicy.cpp b/Engine/Source/Runtime/Engine/Private/WorldPartition/WorldPartitionLevelStreamingPolicy.cpp
--- a/Engine/Source/Runtime/Engine/Private/WorldPartition/WorldPartitionLevelStreamingPolicy.cpp
+++ b/Engine/Source/Runtime/Engine/Private/WorldPartition/WorldPartitionLevelStreamingPolicy.cpp
@@ -198,7 +198,7 @@ UObject* UWorldPartitionLevelStreamingPolicy::GetSubObject(const TCHAR* SubObjec
 	const FString SrcPath = UWorld::RemovePIEPrefix(*SubObjectContext);
 	if (FName* CellName = SubObjectsToCellRemapping.Find(FName(*SrcPath)))
 	{
-		if (UWorldPartitionRuntimeLevelStreamingCell* Cell = (UWorldPartitionRuntimeLevelStreamingCell*)StaticFindObject(UWorldPartitionRuntimeLevelStreamingCell::StaticClass(), GetOuterUWorldPartition(), *(CellName->ToString())))
+		if (UWorldPartitionRuntimeLevelStreamingCell* Cell = (UWorldPartitionRuntimeLevelStreamingCell*)StaticFindObject(UWorldPartitionRuntimeLevelStreamingCell::StaticClass(), GetOuterUWorldPartition()->RuntimeHash, *(CellName->ToString())))
 		{
 			if (UWorldPartitionLevelStreamingDynamic* LevelStreaming = Cell->GetLevelStreaming())
 			{
