Keep track of the most recent hardware Input Device for each platform user. This lets you find out what input device players are using, and get delegates when that changes.

#pf skip

diff --git a/Engine/Source/Runtime/Engine/Classes/GameFramework/InputDeviceSubsystem.h b/Engine/Source/Runtime/Engine/Classes/GameFramework/InputDeviceSubsystem.h
--- a/Engine/Source/Runtime/Engine/Classes/GameFramework/InputDeviceSubsystem.h
+++ b/Engine/Source/Runtime/Engine/Classes/GameFramework/InputDeviceSubsystem.h
@@ -2,6 +2,7 @@
 
 #pragma once
 
+#include "InputSettings.h"
 #include "Subsystems/EngineSubsystem.h"
 #include "Tickable.h"
 #include "Templates/SubclassOf.h"
@@ -60,6 +61,13 @@ struct FSetDevicePropertyParams
 	bool bResetUponCompletion = true;
 };
 
+/**
+ * Delegate called when a user changed the hardware they are using for input.
+ *
+ * @param UserId		The Platform user whose device has changed
+ */
+DECLARE_DYNAMIC_MULTICAST_DELEGATE_OneParam(FHardwareInputDeviceChanged, const FPlatformUserId, UserId);
+
 /**
 * The input device subsystem provides an interface to allow users to set Input Device Properties
 * on any Platform User. 
@@ -67,10 +75,18 @@ struct FSetDevicePropertyParams
 UCLASS(BlueprintType)
 class UInputDeviceSubsystem : public UEngineSubsystem, public FTickableGameObject
 {
+	friend class FInputDeviceSubsystemProcessor;
+	
 	GENERATED_BODY()
 
 public:
 
+	static UInputDeviceSubsystem* Get();
+
+	virtual void Initialize(FSubsystemCollectionBase& Collection) override;
+	virtual void Deinitialize() override;
+	virtual bool ShouldCreateSubsystem(UObject* Outer) const;
+	
 	//~ Begin FTickableGameObject interface	
 	virtual UWorld* GetTickableGameObjectWorld() const override;
 	virtual ETickableTickType GetTickableTickType() const override;
@@ -95,11 +111,31 @@ public:
 	UFUNCTION(BlueprintCallable, Category = "Input Devices", meta = (ReturnDisplayName = "Num Removed"))
 	int32 RemoveDeviceProperty(const FPlatformUserId UserId, TSubclassOf<UInputDeviceProperty> DevicePropertyClass);
 
+	/** Gets the most recently used hardware input device for the given platform user */
+	FHardwareDeviceIdentifier GetMostRecentlyUsedHardwareDevice(const FPlatformUserId InUserId) const;
+
+	/** A delegate that is fired when a platform user changes what Hardware Input device they are using */
+	DECLARE_EVENT_OneParam(UInputDeviceSubsystem, FHardwareInputDeviceChanged, const FPlatformUserId);
+	FHardwareInputDeviceChanged OnInputHardwareDeviceChangedNative;
+	
+	/** A delegate that is fired when a platform user changes what Hardware Input device they are using */
+	UPROPERTY(BlueprintAssignable, Category = "Input Device")
+	FHardwareInputDeviceChanged OnInputHardwareDeviceChanged;
+	
 protected:
 
+	/** Set the most recently used hardware device */
+	void SetMostRecentlyUsedHardwareDevice(const FInputDeviceId InDeviceId, const FHardwareDeviceIdentifier& InHardwareId);
+	
 	/**
 	* Array of the active device properties that are currently being evaluated on Tick.
 	*/
 	UPROPERTY(Transient)
 	TArray<FActiveDeviceProperty> ActiveProperties;
+	
+	/** A map of platform user's to their most recent hardwave device identifier */
+	TMap<FPlatformUserId, FHardwareDeviceIdentifier> LatestInputDeviceIdentifiers;
+
+	/** An input processor that is used to determine the current hardware input device */
+	TSharedPtr<class FInputDeviceSubsystemProcessor> InputPreprocessor;
 };
\ No newline at end of file
diff --git a/Engine/Source/Runtime/Engine/Private/GameFramework/InputDeviceSubsystem.cpp b/Engine/Source/Runtime/Engine/Private/GameFramework/InputDeviceSubsystem.cpp
--- a/Engine/Source/Runtime/Engine/Private/GameFramework/InputDeviceSubsystem.cpp
+++ b/Engine/Source/Runtime/Engine/Private/GameFramework/InputDeviceSubsystem.cpp
@@ -5,15 +5,124 @@
 #include "UnrealEngine.h"	// For GEngine
 #include "Engine/Engine.h"	// For FWorldContext
 #include "Application/SlateApplicationBase.h"
+#include "Framework/Application/IInputProcessor.h"
+#include "GenericPlatform/GenericPlatformApplicationMisc.h"	// For FInputDeviceScope
+#include "GenericPlatform/GenericPlatformInputDeviceMapper.h"
+#include "Framework/Application/SlateApplication.h"			// For RegisterInputPreProcessor
 
 DEFINE_LOG_CATEGORY(LogInputDeviceProperties);
 
+////////////////////////////////////////////////////////
+// FInputDeviceSubsystemProcessor
+
+/** An input processor for detecting changes to input devices based on the current FInputDeviceScope stack */
+class FInputDeviceSubsystemProcessor : public IInputProcessor
+{
+	friend class UInputDeviceSubsystem;
+	
+	void UpdateLatestDevice(const FInputDeviceId InDeviceId)
+	{
+		if (UInputDeviceSubsystem* SubSystem = UInputDeviceSubsystem::Get())
+		{
+			if (const FInputDeviceScope* Scope = FInputDeviceScope::GetCurrent())
+			{
+				SubSystem->SetMostRecentlyUsedHardwareDevice(InDeviceId, { Scope->InputDeviceName, Scope->HardwareDeviceIdentifier });
+			}	
+		}
+	}
+	
+public:
+
+	// Required by IInputProcessor
+	virtual void Tick(const float DeltaTime, FSlateApplication& SlateApp, TSharedRef<ICursor> Cursor) override { }
+	
+	virtual bool HandleKeyDownEvent(FSlateApplication& SlateApp, const FKeyEvent& InEvent) override
+	{
+		UpdateLatestDevice(InEvent.GetInputDeviceId());
+		return false;
+	}
+
+	virtual bool HandleAnalogInputEvent(FSlateApplication& SlateApp, const FAnalogInputEvent& InEvent) override
+	{
+		UpdateLatestDevice(InEvent.GetInputDeviceId());
+		return false;
+	}
+
+	virtual bool HandleMouseMoveEvent(FSlateApplication& SlateApp, const FPointerEvent& InEvent) override
+	{
+		UpdateLatestDevice(InEvent.GetInputDeviceId());
+		return false;
+	}
+
+	virtual bool HandleMouseButtonDownEvent(FSlateApplication& SlateApp, const FPointerEvent& InEvent) override
+	{
+		UpdateLatestDevice(InEvent.GetInputDeviceId());
+		return false;
+	}
+
+	virtual bool HandleMouseButtonDoubleClickEvent(FSlateApplication& SlateApp, const FPointerEvent& InEvent) override
+	{
+		UpdateLatestDevice(InEvent.GetInputDeviceId());
+		return false;
+	}
+	
+	virtual bool HandleMouseWheelOrGestureEvent(FSlateApplication& SlateApp, const FPointerEvent& InEvent, const FPointerEvent* InGestureEvent) override
+	{
+		UpdateLatestDevice(InEvent.GetInputDeviceId());
+		return false;
+	}
+};
+
+////////////////////////////////////////////////////////
+// UInputDeviceSubsystem
+
 FSetDevicePropertyParams::FSetDevicePropertyParams()
 	: UserId(FSlateApplicationBase::SlateAppPrimaryPlatformUser)
 {
 
 }
 
+UInputDeviceSubsystem* UInputDeviceSubsystem::Get()
+{
+	return GEngine ? GEngine->GetEngineSubsystem<UInputDeviceSubsystem>() : nullptr;
+}
+
+void UInputDeviceSubsystem::Initialize(FSubsystemCollectionBase& Collection)
+{
+	Super::Initialize(Collection);
+
+	// We have to have a valid slate app to run this subsystem
+	check(FSlateApplication::IsInitialized());
+	
+	InputPreprocessor = MakeShared<FInputDeviceSubsystemProcessor>();
+	FSlateApplication::Get().RegisterInputPreProcessor(InputPreprocessor, 0);
+}
+
+void UInputDeviceSubsystem::Deinitialize()
+{
+	Super::Deinitialize();
+
+	if (FSlateApplication::IsInitialized())
+	{
+		FSlateApplication::Get().UnregisterInputPreProcessor(InputPreprocessor);
+	}
+	InputPreprocessor.Reset();
+}
+
+bool UInputDeviceSubsystem::ShouldCreateSubsystem(UObject* Outer) const
+{
+	// No slate app means we can't process any input
+	if (!FSlateApplication::IsInitialized() ||
+		// Commandlets and servers have no use for this subsystem
+		IsRunningCommandlet() ||
+		IsRunningDedicatedServer())
+	{
+		return false;
+	}
+	
+	return Super::ShouldCreateSubsystem(Outer);
+}
+
 UWorld* UInputDeviceSubsystem::GetTickableGameObjectWorld() const
 {
 	// Use the default world by default...
@@ -147,4 +256,32 @@ int32 UInputDeviceSubsystem::RemoveDeviceProperty(const FPlatformUserId UserId,
 	}
 
 	return NumRemoved;
+}
+
+FHardwareDeviceIdentifier UInputDeviceSubsystem::GetMostRecentlyUsedHardwareDevice(const FPlatformUserId InUserId) const
+{
+	if (const FHardwareDeviceIdentifier* FoundDevice = LatestInputDeviceIdentifiers.Find(InUserId))
+	{
+		return *FoundDevice;
+	}
+	
+	UE_LOG(LogInputDeviceProperties, Warning, TEXT("Could not determine the hardware device last used by platform user %d"), InUserId.GetInternalId());
+	return FHardwareDeviceIdentifier::Invalid;
+}
+
+void UInputDeviceSubsystem::SetMostRecentlyUsedHardwareDevice(const FInputDeviceId InDeviceId, const FHardwareDeviceIdentifier& InHardwareId)
+{
+	FPlatformUserId OwningUserId = IPlatformInputDeviceMapper::Get().GetUserForInputDevice(InDeviceId); 
+
+	LatestInputDeviceIdentifiers.Add(OwningUserId, InHardwareId);
+
+	if (OnInputHardwareDeviceChanged.IsBound())
+	{
+		OnInputHardwareDeviceChanged.Broadcast(OwningUserId);
+	}
+
+	if (OnInputHardwareDeviceChangedNative.IsBound())
+	{
+		OnInputHardwareDeviceChangedNative.Broadcast(OwningUserId);
+	}
 }
\ No newline at end of file
