Starta - updated slab input naming

diff --git a/Engine/Source/Runtime/Engine/Private/Materials/MaterialExpressions.cpp b/Engine/Source/Runtime/Engine/Private/Materials/MaterialExpressions.cpp
--- a/Engine/Source/Runtime/Engine/Private/Materials/MaterialExpressions.cpp
+++ b/Engine/Source/Runtime/Engine/Private/Materials/MaterialExpressions.cpp
@@ -22653,11 +22653,11 @@ FName UMaterialExpressionStrataSlabBSDF::GetInputName(int32 InputIndex) const
 	{
 		if (InputIndex == 0)
 		{
-			return TEXT("BaseColor");
+			return TEXT("Base Color");
 		}
 		else if (InputIndex == 1)
 		{
-			return TEXT("EdgeColor");
+			return TEXT("Edge Color");
 		}
 		else if (InputIndex == 2)
 		{
@@ -22672,15 +22672,15 @@ FName UMaterialExpressionStrataSlabBSDF::GetInputName(int32 InputIndex) const
 	{
 		if (InputIndex == 0)
 		{
-			return TEXT("DiffuseAlbedo");
+			return TEXT("Diffuse Albedo");
 		}
 		else if (InputIndex == 1)
 		{
-			return TEXT("F0");
+			return TEXT("Specular Color");
 		}
 		else if (InputIndex == 2)
 		{
-			return  TEXT("F90");
+			return  TEXT("Edge Specular Color");
 		}
 	}
 
@@ -22730,11 +22730,11 @@ FName UMaterialExpressionStrataSlabBSDF::GetInputName(int32 InputIndex) const
 	}
 	else if (InputIndex == (15 + SkipDisabledInputOffset))
 	{
-		return TEXT("FuzzAmount");
+		return TEXT("Fuzz Amount");
 	}
 	else if (InputIndex == (16 + SkipDisabledInputOffset))
 	{
-		return TEXT("FuzzColor");
+		return TEXT("Fuzz Color");
 	}
 
 	return TEXT("Unknown");
