SkyAtmosphere - added a way to reset the overriden light direction to disabled.

diff --git a/Engine/Source/Runtime/Engine/Classes/Components/SkyAtmosphereComponent.h b/Engine/Source/Runtime/Engine/Classes/Components/SkyAtmosphereComponent.h
--- a/Engine/Source/Runtime/Engine/Classes/Components/SkyAtmosphereComponent.h
+++ b/Engine/Source/Runtime/Engine/Classes/Components/SkyAtmosphereComponent.h
@@ -164,6 +164,8 @@ class USkyAtmosphereComponent : public USceneComponent
 	ENGINE_API bool IsAtmosphereLightDirectionOverriden(int32 AtmosphereLightIndex);
 	UFUNCTION(BlueprintCallable, Category = "Rendering")
 	ENGINE_API FVector GetOverridenAtmosphereLightDirection(int32 AtmosphereLightIndex);
+	UFUNCTION(BlueprintCallable, Category = "Rendering")
+	ENGINE_API void ResetAtmosphereLightDirectionOverride(int32 AtmosphereLightIndex);
 
 	UFUNCTION(BlueprintCallable, Category = "Rendering")
 	ENGINE_API void SetAtmosphereHeight(float NewValue);
diff --git a/Engine/Source/Runtime/Engine/Private/Components/SkyAtmosphereComponent.cpp b/Engine/Source/Runtime/Engine/Private/Components/SkyAtmosphereComponent.cpp
--- a/Engine/Source/Runtime/Engine/Private/Components/SkyAtmosphereComponent.cpp
+++ b/Engine/Source/Runtime/Engine/Private/Components/SkyAtmosphereComponent.cpp
@@ -329,6 +329,16 @@ FVector USkyAtmosphereComponent::GetOverridenAtmosphereLightDirection(int32 Atmo
 	return FVector::ZeroVector;
 }
 
+void USkyAtmosphereComponent::ResetAtmosphereLightDirectionOverride(int32 AtmosphereLightIndex)
+{
+	check(AtmosphereLightIndex >= 0 && AtmosphereLightIndex < NUM_ATMOSPHERE_LIGHTS);
+	if (AtmosphereLightIndex >= 0 && AtmosphereLightIndex < NUM_ATMOSPHERE_LIGHTS)
+	{
+		OverrideAtmosphericLight[AtmosphereLightIndex] = false;
+		OverrideAtmosphericLightDirection[AtmosphereLightIndex] = FVector::ZeroVector;
+	}
+}
+
 void USkyAtmosphereComponent::GetOverrideLightStatus(bool* OutOverrideAtmosphericLight, FVector* OutOverrideAtmosphericLightDirection) const
 {
 	memcpy(OutOverrideAtmosphericLight, OverrideAtmosphericLight, sizeof(OverrideAtmosphericLight));
