Fix seamless travel to do an explicit garbage reference clear pass if pending kill is disabled
Explicitly clear the net connection's view target during seamless travel, it will reacquire it from the player controller

diff --git a/Engine/Source/Runtime/Engine/Private/NetConnection.cpp b/Engine/Source/Runtime/Engine/Private/NetConnection.cpp
--- a/Engine/Source/Runtime/Engine/Private/NetConnection.cpp
+++ b/Engine/Source/Runtime/Engine/Private/NetConnection.cpp
@@ -4740,6 +4740,9 @@ void UNetConnection::ResetGameWorldState()
 	CleanupDormantActorState();
 	ClientVisibleActorOuters.Empty();
 
+	// Clear the view target as this may be from an out of date world
+	ViewTarget = nullptr;
+
 	// Update any level visibility requests received during the transition
 	// This can occur if client loads faster than the server
 	for (const auto& Pending : PendingUpdateLevelVisibility)
diff --git a/Engine/Source/Runtime/Engine/Private/World.cpp b/Engine/Source/Runtime/Engine/Private/World.cpp
--- a/Engine/Source/Runtime/Engine/Private/World.cpp
+++ b/Engine/Source/Runtime/Engine/Private/World.cpp
@@ -7010,6 +7010,52 @@ void FSeamlessTravelHandler::CopyWorldData()
 	}
 }
 
+/** 
+ * Version of FArchiveReplaceObjectRef that will also clear references to garbage objects not in the replacement map 
+ * This does not try to recursively serialize subobjects because it was unreliable and missed ones without hard parent references
+ */
+template< class T >
+class FArchiveReplaceOrClearGarbageReferences : public FArchiveReplaceObjectRef<T>
+{
+	typedef FArchiveReplaceObjectRef<T> TSuper;
+public:
+	FArchiveReplaceOrClearGarbageReferences
+		( UObject* InSearchObject
+		, const TMap<T*, T*>& InReplacementMap
+		, EArchiveReplaceObjectFlags Flags = EArchiveReplaceObjectFlags::None)
+		: TSuper(InSearchObject, InReplacementMap, EArchiveReplaceObjectFlags::DelayStart | Flags)
+	{
+		if (!(Flags & EArchiveReplaceObjectFlags::DelayStart))
+		{
+			this->SerializeSingleSearchObject();
+		}
+	}
+
+	void SerializeSingleSearchObject()
+	{
+		TSuper::ReplacedReferences.Reset();
+
+		// Difference from parent behavior is to always run even if map is empty, and to ignore subobjects
+		TSuper::SerializedObjects.Add(TSuper::SearchObject);
+		TSuper::SerializingObject = TSuper::SearchObject;
+		TSuper::SerializeObject(TSuper::SearchObject);
+	}
+
+	FArchive& operator<<(UObject*& Obj) override
+	{
+		UObject* Resolved = Obj;
+		TSuper::operator<<(Resolved);
+
+		// if Resolved is garbage, just clear the reference:
+		if (Resolved && !IsValid(Resolved))
+		{
+			Resolved = nullptr;
+		}
+		Obj = Resolved;
+		return *this;
+	}
+};
+
 UWorld* FSeamlessTravelHandler::Tick()
 {
 	bool bWorldChanged = false;
@@ -7303,6 +7349,24 @@ UWorld* FSeamlessTravelHandler::Tick()
 
 			CurrentWorld = nullptr;
 
+			if (!UObjectBaseUtility::IsPendingKillEnabled())
+			{
+				// If pending kill is disabled, run an explicit serializer to clear references to garbage objects on the transferred actors
+				TMap<UObject*, UObject*> ReplacementMap;
+
+				for (AActor* const TheActor : ActuallyKeptActors)
+				{
+					auto ClearReferences = [ReplacementMap](UObject* Object)
+					{
+						FArchiveReplaceOrClearGarbageReferences<UObject> ReplaceAr(Object, ReplacementMap, EArchiveReplaceObjectFlags::IgnoreOuterRef);
+					};
+
+					// Process all subobjects, even unreferenced ones
+					ClearReferences(TheActor);
+					ForEachObjectWithOuter(TheActor, ClearReferences, true, RF_NoFlags, EInternalObjectFlags::Garbage);
+				}
+			}
+
 			// collect garbage to delete the old world
 			// because we marked everything in it pending kill, references will be NULL'ed so we shouldn't end up with any dangling pointers
 			CollectGarbage( GARBAGE_COLLECTION_KEEPFLAGS, true );
