Add missing getters in FSkinWeightLookupVertexBuffer to access raw lookup data similar to FSkinWeightDataVertexBuffer::GetWeightData

diff --git a/Engine/Source/Runtime/Engine/Public/Rendering/SkinWeightVertexBuffer.h b/Engine/Source/Runtime/Engine/Public/Rendering/SkinWeightVertexBuffer.h
--- a/Engine/Source/Runtime/Engine/Public/Rendering/SkinWeightVertexBuffer.h
+++ b/Engine/Source/Runtime/Engine/Public/Rendering/SkinWeightVertexBuffer.h
@@ -90,6 +90,16 @@ public:
 	/** @return true is LookupData is valid */
 	bool IsLookupDataValid() const;
 
+	const uint8* GetLookupData() const
+	{
+		return Data;
+	}
+
+	uint8* GetLookupData()
+	{
+		return Data;
+	}
+
 	void Init(uint32 InNumVertices);
 
 	friend FArchive& operator<<(FArchive& Ar, FSkinWeightLookupVertexBuffer& VertexBuffer);
@@ -529,4 +539,4 @@ private:
 
 	/** Skin weights lookup buffer */
 	FSkinWeightLookupVertexBuffer	LookupVertexBuffer;
-};
\ No newline at end of file
+};
