Fixing RecursiveGetExpressionChain so that it will traverse named reroute nodes

Fixes issues in mobile landscape where we drop weightmaps unexpectedly because their use is hidden behind reroute nodes

#preflgiht 636c599d843e6ac794544ac3

diff --git a/Engine/Source/Runtime/Engine/Private/Materials/Material.cpp b/Engine/Source/Runtime/Engine/Private/Materials/Material.cpp
--- a/Engine/Source/Runtime/Engine/Private/Materials/Material.cpp
+++ b/Engine/Source/Runtime/Engine/Private/Materials/Material.cpp
@@ -40,6 +40,7 @@
 #include "Materials/MaterialExpressionTransform.h"
 #include "Materials/MaterialExpressionExecBegin.h"
 #include "Materials/MaterialExpressionExecEnd.h"
+#include "Materials/MaterialExpressionNamedReroute.h"
 #include "Materials/MaterialFunction.h"
 #include "Materials/MaterialFunctionInstance.h"
 #include "Materials/MaterialExpressionMaterialFunctionCall.h"
@@ -5760,6 +5761,7 @@ bool UMaterial::RecursiveGetExpressionChain(
 	UMaterialExpressionMakeMaterialAttributes* MakeMaterialAttributesExp;
 	UMaterialExpressionSetMaterialAttributes* SetMaterialAttributesExp;
 	UMaterialExpressionShaderStageSwitch* ShaderStageSwitchExp;
+	UMaterialExpressionNamedRerouteUsage* RerouteUsageExp;
 
 	const bool bMobileOnly = InFeatureLevel <= ERHIFeatureLevel::ES3_1;
 
@@ -5905,6 +5907,14 @@ bool UMaterial::RecursiveGetExpressionChain(
 			}
 		}
 	}
+	else if ((RerouteUsageExp = Cast<UMaterialExpressionNamedRerouteUsage>(InExpression)) != nullptr)
+	{
+		// continue searching from the reroute declaration
+		if (RerouteUsageExp->Declaration != nullptr)
+		{
+			RecursiveGetExpressionChain(RerouteUsageExp->Declaration, InOutProcessedInputs, OutExpressions, InStaticParameterSet, InFeatureLevel, InQuality, InShadingPath, InShaderFrequency, InProperty);
+		}
+	}
 	else
 	{
 		Inputs = InExpression->GetInputs();
