Move level streaming report generation to an async task.

#tests Run replayrun on console

#localization none

diff --git a/Engine/Source/Runtime/Engine/Private/ProfilingDebugging/LevelStreamingProfilingSubsystem.cpp b/Engine/Source/Runtime/Engine/Private/ProfilingDebugging/LevelStreamingProfilingSubsystem.cpp
--- a/Engine/Source/Runtime/Engine/Private/ProfilingDebugging/LevelStreamingProfilingSubsystem.cpp
+++ b/Engine/Source/Runtime/Engine/Private/ProfilingDebugging/LevelStreamingProfilingSubsystem.cpp
@@ -117,6 +117,12 @@ TUniquePtr<ULevelStreamingProfilingSubsystem::FActiveLevel> ULevelStreamingProfi
 
 void ULevelStreamingProfilingSubsystem::StartTracking()
 {
+	if (!ReportWritingTask.IsCompleted())
+	{
+		UE_LOG(LogLevelStreamingProfiling, Warning, TEXT("StartTracking called while writing previous report - waiting for completion."));
+		ReportWritingTask.Wait();
+	}
+	
 	if (bIsTracking)
 	{
 		UE_LOG(LogLevelStreamingProfiling, Warning, TEXT("StartTracking called while already tracking"));
@@ -143,84 +149,88 @@ void ULevelStreamingProfilingSubsystem::StopTrackingAndReport()
 	bIsTracking = false;
 	UE_LOG(LogLevelStreamingProfiling, Log, TEXT("Reporting level streaming stats"));
 
-	FString OutDir = FPaths::ProfilingDir() / TEXT("LevelStreaming");
-	IFileManager::Get().MakeDirectory(*OutDir, true);
-
-	FString Filename = OutDir / FString::Printf(TEXT("LevelStreaming_(%s).tsv"), *FDateTime::Now().ToString(TEXT("%Y%m%d_%H%M%S")));
-	TUniquePtr<FArchive> Ar(IFileManager::Get().CreateFileWriter(*Filename));
-	if (Ar.IsValid())
-	{
-		TUtf8StringBuilder<1024> Builder;
-		FUtf8StringView EOL = UTF8TEXTVIEW("\n");
-		auto Write = [&Ar, &Builder, EOL]() {
-			Ar->Serialize(Builder.GetData(), Builder.Len());
-			Ar->Serialize((void*)EOL.GetData(), EOL.Len());
-			Builder.Reset();
-		};
-		auto AddDouble = [&Builder](TOptional<double> Field, double UnitsMultiplier) {
-			Builder << '\t';
-			if (Field)
-			{
-				Builder.Appendf("%.1f", Field.GetValue() * UnitsMultiplier);
-			}
-		};
+	ReportWritingTask = UE::Tasks::Launch(UE_SOURCE_LOCATION,
+	[this]() 
+	{
+		FString OutDir = FPaths::ProfilingDir() / TEXT("LevelStreaming");
+		IFileManager::Get().MakeDirectory(*OutDir, true);
 
-		auto AddVector = [&Builder](TOptional<FVector> Field) {
-			Builder << '\t';
-			if (Field)
+		FString Filename = OutDir / FString::Printf(TEXT("LevelStreaming_(%s).tsv"), *FDateTime::Now().ToString(TEXT("%Y%m%d_%H%M%S")));
+		TUniquePtr<FArchive> Ar(IFileManager::Get().CreateFileWriter(*Filename));
+		if (Ar.IsValid())
+		{
+			TUtf8StringBuilder<1024> Builder;
+			FUtf8StringView EOL = UTF8TEXTVIEW("\n");
+			auto Write = [&Ar, &Builder, EOL]() {
+				Ar->Serialize(Builder.GetData(), Builder.Len());
+				Ar->Serialize((void*)EOL.GetData(), EOL.Len());
+				Builder.Reset();
+			};
+			auto AddDouble = [&Builder](TOptional<double> Field, double UnitsMultiplier) {
+				Builder << '\t';
+				if (Field)
+				{
+					Builder.Appendf("%.1f", Field.GetValue() * UnitsMultiplier);
+				}
+			};
+
+			auto AddVector = [&Builder](TOptional<FVector> Field) {
+				Builder << '\t';
+				if (Field)
+				{
+					Builder.Appendf("%.0f,%.0f,%.0f", Field->X, Field->Y, Field->Z);
+				}
+			};
+
+			auto AddBounds = [&Builder](const FBox& Field) {
+				Builder << '\t';
+				if (Field.IsValid)
+				{
+					Builder.Appendf("%.0f,%.0f,%.0f,%.0f,%.0f,%.0f", Field.Min.X, Field.Min.Y, Field.Min.Z, Field.Max.X, Field.Max.Y, Field.Max.Z);
+				}
+			};
+
+			Builder << "NameOnDisk\tNameInMemory\tTimeQueuedForLoadingMS\tTimeLoadingMS\tTimeQueuedForAddToWorldMS\tTimeAddingToWorldMS\tTimeInWorldS\tTimeQueuedForRemoveFromWorldMS\tTimeRemovingFromWorldMS\tStreamInDistance_Cell\tStreamInDistance_Content\tStreamInSourceLocation\tCellBounds\tContentBounds";
+			AugmentReportHeader(Builder);
+			Write();
+			for (int32 i=0; i < LevelStats.Num(); ++i)
 			{
-				Builder.Appendf("%.0f,%.0f,%.0f", Field->X, Field->Y, Field->Z);
+				const FLevelStats& Stats = LevelStats[i];
+				if (!Stats.bValid) { continue; }
+				Builder << Stats.PackageNameOnDisk;
+				Builder << '\t' << Stats.PackageNameInMemory;
+				AddDouble(Stats.TimeQueuedForLoading, 1000.0);
+				AddDouble(Stats.TimeLoading, 1000.0);
+				AddDouble(Stats.TimeQueueudForAddToWorld, 1000.0);
+				AddDouble(Stats.TimeAddingToWorld, 1000.0);
+				AddDouble(Stats.TimeInWorld, 1.0);
+				AddDouble(Stats.TimeQueuedForRemoveFromWorld, 1000.0);
+				AddDouble(Stats.TimeRemovingFromWorld, 1000.0);
+				AddDouble(Stats.FinalStreamInDistance_Cell, 1.0);
+				AddDouble(Stats.FinalStreamInDistance_Content, 1.0);
+				AddVector(Stats.FinalStreamInLocation);
+				AddBounds(Stats.CellBounds);
+				AddBounds(Stats.ContentBounds);
+				AugmentReportRow(Builder, i);
+				Write();
 			}
-		};
 
-		auto AddBounds = [&Builder](const FBox& Field) {
-			Builder << '\t';
-			if (Field.IsValid)
+	#if CSV_PROFILER
+			// Write a metadata row like CSV files 
+			TMap<FString, FString> Metadata = FCsvProfiler::Get()->GetMetadataMapCopy();
+			if (Metadata.Num() > 0)
 			{
-				Builder.Appendf("%.0f,%.0f,%.0f,%.0f,%.0f,%.0f", Field.Min.X, Field.Min.Y, Field.Min.Z, Field.Max.X, Field.Max.Y, Field.Max.Z);
+				for (const TPair<FString, FString>& Pair : Metadata)
+				{ 		
+					Builder << '[' << Pair.Key << "]\t" << Pair.Value << '\t';
+				}
+				Write();
 			}
-		};
+	#endif
 
-		Builder << "NameOnDisk\tNameInMemory\tTimeQueuedForLoadingMS\tTimeLoadingMS\tTimeQueuedForAddToWorldMS\tTimeAddingToWorldMS\tTimeInWorldS\tTimeQueuedForRemoveFromWorldMS\tTimeRemovingFromWorldMS\tStreamInDistance_Cell\tStreamInDistance_Content\tStreamInSourceLocation\tCellBounds\tContentBounds";
-		AugmentReportHeader(Builder);
-		Write();
-		for (int32 i=0; i < LevelStats.Num(); ++i)
-		{
-			const FLevelStats& Stats = LevelStats[i];
-			if (!Stats.bValid) { continue; }
-			Builder << Stats.PackageNameOnDisk;
-			Builder << '\t' << Stats.PackageNameInMemory;
-			AddDouble(Stats.TimeQueuedForLoading, 1000.0);
-			AddDouble(Stats.TimeLoading, 1000.0);
-			AddDouble(Stats.TimeQueueudForAddToWorld, 1000.0);
-			AddDouble(Stats.TimeAddingToWorld, 1000.0);
-			AddDouble(Stats.TimeInWorld, 1.0);
-			AddDouble(Stats.TimeQueuedForRemoveFromWorld, 1000.0);
-			AddDouble(Stats.TimeRemovingFromWorld, 1000.0);
-			AddDouble(Stats.FinalStreamInDistance_Cell, 1.0);
-			AddDouble(Stats.FinalStreamInDistance_Content, 1.0);
-			AddVector(Stats.FinalStreamInLocation);
-			AddBounds(Stats.CellBounds);
-			AddBounds(Stats.ContentBounds);
-			AugmentReportRow(Builder, i);
-			Write();
-		}
-
-#if CSV_PROFILER
-		// Write a metadata row like CSV files 
-		TMap<FString, FString> Metadata = FCsvProfiler::Get()->GetMetadataMapCopy();
-		if (Metadata.Num() > 0)
-		{
-			for (const TPair<FString, FString>& Pair : Metadata)
-			{ 		
-				Builder << '[' << Pair.Key << "]\t" << Pair.Value << '\t';
-			}
-			Write();
+			PostReport();
 		}
-#endif
-
-		PostReport();
-	}
+	}, UE::Tasks::ETaskPriority::BackgroundLow);
 }
 
 void ULevelStreamingProfilingSubsystem::OnLevelStreamingTargetStateChanged(UWorld* World, const ULevelStreaming* StreamingLevel, ULevel* LevelIfLoaded, ELevelStreamingState CurrentState, ELevelStreamingTargetState PrevTarget, ELevelStreamingTargetState NewTarget)
diff --git a/Engine/Source/Runtime/Engine/Public/ProfilingDebugging/LevelStreamingProfilingSubsystem.h b/Engine/Source/Runtime/Engine/Public/ProfilingDebugging/LevelStreamingProfilingSubsystem.h
--- a/Engine/Source/Runtime/Engine/Public/ProfilingDebugging/LevelStreamingProfilingSubsystem.h
+++ b/Engine/Source/Runtime/Engine/Public/ProfilingDebugging/LevelStreamingProfilingSubsystem.h
@@ -6,6 +6,7 @@
 #include "Containers/ArrayView.h"
 #include "Containers/Map.h"
 #include "Subsystems/WorldSubsystem.h"
+#include "Tasks/Task.h"
 #include "UObject/ObjectKey.h"
 
 #include "LevelStreamingProfilingSubsystem.generated.h"
@@ -186,6 +187,9 @@ private:
 	// Stats for levels which have been at least partially loaded. 
 	// May contain duplicates for a given level if it streamed out and back in again.
 	TArray<FLevelStats> LevelStats;
+
+	// Possibly-executing task dependend on our recorded stats. Cannot start recording again until this task is complete.
+	UE::Tasks::TTask<void> ReportWritingTask;
 #endif // !UE_BUILD_SHIPPING 
 };
 
