Correctness tests for ISPC functions

Submitting on behalf of [REDACTED_EMAIL].

#rnx

diff --git a/Engine/Plugins/Experimental/ChaosCloth/Source/Chaos/Private/ChaosCloth/IspcTestChaosClothingSimulationSolver.cpp b/Engine/Plugins/Experimental/ChaosCloth/Source/Chaos/Private/ChaosCloth/IspcTestChaosClothingSimulationSolver.cpp
new file mode 100644
--- /dev/null
+++ b/Engine/Plugins/Experimental/ChaosCloth/Source/Chaos/Private/ChaosCloth/IspcTestChaosClothingSimulationSolver.cpp
@@ -0,0 +1,195 @@
+// Copyright Epic Games, Inc. All Rights Reserved.
+
+#if WITH_DEV_AUTOMATION_TESTS && INTEL_ISPC
+
+#include "CoreMinimal.h"
+#include "Misc/AutomationTest.h"
+#include "ChaosCloth/ChaosClothingSimulationSolver.h"
+
+using namespace Chaos;
+
+IMPLEMENT_SIMPLE_AUTOMATION_TEST(FIspcTestChaosClothingSimulationSolverAABB, "Ispc.Physics.ChaosClothingSimulationSolver.AABB", EAutomationTestFlags::EditorContext | EAutomationTestFlags::EngineFilter)
+bool FIspcTestChaosClothingSimulationSolverAABB::RunTest(const FString& Parameters)
+{
+	bool InitialState = bChaos_CalculateBounds_ISPC_Enabled;
+
+	const int32 NumParticles = 100;
+
+	bChaos_CalculateBounds_ISPC_Enabled = true;
+	FBoxSphereBounds ISPCAABB;
+	{
+		FClothingSimulationSolver ClothingSimulationSolver;
+		ClothingSimulationSolver.AddParticles(NumParticles, 0);
+		ClothingSimulationSolver.EnableParticles(0, true);
+
+		for (int32 i = 0; i < NumParticles; ++i)
+		{
+			double Value = (i % 2) ? (double)i * 100. : (double)(-i) * 100.;
+
+			Softs::FSolverVec3* Solver = ClothingSimulationSolver.GetParticleXs(i);
+			if (Solver != nullptr)
+			{
+				*Solver = FSolverVec3(Value);
+			}
+		}
+
+		ISPCAABB = ClothingSimulationSolver.CalculateBounds();
+	}
+
+	bChaos_CalculateBounds_ISPC_Enabled = false;
+	FBoxSphereBounds CPPAABB;
+	{
+		FClothingSimulationSolver ClothingSimulationSolver;
+		ClothingSimulationSolver.AddParticles(NumParticles, 0);
+		ClothingSimulationSolver.EnableParticles(0, true);
+
+		for (int32 i = 0; i < NumParticles; ++i)
+		{
+			double Value = (i % 2) ? (double)i * 100. : (double)(-i) * 100.;
+
+			Softs::FSolverVec3* Solver = ClothingSimulationSolver.GetParticleXs(i);
+			if (Solver != nullptr)
+			{
+				*Solver = FSolverVec3(Value);
+			}
+		}
+
+		CPPAABB = ClothingSimulationSolver.CalculateBounds();
+	}
+
+	bChaos_CalculateBounds_ISPC_Enabled = InitialState;
+
+	TestTrue(TEXT("Origin"), ISPCAABB.Origin.Equals(CPPAABB.Origin));
+	TestTrue(TEXT("Extent"), ISPCAABB.BoxExtent.Equals(CPPAABB.BoxExtent));
+	TestEqual(TEXT("Radius"), ISPCAABB.SphereRadius, CPPAABB.SphereRadius);
+
+	return true;
+}
+
+IMPLEMENT_SIMPLE_AUTOMATION_TEST(FIspcTestChaosClothingSimulationSolverPreSimulation, "Ispc.Physics.ChaosClothingSimulationSolver.PreSimulation", EAutomationTestFlags::EditorContext | EAutomationTestFlags::EngineFilter)
+bool FIspcTestChaosClothingSimulationSolverPreSimulation::RunTest(const FString& Parameters)
+{
+	bool InitialState = bChaos_PreSimulationTransforms_ISPC_Enabled;
+
+	const int32 NumParticles = 100;
+
+	const Softs::FSolverReal InvMValue = .5f;
+	Softs::FPAndInvM FPInvMData;
+	FPInvMData.P = Softs::FSolverVec3(0.f);
+	FPInvMData.InvM = InvMValue;
+
+	const Softs::FSolverVec3 XSolverVec(5.f);
+	const Softs::FSolverVec3 VSolverVec(2.f);
+	const Softs::FSolverVec3 CurrentAnimationVec(1.f);
+	const Softs::FSolverVec3 OldAnimationVec(5.f);
+
+	const FQuat4d Rotation(1., 1., 1., 1.);
+	// Debug complains that the rotation isn't normalized.
+	const FQuat4d NormalizedRotation(Rotation.GetNormalized());
+	const FRigidTransform3 OldTransform(Softs::FSolverVec3(1.), NormalizedRotation);
+	const FRigidTransform3 NewTransform(Softs::FSolverVec3(10.), NormalizedRotation);
+	const TVec3<FRealSingle> LinearVelocity(10.f);
+	const FRealSingle AngularVelocityScale = 1.f;
+	const FRealSingle FicticiousAngularScale = 1.f;
+
+	bChaos_PreSimulationTransforms_ISPC_Enabled = true;
+	FClothingSimulationSolver ISPCSimulation;
+	{
+		ISPCSimulation.AddParticles(NumParticles, 0);
+		ISPCSimulation.EnableParticles(0, true);
+
+		for (int32 i = 0; i < NumParticles; ++i)
+		{
+			Softs::FSolverVec3* XSolver = ISPCSimulation.GetParticleXs(i);
+			Softs::FSolverVec3* VSolver = ISPCSimulation.GetParticleVs(i);
+			Softs::FSolverVec3* CurrentAnimation = ISPCSimulation.GetAnimationPositions(i);
+			Softs::FSolverVec3* OldAnimation = ISPCSimulation.GetOldAnimationPositions(i);
+			Softs::FPAndInvM* InvM = ISPCSimulation.GetParticlePandInvMs(i);
+			if ((XSolver != nullptr) &&
+				(VSolver != nullptr) &&
+				(CurrentAnimation != nullptr) &&
+				(OldAnimation != nullptr) &&
+				(InvM != nullptr))
+			{
+				*XSolver = XSolverVec;
+				*VSolver = VSolverVec;
+				*CurrentAnimation = CurrentAnimationVec;
+				*OldAnimation = OldAnimationVec;
+				InvM->P = FPInvMData.P;
+				InvM->InvM = FPInvMData.InvM;
+			}
+		}
+
+		ISPCSimulation.ResetStartPose(0, NumParticles);
+		ISPCSimulation.SetReferenceVelocityScale(0, OldTransform, NewTransform, LinearVelocity, AngularVelocityScale, FicticiousAngularScale);
+
+		ISPCSimulation.Update(.5); // Calls ApplyPreSimulationTransform()
+	}
+
+	bChaos_PreSimulationTransforms_ISPC_Enabled = false;
+	FClothingSimulationSolver CPPSimulation;
+	{
+		CPPSimulation.AddParticles(NumParticles, 0);
+		CPPSimulation.EnableParticles(0, true);
+
+		for (int32 i = 0; i < NumParticles; ++i)
+		{
+			Softs::FSolverVec3* XSolver = CPPSimulation.GetParticleXs(i);
+			Softs::FSolverVec3* VSolver = CPPSimulation.GetParticleVs(i);
+			Softs::FSolverVec3* CurrentAnimation = CPPSimulation.GetAnimationPositions(i);
+			Softs::FSolverVec3* OldAnimation = CPPSimulation.GetOldAnimationPositions(i);
+			Softs::FPAndInvM* InvM = CPPSimulation.GetParticlePandInvMs(i);
+			if ((XSolver != nullptr) && 
+				(VSolver != nullptr) && 
+				(CurrentAnimation != nullptr) &&
+				(OldAnimation != nullptr) &&
+				(InvM != nullptr))
+			{
+				*XSolver = XSolverVec;
+				*VSolver = VSolverVec;
+				*CurrentAnimation = CurrentAnimationVec;
+				*OldAnimation = OldAnimationVec;
+				InvM->P = FPInvMData.P;
+				InvM->InvM = FPInvMData.InvM;
+			}
+		}
+
+		CPPSimulation.ResetStartPose(0, NumParticles);
+		CPPSimulation.SetReferenceVelocityScale(0, OldTransform, NewTransform, LinearVelocity, AngularVelocityScale, FicticiousAngularScale);
+
+		CPPSimulation.Update(.5); // Calls ApplyPreSimulationTransform()
+	}
+
+	bChaos_CalculateBounds_ISPC_Enabled = InitialState;
+
+	for (int32 i = 0; i < NumParticles; ++i)
+	{
+		{
+			Softs::FSolverVec3& ISPCXSolver = *ISPCSimulation.GetParticleXs(i);
+			Softs::FSolverVec3& CPPXSolver = *CPPSimulation.GetParticleXs(i);
+			const FString Message(FString::Format(TEXT("X Solver {0}: ISPC {1}, CPP {2}"), { i, ISPCXSolver.ToString(), CPPXSolver.ToString() }));
+			TestTrue(*Message, ISPCXSolver.Equals(CPPXSolver));
+		}
+
+		{
+			Softs::FSolverVec3& ISPCVSolver = *ISPCSimulation.GetParticleVs(i);
+			Softs::FSolverVec3& CPPVSolver = *CPPSimulation.GetParticleVs(i);
+			const FString Message(FString::Format(TEXT("V Solver {0}: ISPC {1}, CPP {2}"), { i, ISPCVSolver.ToString(), CPPVSolver.ToString() }));
+			TestTrue(*Message, ISPCVSolver.Equals(CPPVSolver));
+		}
+
+		{
+			Softs::FPAndInvM& ISPCInvM = *ISPCSimulation.GetParticlePandInvMs(i);
+			Softs::FPAndInvM& CPPInvM = *CPPSimulation.GetParticlePandInvMs(i);
+
+			const FString PMessage(FString::Format(TEXT("P Vec {0}: ISPC {1}, CPP {2}"), { i, ISPCInvM.P.ToString(), CPPInvM.P.ToString() }));
+			const FString InvMMessage(FString::Format(TEXT("InvM {0}: ISPC {1}, CPP {2}"), { i, ISPCInvM.InvM, CPPInvM.InvM }));
+			TestTrue(*PMessage, ISPCInvM.P.Equals(CPPInvM.P));
+			TestEqual(*InvMMessage, ISPCInvM.InvM, CPPInvM.InvM);
+		}
+	}
+
+	return true;
+}
+
+#endif // WITH_DEV_AUTOMATION_TESTS
diff --git a/Engine/Plugins/FX/Niagara/Source/Niagara/Private/IspcTestNiagaraDataInterfaceVectorField.cpp b/Engine/Plugins/FX/Niagara/Source/Niagara/Private/IspcTestNiagaraDataInterfaceVectorField.cpp
new file mode 100644
--- /dev/null
+++ b/Engine/Plugins/FX/Niagara/Source/Niagara/Private/IspcTestNiagaraDataInterfaceVectorField.cpp
@@ -0,0 +1,325 @@
+// Copyright Epic Games, Inc. All Rights Reserved.
+
+#if WITH_DEV_AUTOMATION_TESTS && INTEL_ISPC
+
+#include "CoreMinimal.h"
+#include "Misc/AutomationTest.h"
+#include "NiagaraDataInterfaceVectorField.h"
+#include "VectorField/VectorField.h"
+#include "VectorVM.h"
+#include "VectorVMExperimental.h"
+#include "VectorVMLegacy.h"
+#include "VectorField/VectorFieldStatic.h"
+
+extern bool GNiagaraVectorFieldUseIspc;
+
+#if VECTORVM_SUPPORTS_EXPERIMENTAL && VECTORVM_SUPPORTS_LEGACY
+IMPLEMENT_SIMPLE_AUTOMATION_TEST(FIspcTestNiagaraDataInterfaceVectorField, "Ispc.Graphics.NiagaraDataInterfaceVectorField", EAutomationTestFlags::EditorContext | EAutomationTestFlags::EngineFilter)
+bool FIspcTestNiagaraDataInterfaceVectorField::RunTest(const FString& Parameters)
+{
+	bool InitialState = GNiagaraVectorFieldUseIspc;
+
+	FFloat16Color RepeatColor;
+	RepeatColor.R = FFloat16(5.f);
+	RepeatColor.G = FFloat16(6.f);
+	RepeatColor.B = FFloat16(7.f);
+	RepeatColor.A = FFloat16(0.f);
+
+	float InX = 10.f;
+	float InY = 10.f;
+	float InZ = 10.f;
+
+	GNiagaraVectorFieldUseIspc = true;
+	float ISPCOutX = 0.f;
+	float ISPCOutY = 0.f;
+	float ISPCOutZ = 0.f;
+	{
+		TObjectPtr<UNiagaraDataInterfaceVectorField> VectorField = NewObject<UNiagaraDataInterfaceVectorField>();
+		TObjectPtr<UVectorFieldStatic> StaticField = NewObject<UVectorFieldStatic>();
+		if (!IsValid(VectorField) || !IsValid(StaticField))
+		{
+			return false;
+		}
+
+		StaticField->SizeX = 10;
+		StaticField->SizeY = 10;
+		StaticField->SizeZ = 10;
+		StaticField->Intensity = 10.f;
+		StaticField->Bounds = FBox(FVector(10.f), FVector(50.f));
+
+		StaticField->InitResource();
+		StaticField->bAllowCPUAccess = true;
+
+		// Fill the static vector field with values.
+		const uint32 VectorCount = (StaticField->SizeX * StaticField->SizeY * StaticField->SizeZ);
+		const uint32 DestBufferSize = VectorCount * sizeof(FFloat16Color);
+		StaticField->SourceData.Lock(LOCK_READ_WRITE);
+		FFloat16Color* RESTRICT DestValues = (FFloat16Color*)StaticField->SourceData.Realloc((int64)DestBufferSize);
+		for (uint32 VectorIndex = 0; VectorIndex < VectorCount; ++VectorIndex)
+		{
+			*DestValues = RepeatColor;
+			DestValues++;
+		}
+		StaticField->SourceData.Unlock();
+
+		StaticField->UpdateCPUData(true);
+		VectorField->Field = StaticField;
+
+		// HACK: Pretend we are executing with the VM.
+		// This requires 6 floats to be pushed onto its registers.
+		// In X, Y, Z
+		// Out X, Y, Z
+		// Assumes one instance
+		FVectorVMExternalFunctionContextExperimental Experimental;
+		FMemory::Memset(&Experimental, 0, sizeof(FVectorVMExternalFunctionContextExperimental));
+		const int32 RegisterCount = 6;
+		Experimental.RegReadCount = 0;
+		Experimental.NumInstances = 1;
+		Experimental.NumRegisters = RegisterCount;
+		Experimental.PerInstanceFnInstanceIdx = 0;
+		float* Values[RegisterCount];
+		Values[0] = &InX;
+		Values[1] = &InY;
+		Values[2] = &InZ;
+		Values[3] = &ISPCOutX;
+		Values[4] = &ISPCOutY;
+		Values[5] = &ISPCOutZ;
+		float** ValuePtr = Values;
+		Experimental.RegisterData = (uint32**)ValuePtr;
+		uint8 IncJunk[RegisterCount] = { 0, 1, 2, 3, 4, 5 };
+		uint16 VecJunk[RegisterCount] = { 0, 1, 2, 3, 4, 5 };
+		Experimental.RegInc = IncJunk;
+		FVectorVMExternalFunctionContextProxy Proxy(Experimental);
+		VectorField->SampleVectorField(Proxy);
+	}
+
+	GNiagaraVectorFieldUseIspc = false;
+	float CPPOutX = 0.f;
+	float CPPOutY = 0.f;
+	float CPPOutZ = 0.f;
+	{
+		TObjectPtr<UNiagaraDataInterfaceVectorField> VectorField = NewObject<UNiagaraDataInterfaceVectorField>();
+		TObjectPtr<UVectorFieldStatic> StaticField = NewObject<UVectorFieldStatic>();
+		if (!IsValid(VectorField) || !IsValid(StaticField))
+		{
+			return false;
+		}
+
+		StaticField->SizeX = 10;
+		StaticField->SizeY = 10;
+		StaticField->SizeZ = 10;
+		StaticField->Intensity = 10.f;
+		StaticField->Bounds = FBox(FVector(10.f), FVector(50.f));
+
+		StaticField->InitResource();
+		StaticField->bAllowCPUAccess = true;
+
+		// Fill the static vector field with values.
+		const uint32 VectorCount = (StaticField->SizeX * StaticField->SizeY * StaticField->SizeZ);
+		const uint32 DestBufferSize = VectorCount * sizeof(FFloat16Color);
+		StaticField->SourceData.Lock(LOCK_READ_WRITE);
+		FFloat16Color* RESTRICT DestValues = (FFloat16Color*)StaticField->SourceData.Realloc((int64)DestBufferSize);
+		for (uint32 VectorIndex = 0; VectorIndex < VectorCount; ++VectorIndex)
+		{
+			*DestValues = RepeatColor;
+			DestValues++;
+		}
+		StaticField->SourceData.Unlock();
+
+		StaticField->UpdateCPUData(true);
+		VectorField->Field = StaticField;
+
+		// HACK: Pretend we are executing with the VM.
+		// This requires 6 floats to be pushed onto its registers.
+		// In X, Y, Z
+		// Out X, Y, Z
+		// Assumes one instance
+		FVectorVMExternalFunctionContextExperimental Experimental;
+		FMemory::Memset(&Experimental, 0, sizeof(FVectorVMExternalFunctionContextExperimental));
+		const int32 RegisterCount = 6;
+		Experimental.RegReadCount = 0;
+		Experimental.NumInstances = 1;
+		Experimental.NumRegisters = RegisterCount;
+		Experimental.PerInstanceFnInstanceIdx = 0;
+		float* Values[RegisterCount];
+		Values[0] = &InX;
+		Values[1] = &InY;
+		Values[2] = &InZ;
+		Values[3] = &CPPOutX;
+		Values[4] = &CPPOutY;
+		Values[5] = &CPPOutZ;
+		float** ValuePtr = Values;
+		Experimental.RegisterData = (uint32**)ValuePtr;
+		uint8 IncJunk[RegisterCount] = { 0, 1, 2, 3, 4, 5 };
+		uint16 VecJunk[RegisterCount] = { 0, 1, 2, 3, 4, 5 };
+		Experimental.RegInc = IncJunk;
+		FVectorVMExternalFunctionContextProxy Proxy(Experimental);
+		VectorField->SampleVectorField(Proxy);
+	}
+
+	GNiagaraVectorFieldUseIspc = InitialState;
+
+	TestEqual(TEXT("X"), ISPCOutX, CPPOutX);
+	TestEqual(TEXT("Y"), ISPCOutY, CPPOutY);
+	TestEqual(TEXT("Z"), ISPCOutZ, CPPOutZ);
+
+	return true;
+}
+#endif
+#if 0
+IMPLEMENT_SIMPLE_AUTOMATION_TEST(FIspcTestNiagaraDataInterfaceVectorField, "Ispc.Graphics.NiagaraDataInterfaceVectorField", EAutomationTestFlags::EditorContext | EAutomationTestFlags::EngineFilter)
+bool FIspcTestNiagaraDataInterfaceVectorField::RunTest(const FString& Parameters)
+{
+	bool InitialState = GNiagaraVectorFieldUseIspc;
+
+	FFloat16Color RepeatColor;
+	RepeatColor.R = FFloat16(5.f);
+	RepeatColor.G = FFloat16(6.f);
+	RepeatColor.B = FFloat16(7.f);
+	RepeatColor.A = FFloat16(0.f);
+
+	float InX = 10.f;
+	float InY = 10.f;
+	float InZ = 10.f;
+
+	GNiagaraVectorFieldUseIspc = true;
+	float ISPCOutX = 0.f;
+	float ISPCOutY = 0.f;
+	float ISPCOutZ = 0.f;
+	{
+		TObjectPtr<UNiagaraDataInterfaceVectorField> VectorField = NewObject<UNiagaraDataInterfaceVectorField>();
+		TObjectPtr<UVectorFieldStatic> StaticField = NewObject<UVectorFieldStatic>();
+		if (!IsValid(VectorField) || !IsValid(StaticField))
+		{
+			return false;
+		}
+
+		StaticField->SizeX = 10;
+		StaticField->SizeY = 10;
+		StaticField->SizeZ = 10;
+		StaticField->Intensity = 10.f;
+		StaticField->Bounds = FBox(FVector(10.f), FVector(50.f));
+
+		StaticField->InitResource();
+		StaticField->bAllowCPUAccess = true;
+
+		// Fill the static vector field with values.
+		const uint32 VectorCount = (StaticField->SizeX * StaticField->SizeY * StaticField->SizeZ);
+		const uint32 DestBufferSize = VectorCount * sizeof(FFloat16Color);
+		StaticField->SourceData.Lock(LOCK_READ_WRITE);
+		FFloat16Color* RESTRICT DestValues = (FFloat16Color*)StaticField->SourceData.Realloc((int64)DestBufferSize);
+		for (uint32 VectorIndex = 0; VectorIndex < VectorCount; ++VectorIndex)
+		{
+			*DestValues = RepeatColor;
+			DestValues++;
+		}
+		StaticField->SourceData.Unlock();
+
+		StaticField->UpdateCPUData(true);
+		VectorField->Field = StaticField;
+
+		// HACK: Pretend we are executing with the VM.
+		// This requires 6 floats to be pushed onto its registers.
+		// In X, Y, Z
+		// Out X, Y, Z
+		// Assumes one instance
+		FVectorVMContext VectorVM;
+		const int32 RegisterCount = 6;
+		float* Values[RegisterCount];
+		Values[0] = &InX;
+		Values[1] = &InY;
+		Values[2] = &InZ;
+		Values[3] = &ISPCOutX;
+		Values[4] = &ISPCOutY;
+		Values[5] = &ISPCOutZ;
+		float** ValuePtr = Values;
+		TArray<FDataSetMeta> JunkData;
+		VectorVM.PrepareForExec(RegisterCount,
+			0,
+			0,
+			0,
+			0,
+			(void**)ValuePtr,
+			JunkData,
+			1,
+			false);
+		FVectorVMExternalFunctionContextLegacy Legacy(&VectorVM);
+		VectorField->SampleVectorField(Legacy);
+	}
+
+	GNiagaraVectorFieldUseIspc = false;
+	float CPPOutX = 0.f;
+	float CPPOutY = 0.f;
+	float CPPOutZ = 0.f;
+	{
+		TObjectPtr<UNiagaraDataInterfaceVectorField> VectorField = NewObject<UNiagaraDataInterfaceVectorField>();
+		TObjectPtr<UVectorFieldStatic> StaticField = NewObject<UVectorFieldStatic>();
+		if (!IsValid(VectorField) || !IsValid(StaticField))
+		{
+			return false;
+		}
+
+		StaticField->SizeX = 10;
+		StaticField->SizeY = 10;
+		StaticField->SizeZ = 10;
+		StaticField->Intensity = 10.f;
+		StaticField->Bounds = FBox(FVector(10.f), FVector(50.f));
+
+		StaticField->InitResource();
+		StaticField->bAllowCPUAccess = true;
+
+		// Fill the static vector field with values.
+		const uint32 VectorCount = (StaticField->SizeX * StaticField->SizeY * StaticField->SizeZ);
+		const uint32 DestBufferSize = VectorCount * sizeof(FFloat16Color);
+		StaticField->SourceData.Lock(LOCK_READ_WRITE);
+		FFloat16Color* RESTRICT DestValues = (FFloat16Color*)StaticField->SourceData.Realloc((int64)DestBufferSize);
+		for (uint32 VectorIndex = 0; VectorIndex < VectorCount; ++VectorIndex)
+		{
+			*DestValues = RepeatColor;
+			DestValues++;
+		}
+		StaticField->SourceData.Unlock();
+
+		StaticField->UpdateCPUData(true);
+		VectorField->Field = StaticField;
+
+		// HACK: Pretend we are executing with the VM.
+		// This requires 6 floats to be pushed onto its registers.
+		// In X, Y, Z
+		// Out X, Y, Z
+		// Assumes one instance
+		FVectorVMContext VectorVM;
+		const int32 RegisterCount = 6;
+		float* Values[RegisterCount];
+		Values[0] = &InX;
+		Values[1] = &InY;
+		Values[2] = &InZ;
+		Values[3] = &ISPCOutX;
+		Values[4] = &ISPCOutY;
+		Values[5] = &ISPCOutZ;
+		float** ValuePtr = Values;
+		TArray<FDataSetMeta> JunkData;
+		VectorVM.PrepareForExec(RegisterCount,
+			0,
+			0,
+			0,
+			0,
+			(void**)ValuePtr,
+			JunkData,
+			1,
+			false);
+		FVectorVMExternalFunctionContextLegacy Legacy(&VectorVM);
+		VectorField->SampleVectorField(Legacy);
+	}
+
+	GNiagaraVectorFieldUseIspc = InitialState;
+
+	TestEqual(TEXT("X"), ISPCOutX, CPPOutX);
+	TestEqual(TEXT("Y"), ISPCOutY, CPPOutY);
+	TestEqual(TEXT("Z"), ISPCOutZ, CPPOutZ);
+
+	return true;
+}
+#endif
+
+#endif // WITH_DEV_AUTOMATION_TESTS
diff --git a/Engine/Source/Runtime/Engine/Private/Animation/IspcTestAnimEncodingConstantKeyLerp.cpp b/Engine/Source/Runtime/Engine/Private/Animation/IspcTestAnimEncodingConstantKeyLerp.cpp
new file mode 100644
--- /dev/null
+++ b/Engine/Source/Runtime/Engine/Private/Animation/IspcTestAnimEncodingConstantKeyLerp.cpp
@@ -0,0 +1,286 @@
+// Copyright Epic Games, Inc. All Rights Reserved.
+
+#if WITH_DEV_AUTOMATION_TESTS && INTEL_ISPC
+
+#include "CoreMinimal.h"
+#include "Engine/Engine.h"
+#include "Misc/AutomationTest.h"
+#include "AnimEncoding_ConstantKeyLerp.h"
+
+IMPLEMENT_SIMPLE_AUTOMATION_TEST(FIspcTestAnimEncodingConstantKeyLerpGetPoseRotations, "Ispc.Animation.ConstantKeyLerp.GetPoseRotations", EAutomationTestFlags::EditorContext | EAutomationTestFlags::EngineFilter)
+bool FIspcTestAnimEncodingConstantKeyLerpGetPoseRotations::RunTest(const FString& Parameters)
+{
+	const FString CommandName(TEXT("a.ConstantKeyLerp.ISPC"));
+	auto FormatCommand = [CommandName](bool State) -> FString {
+		return FString::Format(TEXT("{0} {1}"), { CommandName, State });
+	};
+
+	const IConsoleVariable* CVarISPCEnabled = IConsoleManager::Get().FindConsoleVariable(*CommandName);
+	bool InitialState = CVarISPCEnabled->GetBool();
+	check(GEngine);
+
+	TArray<FTransform> ISPCTransforms;
+	ISPCTransforms.Add(FTransform(FQuat4d(45., 0., 0., 1.)));
+	ISPCTransforms.Add(FTransform(FQuat4d(-450., 180., 400., 42.)));
+	ISPCTransforms.Add(FTransform(FQuat4d(90., 180., 0., 1.)));
+	ISPCTransforms.Add(FTransform(FQuat4d(0., 0., 100., 1.)));
+	ISPCTransforms.Add(FTransform(FQuat4d(100., 100., 100., 1.)));
+	ISPCTransforms.Add(FTransform(FQuat4d(0., 0., 0., 1.)));
+	ISPCTransforms.Add(FTransform(FQuat4d(1., 1., 1., 1.)));
+	ISPCTransforms.Add(FTransform(FQuat4d(-1., 1., -1., 1.)));
+	ISPCTransforms.Add(FTransform(FQuat4d(1., -1., 1., 1.)));
+
+	TArray<FTransform> CPPTransforms(ISPCTransforms);
+	TArrayView<FTransform> ISPCAtoms(ISPCTransforms);
+	TArrayView<FTransform> CPPAtoms(CPPTransforms);
+
+	BoneTrackArray Bones;
+	Bones.Add(BoneTrackPair(0, 0));
+	Bones.Add(BoneTrackPair(1, 1));
+	Bones.Add(BoneTrackPair(2, 2));
+	
+	TArray<uint8> BulkData;
+	BulkData.Init(10, 100);
+	FCompressedOffsetDataBase<TArray<int32>> ScaleOffets;
+	ScaleOffets.OffsetData.Init(30, 100);
+
+	FUECompressedAnimDataMutable CompressedData;
+	CompressedData.CompressedTrackOffsets.Init(30, 100);
+	CompressedData.CompressedScaleOffsets = ScaleOffets;
+	CompressedData.CompressedByteStream = BulkData;
+	CompressedData.KeyEncodingFormat = AKF_ConstantKeyLerp;
+	CompressedData.RotationCompressionFormat = ACF_Identity;
+	CompressedData.TranslationCompressionFormat = ACF_Identity;
+	CompressedData.ScaleCompressionFormat = ACF_Identity;
+	CompressedData.BuildFinalBuffer(BulkData);
+
+	FName AnimName(TEXT(""));
+	FFrameRate FrameRate;
+	TArray<FTrackToSkeletonMap> SkeletonMap;
+
+	GEngine->Exec(nullptr, *FormatCommand(true));
+	{
+		AEFConstantKeyLerp<ACF_Identity> KeyLerp;
+		FAnimSequenceDecompressionContext DecompressContext(FrameRate,
+			1, 
+			EAnimInterpolationType::Linear, 
+			AnimName, 
+			CompressedData,
+			ISPCTransforms,
+			SkeletonMap,
+			nullptr,
+			false);
+		KeyLerp.GetPoseRotations(ISPCAtoms, Bones, DecompressContext);
+	}
+
+	GEngine->Exec(nullptr, *FormatCommand(false));
+	{
+		AEFConstantKeyLerp<ACF_Identity> KeyLerp;
+		FAnimSequenceDecompressionContext DecompressContext(FrameRate,
+			1,
+			EAnimInterpolationType::Linear,
+			AnimName,
+			CompressedData,
+			CPPTransforms,
+			SkeletonMap,
+			nullptr,
+			false);
+		KeyLerp.GetPoseRotations(CPPAtoms, Bones, DecompressContext);
+	}
+
+	GEngine->Exec(nullptr, *FormatCommand(InitialState));
+
+	for (int32 i = 0; i < ISPCAtoms.Num(); ++i)
+	{
+		const FString Message(FString::Format(TEXT("Quat {0}"), { i }));
+		TestTrue(*Message, ISPCAtoms[i].GetRotation().Equals(CPPAtoms[i].GetRotation()));
+	}
+
+	return true;
+}
+
+IMPLEMENT_SIMPLE_AUTOMATION_TEST(FIspcTestAnimEncodingConstantKeyLerpGetPoseTranslations, "Ispc.Animation.ConstantKeyLerp.GetPoseTranslations", EAutomationTestFlags::EditorContext | EAutomationTestFlags::EngineFilter)
+bool FIspcTestAnimEncodingConstantKeyLerpGetPoseTranslations::RunTest(const FString& Parameters)
+{
+	const FString CommandName(TEXT("a.ConstantKeyLerp.ISPC"));
+	auto FormatCommand = [CommandName](bool State) -> FString {
+		return FString::Format(TEXT("{0} {1}"), { CommandName, State });
+	};
+
+	const IConsoleVariable* CVarISPCEnabled = IConsoleManager::Get().FindConsoleVariable(*CommandName);
+	bool InitialState = CVarISPCEnabled->GetBool();
+	check(GEngine);
+
+	TArray<FTransform> ISPCTransforms;
+	ISPCTransforms.Add(FTransform(FVector3d(45., 0., 0.)));
+	ISPCTransforms.Add(FTransform(FVector3d(-450., 180., 400.)));
+	ISPCTransforms.Add(FTransform(FVector3d(90., 180., 0.)));
+	ISPCTransforms.Add(FTransform(FVector3d(0., 0., 100.)));
+	ISPCTransforms.Add(FTransform(FVector3d(100., 100., 100.)));
+	ISPCTransforms.Add(FTransform(FVector3d(0., 0., 0.)));
+	ISPCTransforms.Add(FTransform(FVector3d(1., 1., 1.)));
+	ISPCTransforms.Add(FTransform(FVector3d(-1., 1., -1.)));
+	ISPCTransforms.Add(FTransform(FVector3d(1., -1., 1.)));
+
+	TArray<FTransform> CPPTransforms(ISPCTransforms);
+	TArrayView<FTransform> ISPCAtoms(ISPCTransforms);
+	TArrayView<FTransform> CPPAtoms(CPPTransforms);
+
+	BoneTrackArray Bones;
+	Bones.Add(BoneTrackPair(0, 0));
+	Bones.Add(BoneTrackPair(1, 1));
+	Bones.Add(BoneTrackPair(2, 2));
+
+	TArray<uint8> BulkData;
+	BulkData.Init(10, 100);
+	FCompressedOffsetDataBase<TArray<int32>> ScaleOffets;
+	ScaleOffets.OffsetData.Init(30, 100);
+
+	FUECompressedAnimDataMutable CompressedData;
+	CompressedData.CompressedTrackOffsets.Init(30, 100);
+	CompressedData.CompressedScaleOffsets = ScaleOffets;
+	CompressedData.CompressedByteStream = BulkData;
+	CompressedData.KeyEncodingFormat = AKF_ConstantKeyLerp;
+	CompressedData.RotationCompressionFormat = ACF_Identity;
+	CompressedData.TranslationCompressionFormat = ACF_Identity;
+	CompressedData.ScaleCompressionFormat = ACF_Identity;
+	CompressedData.BuildFinalBuffer(BulkData);
+
+	FName AnimName(TEXT(""));
+	FFrameRate FrameRate;
+	TArray<FTrackToSkeletonMap> SkeletonMap;
+
+	GEngine->Exec(nullptr, *FormatCommand(true));
+	{
+		AEFConstantKeyLerp<ACF_Identity> KeyLerp;
+		FAnimSequenceDecompressionContext DecompressContext(FrameRate,
+			1,
+			EAnimInterpolationType::Linear,
+			AnimName,
+			CompressedData,
+			ISPCTransforms,
+			SkeletonMap,
+			nullptr,
+			false);
+		KeyLerp.GetPoseTranslations(ISPCAtoms, Bones, DecompressContext);
+	}
+
+	GEngine->Exec(nullptr, *FormatCommand(false));
+	{
+		AEFConstantKeyLerp<ACF_Identity> KeyLerp;
+		FAnimSequenceDecompressionContext DecompressContext(FrameRate,
+			1,
+			EAnimInterpolationType::Linear,
+			AnimName,
+			CompressedData,
+			CPPTransforms,
+			SkeletonMap,
+			nullptr,
+			false);
+		KeyLerp.GetPoseTranslations(CPPAtoms, Bones, DecompressContext);
+	}
+
+	GEngine->Exec(nullptr, *FormatCommand(InitialState));
+
+	for (int32 i = 0; i < ISPCAtoms.Num(); ++i)
+	{
+		const FString Message(FString::Format(TEXT("Translation {0}"), { i }));
+		TestTrue(*Message, ISPCAtoms[i].GetTranslation().Equals(CPPAtoms[i].GetTranslation()));
+	}
+
+	return true;
+}
+
+IMPLEMENT_SIMPLE_AUTOMATION_TEST(FIspcTestAnimEncodingConstantKeyLerpGetPoseScales, "Ispc.Animation.ConstantKeyLerp.GetPoseScales", EAutomationTestFlags::EditorContext | EAutomationTestFlags::EngineFilter)
+bool FIspcTestAnimEncodingConstantKeyLerpGetPoseScales::RunTest(const FString& Parameters)
+{
+	const FString CommandName(TEXT("a.ConstantKeyLerp.ISPC"));
+	auto FormatCommand = [CommandName](bool State) -> FString {
+		return FString::Format(TEXT("{0} {1}"), { CommandName, State });
+	};
+
+	const IConsoleVariable* CVarISPCEnabled = IConsoleManager::Get().FindConsoleVariable(*CommandName);
+	bool InitialState = CVarISPCEnabled->GetBool();
+	check(GEngine);
+
+	TArray<FTransform> ISPCTransforms;
+	ISPCTransforms.Add(FTransform(FQuat4d(1., 1., 1., 1.), FVector3d(1., 1., 1.), FVector3d(45., 0., 0.)));
+	ISPCTransforms.Add(FTransform(FQuat4d(1., 1., 1., 1.), FVector3d(1., 1., 1.), FVector3d(45., 180., 0.)));
+	ISPCTransforms.Add(FTransform(FQuat4d(1., 1., 1., 1.), FVector3d(1., 1., 1.), FVector3d(1., -90., 0.)));
+	ISPCTransforms.Add(FTransform(FQuat4d(1., 1., 1., 1.), FVector3d(1., 1., 1.), FVector3d(1., 1., 100.)));
+	ISPCTransforms.Add(FTransform(FQuat4d(1., 1., 1., 1.), FVector3d(1., 1., 1.), FVector3d(0., 0., 0.)));
+	ISPCTransforms.Add(FTransform(FQuat4d(1., 1., 1., 1.), FVector3d(1., 1., 1.), FVector3d(1., 1., 1.)));
+	ISPCTransforms.Add(FTransform(FQuat4d(1., 1., 1., 1.), FVector3d(1., 1., 1.), FVector3d(360., 360., 360.)));
+	ISPCTransforms.Add(FTransform(FQuat4d(1., 1., 1., 1.), FVector3d(1., 1., 1.), FVector3d(9000., 9000., 9000.)));
+	ISPCTransforms.Add(FTransform(FQuat4d(1., 1., 1., 1.), FVector3d(1., 1., 1.), FVector3d(-45., -180., 1.)));
+
+	TArray<FTransform> CPPTransforms(ISPCTransforms);
+	TArrayView<FTransform> ISPCAtoms(ISPCTransforms);
+	TArrayView<FTransform> CPPAtoms(CPPTransforms);
+
+	BoneTrackArray Bones;
+	Bones.Add(BoneTrackPair(0, 0));
+	Bones.Add(BoneTrackPair(1, 1));
+	Bones.Add(BoneTrackPair(2, 2));
+
+	TArray<uint8> BulkData;
+	BulkData.Init(10, 100);
+	FCompressedOffsetDataBase<TArray<int32>> ScaleOffets;
+	ScaleOffets.OffsetData.Init(30, 100);
+
+	FUECompressedAnimDataMutable CompressedData;
+	CompressedData.CompressedTrackOffsets.Init(30, 100);
+	CompressedData.CompressedScaleOffsets = ScaleOffets;
+	CompressedData.CompressedByteStream = BulkData;
+	CompressedData.KeyEncodingFormat = AKF_ConstantKeyLerp;
+	CompressedData.RotationCompressionFormat = ACF_Identity;
+	CompressedData.TranslationCompressionFormat = ACF_Identity;
+	CompressedData.ScaleCompressionFormat = ACF_Identity;
+	CompressedData.BuildFinalBuffer(BulkData);
+
+	FName AnimName(TEXT(""));
+	FFrameRate FrameRate;
+	TArray<FTrackToSkeletonMap> SkeletonMap;
+
+	GEngine->Exec(nullptr, *FormatCommand(true));
+	{
+		AEFConstantKeyLerp<ACF_Identity> KeyLerp;
+		FAnimSequenceDecompressionContext DecompressContext(FrameRate,
+			1,
+			EAnimInterpolationType::Linear,
+			AnimName,
+			CompressedData,
+			ISPCTransforms,
+			SkeletonMap,
+			nullptr,
+			false);
+		KeyLerp.GetPoseScales(ISPCAtoms, Bones, DecompressContext);
+	}
+
+	GEngine->Exec(nullptr, *FormatCommand(false));
+	{
+		AEFConstantKeyLerp<ACF_Identity> KeyLerp;
+		FAnimSequenceDecompressionContext DecompressContext(FrameRate,
+			1,
+			EAnimInterpolationType::Linear,
+			AnimName,
+			CompressedData,
+			CPPTransforms,
+			SkeletonMap,
+			nullptr,
+			false);
+		KeyLerp.GetPoseScales(CPPAtoms, Bones, DecompressContext);
+	}
+
+	GEngine->Exec(nullptr, *FormatCommand(InitialState));
+
+	for (int32 i = 0; i < ISPCAtoms.Num(); ++i)
+	{
+		const FString Message(FString::Format(TEXT("Scale {0}"), { i }));
+		TestTrue(*Message, ISPCAtoms[i].GetScale3D().Equals(CPPAtoms[i].GetScale3D()));
+	}
+
+	return true;
+}
+
+#endif // WITH_DEV_AUTOMATION_TESTS
diff --git a/Engine/Source/Runtime/Engine/Private/Animation/IspcTestAnimEncodingPerTrackCompression.cpp b/Engine/Source/Runtime/Engine/Private/Animation/IspcTestAnimEncodingPerTrackCompression.cpp
new file mode 100644
--- /dev/null
+++ b/Engine/Source/Runtime/Engine/Private/Animation/IspcTestAnimEncodingPerTrackCompression.cpp
@@ -0,0 +1,310 @@
+// Copyright Epic Games, Inc. All Rights Reserved.
+
+#if WITH_DEV_AUTOMATION_TESTS && INTEL_ISPC
+
+#include "CoreMinimal.h"
+#include "Engine/Engine.h"
+#include "Misc/AutomationTest.h"
+#include "AnimEncoding_PerTrackCompression.h"
+
+IMPLEMENT_SIMPLE_AUTOMATION_TEST(FIspcTestAnimEncodingPerTrackCompressionGetPoseRotations, "Ispc.Animation.PerTrackCompression.GetPoseRotations", EAutomationTestFlags::EditorContext | EAutomationTestFlags::EngineFilter)
+bool FIspcTestAnimEncodingPerTrackCompressionGetPoseRotations::RunTest(const FString& Parameters)
+{
+	const FString CommandName(TEXT("a.PerTrackCompression.ISPC"));
+	auto FormatCommand = [CommandName](bool State) -> FString {
+		return FString::Format(TEXT("{0} {1}"), { CommandName, State });
+	};
+
+	const IConsoleVariable* CVarISPCEnabled = IConsoleManager::Get().FindConsoleVariable(*CommandName);
+	bool InitialState = CVarISPCEnabled->GetBool();
+	check(GEngine);
+
+	TArray<FTransform> ISPCTransforms;
+	ISPCTransforms.Add(FTransform(FQuat4d(45., 0., 0., 1.)));
+	ISPCTransforms.Add(FTransform(FQuat4d(-450., 180., 400., 42.)));
+	ISPCTransforms.Add(FTransform(FQuat4d(90., 180., 0., 1.)));
+	ISPCTransforms.Add(FTransform(FQuat4d(0., 0., 100., 1.)));
+	ISPCTransforms.Add(FTransform(FQuat4d(100., 100., 100., 1.)));
+	ISPCTransforms.Add(FTransform(FQuat4d(0., 0., 0., 1.)));
+	ISPCTransforms.Add(FTransform(FQuat4d(1., 1., 1., 1.)));
+	ISPCTransforms.Add(FTransform(FQuat4d(-1., 1., -1., 1.)));
+	ISPCTransforms.Add(FTransform(FQuat4d(1., -1., 1., 1.)));
+
+	TArray<FTransform> CPPTransforms(ISPCTransforms);
+	TArrayView<FTransform> ISPCAtoms(ISPCTransforms);
+	TArrayView<FTransform> CPPAtoms(CPPTransforms);
+
+	BoneTrackArray Bones;
+	Bones.Add(BoneTrackPair(0, 0));
+	Bones.Add(BoneTrackPair(1, 1));
+	Bones.Add(BoneTrackPair(2, 2));
+	
+	TArray<uint8> BulkData;
+	BulkData.SetNumUninitialized(100 * sizeof(int32));
+	for (int32 i = 0; i < BulkData.Num(); i += 4)
+	{
+		BulkData[i] = 0;
+		BulkData[i + 1] = ACF_Identity << 4; // First word is read as the header. Set flag to ACF_Identity.
+		BulkData[i + 2] = 0;
+		BulkData[i + 3] = 0;
+	}
+
+	FCompressedOffsetDataBase<TArray<int32>> ScaleOffets;
+	ScaleOffets.OffsetData.Init(30, 100);
+
+	FUECompressedAnimDataMutable CompressedData;
+	CompressedData.CompressedTrackOffsets.Init(30, 100);
+	CompressedData.CompressedScaleOffsets = ScaleOffets;
+	CompressedData.CompressedByteStream = BulkData;
+	CompressedData.KeyEncodingFormat = AKF_PerTrackCompression;
+	CompressedData.RotationCompressionFormat = ACF_Identity;
+	CompressedData.TranslationCompressionFormat = ACF_Identity;
+	CompressedData.ScaleCompressionFormat = ACF_Identity;
+	CompressedData.BuildFinalBuffer(BulkData);
+
+	FName AnimName(TEXT(""));
+	FFrameRate FrameRate;
+	TArray<FTrackToSkeletonMap> SkeletonMap;
+
+	GEngine->Exec(nullptr, *FormatCommand(true));
+	{
+		AEFPerTrackCompressionCodec KeyLerp;
+		FAnimSequenceDecompressionContext DecompressContext(FrameRate,
+			1,
+			EAnimInterpolationType::Linear,
+			AnimName,
+			CompressedData,
+			ISPCTransforms,
+			SkeletonMap,
+			nullptr,
+			false);
+		KeyLerp.GetPoseRotations(ISPCAtoms, Bones, DecompressContext);
+	}
+
+	GEngine->Exec(nullptr, *FormatCommand(false));
+	{
+		AEFPerTrackCompressionCodec KeyLerp;
+		FAnimSequenceDecompressionContext DecompressContext(FrameRate,
+			1,
+			EAnimInterpolationType::Linear,
+			AnimName,
+			CompressedData,
+			CPPTransforms,
+			SkeletonMap,
+			nullptr,
+			false);
+		KeyLerp.GetPoseRotations(CPPAtoms, Bones, DecompressContext);
+	}
+
+	GEngine->Exec(nullptr, *FormatCommand(InitialState));
+
+	for (int32 i = 0; i < ISPCAtoms.Num(); ++i)
+	{
+		const FString Message(FString::Format(TEXT("Quat {0}"), { i }));
+		TestTrue(*Message, ISPCAtoms[i].GetRotation().Equals(CPPAtoms[i].GetRotation()));
+	}
+
+	return true;
+}
+
+IMPLEMENT_SIMPLE_AUTOMATION_TEST(FIspcTestAnimEncodingPerTrackCompressionGetPoseTranslations, "Ispc.Animation.PerTrackCompression.GetPoseTranslations", EAutomationTestFlags::EditorContext | EAutomationTestFlags::EngineFilter)
+bool FIspcTestAnimEncodingPerTrackCompressionGetPoseTranslations::RunTest(const FString& Parameters)
+{
+	const FString CommandName(TEXT("a.PerTrackCompression.ISPC"));
+	auto FormatCommand = [CommandName](bool State) -> FString {
+		return FString::Format(TEXT("{0} {1}"), { CommandName, State });
+	};
+
+	const IConsoleVariable* CVarISPCEnabled = IConsoleManager::Get().FindConsoleVariable(*CommandName);
+	bool InitialState = CVarISPCEnabled->GetBool();
+	check(GEngine);
+
+	TArray<FTransform> ISPCTransforms;
+	ISPCTransforms.Add(FTransform(FVector3d(45., 0., 0.)));
+	ISPCTransforms.Add(FTransform(FVector3d(-450., 180., 400.)));
+	ISPCTransforms.Add(FTransform(FVector3d(90., 180., 0.)));
+	ISPCTransforms.Add(FTransform(FVector3d(0., 0., 100.)));
+	ISPCTransforms.Add(FTransform(FVector3d(100., 100., 100.)));
+	ISPCTransforms.Add(FTransform(FVector3d(0., 0., 0.)));
+	ISPCTransforms.Add(FTransform(FVector3d(1., 1., 1.)));
+	ISPCTransforms.Add(FTransform(FVector3d(-1., 1., -1.)));
+	ISPCTransforms.Add(FTransform(FVector3d(1., -1., 1.)));
+
+	TArray<FTransform> CPPTransforms(ISPCTransforms);
+	TArrayView<FTransform> ISPCAtoms(ISPCTransforms);
+	TArrayView<FTransform> CPPAtoms(CPPTransforms);
+
+	BoneTrackArray Bones;
+	Bones.Add(BoneTrackPair(0, 0));
+	Bones.Add(BoneTrackPair(1, 1));
+	Bones.Add(BoneTrackPair(2, 2));
+
+	TArray<uint8> BulkData;
+	BulkData.SetNumUninitialized(100 * sizeof(int32));
+	for (int32 i = 0; i < BulkData.Num(); i += 4)
+	{
+		BulkData[i] = 0;
+		BulkData[i + 1] = ACF_Identity << 4; // First word is read as the header. Set flag to ACF_Identity.
+		BulkData[i + 2] = 0;
+		BulkData[i + 3] = 0;
+	}
+
+	FCompressedOffsetDataBase<TArray<int32>> ScaleOffets;
+	ScaleOffets.OffsetData.Init(30, 100);
+
+	FUECompressedAnimDataMutable CompressedData;
+	CompressedData.CompressedTrackOffsets.Init(30, 100);
+	CompressedData.CompressedScaleOffsets = ScaleOffets;
+	CompressedData.CompressedByteStream = BulkData;
+	CompressedData.KeyEncodingFormat = AKF_PerTrackCompression;
+	CompressedData.RotationCompressionFormat = ACF_Identity;
+	CompressedData.TranslationCompressionFormat = ACF_Identity;
+	CompressedData.ScaleCompressionFormat = ACF_Identity;
+	CompressedData.BuildFinalBuffer(BulkData);
+
+	FName AnimName(TEXT(""));
+	FFrameRate FrameRate;
+	TArray<FTrackToSkeletonMap> SkeletonMap;
+
+	GEngine->Exec(nullptr, *FormatCommand(true));
+	{
+		AEFPerTrackCompressionCodec KeyLerp;
+		FAnimSequenceDecompressionContext DecompressContext(FrameRate,
+			1,
+			EAnimInterpolationType::Linear,
+			AnimName,
+			CompressedData,
+			ISPCTransforms,
+			SkeletonMap,
+			nullptr,
+			false);
+		KeyLerp.GetPoseTranslations(ISPCAtoms, Bones, DecompressContext);
+	}
+
+	GEngine->Exec(nullptr, *FormatCommand(false));
+	{
+		AEFPerTrackCompressionCodec KeyLerp;
+		FAnimSequenceDecompressionContext DecompressContext(FrameRate,
+			1,
+			EAnimInterpolationType::Linear,
+			AnimName,
+			CompressedData,
+			CPPTransforms,
+			SkeletonMap,
+			nullptr,
+			false);
+		KeyLerp.GetPoseTranslations(CPPAtoms, Bones, DecompressContext);
+	}
+
+	GEngine->Exec(nullptr, *FormatCommand(InitialState));
+
+	for (int32 i = 0; i < ISPCAtoms.Num(); ++i)
+	{
+		const FString Message(FString::Format(TEXT("Translation {0}"), { i }));
+		TestTrue(*Message, ISPCAtoms[i].GetTranslation().Equals(CPPAtoms[i].GetTranslation()));
+	}
+
+	return true;
+}
+
+IMPLEMENT_SIMPLE_AUTOMATION_TEST(FIspcTestAnimEncodingPerTrackCompressionGetPoseScales, "Ispc.Animation.PerTrackCompression.GetPoseScales", EAutomationTestFlags::EditorContext | EAutomationTestFlags::EngineFilter)
+bool FIspcTestAnimEncodingPerTrackCompressionGetPoseScales::RunTest(const FString& Parameters)
+{
+	const FString CommandName(TEXT("a.PerTrackCompression.ISPC"));
+	auto FormatCommand = [CommandName](bool State) -> FString {
+		return FString::Format(TEXT("{0} {1}"), { CommandName, State });
+	};
+
+	const IConsoleVariable* CVarISPCEnabled = IConsoleManager::Get().FindConsoleVariable(*CommandName);
+	bool InitialState = CVarISPCEnabled->GetBool();
+	check(GEngine);
+
+	TArray<FTransform> ISPCTransforms;
+	ISPCTransforms.Add(FTransform(FQuat4d(1., 1., 1., 1.), FVector3d(1., 1., 1.), FVector3d(45., 0., 0.)));
+	ISPCTransforms.Add(FTransform(FQuat4d(1., 1., 1., 1.), FVector3d(1., 1., 1.), FVector3d(45., 180., 0.)));
+	ISPCTransforms.Add(FTransform(FQuat4d(1., 1., 1., 1.), FVector3d(1., 1., 1.), FVector3d(1., -90., 0.)));
+	ISPCTransforms.Add(FTransform(FQuat4d(1., 1., 1., 1.), FVector3d(1., 1., 1.), FVector3d(1., 1., 100.)));
+	ISPCTransforms.Add(FTransform(FQuat4d(1., 1., 1., 1.), FVector3d(1., 1., 1.), FVector3d(0., 0., 0.)));
+	ISPCTransforms.Add(FTransform(FQuat4d(1., 1., 1., 1.), FVector3d(1., 1., 1.), FVector3d(1., 1., 1.)));
+	ISPCTransforms.Add(FTransform(FQuat4d(1., 1., 1., 1.), FVector3d(1., 1., 1.), FVector3d(360., 360., 360.)));
+	ISPCTransforms.Add(FTransform(FQuat4d(1., 1., 1., 1.), FVector3d(1., 1., 1.), FVector3d(9000., 9000., 9000.)));
+	ISPCTransforms.Add(FTransform(FQuat4d(1., 1., 1., 1.), FVector3d(1., 1., 1.), FVector3d(-45., -180., 1.)));
+
+	TArray<FTransform> CPPTransforms(ISPCTransforms);
+	TArrayView<FTransform> ISPCAtoms(ISPCTransforms);
+	TArrayView<FTransform> CPPAtoms(CPPTransforms);
+
+	BoneTrackArray Bones;
+	Bones.Add(BoneTrackPair(0, 0));
+	Bones.Add(BoneTrackPair(1, 1));
+	Bones.Add(BoneTrackPair(2, 2));
+
+	TArray<uint8> BulkData;
+	BulkData.SetNumUninitialized(100 * sizeof(int32));
+	for (int32 i = 0; i < BulkData.Num(); i += 4)
+	{
+		BulkData[i] = 0;
+		BulkData[i + 1] = ACF_Identity << 4; // First word is read as the header. Set flag to ACF_Identity.
+		BulkData[i + 2] = 0;
+		BulkData[i + 3] = 0;
+	}
+
+	FCompressedOffsetDataBase<TArray<int32>> ScaleOffets;
+	ScaleOffets.OffsetData.Init(30, 100);
+
+	FUECompressedAnimDataMutable CompressedData;
+	CompressedData.CompressedTrackOffsets.Init(30, 100);
+	CompressedData.CompressedScaleOffsets = ScaleOffets;
+	CompressedData.CompressedByteStream = BulkData;
+	CompressedData.KeyEncodingFormat = AKF_PerTrackCompression;
+	CompressedData.RotationCompressionFormat = ACF_Identity;
+	CompressedData.TranslationCompressionFormat = ACF_Identity;
+	CompressedData.ScaleCompressionFormat = ACF_Identity;
+	CompressedData.BuildFinalBuffer(BulkData);
+
+	FName AnimName(TEXT(""));
+	FFrameRate FrameRate;
+	TArray<FTrackToSkeletonMap> SkeletonMap;
+
+	GEngine->Exec(nullptr, *FormatCommand(true));
+	{
+		AEFPerTrackCompressionCodec KeyLerp;
+		FAnimSequenceDecompressionContext DecompressContext(FrameRate,
+			1,
+			EAnimInterpolationType::Linear,
+			AnimName,
+			CompressedData,
+			ISPCTransforms,
+			SkeletonMap,
+			nullptr,
+			false);
+		KeyLerp.GetPoseScales(ISPCAtoms, Bones, DecompressContext);
+	}
+
+	GEngine->Exec(nullptr, *FormatCommand(false));
+	{
+		AEFPerTrackCompressionCodec KeyLerp;
+		FAnimSequenceDecompressionContext DecompressContext(FrameRate,
+			1,
+			EAnimInterpolationType::Linear,
+			AnimName,
+			CompressedData,
+			CPPTransforms,
+			SkeletonMap,
+			nullptr,
+			false);
+		KeyLerp.GetPoseScales(CPPAtoms, Bones, DecompressContext);
+	}
+
+	GEngine->Exec(nullptr, *FormatCommand(InitialState));
+
+	for (int32 i = 0; i < ISPCAtoms.Num(); ++i)
+	{
+		const FString Message(FString::Format(TEXT("Scale {0}"), { i }));
+		TestTrue(*Message, ISPCAtoms[i].GetScale3D().Equals(CPPAtoms[i].GetScale3D()));
+	}
+
+	return true;
+}
+
+#endif // WITH_DEV_AUTOMATION_TESTS
diff --git a/Engine/Source/Runtime/Engine/Private/Animation/IspcTestAnimEncodingVariableKeyLerp.cpp b/Engine/Source/Runtime/Engine/Private/Animation/IspcTestAnimEncodingVariableKeyLerp.cpp
new file mode 100644
--- /dev/null
+++ b/Engine/Source/Runtime/Engine/Private/Animation/IspcTestAnimEncodingVariableKeyLerp.cpp
@@ -0,0 +1,286 @@
+// Copyright Epic Games, Inc. All Rights Reserved.
+
+#if WITH_DEV_AUTOMATION_TESTS && INTEL_ISPC
+
+#include "CoreMinimal.h"
+#include "Engine/Engine.h"
+#include "Misc/AutomationTest.h"
+#include "AnimEncoding_VariableKeyLerp.h"
+
+IMPLEMENT_SIMPLE_AUTOMATION_TEST(FIspcTestAnimEncodingVariableKeyLerpGetPoseRotations, "Ispc.Animation.VariableKeyLerp.GetPoseRotations", EAutomationTestFlags::EditorContext | EAutomationTestFlags::EngineFilter)
+bool FIspcTestAnimEncodingVariableKeyLerpGetPoseRotations::RunTest(const FString& Parameters)
+{
+	const FString CommandName(TEXT("a.VariableKeyLerp.ISPC"));
+	auto FormatCommand = [CommandName](bool State) -> FString {
+		return FString::Format(TEXT("{0} {1}"), { CommandName, State });
+	};
+
+	const IConsoleVariable* CVarISPCEnabled = IConsoleManager::Get().FindConsoleVariable(*CommandName);
+	bool InitialState = CVarISPCEnabled->GetBool();
+	check(GEngine);
+
+	TArray<FTransform> ISPCTransforms;
+	ISPCTransforms.Add(FTransform(FQuat4d(45., 0., 0., 1.)));
+	ISPCTransforms.Add(FTransform(FQuat4d(-450., 180., 400., 42.)));
+	ISPCTransforms.Add(FTransform(FQuat4d(90., 180., 0., 1.)));
+	ISPCTransforms.Add(FTransform(FQuat4d(0., 0., 100., 1.)));
+	ISPCTransforms.Add(FTransform(FQuat4d(100., 100., 100., 1.)));
+	ISPCTransforms.Add(FTransform(FQuat4d(0., 0., 0., 1.)));
+	ISPCTransforms.Add(FTransform(FQuat4d(1., 1., 1., 1.)));
+	ISPCTransforms.Add(FTransform(FQuat4d(-1., 1., -1., 1.)));
+	ISPCTransforms.Add(FTransform(FQuat4d(1., -1., 1., 1.)));
+
+	TArray<FTransform> CPPTransforms(ISPCTransforms);
+	TArrayView<FTransform> ISPCAtoms(ISPCTransforms);
+	TArrayView<FTransform> CPPAtoms(CPPTransforms);
+
+	BoneTrackArray Bones;
+	Bones.Add(BoneTrackPair(0, 0));
+	Bones.Add(BoneTrackPair(1, 1));
+	Bones.Add(BoneTrackPair(2, 2));
+	
+	TArray<uint8> BulkData;
+	BulkData.Init(10, 100);
+	FCompressedOffsetDataBase<TArray<int32>> ScaleOffets;
+	ScaleOffets.OffsetData.Init(30, 100);
+
+	FUECompressedAnimDataMutable CompressedData;
+	CompressedData.CompressedTrackOffsets.Init(30, 100);
+	CompressedData.CompressedScaleOffsets = ScaleOffets;
+	CompressedData.CompressedByteStream = BulkData;
+	CompressedData.KeyEncodingFormat = AKF_VariableKeyLerp;
+	CompressedData.RotationCompressionFormat = ACF_Identity;
+	CompressedData.TranslationCompressionFormat = ACF_Identity;
+	CompressedData.ScaleCompressionFormat = ACF_Identity;
+	CompressedData.BuildFinalBuffer(BulkData);
+
+	FName AnimName(TEXT(""));
+	FFrameRate FrameRate;
+	TArray<FTrackToSkeletonMap> SkeletonMap;
+
+	GEngine->Exec(nullptr, *FormatCommand(true));
+	{
+		AEFVariableKeyLerp<ACF_Identity> KeyLerp;
+		FAnimSequenceDecompressionContext DecompressContext(FrameRate,
+			1,
+			EAnimInterpolationType::Linear,
+			AnimName,
+			CompressedData,
+			ISPCTransforms,
+			SkeletonMap,
+			nullptr,
+			false);
+		KeyLerp.GetPoseRotations(ISPCAtoms, Bones, DecompressContext);
+	}
+
+	GEngine->Exec(nullptr, *FormatCommand(false));
+	{
+		AEFVariableKeyLerp<ACF_Identity> KeyLerp;
+		FAnimSequenceDecompressionContext DecompressContext(FrameRate,
+			1,
+			EAnimInterpolationType::Linear,
+			AnimName,
+			CompressedData,
+			CPPTransforms,
+			SkeletonMap,
+			nullptr,
+			false);
+		KeyLerp.GetPoseRotations(CPPAtoms, Bones, DecompressContext);
+	}
+
+	GEngine->Exec(nullptr, *FormatCommand(InitialState));
+
+	for (int32 i = 0; i < ISPCAtoms.Num(); ++i)
+	{
+		const FString Message(FString::Format(TEXT("Quat {0}"), { i }));
+		TestTrue(*Message, ISPCAtoms[i].GetRotation().Equals(CPPAtoms[i].GetRotation()));
+	}
+
+	return true;
+}
+
+IMPLEMENT_SIMPLE_AUTOMATION_TEST(FIspcTestAnimEncodingVariableKeyLerpGetPoseTranslations, "Ispc.Animation.VariableKeyLerp.GetPoseTranslations", EAutomationTestFlags::EditorContext | EAutomationTestFlags::EngineFilter)
+bool FIspcTestAnimEncodingVariableKeyLerpGetPoseTranslations::RunTest(const FString& Parameters)
+{
+	const FString CommandName(TEXT("a.VariableKeyLerp.ISPC"));
+	auto FormatCommand = [CommandName](bool State) -> FString {
+		return FString::Format(TEXT("{0} {1}"), { CommandName, State });
+	};
+
+	const IConsoleVariable* CVarISPCEnabled = IConsoleManager::Get().FindConsoleVariable(*CommandName);
+	bool InitialState = CVarISPCEnabled->GetBool();
+	check(GEngine);
+
+	TArray<FTransform> ISPCTransforms;
+	ISPCTransforms.Add(FTransform(FVector3d(45., 0., 0.)));
+	ISPCTransforms.Add(FTransform(FVector3d(-450., 180., 400.)));
+	ISPCTransforms.Add(FTransform(FVector3d(90., 180., 0.)));
+	ISPCTransforms.Add(FTransform(FVector3d(0., 0., 100.)));
+	ISPCTransforms.Add(FTransform(FVector3d(100., 100., 100.)));
+	ISPCTransforms.Add(FTransform(FVector3d(0., 0., 0.)));
+	ISPCTransforms.Add(FTransform(FVector3d(1., 1., 1.)));
+	ISPCTransforms.Add(FTransform(FVector3d(-1., 1., -1.)));
+	ISPCTransforms.Add(FTransform(FVector3d(1., -1., 1.)));
+
+	TArray<FTransform> CPPTransforms(ISPCTransforms);
+	TArrayView<FTransform> ISPCAtoms(ISPCTransforms);
+	TArrayView<FTransform> CPPAtoms(CPPTransforms);
+
+	BoneTrackArray Bones;
+	Bones.Add(BoneTrackPair(0, 0));
+	Bones.Add(BoneTrackPair(1, 1));
+	Bones.Add(BoneTrackPair(2, 2));
+
+	TArray<uint8> BulkData;
+	BulkData.Init(10, 100);
+	FCompressedOffsetDataBase<TArray<int32>> ScaleOffets;
+	ScaleOffets.OffsetData.Init(30, 100);
+
+	FUECompressedAnimDataMutable CompressedData;
+	CompressedData.CompressedTrackOffsets.Init(30, 100);
+	CompressedData.CompressedScaleOffsets = ScaleOffets;
+	CompressedData.CompressedByteStream = BulkData;
+	CompressedData.KeyEncodingFormat = AKF_VariableKeyLerp;
+	CompressedData.RotationCompressionFormat = ACF_Identity;
+	CompressedData.TranslationCompressionFormat = ACF_Identity;
+	CompressedData.ScaleCompressionFormat = ACF_Identity;
+	CompressedData.BuildFinalBuffer(BulkData);
+
+	FName AnimName(TEXT(""));
+	FFrameRate FrameRate;
+	TArray<FTrackToSkeletonMap> SkeletonMap;
+
+	GEngine->Exec(nullptr, *FormatCommand(true));
+	{
+		AEFVariableKeyLerp<ACF_Identity> KeyLerp;
+		FAnimSequenceDecompressionContext DecompressContext(FrameRate,
+			1,
+			EAnimInterpolationType::Linear,
+			AnimName,
+			CompressedData,
+			ISPCTransforms,
+			SkeletonMap,
+			nullptr,
+			false);
+		KeyLerp.GetPoseTranslations(ISPCAtoms, Bones, DecompressContext);
+	}
+
+	GEngine->Exec(nullptr, *FormatCommand(false));
+	{
+		AEFVariableKeyLerp<ACF_Identity> KeyLerp;
+		FAnimSequenceDecompressionContext DecompressContext(FrameRate,
+			1,
+			EAnimInterpolationType::Linear,
+			AnimName,
+			CompressedData,
+			CPPTransforms,
+			SkeletonMap,
+			nullptr,
+			false);
+		KeyLerp.GetPoseTranslations(CPPAtoms, Bones, DecompressContext);
+	}
+
+	GEngine->Exec(nullptr, *FormatCommand(InitialState));
+
+	for (int32 i = 0; i < ISPCAtoms.Num(); ++i)
+	{
+		const FString Message(FString::Format(TEXT("Translation {0}"), { i }));
+		TestTrue(*Message, ISPCAtoms[i].GetTranslation().Equals(CPPAtoms[i].GetTranslation()));
+	}
+
+	return true;
+}
+
+IMPLEMENT_SIMPLE_AUTOMATION_TEST(FIspcTestAnimEncodingVariableKeyLerpGetPoseScales, "Ispc.Animation.VariableKeyLerp.GetPoseScales", EAutomationTestFlags::EditorContext | EAutomationTestFlags::EngineFilter)
+bool FIspcTestAnimEncodingVariableKeyLerpGetPoseScales::RunTest(const FString& Parameters)
+{
+	const FString CommandName(TEXT("a.VariableKeyLerp.ISPC"));
+	auto FormatCommand = [CommandName](bool State) -> FString {
+		return FString::Format(TEXT("{0} {1}"), { CommandName, State });
+	};
+
+	const IConsoleVariable* CVarISPCEnabled = IConsoleManager::Get().FindConsoleVariable(*CommandName);
+	bool InitialState = CVarISPCEnabled->GetBool();
+	check(GEngine);
+
+	TArray<FTransform> ISPCTransforms;
+	ISPCTransforms.Add(FTransform(FQuat4d(1., 1., 1., 1.), FVector3d(1., 1., 1.), FVector3d(45., 0., 0.)));
+	ISPCTransforms.Add(FTransform(FQuat4d(1., 1., 1., 1.), FVector3d(1., 1., 1.), FVector3d(45., 180., 0.)));
+	ISPCTransforms.Add(FTransform(FQuat4d(1., 1., 1., 1.), FVector3d(1., 1., 1.), FVector3d(1., -90., 0.)));
+	ISPCTransforms.Add(FTransform(FQuat4d(1., 1., 1., 1.), FVector3d(1., 1., 1.), FVector3d(1., 1., 100.)));
+	ISPCTransforms.Add(FTransform(FQuat4d(1., 1., 1., 1.), FVector3d(1., 1., 1.), FVector3d(0., 0., 0.)));
+	ISPCTransforms.Add(FTransform(FQuat4d(1., 1., 1., 1.), FVector3d(1., 1., 1.), FVector3d(1., 1., 1.)));
+	ISPCTransforms.Add(FTransform(FQuat4d(1., 1., 1., 1.), FVector3d(1., 1., 1.), FVector3d(360., 360., 360.)));
+	ISPCTransforms.Add(FTransform(FQuat4d(1., 1., 1., 1.), FVector3d(1., 1., 1.), FVector3d(9000., 9000., 9000.)));
+	ISPCTransforms.Add(FTransform(FQuat4d(1., 1., 1., 1.), FVector3d(1., 1., 1.), FVector3d(-45., -180., 1.)));
+
+	TArray<FTransform> CPPTransforms(ISPCTransforms);
+	TArrayView<FTransform> ISPCAtoms(ISPCTransforms);
+	TArrayView<FTransform> CPPAtoms(CPPTransforms);
+
+	BoneTrackArray Bones;
+	Bones.Add(BoneTrackPair(0, 0));
+	Bones.Add(BoneTrackPair(1, 1));
+	Bones.Add(BoneTrackPair(2, 2));
+
+	TArray<uint8> BulkData;
+	BulkData.Init(10, 100);
+	FCompressedOffsetDataBase<TArray<int32>> ScaleOffets;
+	ScaleOffets.OffsetData.Init(30, 100);
+
+	FUECompressedAnimDataMutable CompressedData;
+	CompressedData.CompressedTrackOffsets.Init(30, 100);
+	CompressedData.CompressedScaleOffsets = ScaleOffets;
+	CompressedData.CompressedByteStream = BulkData;
+	CompressedData.KeyEncodingFormat = AKF_VariableKeyLerp;
+	CompressedData.RotationCompressionFormat = ACF_Identity;
+	CompressedData.TranslationCompressionFormat = ACF_Identity;
+	CompressedData.ScaleCompressionFormat = ACF_Identity;
+	CompressedData.BuildFinalBuffer(BulkData);
+
+	FName AnimName(TEXT(""));
+	FFrameRate FrameRate;
+	TArray<FTrackToSkeletonMap> SkeletonMap;
+
+	GEngine->Exec(nullptr, *FormatCommand(true));
+	{
+		AEFVariableKeyLerp<ACF_Identity> KeyLerp;
+		FAnimSequenceDecompressionContext DecompressContext(FrameRate,
+			1,
+			EAnimInterpolationType::Linear,
+			AnimName,
+			CompressedData,
+			ISPCTransforms,
+			SkeletonMap,
+			nullptr,
+			false);
+		KeyLerp.GetPoseScales(ISPCAtoms, Bones, DecompressContext);
+	}
+
+	GEngine->Exec(nullptr, *FormatCommand(false));
+	{
+		AEFVariableKeyLerp<ACF_Identity> KeyLerp;
+		FAnimSequenceDecompressionContext DecompressContext(FrameRate,
+			1,
+			EAnimInterpolationType::Linear,
+			AnimName,
+			CompressedData,
+			CPPTransforms,
+			SkeletonMap,
+			nullptr,
+			false);
+		KeyLerp.GetPoseScales(CPPAtoms, Bones, DecompressContext);
+	}
+
+	GEngine->Exec(nullptr, *FormatCommand(InitialState));
+
+	for (int32 i = 0; i < ISPCAtoms.Num(); ++i)
+	{
+		const FString Message(FString::Format(TEXT("Scale {0}"), { i }));
+		TestTrue(*Message, ISPCAtoms[i].GetScale3D().Equals(CPPAtoms[i].GetScale3D()));
+	}
+
+	return true;
+}
+
+#endif // WITH_DEV_AUTOMATION_TESTS
diff --git a/Engine/Source/Runtime/Experimental/Chaos/Private/Chaos/IspcTestPerParticleDampVelocity.cpp b/Engine/Source/Runtime/Experimental/Chaos/Private/Chaos/IspcTestPerParticleDampVelocity.cpp
new file mode 100644
--- /dev/null
+++ b/Engine/Source/Runtime/Experimental/Chaos/Private/Chaos/IspcTestPerParticleDampVelocity.cpp
@@ -0,0 +1,99 @@
+// Copyright Epic Games, Inc. All Rights Reserved.
+
+#if WITH_DEV_AUTOMATION_TESTS && INTEL_ISPC
+
+#include "CoreMinimal.h"
+#include "Misc/AutomationTest.h"
+#include "Chaos/PerParticleDampVelocity.h"
+
+using namespace Chaos;
+
+IMPLEMENT_SIMPLE_AUTOMATION_TEST(FIspcPerParticleDampVelocity, "Ispc.Physics.PerParticleDampVelocity", EAutomationTestFlags::EditorContext | EAutomationTestFlags::EngineFilter)
+bool FIspcPerParticleDampVelocity::RunTest(const FString& Parameters)
+{
+	// Chaos module doesn't let us include Engine.h, therefore we cannot call the console command.
+	bool InitialState = bChaos_DampVelocity_ISPC_Enabled;
+
+	const int32 ParticleOffset = 0;
+	const int32 NumParticles = 101;
+
+	const Softs::FSolverReal InvMValue = 1.f;
+	const Softs::FSolverReal MValue = 1.f;
+	Softs::FPAndInvM FPInvMData;
+	FPInvMData.P = Softs::FSolverVec3(0.f);
+	FPInvMData.InvM = InvMValue;
+
+	Softs::FSolverVec3 XSolverVec(1.f, 2.f, 3.f);
+	Softs::FSolverVec3 VSolverVec(1.f, 2.f, 3.f);
+
+	bChaos_DampVelocity_ISPC_Enabled = true;
+	Softs::FSolverParticles ISPCParticles;
+	// TParticles::Resize() will make all the necessary arrays the same size.
+	ISPCParticles.Resize(NumParticles);
+	{
+		// HACK: Can't get to TDynamicParticles::MInvM any other way.
+		TArray<Softs::FSolverReal>& InvMArray = *const_cast<TArrayCollectionArray<Softs::FSolverReal>*>(&ISPCParticles.GetInvM());
+		TArray<Softs::FPAndInvM>& FPInvMArray = ISPCParticles.GetPAndInvM();
+		TArray<Softs::FSolverVec3>& XArray = ISPCParticles.AllX();
+		TArray<Softs::FSolverVec3>& VArray = *const_cast<TArrayCollectionArray<Softs::FSolverVec3>*>(&ISPCParticles.GetV());
+		TArray<Softs::FSolverReal>& MArray = *const_cast<TArrayCollectionArray<Softs::FSolverReal>*>(&ISPCParticles.GetM());
+
+		for (int32 i = 0; i < NumParticles; ++i)
+		{
+			InvMArray[i] = InvMValue;
+			FPInvMArray[i].P = FPInvMData.P;
+			FPInvMArray[i].InvM = FPInvMData.InvM;
+			XArray[i] = XSolverVec;
+			VArray[i] = VSolverVec;
+			MArray[i] = MValue;
+		}
+
+		Softs::FPerParticleDampVelocity PerParticleDampVelocity;
+		PerParticleDampVelocity.UpdatePositionBasedState(ISPCParticles, 0, NumParticles);
+	}
+
+	bChaos_DampVelocity_ISPC_Enabled = false;
+	Softs::FSolverParticles CPPParticles;
+	// TParticles::Resize() will make all the necessary arrays the same size.
+	CPPParticles.Resize(NumParticles);
+	{
+		// HACK: Can't get to TDynamicParticles::MInvM any other way.
+		TArray<Softs::FSolverReal>& InvMArray = *const_cast<TArrayCollectionArray<Softs::FSolverReal>*>(&CPPParticles.GetInvM());
+		TArray<Softs::FPAndInvM>& FPInvMArray = CPPParticles.GetPAndInvM();
+		TArray<Softs::FSolverVec3>& XArray = CPPParticles.AllX();
+		TArray<Softs::FSolverVec3>& VArray = *const_cast<TArrayCollectionArray<Softs::FSolverVec3>*>(&CPPParticles.GetV());
+		TArray<Softs::FSolverReal>& MArray = *const_cast<TArrayCollectionArray<Softs::FSolverReal>*>(&CPPParticles.GetM());
+
+		for (int32 i = 0; i < NumParticles; ++i)
+		{
+			InvMArray[i] = InvMValue;
+			FPInvMArray[i].P = FPInvMData.P;
+			FPInvMArray[i].InvM = FPInvMData.InvM;
+			XArray[i] = XSolverVec;
+			VArray[i] = VSolverVec;
+			MArray[i] = MValue;
+		}
+
+		Softs::FPerParticleDampVelocity PerParticleDampVelocity;
+		PerParticleDampVelocity.UpdatePositionBasedState(CPPParticles, 0, NumParticles);
+	}
+
+	bChaos_DampVelocity_ISPC_Enabled = InitialState;
+
+	TArray<Softs::FPAndInvM>& ISPCFPInvMArray = ISPCParticles.GetPAndInvM();
+	TArray<Softs::FPAndInvM>& CPPFPInvMArray = CPPParticles.GetPAndInvM();
+	for (int32 i = 0; i < ISPCFPInvMArray.Num(); ++i)
+	{
+		const Softs::FPAndInvM& ISPCElement = ISPCFPInvMArray[i];
+		const Softs::FPAndInvM& CPPElement = CPPFPInvMArray[i];
+
+		const FString PMessage(FString::Format(TEXT("P Vec {0}"), { i }));
+		const FString InvMMessage(FString::Format(TEXT("InvM {0}"), { i }));
+		TestTrue(*PMessage, ISPCElement.P.Equals(CPPElement.P));
+		TestEqual(*InvMMessage, ISPCElement.InvM, CPPElement.InvM, SMALL_NUMBER);
+	}
+
+	return true;
+}
+
+#endif // WITH_DEV_AUTOMATION_TESTS
diff --git a/Engine/Source/Runtime/SignalProcessing/Private/FloatArrayMath.cpp b/Engine/Source/Runtime/SignalProcessing/Private/FloatArrayMath.cpp
--- a/Engine/Source/Runtime/SignalProcessing/Private/FloatArrayMath.cpp
+++ b/Engine/Source/Runtime/SignalProcessing/Private/FloatArrayMath.cpp
@@ -1465,9 +1465,12 @@ namespace Audio
 
 			if (NumNotToSimd)
 			{
-				for (int32 i = 1; i < Num; i += 2)
+				for (int32 i = NumToSimd; i < Num; i++)
 				{
-					InData[i] *= -1.f;
+					if ((i % 2) == 1)
+					{
+						InData[i] *= -1.f;
+					}
 				}
 			}
 		}
@@ -1946,29 +1949,30 @@ namespace Audio
 				const int32 NumToSimd = Num & MathIntrinsics::SimdMask;
 				const int32 NumNotToSimd = Num & MathIntrinsics::NotSimdMask;
 
+				const float DeltaValue = ((EndValue - StartValue) / Num);
+
 				if (NumToSimd)
 				{
-					const int32 NumIterations = Num / AUDIO_NUM_FLOATS_PER_VECTOR_REGISTER;
-					const float DeltaValue = ((EndValue - StartValue) / NumIterations);
-
-					VectorRegister4Float Gain = VectorLoadFloat1(&StartValue);
+					constexpr VectorRegister4Float VectorFour = MakeVectorRegisterFloatConstant(4.f, 4.f, 4.f, 4.f);
+					VectorRegister4Float Accumulator = MakeVectorRegisterFloat(0.f, 1.f, 2.f, 3.f);
 					VectorRegister4Float Delta = VectorLoadFloat1(&DeltaValue);
+					VectorRegister4Float Start = VectorLoadFloat1(&StartValue);
 
 					for (int32 i = 0; i < NumToSimd; i += AUDIO_NUM_FLOATS_PER_VECTOR_REGISTER)
 					{
 						VectorRegister4Float Output = VectorLoad(&OutFloatBuffer[i]);
+						VectorRegister4Float Gain = VectorMultiplyAdd(Accumulator, Delta, Start);
 						Output = VectorMultiply(Output, Gain);
-						Gain = VectorAdd(Gain, Delta);
+						Accumulator = VectorAdd(Accumulator, VectorFour);
 						VectorStore(Output, &OutFloatBuffer[i]);
 					}
 				}
 
 				if (NumNotToSimd)
 				{
-					float Gain = StartValue;
+					float Gain = (NumToSimd * DeltaValue) + StartValue;
 
 					// Do a fade from start to end
-					const float DeltaValue = ((EndValue - StartValue) / Num);
 					for (int32 i = NumToSimd; i < Num; ++i)
 					{
 						OutFloatBuffer[i] = OutFloatBuffer[i] * Gain;
@@ -2098,29 +2102,29 @@ namespace Audio
 				const int32 NumToSimd = Num & MathIntrinsics::SimdMask;
 				const int32 NumNotToSimd = Num & MathIntrinsics::NotSimdMask;
 
+				const float DeltaValue = ((EndGain - StartGain) / Num);
+
 				if (NumToSimd)
 				{
-					const int32 NumIterations = Num / AUDIO_NUM_FLOATS_PER_VECTOR_REGISTER;
-					const float DeltaValue = ((EndGain - StartGain) / NumIterations);
+					constexpr VectorRegister4Float VectorFour = MakeVectorRegisterFloatConstant(4.f, 4.f, 4.f, 4.f);
+					VectorRegister4Float Accumulator = MakeVectorRegisterFloat(0.f, 1.f, 2.f, 3.f);
 
-					VectorRegister4Float Gain = VectorLoadFloat1(&StartGain);
+					VectorRegister4Float Start = VectorLoadFloat1(&StartGain);
 					VectorRegister4Float Delta = VectorLoadFloat1(&DeltaValue);
 
 					for (int32 i = 0; i < NumToSimd; i += AUDIO_NUM_FLOATS_PER_VECTOR_REGISTER)
 					{
 						VectorRegister4Float Input = VectorLoad(&InFloatBuffer[i]);
 						VectorRegister4Float Output = VectorLoad(&BufferToSumTo[i]);
+						VectorRegister4Float Gain = VectorMultiplyAdd(Accumulator, Delta, Start);
 						Output = VectorMultiplyAdd(Input, Gain, Output);
-
+						Accumulator = VectorAdd(Accumulator, VectorFour);
 						VectorStore(Output, &BufferToSumTo[i]);
-
-						Gain = VectorAdd(Gain, Delta);
 					}
 				}
 
 				if (NumNotToSimd)
 				{
-					const float DeltaValue = ((EndGain - StartGain) / Num);
 					float Gain = (NumToSimd * DeltaValue) + StartGain;
 
 					for (int32 i = NumToSimd; i < Num; ++i)
