Changed the canvas Draw* API to use doubles.

Enabling truncation warnings has revealed numerous spots where we need to do an int->float conversion to get around a warning. Since the float parameters of these functions will be wrapped by a FVector2D, we can just change the type to double, and avoid the manual casts.

diff --git a/Engine/Source/Runtime/Engine/Private/UserInterface/Canvas.cpp b/Engine/Source/Runtime/Engine/Private/UserInterface/Canvas.cpp
--- a/Engine/Source/Runtime/Engine/Private/UserInterface/Canvas.cpp
+++ b/Engine/Source/Runtime/Engine/Private/UserInterface/Canvas.cpp
@@ -1054,13 +1054,13 @@ void FCanvas::Clear(const FLinearColor& ClearColor)
 		});
 }
 
-void FCanvas::DrawTile(float X, float Y, float SizeX, float SizeY, float U, float V, float SizeU, float SizeV, const FLinearColor& Color,	const FTexture* Texture, bool AlphaBlend)
+void FCanvas::DrawTile(double X, double Y, double SizeX, double SizeY, float U, float V, float SizeU, float SizeV, const FLinearColor& Color,	const FTexture* Texture, bool AlphaBlend)
 {
 	ESimpleElementBlendMode BlendMode = AlphaBlend ? (bUseInternalTexture ? SE_BLEND_TranslucentAlphaOnlyWriteAlpha : SE_BLEND_Translucent) : SE_BLEND_Opaque;
 	DrawTile(X, Y, SizeX, SizeY, U, V, SizeU, SizeV, Color, Texture, BlendMode);
 }
 
-void FCanvas::DrawTile(float X, float Y, float SizeX, float SizeY, float U, float V, float SizeU, float SizeV, const FLinearColor& Color, const FTexture* Texture, ESimpleElementBlendMode BlendMode)
+void FCanvas::DrawTile(double X, double Y, double SizeX, double SizeY, float U, float V, float SizeU, float SizeV, const FLinearColor& Color, const FTexture* Texture, ESimpleElementBlendMode BlendMode)
 {
 	SCOPE_CYCLE_COUNTER(STAT_Canvas_DrawTextureTileTime);
 
@@ -1069,7 +1069,7 @@ void FCanvas::DrawTile(float X, float Y, float SizeX, float SizeY, float U, floa
 	DrawItem(TileItem);
 }
 
-int32 FCanvas::DrawShadowedString( float StartX,float StartY,const TCHAR* Text,const UFont* Font,const FLinearColor& Color, const FLinearColor& ShadowColor)
+int32 FCanvas::DrawShadowedString(double StartX, double StartY, const TCHAR* Text, const UFont* Font, const FLinearColor& Color, const FLinearColor& ShadowColor)
 {
 	const float Z = 1.0f;
 	FCanvasTextItem TextItem( FVector2D( StartX, StartY ), FText::FromString( Text ), Font, Color);
@@ -1093,7 +1093,7 @@ void FCanvas::DrawNGon(const FVector2D& Center, const FColor& Color, int32 NumSi
 	DrawItem(NGonItem);
 }
 
-int32 FCanvas::DrawShadowedText( float StartX,float StartY,const FText& Text,const UFont* Font,const FLinearColor& Color, const FLinearColor& ShadowColor )
+int32 FCanvas::DrawShadowedText(double StartX, double StartY, const FText& Text, const UFont* Font, const FLinearColor& Color, const FLinearColor& ShadowColor)
 {
 	const float Z = 1.0f;
 	FCanvasTextItem TextItem( FVector2D( StartX, StartY ), Text, Font, Color );
diff --git a/Engine/Source/Runtime/Engine/Public/CanvasTypes.h b/Engine/Source/Runtime/Engine/Public/CanvasTypes.h
--- a/Engine/Source/Runtime/Engine/Public/CanvasTypes.h
+++ b/Engine/Source/Runtime/Engine/Public/CanvasTypes.h
@@ -731,8 +731,8 @@ public:
 	* @param Texture - Texture to draw
 	* @param AlphaBlend - true to alphablend
 	*/
-	ENGINE_API void DrawTile(float X, float Y, float SizeX, float SizeY, float U, float V,  float SizeU, float SizeV, const FLinearColor& Color, const FTexture* Texture = NULL, bool AlphaBlend = true);
-	ENGINE_API void DrawTile(float X, float Y, float SizeX, float SizeY, float U, float V, float SizeU, float SizeV, const FLinearColor& Color, const FTexture* Texture, ESimpleElementBlendMode BlendMode);
+	ENGINE_API void DrawTile(double X, double Y, double SizeX, double SizeY, float U, float V, float SizeU, float SizeV, const FLinearColor& Color, const FTexture* Texture = NULL, bool AlphaBlend = true);
+	ENGINE_API void DrawTile(double X, double Y, double SizeX, double SizeY, float U, float V, float SizeU, float SizeV, const FLinearColor& Color, const FTexture* Texture, ESimpleElementBlendMode BlendMode);
 
 	/** 
 	* Draw an string centered on given location. 
@@ -746,9 +746,9 @@ public:
 	* @param ShadowColor - Shadow color to draw underneath the text (ignored for distance field fonts)
 	* @return total size in pixels of text drawn
 	*/
-	ENGINE_API int32 DrawShadowedString( float StartX, float StartY, const TCHAR* Text, const UFont* Font, const FLinearColor& Color, const FLinearColor& ShadowColor = FLinearColor::Black );
+	ENGINE_API int32 DrawShadowedString(double StartX, double StartY, const TCHAR* Text, const UFont* Font, const FLinearColor& Color, const FLinearColor& ShadowColor = FLinearColor::Black );
 	
-	ENGINE_API int32 DrawShadowedText( float StartX, float StartY, const FText& Text, const UFont* Font, const FLinearColor& Color, const FLinearColor& ShadowColor = FLinearColor::Black );
+	ENGINE_API int32 DrawShadowedText(double StartX, double StartY, const FText& Text, const UFont* Font, const FLinearColor& Color, const FLinearColor& ShadowColor = FLinearColor::Black );
 
 	ENGINE_API void WrapString( FTextSizingParameters& Parameters, const float InCurX, const TCHAR* const pText, TArray<FWrappedStringElement>& out_Lines, FCanvasWordWrapper::FWrappedLineData* const OutWrappedLineData = nullptr);
 
