Fix ray tracing dynamic instances for static meshes not taking into account streaming:
* We need to find the first valid geometry in a similar way we do it for cached instances or dynamic nanite instances.

diff --git a/Engine/Source/Runtime/Engine/Private/StaticMeshRender.cpp b/Engine/Source/Runtime/Engine/Private/StaticMeshRender.cpp
--- a/Engine/Source/Runtime/Engine/Private/StaticMeshRender.cpp
+++ b/Engine/Source/Runtime/Engine/Private/StaticMeshRender.cpp
@@ -1804,7 +1804,23 @@ void FStaticMeshSceneProxy::GetDynamicRayTracingInstances(FRayTracingMaterialGat
 
 	ESceneDepthPriorityGroup PrimitiveDPG = GetStaticDepthPriorityGroup();
 	const int32 LODIndex = FMath::Max(GetLOD(Context.ReferenceView), (int32)GetCurrentFirstLODIdx_RenderThread());
-	const FStaticMeshLODResources& LODModel = RenderData->LODResources[LODIndex];
+
+	int32 RayTracingLODIndex = INDEX_NONE;
+	for (int32 LODIdx = LODIndex; LODIdx < RenderData->LODResources.Num(); LODIdx++)
+	{
+		if (RenderData->LODResources[LODIdx].RayTracingGeometry.IsValid())
+		{
+			RayTracingLODIndex = LODIdx;
+			break;
+		}
+	}
+
+	if (RayTracingLODIndex == INDEX_NONE)
+	{
+		return;
+	}
+
+	const FStaticMeshLODResources& LODModel = RenderData->LODResources[RayTracingLODIndex];
 
 	bool bEvaluateWPO = bDynamicRayTracingGeometry;
 
@@ -1822,7 +1838,7 @@ void FStaticMeshSceneProxy::GetDynamicRayTracingInstances(FRayTracingMaterialGat
 	}
 
 	// TODO: Need to validate that the DynamicRayTracingGeometries are still valid - they could contain streamed out IndexBuffers from the shared StaticMesh (UE-139474)
-	FRayTracingGeometry& Geometry = bEvaluateWPO? DynamicRayTracingGeometries[LODIndex] : RenderData->LODResources[LODIndex].RayTracingGeometry;
+	FRayTracingGeometry& Geometry = bEvaluateWPO? DynamicRayTracingGeometries[LODIndex] : RenderData->LODResources[RayTracingLODIndex].RayTracingGeometry;
 	
 	if (LODModel.GetNumVertices() <= 0 || Geometry.Initializer.TotalPrimitiveCount <= 0)
 	{
