- add GRHISupportAsyncTextureStreamOut which will perform new texture creation for stream out ops to happen on async thread
- add cvar r.D3D12.UseCommittedTexturesNV to force committed resource allocation for all textures as possible texture corruption work around

#rnx

diff --git a/Engine/Source/Runtime/D3D12RHI/Private/D3D12PoolAllocator.cpp b/Engine/Source/Runtime/D3D12RHI/Private/D3D12PoolAllocator.cpp
--- a/Engine/Source/Runtime/D3D12RHI/Private/D3D12PoolAllocator.cpp
+++ b/Engine/Source/Runtime/D3D12RHI/Private/D3D12PoolAllocator.cpp
@@ -7,6 +7,14 @@
 #define NEEDS_D3D12_INDIRECT_ARGUMENT_HEAP_WORKAROUND 0
 #endif
 
+int32 GUseCommittedTexturesNV = 0;
+static FAutoConsoleVariableRef CVarUseCommittedTexturesNV(
+	TEXT("r.D3D12.UseCommittedTexturesNV"),
+	GUseCommittedTexturesNV,
+	TEXT("Force committed resource creation for textures to fix possible driver bug"),
+	ECVF_Default
+);
+
 //-----------------------------------------------------------------------------
 //	FD3D12MemoryPool
 //-----------------------------------------------------------------------------
@@ -293,7 +301,18 @@ void FD3D12PoolAllocator::AllocateResource(uint32 GPUIndex, D3D12_HEAP_TYPE InHe
 
 	const bool bSharedResource = EnumHasAnyFlags(InDesc.Flags, D3D12_RESOURCE_FLAG_ALLOW_SIMULTANEOUS_ACCESS);
 	const bool PoolResource = InSize <= MaxAllocationSize && !bSharedResource;
-	if (PoolResource)
+
+	// Disable pooling for VRAM allocated textures and force use the committed resource path
+	bool bForceCommittedResourcePath = false;
+#if PLATFORM_WINDOWS
+	const bool bUseCommittedTexturesNV = GUseCommittedTexturesNV && IsRHIDeviceNVIDIA();
+	if (PoolResource && bUseCommittedTexturesNV && InHeapType == D3D12_HEAP_TYPE_DEFAULT && InDesc.Dimension != D3D12_RESOURCE_DIMENSION_BUFFER)
+	{
+		bForceCommittedResourcePath = true;
+	}
+#endif // PLATFORM_WINDOWS
+
+	if (PoolResource && !bForceCommittedResourcePath)
 	{
 		const bool bPlacedResource = (AllocationStrategy == EResourceAllocationStrategy::kPlacedResource);
 
@@ -376,10 +395,10 @@ void FD3D12PoolAllocator::AllocateResource(uint32 GPUIndex, D3D12_HEAP_TYPE InHe
 
 		// If we are tracking all allocation data and allocating a standalone texture, then first create a heap so we can retrieve the GPU virtual address as well
 		// UAV Aliasing needs a Heap to create the aliased resource in.
-		if (InDesc.NeedsUAVAliasWorkarounds() || (Adapter->IsTrackingAllAllocations() && InDesc.Dimension != D3D12_RESOURCE_DIMENSION_BUFFER))
+		if (InDesc.NeedsUAVAliasWorkarounds() || (Adapter->IsTrackingAllAllocations() && InDesc.Dimension != D3D12_RESOURCE_DIMENSION_BUFFER && !bForceCommittedResourcePath))
 		{
 			D3D12_HEAP_DESC HeapDesc = {};
-			HeapDesc.SizeInBytes = InSize;
+			HeapDesc.SizeInBytes = FMath::Max((uint64)D3D12_DEFAULT_RESOURCE_PLACEMENT_ALIGNMENT, InSize);
 			HeapDesc.Properties = HeapProps;
 			HeapDesc.Alignment = InDesc.SampleDesc.Count > 1 ? D3D12_DEFAULT_MSAA_RESOURCE_PLACEMENT_ALIGNMENT : 0;
 			HeapDesc.Flags = D3D12_HEAP_FLAG_ALLOW_ALL_BUFFERS_AND_TEXTURES;
diff --git a/Engine/Source/Runtime/D3D12RHI/Private/Windows/WindowsD3D12Device.cpp b/Engine/Source/Runtime/D3D12RHI/Private/Windows/WindowsD3D12Device.cpp
--- a/Engine/Source/Runtime/D3D12RHI/Private/Windows/WindowsD3D12Device.cpp
+++ b/Engine/Source/Runtime/D3D12RHI/Private/Windows/WindowsD3D12Device.cpp
@@ -1527,6 +1527,7 @@ void FD3D12DynamicRHI::Init()
 	GRHISupportsAsyncTextureCreation = D3D12RHI_ShouldAllowAsyncResourceCreation();
 	if (GRHISupportsAsyncTextureCreation)
 	{
+		GRHISupportAsyncTextureStreamOut = true;
 		UE_LOG(LogD3D12RHI, Log, TEXT("Async texture creation enabled"));
 	}
 	else
diff --git a/Engine/Source/Runtime/Engine/Private/Streaming/Texture2DStreamOut_AsyncCreate.cpp b/Engine/Source/Runtime/Engine/Private/Streaming/Texture2DStreamOut_AsyncCreate.cpp
new file mode 100644
--- /dev/null
+++ b/Engine/Source/Runtime/Engine/Private/Streaming/Texture2DStreamOut_AsyncCreate.cpp
@@ -0,0 +1,73 @@
+// Copyright Epic Games, Inc. All Rights Reserved.
+
+/*=============================================================================
+Texture2DStreamOut_AsyncCreate.cpp: Definitions of classes used for texture.
+=============================================================================*/
+
+#include "Streaming/Texture2DStreamOut_AsyncCreate.h"
+#include "RenderUtils.h"
+
+
+// ****************************
+// ******* Update Steps *******
+// ****************************
+
+FTexture2DStreamOut_AsyncCreate::FTexture2DStreamOut_AsyncCreate(UTexture2D* InTexture)
+	: FTexture2DUpdate(InTexture)
+{
+	if (!ensure(ResourceState.NumRequestedLODs < ResourceState.NumResidentLODs))
+	{
+		bIsCancelled = true;
+	}
+
+	PushTask(FContext(InTexture, TT_None), TT_Async, SRA_UPDATE_CALLBACK(AsyncCreate), TT_None, nullptr);
+}
+
+void FTexture2DStreamOut_AsyncCreate::AsyncCreate(const FContext& Context)
+{
+	DECLARE_SCOPE_CYCLE_COUNTER(TEXT("FTexture2DStreamOut_AsyncReallocate::AsyncReallocate"), STAT_Texture2DStreamOutAsyncCreate_AsyncReallocate, STATGROUP_StreamingDetails);
+	check(Context.CurrentThread == TT_Async);
+
+	if (!IsCancelled() && Context.Resource && Context.Texture)
+	{
+		const FTexture2DMipMap& RequestedMipMap = *Context.MipsView[PendingFirstLODIdx];
+
+		ensure(IntermediateTextureRHI == nullptr);
+		IntermediateTextureRHI = RHIAsyncCreateTexture2D(
+			RequestedMipMap.SizeX,
+			RequestedMipMap.SizeY,
+			Context.Resource->GetPixelFormat(),
+			ResourceState.NumRequestedLODs,
+			Context.Resource->GetCreationFlags(),
+			nullptr, 0);
+	}
+
+	PushTask(Context, TT_Render, SRA_UPDATE_CALLBACK(Finalize), TT_Render, SRA_UPDATE_CALLBACK(Cancel));
+}
+
+void FTexture2DStreamOut_AsyncCreate::Finalize(const FContext& Context)
+{
+	DECLARE_SCOPE_CYCLE_COUNTER(TEXT("FTexture2DStreamOut_AsyncReallocate::Finalize"), STAT_Texture2DStreamOutAsyncCreate_Finalize, STATGROUP_StreamingDetails);
+	check(Context.CurrentThread == TT_Render);
+
+	if (!IsCancelled() && IntermediateTextureRHI && Context.Resource)
+	{
+		RHICopySharedMips(IntermediateTextureRHI, Context.Resource->GetTexture2DRHI());
+	}
+
+	RHIFinalizeAsyncReallocateTexture2D(IntermediateTextureRHI, true);
+	DoFinishUpdate(Context);
+}
+
+// ****************************
+// ******* Cancel Steps *******
+// ****************************
+
+
+void FTexture2DStreamOut_AsyncCreate::Cancel(const FContext& Context)
+{
+	DECLARE_SCOPE_CYCLE_COUNTER(TEXT("FTexture2DStreamOut_AsyncCreate::Cancel"), STAT_Texture2DStreamOutAsyncCreate_Cancel, STATGROUP_StreamingDetails);
+	check(Context.CurrentThread == TT_Render);
+
+	DoFinishUpdate(Context);
+}
diff --git a/Engine/Source/Runtime/Engine/Private/Streaming/Texture2DStreamOut_AsyncCreate.h b/Engine/Source/Runtime/Engine/Private/Streaming/Texture2DStreamOut_AsyncCreate.h
new file mode 100644
--- /dev/null
+++ b/Engine/Source/Runtime/Engine/Private/Streaming/Texture2DStreamOut_AsyncCreate.h
@@ -0,0 +1,35 @@
+// Copyright Epic Games, Inc. All Rights Reserved.
+
+/*=============================================================================
+Texture2DStreamOut_AsyncReallocate.h: Stream out logic for texture 2D.
+=============================================================================*/
+
+#pragma once
+
+#include "CoreMinimal.h"
+#include "Texture2DUpdate.h"
+
+class FTexture2DStreamOut_AsyncCreate : public FTexture2DUpdate
+{
+public:
+
+	FTexture2DStreamOut_AsyncCreate(UTexture2D* InTexture);
+
+protected:
+
+	// ****************************
+	// ******* Update Steps *******
+	// ****************************
+
+	// Create the new texture on the async thread. (AsyncThread)
+	void AsyncCreate(const FContext& Context);
+	// Apply the intermediate texture and cleanup. (RenderThread)
+	void Finalize(const FContext& Context);
+
+	// ****************************
+	// ******* Cancel Steps *******
+	// ****************************
+
+	// Cancel the update. (RenderThread)
+	void Cancel(const FContext& Context);
+};
diff --git a/Engine/Source/Runtime/Engine/Private/Texture2D.cpp b/Engine/Source/Runtime/Engine/Private/Texture2D.cpp
--- a/Engine/Source/Runtime/Engine/Private/Texture2D.cpp
+++ b/Engine/Source/Runtime/Engine/Private/Texture2D.cpp
@@ -27,6 +27,7 @@
 #include "DerivedDataRequestOwner.h"
 #include "Engine/TextureStreamingTypes.h"
 #include "Streaming/TextureStreamingHelpers.h"
+#include "Streaming/Texture2DStreamOut_AsyncCreate.h"
 #include "Streaming/Texture2DStreamOut_AsyncReallocate.h"
 #include "Streaming/Texture2DStreamOut_Virtual.h"
 #include "Streaming/Texture2DStreamIn_DDC_AsyncCreate.h"
@@ -1746,6 +1747,11 @@ bool UTexture2D::StreamOut(int32 NewMipCount)
 		{
 			PendingUpdate = new FTexture2DStreamOut_Virtual(this);
 		}
+		// If the platform supports creating the new texture on an async thread, use that path.
+		else if (GRHISupportAsyncTextureStreamOut)
+		{
+			PendingUpdate = new FTexture2DStreamOut_AsyncCreate(this);
+		}
 		else
 		{
 			PendingUpdate = new FTexture2DStreamOut_AsyncReallocate(this);
diff --git a/Engine/Source/Runtime/RHI/Private/RHI.cpp b/Engine/Source/Runtime/RHI/Private/RHI.cpp
--- a/Engine/Source/Runtime/RHI/Private/RHI.cpp
+++ b/Engine/Source/Runtime/RHI/Private/RHI.cpp
@@ -1554,6 +1554,7 @@ bool GRHISupportsDirectGPUMemoryLock = false;
 bool GRHISupportsMultithreadedShaderCreation = true;
 bool GRHISupportsMultithreadedResources = false;
 
+bool GRHISupportAsyncTextureStreamOut = false;
 bool GRHISupportsMSAADepthSampleAccess = false;
 
 bool GRHISupportsBackBufferWithCustomDepthStencil = true;
diff --git a/Engine/Source/Runtime/RHI/Public/RHI.h b/Engine/Source/Runtime/RHI/Public/RHI.h
--- a/Engine/Source/Runtime/RHI/Public/RHI.h
+++ b/Engine/Source/Runtime/RHI/Public/RHI.h
@@ -550,6 +550,9 @@ extern RHI_API bool GRHISupportsRawViewsForAnyBuffer;
 /** Whether depth or stencil can individually be set to CopySrc/Dest access. */
 extern RHI_API bool GRHISupportsSeparateDepthStencilCopyAccess;
 
+/** Support using async thread for texture stream out operations */
+extern RHI_API bool GRHISupportAsyncTextureStreamOut;
+
 FORCEINLINE uint64 GetMaxBufferDimension()
 {
 	return GMaxBufferDimensions;
