Clear linked bunch Next pointer when releasing incoming bunches from InRec
Prevents potential freed memory access with InPartialBunch

diff --git a/Engine/Source/Runtime/Engine/Private/DataChannel.cpp b/Engine/Source/Runtime/Engine/Private/DataChannel.cpp
--- a/Engine/Source/Runtime/Engine/Private/DataChannel.cpp
+++ b/Engine/Source/Runtime/Engine/Private/DataChannel.cpp
@@ -698,6 +698,9 @@ void UChannel::ReceivedRawBunch( FInBunch & Bunch, bool & bOutSkipAck )
 			FInBunch* Release = InRec;
 			InRec = InRec->Next;
 			NumInRec--;
+
+			// Removed from InRec, clear pointer to the rest of the chain before processing
+			Release->Next = nullptr;
 			
 			// Just keep a local copy of the bSkipAck flag, since these have already been acked and it doesn't make sense on this context
 			// Definitely want to warn when this happens, since it's really not possible
