Fix CIS error about struct initalizer list

#rnx

diff --git a/Engine/Source/Runtime/Engine/Classes/GameFramework/InputSettings.h b/Engine/Source/Runtime/Engine/Classes/GameFramework/InputSettings.h
--- a/Engine/Source/Runtime/Engine/Classes/GameFramework/InputSettings.h
+++ b/Engine/Source/Runtime/Engine/Classes/GameFramework/InputSettings.h
@@ -292,6 +292,7 @@ struct FHardwareDeviceIdentifier
 	GENERATED_BODY()
 
 	FHardwareDeviceIdentifier() = default;
+	FHardwareDeviceIdentifier(FName InClassName, FString InHardwareDeviceIdentifier);
 	
 	/** 
 	* The name of the Input Class that uses this hardware device.
diff --git a/Engine/Source/Runtime/Engine/Private/UserInterface/InputSettings.cpp b/Engine/Source/Runtime/Engine/Private/UserInterface/InputSettings.cpp
--- a/Engine/Source/Runtime/Engine/Private/UserInterface/InputSettings.cpp
+++ b/Engine/Source/Runtime/Engine/Private/UserInterface/InputSettings.cpp
@@ -494,6 +494,13 @@ void UInputSettings::SetDefaultInputComponentClass(TSubclassOf<UInputComponent>
 
 FHardwareDeviceIdentifier FHardwareDeviceIdentifier::Invalid = { NAME_None, TEXT("Invalid") };
 
+FHardwareDeviceIdentifier::FHardwareDeviceIdentifier(FName InClassName, FString InHardwareDeviceIdentifier)
+	: InputClassName(InClassName)
+	, HardwareDeviceIdentifier(InHardwareDeviceIdentifier)
+{
+	
+}
+
 bool FHardwareDeviceIdentifier::IsValid() const
 {
 	return InputClassName.IsValid() && !HardwareDeviceIdentifier.IsEmpty();
