Fixup serialization in Postload because UMaterialInterfaceEditorOnlyData is serialized independently.

diff --git a/Engine/Source/Runtime/Engine/Private/Materials/MaterialInstance.cpp b/Engine/Source/Runtime/Engine/Private/Materials/MaterialInstance.cpp
--- a/Engine/Source/Runtime/Engine/Private/Materials/MaterialInstance.cpp
+++ b/Engine/Source/Runtime/Engine/Private/Materials/MaterialInstance.cpp
@@ -2805,13 +2805,6 @@ void UMaterialInstance::Serialize(FArchive& Ar)
 	{
 		ValidateTextureOverrides(GMaxRHIFeatureLevel);
 	}
-
-#if WITH_EDITORONLY_DATA
-	if (Ar.IsLoading() && !GetEditorOnlyData()->StaticParameters.StaticSwitchParameters_DEPRECATED.IsEmpty())
-	{
-		StaticParametersRuntime.StaticSwitchParameters = MoveTemp(GetEditorOnlyData()->StaticParameters.StaticSwitchParameters_DEPRECATED);
-	}
-#endif
 }
 
 void UMaterialInstance::PostLoad()
@@ -2819,6 +2812,14 @@ void UMaterialInstance::PostLoad()
 	LLM_SCOPE(ELLMTag::MaterialInstance);
 	SCOPED_LOADTIMER(MaterialInstancePostLoad);
 
+#if WITH_EDITORONLY_DATA // fixup serialization before everything else
+	if (IsEditorOnlyDataValid() && !GetEditorOnlyData()->StaticParameters.StaticSwitchParameters_DEPRECATED.IsEmpty())
+	{
+		ensure(StaticParametersRuntime.StaticSwitchParameters.IsEmpty());
+		StaticParametersRuntime.StaticSwitchParameters = MoveTemp(GetEditorOnlyData()->StaticParameters.StaticSwitchParameters_DEPRECATED);
+	}
+#endif
+
 	Super::PostLoad();
 
 #if WITH_EDITOR
