Include SCS modified properties in the transaction diff

These need to be tracked as if we're syncing state over the network using this diff info, then there's no guarantee that the receiving side will actually re-run its construction script (assuming that the construction script changes are even deterministic).

This can be disabled again via the Actor.IncludeSCSModifiedPropertiesInDiff CVar.

diff --git a/Engine/Source/Runtime/Engine/Private/ActorEditor.cpp b/Engine/Source/Runtime/Engine/Private/ActorEditor.cpp
--- a/Engine/Source/Runtime/Engine/Private/ActorEditor.cpp
+++ b/Engine/Source/Runtime/Engine/Private/ActorEditor.cpp
@@ -32,6 +32,7 @@
 #if WITH_EDITOR
 
 #include "Editor.h"
+#include "HAL/IConsoleManager.h"
 #include "Misc/TransactionObjectEvent.h"
 #include "ActorTransactionAnnotation.h"
 #include "Engine/LevelStreaming.h"
@@ -47,6 +48,14 @@
 
 #define LOCTEXT_NAMESPACE "ErrorChecking"
 
+namespace ActorEditorSettings
+{
+
+static bool bIncludeSCSModifiedPropertiesInDiff = true;
+static FAutoConsoleVariableRef CVarIncludeSCSModifiedPropertiesInDiff(TEXT("Actor.IncludeSCSModifiedPropertiesInDiff"), bIncludeSCSModifiedPropertiesInDiff, TEXT("True to include SCS modified properties in any transaction diffs, or False to skip them"));
+
+}
+
 void AActor::PreEditChange(FProperty* PropertyThatWillChange)
 {
 	Super::PreEditChange(PropertyThatWillChange);
@@ -523,18 +532,21 @@ void FActorTransactionAnnotation::ComputeAdditionalObjectChanges(const ITransact
 			{
 				TSet<const FProperty*> PropertiesToSkip;
 
-				// Skip properties modified during construction when calculating the diff
-				CurrentComponent->GetUCSModifiedProperties(PropertiesToSkip);
-
-				// If this is the owning Actor's root scene component, always include relative transform properties as GetUCSModifiedProperties incorrectly considers them modified (due to changing during placement)
-				if (CurrentComponent->IsA<USceneComponent>())
+				if (!ActorEditorSettings::bIncludeSCSModifiedPropertiesInDiff)
 				{
-					const AActor* ComponentOwner = CurrentComponent->GetOwner();
-					if (ComponentOwner && ComponentOwner->GetRootComponent() == CurrentComponent)
+					// Skip properties modified during construction when calculating the diff
+					CurrentComponent->GetUCSModifiedProperties(PropertiesToSkip);
+
+					// If this is the owning Actor's root scene component, always include relative transform properties as GetUCSModifiedProperties incorrectly considers them modified (due to changing during placement)
+					if (CurrentComponent->IsA<USceneComponent>())
 					{
-						PropertiesToSkip.Remove(FindFProperty<FProperty>(USceneComponent::StaticClass(), USceneComponent::GetRelativeLocationPropertyName()));
-						PropertiesToSkip.Remove(FindFProperty<FProperty>(USceneComponent::StaticClass(), USceneComponent::GetRelativeRotationPropertyName()));
-						PropertiesToSkip.Remove(FindFProperty<FProperty>(USceneComponent::StaticClass(), USceneComponent::GetRelativeScale3DPropertyName()));
+						const AActor* ComponentOwner = CurrentComponent->GetOwner();
+						if (ComponentOwner && ComponentOwner->GetRootComponent() == CurrentComponent)
+						{
+							PropertiesToSkip.Remove(FindFProperty<FProperty>(USceneComponent::StaticClass(), USceneComponent::GetRelativeLocationPropertyName()));
+							PropertiesToSkip.Remove(FindFProperty<FProperty>(USceneComponent::StaticClass(), USceneComponent::GetRelativeRotationPropertyName()));
+							PropertiesToSkip.Remove(FindFProperty<FProperty>(USceneComponent::StaticClass(), USceneComponent::GetRelativeScale3DPropertyName()));
+						}
 					}
 				}
 
