Add CVar to take a GPU capture of hit proxy rendering upon next click.

diff --git a/Engine/Source/Runtime/Engine/Private/UnrealClient.cpp b/Engine/Source/Runtime/Engine/Private/UnrealClient.cpp
--- a/Engine/Source/Runtime/Engine/Private/UnrealClient.cpp
+++ b/Engine/Source/Runtime/Engine/Private/UnrealClient.cpp
@@ -32,11 +32,18 @@
 #include "EngineUtils.h"
 #include "RenderGraphUtils.h"
 #include "DynamicResolutionState.h"
+#include "RenderCaptureInterface.h"
 
 DEFINE_LOG_CATEGORY_STATIC(LogClient, Log, All);
 
 UE_IMPLEMENT_STRUCT("/Script/Engine", PostProcessSettings);
 
+static int32 GHitProxyCaptureNextUpdate = 0;
+static FAutoConsoleVariableRef CVarHitProxyCaptureEnable(
+	TEXT("r.HitProxy.CaptureNextUpdate"),
+	GHitProxyCaptureNextUpdate,
+	TEXT("Enables GPU capture of hit proxy rendering on the next update."));
+
 bool FViewport::bIsGameRenderingEnabled = true;
 int32 FViewport::PresentAndStopMovieDelay = 0;
 
@@ -1741,6 +1748,9 @@ const TArray<FColor>& FViewport::GetRawHitProxyData(FIntRect InRect)
 	// If the hit proxy map isn't up to date, render the viewport client's hit proxies to it.
 	else if (!bHitProxiesCached)
 	{
+		RenderCaptureInterface::FScopedCapture RenderCapture(GHitProxyCaptureNextUpdate != 0, TEXT("Update Hit Proxies"));
+		GHitProxyCaptureNextUpdate = 0;
+
 		EnqueueBeginRenderFrame(false);
 
 		FViewport* Viewport = this;
