Implements sg.Test.{CPU,GPU}PerfIndexOverride for AutoSet scalability testing by QA

diff --git a/Engine/Source/Runtime/Engine/Private/Scalability.cpp b/Engine/Source/Runtime/Engine/Private/Scalability.cpp
--- a/Engine/Source/Runtime/Engine/Private/Scalability.cpp
+++ b/Engine/Source/Runtime/Engine/Private/Scalability.cpp
@@ -13,6 +13,20 @@
 
 Scalability::FOnScalabilitySettingsChanged Scalability::OnScalabilitySettingsChanged;
 
+#if !UE_BUILD_SHIPPING
+
+static TAutoConsoleVariable<float> CVarTestCPUPerfIndexOverride(
+	TEXT("sg.Test.CPUPerfIndexOverride"), 0.0f,
+	TEXT("Custom override for the CPU perf index returned by the GPU benchmark."),
+	ECVF_Default);
+
+static TAutoConsoleVariable<float> CVarTestGPUPerfIndexOverride(
+	TEXT("sg.Test.GPUPerfIndexOverride"), 0.0f,
+	TEXT("Custom override for the GPU perf index returned by the GPU benchmark."),
+	ECVF_Default);
+
+#endif
+
 static TAutoConsoleVariable<float> CVarResolutionQuality(
 	TEXT("sg.ResolutionQuality"),
 	100.0f,
@@ -583,8 +597,19 @@ FQualityLevels BenchmarkQualityLevels(uint32 WorkScale, float CPUMultiplier, flo
 	FSynthBenchmarkResults SynthBenchmark;
 	ISynthBenchmark::Get().Run(SynthBenchmark, true, WorkScale);
 
-	const float CPUPerfIndex = SynthBenchmark.ComputeCPUPerfIndex(/*out*/ &Results.CPUBenchmarkSteps) * CPUMultiplier;
-	const float GPUPerfIndex = SynthBenchmark.ComputeGPUPerfIndex(/*out*/ &Results.GPUBenchmarkSteps) * GPUMultiplier;
+	float CPUPerfIndex = SynthBenchmark.ComputeCPUPerfIndex(/*out*/ &Results.CPUBenchmarkSteps) * CPUMultiplier;
+	float GPUPerfIndex = SynthBenchmark.ComputeGPUPerfIndex(/*out*/ &Results.GPUBenchmarkSteps) * GPUMultiplier;
+
+#if !UE_BUILD_SHIPPING
+	if (CVarTestCPUPerfIndexOverride.GetValueOnAnyThread() > 0.0f)
+	{
+		CPUPerfIndex = CVarTestCPUPerfIndexOverride.GetValueOnAnyThread();
+	}
+	if (CVarTestGPUPerfIndexOverride.GetValueOnAnyThread() > 0.0f)
+	{
+		GPUPerfIndex = CVarTestGPUPerfIndexOverride.GetValueOnAnyThread();
+	}
+#endif
 
 	// decide on the actual quality needed
 	Results.ResolutionQuality = GetResolutionQualityFromGPUPerfIndex(GPUPerfIndex);
