Add per-platform overrides to Force Feedback effects

#rnx

diff --git a/Engine/Source/Runtime/Engine/Classes/Components/ForceFeedbackComponent.h b/Engine/Source/Runtime/Engine/Classes/Components/ForceFeedbackComponent.h
--- a/Engine/Source/Runtime/Engine/Classes/Components/ForceFeedbackComponent.h
+++ b/Engine/Source/Runtime/Engine/Classes/Components/ForceFeedbackComponent.h
@@ -28,9 +28,9 @@ public:
 	void AddActiveComponent(UForceFeedbackComponent* ForceFeedbackComponent);
 	void RemoveActiveComponent(UForceFeedbackComponent* ForceFeedbackComponent);
 
-	void Update(FVector Location, FForceFeedbackValues& Values) const;
+	void Update(FVector Location, FForceFeedbackValues& Values, const FPlatformUserId UserId) const;
 
-	void DrawDebug(const FVector Location, FDisplayDebugManager& DisplayDebugManager) const;
+	void DrawDebug(const FVector Location, FDisplayDebugManager& DisplayDebugManager, const FPlatformUserId UserId) const;
 
 private:
 	virtual void AddReferencedObjects( FReferenceCollector& Collector ) override;
@@ -162,7 +162,7 @@ private:
 #endif
 
 	bool Advance(float DeltaTime);
-	void Update(FVector Location, FForceFeedbackValues& Values) const;
+	void Update(FVector Location, FForceFeedbackValues& Values, const FPlatformUserId UserId) const;
 	void StopInternal(bool bRemoveFromManager = true);
 
 	friend FForceFeedbackManager;
diff --git a/Engine/Source/Runtime/Engine/Classes/GameFramework/ForceFeedbackEffect.h b/Engine/Source/Runtime/Engine/Classes/GameFramework/ForceFeedbackEffect.h
--- a/Engine/Source/Runtime/Engine/Classes/GameFramework/ForceFeedbackEffect.h
+++ b/Engine/Source/Runtime/Engine/Classes/GameFramework/ForceFeedbackEffect.h
@@ -80,6 +80,18 @@ struct ENGINE_API FActiveForceFeedbackEffect
 	void GetValues(FForceFeedbackValues& Values) const;
 };
 
+/** A wrapper struct for setting channel details on a per-platform basis */
+USTRUCT()
+struct FForceFeedbackEffectOverridenChannelDetails
+{
+	GENERATED_BODY()
+
+	FForceFeedbackEffectOverridenChannelDetails();
+
+	UPROPERTY(EditAnywhere, Category="ForceFeedbackEffect")
+	TArray<FForceFeedbackChannelDetails> ChannelDetails;
+};
+
 /**
  * A predefined force-feedback effect to be played on a controller
  */
@@ -91,6 +103,10 @@ class UForceFeedbackEffect : public UObject
 	UPROPERTY(EditAnywhere, Category="ForceFeedbackEffect")
 	TArray<FForceFeedbackChannelDetails> ChannelDetails;
 
+	/** A map of platform name -> ForceFeedback channel details */
+	UPROPERTY(EditAnywhere, Category = "ForceFeedbackEffect", Meta = (GetOptions = "Engine.InputPlatformSettings.GetAllHardwareDeviceNames"))
+	TMap<FString, FForceFeedbackEffectOverridenChannelDetails> PerDeviceOverides;
+
 	/** A map of input device properties that we want to set while this effect is playing */
 	UPROPERTY(EditAnywhere, Instanced, Category = "ForceFeedbackEffect")
 	TArray<TObjectPtr<UInputDeviceProperty>> DeviceProperties;
@@ -108,10 +124,17 @@ class UForceFeedbackEffect : public UObject
 	/** Returns the longest duration of any active UInputDeviceProperty's that this effect has on it. */
 	float GetTotalDevicePropertyDuration();
 
-	void GetValues(const float EvalTime, FForceFeedbackValues& Values, float ValueMultiplier = 1.f) const;
+	void GetValues(const float EvalTime, FForceFeedbackValues& Values, const FPlatformUserId PlatformUser, float ValueMultiplier = 1.f) const;
 
 	void SetDeviceProperties(const FPlatformUserId PlatformUser, const float DeltaTime, const float EvalTime);
 
 	/** Reset any device properties that may need to be after the duration of this effect has ended. */
 	void ResetDeviceProperties(const FPlatformUserId PlatformUser);
+
+	/**
+	 * Returns the channel details that should currently be used for the given platform.
+	 * This will return one of the OverridenDetails if the current platform has one specified.
+	 * If there isn't an override specified, then it will return the generic "ChannelDetails" array.
+	 **/
+	const TArray<FForceFeedbackChannelDetails>& GetCurrentChannelDetails(const FPlatformUserId PlatformUser) const;
 };
diff --git a/Engine/Source/Runtime/Engine/Private/Components/ForceFeedbackComponent.cpp b/Engine/Source/Runtime/Engine/Private/Components/ForceFeedbackComponent.cpp
--- a/Engine/Source/Runtime/Engine/Private/Components/ForceFeedbackComponent.cpp
+++ b/Engine/Source/Runtime/Engine/Private/Components/ForceFeedbackComponent.cpp
@@ -106,25 +106,25 @@ void FForceFeedbackManager::Tick(float DeltaTime)
 	}
 }
 
-void FForceFeedbackManager::Update(const FVector Location, FForceFeedbackValues& Values) const
+void FForceFeedbackManager::Update(const FVector Location, FForceFeedbackValues& Values, const FPlatformUserId UserId) const
 {
 	for (UForceFeedbackComponent* FFC : ActiveForceFeedbackComponents)
 	{
 		if (FFC)
 		{
-			FFC->Update(Location, Values);
+			FFC->Update(Location, Values, UserId);
 		}
 	}
 }
 
-void FForceFeedbackManager::DrawDebug(const FVector Location, FDisplayDebugManager& DisplayDebugManager) const
+void FForceFeedbackManager::DrawDebug(const FVector Location, FDisplayDebugManager& DisplayDebugManager, const FPlatformUserId UserId) const
 {
 	for (UForceFeedbackComponent* FFC : ActiveForceFeedbackComponents)
 	{
 		if (FFC && FFC->ForceFeedbackEffect)
 		{
 			FForceFeedbackValues ActiveValues;
-			FFC->Update(Location, ActiveValues);
+			FFC->Update(Location, ActiveValues, UserId);
 
 			const FString ActiveEntry = FString::Printf(TEXT("%s %s %.2f %.2f %s %.2f - LL: %.2f LS: %.2f RL: %.2f RS: %.2f"), 
 				*FFC->ForceFeedbackEffect->GetFName().ToString(), 
@@ -373,7 +373,7 @@ bool UForceFeedbackComponent::Advance(const float DeltaTime)
 	return true;
 }
 
-void UForceFeedbackComponent::Update(FVector Location, FForceFeedbackValues& Values) const
+void UForceFeedbackComponent::Update(FVector Location, FForceFeedbackValues& Values, const FPlatformUserId UserId) const
 {
 	if (ForceFeedbackEffect)
 	{
@@ -392,7 +392,7 @@ void UForceFeedbackComponent::Update(FVector Location, FForceFeedbackValues& Val
 
 		if (ValueMultiplier > 0.f)
 		{
-			ForceFeedbackEffect->GetValues(EvalTime, Values, ValueMultiplier);
+			ForceFeedbackEffect->GetValues(EvalTime, Values, UserId, ValueMultiplier);
 		}
 	}
 }
diff --git a/Engine/Source/Runtime/Engine/Private/GameFramework/ForceFeedbackEffect.cpp b/Engine/Source/Runtime/Engine/Private/GameFramework/ForceFeedbackEffect.cpp
--- a/Engine/Source/Runtime/Engine/Private/GameFramework/ForceFeedbackEffect.cpp
+++ b/Engine/Source/Runtime/Engine/Private/GameFramework/ForceFeedbackEffect.cpp
@@ -4,6 +4,8 @@
 #include "GenericPlatform/IInputInterface.h"
 #include "Misc/App.h"
 #include "GameFramework/InputDeviceProperties.h"
+#include "GameFramework/InputDeviceSubsystem.h"
+#include "GenericPlatform/GenericPlatformInputDeviceMapper.h"
 
 #include UE_INLINE_GENERATED_CPP_BY_NAME(ForceFeedbackEffect)
 
@@ -15,6 +17,12 @@ UForceFeedbackEffect::UForceFeedbackEffect(const FObjectInitializer& ObjectIniti
 	ChannelDetails.Add(ChannelDetail);
 }
 
+FForceFeedbackEffectOverridenChannelDetails::FForceFeedbackEffectOverridenChannelDetails()
+{
+	// Add one default channel details by default
+	ChannelDetails.AddDefaulted(1);
+}
+
 #if WITH_EDITOR
 void UForceFeedbackEffect::PostEditChangeChainProperty(struct FPropertyChangedChainEvent& PropertyChangedEvent)
 {
@@ -31,10 +39,13 @@ float UForceFeedbackEffect::GetDuration()
 	{
 		Duration = 0.f;
 
+		// Just use the primary platform user when calculating duration, this won't be affected by which player the effect is for
+		const TArray<FForceFeedbackChannelDetails>& CurrentDetails = GetCurrentChannelDetails(IPlatformInputDeviceMapper::Get().GetPrimaryPlatformUser());
+
 		float MinTime, MaxTime;
-		for (int32 Index = 0; Index < ChannelDetails.Num(); ++Index)
+		for (int32 Index = 0; Index < CurrentDetails.Num(); ++Index)
 		{
-			ChannelDetails[Index].Curve.GetRichCurve()->GetTimeRange(MinTime, MaxTime);
+			CurrentDetails[Index].Curve.GetRichCurveConst()->GetTimeRange(MinTime, MaxTime);
 
 			if (MaxTime > Duration)
 			{
@@ -66,11 +77,13 @@ float UForceFeedbackEffect::GetTotalDevicePropertyDuration()
 	return LongestDuration;
 }
 
-void UForceFeedbackEffect::GetValues(const float EvalTime, FForceFeedbackValues& Values, const float ValueMultiplier) const
+void UForceFeedbackEffect::GetValues(const float EvalTime, FForceFeedbackValues& Values, const FPlatformUserId PlatformUser, const float ValueMultiplier) const
 {
-	for (int32 Index = 0; Index < ChannelDetails.Num(); ++Index)
+	const TArray<FForceFeedbackChannelDetails>& CurrentDetails = GetCurrentChannelDetails(PlatformUser);
+
+	for (int32 Index = 0; Index < CurrentDetails.Num(); ++Index)
 	{
-		const FForceFeedbackChannelDetails& Details = ChannelDetails[Index];
+		const FForceFeedbackChannelDetails& Details = CurrentDetails[Index];
 		const float Value = Details.Curve.GetRichCurveConst()->Eval(EvalTime) * ValueMultiplier;
 
 		if (Details.bAffectsLeftLarge)
@@ -122,13 +135,28 @@ void UForceFeedbackEffect::ResetDeviceProperties(const FPlatformUserId PlatformU
 	}
 }
 
+const TArray<FForceFeedbackChannelDetails>& UForceFeedbackEffect::GetCurrentChannelDetails(const FPlatformUserId PlatformUser) const
+{	
+	if (const UInputDeviceSubsystem* SubSystem = UInputDeviceSubsystem::Get())
+	{
+		FHardwareDeviceIdentifier Hardware = SubSystem->GetMostRecentlyUsedHardwareDevice(PlatformUser);
+		// Check if there are any per-input device overrides available
+		if (const FForceFeedbackEffectOverridenChannelDetails* Details = PerDeviceOverides.Find(Hardware.HardwareDeviceIdentifier))
+		{
+			return Details->ChannelDetails;
+		}
+	}
+
+	return ChannelDetails;
+}
+
 void FActiveForceFeedbackEffect::GetValues(FForceFeedbackValues& Values) const
 {
 	if (ForceFeedbackEffect)
 	{
 		const float Duration = ForceFeedbackEffect->GetDuration();
 		const float EvalTime = PlayTime - Duration * FMath::FloorToFloat(PlayTime / Duration);
-		ForceFeedbackEffect->GetValues(EvalTime, Values);
+		ForceFeedbackEffect->GetValues(EvalTime, Values, PlatformUser);
 	}
 	else
 	{
diff --git a/Engine/Source/Runtime/Engine/Private/PlayerController.cpp b/Engine/Source/Runtime/Engine/Private/PlayerController.cpp
--- a/Engine/Source/Runtime/Engine/Private/PlayerController.cpp
+++ b/Engine/Source/Runtime/Engine/Private/PlayerController.cpp
@@ -3396,7 +3396,7 @@ void APlayerController::DisplayDebug(class UCanvas* Canvas, const FDebugDisplayI
 		}
 		if (FForceFeedbackManager* FFM = FForceFeedbackManager::Get(GetWorld()))
 		{
-			FFM->DrawDebug(GetFocalLocation(), DisplayDebugManager);
+			FFM->DrawDebug(GetFocalLocation(), DisplayDebugManager, GetPlatformUserId());
 		}
 		DisplayDebugManager.DrawString(TEXT("-----------------------------------------------------"));
 #endif
@@ -4577,7 +4577,7 @@ void APlayerController::ProcessForceFeedbackAndHaptics(const float DeltaTime, co
 
 		if (FForceFeedbackManager* ForceFeedbackManager = FForceFeedbackManager::Get(World))
 		{
-			ForceFeedbackManager->Update(GetFocalLocation(), ForceFeedbackValues);
+			ForceFeedbackManager->Update(GetFocalLocation(), ForceFeedbackValues, GetPlatformUserId());
 		}
 
 		// Apply ForceFeedbackScale
