Move TimerHandle to its own file

diff --git a/Engine/Source/Runtime/Engine/Classes/Engine/EngineTypes.h b/Engine/Source/Runtime/Engine/Classes/Engine/EngineTypes.h
--- a/Engine/Source/Runtime/Engine/Classes/Engine/EngineTypes.h
+++ b/Engine/Source/Runtime/Engine/Classes/Engine/EngineTypes.h
@@ -20,6 +20,7 @@
 #include "Engine/DamageEvents.h"
 #include "Engine/ReplicatedState.h"
 #endif
+#include "Engine/TimerHandle.h"
 #include "EngineTypes.generated.h"
 
 class AActor;
@@ -2845,83 +2846,6 @@ namespace EEndPlayReason
 	};
 }
 
-DECLARE_DYNAMIC_DELEGATE(FTimerDynamicDelegate);
-
-/** Unique handle that can be used to distinguish timers that have identical delegates. */
-USTRUCT(BlueprintType)
-struct FTimerHandle
-{
-	GENERATED_BODY()
-
-	friend class FTimerManager;
-	friend struct FTimerHeapOrder;
-
-	FTimerHandle()
-	: Handle(0)
-	{
-	}
-
-	/** True if this handle was ever initialized by the timer manager */
-	bool IsValid() const
-	{
-		return Handle != 0;
-	}
-
-	/** Explicitly clear handle */
-	void Invalidate()
-	{
-		Handle = 0;
-	}
-
-	bool operator==(const FTimerHandle& Other) const
-	{
-		return Handle == Other.Handle;
-	}
-
-	bool operator!=(const FTimerHandle& Other) const
-	{
-		return Handle != Other.Handle;
-	}
-
-	FString ToString() const
-	{
-		return FString::Printf(TEXT("%llu"), Handle);
-	}
-
-private:
-	static const uint32 IndexBits        = 24;
-	static const uint32 SerialNumberBits = 40;
-
-	static_assert(IndexBits + SerialNumberBits == 64, "The space for the timer index and serial number should total 64 bits");
-
-	static const int32  MaxIndex        = (int32)1 << IndexBits;
-	static const uint64 MaxSerialNumber = (uint64)1 << SerialNumberBits;
-
-	void SetIndexAndSerialNumber(int32 Index, uint64 SerialNumber)
-	{
-		check(Index >= 0 && Index < MaxIndex);
-		check(SerialNumber < MaxSerialNumber);
-		Handle = (SerialNumber << IndexBits) | (uint64)(uint32)Index;
-	}
-
-	FORCEINLINE int32 GetIndex() const
-	{
-		return (int32)(Handle & (uint64)(MaxIndex - 1));
-	}
-
-	FORCEINLINE uint64 GetSerialNumber() const
-	{
-		return Handle >> IndexBits;
-	}
-
-	UPROPERTY(Transient)
-	uint64 Handle;
-
-	friend uint32 GetTypeHash(const FTimerHandle& InHandle)
-	{
-		return GetTypeHash(InHandle.Handle);
-	}
-};
 
 /**
  * Controls behavior of WalkableSlopeOverride, determining how to affect walkability of surfaces for Characters.
diff --git a/Engine/Source/Runtime/Engine/Classes/Engine/TimerHandle.h b/Engine/Source/Runtime/Engine/Classes/Engine/TimerHandle.h
new file mode 100644
--- /dev/null
+++ b/Engine/Source/Runtime/Engine/Classes/Engine/TimerHandle.h
@@ -0,0 +1,83 @@
+// Copyright Epic Games, Inc. All Rights Reserved.
+
+#pragma once
+
+#include "TimerHandle.generated.h"
+
+DECLARE_DYNAMIC_DELEGATE(FTimerDynamicDelegate);
+
+/** Unique handle that can be used to distinguish timers that have identical delegates. */
+USTRUCT(BlueprintType)
+struct FTimerHandle
+{
+	GENERATED_BODY()
+
+	friend class FTimerManager;
+	friend struct FTimerHeapOrder;
+
+	FTimerHandle()
+	: Handle(0)
+	{
+	}
+
+	/** True if this handle was ever initialized by the timer manager */
+	bool IsValid() const
+	{
+		return Handle != 0;
+	}
+
+	/** Explicitly clear handle */
+	void Invalidate()
+	{
+		Handle = 0;
+	}
+
+	bool operator==(const FTimerHandle& Other) const
+	{
+		return Handle == Other.Handle;
+	}
+
+	bool operator!=(const FTimerHandle& Other) const
+	{
+		return Handle != Other.Handle;
+	}
+
+	FString ToString() const
+	{
+		return FString::Printf(TEXT("%llu"), Handle);
+	}
+
+private:
+	static constexpr uint32 IndexBits        = 24;
+	static constexpr uint32 SerialNumberBits = 40;
+
+	static_assert(IndexBits + SerialNumberBits == 64, "The space for the timer index and serial number should total 64 bits");
+
+	static constexpr int32  MaxIndex        = (int32)1 << IndexBits;
+	static constexpr uint64 MaxSerialNumber = (uint64)1 << SerialNumberBits;
+
+	void SetIndexAndSerialNumber(int32 Index, uint64 SerialNumber)
+	{
+		check(Index >= 0 && Index < MaxIndex);
+		check(SerialNumber < MaxSerialNumber);
+		Handle = (SerialNumber << IndexBits) | (uint64)(uint32)Index;
+	}
+
+	FORCEINLINE int32 GetIndex() const
+	{
+		return (int32)(Handle & (uint64)(MaxIndex - 1));
+	}
+
+	FORCEINLINE uint64 GetSerialNumber() const
+	{
+		return Handle >> IndexBits;
+	}
+
+	UPROPERTY(Transient)
+	uint64 Handle;
+
+	friend uint32 GetTypeHash(const FTimerHandle& InHandle)
+	{
+		return GetTypeHash(InHandle.Handle);
+	}
+};
\ No newline at end of file
diff --git a/Engine/Source/Runtime/Engine/Private/Engine/TiimerHandle.cpp b/Engine/Source/Runtime/Engine/Private/Engine/TiimerHandle.cpp
new file mode 100644
--- /dev/null
+++ b/Engine/Source/Runtime/Engine/Private/Engine/TiimerHandle.cpp
@@ -0,0 +1,5 @@
+// Copyright Epic Games, Inc. All Rights Reserved.
+
+#include "Engine/TimerHandle.h"
+
+#include UE_INLINE_GENERATED_CPP_BY_NAME(TimerHandle)
\ No newline at end of file
diff --git a/Engine/Source/Runtime/Engine/Private/TimerManager.cpp b/Engine/Source/Runtime/Engine/Private/TimerManager.cpp
--- a/Engine/Source/Runtime/Engine/Private/TimerManager.cpp
+++ b/Engine/Source/Runtime/Engine/Private/TimerManager.cpp
@@ -5,6 +5,7 @@
 =============================================================================*/
 
 #include "TimerManager.h"
+#include "Engine/EngineTypes.h"
 #include "Engine/GameInstance.h"
 #include "HAL/IConsoleManager.h"
 #include "Misc/CoreDelegates.h"
@@ -12,7 +13,6 @@
 #include "UnrealEngine.h"
 #include "Misc/TimeGuard.h"
 #include "ProfilingDebugging/CsvProfiler.h"
-#include "Algo/Transform.h"
 #include "HAL/PlatformStackWalk.h"
 
 DECLARE_CYCLE_STAT(TEXT("SetTimer"), STAT_SetTimer, STATGROUP_Engine);
@@ -108,7 +108,6 @@ struct FTimerSourceList
 	{
 		FString FunctionNameStr;
 		FString ObjectNameStr;
-		bool bDynDelegate = false;
 
 		if (Delegate.FuncDelegate.IsBound())
 		{
@@ -149,7 +148,6 @@ struct FTimerSourceList
 		{
 			const FName FuncFName = Delegate.FuncDynDelegate.GetFunctionName();
 			FunctionNameStr = FuncFName.ToString();
-			bDynDelegate = true;
 
 			UClass* SourceClass = nullptr;
 			if (const UObject* Object = Delegate.FuncDynDelegate.GetUObject())
@@ -266,7 +264,7 @@ FTimerManager::FTimerManager(UGameInstance* GameInstance)
 {
 	if (IsRunningDedicatedServer())
 	{
-		// Off by default, renable if needed
+		// Off by default, reenable if needed
 		//FCoreDelegates::OnHandleSystemError.AddRaw(this, &FTimerManager::OnCrash);
 	}
 
@@ -762,6 +760,15 @@ void FTimerManager::UnPauseTimer(FTimerHandle InHandle)
 	PausedTimerSet.Remove(InHandle);
 }
 
+FTimerData::FTimerData()
+	: bLoop(false)
+	, bRequiresDelegate(false)
+	, Status(ETimerStatus::Active)
+	, Rate(0)
+	, ExpireTime(0)
+	, LevelCollection(ELevelCollectionType::DynamicSourceLevels)
+{}
+
 // ---------------------------------
 // Public members
 // ---------------------------------
diff --git a/Engine/Source/Runtime/Engine/Public/TimerManager.h b/Engine/Source/Runtime/Engine/Public/TimerManager.h
--- a/Engine/Source/Runtime/Engine/Public/TimerManager.h
+++ b/Engine/Source/Runtime/Engine/Public/TimerManager.h
@@ -9,9 +9,14 @@
 #include "CoreMinimal.h"
 #include "Stats/Stats.h"
 #include "UObject/Object.h"
+#if UE_ENABLE_INCLUDE_ORDER_DEPRECATED_IN_5_2
 #include "Engine/EngineTypes.h"
+#endif
+#include "Engine/TimerHandle.h"
+#include "Templates/Function.h"
 
 class UGameInstance;
+enum class ELevelCollectionType : uint8;
 struct FTimerSourceList;
 
 DECLARE_DELEGATE(FTimerDelegate);
@@ -127,19 +132,12 @@ struct FTimerData
 	FTimerHandle Handle;
 
 	/** This is the key to the TimerIndicesByObject map - this is kept so that we can look up even if the referenced object is expired */
-	const void* TimerIndicesByObjectKey;
+	const void* TimerIndicesByObjectKey = nullptr;
 
 	/** The level collection that was active when this timer was created. Used to set the correct context before executing the timer's delegate. */
 	ELevelCollectionType LevelCollection;
 
-	FTimerData()
-		: bLoop(false)
-		, bRequiresDelegate(false)
-		, Status(ETimerStatus::Active)
-		, Rate(0)
-		, ExpireTime(0)
-		, LevelCollection(ELevelCollectionType::DynamicSourceLevels)
-	{}
+	ENGINE_API FTimerData();
 
 	// Movable only
 	FTimerData(FTimerData&&) = default;
