Move RenderedTextureStats to its own file

diff --git a/Engine/Source/Runtime/Engine/Public/ContentStreaming.h b/Engine/Source/Runtime/Engine/Public/ContentStreaming.h
--- a/Engine/Source/Runtime/Engine/Public/ContentStreaming.h
+++ b/Engine/Source/Runtime/Engine/Public/ContentStreaming.h
@@ -9,6 +9,7 @@
 #include "Async/TaskGraphInterfaces.h"
 #include "CoreMinimal.h"
 #include "UObject/WeakObjectPtr.h"
+#include "RenderedTextureStats.h"
 #if UE_ENABLE_INCLUDE_ORDER_DEPRECATED_IN_5_2
 #include "CanvasTypes.h"
 #include "UnrealClient.h"
@@ -419,15 +420,6 @@ protected:
 	int32		NumWantingResourcesCounter;
 };
 
-/**
- * Lightweight struct used to list the MIP levels of rendered assets.
- */
-struct FRenderedTextureStats
-{
-	int32 MaxMipLevelShown;
-	FString TextureGroup;
-};
-
 /**
  * Interface to add functions specifically related to texture/mesh streaming
  */
diff --git a/Engine/Source/Runtime/Engine/Public/RenderedTextureStats.h b/Engine/Source/Runtime/Engine/Public/RenderedTextureStats.h
new file mode 100644
--- /dev/null
+++ b/Engine/Source/Runtime/Engine/Public/RenderedTextureStats.h
@@ -0,0 +1,14 @@
+// Copyright Epic Games, Inc. All Rights Reserved.
+
+#pragma once
+
+#include "CoreMinimal.h"
+
+/**
+ * Lightweight struct used to list the MIP levels of rendered assets.
+ */
+struct FRenderedTextureStats
+{
+	int32 MaxMipLevelShown;
+	FString TextureGroup;
+};
