Protect against a null source object when handling re-instancing in custom property nodes

diff --git a/Engine/Source/Runtime/Engine/Private/Animation/AnimNode_CustomProperty.cpp b/Engine/Source/Runtime/Engine/Private/Animation/AnimNode_CustomProperty.cpp
--- a/Engine/Source/Runtime/Engine/Private/Animation/AnimNode_CustomProperty.cpp
+++ b/Engine/Source/Runtime/Engine/Private/Animation/AnimNode_CustomProperty.cpp
@@ -103,7 +103,10 @@ void FAnimNode_CustomProperty::HandleObjectsReinstanced_Impl(UObject* InSourceOb
 	static IConsoleVariable* UseLegacyAnimInstanceReinstancingBehavior = IConsoleManager::Get().FindConsoleVariable(TEXT("bp.UseLegacyAnimInstanceReinstancingBehavior"));
 	if(UseLegacyAnimInstanceReinstancingBehavior == nullptr || !UseLegacyAnimInstanceReinstancingBehavior->GetBool())
 	{
-		InitializeProperties(CastChecked<UAnimInstance>(InSourceObject), GetTargetClass());
+		if(InSourceObject)
+		{
+			InitializeProperties(CastChecked<UAnimInstance>(InSourceObject), GetTargetClass());
+		}
 	}
 }
 
