Fix for keytimes in CurveTables not supporting decimals when using WriteTableAsJSON

diff --git a/Engine/Source/Runtime/Engine/Private/CurveTable.cpp b/Engine/Source/Runtime/Engine/Private/CurveTable.cpp
--- a/Engine/Source/Runtime/Engine/Private/CurveTable.cpp
+++ b/Engine/Source/Runtime/Engine/Private/CurveTable.cpp
@@ -487,7 +487,7 @@ void WriteTableAsJSON_Internal(const TMap<FName, T*>& RowMap, const TSharedRef<
 		auto LongIt(Curves[LongestCurveIndex]->GetKeyIterator());
 		for (auto It(Curves[CurvesIdx]->GetKeyIterator()); It; ++It)
 		{
-			JsonWriter->WriteValue(FString::Printf(TEXT("%d"), (int32)LongIt->Time), It->Value);
+			JsonWriter->WriteValue(FString::SanitizeFloat(LongIt->Time, 0), It->Value);
 			++LongIt;
 		}
 		JsonWriter->WriteObjectEnd();
