Fix for ServerWorldTimeSecondsDelta being float when it should be double. I have had to break compilation for any licensee code overiding virtual float GetServerWorldTimeSeconds() as it seems like the only sensible option here (typical virtual function deprecation issues).

diff --git a/Engine/Source/Runtime/Engine/Classes/Engine/DemoNetDriver.h b/Engine/Source/Runtime/Engine/Classes/Engine/DemoNetDriver.h
--- a/Engine/Source/Runtime/Engine/Classes/Engine/DemoNetDriver.h
+++ b/Engine/Source/Runtime/Engine/Classes/Engine/DemoNetDriver.h
@@ -282,7 +282,7 @@ private:
 	FOnGotoTimeDelegate OnGotoTimeDelegate_Transient;
 	
 	/** Saved server time after loading a checkpoint, so that we can set the server time as accurately as possible after the fast-forward */
-	float SavedReplicatedWorldTimeSeconds;
+	double SavedReplicatedWorldTimeSeconds;
 
 	/** Saved fast-forward time, used for correcting world time after the fast-forward is complete */
 	float SavedSecondsToSkip;
diff --git a/Engine/Source/Runtime/Engine/Classes/GameFramework/GameStateBase.h b/Engine/Source/Runtime/Engine/Classes/GameFramework/GameStateBase.h
--- a/Engine/Source/Runtime/Engine/Classes/GameFramework/GameStateBase.h
+++ b/Engine/Source/Runtime/Engine/Classes/GameFramework/GameStateBase.h
@@ -69,7 +69,7 @@ public:
 
 	/** Returns the simulated TimeSeconds on the server, will be synchronized on client and server */
 	UFUNCTION(BlueprintCallable, Category=GameState)
-	virtual float GetServerWorldTimeSeconds() const;
+	virtual double GetServerWorldTimeSeconds() const;
 
 	/** Returns true if the world has started play (called BeginPlay on actors) */
 	UFUNCTION(BlueprintCallable, Category=GameState)
@@ -148,7 +148,7 @@ protected:
 
 	/** Server TimeSeconds. Useful for syncing up animation and gameplay. */
 	UPROPERTY(Transient, ReplicatedUsing=OnRep_ReplicatedWorldTimeSeconds)
-	float ReplicatedWorldTimeSeconds;
+	double ReplicatedWorldTimeSeconds;
 
 	/** The difference from the local world's TimeSeconds and the server world's TimeSeconds. */
 	UPROPERTY(Transient)
diff --git a/Engine/Source/Runtime/Engine/Private/DemoNetDriver.cpp b/Engine/Source/Runtime/Engine/Private/DemoNetDriver.cpp
--- a/Engine/Source/Runtime/Engine/Private/DemoNetDriver.cpp
+++ b/Engine/Source/Runtime/Engine/Private/DemoNetDriver.cpp
@@ -830,7 +830,7 @@ bool UDemoNetDriver::InitBase(bool bInitAsClient, FNetworkNotify* InNotify, cons
 		bIsFastForwardingForCheckpoint	= false;
 		bIsRestoringStartupActors		= false;
 		bWasStartStreamingSuccessful	= true;
-		SavedReplicatedWorldTimeSeconds	= 0.0f;
+		SavedReplicatedWorldTimeSeconds	= 0.0;
 		SavedSecondsToSkip				= 0.0f;
 		MaxDesiredRecordTimeMS			= -1.0f;
 		ViewerOverride					= nullptr;
@@ -2969,7 +2969,7 @@ void UDemoNetDriver::FinalizeFastForward(const double StartTime)
 	{
 		if (bIsFastForwardingForCheckpoint)
 		{
-			const float PostCheckpointServerTime = SavedReplicatedWorldTimeSeconds + SavedSecondsToSkip;
+			const double PostCheckpointServerTime = SavedReplicatedWorldTimeSeconds + SavedSecondsToSkip;
 			GameState->ReplicatedWorldTimeSeconds = PostCheckpointServerTime;
 		}
 
diff --git a/Engine/Source/Runtime/Engine/Private/GameStateBase.cpp b/Engine/Source/Runtime/Engine/Private/GameStateBase.cpp
--- a/Engine/Source/Runtime/Engine/Private/GameStateBase.cpp
+++ b/Engine/Source/Runtime/Engine/Private/GameStateBase.cpp
@@ -138,7 +138,7 @@ void AGameStateBase::RemovePlayerState(APlayerState* PlayerState)
 	}
 }
 
-float AGameStateBase::GetServerWorldTimeSeconds() const
+double AGameStateBase::GetServerWorldTimeSeconds() const
 {
 	UWorld* World = GetWorld();
 	if (World)
@@ -146,7 +146,7 @@ float AGameStateBase::GetServerWorldTimeSeconds() const
 		return World->GetTimeSeconds() + ServerWorldTimeSecondsDelta;
 	}
 
-	return 0.f;
+	return 0.;
 }
 
 void AGameStateBase::UpdateServerTimeSeconds()
@@ -163,11 +163,11 @@ void AGameStateBase::OnRep_ReplicatedWorldTimeSeconds()
 	UWorld* World = GetWorld();
 	if (World)
 	{
-		const float ServerWorldTimeDelta = ReplicatedWorldTimeSeconds - World->GetTimeSeconds();
+		const double ServerWorldTimeDelta = ReplicatedWorldTimeSeconds - World->GetTimeSeconds();
 
 		// Accumulate the computed server world delta
 		SumServerWorldTimeSecondsDelta += ServerWorldTimeDelta;
-		NumServerWorldTimeSecondsDeltas += 1.0;
+		NumServerWorldTimeSecondsDeltas += 1;
 
 		// Reset the accumulated values to ensure that we remain representative of the current delta
 		if (NumServerWorldTimeSecondsDeltas > 250)
diff --git a/Engine/Source/Runtime/Engine/Private/ReplayHelper.cpp b/Engine/Source/Runtime/Engine/Private/ReplayHelper.cpp
--- a/Engine/Source/Runtime/Engine/Private/ReplayHelper.cpp
+++ b/Engine/Source/Runtime/Engine/Private/ReplayHelper.cpp
@@ -749,7 +749,7 @@ void FReplayHelper::ProcessCheckpointActors(UNetConnection* Connection, TArrayVi
 	// more often than the normal update rate.
 	AGameStateBase* const GameState = World != nullptr ? World->GetGameState() : nullptr;
 
-	const float SavedReplicatedServerTimeSeconds = GameState ? GameState->ReplicatedWorldTimeSeconds : -1.0f;
+	const double SavedReplicatedServerTimeSeconds = GameState ? GameState->ReplicatedWorldTimeSeconds : -1.0;
 
 	// Normally AGameStateBase::ReplicatedWorldTimeSeconds is only updated periodically,
 	// but we want to make sure it's accurate for the checkpoint.
