Make sure derived net drivers are marked in tick for the duration of TickFlush/Dispatch

diff --git a/Engine/Source/Runtime/Engine/Classes/Engine/NetDriver.h b/Engine/Source/Runtime/Engine/Classes/Engine/NetDriver.h
--- a/Engine/Source/Runtime/Engine/Classes/Engine/NetDriver.h
+++ b/Engine/Source/Runtime/Engine/Classes/Engine/NetDriver.h
@@ -1798,6 +1798,11 @@ protected:
 	/** Unregister all TickDispatch, TickFlush, PostTickFlush to tick in World */
 	ENGINE_API void UnregisterTickEvents(class UWorld* InWorld);
 
+private:
+	void InternalTickDispatch(float DeltaSeconds);
+	void InternalTickFlush(float DeltaSeconds);
+
+protected:
 	/** Subclasses may override this to customize channel creation. Called by GetOrCreateChannel if the pool is exhausted and a new channel must be allocated. */
 	ENGINE_API virtual UChannel* InternalCreateChannelByName(const FName& ChName);
 
diff --git a/Engine/Source/Runtime/Engine/Private/NetDriver.cpp b/Engine/Source/Runtime/Engine/Private/NetDriver.cpp
--- a/Engine/Source/Runtime/Engine/Private/NetDriver.cpp
+++ b/Engine/Source/Runtime/Engine/Private/NetDriver.cpp
@@ -870,8 +870,6 @@ void UNetDriver::TickFlush(float DeltaSeconds)
 {
 	LLM_SCOPE_BYTAG(NetDriver);
 
-	TGuardValue<bool> GuardInNetTick(bInTick, true);
-
 	CSV_SCOPED_TIMING_STAT_EXCLUSIVE(NetworkOutgoing);
 	SCOPE_CYCLE_COUNTER(STAT_NetTickFlush);
 	bool bEnableTimer = (NetDriverName == NAME_GameNetDriver) && ShouldEnableScopeSecondsTimers();
@@ -1541,9 +1539,9 @@ void UNetDriver::RegisterTickEvents(class UWorld* InWorld)
 {
 	if (InWorld)
 	{
-		TickDispatchDelegateHandle  = InWorld->OnTickDispatch ().AddUObject(this, &UNetDriver::TickDispatch);
+		TickDispatchDelegateHandle  = InWorld->OnTickDispatch ().AddUObject(this, &UNetDriver::InternalTickDispatch);
 		PostTickDispatchDelegateHandle	= InWorld->OnPostTickDispatch().AddUObject(this, &UNetDriver::PostTickDispatch);
-		TickFlushDelegateHandle     = InWorld->OnTickFlush    ().AddUObject(this, &UNetDriver::TickFlush);
+		TickFlushDelegateHandle     = InWorld->OnTickFlush    ().AddUObject(this, &UNetDriver::InternalTickFlush);
 		PostTickFlushDelegateHandle		= InWorld->OnPostTickFlush	 ().AddUObject(this, &UNetDriver::PostTickFlush);
 	}
 }
@@ -1559,6 +1557,18 @@ void UNetDriver::UnregisterTickEvents(class UWorld* InWorld)
 	}
 }
 
+void UNetDriver::InternalTickDispatch(float DeltaSeconds)
+{
+	TGuardValue<bool> GuardInNetTick(bInTick, true);
+	TickDispatch(DeltaSeconds);
+}
+
+void UNetDriver::InternalTickFlush(float DeltaSeconds)
+{
+	TGuardValue<bool> GuardInNetTick(bInTick, true);
+	TickFlush(DeltaSeconds);
+}
+
 static bool bCVarLogPendingGuidsOnShutdown = false;
 static FAutoConsoleVariableRef CVarLogPendingGuidsOnShutdown(
 	TEXT("Net.LogPendingGuidsOnShutdown"),
@@ -1762,8 +1772,6 @@ struct FRPCCSVTracker
 
 void UNetDriver::TickDispatch( float DeltaTime )
 {
-	TGuardValue<bool> GuardInNetTick(bInTick, true);
-
 	SendCycles=0;
 
 	const double CurrentRealtime = FPlatformTime::Seconds();
