Rename LastEditedDocuments when renaming Blueprints so they map to the new remapped package/object.

diff --git a/Engine/Source/Runtime/Engine/Private/Blueprint.cpp b/Engine/Source/Runtime/Engine/Private/Blueprint.cpp
--- a/Engine/Source/Runtime/Engine/Private/Blueprint.cpp
+++ b/Engine/Source/Runtime/Engine/Private/Blueprint.cpp
@@ -562,6 +562,16 @@ bool UBlueprint::RenameGeneratedClasses( const TCHAR* InName, UObject* NewOuter,
 bool UBlueprint::Rename( const TCHAR* InName, UObject* NewOuter, ERenameFlags Flags )
 {
 	const FName OldName = GetFName();
+	
+	TArray<UObject*> LastEditedDocumentsObjects;
+	if (!(Flags & REN_Test))
+	{
+		LastEditedDocumentsObjects.Reserve(LastEditedDocuments.Num());
+		for (const FEditedDocumentInfo& LastEditedDocument : LastEditedDocuments)
+		{
+			LastEditedDocumentsObjects.Add(LastEditedDocument.EditedObjectPath.ResolveObject());
+		}
+	}
 
 	// Move generated class/CDO to the new package, to create redirectors
 	if ( !RenameGeneratedClasses(InName, NewOuter, Flags) )
@@ -569,7 +579,22 @@ bool UBlueprint::Rename( const TCHAR* InName, UObject* NewOuter, ERenameFlags Fl
 		return false;
 	}
 
-	return Super::Rename( InName, NewOuter, Flags );
+	if (Super::Rename( InName, NewOuter, Flags ))
+	{
+		if (!(Flags & REN_Test))
+		{
+			for (int32 i=0; i<LastEditedDocuments.Num(); i++)
+			{
+				if (LastEditedDocumentsObjects[i])
+				{
+					LastEditedDocuments[i].EditedObjectPath = FSoftObjectPath(LastEditedDocumentsObjects[i]);
+				}
+			}
+		}
+		return true;
+	}
+
+	return false;
 }
 
 void UBlueprint::PostDuplicate(bool bDuplicateForPIE)
