World Partition: fix support for disabling loading in the editor to treat actors as streamed, otherwise it would not compute world bounds properly.

#rnx

diff --git a/Engine/Source/Runtime/Engine/Private/WorldPartition/WorldPartition.cpp b/Engine/Source/Runtime/Engine/Private/WorldPartition/WorldPartition.cpp
--- a/Engine/Source/Runtime/Engine/Private/WorldPartition/WorldPartition.cpp
+++ b/Engine/Source/Runtime/Engine/Private/WorldPartition/WorldPartition.cpp
@@ -280,6 +280,7 @@ UWorldPartition::UWorldPartition(const FObjectInitializer& ObjectInitializer)
 #if WITH_EDITOR
 	, EditorHash(nullptr)
 	, AlwaysLoadedActors(nullptr)
+	, ForceLoadedActors(nullptr)
 	, PinnedActors(nullptr)
 	, WorldPartitionEditor(nullptr)
 	, bStreamingWasEnabled(true)
@@ -495,6 +496,9 @@ void UWorldPartition::Initialize(UWorld* InWorld, const FTransform& InTransform)
 
 		AlwaysLoadedActors = new FLoaderAdapterAlwaysLoadedActors(OuterWorld);
 		PinnedActors = new FLoaderAdapterActorList(OuterWorld);
+		
+		IWorldPartitionEditorModule& WorldPartitionEditorModule = FModuleManager::LoadModuleChecked<IWorldPartitionEditorModule>("WorldPartitionEditor");
+		ForceLoadedActors = WorldPartitionEditorModule.GetDisableLoadingInEditor() ? new FLoaderAdapterActorList(OuterWorld) : nullptr;
 	}
 
 	check(RuntimeHash);
@@ -697,6 +701,12 @@ void UWorldPartition::Uninitialize()
 			PinnedActors = nullptr;
 		}
 
+		if (ForceLoadedActors)
+		{
+			delete ForceLoadedActors;
+			ForceLoadedActors = nullptr;
+		}
+
 		if (RegisteredEditorLoaderAdapters.Num())
 		{
 			for (UWorldPartitionEditorLoaderAdapter* RegisteredEditorLoaderAdapter : RegisteredEditorLoaderAdapters)
@@ -1160,6 +1170,11 @@ void UWorldPartition::HashActorDesc(FWorldPartitionActorDesc* ActorDesc)
 	EditorHash->HashActor(ActorHandle);
 
 	bShouldCheckEnableStreamingWarning = IsMainWorldPartition();
+
+	if (ForceLoadedActors)
+	{
+		ForceLoadedActors->AddActors({ ActorDesc->GetGuid() });
+	}
 }
 
 void UWorldPartition::UnhashActorDesc(FWorldPartitionActorDesc* ActorDesc)
@@ -1169,6 +1184,11 @@ void UWorldPartition::UnhashActorDesc(FWorldPartitionActorDesc* ActorDesc)
 
 	FWorldPartitionHandle ActorHandle(this, ActorDesc->GetGuid());
 	EditorHash->UnhashActor(ActorHandle);
+
+	if (ForceLoadedActors)
+	{
+		ForceLoadedActors->RemoveActors({ ActorDesc->GetGuid() });
+	}
 }
 #endif
 
diff --git a/Engine/Source/Runtime/Engine/Private/WorldPartition/WorldPartitionEditorSpatialHash.cpp b/Engine/Source/Runtime/Engine/Private/WorldPartition/WorldPartitionEditorSpatialHash.cpp
--- a/Engine/Source/Runtime/Engine/Private/WorldPartition/WorldPartitionEditorSpatialHash.cpp
+++ b/Engine/Source/Runtime/Engine/Private/WorldPartition/WorldPartitionEditorSpatialHash.cpp
@@ -128,7 +128,7 @@ void UWorldPartitionEditorSpatialHash::HashActor(FWorldPartitionHandle& InActorH
 {
 	check(InActorHandle.IsValid());
 
-	const bool bConsiderActorSpatiallyLoaded = InActorHandle->GetIsSpatiallyLoaded() && !WorldPartitionEditorModule->GetDisableLoadingInEditor();
+	const bool bConsiderActorSpatiallyLoaded = InActorHandle->GetIsSpatiallyLoaded();
 	const FBox ActorBounds = bConsiderActorSpatiallyLoaded ? InActorHandle->GetBounds() : FBox(ForceInit);
 
 #if DO_CHECK
@@ -240,7 +240,7 @@ void UWorldPartitionEditorSpatialHash::UnhashActor(FWorldPartitionHandle& InActo
 {
 	check(InActorHandle.IsValid());
 
-	const bool bConsiderActorSpatiallyLoaded = InActorHandle->GetIsSpatiallyLoaded()  && !WorldPartitionEditorModule->GetDisableLoadingInEditor();
+	const bool bConsiderActorSpatiallyLoaded = InActorHandle->GetIsSpatiallyLoaded();
 	const FBox ActorBounds = bConsiderActorSpatiallyLoaded ? InActorHandle->GetBounds() : FBox(ForceInit);
 
 #if DO_CHECK
diff --git a/Engine/Source/Runtime/Engine/Public/WorldPartition/WorldPartition.h b/Engine/Source/Runtime/Engine/Public/WorldPartition/WorldPartition.h
--- a/Engine/Source/Runtime/Engine/Public/WorldPartition/WorldPartition.h
+++ b/Engine/Source/Runtime/Engine/Public/WorldPartition/WorldPartition.h
@@ -280,6 +280,7 @@ public:
 	TObjectPtr<UWorldPartitionEditorHash> EditorHash;
 
 	FLoaderAdapterAlwaysLoadedActors* AlwaysLoadedActors;
+	FLoaderAdapterActorList* ForceLoadedActors;
 	FLoaderAdapterActorList* PinnedActors;
 
 	IWorldPartitionEditor* WorldPartitionEditor;
