Fixed wp.Runtime.ShowRuntimeSpatialHashCellStreamingPriority

diff --git a/Engine/Source/Runtime/Engine/Private/WorldPartition/WorldPartitionStreamingPolicy.cpp b/Engine/Source/Runtime/Engine/Private/WorldPartition/WorldPartitionStreamingPolicy.cpp
--- a/Engine/Source/Runtime/Engine/Private/WorldPartition/WorldPartitionStreamingPolicy.cpp
+++ b/Engine/Source/Runtime/Engine/Private/WorldPartition/WorldPartitionStreamingPolicy.cpp
@@ -621,6 +621,11 @@ void UWorldPartitionStreamingPolicy::UpdateDebugCellsStreamingPriority(const TSe
 		TArray<const UWorldPartitionRuntimeCell*> Cells = ActivateStreamingCells.Array();
 		Cells.Append(LoadStreamingCells.Array());
 
+		for (const UWorldPartitionRuntimeCell* Cell : Cells)
+		{
+			Cell->MergeStreamingSourceInfo();
+		}
+
 		SortStreamingCellsByImportance(Cells);
 
 		const int32 CellCount = Cells.Num();
