Add a cvar TypedElements.EnableViewportSMInstanceSelection, which can be used to enable/disable ISMC Instance selection without disabling the entire TypedElement system for SMInstances

diff --git a/Engine/Source/Runtime/Engine/Private/InstancedStaticMesh.cpp b/Engine/Source/Runtime/Engine/Private/InstancedStaticMesh.cpp
--- a/Engine/Source/Runtime/Engine/Private/InstancedStaticMesh.cpp
+++ b/Engine/Source/Runtime/Engine/Private/InstancedStaticMesh.cpp
@@ -131,6 +131,14 @@ static TAutoConsoleVariable<int32> CVarRayTracingSimulatedInstanceCount(
 	TEXT("Maximum number of instances to simulate per instanced static mesh, presently capped to 256")
 );
 
+#if WITH_EDITOR
+static TAutoConsoleVariable<int32> CVarEnableViewportSMInstanceSelection(
+	TEXT("TypedElements.EnableViewportSMInstanceSelection"),
+	1,
+	TEXT("Enable direct selection of Instanced Static Mesh Component (ISMC) Instances in the Level Editor Viewport")
+);
+#endif
+
 class FISMExecHelper : public FSelfRegisteringExec
 {
 	virtual bool Exec(class UWorld* InWorld, const TCHAR* Cmd, FOutputDevice& Ar)
@@ -210,11 +218,14 @@ FTypedElementHandle HInstancedStaticMeshInstance::GetElementHandle() const
 #if WITH_EDITOR
 	if (Component)
 	{
-		// Prefer per-instance selection if available
-		// This may fail to return a handle if the feature is disabled, or if per-instance editing is disabled for this component
-		if (FTypedElementHandle ElementHandle = UEngineElementsLibrary::AcquireEditorSMInstanceElementHandle(Component, InstanceIndex))
+		if (CVarEnableViewportSMInstanceSelection.GetValueOnAnyThread() != 0)
 		{
-			return ElementHandle;
+			// Prefer per-instance selection if available
+			// This may fail to return a handle if the feature is disabled, or if per-instance editing is disabled for this component
+			if (FTypedElementHandle ElementHandle = UEngineElementsLibrary::AcquireEditorSMInstanceElementHandle(Component, InstanceIndex))
+			{
+				return ElementHandle;
+			}
 		}
 
 		// If per-instance selection isn't possible, fallback to general per-component selection (which may choose to select the owner actor instead)
