Expose the Input Device Subsystem with ENGINE_API. Add a constructor to the FSetDevicePropertyParams to make it easier to use in gameplay code.

#rnx

diff --git a/Engine/Source/Runtime/Engine/Classes/GameFramework/InputDeviceSubsystem.h b/Engine/Source/Runtime/Engine/Classes/GameFramework/InputDeviceSubsystem.h
--- a/Engine/Source/Runtime/Engine/Classes/GameFramework/InputDeviceSubsystem.h
+++ b/Engine/Source/Runtime/Engine/Classes/GameFramework/InputDeviceSubsystem.h
@@ -14,7 +14,7 @@ class UInputDeviceProperty;
 
 /** Contains a pointer to an active device property and keeps track of how long it has been evaluated for */
 USTRUCT()
-struct FActiveDeviceProperty
+struct ENGINE_API FActiveDeviceProperty
 {
 	GENERATED_BODY()
 
@@ -38,11 +38,13 @@ struct FActiveDeviceProperty
 
 /** Parameters for the UInputDeviceSubsystem::SetDeviceProperty function */
 USTRUCT(BlueprintType)
-struct FSetDevicePropertyParams
+struct ENGINE_API FSetDevicePropertyParams
 {
 	GENERATED_BODY()
 
 	FSetDevicePropertyParams();
+
+	FSetDevicePropertyParams(TSubclassOf<UInputDeviceProperty> InClass, const FPlatformUserId InUserId, const bool bInResetUponCompletion = true);
 	
 	/** The device Property to set */
 	UPROPERTY(EditAnywhere, BlueprintReadWrite, Category = "Input Devices")
@@ -73,7 +75,7 @@ DECLARE_DYNAMIC_MULTICAST_DELEGATE_OneParam(FHardwareInputDeviceChanged, const F
 * on any Platform User. 
 */
 UCLASS(BlueprintType)
-class UInputDeviceSubsystem : public UEngineSubsystem, public FTickableGameObject
+class ENGINE_API UInputDeviceSubsystem : public UEngineSubsystem, public FTickableGameObject
 {
 	friend class FInputDeviceSubsystemProcessor;
 	
@@ -111,6 +113,10 @@ public:
 	UFUNCTION(BlueprintCallable, Category = "Input Devices", meta = (ReturnDisplayName = "Num Removed"))
 	int32 RemoveDeviceProperty(const FPlatformUserId UserId, TSubclassOf<UInputDeviceProperty> DevicePropertyClass);
 
+	/** Removes all the current Input Device Properties that are active, regardless of the Platform User */
+	UFUNCTION(BlueprintCallable, Category = "Input Devices")
+	void RemoveAllDeviceProperties();
+
 	/** Gets the most recently used hardware input device for the given platform user */
 	FHardwareDeviceIdentifier GetMostRecentlyUsedHardwareDevice(const FPlatformUserId InUserId) const;
 
diff --git a/Engine/Source/Runtime/Engine/Private/GameFramework/InputDeviceSubsystem.cpp b/Engine/Source/Runtime/Engine/Private/GameFramework/InputDeviceSubsystem.cpp
--- a/Engine/Source/Runtime/Engine/Private/GameFramework/InputDeviceSubsystem.cpp
+++ b/Engine/Source/Runtime/Engine/Private/GameFramework/InputDeviceSubsystem.cpp
@@ -82,6 +82,14 @@ FSetDevicePropertyParams::FSetDevicePropertyParams()
 
 }
 
+FSetDevicePropertyParams::FSetDevicePropertyParams(TSubclassOf<UInputDeviceProperty> InClass, const FPlatformUserId InUserId, const bool bInResetUponCompletion /* = true */)
+	: DevicePropertyClass(InClass)
+	, UserId(InUserId)
+	, bResetUponCompletion(bInResetUponCompletion)
+{
+	
+}
+
 UInputDeviceSubsystem* UInputDeviceSubsystem::Get()
 {
 	return GEngine ? GEngine->GetEngineSubsystem<UInputDeviceSubsystem>() : nullptr;
@@ -258,6 +266,11 @@ int32 UInputDeviceSubsystem::RemoveDeviceProperty(const FPlatformUserId UserId,
 	return NumRemoved;
 }
 
+void UInputDeviceSubsystem::RemoveAllDeviceProperties()
+{
+	ActiveProperties.Empty();
+}
+
 FHardwareDeviceIdentifier UInputDeviceSubsystem::GetMostRecentlyUsedHardwareDevice(const FPlatformUserId InUserId) const
 {
 	if (const FHardwareDeviceIdentifier* FoundDevice = LatestInputDeviceIdentifiers.Find(InUserId))
