Mobile: Use 'ApproximateOcclusion' for more mesh mobility types, to reduce amount of individual queries for a movable meshes.

diff --git a/Engine/Source/Runtime/Engine/Private/SkeletalMesh.cpp b/Engine/Source/Runtime/Engine/Private/SkeletalMesh.cpp
--- a/Engine/Source/Runtime/Engine/Private/SkeletalMesh.cpp
+++ b/Engine/Source/Runtime/Engine/Private/SkeletalMesh.cpp
@@ -6031,6 +6031,12 @@ FSkeletalMeshSceneProxy::FSkeletalMeshSceneProxy(const USkinnedMeshComponent* Co
 		}
 	}
 #endif
+
+	if (FeatureLevel == ERHIFeatureLevel::ES3_1)
+	{
+		// batch occlusion for a "Static" skeletal mesh
+		bAllowApproximateOcclusion = (bAllowApproximateOcclusion || bRenderStatic);
+	}
 }
 
 
diff --git a/Engine/Source/Runtime/Engine/Private/StaticMeshRender.cpp b/Engine/Source/Runtime/Engine/Private/StaticMeshRender.cpp
--- a/Engine/Source/Runtime/Engine/Private/StaticMeshRender.cpp
+++ b/Engine/Source/Runtime/Engine/Private/StaticMeshRender.cpp
@@ -419,6 +419,12 @@ FStaticMeshSceneProxy::FStaticMeshSceneProxy(UStaticMeshComponent* InComponent,
 
 	// Enable dynamic triangle reordering to remove/reduce sorting issue when rendered with a translucent material (i.e., order-independent-transparency)
 	bSupportsSortedTriangles = InComponent->bSortTriangles;
+
+	// In general it's more beneficial to batch occlusion queires for all meshes on mobile
+	if (FeatureLevel == ERHIFeatureLevel::ES3_1)
+	{
+		bAllowApproximateOcclusion = true;
+	}
 }
 
 void FStaticMeshSceneProxy::SetEvaluateWorldPositionOffsetInRayTracing(bool NewValue)
