[Engine]
* Fixed so all headers that are supposed to be included independently can compile (there are headers that must be included through other headers)

* Added HEADER_UNIT_SKIP to headers that must be included through others

* Added HEADER_UNIT_UNSUPPORTED to headers that are not supported by platform but is part of the include list anyway (files in Clang folder and pthread files).

* Added HEADER_UNIT_UNSUPPORTED to headers that have bad dependencies and include files from modules they don't depend on. (KeyChainUtilities.h and PackageWriter.h)

diff --git a/Engine/Source/Runtime/Engine/Classes/Animation/AnimEnums.h b/Engine/Source/Runtime/Engine/Classes/Animation/AnimEnums.h
--- a/Engine/Source/Runtime/Engine/Classes/Animation/AnimEnums.h
+++ b/Engine/Source/Runtime/Engine/Classes/Animation/AnimEnums.h
@@ -66,7 +66,7 @@ enum EAdditiveBasePoseType
 * Indicates animation data compression format.
 */
 UENUM()
-enum AnimationCompressionFormat
+enum AnimationCompressionFormat : int
 {
 	ACF_None,
 	ACF_Float96NoW,
diff --git a/Engine/Source/Runtime/Engine/Classes/Animation/AnimMontageEvaluationState.h b/Engine/Source/Runtime/Engine/Classes/Animation/AnimMontageEvaluationState.h
--- a/Engine/Source/Runtime/Engine/Classes/Animation/AnimMontageEvaluationState.h
+++ b/Engine/Source/Runtime/Engine/Classes/Animation/AnimMontageEvaluationState.h
@@ -2,6 +2,8 @@
 
 #pragma once
 
+#include "Animation/AnimMontage.h"
+
 struct FMontageEvaluationState
 {
 	FMontageEvaluationState(UAnimMontage* InMontage, float InPosition, FDeltaTimeRecord InDeltaTimeRecord, bool bInIsPlaying, bool bInIsActive, const FAlphaBlend& InBlendInfo, const UBlendProfile* InActiveBlendProfile, float InBlendStartAlpha)
diff --git a/Engine/Source/Runtime/Engine/Classes/Animation/AnimNodeAlphaOptions.h b/Engine/Source/Runtime/Engine/Classes/Animation/AnimNodeAlphaOptions.h
--- a/Engine/Source/Runtime/Engine/Classes/Animation/AnimNodeAlphaOptions.h
+++ b/Engine/Source/Runtime/Engine/Classes/Animation/AnimNodeAlphaOptions.h
@@ -3,6 +3,8 @@
 #pragma once
 
 #include "Animation/AnimInstanceProxy.h"
+#include "Animation/AnimNodeBase.h"
+#include "Animation/InputScaleBias.h"
 
 // Alpha blending options helper functions for anim nodes
 // Assumes that the specified node contains the members:
diff --git a/Engine/Source/Runtime/Engine/Classes/Animation/BlendSpaceBase.h b/Engine/Source/Runtime/Engine/Classes/Animation/BlendSpaceBase.h
--- a/Engine/Source/Runtime/Engine/Classes/Animation/BlendSpaceBase.h
+++ b/Engine/Source/Runtime/Engine/Classes/Animation/BlendSpaceBase.h
@@ -5,6 +5,8 @@
  *
  */
 
+// HEADER_UNIT_SKIP
+
 #pragma once
 
 #include "Animation/BlendSpace.h"
diff --git a/Engine/Source/Runtime/Engine/Classes/Animation/CustomAttributesRuntime.h b/Engine/Source/Runtime/Engine/Classes/Animation/CustomAttributesRuntime.h
--- a/Engine/Source/Runtime/Engine/Classes/Animation/CustomAttributesRuntime.h
+++ b/Engine/Source/Runtime/Engine/Classes/Animation/CustomAttributesRuntime.h
@@ -2,6 +2,8 @@
 
 #pragma once
 
+// HEADER_UNIT_SKIP
+
 #include "CoreMinimal.h"
 
 PRAGMA_DISABLE_DEPRECATION_WARNINGS
diff --git a/Engine/Source/Runtime/Engine/Classes/Animation/IAttributeBlendOperator.h b/Engine/Source/Runtime/Engine/Classes/Animation/IAttributeBlendOperator.h
--- a/Engine/Source/Runtime/Engine/Classes/Animation/IAttributeBlendOperator.h
+++ b/Engine/Source/Runtime/Engine/Classes/Animation/IAttributeBlendOperator.h
@@ -2,6 +2,8 @@
 
 #pragma once
 
+#include "CoreTypes.h"
+
 namespace UE
 {
 	namespace Anim
diff --git a/Engine/Source/Runtime/Engine/Classes/Animation/WrappedAttribute.h b/Engine/Source/Runtime/Engine/Classes/Animation/WrappedAttribute.h
--- a/Engine/Source/Runtime/Engine/Classes/Animation/WrappedAttribute.h
+++ b/Engine/Source/Runtime/Engine/Classes/Animation/WrappedAttribute.h
@@ -4,6 +4,7 @@
 
 #include "CoreMinimal.h"
 #include "Templates/UnrealTypeTraits.h"
+#include "UObject/Class.h"
 
 namespace UE
 {
diff --git a/Engine/Source/Runtime/Engine/Classes/Engine/ContentEncryptionConfig.h b/Engine/Source/Runtime/Engine/Classes/Engine/ContentEncryptionConfig.h
--- a/Engine/Source/Runtime/Engine/Classes/Engine/ContentEncryptionConfig.h
+++ b/Engine/Source/Runtime/Engine/Classes/Engine/ContentEncryptionConfig.h
@@ -1,6 +1,8 @@
 // Copyright Epic Games, Inc. All Rights Reserved.
 #pragma once
 
+#include "CoreTypes.h"
+
 /** Project specific configuration for content encryption */
 class FContentEncryptionConfig
 {
diff --git a/Engine/Source/Runtime/Engine/Classes/Engine/LevelStreamingKismet.h b/Engine/Source/Runtime/Engine/Classes/Engine/LevelStreamingKismet.h
--- a/Engine/Source/Runtime/Engine/Classes/Engine/LevelStreamingKismet.h
+++ b/Engine/Source/Runtime/Engine/Classes/Engine/LevelStreamingKismet.h
@@ -6,6 +6,8 @@
 
 #pragma once
 
+// HEADER_UNIT_SKIP - Deprecated
+
 #ifdef _MSC_VER
 #pragma message(__FILE__"(9): warning : #include LevelStreamingDynamic.h instead of LevelStreamingKismet.h")
 #else
diff --git a/Engine/Source/Runtime/Engine/Classes/Engine/PawnIterator.h b/Engine/Source/Runtime/Engine/Classes/Engine/PawnIterator.h
--- a/Engine/Source/Runtime/Engine/Classes/Engine/PawnIterator.h
+++ b/Engine/Source/Runtime/Engine/Classes/Engine/PawnIterator.h
@@ -2,6 +2,7 @@
 
 #pragma once
 
+#include "EngineUtils.h"
 #include "GameFramework/Pawn.h"
 
 /** Wrapper object that tries to imitate the TWeakObjectPtr interface for the objects previously in the PawnList and iterated by FConstPawnIterator. */
diff --git a/Engine/Source/Runtime/Engine/Classes/Engine/ViewportSplitScreen.h b/Engine/Source/Runtime/Engine/Classes/Engine/ViewportSplitScreen.h
--- a/Engine/Source/Runtime/Engine/Classes/Engine/ViewportSplitScreen.h
+++ b/Engine/Source/Runtime/Engine/Classes/Engine/ViewportSplitScreen.h
@@ -2,7 +2,7 @@
 
 #pragma once
 
-#include "CoreMinimal.h"
+#include "UObject/ObjectMacros.h"
 
 /** 
  * Various data declarations relating to split screen on a GameViewportClient
diff --git a/Engine/Source/Runtime/Engine/Classes/GameFramework/LightWeightInstanceUtils.h b/Engine/Source/Runtime/Engine/Classes/GameFramework/LightWeightInstanceUtils.h
--- a/Engine/Source/Runtime/Engine/Classes/GameFramework/LightWeightInstanceUtils.h
+++ b/Engine/Source/Runtime/Engine/Classes/GameFramework/LightWeightInstanceUtils.h
@@ -2,6 +2,8 @@
 
 #pragma once
 
+// HEADER_UNIT_UNSUPPORTED - I don't think this header is used. dll export of templated function below indicates that.
+
 #include "Kismet/BlueprintFunctionLibrary.h"
 #include "Engine/EngineTypes.h"
 
diff --git a/Engine/Source/Runtime/Engine/Classes/Intrinsic/Model.h b/Engine/Source/Runtime/Engine/Classes/Intrinsic/Model.h
--- a/Engine/Source/Runtime/Engine/Classes/Intrinsic/Model.h
+++ b/Engine/Source/Runtime/Engine/Classes/Intrinsic/Model.h
@@ -1,6 +1,9 @@
 // Copyright Epic Games, Inc. All Rights Reserved.
 
 #pragma once
+
+// HEADER_UNIT_SKIP
+
 #if !CPP && !UE_BUILD_DOCS      //noexport class
 UCLASS(noexport, Intrinsic)
 class UModel : public UObject
diff --git a/Engine/Source/Runtime/Engine/Private/Shader/PreshaderEvaluate.h b/Engine/Source/Runtime/Engine/Private/Shader/PreshaderEvaluate.h
--- a/Engine/Source/Runtime/Engine/Private/Shader/PreshaderEvaluate.h
+++ b/Engine/Source/Runtime/Engine/Private/Shader/PreshaderEvaluate.h
@@ -2,7 +2,7 @@
 #pragma once
 
 #include "CoreMinimal.h"
-#include "Shader/PreshaderTypes.h"
+#include "Shader/Preshader.h"
 
 class FUniformExpressionSet;
 struct FMaterialRenderContext;
diff --git a/Engine/Source/Runtime/Engine/Public/Animation/AnimMTStats.h b/Engine/Source/Runtime/Engine/Public/Animation/AnimMTStats.h
--- a/Engine/Source/Runtime/Engine/Public/Animation/AnimMTStats.h
+++ b/Engine/Source/Runtime/Engine/Public/Animation/AnimMTStats.h
@@ -2,6 +2,8 @@
 
 // [[ IncludeTool: Inline ]] // Markup to tell IncludeTool that this file is state changing and cannot be optimized out.
 
+// HEADER_UNIT_SKIP - Included through other header
+
 DO_ANIMSTAT_PROCESSING(PerformAnimEvaluation);
 DO_ANIMSTAT_PROCESSING(SkeletalComponentAnimEvaluate);
 DO_ANIMSTAT_PROCESSING(EvaluateAnimGraph);
diff --git a/Engine/Source/Runtime/Engine/Public/Animation/AnimPhysicsSolver.h b/Engine/Source/Runtime/Engine/Public/Animation/AnimPhysicsSolver.h
--- a/Engine/Source/Runtime/Engine/Public/Animation/AnimPhysicsSolver.h
+++ b/Engine/Source/Runtime/Engine/Public/Animation/AnimPhysicsSolver.h
@@ -48,16 +48,16 @@ class FAnimPhys;
 namespace AnimPhysicsConstants
 {
 	// Added bias for angular joints
-	const float		JointBiasFactor = 0.3f;
+	inline const float		JointBiasFactor = 0.3f;
 
 	// Damping for linear momentum (1 = critically damped)
-	const float		LinearDamping = 0.7f;
+	inline const float		LinearDamping = 0.7f;
 
 	// Damping for angular momentum (1 = critically damped)
-	const float		AngularDamping = 0.7f;
+	inline const float		AngularDamping = 0.7f;
 
-	const float		DefaultSpringConstantLinear = 12.0f;
-	const float		DefaultSpringConstantAngular = 4.0f;
+	inline const float		DefaultSpringConstantLinear = 12.0f;
+	inline const float		DefaultSpringConstantAngular = 4.0f;
 }
 
 // Enum for picking current angular twist axis
diff --git a/Engine/Source/Runtime/Engine/Public/Animation/AnimStats.h b/Engine/Source/Runtime/Engine/Public/Animation/AnimStats.h
--- a/Engine/Source/Runtime/Engine/Public/Animation/AnimStats.h
+++ b/Engine/Source/Runtime/Engine/Public/Animation/AnimStats.h
@@ -7,7 +7,7 @@
 #pragma once
 
 #include "CoreMinimal.h"
-#include "Stats/Stats.h"
+#include "Stats/Stats2.h"
 
 /** Skeletal stats */
 DECLARE_CYCLE_STAT_EXTERN(TEXT("SkinnedMeshComp Tick"), STAT_SkinnedMeshCompTick, STATGROUP_Anim, );
diff --git a/Engine/Source/Runtime/Engine/Public/AnimationCompression.h b/Engine/Source/Runtime/Engine/Public/AnimationCompression.h
--- a/Engine/Source/Runtime/Engine/Public/AnimationCompression.h
+++ b/Engine/Source/Runtime/Engine/Public/AnimationCompression.h
@@ -36,7 +36,7 @@ ENGINE_API DECLARE_LOG_CATEGORY_EXTERN(LogAnimationCompression, Log, All);
 #define Quant11BitOffs    (1023)
 
 //@TODO: Explore different scales
-static constexpr int32 LogScale = 7;
+inline static constexpr int32 LogScale = 7;
 
 namespace AnimationCompressionUtils
 {
diff --git a/Engine/Source/Runtime/Engine/Public/DistanceFieldAtlas.h b/Engine/Source/Runtime/Engine/Public/DistanceFieldAtlas.h
--- a/Engine/Source/Runtime/Engine/Public/DistanceFieldAtlas.h
+++ b/Engine/Source/Runtime/Engine/Public/DistanceFieldAtlas.h
@@ -36,18 +36,18 @@ template <class T> class TLockFreePointerListLIFO;
 namespace DistanceField
 {
 	// One voxel border around object for handling gradient
-	constexpr int32 MeshDistanceFieldObjectBorder = 1;
-	constexpr int32 UniqueDataBrickSize = 7;
+	inline constexpr int32 MeshDistanceFieldObjectBorder = 1;
+	inline constexpr int32 UniqueDataBrickSize = 7;
 	// Half voxel border around brick for trilinear filtering
-	constexpr int32 BrickSize = 8;
+	inline constexpr int32 BrickSize = 8;
 	// Trade off between SDF memory and number of steps required to find intersection
-	constexpr int32 BandSizeInVoxels = 4;
-	constexpr int32 NumMips = 3;
-	constexpr uint32 InvalidBrickIndex = 0xFFFFFFFF;
-	constexpr EPixelFormat DistanceFieldFormat = PF_G8;
+	inline constexpr int32 BandSizeInVoxels = 4;
+	inline constexpr int32 NumMips = 3;
+	inline constexpr uint32 InvalidBrickIndex = 0xFFFFFFFF;
+	inline constexpr EPixelFormat DistanceFieldFormat = PF_G8;
 
 	// Must match LoadDFAssetData
-	constexpr uint32 MaxIndirectionDimension = 1024;
+	inline constexpr uint32 MaxIndirectionDimension = 1024;
 };
 
 class ENGINE_API FLandscapeTextureAtlas : public FRenderResource
diff --git a/Engine/Source/Runtime/Engine/Public/GPUSceneWriter.h b/Engine/Source/Runtime/Engine/Public/GPUSceneWriter.h
--- a/Engine/Source/Runtime/Engine/Public/GPUSceneWriter.h
+++ b/Engine/Source/Runtime/Engine/Public/GPUSceneWriter.h
@@ -3,6 +3,7 @@
 #pragma once
 
 #include "CoreMinimal.h"
+#include "ShaderParameterMacros.h"
 #include "ShaderParameters.h"
 
 class FSceneView;
diff --git a/Engine/Source/Runtime/Engine/Public/GameTime.h b/Engine/Source/Runtime/Engine/Public/GameTime.h
--- a/Engine/Source/Runtime/Engine/Public/GameTime.h
+++ b/Engine/Source/Runtime/Engine/Public/GameTime.h
@@ -2,6 +2,7 @@
 
 #pragma once
 
+#include "Misc/AssertionMacros.h"
 #include "Misc/CoreMiscDefines.h"
 
 /** Contains all the timings of a gaming frame, to handle pause and time dilation (for instance bullet time) of the world. */
diff --git a/Engine/Source/Runtime/Engine/Public/GeneratedCodeHelpers.h b/Engine/Source/Runtime/Engine/Public/GeneratedCodeHelpers.h
--- a/Engine/Source/Runtime/Engine/Public/GeneratedCodeHelpers.h
+++ b/Engine/Source/Runtime/Engine/Public/GeneratedCodeHelpers.h
@@ -2,6 +2,8 @@
 
 #pragma once
 
+// HEADER_UNIT_SKIP
+
 // DO NOT INCLUDE THIS HEADER!
 // THIS FILE SHOULD BE USED ONLY BY AUTOMATICALLY GENERATED CODE. 
 
diff --git a/Engine/Source/Runtime/Engine/Public/HLOD/HLODBatchingPolicy.h b/Engine/Source/Runtime/Engine/Public/HLOD/HLODBatchingPolicy.h
--- a/Engine/Source/Runtime/Engine/Public/HLOD/HLODBatchingPolicy.h
+++ b/Engine/Source/Runtime/Engine/Public/HLOD/HLODBatchingPolicy.h
@@ -2,7 +2,9 @@
 
 #pragma once
 
-#include "CoreMinimal.h"
+// HEADER_UNIT_SKIP - Included through other
+
+#include "CoreTypes.h"
 
 /** Determines how the geometry of a component will be incorporated in proxy (simplified) HLODs. */
 UENUM()
diff --git a/Engine/Source/Runtime/Engine/Public/Net/Util/ResizableCircularQueue.h b/Engine/Source/Runtime/Engine/Public/Net/Util/ResizableCircularQueue.h
--- a/Engine/Source/Runtime/Engine/Public/Net/Util/ResizableCircularQueue.h
+++ b/Engine/Source/Runtime/Engine/Public/Net/Util/ResizableCircularQueue.h
@@ -2,6 +2,8 @@
 
 #pragma once
 
+// HEADER_UNIT_SKIP - Deprecated
+
 // forward to new location
 #include "Net/Core/Misc/ResizableCircularQueue.h"
 
diff --git a/Engine/Source/Runtime/Engine/Public/ObjectCacheEventSink.h b/Engine/Source/Runtime/Engine/Public/ObjectCacheEventSink.h
--- a/Engine/Source/Runtime/Engine/Public/ObjectCacheEventSink.h
+++ b/Engine/Source/Runtime/Engine/Public/ObjectCacheEventSink.h
@@ -5,6 +5,11 @@
 
 #if WITH_EDITOR
 
+class UActorComponent;
+class UMaterialInterface;
+class UPrimitiveComponent;
+class UStaticMeshComponent;
+
 struct ENGINE_API FObjectCacheEventSink
 {
 	// For parallel processing.  Calling BeginQueueNotifyEvents will prevent notify events from being processed, an instead placed in a list for later processing.
diff --git a/Engine/Source/Runtime/Engine/Public/OverrideVoidReturnInvoker.h b/Engine/Source/Runtime/Engine/Public/OverrideVoidReturnInvoker.h
--- a/Engine/Source/Runtime/Engine/Public/OverrideVoidReturnInvoker.h
+++ b/Engine/Source/Runtime/Engine/Public/OverrideVoidReturnInvoker.h
@@ -3,6 +3,7 @@
 #pragma once
 
 #include "CoreTypes.h"
+#include <type_traits>
 
 /**
  * Wraps invocation of a function that can return a value or not. If it doesn't, the function call is wrapped into
diff --git a/Engine/Source/Runtime/Engine/Public/ParticleHelper.h b/Engine/Source/Runtime/Engine/Public/ParticleHelper.h
--- a/Engine/Source/Runtime/Engine/Public/ParticleHelper.h
+++ b/Engine/Source/Runtime/Engine/Public/ParticleHelper.h
@@ -45,7 +45,7 @@ enum EParticleDetailMode
 	PDM_High UMETA(DisplayName = "High"),
 	PDM_MAX UMETA(Hidden),
 };
-const int32 PDM_DefaultValue = 0xFFFF;
+inline const int32 PDM_DefaultValue = 0xFFFF;
 
 
 /*-----------------------------------------------------------------------------
diff --git a/Engine/Source/Runtime/Engine/Public/Physics/Experimental/PhysicsUserData_Chaos.h b/Engine/Source/Runtime/Engine/Public/Physics/Experimental/PhysicsUserData_Chaos.h
--- a/Engine/Source/Runtime/Engine/Public/Physics/Experimental/PhysicsUserData_Chaos.h
+++ b/Engine/Source/Runtime/Engine/Public/Physics/Experimental/PhysicsUserData_Chaos.h
@@ -5,6 +5,8 @@
 #include "EngineDefines.h"
 #include "Physics/PhysicsInterfaceDeclares.h"
 
+class UPhysicalMaterial;
+struct FBodyInstance;
 struct FConstraintInstanceBase;
 
 // Note: this must match the enum EChaosUserDataType in PhysicsInterfaceTypesCore
diff --git a/Engine/Source/Runtime/Engine/Public/PlanarReflectionSceneProxy.h b/Engine/Source/Runtime/Engine/Public/PlanarReflectionSceneProxy.h
--- a/Engine/Source/Runtime/Engine/Public/PlanarReflectionSceneProxy.h
+++ b/Engine/Source/Runtime/Engine/Public/PlanarReflectionSceneProxy.h
@@ -14,7 +14,7 @@
 
 // Currently we support at most 2 views for each planar reflection, one view per stereo pass
 // Must match FPlanarReflectionUniformParameters.
-const int32 GMaxPlanarReflectionViews = 2;
+inline const int32 GMaxPlanarReflectionViews = 2;
 
 class UPlanarReflectionComponent;
 
diff --git a/Engine/Source/Runtime/Engine/Public/PrecomputedVolumetricLightmap.h b/Engine/Source/Runtime/Engine/Public/PrecomputedVolumetricLightmap.h
--- a/Engine/Source/Runtime/Engine/Public/PrecomputedVolumetricLightmap.h
+++ b/Engine/Source/Runtime/Engine/Public/PrecomputedVolumetricLightmap.h
@@ -307,7 +307,7 @@ inline FLinearColor ConvertToLinearColor<uint8>(uint8 InColor)
 	return FLinearColor(InColor * Scale, 0, 0, 0);
 };
 
-static const float GPointFilteringThreshold = .001f;
+inline static const float GPointFilteringThreshold = .001f;
 
 template<typename VoxelDataType>
 FLinearColor FilteredVolumeLookup(FVector Coordinate, FIntVector DataDimensions, const VoxelDataType* Data)
diff --git a/Engine/Source/Runtime/Engine/Public/PropertyPairsMap.h b/Engine/Source/Runtime/Engine/Public/PropertyPairsMap.h
--- a/Engine/Source/Runtime/Engine/Public/PropertyPairsMap.h
+++ b/Engine/Source/Runtime/Engine/Public/PropertyPairsMap.h
@@ -2,9 +2,11 @@
 
 #pragma once
 
+#include "Containers/Map.h"
 #include "CoreTypes.h"
 #include "Misc/StringBuilder.h"
 #include "OverrideVoidReturnInvoker.h"
+#include "UObject/NameTypes.h"
 
 /** A simple name/value pairs map */
 class FPropertyPairsMap
diff --git a/Engine/Source/Runtime/Engine/Public/Rendering/StaticMeshVertexDataInterface.h b/Engine/Source/Runtime/Engine/Public/Rendering/StaticMeshVertexDataInterface.h
--- a/Engine/Source/Runtime/Engine/Public/Rendering/StaticMeshVertexDataInterface.h
+++ b/Engine/Source/Runtime/Engine/Public/Rendering/StaticMeshVertexDataInterface.h
@@ -2,6 +2,8 @@
 
 #pragma once
 
+#include "CoreTypes.h"
+
 class FResourceArrayInterface;
 
 enum class EResizeBufferFlags
diff --git a/Engine/Source/Runtime/Engine/Public/ReplayTypes.h b/Engine/Source/Runtime/Engine/Public/ReplayTypes.h
--- a/Engine/Source/Runtime/Engine/Public/ReplayTypes.h
+++ b/Engine/Source/Runtime/Engine/Public/ReplayTypes.h
@@ -116,12 +116,12 @@ enum ENetworkVersionHistory
 	HISTORY_LATEST = HISTORY_PLUS_ONE - 1
 };
 
-static const uint32 NETWORK_DEMO_MAGIC = 0x2CF5A13D;
-static const uint32 NETWORK_DEMO_VERSION = HISTORY_LATEST;
-static const uint32 MIN_NETWORK_DEMO_VERSION = HISTORY_CHARACTER_MOVEMENT;
+inline static const uint32 NETWORK_DEMO_MAGIC = 0x2CF5A13D;
+inline static const uint32 NETWORK_DEMO_VERSION = HISTORY_LATEST;
+inline static const uint32 MIN_NETWORK_DEMO_VERSION = HISTORY_CHARACTER_MOVEMENT;
 
-static const uint32 NETWORK_DEMO_METADATA_MAGIC = 0x3D06B24E;
-static const uint32 NETWORK_DEMO_METADATA_VERSION = 0;
+inline static const uint32 NETWORK_DEMO_METADATA_MAGIC = 0x3D06B24E;
+inline static const uint32 NETWORK_DEMO_METADATA_VERSION = 0;
 
 struct FNetworkDemoHeader
 {
diff --git a/Engine/Source/Runtime/Engine/Public/StaticMeshSceneProxy.h b/Engine/Source/Runtime/Engine/Public/StaticMeshSceneProxy.h
--- a/Engine/Source/Runtime/Engine/Public/StaticMeshSceneProxy.h
+++ b/Engine/Source/Runtime/Engine/Public/StaticMeshSceneProxy.h
@@ -9,6 +9,7 @@
 #include "Components/StaticMeshComponent.h"
 #include "PrimitiveSceneProxy.h"
 #include "RayTracingInstance.h"
+#include "SceneManagement.h"
 
 class FLocalVertexFactoryUniformShaderParameters;
 class FRawStaticIndexBuffer;
diff --git a/Engine/Source/Runtime/Engine/Public/TextureReferenceResolver.h b/Engine/Source/Runtime/Engine/Public/TextureReferenceResolver.h
--- a/Engine/Source/Runtime/Engine/Public/TextureReferenceResolver.h
+++ b/Engine/Source/Runtime/Engine/Public/TextureReferenceResolver.h
@@ -3,6 +3,8 @@
 #pragma once
 
 #include "Engine/Texture.h"
+#include "RenderResource.h"
+#include "RHI.h"
 
 /** 
 * Replaces the RHI reference of one texture with another.
diff --git a/Engine/Source/Runtime/Engine/Public/WorldPartition/ContentBundle/ContentBundlePaths.h b/Engine/Source/Runtime/Engine/Public/WorldPartition/ContentBundle/ContentBundlePaths.h
--- a/Engine/Source/Runtime/Engine/Public/WorldPartition/ContentBundle/ContentBundlePaths.h
+++ b/Engine/Source/Runtime/Engine/Public/WorldPartition/ContentBundle/ContentBundlePaths.h
@@ -8,8 +8,8 @@ class FContentBundleBase;
 
 namespace ContentBundlePaths
 {
-	constexpr FStringView ContentBundleFolder = TEXTVIEW("/ContentBundle/");
-	constexpr FStringView GeneratedFolder = TEXTVIEW("/_Generated_/");
+	inline constexpr FStringView ContentBundleFolder = TEXTVIEW("/ContentBundle/");
+	inline constexpr FStringView GeneratedFolder = TEXTVIEW("/_Generated_/");
 
 	// Returns "/ContentBundle/"
 	constexpr FStringView GetContentBundleFolder() { return ContentBundleFolder; }
diff --git a/Engine/Source/Runtime/Engine/Public/WorldPartition/WorldPartitionActorDescType.h b/Engine/Source/Runtime/Engine/Public/WorldPartition/WorldPartitionActorDescType.h
--- a/Engine/Source/Runtime/Engine/Public/WorldPartition/WorldPartitionActorDescType.h
+++ b/Engine/Source/Runtime/Engine/Public/WorldPartition/WorldPartitionActorDescType.h
@@ -3,6 +3,8 @@
 
 #include "CoreMinimal.h"
 
+class AActor;
+
 #if WITH_EDITOR
 template <typename ActorType, TEMPLATE_REQUIRES(TIsDerivedFrom<ActorType, AActor>::IsDerived)>
 struct FWorldPartitionActorDescType
diff --git a/Engine/Source/Runtime/Engine/Public/kDOP.h b/Engine/Source/Runtime/Engine/Public/kDOP.h
--- a/Engine/Source/Runtime/Engine/Public/kDOP.h
+++ b/Engine/Source/Runtime/Engine/Public/kDOP.h
@@ -1091,9 +1091,9 @@ template <typename COLL_DATA_PROVIDER, typename KDOP_IDX_TYPE> struct TkDOPLineC
 		// Calculate the vector's direction in local space
 		LocalDir = LocalEnd - LocalStart;
 		// Build the one over dir
-		LocalOneOverDir.X = LocalDir.X ? 1.f / LocalDir.X : MAX_FLT;
-		LocalOneOverDir.Y = LocalDir.Y ? 1.f / LocalDir.Y : MAX_FLT;
-		LocalOneOverDir.Z = LocalDir.Z ? 1.f / LocalDir.Z : MAX_FLT;
+		LocalOneOverDir.X = LocalDir.X ? 1.f / LocalDir.X : UE_MAX_FLT;
+		LocalOneOverDir.Y = LocalDir.Y ? 1.f / LocalDir.Y : UE_MAX_FLT;
+		LocalOneOverDir.Z = LocalDir.Z ? 1.f / LocalDir.Z : UE_MAX_FLT;
 
 		// Construct the SOA data
 		StartSOA.X = VectorLoadFloat1( &LocalStart.X );
