Fixed wp.Runtime.RuntimeSpatialHashPlacePartitionActorsUsingLocation ignored when wp.Runtime.RuntimeSpatialHashPlaceSmallActorsUsingLocation was enabled.

diff --git a/Engine/Source/Runtime/Engine/Private/WorldPartition/RuntimeSpatialHash/RuntimeSpatialHashGridHelper.cpp b/Engine/Source/Runtime/Engine/Private/WorldPartition/RuntimeSpatialHash/RuntimeSpatialHashGridHelper.cpp
--- a/Engine/Source/Runtime/Engine/Private/WorldPartition/RuntimeSpatialHash/RuntimeSpatialHashGridHelper.cpp
+++ b/Engine/Source/Runtime/Engine/Private/WorldPartition/RuntimeSpatialHash/RuntimeSpatialHashGridHelper.cpp
@@ -171,9 +171,9 @@ FSquare2DGridHelper GetPartitionedActors(const FBox& WorldBounds, const FSpatial
 	auto ShouldActorUseLocationPlacement = [CellArea, CellSize, GridLevelCount](const IStreamingGenerationContext::FActorSetInstance* InActorSetInstance, const FBox2D& InActorSetInstanceBounds, int32& OutGridLevel)
 	{
 		OutGridLevel = 0;
-		if (GRuntimeSpatialHashPlaceSmallActorsUsingLocation)
+		if (GRuntimeSpatialHashPlaceSmallActorsUsingLocation && (InActorSetInstanceBounds.GetArea() <= CellArea))
 		{
-			return InActorSetInstanceBounds.GetArea() <= CellArea;
+			return true;
 		}
 
 		if (GRuntimeSpatialHashPlacePartitionActorsUsingLocation)
