Add methods for setting engine analytics api key before provider creation.

diff --git a/Engine/Source/Runtime/Engine/Private/EngineAnalytics.cpp b/Engine/Source/Runtime/Engine/Private/EngineAnalytics.cpp
--- a/Engine/Source/Runtime/Engine/Private/EngineAnalytics.cpp
+++ b/Engine/Source/Runtime/Engine/Private/EngineAnalytics.cpp
@@ -37,9 +37,31 @@ FSimpleMulticastDelegate FEngineAnalytics::OnInitializeEngineAnalytics;
 FSimpleMulticastDelegate FEngineAnalytics::OnShutdownEngineAnalytics;
 #endif
 
+namespace UE
+{
+namespace Analytics
+{
+static TOptional<FString> GAPIKey;
+static TOptional<FString> GBackendEnvironment;
+}
+}
+
+void FEngineAnalytics::SetAPIKey(const FString& InAPIKey, const FString& BackendEnv)
+{
+	UE::Analytics::GAPIKey = InAPIKey;
+	UE::Analytics::GBackendEnvironment = BackendEnv;
+}
+
 static TSharedPtr<IAnalyticsProviderET> CreateEpicAnalyticsProvider()
 {
 	FAnalyticsET::Config Config;
+	Config.APIServerET = TEXT("https://datarouter.ol.epicgames.com/");
+
+	if (UE::Analytics::GAPIKey)
+	{
+		Config.APIKeyET = *UE::Analytics::GAPIKey;
+	}
+	else
 	{
 		// We always use the "Release" analytics account unless we're running in analytics test mode (usually with
 		// a command-line parameter), or we're an internal Epic build
@@ -52,13 +74,18 @@ static TSharedPtr<IAnalyticsProviderET> CreateEpicAnalyticsProvider()
 		FString UETypeOverride;
 		bool bHasOverride = GConfig->GetString(TEXT("Analytics"), TEXT("UE4TypeOverride"), UETypeOverride, GEngineIni);
 		const TCHAR* UETypeStr = bHasOverride ? *UETypeOverride : FEngineBuildSettings::IsPerforceBuild() ? TEXT("Perforce") : TEXT("UnrealEngine");
+		Config.APIKeyET = FString::Printf(TEXT("UEEditor.%s.%s"), UETypeStr, BuildTypeStr);
+	}
 
-		FString AppID;
-		GConfig->GetString(TEXT("Analytics"), TEXT("AppIdOverride"), AppID, GEditorIni);
-		Config.APIKeyET = FString::Printf(TEXT("%s.%s.%s"), AppID.IsEmpty() ? TEXT("UEEditor") : *AppID, UETypeStr, BuildTypeStr);
+	if (UE::Analytics::GBackendEnvironment)
+	{
+		Config.AppEnvironment = *UE::Analytics::GBackendEnvironment;
 	}
-	Config.APIServerET = TEXT("https://datarouter.ol.epicgames.com/");
-	Config.AppEnvironment = TEXT("datacollector-binary");
+	else
+	{
+		Config.AppEnvironment = TEXT("datacollector-binary");
+	}
+
 	Config.AppVersionET = FEngineVersion::Current().ToString();
 
 	// Connect the engine analytics provider (if there is a configuration delegate installed)
diff --git a/Engine/Source/Runtime/Engine/Public/EngineAnalytics.h b/Engine/Source/Runtime/Engine/Public/EngineAnalytics.h
--- a/Engine/Source/Runtime/Engine/Public/EngineAnalytics.h
+++ b/Engine/Source/Runtime/Engine/Public/EngineAnalytics.h
@@ -56,6 +56,8 @@ public:
 
 	static ENGINE_API void LowDriveSpaceDetected();
 
+	static ENGINE_API void SetAPIKey(const FString& InAPIKey, const FString& Backend);
+
 private:
 	static bool bIsInitialized;
 	static ENGINE_API TSharedPtr<IAnalyticsProviderET> Analytics;
