- Add ability to force disable Nanite on a Static Mesh Component with a run-time switch.
- Fix a bug where transparent meshes that were purposefully disallowing Nanite would get reset to Nanite, causing them to render with a default material.

#localization none
#tests confirmed fix in editor

diff --git a/Engine/Source/Runtime/Engine/Classes/Components/StaticMeshComponent.h b/Engine/Source/Runtime/Engine/Classes/Components/StaticMeshComponent.h
--- a/Engine/Source/Runtime/Engine/Classes/Components/StaticMeshComponent.h
+++ b/Engine/Source/Runtime/Engine/Classes/Components/StaticMeshComponent.h
@@ -141,9 +141,14 @@ public:
 	UPROPERTY(EditAnywhere, AdvancedDisplay, BlueprintReadOnly, Category=Rendering, meta=(editcondition = "bOverrideWireframeColor"))
 	FColor WireframeColorOverride;
 
+	/** Forces this component to use fallback mesh for rendering if Nanite is enabled on the mesh. */
 	UPROPERTY(EditAnywhere, AdvancedDisplay, BlueprintReadWrite, Category = Rendering)
 	uint8 bDisallowNanite : 1;
 
+	/** Forces this component to use fallback mesh for rendering if Nanite is enabled on the mesh (run-time override) */
+	UPROPERTY()
+	uint8 bForceDisableNanite : 1;
+
 	/** 
 	 * Whether to evaluate World Position Offset. 
 	 * This is only used when running with r.OptimizedWPO=1 
@@ -392,6 +397,10 @@ public:
 	UFUNCTION(BlueprintCallable, Category="Rendering|Lighting")
 	void SetReverseCulling(bool ReverseCulling);
 
+	/** Force disabling of Nanite rendering. When true, Will swap to the the fallback mesh instead. */
+	UFUNCTION(BlueprintCallable, Category="Rendering")
+	void SetForceDisableNanite(bool bInForceDisableNanite);
+
 	virtual void SetCollisionProfileName(FName InCollisionProfileName, bool bUpdateOverlaps=true) override;
 
 public:
diff --git a/Engine/Source/Runtime/Engine/Private/Components/StaticMeshComponent.cpp b/Engine/Source/Runtime/Engine/Private/Components/StaticMeshComponent.cpp
--- a/Engine/Source/Runtime/Engine/Private/Components/StaticMeshComponent.cpp
+++ b/Engine/Source/Runtime/Engine/Private/Components/StaticMeshComponent.cpp
@@ -208,6 +208,7 @@ UStaticMeshComponent::UStaticMeshComponent(const FObjectInitializer& ObjectIniti
 	bForceNavigationObstacle = true;
 	bDisallowMeshPaintPerInstance = false;
 	bDisallowNanite = false;
+	bForceDisableNanite = false;
 	bEvaluateWorldPositionOffset = true;
 	bEvaluateWorldPositionOffsetInRayTracing = false;
 	bInitialEvaluateWorldPositionOffset = false;
@@ -1721,7 +1722,7 @@ void UStaticMeshComponent::UpdatePreCulledData(int32 LODIndex, const TArray<uint
 
 bool UStaticMeshComponent::ShouldCreateNaniteProxy() const
 {
-	if (bDisallowNanite || GetScene() == nullptr)
+	if (bDisallowNanite || bForceDisableNanite || GetScene() == nullptr)
 	{
 		// Regardless of the static mesh asset supporting Nanite, this component does not want Nanite to be used
 		return false;
@@ -2371,6 +2372,17 @@ void UStaticMeshComponent::SetReverseCulling(bool ReverseCulling)
 	}
 }
 
+void UStaticMeshComponent::SetForceDisableNanite(bool bInForceDisableNanite)
+{
+	bForceDisableNanite = bInForceDisableNanite;
+
+	// Check if we now need to recreate our scene proxy
+	if (SceneProxy != nullptr && SceneProxy->IsNaniteMesh() != ShouldCreateNaniteProxy())
+	{
+		MarkRenderStateDirty();
+	}
+}
+
 void UStaticMeshComponent::GetLocalBounds(FVector& Min, FVector& Max) const
 {
 	if (GetStaticMesh())
