Improve EyeAdaptationInverse node to support any float numeric parameter type.
- New version of EyeAdaptationInverseLookup hlsl function that returns inverse eye adaptation scale.
- Setup multiplication with LightValue parameter in HLSLMaterialTranslator.

diff --git a/Engine/Shaders/Private/EyeAdaptationCommon.ush b/Engine/Shaders/Private/EyeAdaptationCommon.ush
--- a/Engine/Shaders/Private/EyeAdaptationCommon.ush
+++ b/Engine/Shaders/Private/EyeAdaptationCommon.ush
@@ -66,7 +66,7 @@ float EyeAdaptationLookup()
 #endif
 }
 
-float3 EyeAdaptationInverseLookup(float3 LightValue, float Alpha)
+float EyeAdaptationInverseLookup(float Alpha)
 {
 	float Adaptation = EyeAdaptationLookup();
 
@@ -79,8 +79,10 @@ float3 EyeAdaptationInverseLookup(float3 LightValue, float Alpha)
 
 	float LerpLogScale = -Alpha * log(Adaptation);
 	float Scale = exp(LerpLogScale);
-	return LightValue * Scale;
+	return Scale;
 }
 
-
-
+float3 EyeAdaptationInverseLookup(float3 LightValue, float Alpha)
+{
+	return LightValue * EyeAdaptationInverseLookup(Alpha);
+}
\ No newline at end of file
diff --git a/Engine/Source/Runtime/Engine/Private/Materials/HLSLMaterialTranslator.cpp b/Engine/Source/Runtime/Engine/Private/Materials/HLSLMaterialTranslator.cpp
--- a/Engine/Source/Runtime/Engine/Private/Materials/HLSLMaterialTranslator.cpp
+++ b/Engine/Source/Runtime/Engine/Private/Materials/HLSLMaterialTranslator.cpp
@@ -12039,23 +12039,27 @@ int32 FHLSLMaterialTranslator::EyeAdaptationInverse(int32 LightValueArg, int32 A
 		return INDEX_NONE;
 	}
 
-	if (GetParameterType(LightValueArg) != MCT_Float3)
+	const EMaterialValueType LightValueType = GetParameterType(LightValueArg);
+	const EMaterialValueType ResultType = LightValueType;
+
+	if (!IsFloatNumericType(LightValueType))
 	{
-		Errorf(TEXT("EyeAdaptationInverse expects a float3 type for LightValue"));
+		Errorf(TEXT("EyeAdaptationInverse expects a float numeric type for LightValue"));
 		return INDEX_NONE;
 	}
-	int32 LightValueName = LightValueArg;
 
 	if (GetParameterType(AlphaArg) != MCT_Float)
 	{
 		Errorf(TEXT("EyeAdaptationInverse expects a float type for Alpha"));
 		return INDEX_NONE;
 	}
-	int32 AlphaName = AlphaArg;
 
 	MaterialCompilationOutput.bUsesEyeAdaptation = true;
 
-	return AddInlinedCodeChunk(MCT_Float3, TEXT("EyeAdaptationInverseLookup(%s,%s)"), *GetParameterCode(LightValueName), *GetParameterCode(AlphaName));
+	int32 Multiplier = AddInlinedCodeChunk(MCT_Float, TEXT("EyeAdaptationInverseLookup(%s)"), *GetParameterCode(AlphaArg));
+
+	// return LightValue scaled by inverse eye adaptation
+	return Mul(LightValueArg, Multiplier);
 }
 
 // to only have one piece of code dealing with error handling if the Primitive constant buffer is not used.
