Fix Editor crash when updating volume texture source from a function.

Note: Previously the source was first initialized with a dummy texture in order to determine the number of bytes per pixel. However the dummy texture was created with zero size and non-zero mip count, which was incorrect. Creating a dummy texture is no longer necessary, as we now have a static function to determine bytes per pixel. In addition to that, the referenced source texture is now reset when generating the volume source from a function.

diff --git a/Engine/Source/Runtime/Engine/Private/VolumeTexture.cpp b/Engine/Source/Runtime/Engine/Private/VolumeTexture.cpp
--- a/Engine/Source/Runtime/Engine/Private/VolumeTexture.cpp
+++ b/Engine/Source/Runtime/Engine/Private/VolumeTexture.cpp
@@ -21,6 +21,7 @@
 #include "TextureResource.h"
 #include "UObject/Package.h"
 #include "UObject/StrongObjectPtr.h"
+#include "ImageCoreUtils.h"
 
 #include UE_INLINE_GENERATED_CPP_BY_NAME(VolumeTexture)
 
@@ -248,10 +249,8 @@ ENGINE_API bool UVolumeTexture::UpdateSourceFromFunction(TFunction<void(int32, i
 	
 	Modify(true);
 
-	// First clear up the existing source with the requested TextureSourceFormat
-	Source.Init(0, 0, 0, 1, Format, nullptr);
-	// It is now possible to query the correct FormatDataSize (there is no static version of GetBytesPerPixel)
-	const int32 FormatDataSize = Source.GetBytesPerPixel();
+	Source2DTexture = nullptr;
+	const int32 FormatDataSize = ERawImageFormat::GetBytesPerPixel(FImageCoreUtils::ConvertToRawImageFormat(Format));
 
 	// Allocate temp buffer used to fill texture
 	uint8* const NewData = (uint8*)FMemory::Malloc(SizeX * SizeY * SizeZ * FormatDataSize);
