ComputeFramework: Add execution groups so that workloads can be scheduled for different points in the frame.
Added an Immediate execution group which will immediately enqueue the work to the render thread.
Optimus: Use Immediate execution to submit a deformer update when mesh component render state is recreated.
This prevents potential issue with one frame not having any deformerd vertices.

diff --git a/Engine/Plugins/Animation/DeformerGraph/Source/OptimusCore/Private/OptimusDeformerInstance.cpp b/Engine/Plugins/Animation/DeformerGraph/Source/OptimusCore/Private/OptimusDeformerInstance.cpp
--- a/Engine/Plugins/Animation/DeformerGraph/Source/OptimusCore/Private/OptimusDeformerInstance.cpp
+++ b/Engine/Plugins/Animation/DeformerGraph/Source/OptimusCore/Private/OptimusDeformerInstance.cpp
@@ -3,6 +3,8 @@
 #include "OptimusDeformerInstance.h"
 
 #include "Components/MeshComponent.h"
+#include "ComputeFramework/ComputeFramework.h"
+#include "ComputeWorkerInterface.h"
 #include "DataInterfaces/OptimusDataInterfaceGraph.h"
 #include "OptimusComputeGraph.h"
 #include "OptimusDataTypeRegistry.h"
@@ -427,21 +429,47 @@ bool UOptimusDeformerInstance::IsActive() const
 	return !ComputeGraphExecInfos.IsEmpty();
 }
 
-void UOptimusDeformerInstance::EnqueueWork(FSceneInterface* InScene, EWorkLoad InWorkLoadType, FName InOwnerName)
+void UOptimusDeformerInstance::EnqueueWork(FSceneInterface* InScene, EWorkLoad InWorkLoadType, EExectutionGroup InExecutionGroup, FName InOwnerName)
 {
+	// Convert execution group enum to ComputeTaskExecutionGroup name.
+	FName ExecutionGroupName;
+	switch (InExecutionGroup)
+	{
+	case UMeshDeformerInstance::ExecutionGroup_Immediate:
+		ExecutionGroupName = ComputeTaskExecutionGroup::Immediate;
+		break;
+	case UMeshDeformerInstance::ExecutionGroup_Default:
+	case UMeshDeformerInstance::ExecutionGroup_EndOfFrameUodate:
+		ExecutionGroupName = ComputeTaskExecutionGroup::EndOfFrameUpdate;
+		break;
+	default:
+		ensure(0);
+		return;
+	}
+
+	// Get the current queued graphs.
 	TSet<FName> GraphsToRun;
 	{
 		UE::TScopeLock<FCriticalSection> Lock(GraphsToRunOnNextTickLock);
 		Swap(GraphsToRunOnNextTick, GraphsToRun);
 	}
 	
+	// Enqueue work.
+	bool bIsWorkEnqueued = false;
 	for (FOptimusDeformerInstanceExecInfo& Info: ComputeGraphExecInfos)
 	{
 		if (Info.GraphType == EOptimusNodeGraphType::Update || GraphsToRun.Contains(Info.GraphName))
 		{
-			Info.ComputeGraphInstance.EnqueueWork(Info.ComputeGraph, InScene, InOwnerName);
+			Info.ComputeGraphInstance.EnqueueWork(Info.ComputeGraph, InScene, ExecutionGroupName, InOwnerName);
+			bIsWorkEnqueued = true;
 		}
 	}
+
+	// If this was enqueued to the Immediate group then flush all work on that group now.
+	if (bIsWorkEnqueued && InExecutionGroup == UMeshDeformerInstance::ExecutionGroup_Immediate)
+	{
+		ComputeFramework::FlushWork(InScene, ExecutionGroupName);
+	}
 }
 
 namespace
diff --git a/Engine/Plugins/Animation/DeformerGraph/Source/OptimusCore/Public/OptimusDeformerInstance.h b/Engine/Plugins/Animation/DeformerGraph/Source/OptimusCore/Public/OptimusDeformerInstance.h
--- a/Engine/Plugins/Animation/DeformerGraph/Source/OptimusCore/Public/OptimusDeformerInstance.h
+++ b/Engine/Plugins/Animation/DeformerGraph/Source/OptimusCore/Public/OptimusDeformerInstance.h
@@ -208,7 +208,7 @@ protected:
 	void AllocateResources() override;
 	void ReleaseResources() override;
 	bool IsActive() const override;
-	void EnqueueWork(FSceneInterface* InScene, EWorkLoad InWorkLoadType, FName InOwnerName) override;
+	void EnqueueWork(FSceneInterface* InScene, EWorkLoad InWorkLoadType, EExectutionGroup InExecutionGroup, FName InOwnerName) override;
 
 private:
 	/** The Mesh Component that owns this Mesh Deformer Instance. */
diff --git a/Engine/Plugins/Runtime/ComputeFramework/Source/ComputeFramework/Private/ComputeFramework.cpp b/Engine/Plugins/Runtime/ComputeFramework/Source/ComputeFramework/Private/ComputeFramework.cpp
--- a/Engine/Plugins/Runtime/ComputeFramework/Source/ComputeFramework/Private/ComputeFramework.cpp
+++ b/Engine/Plugins/Runtime/ComputeFramework/Source/ComputeFramework/Private/ComputeFramework.cpp
@@ -2,9 +2,13 @@
 
 #include "ComputeFramework/ComputeFramework.h"
 
+#include "ComputeFramework/ComputeFrameworkModule.h"
 #include "ComputeFramework/ComputeGraph.h"
+#include "ComputeFramework/ComputeGraphWorker.h"
 #include "ComputeFramework/ComputeKernelFromText.h"
 #include "ComputeFramework/ComputeKernelShaderCompilationManager.h"
+#include "ComputeFramework/ComputeSystem.h"
+#include "RenderGraphBuilder.h"
 #include "ShaderCore.h"
 #include "UObject/UObjectIterator.h"
 
@@ -58,4 +62,19 @@ namespace ComputeFramework
 		GComputeKernelShaderCompilationManager.Tick(DeltaSeconds);
 #endif // WITH_EDITOR
 	}
+
+	void FlushWork(FSceneInterface const* InScene, FName InExecutionGroupName)
+	{
+		FComputeGraphTaskWorker* ComputeGraphWorker = FComputeFrameworkModule::GetComputeSystem()->GetComputeWorker(InScene);
+		if (ensure(ComputeGraphWorker))
+		{
+			ENQUEUE_RENDER_COMMAND(ComputeFrameworkFlushCommand)(
+				[ComputeGraphWorker, InExecutionGroupName](FRHICommandListImmediate& RHICmdList)
+				{
+					FRDGBuilder GraphBuilder(RHICmdList);
+					ComputeGraphWorker->SubmitWork(GraphBuilder, InExecutionGroupName, GMaxRHIFeatureLevel);
+					GraphBuilder.Execute();
+				});
+		}
+	}
 }
diff --git a/Engine/Plugins/Runtime/ComputeFramework/Source/ComputeFramework/Private/ComputeFramework/ComputeGraphWorker.h b/Engine/Plugins/Runtime/ComputeFramework/Source/ComputeFramework/Private/ComputeFramework/ComputeGraphWorker.h
--- a/Engine/Plugins/Runtime/ComputeFramework/Source/ComputeFramework/Private/ComputeFramework/ComputeGraphWorker.h
+++ b/Engine/Plugins/Runtime/ComputeFramework/Source/ComputeFramework/Private/ComputeFramework/ComputeGraphWorker.h
@@ -21,14 +21,16 @@ class FComputeGraphTaskWorker : public IComputeTaskWorker
 public:
 	/** Enqueue a compute graph for execution. */
 	void Enqueue(
+		FName InExecutionGroupName,
 		FName InOwnerName,
 		FComputeGraphRenderProxy const* InGraphRenderProxy, 
 		TArray<FComputeDataProviderRenderProxy*> InDataProviderRenderProxies );
 
 	/** Submit enqueued compute graph work. */
 	void SubmitWork(
-		FRHICommandListImmediate& RHICmdList,
-		ERHIFeatureLevel::Type FeatureLevel ) override;
+		FRDGBuilder& GraphBuilder,
+		FName InExecutionGroupName,
+		ERHIFeatureLevel::Type InFeatureLevel ) override;
 
 private:
 	/** Description of each graph that is enqueued. */
@@ -44,5 +46,6 @@ private:
 		~FGraphInvocation();
 	};
 
-	TArray<FGraphInvocation> GraphInvocations;
+	/** Map of enqueued work per execution group . */
+	TMap<FName, TArray<FGraphInvocation> > GraphInvocationsPerGroup;
 };
diff --git a/Engine/Plugins/Runtime/ComputeFramework/Source/ComputeFramework/Private/ComputeGraphComponent.cpp b/Engine/Plugins/Runtime/ComputeFramework/Source/ComputeFramework/Private/ComputeGraphComponent.cpp
--- a/Engine/Plugins/Runtime/ComputeFramework/Source/ComputeFramework/Private/ComputeGraphComponent.cpp
+++ b/Engine/Plugins/Runtime/ComputeFramework/Source/ComputeFramework/Private/ComputeGraphComponent.cpp
@@ -3,6 +3,7 @@
 #include "ComputeFramework/ComputeGraphComponent.h"
 
 #include "ComputeFramework/ComputeGraph.h"
+#include "ComputeWorkerInterface.h"
 #include "GameFramework/Actor.h"
 
 UComputeGraphComponent::UComputeGraphComponent()
@@ -61,5 +62,5 @@ void UComputeGraphComponent::SendRenderDynamicData_Concurrent()
 {
 	Super::SendRenderDynamicData_Concurrent();
 	
-	ComputeGraphInstance.EnqueueWork(ComputeGraph, GetScene(), GetOwner()->GetFName());
+	ComputeGraphInstance.EnqueueWork(ComputeGraph, GetScene(), ComputeTaskExecutionGroup::EndOfFrameUpdate, GetOwner()->GetFName());
 }
diff --git a/Engine/Plugins/Runtime/ComputeFramework/Source/ComputeFramework/Private/ComputeGraphInstance.cpp b/Engine/Plugins/Runtime/ComputeFramework/Source/ComputeFramework/Private/ComputeGraphInstance.cpp
--- a/Engine/Plugins/Runtime/ComputeFramework/Source/ComputeFramework/Private/ComputeGraphInstance.cpp
+++ b/Engine/Plugins/Runtime/ComputeFramework/Source/ComputeFramework/Private/ComputeGraphInstance.cpp
@@ -29,7 +29,7 @@ bool FComputeGraphInstance::ValidateDataProviders(UComputeGraph* InComputeGraph)
 	return InComputeGraph != nullptr && InComputeGraph->IsCompiled() && InComputeGraph->ValidateGraph() && InComputeGraph->ValidateProviders(DataProviders);
 }
 
-bool FComputeGraphInstance::EnqueueWork(UComputeGraph* InComputeGraph, FSceneInterface const* InScene, FName InOwnerName)
+bool FComputeGraphInstance::EnqueueWork(UComputeGraph* InComputeGraph, FSceneInterface const* InScene, FName InExecutionGroupName, FName InOwnerName)
 {
 	if (InComputeGraph == nullptr || InScene == nullptr)
 	{
@@ -74,10 +74,10 @@ bool FComputeGraphInstance::EnqueueWork(UComputeGraph* InComputeGraph, FSceneInt
 	}
 
 	ENQUEUE_RENDER_COMMAND(ComputeFrameworkEnqueueExecutionCommand)(
-		[ComputeGraphWorker, InOwnerName, GraphRenderProxy, MovedDataProviderRenderProxies = MoveTemp(DataProviderRenderProxies)](FRHICommandListImmediate& RHICmdList)
+		[ComputeGraphWorker, InExecutionGroupName, InOwnerName, GraphRenderProxy, MovedDataProviderRenderProxies = MoveTemp(DataProviderRenderProxies)](FRHICommandListImmediate& RHICmdList)
 		{
 			// Compute graph scheduler will take ownership of the provider proxies.
-			ComputeGraphWorker->Enqueue(InOwnerName, GraphRenderProxy, MovedDataProviderRenderProxies);
+			ComputeGraphWorker->Enqueue(InExecutionGroupName, InOwnerName, GraphRenderProxy, MovedDataProviderRenderProxies);
 		});
 
 	return true;
diff --git a/Engine/Plugins/Runtime/ComputeFramework/Source/ComputeFramework/Private/ComputeGraphWorker.cpp b/Engine/Plugins/Runtime/ComputeFramework/Source/ComputeFramework/Private/ComputeGraphWorker.cpp
--- a/Engine/Plugins/Runtime/ComputeFramework/Source/ComputeFramework/Private/ComputeGraphWorker.cpp
+++ b/Engine/Plugins/Runtime/ComputeFramework/Source/ComputeFramework/Private/ComputeGraphWorker.cpp
@@ -14,26 +14,25 @@
 
 DECLARE_GPU_STAT_NAMED(ComputeFramework_ExecuteBatches, TEXT("ComputeFramework::ExecuteBatches"));
 
-void FComputeGraphTaskWorker::Enqueue(FName InOwnerName, FComputeGraphRenderProxy const* InGraphRenderProxy, TArray<FComputeDataProviderRenderProxy*> InDataProviderRenderProxies)
+void FComputeGraphTaskWorker::Enqueue(FName InExecutionGroupName, FName InOwnerName, FComputeGraphRenderProxy const* InGraphRenderProxy, TArray<FComputeDataProviderRenderProxy*> InDataProviderRenderProxies)
 {
-	FGraphInvocation& GraphInvocation = GraphInvocations.AddDefaulted_GetRef();
+	FGraphInvocation& GraphInvocation = GraphInvocationsPerGroup.FindOrAdd(InExecutionGroupName).AddDefaulted_GetRef();
 	GraphInvocation.OwnerName = InOwnerName;
 	GraphInvocation.GraphRenderProxy = InGraphRenderProxy;
 	GraphInvocation.DataProviderRenderProxies = MoveTemp(InDataProviderRenderProxies);
 }
 
-void FComputeGraphTaskWorker::SubmitWork(FRHICommandListImmediate& RHICmdList, ERHIFeatureLevel::Type FeatureLevel)
+void FComputeGraphTaskWorker::SubmitWork(FRDGBuilder& GraphBuilder, FName InExecutionGroupName, ERHIFeatureLevel::Type FeatureLevel)
 {
+	TArray<FGraphInvocation>& GraphInvocations = GraphInvocationsPerGroup.FindOrAdd(InExecutionGroupName);
 	if (GraphInvocations.IsEmpty())
 	{
 		return;
 	}
 
 	{
-		SCOPED_DRAW_EVENTF(RHICmdList, ComputeFramework_ExecuteBatches, TEXT("ComputeFramework::ExecuteBatches"));
-		SCOPED_GPU_STAT(RHICmdList, ComputeFramework_ExecuteBatches);
-
-		FRDGBuilder GraphBuilder(RHICmdList);
+		SCOPED_DRAW_EVENTF(GraphBuilder.RHICmdList, ComputeFramework_ExecuteBatches, TEXT("ComputeFramework::ExecuteBatches"));
+		SCOPED_GPU_STAT(GraphBuilder.RHICmdList, ComputeFramework_ExecuteBatches);
 
 		for (int32 GraphIndex = 0; GraphIndex < GraphInvocations.Num(); ++GraphIndex)
 		{
@@ -113,16 +112,12 @@ void FComputeGraphTaskWorker::SubmitWork(FRHICommandListImmediate& RHICmdList, E
 
 		// Release any graph resources at the end of graph execution.
 		GraphBuilder.AddPass(
-			RDG_EVENT_NAME("Release Data Providers"), 
-			ERDGPassFlags::None, 
-			[this](FRHICommandList&) 
-		{
-			GraphInvocations.Reset(); 
-		});
-
-		// Execute graph.
-		// todo[CF]: We can pull this out into calling code so that graph can be shared with other work.
-		GraphBuilder.Execute();
+			RDG_EVENT_NAME("Release Data Providers"),
+			ERDGPassFlags::None,
+			[this, InExecutionGroupName](FRHICommandList&)
+			{
+				GraphInvocationsPerGroup.FindChecked(InExecutionGroupName).Reset();
+			});
 	}
 }
 
diff --git a/Engine/Plugins/Runtime/ComputeFramework/Source/ComputeFramework/Public/ComputeFramework/ComputeFramework.h b/Engine/Plugins/Runtime/ComputeFramework/Source/ComputeFramework/Public/ComputeFramework/ComputeFramework.h
--- a/Engine/Plugins/Runtime/ComputeFramework/Source/ComputeFramework/Public/ComputeFramework/ComputeFramework.h
+++ b/Engine/Plugins/Runtime/ComputeFramework/Source/ComputeFramework/Public/ComputeFramework/ComputeFramework.h
@@ -5,6 +5,8 @@
 #include "CoreMinimal.h"
 #include "RHIDefinitions.h"
 
+class FSceneInterface;
+
 DECLARE_LOG_CATEGORY_EXTERN(LogComputeFramework, Log, All);
 
 namespace ComputeFramework
@@ -20,4 +22,7 @@ namespace ComputeFramework
 
 	/** Tick shader compilation. */
 	COMPUTEFRAMEWORK_API void TickCompilation(float DeltaSeconds);
+
+	/** Flush any enqueued ComputeGraph work for a given execution group. */
+	COMPUTEFRAMEWORK_API void FlushWork(FSceneInterface const* InScene, FName InExecutionGroupName);
 }
diff --git a/Engine/Plugins/Runtime/ComputeFramework/Source/ComputeFramework/Public/ComputeFramework/ComputeGraphInstance.h b/Engine/Plugins/Runtime/ComputeFramework/Source/ComputeFramework/Public/ComputeFramework/ComputeGraphInstance.h
--- a/Engine/Plugins/Runtime/ComputeFramework/Source/ComputeFramework/Public/ComputeFramework/ComputeGraphInstance.h
+++ b/Engine/Plugins/Runtime/ComputeFramework/Source/ComputeFramework/Public/ComputeFramework/ComputeGraphInstance.h
@@ -35,7 +35,7 @@ public:
 	TArray< TObjectPtr<UComputeDataProvider> >& GetDataProviders() { return DataProviders; }
 
 	/** Enqueue the ComputeGraph work. */
-	bool EnqueueWork(UComputeGraph* InComputeGraph, FSceneInterface const* Scene, FName InOwnerName);
+	bool EnqueueWork(UComputeGraph* InComputeGraph, FSceneInterface const* InScene, FName InExecutionGroupName, FName InOwnerName);
 
 private:
 	/** The currently bound Data Provider objects. */
diff --git a/Engine/Source/Runtime/Engine/Classes/Animation/MeshDeformerInstance.h b/Engine/Source/Runtime/Engine/Classes/Animation/MeshDeformerInstance.h
--- a/Engine/Source/Runtime/Engine/Classes/Animation/MeshDeformerInstance.h
+++ b/Engine/Source/Runtime/Engine/Classes/Animation/MeshDeformerInstance.h
@@ -37,6 +37,14 @@ public:
 		WorkLoad_Update,
 	};
 
+	/** Enumeration for exectution groups to enqueue on. */
+	enum EExectutionGroup
+	{
+		ExecutionGroup_Default,
+		ExecutionGroup_Immediate,
+		ExecutionGroup_EndOfFrameUodate,
+	};
+
 	/** Called to allocate any persistent render resources */
 	virtual void AllocateResources() PURE_VIRTUAL(, );
 
@@ -47,5 +55,11 @@ public:
 	virtual bool IsActive() const PURE_VIRTUAL(, return false;);
 	
 	/** Enqueue the mesh deformer workload on a scene. InOwnerName is used for debugging. */
-	virtual void EnqueueWork(FSceneInterface* InScene, EWorkLoad InWorkLoadType, FName InOwnerName) PURE_VIRTUAL(, );
+	virtual void EnqueueWork(FSceneInterface* InScene, EWorkLoad InWorkLoadType, EExectutionGroup InExecutionGroup, FName InOwnerName) PURE_VIRTUAL(, );
+
+	/** For backwards compatibility. To deprecate. */
+	void EnqueueWork(FSceneInterface* InScene, EWorkLoad InWorkLoadType, FName InOwnerName)
+	{
+		EnqueueWork(InScene, InWorkLoadType, ExecutionGroup_Default, InOwnerName);
+	}
 };
diff --git a/Engine/Source/Runtime/Engine/Classes/Components/SkinnedMeshComponent.h b/Engine/Source/Runtime/Engine/Classes/Components/SkinnedMeshComponent.h
--- a/Engine/Source/Runtime/Engine/Classes/Components/SkinnedMeshComponent.h
+++ b/Engine/Source/Runtime/Engine/Classes/Components/SkinnedMeshComponent.h
@@ -707,9 +707,6 @@ protected:
 	/** Whether we have updated bone visibility this tick */
 	uint8 bBoneVisibilityDirty:1;
 
-	/** Whether mesh deformer state is dirty and will need updating at the next tick. */
-	uint8 bUpdateDeformerAtNextTick : 1;
-
 private:
 	/** If true, UpdateTransform will always result in a call to MeshObject->Update. */
 	UPROPERTY(transient)
diff --git a/Engine/Source/Runtime/Engine/Private/Components/SkinnedMeshComponent.cpp b/Engine/Source/Runtime/Engine/Private/Components/SkinnedMeshComponent.cpp
--- a/Engine/Source/Runtime/Engine/Private/Components/SkinnedMeshComponent.cpp
+++ b/Engine/Source/Runtime/Engine/Private/Components/SkinnedMeshComponent.cpp
@@ -412,8 +412,6 @@ USkinnedMeshComponent::USkinnedMeshComponent(const FObjectInitializer& ObjectIni
 	bNeedToFlipSpaceBaseBuffers = false;
 	bBoneVisibilityDirty = false;
 
-	bUpdateDeformerAtNextTick = false;
-
 	bCanEverAffectNavigation = false;
 	LeaderBoneMapCacheCount = 0;
 	bSyncAttachParentLOD = true;
@@ -904,14 +902,6 @@ void USkinnedMeshComponent::CreateRenderState_Concurrent(FRegisterComponentConte
 				}
 			}
 
-			if (MeshDeformerInstance)
-			{
-				MeshDeformerInstance->AllocateResources();
-				
-				// After creating the MeshObject we need to run the MeshDeformer at least once to set up the vertex data.
-				bUpdateDeformerAtNextTick = true;
-			}
-
 			//Allow the editor a chance to manipulate it before its added to the scene
 			PostInitMeshObject(MeshObject);
 		}
@@ -965,6 +955,16 @@ void USkinnedMeshComponent::CreateRenderState_Concurrent(FRegisterComponentConte
 		// scene proxy update of material usage based on active morphs
 		UpdateMorphMaterialUsageOnProxy();
 	}
+ 
+ 	if (MeshDeformerInstance)
+ 	{
+		MeshDeformerInstance->AllocateResources();
+		if (MeshDeformerInstance->IsActive())
+		{
+			// Enqueue immediate execution of work here to ensure that we have some deformer outputs written for the next frame.
+			MeshDeformerInstance->EnqueueWork(GetScene(), UMeshDeformerInstance::WorkLoad_Update, UMeshDeformerInstance::ExecutionGroup_Immediate, GetSkinnedAsset()->GetFName());
+		}
+	}
 }
 
 void USkinnedMeshComponent::DestroyRenderState_Concurrent()
@@ -1079,7 +1079,7 @@ void USkinnedMeshComponent::SendRenderDynamicData_Concurrent()
 		{
 			if (MeshDeformerInstance->IsActive())
 			{
-				MeshDeformerInstance->EnqueueWork(GetScene(), UMeshDeformerInstance::WorkLoad_Update, GetSkinnedAsset()->GetFName());
+				MeshDeformerInstance->EnqueueWork(GetScene(), UMeshDeformerInstance::WorkLoad_Update, UMeshDeformerInstance::ExecutionGroup_Default, GetSkinnedAsset()->GetFName());
 			}
 			else
 			{
@@ -1306,12 +1306,6 @@ void USkinnedMeshComponent::TickComponent(float DeltaTime, enum ELevelTick TickT
 		}
 	}
 #endif // WITH_EDITOR
-
-	if (bUpdateDeformerAtNextTick)
-	{
-		MarkRenderDynamicDataDirty();
-		bUpdateDeformerAtNextTick = false;
-	}
 }
 
 UObject const* USkinnedMeshComponent::AdditionalStatObject() const
@@ -1914,7 +1908,7 @@ bool USkinnedMeshComponent::IsSkinCacheAllowed(int32 LodIdx) const
 	static const IConsoleVariable* CVarDefaultGPUSkinCacheBehavior = IConsoleManager::Get().FindConsoleVariable(TEXT("r.SkinCache.DefaultBehavior"));
 	const bool bGlobalDefault = CVarDefaultGPUSkinCacheBehavior && ESkinCacheDefaultBehavior(CVarDefaultGPUSkinCacheBehavior->GetInt()) == ESkinCacheDefaultBehavior::Inclusive;
 
-	if (HasMeshDeformer())
+	if (MeshDeformerInstance != nullptr)
 	{
 		// Disable skin cache if a mesh deformer is in use.
 		// Any animation buffers are expected to be owned by the MeshDeformer.
diff --git a/Engine/Source/Runtime/Engine/Private/ComputeWorkerInterface.cpp b/Engine/Source/Runtime/Engine/Private/ComputeWorkerInterface.cpp
new file mode 100644
--- /dev/null
+++ b/Engine/Source/Runtime/Engine/Private/ComputeWorkerInterface.cpp
@@ -0,0 +1,6 @@
+// Copyright Epic Games, Inc. All Rights Reserved.
+
+#include "ComputeWorkerInterface.h"
+
+FName ComputeTaskExecutionGroup::Immediate("Immediate");
+FName ComputeTaskExecutionGroup::EndOfFrameUpdate("EndOfFrameUpdate");
diff --git a/Engine/Source/Runtime/Engine/Private/LevelTick.cpp b/Engine/Source/Runtime/Engine/Private/LevelTick.cpp
--- a/Engine/Source/Runtime/Engine/Private/LevelTick.cpp
+++ b/Engine/Source/Runtime/Engine/Private/LevelTick.cpp
@@ -67,6 +67,7 @@
 #include "ProfilingDebugging/RealtimeGPUProfiler.h"
 #include "GPUSkinCache.h"
 #include "ComputeWorkerInterface.h"
+#include "RenderGraphBuilder.h"
 
 #if WITH_EDITOR
 #include "Editor.h"
@@ -1032,7 +1033,9 @@ void EndSendEndOfFrameUpdatesDrawEvent(FSendAllEndOfFrameUpdates& SendAllEndOfFr
 				SCOPED_GPU_STAT(RHICmdList, ComputeTaskWorkerUpdates);
 				for (IComputeTaskWorker* ComputeTaskWorker : ComputeTaskWorkers)
 				{
-					ComputeTaskWorker->SubmitWork(RHICmdList, FeatureLevel);
+					FRDGBuilder GraphBuilder(RHICmdList);
+					ComputeTaskWorker->SubmitWork(GraphBuilder, ComputeTaskExecutionGroup::EndOfFrameUpdate, FeatureLevel);
+					GraphBuilder.Execute();
 				}
 			}
 		});
diff --git a/Engine/Source/Runtime/Engine/Private/SkeletalRenderGPUSkin.cpp b/Engine/Source/Runtime/Engine/Private/SkeletalRenderGPUSkin.cpp
--- a/Engine/Source/Runtime/Engine/Private/SkeletalRenderGPUSkin.cpp
+++ b/Engine/Source/Runtime/Engine/Private/SkeletalRenderGPUSkin.cpp
@@ -2536,7 +2536,7 @@ void FDynamicSkelMeshObjectDataGPUSkin::InitDynamicSkelMeshObjectDataGPUSkin(
 	}
 
 	bIsSkinCacheAllowed = InMeshComponent ? InMeshComponent->IsSkinCacheAllowed(InLODIndex) : false;
-	bHasMeshDeformer = InMeshComponent ? InMeshComponent->HasMeshDeformer() : false;
+	bHasMeshDeformer = InMeshComponent ? InMeshComponent->GetMeshDeformerInstance() != nullptr : false;
 
 	if (bIsSkinCacheAllowed && InMeshObject->FeatureLevel == ERHIFeatureLevel::ES3_1)
 	{
diff --git a/Engine/Source/Runtime/Engine/Public/ComputeWorkerInterface.h b/Engine/Source/Runtime/Engine/Public/ComputeWorkerInterface.h
--- a/Engine/Source/Runtime/Engine/Public/ComputeWorkerInterface.h
+++ b/Engine/Source/Runtime/Engine/Public/ComputeWorkerInterface.h
@@ -14,5 +14,12 @@ public:
 	virtual ~IComputeTaskWorker() {}
 
 	/** Add any scheduled work to an RDGBuilder ready for execution. */
-	virtual void SubmitWork(class FRHICommandListImmediate& RHIComdList, ERHIFeatureLevel::Type FeatureLevel) = 0;
+	virtual void SubmitWork(class FRDGBuilder& GraphBuilder, FName InExecutionGroupName, ERHIFeatureLevel::Type FeatureLevel) = 0;
+};
+
+/** Core execution group names for use in IComputeTaskWorker::SubmitWork(). */
+struct ENGINE_API ComputeTaskExecutionGroup
+{
+	static FName Immediate;
+	static FName EndOfFrameUpdate;
 };
