Content Bundle: null pointer check until we figure out a good repro.

#localization none
#tests tried to repro but without luck, but narrowing down the cause slowly

diff --git a/Engine/Source/Runtime/Engine/Private/WorldPartition/ContentBundle/ContentBundleEditor.cpp b/Engine/Source/Runtime/Engine/Private/WorldPartition/ContentBundle/ContentBundleEditor.cpp
--- a/Engine/Source/Runtime/Engine/Private/WorldPartition/ContentBundle/ContentBundleEditor.cpp
+++ b/Engine/Source/Runtime/Engine/Private/WorldPartition/ContentBundle/ContentBundleEditor.cpp
@@ -573,9 +573,13 @@ void UContentBundleUnsavedActorMonitor::Uninitialize()
 {
 	StopListeningOnActorEvents();
 
-	for (TWeakObjectPtr<AActor>& Actor : UnsavedActors)
+	for (TWeakObjectPtr<AActor>& ActorPtr : UnsavedActors)
 	{
-		ContentBundle->GetInjectedWorld()->DestroyActor(Actor.Get());
+		// @todo_ow: figure out how can this happen
+		if (AActor* Actor = ActorPtr.Get())
+		{
+			ContentBundle->GetInjectedWorld()->DestroyActor(Actor);
+		}
 	}
 	UnsavedActors.Empty();
 
