Add a cvar to allow disabling shadows on HLOD based on platforms scalability settings

diff --git a/Engine/Source/Runtime/Engine/Private/WorldPartition/HLOD/HLODActor.cpp b/Engine/Source/Runtime/Engine/Private/WorldPartition/HLOD/HLODActor.cpp
--- a/Engine/Source/Runtime/Engine/Private/WorldPartition/HLOD/HLODActor.cpp
+++ b/Engine/Source/Runtime/Engine/Private/WorldPartition/HLOD/HLODActor.cpp
@@ -22,6 +22,13 @@
 #include "Engine/TextureStreamingTypes.h"
 #endif
 
+static int32 GWorldPartitionHLODForceDisableShadows = 0;
+static FAutoConsoleVariableRef CVarWorldPartitionHLODForceDisableShadows(
+	TEXT("wp.Runtime.HLOD.ForceDisableShadows"),
+	GWorldPartitionHLODForceDisableShadows,
+	TEXT("Force disable CastShadow flag on World Partition HLOD actors"),
+	ECVF_Scalability);
+
 #if WITH_EDITORONLY_DATA
 FArchive& operator<<(FArchive& Ar, FHLODSubActorDesc& SubActor)
 {
@@ -107,12 +114,19 @@ bool AWorldPartitionHLOD::NeedsLoadForServer() const
 	return GetIsReplicated();
 }
 
-#if WITH_EDITOR
-
 void AWorldPartitionHLOD::PostLoad()
 {
 	Super::PostLoad();
 
+	if (GWorldPartitionHLODForceDisableShadows && GetWorld() && GetWorld()->IsGameWorld())
+	{
+		ForEachComponent<UPrimitiveComponent>(false, [](UPrimitiveComponent* PrimitiveComponent)
+		{
+			PrimitiveComponent->SetCastShadow(false);
+		});
+	}
+
+#if WITH_EDITOR
 	if (GetLinkerCustomVersion(FFortniteMainBranchObjectVersion::GUID) < FFortniteMainBranchObjectVersion::WorldPartitionStreamingCellsNamingShortened)
 	{
 		if (!HLODSubActors.IsEmpty())
@@ -134,8 +148,11 @@ void AWorldPartitionHLOD::PostLoad()
 
 	// Update the disk size stat on load, as we can't really know it when saving
 	HLODStats.Add(FWorldPartitionHLODStats::MemoryDiskSizeBytes, FHLODActorDesc::GetPackageSize(this));
+#endif
 }
 
+#if WITH_EDITOR
+
 void AWorldPartitionHLOD::RerunConstructionScripts()
 {
 }
