[BugFix] CIS UE NonUnity Error: DiffResults.h(117): error C2653: 'FAppStyle': is not a class or namespace name

diff --git a/Engine/Source/Runtime/Engine/Public/DiffResults.h b/Engine/Source/Runtime/Engine/Public/DiffResults.h
--- a/Engine/Source/Runtime/Engine/Public/DiffResults.h
+++ b/Engine/Source/Runtime/Engine/Public/DiffResults.h
@@ -4,6 +4,7 @@
 
 #include "CoreMinimal.h"
 #include "Styling/StyleColors.h"
+#include "Styling/AppStyle.h"
 
 /** Differences found within a graph or object */
 namespace EDiffType
