Raising the default value of net.MaxSerializedReplayNetGuids

diff --git a/Engine/Source/Runtime/Engine/Private/PackageMapClient.cpp b/Engine/Source/Runtime/Engine/Private/PackageMapClient.cpp
--- a/Engine/Source/Runtime/Engine/Private/PackageMapClient.cpp
+++ b/Engine/Source/Runtime/Engine/Private/PackageMapClient.cpp
@@ -65,7 +65,7 @@ namespace UE
 		int32 MaxSerializedNetGuids = 2048;
 		static FAutoConsoleVariableRef CVarMaxSerializedNetGuids(TEXT("net.MaxSerializedNetGuids"), MaxSerializedNetGuids, TEXT("Maximum number of network guids we would expect to receive in a bunch"));
 
-		int32 MaxSerializedReplayNetGuids = 16 * 1024;
+		int32 MaxSerializedReplayNetGuids = 32 * 1024;
 		static FAutoConsoleVariableRef CVarMaxSerializedReplayNetGuids(TEXT("net.MaxSerializedReplayNetGuids"), MaxSerializedReplayNetGuids, TEXT("Maximum number of network guids we would expect to receive in replay export data."));
 
 		int32 MaxSerializedNetExportGroups = 64 * 1024;
