Shrink huge maps after AssetManager scanned all paths.

#tests replay run

#localization none

diff --git a/Engine/Source/Runtime/Engine/Private/AssetManager.cpp b/Engine/Source/Runtime/Engine/Private/AssetManager.cpp
--- a/Engine/Source/Runtime/Engine/Private/AssetManager.cpp
+++ b/Engine/Source/Runtime/Engine/Private/AssetManager.cpp
@@ -985,6 +985,16 @@ void UAssetManager::StopBulkScanning()
 	GetAssetRegistry().SetTemporaryCachingMode(bOldTemporaryCachingMode);
 	
 	RebuildObjectReferenceList();
+
+	// Shrink some big containers because we've finished modifying them.
+	AssetPathMap.Shrink();
+	AssetTypeMap.Shrink();
+	for (const auto& Pair : AssetTypeMap)
+	{
+		FPrimaryAssetTypeData& AssetTypeData = *Pair.Value;
+		AssetTypeData.AssetMap.Shrink();
+		AssetTypeData.RealAssetScanPaths.Shrink();
+	}
 }
 
 bool UAssetManager::RegisterSpecificPrimaryAsset(const FPrimaryAssetId& PrimaryAssetId, const FAssetData& NewAssetData)
