[UE-110638] Blend profiles in blend spaces

* Added support for using blend profiles for per bone overrides in blend spaces.
* Backward compatibility to previous blend space assets.
* Manual per bone overrides are remembered when changing between blend profiles and manual mode.
* Combo box to change between the modes with automatic hiding of the per bone or blend profile properties.

diff --git a/Engine/Config/BaseEngine.ini b/Engine/Config/BaseEngine.ini
--- a/Engine/Config/BaseEngine.ini
+++ b/Engine/Config/BaseEngine.ini
@@ -1099,6 +1099,7 @@ DefaultWorldPartitionHLODLayer=/Engine/Maps/Templates/HLODs/HLODLayer_Instanced.
 ; 5.2
 +StructRedirects=(OldName="/Script/AssetTools.RevisionInfo", NewName="/Script/AssetDefinition.RevisionInfo")
 +PropertyRedirects=(OldName="DecalComponent.SetFadeIn.Duaration", NewName="DecalComponent.SetFadeIn.Duration")
++PropertyRedirects=(OldName="BlendSpace.PerBoneBlend", NewName="BlendSpace.ManualPerBoneOverrides")
 +ClassRedirects=(OldName="/Script/MovieRenderPipelineCore.MoviePipelineMasterConfig", NewName="/Script/MovieRenderPipelineCore.MoviePipelinePrimaryConfig")
 
 [CoreUObject.Metadata]
diff --git a/Engine/Source/Runtime/Engine/Classes/Animation/BlendProfile.h b/Engine/Source/Runtime/Engine/Classes/Animation/BlendProfile.h
--- a/Engine/Source/Runtime/Engine/Classes/Animation/BlendProfile.h
+++ b/Engine/Source/Runtime/Engine/Classes/Animation/BlendProfile.h
@@ -87,6 +87,13 @@ public:
 	*/
 	void RemoveEntry(int32 InBoneIdx);
 
+	/**
+	 * Get the bone entry by entry index.
+	 * @param[in] InEntryIdx The index to the bone entry in range [0, GetNumBlendEntries()-1].
+	 * @return The bone entry containing the bone reference and blend scale.
+	 **/
+	const FBlendProfileBoneEntry& GetEntry(const int32 InEntryIdx) const;
+
 	/** Get the set blend scale for the specified bone, will return 1.0f if no entry was found (no scale)
 	 *  @param InBoneIdx Index of the bone to retrieve
 	 */
diff --git a/Engine/Source/Runtime/Engine/Classes/Animation/BlendSpace.h b/Engine/Source/Runtime/Engine/Classes/Animation/BlendSpace.h
--- a/Engine/Source/Runtime/Engine/Classes/Animation/BlendSpace.h
+++ b/Engine/Source/Runtime/Engine/Classes/Animation/BlendSpace.h
@@ -56,6 +56,24 @@ enum class EPreferredTriangulationDirection : uint8
 	Radial UMETA(DisplayName = "Radial", ToolTip = "When there is ambiguity, rectangles will be split so that the inserted edge tends to point towards the origin")
 };
 
+UENUM()
+enum class EBlendSpacePerBoneBlendMode : uint8
+{
+	ManualPerBoneOverride UMETA(DisplayName = "Manual Per Bone Override", ToolTip = "Manually specify the bones and their smoothing interpolation times."),
+	BlendProfile UMETA(DisplayName = "Blend Profile", ToolTip = "Use a blend profile to specify the bone smoothing interpolation times.")
+};
+
+USTRUCT()
+struct FBlendSpaceBlendProfile
+{
+	GENERATED_BODY()
+
+	UPROPERTY(EditAnywhere, Category = SampleSmoothing, meta = (DisplayName = "Blend Profile", UseAsBlendMask = true, EditConditionHides))
+	TObjectPtr<UBlendProfile> BlendProfile;
+
+	UPROPERTY(EditAnywhere, Category = SampleSmoothing, meta = (DisplayName = "Weight Speed", ClampMin = "0"))
+	float TargetWeightInterpolationSpeedPerSec;
+};
 
 USTRUCT()
 struct FInterpolationParameter
@@ -792,13 +810,31 @@ public:
 protected:
 
 	/**
-	* Per bone sample smoothing settings, which affect the specified bone and all its descendants in the skeleton.
-	* These act as overrides to the global sample smoothing speed, which means the global sample smoothing speed does
-	* not affect these bones. Note that they also override each other - so a per-bone setting on the chest will not
-	* affect the hand if there is a per-bone setting on the arm.
-	*/
-	UPROPERTY(EditAnywhere, Category = SampleSmoothing, meta = (DisplayName="Per Bone Overrides"))
-	TArray<FPerBoneInterpolation> PerBoneBlend;
+	 * There are two ways to use per pone sample smoothing: Blend profiles and manually maintaining the per bone overrides.
+	 */
+	UPROPERTY(EditAnywhere, Category = SampleSmoothing)
+	EBlendSpacePerBoneBlendMode PerBoneBlendMode = EBlendSpacePerBoneBlendMode::ManualPerBoneOverride;
+
+	/**
+	 * Per bone sample smoothing settings, which affect the specified bone and all its descendants in the skeleton.
+	 * These act as overrides to the global sample smoothing speed, which means the global sample smoothing speed does
+	 * not affect these bones. Note that they also override each other - so a per-bone setting on the chest will not
+	 * affect the hand if there is a per-bone setting on the arm.
+	 */
+	UPROPERTY(EditAnywhere, Category = SampleSmoothing, meta = (DisplayName="Per Bone Overrides", EditCondition = "PerBoneBlendMode == EBlendSpacePerBoneBlendMode::ManualPerBoneOverride", EditConditionHides))
+	TArray<FPerBoneInterpolation> ManualPerBoneOverrides;
+
+	/**
+	 * Reference to a blend profile of the corresponding skeleton to be used for per bone smoothing in case the per bone blend mode is set to use a blend profile.
+	 **/
+	UPROPERTY(EditAnywhere, Category = SampleSmoothing, meta = (DisplayName = "Per Bone Overrides", EditCondition = "PerBoneBlendMode == EBlendSpacePerBoneBlendMode::BlendProfile", EditConditionHides))
+	FBlendSpaceBlendProfile PerBoneBlendProfile;
+
+	/**
+	 * Stores the actual bone references and their smoothing interpolation speeds used by the blend space. This will be either filled by the manual per bone overrides or
+	 * the blend profile, depending on the set per bone blend mode.
+	 **/
+	TArray<FPerBoneInterpolation> PerBoneBlendValues;
 
 	/** Track index to get marker data from. Samples are tested for the suitability of marker based sync
 	    during load and if we can use marker based sync we cache an index to a representative sample here */
diff --git a/Engine/Source/Runtime/Engine/Private/Animation/BlendProfile.cpp b/Engine/Source/Runtime/Engine/Private/Animation/BlendProfile.cpp
--- a/Engine/Source/Runtime/Engine/Private/Animation/BlendProfile.cpp
+++ b/Engine/Source/Runtime/Engine/Private/Animation/BlendProfile.cpp
@@ -51,6 +51,11 @@ void UBlendProfile::RemoveEntry(int32 InBoneIdx)
 		});
 }
 
+const FBlendProfileBoneEntry& UBlendProfile::GetEntry(const int32 InEntryIdx) const
+{
+	return ProfileEntries[InEntryIdx];
+}
+
 float UBlendProfile::GetBoneBlendScale(int32 InBoneIdx) const
 {
 	const FBlendProfileBoneEntry* FoundEntry = ProfileEntries.FindByPredicate([InBoneIdx](const FBlendProfileBoneEntry& Entry)
diff --git a/Engine/Source/Runtime/Engine/Private/Animation/BlendSpace.cpp b/Engine/Source/Runtime/Engine/Private/Animation/BlendSpace.cpp
--- a/Engine/Source/Runtime/Engine/Private/Animation/BlendSpace.cpp
+++ b/Engine/Source/Runtime/Engine/Private/Animation/BlendSpace.cpp
@@ -5,6 +5,7 @@
 =============================================================================*/ 
 
 #include "Animation/BlendSpace.h"
+#include "Animation/BlendProfile.h"
 #include "Animation/AnimNotifyQueue.h"
 #include "AnimationUtils.h"
 #include "Animation/BlendSpaceUtilities.h"
@@ -12,6 +13,7 @@
 #include "Animation/AttributesRuntime.h"
 #include "Animation/BlendSpaceHelpers.h"
 #include "Animation/BlendSpace1DHelpers.h"
+#include "Math/UnrealMathUtility.h"
 #include "UObject/FrameworkObjectVersion.h"
 #include "UObject/UE5MainStreamObjectVersion.h"
 #include "UObject/UObjectIterator.h"
@@ -152,8 +154,11 @@ void UBlendSpace::PostEditChangeProperty(struct FPropertyChangedEvent& PropertyC
 	const FName MemberPropertyName = PropertyChangedEvent.MemberProperty ? PropertyChangedEvent.MemberProperty->GetFName() : NAME_None;
 	const FName PropertyName = PropertyChangedEvent.Property ? PropertyChangedEvent.Property->GetFName() : NAME_None;
 
-	if (MemberPropertyName == GET_MEMBER_NAME_CHECKED(UBlendSpace, PerBoneBlend) ||
-		PropertyName == GET_MEMBER_NAME_CHECKED(FPerBoneInterpolation, InterpolationSpeedPerSec))
+	if (MemberPropertyName == GET_MEMBER_NAME_CHECKED(UBlendSpace, ManualPerBoneOverrides) ||
+		PropertyName == GET_MEMBER_NAME_CHECKED(FPerBoneInterpolation, InterpolationSpeedPerSec) ||
+		MemberPropertyName == GET_MEMBER_NAME_CHECKED(UBlendSpace, PerBoneBlendProfile) ||
+		MemberPropertyName == GET_MEMBER_NAME_CHECKED(UBlendSpace, PerBoneBlendMode) ||
+		MemberPropertyName == GET_MEMBER_NAME_CHECKED(UBlendSpace, TargetWeightInterpolationSpeedPerSec))
 	{
 		InitializePerBoneBlend();
 	}
@@ -244,7 +249,7 @@ bool UBlendSpace::UpdateBlendSamples_Internal(
 	if (GetSamplesFromBlendInput(InBlendSpacePosition, NewSampleDataList, InOutCachedTriangulationIndex, true))
 	{
 		// if target weight interpolation is set
-		if (TargetWeightInterpolationSpeedPerSec > 0.f || PerBoneBlend.Num() > 0)
+		if (TargetWeightInterpolationSpeedPerSec > 0.f || PerBoneBlendValues.Num() > 0)
 		{
 			// target weight interpolation
 			if (InterpolateWeightOfSampleData(InDeltaTime, InOutOldSampleDataList, NewSampleDataList, InOutSampleDataCache))
@@ -731,10 +736,10 @@ struct FSortedPerBoneInterpolationData : public IInterpolationIndexProvider::FPe
 TSharedPtr<IInterpolationIndexProvider::FPerBoneInterpolationData> UBlendSpace::GetPerBoneInterpolationData(const USkeleton* RuntimeSkeleton) const 
 {
 	FSortedPerBoneInterpolationData* Data = new FSortedPerBoneInterpolationData();
-	Data->Data.SetNumUninitialized(PerBoneBlend.Num());
-	for (int32 Iter = 0 ; Iter != PerBoneBlend.Num() ; ++Iter)
+	Data->Data.SetNumUninitialized(PerBoneBlendValues.Num());
+	for (int32 Iter = 0 ; Iter != PerBoneBlendValues.Num() ; ++Iter)
 	{
-		Data->Data[Iter] = FSortedPerBoneInterpolation(PerBoneBlend[Iter], Iter);
+		Data->Data[Iter] = FSortedPerBoneInterpolation(PerBoneBlendValues[Iter], Iter);
 		if (!IsAsset())
 		{
 			FBoneReference& Bone = Data->Data[Iter].PerBoneBlend.BoneReference;
@@ -911,7 +916,7 @@ void UBlendSpace::GetAnimationPose_Internal(TArray<FBlendSampleData>& BlendSampl
 
 	TArrayView<FCompactPose> ChildrenPosesView(ChildrenPoses);
 
-	if (PerBoneBlend.Num() > 0)
+	if (PerBoneBlendValues.Num() > 0)
 	{
 		bool bValidAdditive = IsValidAdditive();
 
@@ -1516,10 +1521,66 @@ bool UBlendSpace::IsTooCloseToExistingSamplePoint(const FVector& SampleValue, in
 
 #endif // WITH_EDITOR
 
+// When using CriticallyDampedSmoothing, how to go from the interpolation speed to the smooth
+// time? What would the critically damped velocity be as it goes from a starting value of 0 to a
+// target of 1 (see eq in CriticallyDampedSmoothing), starting with v = 0?
+//
+// v = W^2 t exp(-W t)
+//
+// Differentiate and set equal to zero to find maximum v is at t = 1 / W
+//
+// vMax = W / e = 2 / (SmoothingTime * e)
+//
+// Set this equal to TargetWeightInterpolationSpeedPerSec, we get
+//
+// SmoothingTime = 2 / (e * TargetWeightInterpolationSpeedPerSec)
+//
+// However - this looks significantly slower than when using a constant speed, because we're
+// easing in/out, so aim for twice this speed (i.e. smooth over half the time)
+static float SmoothingTimeFromSpeed(float Speed)
+{
+	return Speed > FLT_EPSILON ? 1.0f / (UE_EULERS_NUMBER * Speed) : 0.0f;
+}
+
+static float SpeedFromSmoothingTime(float SmoothingTime)
+{
+	return SmoothingTime > FLT_EPSILON ? 1 / (UE_EULERS_NUMBER * SmoothingTime) : 0.0f;
+}
+
 void UBlendSpace::InitializePerBoneBlend()
 {
+	if (PerBoneBlendMode == EBlendSpacePerBoneBlendMode::ManualPerBoneOverride)
+	{
+		PerBoneBlendValues = ManualPerBoneOverrides;
+	}
+	else
+	{
+		PerBoneBlendValues.Empty();
+
+		UBlendProfile* BlendProfile = PerBoneBlendProfile.BlendProfile.Get();
+		if (BlendProfile)
+		{
+			int32 NumBlendEntries = BlendProfile->GetNumBlendEntries();
+			for (int32 EntryIndex = 0; EntryIndex < NumBlendEntries; ++EntryIndex)
+			{
+				const FBlendProfileBoneEntry& BoneEntry = BlendProfile->GetEntry(EntryIndex);
+
+				FPerBoneInterpolation BoneInterpolation;
+				BoneInterpolation.BoneReference = BoneEntry.BoneReference;
+
+				const float TargetWeightInterpolationTime = SmoothingTimeFromSpeed(TargetWeightInterpolationSpeedPerSec);
+				const float BlendProfileInterpolationTime = SmoothingTimeFromSpeed(PerBoneBlendProfile.TargetWeightInterpolationSpeedPerSec);
+
+				const float InterpolatedTime = FMath::Lerp(TargetWeightInterpolationTime, BlendProfileInterpolationTime, FMath::Clamp(BoneEntry.BlendScale, 0.0f, 1.0f));
+				BoneInterpolation.InterpolationSpeedPerSec = SpeedFromSmoothingTime(InterpolatedTime);
+
+				PerBoneBlendValues.Add(BoneInterpolation);
+			}
+		}
+	}
+
 	const USkeleton* MySkeleton = GetSkeleton();
-	for (FPerBoneInterpolation& BoneInterpolationData : PerBoneBlend)
+	for (FPerBoneInterpolation& BoneInterpolationData : PerBoneBlendValues)
 	{
 		BoneInterpolationData.Initialize(MySkeleton);
 	}
@@ -1627,27 +1688,6 @@ const FEditorElement* UBlendSpace::GetGridSampleInternal(int32 Index) const
 	return GridSamples.IsValidIndex(Index) ? &GridSamples[Index] : nullptr;
 }
 
-// When using CriticallyDampedSmoothing, how to go from the interpolation speed to the smooth
-// time? What would the critically damped velocity be as it goes from a starting value of 0 to a
-// target of 1 (see eq in CriticallyDampedSmoothing), starting with v = 0?
-//
-// v = W^2 t exp(-W t)
-//
-// Differentiate and set equal to zero to find maximum v is at t = 1 / W
-//
-// vMax = W / e = 2 / (SmoothingTime * e)
-//
-// Set this equal to TargetWeightInterpolationSpeedPerSec, we get
-//
-// SmoothingTime = 2 / (e * TargetWeightInterpolationSpeedPerSec)
-//
-// However - this looks significantly slower than when using a constant speed, because we're
-// easing in/out, so aim for twice this speed (i.e. smooth over half the time)
-static float SmoothingTimeFromSpeed(float Speed)
-{
-	return 1.0f / (UE_EULERS_NUMBER * Speed);
-}
-
 static void SmoothWeight(float& Output, float& OutputRate, float Input, float InputRate, float Target, float DeltaTime, float Speed, bool bUseEaseInOut)
 {
 	if (Speed <= 0.0f)
@@ -1680,10 +1720,10 @@ bool UBlendSpace::InterpolateWeightOfSampleData(float DeltaTime, const TArray<FB
 		FBlendSampleData OldSample = *OldIt;
 		bool bTargetSampleExists = false;
 
-		if (OldSample.PerBoneBlendData.Num() != PerBoneBlend.Num())
+		if (OldSample.PerBoneBlendData.Num() != PerBoneBlendValues.Num())
 		{
-			OldSample.PerBoneBlendData.Init(OldSample.TotalWeight, PerBoneBlend.Num());
-			OldSample.PerBoneWeightRate.Init(OldSample.WeightRate, PerBoneBlend.Num());
+			OldSample.PerBoneBlendData.Init(OldSample.TotalWeight, PerBoneBlendValues.Num());
+			OldSample.PerBoneWeightRate.Init(OldSample.WeightRate, PerBoneBlendValues.Num());
 		}
 
 		for (auto NewIt = NewSampleDataList.CreateConstIterator(); NewIt; ++NewIt)
@@ -1704,7 +1744,7 @@ bool UBlendSpace::InterpolateWeightOfSampleData(float DeltaTime, const TArray<FB
 					SmoothWeight(
 						InterpData.PerBoneBlendData[Iter], InterpData.PerBoneWeightRate[Iter],
 						OldSample.PerBoneBlendData[Iter], OldSample.PerBoneWeightRate[Iter], NewSample.TotalWeight,
-						DeltaTime, PerBoneBlend[Iter].InterpolationSpeedPerSec, bTargetWeightInterpolationEaseInOut);
+						DeltaTime, PerBoneBlendValues[Iter].InterpolationSpeedPerSec, bTargetWeightInterpolationEaseInOut);
 					TotalPerBoneWeight += InterpData.PerBoneBlendData[Iter];
 				}
 
@@ -1731,7 +1771,7 @@ bool UBlendSpace::InterpolateWeightOfSampleData(float DeltaTime, const TArray<FB
 				SmoothWeight(
 					InterpData.PerBoneBlendData[Iter], InterpData.PerBoneWeightRate[Iter],
 					OldSample.PerBoneBlendData[Iter], OldSample.PerBoneWeightRate[Iter], 0.0f,
-					DeltaTime, PerBoneBlend[Iter].InterpolationSpeedPerSec, bTargetWeightInterpolationEaseInOut);
+					DeltaTime, PerBoneBlendValues[Iter].InterpolationSpeedPerSec, bTargetWeightInterpolationEaseInOut);
 				TotalPerBoneWeight += InterpData.PerBoneBlendData[Iter];
 			}
 
@@ -1752,10 +1792,10 @@ bool UBlendSpace::InterpolateWeightOfSampleData(float DeltaTime, const TArray<FB
 		FBlendSampleData OldSample = *OldIt;
 		bool bOldSampleExists = false;
 
-		if (OldSample.PerBoneBlendData.Num() != PerBoneBlend.Num())
+		if (OldSample.PerBoneBlendData.Num() != PerBoneBlendValues.Num())
 		{
-			OldSample.PerBoneBlendData.Init(OldSample.TotalWeight, PerBoneBlend.Num());
-			OldSample.PerBoneWeightRate.Init(OldSample.WeightRate, PerBoneBlend.Num());
+			OldSample.PerBoneBlendData.Init(OldSample.TotalWeight, PerBoneBlendValues.Num());
+			OldSample.PerBoneWeightRate.Init(OldSample.WeightRate, PerBoneBlendValues.Num());
 		}
 
 		for (auto NewIt = FinalSampleDataList.CreateConstIterator(); NewIt; ++NewIt)
@@ -1786,7 +1826,7 @@ bool UBlendSpace::InterpolateWeightOfSampleData(float DeltaTime, const TArray<FB
 				SmoothWeight(
 					InterpData.PerBoneBlendData[Iter], InterpData.PerBoneWeightRate[Iter],
 					OldSample.PerBoneBlendData[Iter], OldSample.PerBoneWeightRate[Iter], Target,
-					DeltaTime, PerBoneBlend[Iter].InterpolationSpeedPerSec, bTargetWeightInterpolationEaseInOut);
+					DeltaTime, PerBoneBlendValues[Iter].InterpolationSpeedPerSec, bTargetWeightInterpolationEaseInOut);
 				TotalPerBoneWeight += InterpData.PerBoneBlendData[Iter];
 			}
 			if (InterpData.TotalWeight > ZERO_ANIMWEIGHT_THRESH || TotalPerBoneWeight > ZERO_ANIMWEIGHT_THRESH)
