Remove logging when attempting to find most recent hardware device info for a player

#rnx

diff --git a/Engine/Source/Runtime/Engine/Private/GameFramework/InputDeviceSubsystem.cpp b/Engine/Source/Runtime/Engine/Private/GameFramework/InputDeviceSubsystem.cpp
--- a/Engine/Source/Runtime/Engine/Private/GameFramework/InputDeviceSubsystem.cpp
+++ b/Engine/Source/Runtime/Engine/Private/GameFramework/InputDeviceSubsystem.cpp
@@ -277,8 +277,7 @@ FHardwareDeviceIdentifier UInputDeviceSubsystem::GetMostRecentlyUsedHardwareDevi
 	{
 		return *FoundDevice;
 	}
-	
-	UE_LOG(LogInputDeviceProperties, Warning, TEXT("Could not determine the hardware device last used by platform user %d"), InUserId.GetInternalId());
+		
 	return FHardwareDeviceIdentifier::Invalid;
 }
 
@@ -297,4 +296,4 @@ void UInputDeviceSubsystem::SetMostRecentlyUsedHardwareDevice(const FInputDevice
 	{
 		OnInputHardwareDeviceChangedNative.Broadcast(OwningUserId);
 	}
-}
\ No newline at end of file
+}
