Strata - fixed static variable shadowing.

diff --git a/Engine/Source/Runtime/Engine/Private/ShaderCompiler/ShaderCompiler.cpp b/Engine/Source/Runtime/Engine/Private/ShaderCompiler/ShaderCompiler.cpp
--- a/Engine/Source/Runtime/Engine/Private/ShaderCompiler/ShaderCompiler.cpp
+++ b/Engine/Source/Runtime/Engine/Private/ShaderCompiler/ShaderCompiler.cpp
@@ -6296,8 +6296,8 @@ void GlobalBeginCompileShader(
 				Input.Environment.SetDefine(TEXT("USE_ACHROMATIC_BXDF_ENERGY"), 1u);
 			}
 
-			static const auto CVarGBufferFormat = IConsoleManager::Get().FindTConsoleVariableDataInt(TEXT("r.GBufferFormat"));
-			const uint32 StrataNormalQuality = CVarGBufferFormat && CVarGBufferFormat->GetValueOnAnyThread() > 1 ? 1 : 0;
+			static const auto CVarStrataGBufferFormat = IConsoleManager::Get().FindTConsoleVariableDataInt(TEXT("r.GBufferFormat"));
+			const uint32 StrataNormalQuality = CVarStrataGBufferFormat && CVarStrataGBufferFormat->GetValueOnAnyThread() > 1 ? 1 : 0;
 			Input.Environment.SetDefine(TEXT("STRATA_NORMAL_QUALITY"), StrataNormalQuality);
 
 			static FShaderPlatformCachedIniValue<int32> CVarStrataTileCoord8Bits(TEXT("r.Strata.TileCoord8bits"));
diff --git a/Engine/Source/Runtime/RenderCore/Private/Shader.cpp b/Engine/Source/Runtime/RenderCore/Private/Shader.cpp
--- a/Engine/Source/Runtime/RenderCore/Private/Shader.cpp
+++ b/Engine/Source/Runtime/RenderCore/Private/Shader.cpp
@@ -1924,8 +1924,8 @@ void ShaderMapAppendKeyString(EShaderPlatform Platform, FString& KeyString)
 			KeyString += FString::Printf(TEXT("_ROUGHDIFF"));
 		}
 
-		static const auto CVarGBufferFormat = IConsoleManager::Get().FindTConsoleVariableDataInt(TEXT("r.GBufferFormat"));
-		const uint32 StrataNormalQuality = CVarGBufferFormat && CVarGBufferFormat->GetValueOnAnyThread() > 1 ? 1 : 0;
+		static const auto CVarStrataGBufferFormat = IConsoleManager::Get().FindTConsoleVariableDataInt(TEXT("r.GBufferFormat"));
+		const uint32 StrataNormalQuality = CVarStrataGBufferFormat && CVarStrataGBufferFormat->GetValueOnAnyThread() > 1 ? 1 : 0;
 		if (bStrataEnabled && StrataNormalQuality > 0)
 		{
 			KeyString += FString::Printf(TEXT("_STRTNRMQ"));
diff --git a/Engine/Source/Runtime/Renderer/Private/Strata/Strata.cpp b/Engine/Source/Runtime/Renderer/Private/Strata/Strata.cpp
--- a/Engine/Source/Runtime/Renderer/Private/Strata/Strata.cpp
+++ b/Engine/Source/Runtime/Renderer/Private/Strata/Strata.cpp
@@ -376,8 +376,8 @@ static void RecordStrataAnalytics()
 
 static EPixelFormat GetTopLayerTextureFormat()
 {
-	static const auto CVarGBufferFormat = IConsoleManager::Get().FindTConsoleVariableDataInt(TEXT("r.GBufferFormat"));
-	return CVarGBufferFormat && CVarGBufferFormat->GetValueOnAnyThread() > 1 ? PF_R32G32_UINT : PF_R32_UINT;
+	static const auto CVarStrataGBufferFormat = IConsoleManager::Get().FindTConsoleVariableDataInt(TEXT("r.GBufferFormat"));
+	return CVarStrataGBufferFormat && CVarStrataGBufferFormat->GetValueOnAnyThread() > 1 ? PF_R32G32_UINT : PF_R32_UINT;
 }
 
 void InitialiseStrataFrameSceneData(FRDGBuilder& GraphBuilder, FSceneRenderer& SceneRenderer)
