Update EngineAnalytics to allow more general config and provider changes during creation.

diff --git a/Engine/Source/Runtime/Engine/Private/EngineAnalytics.cpp b/Engine/Source/Runtime/Engine/Private/EngineAnalytics.cpp
--- a/Engine/Source/Runtime/Engine/Private/EngineAnalytics.cpp
+++ b/Engine/Source/Runtime/Engine/Private/EngineAnalytics.cpp
@@ -37,31 +37,16 @@ FSimpleMulticastDelegate FEngineAnalytics::OnInitializeEngineAnalytics;
 FSimpleMulticastDelegate FEngineAnalytics::OnShutdownEngineAnalytics;
 #endif
 
-namespace UE
+namespace UE::Analytics::Private
 {
-namespace Analytics
-{
-static TOptional<FString> GAPIKey;
-static TOptional<FString> GBackendEnvironment;
-}
-}
 
-void FEngineAnalytics::SetAPIKey(const FString& InAPIKey, const FString& BackendEnv)
-{
-	UE::Analytics::GAPIKey = InAPIKey;
-	UE::Analytics::GBackendEnvironment = BackendEnv;
+IEngineAnalyticsConfigOverride* EngineAnalyticsConfigOverride = nullptr;
+
 }
 
 static TSharedPtr<IAnalyticsProviderET> CreateEpicAnalyticsProvider()
 {
 	FAnalyticsET::Config Config;
-	Config.APIServerET = TEXT("https://datarouter.ol.epicgames.com/");
-
-	if (UE::Analytics::GAPIKey)
-	{
-		Config.APIKeyET = *UE::Analytics::GAPIKey;
-	}
-	else
 	{
 		// We always use the "Release" analytics account unless we're running in analytics test mode (usually with
 		// a command-line parameter), or we're an internal Epic build
@@ -74,20 +59,20 @@ static TSharedPtr<IAnalyticsProviderET> CreateEpicAnalyticsProvider()
 		FString UETypeOverride;
 		bool bHasOverride = GConfig->GetString(TEXT("Analytics"), TEXT("UE4TypeOverride"), UETypeOverride, GEngineIni);
 		const TCHAR* UETypeStr = bHasOverride ? *UETypeOverride : FEngineBuildSettings::IsPerforceBuild() ? TEXT("Perforce") : TEXT("UnrealEngine");
-		Config.APIKeyET = FString::Printf(TEXT("UEEditor.%s.%s"), UETypeStr, BuildTypeStr);
-	}
 
-	if (UE::Analytics::GBackendEnvironment)
-	{
-		Config.AppEnvironment = *UE::Analytics::GBackendEnvironment;
+		FString AppID;
+		GConfig->GetString(TEXT("Analytics"), TEXT("AppIdOverride"), AppID, GEditorIni);
+		Config.APIKeyET = FString::Printf(TEXT("%s.%s.%s"), AppID.IsEmpty() ? TEXT("UEEditor") : *AppID, UETypeStr, BuildTypeStr);
 	}
-	else
+	Config.APIServerET = TEXT("https://datarouter.ol.epicgames.com/");
+	Config.AppEnvironment = TEXT("datacollector-binary");
+	Config.AppVersionET = FEngineVersion::Current().ToString();
+
+	if (UE::Analytics::Private::EngineAnalyticsConfigOverride)
 	{
-		Config.AppEnvironment = TEXT("datacollector-binary");
+		UE::Analytics::Private::EngineAnalyticsConfigOverride->ApplyConfiguration(Config);
 	}
 
-	Config.AppVersionET = FEngineVersion::Current().ToString();
-
 	// Connect the engine analytics provider (if there is a configuration delegate installed)
 	return FAnalyticsET::Get().CreateAnalyticsProvider(Config);
 }
@@ -135,7 +120,14 @@ void FEngineAnalytics::Initialize()
 
 		if (Analytics.IsValid())
 		{
-			Analytics->SetUserID(FString::Printf(TEXT("%s|%s|%s"), *FPlatformMisc::GetLoginId(), *FPlatformMisc::GetEpicAccountId(), *FPlatformMisc::GetOperatingSystemId()));
+			if (UE::Analytics::Private::EngineAnalyticsConfigOverride)
+			{
+				UE::Analytics::Private::EngineAnalyticsConfigOverride->OnProviderCreated(*Analytics);
+			}
+			else
+			{
+				Analytics->SetUserID(FString::Printf(TEXT("%s|%s|%s"), *FPlatformMisc::GetLoginId(), *FPlatformMisc::GetEpicAccountId(), *FPlatformMisc::GetOperatingSystemId()));
+			}
 
 			const UGeneralProjectSettings& ProjectSettings = *GetDefault<UGeneralProjectSettings>();
 
diff --git a/Engine/Source/Runtime/Engine/Public/EngineAnalytics.h b/Engine/Source/Runtime/Engine/Public/EngineAnalytics.h
--- a/Engine/Source/Runtime/Engine/Public/EngineAnalytics.h
+++ b/Engine/Source/Runtime/Engine/Public/EngineAnalytics.h
@@ -56,10 +56,26 @@ public:
 
 	static ENGINE_API void LowDriveSpaceDetected();
 
-	static ENGINE_API void SetAPIKey(const FString& InAPIKey, const FString& Backend);
-
 private:
 	static bool bIsInitialized;
 	static ENGINE_API TSharedPtr<IAnalyticsProviderET> Analytics;
 };
 
+namespace UE::Analytics::Private {
+
+/**
+  * Interface for allowing changes to engine analytics configuration intended for use by internal tools.
+  * 
+  * Internal use only, not intended for licensee use.
+  */
+class IEngineAnalyticsConfigOverride
+{
+public:
+	virtual ~IEngineAnalyticsConfigOverride() = default;
+	virtual void ApplyConfiguration(FAnalyticsET::Config& Config) = 0;
+	virtual void OnProviderCreated(IAnalyticsProviderET& provider) = 0;
+};
+
+extern ENGINE_API IEngineAnalyticsConfigOverride* EngineAnalyticsConfigOverride;
+
+/* UE::Analytics::Private */ }
