Made sure that HLOD runtime grids that contains net relevant actors are flagged as not bClientOnlyVisible.
Multiple HLOD layers can point to the same runtime grid. The runtime grid must be available on the server if any of the layer requires it.

diff --git a/Engine/Source/Runtime/Engine/Private/WorldPartition/RuntimeSpatialHash/RuntimeSpatialHashHLOD.cpp b/Engine/Source/Runtime/Engine/Private/WorldPartition/RuntimeSpatialHash/RuntimeSpatialHashHLOD.cpp
--- a/Engine/Source/Runtime/Engine/Private/WorldPartition/RuntimeSpatialHash/RuntimeSpatialHashHLOD.cpp
+++ b/Engine/Source/Runtime/Engine/Private/WorldPartition/RuntimeSpatialHash/RuntimeSpatialHashHLOD.cpp
@@ -314,28 +314,31 @@ static TMap<FName, FSpatialHashRuntimeGrid> CreateHLODGrids(TMap<UHLODLayer*, in
 		if (HLODLayer->IsSpatiallyLoaded())
 		{
 			int32 HLODLevel = HLODLayerLevel.Value;
+			FName GridName = HLODLayer->GetRuntimeGrid(HLODLevel);
 
-			FSpatialHashRuntimeGrid HLODGrid;
-			HLODGrid.CellSize = HLODLayer->GetCellSize();
-			HLODGrid.LoadingRange = HLODLayer->GetLoadingRange();
-			HLODGrid.DebugColor = FLinearColor::Red;
-			HLODGrid.GridName = HLODLayer->GetRuntimeGrid(HLODLevel);
-			HLODGrid.HLODLayer = HLODLayer;
+			FSpatialHashRuntimeGrid& HLODGrid = HLODGrids.FindOrAdd(GridName);
+			if (HLODGrid.GridName.IsNone())
+			{
+				HLODGrid.CellSize = HLODLayer->GetCellSize();
+				HLODGrid.LoadingRange = HLODLayer->GetLoadingRange();
+				HLODGrid.GridName = GridName;
+				HLODGrid.HLODLayer = HLODLayer;
+				HLODGrid.bClientOnlyVisible = true;
+
+				// HLOD grid have a lower priority than standard grids
+				HLODGrid.Priority = 100 + HLODLevel;
+
+				// Setup grid debug color to match HLODColorationColors
+				const int32 LastLODColorationColorIdx = GEngine->HLODColorationColors.Num() - 1;
+				const int32 HLODColorIdx = FMath::Clamp(HLODLevel + 2, 0U, (int32)LastLODColorationColorIdx);
+				HLODGrid.DebugColor = GEngine->HLODColorationColors[HLODColorIdx];
+			}
 
 			// Temp solution to have replication working for hlod layers that have a modifier
+			// Modifiers are used for example for destructible HLODs, that need to have a server side component
 			// The real solution would probably be to check if we have any actors that are
 			// returning true for NeedsLoadForServer()
-			HLODGrid.bClientOnlyVisible = HLODLayer->GetHLODModifierClass() == nullptr;
-
-			// HLOD grid have a lower priority than standard grids
-			HLODGrid.Priority = 100 + HLODLevel;
-
-			// Setup grid debug color to match HLODColorationColors
-			const int32 LastLODColorationColorIdx = GEngine->HLODColorationColors.Num() - 1;
-			const int32 HLODColorIdx = FMath::Clamp(HLODLevel + 2, 0U, (int32)LastLODColorationColorIdx);
-			HLODGrid.DebugColor = GEngine->HLODColorationColors[HLODColorIdx];
-
-			HLODGrids.Emplace(HLODGrid.GridName, HLODGrid);
+			HLODGrid.bClientOnlyVisible &= HLODLayer->GetHLODModifierClass() == nullptr;
 		}
 	}
 
