Fix a build break caused by a missing forward declaration.

#rnx

diff --git a/Engine/Source/Runtime/Engine/Classes/Engine/SkinnedAssetCommon.h b/Engine/Source/Runtime/Engine/Classes/Engine/SkinnedAssetCommon.h
--- a/Engine/Source/Runtime/Engine/Classes/Engine/SkinnedAssetCommon.h
+++ b/Engine/Source/Runtime/Engine/Classes/Engine/SkinnedAssetCommon.h
@@ -6,18 +6,20 @@
  * Contains the shared data that is used by skinned assets
  */
 
-#include "CoreMinimal.h"
-#include "Components.h"
 #include "BoneContainer.h"
+#include "Components.h"
+#include "CoreMinimal.h"
 #include "Engine/EngineTypes.h"
 #include "PerPlatformProperties.h"
-#include "PerQualityLevelProperties.h"
 #include "SkeletalMeshReductionSettings.h"
 #include "SkinnedAssetCommon.generated.h"
 
-struct FSkeletalMeshLODGroupSettings;
-class UAnimSequence;
+
 class FSkeletalMeshLODModel;
+class UAnimSequence;
+class UMaterialInterface;
+struct FSkelMeshSection;
+struct FSkeletalMeshLODGroupSettings;
 
 
 UENUM()
@@ -71,7 +73,7 @@ struct FSectionReference
 	/** return true if it has a valid section index for LodModel parameter **/
 	ENGINE_API bool IsValidToEvaluate(const FSkeletalMeshLODModel& LodModel) const;
 
-	const struct FSkelMeshSection* GetMeshLodSection(const FSkeletalMeshLODModel& LodModel) const;
+	const FSkelMeshSection* GetMeshLodSection(const FSkeletalMeshLODModel& LodModel) const;
 	int32 GetMeshLodSectionIndex(const FSkeletalMeshLODModel& LodModel) const;
 #endif
 
@@ -257,7 +259,7 @@ struct FSkeletalMaterial
 
 	}
 
-	FSkeletalMaterial( class UMaterialInterface* InMaterialInterface
+	FSkeletalMaterial( UMaterialInterface* InMaterialInterface
 						, bool bInEnableShadowCasting = true
 						, bool bInRecomputeTangent = false
 						, FName InMaterialSlotName = NAME_None
@@ -283,7 +285,7 @@ struct FSkeletalMaterial
 	ENGINE_API friend bool operator==( const UMaterialInterface& LHS, const FSkeletalMaterial& RHS );
 
 	UPROPERTY(EditAnywhere, BlueprintReadWrite, Category=SkeletalMesh)
-	TObjectPtr<class UMaterialInterface> 	MaterialInterface;
+	TObjectPtr<UMaterialInterface> 	MaterialInterface;
 	
 	/*This name should be use by the gameplay to avoid error if the skeletal mesh Materials array topology change*/
 	UPROPERTY(EditAnywhere, BlueprintReadWrite, Category = SkeletalMesh)
