Fix static analysis issues for the texture preview code.

#rnx

diff --git a/Engine/Source/Editor/UnrealEd/Private/Texture2DPreview.cpp b/Engine/Source/Editor/UnrealEd/Private/Texture2DPreview.cpp
--- a/Engine/Source/Editor/UnrealEd/Private/Texture2DPreview.cpp
+++ b/Engine/Source/Editor/UnrealEd/Private/Texture2DPreview.cpp
@@ -101,8 +101,11 @@ public:
 			SetTextureParameter(InRHICmdList, ShaderRHI, InPageTableTexture0, GBlackTexture->TextureRHI);
 			SetTextureParameter(InRHICmdList, ShaderRHI, InPageTableTexture1, GBlackTexture->TextureRHI);
 
-			FRHISamplerState* SamplerState = bInUsePointSampling ? TStaticSamplerState<SF_Point>::GetRHI() : InTextureValue->SamplerStateRHI.GetReference();
-			SetTextureParameter(InRHICmdList, ShaderRHI, InTexture, InTextureSampler, SamplerState, InTextureValue->TextureRHI);
+			if (InTextureValue != nullptr)
+			{
+				FRHISamplerState* SamplerState = bInUsePointSampling ? TStaticSamplerState<SF_Point>::GetRHI() : InTextureValue->SamplerStateRHI.GetReference();
+				SetTextureParameter(InRHICmdList, ShaderRHI, InTexture, InTextureSampler, SamplerState, InTextureValue->TextureRHI);
+			}
 		}
 		
 		SetShaderValue(InRHICmdList, ShaderRHI,ColorWeights,InColorWeightsValue);
@@ -112,7 +115,7 @@ public:
 		// Store slice count and selected slice index for texture array
 		if (bInIsTextureArray)
 		{
-			const float NumSlicesData = InTextureValue ? static_cast<float>(InTextureValue->GetSizeZ()) : 1.0f;
+			const float NumSlicesData = InTextureValue != nullptr ? static_cast<float>(InTextureValue->GetSizeZ()) : 1.0f;
 			SetShaderValue(InRHICmdList, ShaderRHI, NumSlices, NumSlicesData);
 			SetShaderValue(InRHICmdList, ShaderRHI, SliceIndex, InSliceIndex);
 		}
diff --git a/Engine/Source/Runtime/Engine/Private/CubemapUnwrapUtils.cpp b/Engine/Source/Runtime/Engine/Private/CubemapUnwrapUtils.cpp
--- a/Engine/Source/Runtime/Engine/Private/CubemapUnwrapUtils.cpp
+++ b/Engine/Source/Runtime/Engine/Private/CubemapUnwrapUtils.cpp
@@ -152,8 +152,11 @@ void FCubemapTexturePropertiesPS::SetParameters(FRHICommandList& RHICmdList, con
 {
 	FRHIPixelShader* ShaderRHI = RHICmdList.GetBoundPixelShader();
 
-	FRHISamplerState* SamplerState = bInUsePointSampling ? TStaticSamplerState<SF_Point>::GetRHI() : InTexture->SamplerStateRHI.GetReference();
-	SetTextureParameter(RHICmdList, ShaderRHI, CubeTexture, CubeTextureSampler, SamplerState, InTexture->TextureRHI);
+	if (InTexture != nullptr)
+	{
+		FRHISamplerState* SamplerState = bInUsePointSampling ? TStaticSamplerState<SF_Point>::GetRHI() : InTexture->SamplerStateRHI.GetReference();
+		SetTextureParameter(RHICmdList, ShaderRHI, CubeTexture, CubeTextureSampler, SamplerState, InTexture->TextureRHI);
+	}
 
 	FVector4f PackedProperties0Value(InMipLevel, bInShowLongLatUnwrap ? 1.0f : -1.0f, 0, 0);
 	SetShaderValue(RHICmdList, ShaderRHI, PackedProperties0, PackedProperties0Value);
@@ -166,7 +169,7 @@ void FCubemapTexturePropertiesPS::SetParameters(FRHICommandList& RHICmdList, con
 	{
 		// GetSizeZ() returns the total number of slices stored in the platform data
 		// for a TextureCube array this value is equal to the size of the array multiplied by 6
-		const float NumSlicesData = (float)(InTexture ? FMath::Max((int32)InTexture->GetSizeZ() / 6, 1) : 1);
+		const float NumSlicesData = (float)(InTexture != nullptr ? FMath::Max((int32)InTexture->GetSizeZ() / 6, 1) : 1);
 		SetShaderValue(RHICmdList, ShaderRHI, NumSlices, NumSlicesData);
 		SetShaderValue(RHICmdList, ShaderRHI, SliceIndex, InSliceIndex);
 	}
