Dir Light - exposed ForwardShadingPriority to blueprint via a function.

diff --git a/Engine/Source/Runtime/Engine/Classes/Components/DirectionalLightComponent.h b/Engine/Source/Runtime/Engine/Classes/Components/DirectionalLightComponent.h
--- a/Engine/Source/Runtime/Engine/Classes/Components/DirectionalLightComponent.h
+++ b/Engine/Source/Runtime/Engine/Classes/Components/DirectionalLightComponent.h
@@ -312,6 +312,9 @@ class ENGINE_API UDirectionalLightComponent : public ULightComponent
 	UFUNCTION(BlueprintCallable, Category = "Rendering|Lighting")
 	void SetAtmosphereSunLightIndex(int32 NewValue);
 
+	UFUNCTION(BlueprintCallable, Category = "Rendering|Lighting")
+	void SetForwardShadingPriority(int32 NewValue);
+
 	//~ Begin ULightComponent Interface
 	virtual FVector4 GetLightPosition() const override;
 	virtual ELightComponentType GetLightType() const override;
diff --git a/Engine/Source/Runtime/Engine/Private/Components/DirectionalLightComponent.cpp b/Engine/Source/Runtime/Engine/Private/Components/DirectionalLightComponent.cpp
--- a/Engine/Source/Runtime/Engine/Private/Components/DirectionalLightComponent.cpp
+++ b/Engine/Source/Runtime/Engine/Private/Components/DirectionalLightComponent.cpp
@@ -1346,6 +1346,16 @@ void UDirectionalLightComponent::SetAtmosphereSunLightIndex(int32 NewValue)
 	}
 }
 
+void UDirectionalLightComponent::SetForwardShadingPriority(int32 NewValue)
+{
+	if (AreDynamicDataChangesAllowed()
+		&& ForwardShadingPriority != NewValue)
+	{
+		ForwardShadingPriority = FMath::Max(0, NewValue);
+		MarkRenderStateDirty();
+	}
+}
+
 void UDirectionalLightComponent::Serialize(FArchive& Ar)
 {
 	Ar.UsingCustomVersion(FUE5ReleaseStreamObjectVersion::GUID);
