Iris - Implemented globally unique NetHandles.

diff --git a/Engine/Source/Runtime/Engine/Private/Net/Iris/ReplicationSystem/ReplicationSystemUtil.cpp b/Engine/Source/Runtime/Engine/Private/Net/Iris/ReplicationSystem/ReplicationSystemUtil.cpp
--- a/Engine/Source/Runtime/Engine/Private/Net/Iris/ReplicationSystem/ReplicationSystemUtil.cpp
+++ b/Engine/Source/Runtime/Engine/Private/Net/Iris/ReplicationSystem/ReplicationSystemUtil.cpp
@@ -5,6 +5,7 @@
 #if UE_WITH_IRIS
 
 #include "EngineUtils.h"
+#include "Engine/Engine.h"
 #include "Engine/NetConnection.h"
 #include "Engine/NetDriver.h"
 #include "Engine/World.h"
@@ -18,7 +19,45 @@
 #include "Templates/Casts.h"
 #include "Net/Subsystems/NetworkSubsystem.h"
 #include "Net/Core/Misc/NetConditionGroupManager.h"
+#include "Net/Core/NetHandle/NetHandleManager.h"
 #include "GameFramework/PlayerController.h"
+#include "Templates/Function.h"
+
+namespace UE::Net::ReplicationSystemUtil
+{
+	static void ForEachReplicationSystem(const UEngine* Engine, const UWorld* World, TFunctionRef<void(UReplicationSystem*)> Function)
+	{
+		if (Engine == nullptr || World == nullptr)
+		{
+			return;
+		}
+
+		if (const FWorldContext* Context = Engine->GetWorldContextFromWorld(World))
+		{
+			for (const FNamedNetDriver& NamedNetDriver : Context->ActiveNetDrivers)
+			{
+				if (UNetDriver* NetDriver = NamedNetDriver.NetDriver)
+				{
+					if (UReplicationSystem* ReplicationSystem = NetDriver->GetReplicationSystem())
+					{
+						Function(ReplicationSystem);
+					}
+				}
+			}
+		}
+	}
+
+	static void ForEachReplicationSystem(TFunctionRef<void(UReplicationSystem*)> Function)
+	{
+		for (UReplicationSystem* ReplicationSystem : FReplicationSystemFactory::GetAllReplicationSystems())
+		{
+			if (ReplicationSystem != nullptr)
+			{
+				Function(ReplicationSystem);
+			}
+		}
+	}
+}
 
 namespace UE::Net
 {
@@ -56,187 +95,315 @@ UActorReplicationBridge* FReplicationSystemUtil::GetActorReplicationBridge(const
 
 FNetHandle FReplicationSystemUtil::GetNetHandle(const AActor* Actor)
 {
-	if (!Actor)
+	FNetHandle NetHandle = FNetHandleManager::GetNetHandle(Actor);
+	return NetHandle;
+}
+
+FNetHandle FReplicationSystemUtil::GetNetHandle(const UActorComponent* SubObject)
+{
+	FNetHandle NetHandle = FNetHandleManager::GetNetHandle(SubObject);
+	return NetHandle;
+}
+
+void FReplicationSystemUtil::BeginReplication(AActor* Actor, const FActorBeginReplicationParams& Params)
+{
+	if (const UWorld* World = Actor->GetWorld())
 	{
-		return FNetHandle();
+		ReplicationSystemUtil::ForEachReplicationSystem(GEngine, World, [Actor, &Params](UReplicationSystem* ReplicationSystem)
+		{
+			if (ReplicationSystem->IsServer())
+			{
+				if (UActorReplicationBridge* Bridge = Cast<UActorReplicationBridge>(ReplicationSystem->GetReplicationBridge()))
+				{
+					Bridge->BeginReplication(Actor, Params);
+				}
+			}
+		});
 	}
+}
 
-	// We really should cache the handle in the actor (or UObject if we can afford it)
-	UActorReplicationBridge* Bridge = GetActorReplicationBridge(Actor);
-	return Bridge ? Bridge->GetReplicatedHandle(Actor) : FNetHandle();
+void FReplicationSystemUtil::BeginReplication(AActor* Actor)
+{
+	const FActorBeginReplicationParams BeginReplicationParams;
+	return BeginReplication(Actor, BeginReplicationParams);
 }
 
-FNetHandle FReplicationSystemUtil::GetNetHandle(const UActorComponent* SubObject)
+void FReplicationSystemUtil::EndReplication(AActor* Actor, EEndPlayReason::Type EndPlayReason)
 {
-	UActorReplicationBridge* Bridge = SubObject ? GetActorReplicationBridge(SubObject->GetOwner()) : nullptr;
-	return Bridge ? Bridge->GetReplicatedHandle(SubObject) : FNetHandle();
+	// If the call is coming from for example destroying an actor then a formerly associated NetRefHandle will no longer be valid.
+	// The bridge itself will verify that the actor is replicated by it so there's no reason to check that either.
+	ReplicationSystemUtil::ForEachReplicationSystem([Actor, EndPlayReason](UReplicationSystem* ReplicationSystem)
+	{
+		if (UActorReplicationBridge* Bridge = Cast<UActorReplicationBridge>(ReplicationSystem->GetReplicationBridge()))
+		{
+			Bridge->EndReplication(Actor, EndPlayReason);
+		}
+	});
 }
 
-FNetHandle FReplicationSystemUtil::BeginReplicationForActorComponent(FNetHandle ActorHandle, UActorComponent* ActorComp)
+
+void FReplicationSystemUtil::BeginReplicationForActorComponent(FNetHandle ActorHandle, UActorComponent* ActorComp)
 {
-	if (ActorHandle.IsValid())
+	if (!ActorHandle.IsValid())
 	{
-		UReplicationSystem* ReplicationSystem = Net::GetReplicationSystem(ActorHandle.GetReplicationSystemId());
-		if (ReplicationSystem && ReplicationSystem->IsServer())
+		return;
+	}
+
+	if (!ensureAlways(ActorComp != nullptr))
+	{
+		return;
+	}
+
+	AActor* Actor = Cast<AActor>(ActorComp->GetOwner());
+	ensureAlways(FNetHandleManager::GetNetHandle(Actor) == ActorHandle);
+	if (Actor)
+	{
+		if (const UWorld* World = Actor->GetWorld())
 		{
-			if (UActorReplicationBridge* Bridge = Cast<UActorReplicationBridge>(ReplicationSystem->GetReplicationBridge()))
+			ReplicationSystemUtil::ForEachReplicationSystem(GEngine, World, [ActorHandle, ActorComp](UReplicationSystem* ReplicationSystem)
 			{
-				return Bridge->BeginReplication(ActorHandle, ActorComp);
-			}
+				if (ReplicationSystem->IsServer())
+				{
+					if (UActorReplicationBridge* Bridge = Cast<UActorReplicationBridge>(ReplicationSystem->GetReplicationBridge()))
+					{
+						const FNetRefHandle OwnerRefHandle = Bridge->GetReplicatedRefHandle(ActorHandle);
+						if (OwnerRefHandle.IsValid())
+						{
+							Bridge->BeginReplication(OwnerRefHandle, ActorComp);
+						}
+					}
+				}
+			});
 		}
 	}
-	return FNetHandle();
 }
 
-FNetHandle FReplicationSystemUtil::BeginReplicationForActorComponent(const AActor* Actor, UActorComponent* ActorComp)
+void FReplicationSystemUtil::BeginReplicationForActorComponent(const AActor* Actor, UActorComponent* ActorComp)
 {
-	return FReplicationSystemUtil::BeginReplicationForActorComponent(FReplicationSystemUtil::GetNetHandle(Actor), ActorComp);
+	const FNetHandle ActorHandle = GetNetHandle(Actor);
+	// If the actor doesn't have a valid handle we assume it's not replicated by any ReplicationSystem
+	if (!ActorHandle.IsValid())
+	{
+		return;
+	}
+
+	if (const UWorld* World = Actor->GetWorld())
+	{
+		ReplicationSystemUtil::ForEachReplicationSystem(GEngine, World, [ActorHandle, ActorComp](UReplicationSystem* ReplicationSystem)
+		{
+			if (ReplicationSystem->IsServer())
+			{
+				if (UActorReplicationBridge* Bridge = Cast<UActorReplicationBridge>(ReplicationSystem->GetReplicationBridge()))
+				{
+					const FNetRefHandle ActorRefHandle = Bridge->GetReplicatedRefHandle(ActorHandle);
+					if (ActorRefHandle.IsValid())
+					{
+						Bridge->BeginReplication(ActorRefHandle, ActorComp);
+					}
+				}
+			}
+		});
+	}
 }
 
-FNetHandle FReplicationSystemUtil::BeginReplicationForActorSubObject(const AActor* Actor, UObject* ActorSubObject, ELifetimeCondition NetCondition)
+void FReplicationSystemUtil::BeginReplicationForActorSubObject(const AActor* Actor, UObject* ActorSubObject, ELifetimeCondition NetCondition)
 {
-	if (NetCondition != ELifetimeCondition::COND_Never)
+	if (NetCondition == ELifetimeCondition::COND_Never)
+	{
+		return;
+	}
+
+	// Assume an actor without NetHandle isn't replicated.
+	const FNetHandle ActorHandle = GetNetHandle(Actor);
+	if (!ActorHandle.IsValid())
 	{
-		const FNetHandle OwnerHandle = GetNetHandle(Actor);
-		if (OwnerHandle.IsValid())
+		return;
+	}
+
+	if (const UWorld* World = Actor->GetWorld())
+	{
+		ReplicationSystemUtil::ForEachReplicationSystem(GEngine, World, [ActorHandle, ActorSubObject, NetCondition](UReplicationSystem* ReplicationSystem)
 		{
-			UReplicationSystem* ReplicationSystem = Net::GetReplicationSystem(OwnerHandle.GetReplicationSystemId());
-			if (ReplicationSystem && ReplicationSystem->IsServer())
+			if (ReplicationSystem->IsServer())
 			{
 				if (UObjectReplicationBridge* Bridge = ReplicationSystem->GetReplicationBridgeAs<UObjectReplicationBridge>())
 				{
-					const FNetHandle SubObjectHandle = Bridge->BeginReplication(OwnerHandle, ActorSubObject);
-					if (SubObjectHandle.IsValid() && NetCondition != ELifetimeCondition::COND_None)
+					const FNetRefHandle ActorRefHandle = Bridge->GetReplicatedRefHandle(ActorHandle);
+					if (ActorRefHandle.IsValid())
 					{
-						Bridge->SetSubObjectNetCondition(SubObjectHandle, NetCondition);
+						const FNetRefHandle SubObjectRefHandle = Bridge->BeginReplication(ActorRefHandle, ActorSubObject);
+						if (SubObjectRefHandle.IsValid() && NetCondition != ELifetimeCondition::COND_None)
+						{
+							Bridge->SetSubObjectNetCondition(SubObjectRefHandle, NetCondition);
+						}
 					}
-					return SubObjectHandle;
 				}
 			}
-		}
+		});
 	}
-
-	return FNetHandle();
 }
 
-FNetHandle FReplicationSystemUtil::BeginReplicationForActorComponentSubObject(UActorComponent* ActorComponent, UObject* SubObject, ELifetimeCondition NetCondition)
+void FReplicationSystemUtil::BeginReplicationForActorComponentSubObject(UActorComponent* ActorComponent, UObject* SubObject, ELifetimeCondition NetCondition)
 {
 	const AActor* Actor = ActorComponent->GetOwner();
 	if (Actor && NetCondition != ELifetimeCondition::COND_Never)
 	{		
-		const FNetHandle OwnerHandle = GetNetHandle(ActorComponent->GetOwner());
-		if (OwnerHandle.IsValid())
+		const FNetHandle ActorHandle = GetNetHandle(Actor);
+		if (!ActorHandle.IsValid())
+		{
+			return;
+		}
+
+		if (const UWorld* World = Actor->GetWorld())
 		{
-			UReplicationSystem* ReplicationSystem = Net::GetReplicationSystem(OwnerHandle.GetReplicationSystemId());
-			if (ReplicationSystem && ReplicationSystem->IsServer())
+			ReplicationSystemUtil::ForEachReplicationSystem(GEngine, World, [ActorHandle, ActorComponent, SubObject, NetCondition](UReplicationSystem* ReplicationSystem)
 			{
-				if (UObjectReplicationBridge* Bridge = ReplicationSystem->GetReplicationBridgeAs<UObjectReplicationBridge>())
+				if (ReplicationSystem->IsServer())
 				{
-					const FNetHandle ActorComponentHandle = Bridge->GetReplicatedHandle(ActorComponent);
-					if (ActorComponentHandle.IsValid())
+					if (UObjectReplicationBridge* Bridge = ReplicationSystem->GetReplicationBridgeAs<UObjectReplicationBridge>())
 					{
-						const FNetHandle SubObjectHandle = Bridge->BeginReplication(OwnerHandle, SubObject, ActorComponentHandle, UReplicationBridge::ESubObjectInsertionOrder::ReplicateWith);
-						if (NetCondition != ELifetimeCondition::COND_None)
+						const FNetRefHandle ActorRefHandle = Bridge->GetReplicatedRefHandle(ActorHandle);
+						const FNetRefHandle ActorComponentRefHandle = Bridge->GetReplicatedRefHandle(ActorComponent);
+						if (ActorRefHandle.IsValid() && ActorComponentRefHandle.IsValid())
 						{
-							Bridge->SetSubObjectNetCondition(SubObjectHandle, NetCondition);
+							const FNetRefHandle SubObjectRefHandle = Bridge->BeginReplication(ActorRefHandle, SubObject, ActorComponentRefHandle, UReplicationBridge::ESubObjectInsertionOrder::ReplicateWith);
+							if (SubObjectRefHandle.IsValid() && NetCondition != ELifetimeCondition::COND_None)
+							{
+								Bridge->SetSubObjectNetCondition(SubObjectRefHandle, NetCondition);
+							}
 						}
-						return SubObjectHandle;
 					}
 				}
-			}
+			});
 		}
 	}
-
-	return FNetHandle();
 }
 
-void FReplicationSystemUtil::EndReplicationForActorSubObject(const AActor* Actor, UObject* SubObject)
+void FReplicationSystemUtil::EndReplicationForActorComponent(UActorComponent* SubObject)
 {
-	if (UObjectReplicationBridge* Bridge = Cast<UObjectReplicationBridge>(GetActorReplicationBridge(Actor)))
+	ReplicationSystemUtil::ForEachReplicationSystem([SubObject](UReplicationSystem* ReplicationSystem)
 	{
-		Bridge->EndReplication(SubObject);
-	}
-}
-
-void FReplicationSystemUtil::EndReplicationForActorComponentSubObject(UActorComponent* ActorComponent, UObject* SubObject)
-{
-	EndReplicationForActorSubObject(ActorComponent->GetOwner(), SubObject);
+		if (UActorReplicationBridge* Bridge = Cast<UActorReplicationBridge>(ReplicationSystem->GetReplicationBridge()))
+		{
+			constexpr EEndReplicationFlags EndReplicationFlags = EEndReplicationFlags::DestroyNetHandle | EEndReplicationFlags::ClearNetPushId;
+			Bridge->EndReplicationForActorComponent(SubObject, EndReplicationFlags);
+		}
+	});
 }
 
-FNetHandle FReplicationSystemUtil::BeginReplication(AActor* Actor, const FActorBeginReplicationParams& Params)
+void FReplicationSystemUtil::EndReplicationForActorSubObject(const AActor* Actor, UObject* SubObject)
 {
-	UReplicationSystem* ReplicationSystem = GetReplicationSystem(Actor);
-	if (ReplicationSystem && ReplicationSystem->IsServer())
+	if (!IsValid(SubObject))
 	{
-		if (UActorReplicationBridge* Bridge = ReplicationSystem->GetReplicationBridgeAs<UActorReplicationBridge>())
-		{
-			return Bridge->BeginReplication(Actor, Params);
-		}
+		return;
 	}
 
-	return FNetHandle();
+	ReplicationSystemUtil::ForEachReplicationSystem([SubObject](UReplicationSystem* ReplicationSystem)
+	{
+		if (UActorReplicationBridge* Bridge = Cast<UActorReplicationBridge>(ReplicationSystem->GetReplicationBridge()))
+		{				
+			constexpr EEndReplicationFlags EndReplicationFlags = EEndReplicationFlags::Destroy | EEndReplicationFlags::DestroyNetHandle | EEndReplicationFlags::ClearNetPushId;
+			Bridge->EndReplication(SubObject, EndReplicationFlags);
+		}
+	});
 }
 
-FNetHandle FReplicationSystemUtil::BeginReplication(AActor* Actor)
+void FReplicationSystemUtil::EndReplicationForActorComponentSubObject(UActorComponent* ActorComponent, UObject* SubObject)
 {
-	const FActorBeginReplicationParams BeginReplicationParams;
-	return BeginReplication(Actor, BeginReplicationParams);
+	EndReplicationForActorSubObject(ActorComponent->GetOwner(), SubObject);
 }
 
 void FReplicationSystemUtil::AddDependentActor(const AActor* Parent, AActor* Child)
 {
 	// Can only add dependent actors on already replicating actors
 	FNetHandle ParentHandle = GetNetHandle(Parent);
-	if (ensureMsgf(ParentHandle.IsValid(), TEXT("FReplicationSystemUtil::AddDependentActor Parent %s must be replicated"), *GetPathNameSafe(Parent)))
+	if (!ensureMsgf(ParentHandle.IsValid(), TEXT("FReplicationSystemUtil::AddDependentActor Parent %s must be replicated"), *GetPathNameSafe(Parent)))
+	{
+		return;
+	}
+
+	if (const UWorld* World = Parent->GetWorld())
 	{
-		FNetHandle ChildHandle = BeginReplication(Child);
-		if (ensureMsgf(ChildHandle.IsValid(), TEXT("FReplicationSystemUtil::AddDependentActor Child %s must be replicated"), *GetPathNameSafe(Child)))
+		ReplicationSystemUtil::ForEachReplicationSystem(GEngine, World, [ParentHandle, Child](UReplicationSystem* ReplicationSystem)
 		{
-			UReplicationSystem* ReplicationSystem = Net::GetReplicationSystem(ParentHandle.GetReplicationSystemId());
-			if (UObjectReplicationBridge* Bridge = ReplicationSystem ? ReplicationSystem->GetReplicationBridgeAs<UObjectReplicationBridge>() : nullptr)
+			if (ReplicationSystem->IsServer())
 			{
-				Bridge->AddDependentObject(ParentHandle, ChildHandle);
+				if (UObjectReplicationBridge* Bridge = ReplicationSystem->GetReplicationBridgeAs<UObjectReplicationBridge>())
+				{
+					const FNetRefHandle ParentRefHandle = Bridge->GetReplicatedRefHandle(ParentHandle);
+					if (ParentRefHandle.IsValid())
+					{
+						const FNetRefHandle ChildRefHandle = Bridge->BeginReplication(Child);
+						if (ensureMsgf(ChildRefHandle.IsValid(), TEXT("FReplicationSystemUtil::AddDependentActor Child %s must be replicated"), *GetPathNameSafe(Child)))
+						{
+							Bridge->AddDependentObject(ParentRefHandle, ChildRefHandle);
+						}
+					}
+				}
 			}
-		}
+		});
 	}
 }
 
 void FReplicationSystemUtil::RemoveDependentActor(const AActor* Parent, AActor* Child)
 {
-	// Need to find replication system
-	FNetHandle ParentHandle = GetNetHandle(Parent);
-	FNetHandle ChildHandle = GetNetHandle(Child);
-	if (ParentHandle.IsValid() && ChildHandle.IsValid())
+	const FNetHandle ParentHandle = GetNetHandle(Parent);
+	const FNetHandle ChildHandle = GetNetHandle(Child);
+	if (!ParentHandle.IsValid() || !ChildHandle.IsValid())
 	{
-		UReplicationSystem* ReplicationSystem = Net::GetReplicationSystem(ChildHandle.GetReplicationSystemId());
-		if (UObjectReplicationBridge* Bridge = ReplicationSystem ? ReplicationSystem->GetReplicationBridgeAs<UObjectReplicationBridge>() : nullptr)
-		{
-			Bridge->RemoveDependentObject(ParentHandle, ChildHandle);
-		}
+		return;
 	}
-}
 
-void FReplicationSystemUtil::EndReplication(AActor* Actor, EEndPlayReason::Type EndPlayReason)
-{	
-	if (UActorReplicationBridge* Bridge = GetActorReplicationBridge(Actor))
+	if (const UWorld* World = Parent->GetWorld())
 	{
-		Bridge->EndReplication(Actor, EndPlayReason);
+		ReplicationSystemUtil::ForEachReplicationSystem(GEngine, World, [ParentHandle, ChildHandle](UReplicationSystem* ReplicationSystem)
+		{
+			if (ReplicationSystem->IsServer())
+			{
+				if (UObjectReplicationBridge* Bridge = ReplicationSystem->GetReplicationBridgeAs<UObjectReplicationBridge>())
+				{
+					const FNetRefHandle ParentRefHandle = Bridge->GetReplicatedRefHandle(ParentHandle);
+					if (ParentRefHandle.IsValid())
+					{
+						const FNetRefHandle ChildRefHandle = Bridge->GetReplicatedRefHandle(ChildHandle);
+						Bridge->RemoveDependentObject(ParentRefHandle, ChildRefHandle);
+					}
+				}
+			}
+		});
 	}
 }
 
-void FReplicationSystemUtil::EndReplicationForActorComponent(UActorComponent* SubObject)
+void FReplicationSystemUtil::SetActorComponentNetCondition(const UActorComponent* ActorComponent, ELifetimeCondition NetCondition)
 {
-	if (UActorReplicationBridge* Bridge = IsValid(SubObject) ? GetActorReplicationBridge(SubObject->GetOwner()) : nullptr)
+	if (!IsValid(ActorComponent))
 	{
-		Bridge->EndReplicationForActorComponent(SubObject);
+		return;
 	}
-}
 
-void FReplicationSystemUtil::SetActorComponentNetCondition(const UActorComponent* ActorComponent, ELifetimeCondition Condition)
-{
-	if (UActorReplicationBridge* Bridge = IsValid(ActorComponent) ? GetActorReplicationBridge(ActorComponent->GetOwner()) : nullptr)
+	const AActor* Actor = ActorComponent->GetOwner();
+	const FNetHandle ActorHandle = GetNetHandle(Actor);
+	if (!ActorHandle.IsValid())
 	{
-		Bridge->SetSubObjectNetCondition(Bridge->GetReplicatedHandle(ActorComponent), Condition);
+		return;
+	}
+
+	if (const UWorld* World = Actor->GetWorld())
+	{
+		ReplicationSystemUtil::ForEachReplicationSystem(GEngine, World, [ActorComponent, NetCondition](UReplicationSystem* ReplicationSystem)
+		{
+			if (ReplicationSystem->IsServer())
+			{
+				if (UObjectReplicationBridge* Bridge = ReplicationSystem->GetReplicationBridgeAs<UObjectReplicationBridge>())
+				{
+					const FNetRefHandle ActorComponentRefHandle = Bridge->GetReplicatedRefHandle(ActorComponent);
+					if (ActorComponentRefHandle.IsValid())
+					{
+						Bridge->SetSubObjectNetCondition(ActorComponentRefHandle, NetCondition);
+					}
+				}
+			}
+		});
 	}
 }
 
@@ -259,29 +426,55 @@ void FReplicationSystemUtil::BeginReplicationForActorsInWorld(UWorld* World)
 
 void FReplicationSystemUtil::NotifyActorDormancyChange(AActor* Actor, ENetDormancy OldDormancyState)
 {
+	FNetHandle ActorHandle = GetNetHandle(Actor);
+	if (!ActorHandle.IsValid())
+	{
+		return;
+	}
+	
 	const ENetDormancy Dormancy = Actor->NetDormancy;
 	const bool bIsPendingDormancy = (Dormancy > DORM_Awake);
 
-	FNetHandle NetHandle = GetNetHandle(Actor);
-	if (NetHandle.IsValid())
+	if (const UWorld* World = Actor->GetWorld())
 	{
-		UReplicationSystem* ReplicationSystem = Net::GetReplicationSystem(NetHandle.GetReplicationSystemId());
-		if (UObjectReplicationBridge* Bridge = ReplicationSystem ? ReplicationSystem->GetReplicationBridgeAs<UObjectReplicationBridge>() : nullptr)
+		ReplicationSystemUtil::ForEachReplicationSystem(GEngine, World, [ActorHandle, bIsPendingDormancy](UReplicationSystem* ReplicationSystem)
 		{
-			Bridge->SetObjectWantsToBeDormant(NetHandle, bIsPendingDormancy);
-		}
+			if (ReplicationSystem->IsServer())
+			{
+				if (UObjectReplicationBridge* Bridge = ReplicationSystem->GetReplicationBridgeAs<UObjectReplicationBridge>())
+				{
+					const FNetRefHandle ActorRefHandle = Bridge->GetReplicatedRefHandle(ActorHandle);
+					if (ActorRefHandle.IsValid())
+					{
+						Bridge->SetObjectWantsToBeDormant(ActorRefHandle, bIsPendingDormancy);
+					}
+				}
+			}
+		});
 	}
 }
 
 void FReplicationSystemUtil::FlushNetDormancy(AActor* Actor, bool bWasDormInitial)
 {
-	FNetHandle NetHandle = GetNetHandle(Actor);
-	if (NetHandle.IsValid())
+	FNetHandle ActorHandle = GetNetHandle(Actor);
+	if (ActorHandle.IsValid())
 	{
-		UReplicationSystem* ReplicationSystem = Net::GetReplicationSystem(NetHandle.GetReplicationSystemId());
-		if (UObjectReplicationBridge* Bridge = ReplicationSystem ? ReplicationSystem->GetReplicationBridgeAs<UObjectReplicationBridge>() : nullptr)
+		if (const UWorld* World = Actor->GetWorld())
 		{
-			Bridge->ForceUpdateWantsToBeDormantObject(NetHandle);
+			ReplicationSystemUtil::ForEachReplicationSystem(GEngine, World, [ActorHandle, bWasDormInitial](UReplicationSystem* ReplicationSystem)
+			{
+				if (ReplicationSystem->IsServer())
+				{
+					if (UObjectReplicationBridge* Bridge = ReplicationSystem->GetReplicationBridgeAs<UObjectReplicationBridge>())
+					{
+						const FNetRefHandle ActorRefHandle = Bridge->GetReplicatedRefHandle(ActorHandle);
+						if (ActorRefHandle.IsValid())
+						{
+							Bridge->ForceUpdateWantsToBeDormantObject(ActorRefHandle);
+						}
+					}
+				}
+			});
 		}
 	}
 	else
@@ -300,58 +493,85 @@ void FReplicationSystemUtil::UpdateSubObjectGroupMemberships(const UObject* SubO
 {
 	if (const UNetworkSubsystem* NetSubsystem = World ? World->GetSubsystem<UNetworkSubsystem>() : nullptr)
 	{
-		const UNetDriver* NetDriver = World->GetNetDriver();
-		if (UReplicationSystem* ReplicationSystem = NetDriver ? NetDriver->GetReplicationSystem() : nullptr)
+		FObjectKey SubObjectKey(SubObject);
+		ReplicationSystemUtil::ForEachReplicationSystem(GEngine, World, [NetSubsystem, SubObject, SubObjectKey](UReplicationSystem* ReplicationSystem)
 		{
-			if (UActorReplicationBridge* Bridge = ReplicationSystem->GetReplicationBridgeAs<UActorReplicationBridge>())
+			if (ReplicationSystem->IsServer())
 			{
-				FNetHandle Handle = Bridge->GetReplicatedHandle(SubObject);
-				if (Handle.IsValid())
+				if (UObjectReplicationBridge* Bridge = ReplicationSystem->GetReplicationBridgeAs<UObjectReplicationBridge>())
 				{
-					for (const FName NetGroup : NetSubsystem->GetNetConditionGroupManager().GetSubObjectNetConditionGroups(FObjectKey(SubObject)))
+					FNetRefHandle RefHandle = Bridge->GetReplicatedRefHandle(SubObject);
+					if (RefHandle.IsValid())
 					{
-						FNetObjectGroupHandle SubObjectGroupHandle = ReplicationSystem->GetOrCreateSubObjectFilter(NetGroup);
-						ReplicationSystem->AddToGroup(SubObjectGroupHandle, Handle);
+						for (const FName NetGroup : NetSubsystem->GetNetConditionGroupManager().GetSubObjectNetConditionGroups(SubObjectKey))
+						{
+							FNetObjectGroupHandle SubObjectGroupHandle = ReplicationSystem->GetOrCreateSubObjectFilter(NetGroup);
+							ReplicationSystem->AddToGroup(SubObjectGroupHandle, RefHandle);
+						}
 					}
 				}
 			}
-		}
+		});
 	}
 }
 
 void FReplicationSystemUtil::RemoveSubObjectGroupMembership(const APlayerController* PC, const FName NetGroup)
 {
+	if (IsSpecialNetConditionGroup(NetGroup))
+	{
+		return;
+	}
+	
+	// We assume the player controller is tied to a single connection.
 	UNetConnection* Conn = PC->GetNetConnection();
 	if (ensureAlways(Conn))
 	{
 		if (UReplicationSystem* ReplicationSystem = Conn->GetDriver() ? Conn->GetDriver()->GetReplicationSystem() : nullptr)
 		{
-			if (!IsSpecialNetConditionGroup(NetGroup))
-			{
-				ReplicationSystem->SetSubObjectFilterStatus(NetGroup, Conn->GetConnectionId(), ENetFilterStatus::Disallow);
-			}
+			ReplicationSystem->SetSubObjectFilterStatus(NetGroup, Conn->GetConnectionId(), ENetFilterStatus::Disallow);
 		}
 	}
 }
 
 void FReplicationSystemUtil::UpdateSubObjectGroupMemberships(const APlayerController* PC)
 {
+	// We assume the player controller is tied to a single connection.
 	if (UNetConnection* Conn = PC->GetNetConnection())
 	{
 		if (UReplicationSystem* ReplicationSystem = Conn->GetDriver() ? Conn->GetDriver()->GetReplicationSystem() : nullptr)
 		{
+			const uint32 ConnId = Conn->GetConnectionId();
 			for (const FName NetGroup : PC->GetNetConditionGroups())
 			{
 				if (!IsSpecialNetConditionGroup(NetGroup))
 				{
 					FNetObjectGroupHandle SubObjectGroupHandle = ReplicationSystem->GetOrCreateSubObjectFilter(NetGroup);
-					ReplicationSystem->SetSubObjectFilterStatus(NetGroup, Conn->GetConnectionId(), ENetFilterStatus::Allow);
+					ReplicationSystem->SetSubObjectFilterStatus(NetGroup, ConnId, ENetFilterStatus::Allow);
 				}
 			}
 		}
 	}
 }
 
+void FReplicationSystemUtil::SetReplicationCondition(FNetHandle NetHandle, EReplicationCondition Condition, bool bEnableCondition)
+{
+	ReplicationSystemUtil::ForEachReplicationSystem([NetHandle, Condition, bEnableCondition](UReplicationSystem* ReplicationSystem)
+	{
+		if (ReplicationSystem->IsServer())
+		{
+			if (UObjectReplicationBridge* Bridge = ReplicationSystem->GetReplicationBridgeAs<UObjectReplicationBridge>())
+			{
+				FNetRefHandle RefHandle = Bridge->GetReplicatedRefHandle(NetHandle);
+				if (RefHandle.IsValid())
+				{
+					ReplicationSystem->SetReplicationCondition(RefHandle, Condition, bEnableCondition);
+				}
+			}
+		}
+	});
+}
+
+
 }
 
 #endif
diff --git a/Engine/Source/Runtime/Engine/Private/NetDriver.cpp b/Engine/Source/Runtime/Engine/Private/NetDriver.cpp
--- a/Engine/Source/Runtime/Engine/Private/NetDriver.cpp
+++ b/Engine/Source/Runtime/Engine/Private/NetDriver.cpp
@@ -768,12 +768,16 @@ FNetworkObjectInfo* UNetDriver::FindNetworkObjectInfo(const AActor* InActor)
 #if UE_WITH_IRIS
 	if (ReplicationSystem)
 	{
-		// Return default initialized NetworkObjectInfo so the presence of a pointer can at least indicate whether the actor is replicated or not.
-		const UE::Net::FNetHandle NetHandle = UE::Net::FReplicationSystemUtil::GetNetHandle(InActor);
-		if (NetHandle.IsValid())
+		if (UObjectReplicationBridge* Bridge = ReplicationSystem->GetReplicationBridgeAs<UObjectReplicationBridge>())
 		{
-			*DummyNetworkObjectInfo = FNetworkObjectInfo();
-			return DummyNetworkObjectInfo.Get();
+			const UE::Net::FNetRefHandle ActorRefHandle = Bridge->GetReplicatedRefHandle(InActor);
+			
+			// Return default initialized NetworkObjectInfo so the presence of a pointer can at least indicate whether the actor is replicated or not.
+			if (ActorRefHandle.IsValid())
+			{
+				*DummyNetworkObjectInfo = FNetworkObjectInfo();
+				return DummyNetworkObjectInfo.Get();
+			}
 		}
 
 		return nullptr;
@@ -940,13 +944,6 @@ void UNetDriver::TickFlush(float DeltaSeconds)
 		{
 			Connection->Tick(DeltaSeconds);
 		}
-
-#if UE_WITH_IRIS
-		if (ReplicationSystem && (ServerConnection != nullptr || ClientConnections.Num() > 0))
-		{
-			ReplicationSystem->PostSendUpdate();
-		}
-#endif // UE_WITH_IRIS
 	}
 
 	if (ConnectionlessHandler.IsValid())
@@ -1315,6 +1312,13 @@ void UNetDriver::ProcessLocalClientPackets()
 
 void UNetDriver::PostTickFlush()
 {
+#if UE_WITH_IRIS
+	if (ReplicationSystem && (ServerConnection != nullptr || ClientConnections.Num() > 0))
+	{
+		ReplicationSystem->PostSendUpdate();
+	}
+#endif // UE_WITH_IRIS
+
 	if (World && !bSkipClearVoicePackets)
 	{
 		UOnlineEngineInterface::Get()->ClearVoicePackets(World);
@@ -3532,8 +3536,12 @@ void UNetDriver::NotifyActorTearOff(AActor* Actor)
 #if UE_WITH_IRIS
 	if (ReplicationSystem)
 	{
-		// Set the actor to be torn-off during the next update of the replication systeem
-		ReplicationSystem->TearOffNextUpdate(UE::Net::FReplicationSystemUtil::GetNetHandle(Actor));
+		if (UObjectReplicationBridge* Bridge = ReplicationSystem->GetReplicationBridgeAs<UObjectReplicationBridge>())
+		{
+			// Set the actor to be torn-off during the next update of the replication systeem
+			const UE::Net::FNetRefHandle ActorRefHandle = Bridge->GetReplicatedRefHandle(Actor);
+			ReplicationSystem->TearOffNextUpdate(ActorRefHandle);
+		}
 	}
 #endif // UE_WITH_IRIS
 }
@@ -3550,7 +3558,11 @@ void UNetDriver::ForceNetUpdate(AActor* Actor)
 #if UE_WITH_IRIS
 	if (ReplicationSystem)
 	{
-		ReplicationSystem->MarkDirty(UE::Net::FReplicationSystemUtil::GetNetHandle(Actor));
+		if (UObjectReplicationBridge* Bridge = ReplicationSystem->GetReplicationBridgeAs<UObjectReplicationBridge>())
+		{
+			const UE::Net::FNetRefHandle ActorRefHandle = Bridge->GetReplicatedRefHandle(Actor);
+			ReplicationSystem->MarkDirty(ActorRefHandle);
+		}
 		return;
 	}
 #endif // UE_WITH_IRIS
diff --git a/Engine/Source/Runtime/Engine/Private/RepLayout.cpp b/Engine/Source/Runtime/Engine/Private/RepLayout.cpp
--- a/Engine/Source/Runtime/Engine/Private/RepLayout.cpp
+++ b/Engine/Source/Runtime/Engine/Private/RepLayout.cpp
@@ -29,6 +29,7 @@
 #include "UObject/EnumProperty.h"
 #include "UObject/UnrealType.h"
 #if UE_WITH_IRIS
+#include "Iris/IrisConfig.h"
 #include "Net/Iris/ReplicationSystem/ReplicationSystemUtil.h"
 #endif // UE_WITH_IRIS
 
diff --git a/Engine/Source/Runtime/Engine/Public/Net/Iris/ReplicationSystem/ActorReplicationBridge.h b/Engine/Source/Runtime/Engine/Public/Net/Iris/ReplicationSystem/ActorReplicationBridge.h
--- a/Engine/Source/Runtime/Engine/Public/Net/Iris/ReplicationSystem/ActorReplicationBridge.h
+++ b/Engine/Source/Runtime/Engine/Public/Net/Iris/ReplicationSystem/ActorReplicationBridge.h
@@ -59,7 +59,7 @@ public:
 	inline UNetDriver* GetNetDriver() const { return NetDriver; }
 
 	/** Begin replication of an actor and its registered ActorComponents and SubObjects. */
-	ENGINE_API FNetHandle BeginReplication(AActor* Instance, const FActorBeginReplicationParams& Params);
+	ENGINE_API FNetRefHandle BeginReplication(AActor* Instance, const FActorBeginReplicationParams& Params);
 
 	/** Stop replicating an actor. Will destroy handle for actor and registered subobjects. */
 	ENGINE_API void EndReplication(AActor* Actor, EEndPlayReason::Type EndPlayReason);
@@ -68,24 +68,27 @@ public:
 	 * Begin replication of an ActorComponent and its registered SubObjects, 
 	 * if the ActorComponent already is replicated any set NetObjectConditions will be updated.
 	*/
-	ENGINE_API FNetHandle BeginReplication(FNetHandle OwnerHandle, UActorComponent* ActorComponent);
+	ENGINE_API FNetRefHandle BeginReplication(FNetRefHandle OwnerHandle, UActorComponent* ActorComponent);
 
 	/** Stop replicating an ActorComponent and its associated SubObjects. */
-	ENGINE_API void EndReplicationForActorComponent(UActorComponent* ActorComponent);
+	ENGINE_API void EndReplicationForActorComponent(UActorComponent* ActorComponent, EEndReplicationFlags EndReplicationFlags = EEndReplicationFlags::None);
 	
 	/** Get object reference packagemap. Used in special cases where serialization hasn't been converted to use NetSerializers.  */
 	ENGINE_API UIrisObjectReferencePackageMap* GetObjectReferencePackageMap() const { return ObjectReferencePackageMap; }
 	
+	using UObjectReplicationBridge::EndReplication;
+
 protected:
+
 	// UObjectReplicationBridge
 	virtual void Initialize(UReplicationSystem* ReplicationSystem) override;
-	virtual bool WriteCreationHeader(UE::Net::FNetSerializationContext& Context, FNetHandle Handle) override;
+	virtual bool WriteCreationHeader(UE::Net::FNetSerializationContext& Context, FNetRefHandle Handle) override;
 	virtual FCreationHeader* ReadCreationHeader(UE::Net::FNetSerializationContext& Context) override;
-	virtual UObject* BeginInstantiateFromRemote(FNetHandle SubObjectOwnerNetHandle, const UE::Net::FNetObjectResolveContext& ResolveContext, const FCreationHeader* InHeader) override;
+	virtual UObject* BeginInstantiateFromRemote(FNetRefHandle SubObjectOwnerNetHandle, const UE::Net::FNetObjectResolveContext& ResolveContext, const FCreationHeader* InHeader) override;
 	virtual bool OnInstantiatedFromRemote(UObject* Instance, const FCreationHeader* InHeader, uint32 ConnectionId) const override;
-	virtual void EndInstantiateFromRemote(FNetHandle Handle) override;
+	virtual void EndInstantiateFromRemote(FNetRefHandle Handle) override;
 	virtual void DestroyInstanceFromRemote(UObject* Instance, bool bTearOff) override;
-	virtual void GetInitialDependencies(FNetHandle Handle, FNetDependencyInfoArray& OutDependencies) const override;
+	virtual void GetInitialDependencies(FNetRefHandle Handle, FNetDependencyInfoArray& OutDependencies) const override;
 	virtual bool RemapPathForPIE(uint32 ConnectionId, FString& Path, bool bReading) const override;
 	virtual bool ObjectLevelHasFinishedLoading(UObject* Object) const override;
 
diff --git a/Engine/Source/Runtime/Engine/Public/Net/Iris/ReplicationSystem/ReplicationSystemUtil.h b/Engine/Source/Runtime/Engine/Public/Net/Iris/ReplicationSystem/ReplicationSystemUtil.h
--- a/Engine/Source/Runtime/Engine/Public/Net/Iris/ReplicationSystem/ReplicationSystemUtil.h
+++ b/Engine/Source/Runtime/Engine/Public/Net/Iris/ReplicationSystem/ReplicationSystemUtil.h
@@ -6,7 +6,8 @@
 
 #if UE_WITH_IRIS
 
-#include "Iris/ReplicationSystem/NetHandle.h"
+#include "Iris/ReplicationSystem/Conditionals/ReplicationCondition.h"
+#include "Net/Core/NetHandle/NetHandle.h"
 #include "Engine/EngineTypes.h"
 
 class AActor;
@@ -25,51 +26,47 @@ namespace UE::Net
 // Helper methods to interact with ReplicationSystem from engine code
 struct FReplicationSystemUtil
 {
-	/** Returns the UReplicationSystem the actor is replicated by, or null if it isn't. */
+	/** Returns the UReplicationSystem for the main NetDriver assigned to the Actor. Note that an Actor may be replicated by multiple ReplicationSystems. May return null. */
 	ENGINE_API static UReplicationSystem* GetReplicationSystem(const AActor* Actor);
 
-	/** Returns the UActorReplicationBridge of the UReplicationSystem the actor is replicated by, or null if it isn't. */
+	/** Returns the UActorReplicationBridge of the UReplicationSystem belogning to the main NetDriver assigned to the Actor. May return null. */
 	ENGINE_API static UActorReplicationBridge* GetActorReplicationBridge(const AActor* Actor);
 
 	/** Returns the UReplicationSystem of the UNetDriver the UNetConnection belongs to. */
 	ENGINE_API static UActorReplicationBridge* GetActorReplicationBridge(const UNetConnection* NetConnection);
 
-	/** Get the NetHandle for the actor. */
+	/** Returns the NetHandle for the actor. The returned handle may be invalid. */
 	ENGINE_API static FNetHandle GetNetHandle(const AActor* Actor);
 
-	/** Get the NetHandle for an actor component. */
+	/** Returns the NetHandle for an actor component. The returned handle may be invalid. */
 	ENGINE_API static FNetHandle GetNetHandle(const UActorComponent* SubObject);
 
-	/** 
-		Begins replication an actor and all of its registered subobjects returns a NetHandle for the actor if it succeeds. 
-	*/
-	ENGINE_API static FNetHandle BeginReplication(AActor* Actor, const FActorBeginReplicationParams& Params);
+	/** Begins replication of an actor and all of its registered subobjects. If any ReplicationSystem wants to replicate the actor a NetHandle will be created. */
+	ENGINE_API static void BeginReplication(AActor* Actor, const FActorBeginReplicationParams& Params);
 
-	/** 
-		Begins replication an actor and all of its registered subobjects returns a NetHandle for the actor if it succeeds. 
-	*/
-	ENGINE_API static FNetHandle BeginReplication(AActor* Actor);
+	/** Begins replication of an actor and all of its registered subobjects. If any ReplicationSystem wants to replicate the actor a NetHandle will be created. */
+	ENGINE_API static void BeginReplication(AActor* Actor);
 	
 	/** Stop replicating an actor. Will destroy handle for actor and registered subobjects. */
 	ENGINE_API static void EndReplication(AActor* Actor, EEndPlayReason::Type EndPlayReason);
 
 	/** Create NetHandle for actor component and add it as a subobject to the actor handle. */
-	ENGINE_API static FNetHandle BeginReplicationForActorComponent(FNetHandle ActorHandle, UActorComponent* ActorComponent);
+	ENGINE_API static void BeginReplicationForActorComponent(FNetHandle ActorHandle, UActorComponent* ActorComponent);
 
 	/** Create NetHandle for actor component and add it as a subobject to the actor. */
-	ENGINE_API static FNetHandle BeginReplicationForActorComponent(const AActor* Actor, UActorComponent* ActorComponent);
+	ENGINE_API static void BeginReplicationForActorComponent(const AActor* Actor, UActorComponent* ActorComponent);
 
 	/** Stop replicating an ActorComponent and its associated SubObjects. */
 	ENGINE_API static void EndReplicationForActorComponent(UActorComponent* SubObject);
 
 	/** Create NetHandle for SubObject and add it as a subobject to the actor. */
-	ENGINE_API static FNetHandle BeginReplicationForActorSubObject(const AActor* Actor, UObject* SubObject, ELifetimeCondition NetCondition);
+	ENGINE_API static void BeginReplicationForActorSubObject(const AActor* Actor, UObject* SubObject, ELifetimeCondition NetCondition);
 
 	/** Stop replicating a SubObject ActorComponent and its associated SubObjects. */
 	ENGINE_API static void EndReplicationForActorSubObject(const AActor* Actor, UObject* SubObject);
 
-	/** Create NetHandle for SubObject and add it as a subobject to the actor but only replicated if the ActorComponent replicates. */
-	ENGINE_API static FNetHandle BeginReplicationForActorComponentSubObject(UActorComponent* ActorComponent, UObject* SubObject, ELifetimeCondition Condition);
+	/** Create NetHandle for SubObject and add it as a subobject to the actor component, but only replicated if the ActorComponent replicates. */
+	ENGINE_API static void BeginReplicationForActorComponentSubObject(UActorComponent* ActorComponent, UObject* SubObject, ELifetimeCondition Condition);
 
 	/** Stop replicating an ActorComponent and its associated SubObjects. */
 	ENGINE_API static void EndReplicationForActorComponentSubObject(UActorComponent* ActorComponent, UObject* SubObject);
@@ -108,6 +105,9 @@ struct FReplicationSystemUtil
 
 	/** Trigger replication of dirty state for actor wanting to be dormant. */
 	ENGINE_API static void FlushNetDormancy(AActor* Actor, bool bWasDormInitial);
+
+	/** Enable or disable a replication condition. This will affect the replication of properties with conditions. */
+	ENGINE_API static void SetReplicationCondition(FNetHandle NetHandle, EReplicationCondition Condition, bool bEnableCondition);
 };
 
 }
diff --git a/Engine/Source/Runtime/Experimental/Iris/Core/Private/Iris/Core/IrisDebugging.cpp b/Engine/Source/Runtime/Experimental/Iris/Core/Private/Iris/Core/IrisDebugging.cpp
--- a/Engine/Source/Runtime/Experimental/Iris/Core/Private/Iris/Core/IrisDebugging.cpp
+++ b/Engine/Source/Runtime/Experimental/Iris/Core/Private/Iris/Core/IrisDebugging.cpp
@@ -36,24 +36,24 @@ namespace IrisDebugHelperInternal
 	FAutoConsoleVariableRef NetRPCSetDebugRPCName(
 		TEXT("Net.Iris.DebugRPCName"),
 		GIrisDebugRPCName,
-		TEXT("Set the Name of an RPC to break on."),
+		TEXT("Set the name of an RPC to break on."),
 		ECVF_Default);
 
-	static FNetHandle GIrisDebugNetHandle;
-	static FAutoConsoleCommand NetIrisDebugNetHandle(
-		TEXT("Net.Iris.DebugNetHandle"), 
-		TEXT("Specify an handle index that we will break on (or none to turn off)."), 
+	static FNetRefHandle GIrisDebugNetRefHandle;
+	static FAutoConsoleCommand NetIrisDebugNetRefHandle(
+		TEXT("Net.Iris.DebugNetRefHandle"), 
+		TEXT("Specify a NetRefHandle ID that we will break on (or none to turn off)."), 
 		FConsoleCommandWithArgsDelegate::CreateLambda([](const TArray<FString>& Args)
 		{
 			if (Args.Num() > 0)
 			{
 				uint32 NetId = 0;
 				LexFromString(NetId, *Args[0]);
-				GIrisDebugNetHandle = Private::FNetHandleManager::MakeNetHandleFromId(NetId);
+				GIrisDebugNetRefHandle = Private::FNetRefHandleManager::MakeNetRefHandleFromId(NetId);
 			}
 			else
 			{
-				GIrisDebugNetHandle = FNetHandle();
+				GIrisDebugNetRefHandle = FNetRefHandle();
 			}
 		}));
 }; // namespace IrisDebugHelperInternal
@@ -69,9 +69,9 @@ bool BreakOnObjectName(UObject* Object)
 	return false;
 }
 
-bool BreakOnNetHandle(FNetHandle NetHandle)
+bool BreakOnNetRefHandle(FNetRefHandle NetRefHandle)
 {
-	if (IrisDebugHelperInternal::GIrisDebugNetHandle.IsValid() && IrisDebugHelperInternal::GIrisDebugNetHandle == NetHandle)
+	if (IrisDebugHelperInternal::GIrisDebugNetRefHandle.IsValid() && IrisDebugHelperInternal::GIrisDebugNetRefHandle == NetRefHandle)
 	{
 		UE_DEBUG_BREAK();
 		return true;
@@ -103,9 +103,9 @@ void SetIrisDebugObjectName(const ANSICHAR* NameBuffer)
 	}
 }
 
-void SetIrisDebugNetHandle(uint32 NetHandleId)
+void SetIrisDebugNetRefHandle(uint32 NetRefHandleId)
 {
-	IrisDebugHelperInternal::GIrisDebugNetHandle = Private::FNetHandleManager::MakeNetHandleFromId(NetHandleId);
+	IrisDebugHelperInternal::GIrisDebugNetRefHandle = Private::FNetRefHandleManager::MakeNetRefHandleFromId(NetRefHandleId);
 }
 
 void SetIrisDebugRPCName(const ANSICHAR* NameBuffer)
@@ -137,11 +137,11 @@ uint64 Init()
 	UE_NET_FORCE_REFERENCE_DEBUGFUNCTION(DebugOutputNetObjectProtocolReferences);
 	UE_NET_FORCE_REFERENCE_DEBUGFUNCTION(DebugNetObject);
 	UE_NET_FORCE_REFERENCE_DEBUGFUNCTION(DebugNetObjectById);
-	UE_NET_FORCE_REFERENCE_DEBUGFUNCTION(DebugNetHandle);
-	UE_NET_FORCE_REFERENCE_DEBUGFUNCTION(DebugNetHandleById);
+	UE_NET_FORCE_REFERENCE_DEBUGFUNCTION(DebugNetRefHandle);
+	UE_NET_FORCE_REFERENCE_DEBUGFUNCTION(DebugNetRefHandleById);
 	UE_NET_FORCE_REFERENCE_DEBUGFUNCTION(DebugNetObjectProtocolReferencesToString);
 	UE_NET_FORCE_REFERENCE_DEBUGFUNCTION(SetIrisDebugObjectName);
-	UE_NET_FORCE_REFERENCE_DEBUGFUNCTION(SetIrisDebugNetHandle);
+	UE_NET_FORCE_REFERENCE_DEBUGFUNCTION(SetIrisDebugNetRefHandle);
 	UE_NET_FORCE_REFERENCE_DEBUGFUNCTION(SetIrisDebugRPCName);
 	
 	return FunctionReferenceAccumulator;
@@ -149,33 +149,33 @@ uint64 Init()
 
 #undef UE_NET_FORCE_REFERENCE_DEBUGFUNCTION
 
-void NetObjectStateToString(FStringBuilderBase& StringBuilder, FNetHandle NetHandle)
+void NetObjectStateToString(FStringBuilderBase& StringBuilder, FNetRefHandle RefHandle)
 {
 	using namespace UE::Net::Private;
 
-	if (!NetHandle.IsValid())
+	if (!RefHandle.IsValid())
 	{
 		return;
 	}
 
-	UReplicationSystem* ReplicationSystem = GetReplicationSystem(NetHandle.GetReplicationSystemId());
+	UReplicationSystem* ReplicationSystem = GetReplicationSystem(RefHandle.GetReplicationSystemId());
 	if (!ReplicationSystem)
 	{
 		return;
 	}
 
 	FReplicationSystemInternal* ReplicationSystemInternal = ReplicationSystem->GetReplicationSystemInternal();
-	const FNetHandleManager& NetHandleManager = ReplicationSystemInternal->GetNetHandleManager();
+	const FNetRefHandleManager& NetRefHandleManager = ReplicationSystemInternal->GetNetRefHandleManager();
 
-	const uint32 NetObjectInternalIndex = NetHandleManager.GetInternalIndex(NetHandle);
-	if (!NetObjectInternalIndex)
+	const FInternalNetRefIndex InternalNetRefIndex = NetRefHandleManager.GetInternalIndex(RefHandle);
+	if (!InternalNetRefIndex)
 	{
 		return;
 	}
 
 	const bool bIsServer = ReplicationSystem->IsServer();
-	const FNetHandleManager::FReplicatedObjectData& ReplicatedObjectData = NetHandleManager.GetReplicatedObjectDataNoCheck(NetObjectInternalIndex);	
-	const uint8* InternalStateBuffer = bIsServer ? NetHandleManager.GetReplicatedObjectStateBufferNoCheck(NetObjectInternalIndex) : ReplicatedObjectData.ReceiveStateBuffer;
+	const FNetRefHandleManager::FReplicatedObjectData& ReplicatedObjectData = NetRefHandleManager.GetReplicatedObjectDataNoCheck(InternalNetRefIndex);	
+	const uint8* InternalStateBuffer = bIsServer ? NetRefHandleManager.GetReplicatedObjectStateBufferNoCheck(InternalNetRefIndex) : ReplicatedObjectData.ReceiveStateBuffer;
 	if (ReplicatedObjectData.InstanceProtocol == nullptr || InternalStateBuffer == nullptr)
 	{
 		return;
@@ -201,30 +201,30 @@ void NetObjectStateToString(FStringBuilderBase& StringBuilder, FNetHandle NetHan
 	NetSerializationContext.SetInternalContext(&InternalContext);
 	NetSerializationContext.SetLocalConnectionId(InternalContextInitParams.ObjectResolveContext.ConnectionId);
 
-	StringBuilder << TEXT("State for ") << NetHandle;
+	StringBuilder << TEXT("State for ") << RefHandle;
 	StringBuilder.Appendf(TEXT(" of type %s\n"), Protocol->DebugName->Name);
 
 	FReplicationInstanceOperations::OutputInternalStateToString(NetSerializationContext, StringBuilder, nullptr, InternalStateBuffer, ReplicatedObjectData.InstanceProtocol, ReplicatedObjectData.Protocol);
 }
 
-void DebugOutputNetObjectState(uint32 NetHandleId, uint32 ReplicationSystemId)
+void DebugOutputNetObjectState(uint32 NetRefHandleId, uint32 ReplicationSystemId)
 {
-	const FNetHandle NetHandle = Private::FNetHandleManager::MakeNetHandle(NetHandleId, ReplicationSystemId);
+	const FNetRefHandle RefHandle = Private::FNetRefHandleManager::MakeNetRefHandle(NetRefHandleId, ReplicationSystemId);
 
 	TStringBuilder<4096> StringBuilder;
 
-	NetObjectStateToString(StringBuilder, NetHandle);
+	NetObjectStateToString(StringBuilder, RefHandle);
 
 	FPlatformMisc::LowLevelOutputDebugString(StringBuilder.ToString());
 }
 
-const TCHAR* DebugNetObjectStateToString(uint32 NetHandleId, uint32 ReplicationSystemId)
+const TCHAR* DebugNetObjectStateToString(uint32 NetRefHandleId, uint32 ReplicationSystemId)
 {
 	static TStringBuilder<4096> StringBuilder;
 
 	StringBuilder.Reset();
-	const FNetHandle NetHandle = Private::FNetHandleManager::MakeNetHandle(NetHandleId, ReplicationSystemId);
-	NetObjectStateToString(StringBuilder, NetHandle);
+	const FNetRefHandle RefHandle = Private::FNetRefHandleManager::MakeNetRefHandle(NetRefHandleId, ReplicationSystemId);
+	NetObjectStateToString(StringBuilder, RefHandle);
 
 	return StringBuilder.ToString();
 }
@@ -242,15 +242,15 @@ FNetReplicatedObjectDebugInfo DebugNetObject(UObject* Instance)
 		{
 			if (const UObjectReplicationBridge* Bridge = ReplicationSystem->GetReplicationBridgeAs<UObjectReplicationBridge>())
 			{
-				FNetHandle Handle = Bridge->GetReplicatedHandle(Instance);
+				FNetRefHandle Handle = Bridge->GetReplicatedRefHandle(Instance);
 				if (Handle.IsValid())
 				{
-					const FNetHandleManager& NetHandleManager = ReplicationSystem->GetReplicationSystemInternal()->GetNetHandleManager();
-					Info.InternalNetHandleIndex = NetHandleManager.GetInternalIndex(Handle);
+					const FNetRefHandleManager& NetRefHandleManager = ReplicationSystem->GetReplicationSystemInternal()->GetNetRefHandleManager();
+					Info.InternalNetRefIndex = NetRefHandleManager.GetInternalIndex(Handle);
 
-					const FNetHandleManager::FReplicatedObjectData& ObjectData = NetHandleManager.GetReplicatedObjectDataNoCheck(Info.InternalNetHandleIndex);				
+					const FNetRefHandleManager::FReplicatedObjectData& ObjectData = NetRefHandleManager.GetReplicatedObjectDataNoCheck(Info.InternalNetRefIndex);				
 
-					Info.Handle = &ObjectData.Handle;
+					Info.RefHandle = &ObjectData.RefHandle;
 					Info.Protocol = ObjectData.Protocol;
 					Info.InstanceProtocol = ObjectData.InstanceProtocol;
 					Info.ReplicationSystem = ReplicationSystem;
@@ -274,15 +274,15 @@ FNetReplicatedObjectDebugInfo DebugNetObjectById(UObject* Instance, uint32 Repli
 	{
 		if (const UObjectReplicationBridge* Bridge = ReplicationSystem->GetReplicationBridgeAs<UObjectReplicationBridge>())
 		{
-			FNetHandle Handle = Bridge->GetReplicatedHandle(Instance);
+			FNetRefHandle Handle = Bridge->GetReplicatedRefHandle(Instance);
 			if (Handle.IsValid())
 			{
-				const FNetHandleManager& NetHandleManager = ReplicationSystem->GetReplicationSystemInternal()->GetNetHandleManager();
-				Info.InternalNetHandleIndex = NetHandleManager.GetInternalIndex(Handle);
+				const FNetRefHandleManager& NetRefHandleManager = ReplicationSystem->GetReplicationSystemInternal()->GetNetRefHandleManager();
+				Info.InternalNetRefIndex = NetRefHandleManager.GetInternalIndex(Handle);
 
-				const FNetHandleManager::FReplicatedObjectData& ObjectData = NetHandleManager.GetReplicatedObjectDataNoCheck(Info.InternalNetHandleIndex);
+				const FNetRefHandleManager::FReplicatedObjectData& ObjectData = NetRefHandleManager.GetReplicatedObjectDataNoCheck(Info.InternalNetRefIndex);
 
-				Info.Handle = &ObjectData.Handle;
+				Info.RefHandle = &ObjectData.RefHandle;
 				Info.Protocol = ObjectData.Protocol;
 				Info.InstanceProtocol = ObjectData.InstanceProtocol;
 				Info.ReplicationSystem = ReplicationSystem;
@@ -295,12 +295,12 @@ FNetReplicatedObjectDebugInfo DebugNetObjectById(UObject* Instance, uint32 Repli
 	return Info;
 }
 
-FNetReplicatedObjectDebugInfo DebugNetHandle(FNetHandle Handle)
+FNetReplicatedObjectDebugInfo DebugNetRefHandle(FNetRefHandle Handle)
 {
-	return DebugNetHandleById(Handle.GetId(), Handle.GetReplicationSystemId());
+	return DebugNetRefHandleById(Handle.GetId(), Handle.GetReplicationSystemId());
 }
 
-FNetReplicatedObjectDebugInfo DebugNetHandleById(uint32 NetHandleId, uint32 ReplicationSystemId)
+FNetReplicatedObjectDebugInfo DebugNetRefHandleById(uint32 NetRefHandleId, uint32 ReplicationSystemId)
 {
 	using namespace UE::Net::Private;
 	FNetReplicatedObjectDebugInfo Info = {};
@@ -310,16 +310,16 @@ FNetReplicatedObjectDebugInfo DebugNetHandleById(uint32 NetHandleId, uint32 Repl
 		const UReplicationSystem* ReplicationSystem = GetReplicationSystemForDebug(ReplicationSystemId);
 		if (ReplicationSystem)
 		{
-			const FNetHandleManager& NetHandleManager = ReplicationSystem->GetReplicationSystemInternal()->GetNetHandleManager();
-			FNetHandle IncompleteHandle = FNetHandleManager::MakeNetHandle(NetHandleId, ReplicationSystemId);
-			const uint32 InternalNetHandleIndex = Info.InternalNetHandleIndex = NetHandleManager.GetInternalIndex(IncompleteHandle);
-			if (InternalNetHandleIndex != FNetHandleManager::InvalidInternalIndex)
+			const FNetRefHandleManager& NetRefHandleManager = ReplicationSystem->GetReplicationSystemInternal()->GetNetRefHandleManager();
+			FNetRefHandle IncompleteHandle = FNetRefHandleManager::MakeNetRefHandle(NetRefHandleId, ReplicationSystemId);
+			const uint32 InternalNetRefIndex = Info.InternalNetRefIndex = NetRefHandleManager.GetInternalIndex(IncompleteHandle);
+			if (InternalNetRefIndex != FNetRefHandleManager::InvalidInternalIndex)
 			{
-				const FNetHandleManager::FReplicatedObjectData& ObjectData = NetHandleManager.GetReplicatedObjectDataNoCheck(InternalNetHandleIndex);
-				if (ObjectData.Handle == IncompleteHandle)
+				const FNetRefHandleManager::FReplicatedObjectData& ObjectData = NetRefHandleManager.GetReplicatedObjectDataNoCheck(InternalNetRefIndex);
+				if (ObjectData.RefHandle == IncompleteHandle)
 				{
-					Info.Handle = &ObjectData.Handle;
-					Info.InternalNetHandleIndex = InternalNetHandleIndex;
+					Info.RefHandle = &ObjectData.RefHandle;
+					Info.InternalNetRefIndex = InternalNetRefIndex;
 					Info.Protocol = ObjectData.Protocol;
 					Info.InstanceProtocol = ObjectData.InstanceProtocol;
 					Info.ReplicationSystem = ReplicationSystem;
@@ -344,25 +344,25 @@ void NetObjectProtocolReferencesToString(FStringBuilderBase& StringBuilder, uint
 	}
 
 	FReplicationSystemInternal* ReplicationSystemInternal = ReplicationSystem->GetReplicationSystemInternal();
-	const FNetHandleManager& NetHandleManager = ReplicationSystemInternal->GetNetHandleManager();
+	const FNetRefHandleManager& NetRefHandleManager = ReplicationSystemInternal->GetNetRefHandleManager();
 	const FReplicationProtocolManager& ProtocolManager = ReplicationSystemInternal->GetReplicationProtocolManager();
 
 	// As there might be multiple protocols sharing the same identifier
-	ProtocolManager.ForEachProtocol(ProtocolId, [&StringBuilder, &NetHandleManager](const FReplicationProtocol* Protocol, const UObject* ArchetypeOrCDOUsedAsKey)
+	ProtocolManager.ForEachProtocol(ProtocolId, [&StringBuilder, &NetRefHandleManager](const FReplicationProtocol* Protocol, const UObject* ArchetypeOrCDOUsedAsKey)
 		{
 			StringBuilder << TEXT("Protocol: ") << ToCStr(Protocol->DebugName);
 			StringBuilder.Appendf(TEXT("Id: 0x%" UINT64_x_FMT " Created From : 0x%p"), Protocol->ProtocolIdentifier, ArchetypeOrCDOUsedAsKey) << TEXT(" Used by : \n");
 
-			auto FindMatchingProtocol = [&StringBuilder, &Protocol, &NetHandleManager](uint32 InternalObjectIndex)
+			auto FindMatchingProtocol = [&StringBuilder, &Protocol, &NetRefHandleManager](uint32 InternalObjectIndex)
 			{
-				const FNetHandleManager::FReplicatedObjectData& ObjectData = NetHandleManager.GetReplicatedObjectDataNoCheck(InternalObjectIndex);
+				const FNetRefHandleManager::FReplicatedObjectData& ObjectData = NetRefHandleManager.GetReplicatedObjectDataNoCheck(InternalObjectIndex);
 				if (ObjectData.Protocol == Protocol)
 				{
-					StringBuilder << ObjectData.Handle << TEXT(" ( InternalIndex: ") << InternalObjectIndex << TEXT(", )\n");
+					StringBuilder << ObjectData.RefHandle << TEXT(" ( InternalIndex: ") << InternalObjectIndex << TEXT(", )\n");
 				}
 			};
 			
-			NetHandleManager.GetAssignedInternalIndices().ForAllSetBits(FindMatchingProtocol);
+			NetRefHandleManager.GetAssignedInternalIndices().ForAllSetBits(FindMatchingProtocol);
 		}
 	);
 }
diff --git a/Engine/Source/Runtime/Experimental/Iris/Core/Private/Iris/IrisCoreModule.cpp b/Engine/Source/Runtime/Experimental/Iris/Core/Private/Iris/IrisCoreModule.cpp
--- a/Engine/Source/Runtime/Experimental/Iris/Core/Private/Iris/IrisCoreModule.cpp
+++ b/Engine/Source/Runtime/Experimental/Iris/Core/Private/Iris/IrisCoreModule.cpp
@@ -31,6 +31,9 @@ private:
 
 	virtual void StartupModule() override
 	{
+		// Iris requires NetCore
+		FModuleManager::LoadModuleChecked<IModuleInterface>("NetCore");
+
 		// Check command line for whether we should override the net.Iris.UseIrisReplication cvar, as we need to do that early
 		int32 UseIrisReplication;
 		if(FParse::Value(FCommandLine::Get(), TEXT("UseIrisReplication="), UseIrisReplication))
diff --git a/Engine/Source/Runtime/Experimental/Iris/Core/Private/Iris/ReplicationState/ReplicationStateStorage.cpp b/Engine/Source/Runtime/Experimental/Iris/Core/Private/Iris/ReplicationState/ReplicationStateStorage.cpp
--- a/Engine/Source/Runtime/Experimental/Iris/Core/Private/Iris/ReplicationState/ReplicationStateStorage.cpp
+++ b/Engine/Source/Runtime/Experimental/Iris/Core/Private/Iris/ReplicationState/ReplicationStateStorage.cpp
@@ -30,7 +30,7 @@ void FReplicationStateStorage::Init(FReplicationStateStorageInitParams& InitPara
 	// Make sure the MaxObjectInfoCount calculation can't overflow.
 	static_assert(std::numeric_limits<decltype(MaxObjectInfoCount)>::max() > std::numeric_limits<ObjectInfoIndexType>::max(), "");
 
-	NetHandleManager = InitParams.NetHandleManager;
+	NetRefHandleManager = InitParams.NetRefHandleManager;
 
 	UsedPerObjectInfos.Init(MaxObjectInfoCount);
 	UsedPerObjectInfos.SetBit(InvalidObjectInfoIndex);
@@ -79,11 +79,11 @@ const uint8* FReplicationStateStorage::GetState(uint32 ObjectIndex, EReplication
 	{
 		case EReplicationStateType::CurrentSendState:
 		{
-			return NetHandleManager->GetReplicatedObjectStateBufferNoCheck(ObjectIndex);
+			return NetRefHandleManager->GetReplicatedObjectStateBufferNoCheck(ObjectIndex);
 		}
 		case EReplicationStateType::CurrentRecvState:
 		{
-			const Private::FNetHandleManager::FReplicatedObjectData& ReplicatedObjectData = NetHandleManager->GetReplicatedObjectData(ObjectIndex);
+			const Private::FNetRefHandleManager::FReplicatedObjectData& ReplicatedObjectData = NetRefHandleManager->GetReplicatedObjectData(ObjectIndex);
 			return ReplicatedObjectData.ReceiveStateBuffer;
 		}
 		default:
@@ -332,10 +332,10 @@ FReplicationStateStorage::FPerObjectInfo* FReplicationStateStorage::GetOrCreateP
 		FPerObjectInfo& ObjectInfo = ObjectInfos[NewInfoIndex];
 		
 		// Cache some info
-		const Private::FNetHandleManager::FReplicatedObjectData& ReplicatedObjectData = NetHandleManager->GetReplicatedObjectData(ObjectIndex);
+		const Private::FNetRefHandleManager::FReplicatedObjectData& ReplicatedObjectData = NetRefHandleManager->GetReplicatedObjectData(ObjectIndex);
 		const FReplicationProtocol* Protocol = ReplicatedObjectData.Protocol;
 		ObjectInfo.Protocol = Protocol;
-		ObjectInfo.StateBuffers[(unsigned)EStateBufferType::SendState] = NetHandleManager->GetReplicatedObjectStateBufferNoCheck(ObjectIndex);
+		ObjectInfo.StateBuffers[(unsigned)EStateBufferType::SendState] = NetRefHandleManager->GetReplicatedObjectStateBufferNoCheck(ObjectIndex);
 		ObjectInfo.StateBuffers[(unsigned)EStateBufferType::RecvState] = ReplicatedObjectData.ReceiveStateBuffer;
 		return &ObjectInfo;
 	}
diff --git a/Engine/Source/Runtime/Experimental/Iris/Core/Private/Iris/ReplicationState/ReplicationStateStorage.h b/Engine/Source/Runtime/Experimental/Iris/Core/Private/Iris/ReplicationState/ReplicationStateStorage.h
--- a/Engine/Source/Runtime/Experimental/Iris/Core/Private/Iris/ReplicationState/ReplicationStateStorage.h
+++ b/Engine/Source/Runtime/Experimental/Iris/Core/Private/Iris/ReplicationState/ReplicationStateStorage.h
@@ -17,7 +17,7 @@ namespace UE::Net
 	struct FReplicationProtocol;
 	namespace Private
 	{
-		class FNetHandleManager;
+		class FNetRefHandleManager;
 	};
 }
 
@@ -27,7 +27,7 @@ namespace UE::Net
 struct FReplicationStateStorageInitParams
 {
 	UReplicationSystem* ReplicationSystem = nullptr;
-	const Private::FNetHandleManager* NetHandleManager = nullptr;
+	const Private::FNetRefHandleManager* NetRefHandleManager = nullptr;
 	uint32 MaxObjectCount = 0;
 	uint32 MaxConnectionCount = 0;
 	uint32 MaxDeltaCompressedObjectCount = 0;
@@ -140,7 +140,7 @@ private:
 	TMultiMap<uint32, const void*> BaselineStorageValidation;
 #endif
 
-	const Private::FNetHandleManager* NetHandleManager = nullptr;
+	const Private::FNetRefHandleManager* NetRefHandleManager = nullptr;
 
 	FNetSerializationContext SerializationContext;
 	Private::FInternalNetSerializationContext InternalSerializationContext;
diff --git a/Engine/Source/Runtime/Experimental/Iris/Core/Private/Iris/ReplicationSystem/AttachmentReplication.cpp b/Engine/Source/Runtime/Experimental/Iris/Core/Private/Iris/ReplicationSystem/AttachmentReplication.cpp
--- a/Engine/Source/Runtime/Experimental/Iris/Core/Private/Iris/ReplicationSystem/AttachmentReplication.cpp
+++ b/Engine/Source/Runtime/Experimental/Iris/Core/Private/Iris/ReplicationSystem/AttachmentReplication.cpp
@@ -106,11 +106,11 @@ public:
 		return true;
 	}
 
-	uint32 Serialize(FNetSerializationContext& Context, FNetHandle NetHandle, FReliableNetBlobQueue::ReplicationRecord& OutRecord)
+	uint32 Serialize(FNetSerializationContext& Context, FNetRefHandle RefHandle, FReliableNetBlobQueue::ReplicationRecord& OutRecord)
 	{
-		if (NetHandle.IsValid())
+		if (RefHandle.IsValid())
 		{
-			return ReliableQueue.SerializeWithObject(Context, NetHandle, OutRecord);
+			return ReliableQueue.SerializeWithObject(Context, RefHandle, OutRecord);
 		}
 		else
 		{
@@ -251,7 +251,7 @@ void FNetObjectAttachmentSendQueue::SetUnreliableQueueCapacity(uint32 QueueCapac
 	UnreliableQueue.PopNoCheck(DropCount);
 }
 
-EAttachmentWriteStatus FNetObjectAttachmentSendQueue::Serialize(FNetSerializationContext& Context, FNetHandle NetHandle, FNetObjectAttachmentSendQueue::ReplicationRecord& OutRecord, bool& bOutHasUnsentAttachments)
+EAttachmentWriteStatus FNetObjectAttachmentSendQueue::Serialize(FNetSerializationContext& Context, FNetRefHandle RefHandle, FNetObjectAttachmentSendQueue::ReplicationRecord& OutRecord, bool& bOutHasUnsentAttachments)
 {
 	FNetBitStreamWriter& Writer = *Context.GetBitStreamWriter();
 
@@ -295,7 +295,7 @@ EAttachmentWriteStatus FNetObjectAttachmentSendQueue::Serialize(FNetSerializatio
 	if (bCanSendReliableAttachments)
 	{
 		UE_NET_TRACE_SCOPE(Reliable, Writer, Context.GetTraceCollector(), ENetTraceVerbosity::Trace);
-		SerializedReliableCount = SerializeReliable(Context, NetHandle, ReplicationRecord.ReliableRecord);
+		SerializedReliableCount = SerializeReliable(Context, RefHandle, ReplicationRecord.ReliableRecord);
 		// If we couldn't fit any reliable attachments then don't even try unreliable
 		if (SerializedReliableCount == 0)
 		{
@@ -313,7 +313,7 @@ EAttachmentWriteStatus FNetObjectAttachmentSendQueue::Serialize(FNetSerializatio
 	if (bHasUnreliableAttachments)
 	{
 		UE_NET_TRACE_SCOPE(Unreliable, Writer, Context.GetTraceCollector(), ENetTraceVerbosity::Trace);
-		SerializedUnreliableCount = SerializeUnreliable(Context, NetHandle, ReplicationRecord.UnreliableRecord);
+		SerializedUnreliableCount = SerializeUnreliable(Context, RefHandle, ReplicationRecord.UnreliableRecord);
 		if (SerializedUnreliableCount == 0)
 		{
 			// If we didn't manage to send anything then inform the caller of this via overflowing the bitstream
@@ -346,19 +346,19 @@ EAttachmentWriteStatus FNetObjectAttachmentSendQueue::Serialize(FNetSerializatio
 	return WriteStatus;
 }
 
-uint32 FNetObjectAttachmentSendQueue::SerializeReliable(FNetSerializationContext& Context, FNetHandle NetHandle, FReliableNetBlobQueue::ReplicationRecord& OutRecord)
+uint32 FNetObjectAttachmentSendQueue::SerializeReliable(FNetSerializationContext& Context, FNetRefHandle RefHandle, FReliableNetBlobQueue::ReplicationRecord& OutRecord)
 {
-	return ReliableQueue->Serialize(Context, NetHandle, OutRecord);
+	return ReliableQueue->Serialize(Context, RefHandle, OutRecord);
 }
 
-uint32 FNetObjectAttachmentSendQueue::SerializeUnreliable(FNetSerializationContext& Context, FNetHandle NetHandle, uint32& OutRecord)
+uint32 FNetObjectAttachmentSendQueue::SerializeUnreliable(FNetSerializationContext& Context, FNetRefHandle RefHandle, uint32& OutRecord)
 {
 	FNetBitStreamWriter* Writer = Context.GetBitStreamWriter();
 	const FObjectReferenceCache* ObjectReferenceCache = Context.GetInternalContext()->ObjectReferenceCache;
 
 	uint32 SerializedUnreliableCount = 0;
 	uint32 PrevHasMoreAttachmentsWritePos = 0;
-	const bool bSerializeWithObject = NetHandle.IsValid();
+	const bool bSerializeWithObject = RefHandle.IsValid();
 	for (uint32 AttachmentIt = 0, AttachmentEndIt = UnreliableQueue.Count(); AttachmentIt != AttachmentEndIt; ++AttachmentIt)
 	{
 		FNetBitStreamRollbackScope RollbackScope(*Writer);
@@ -372,7 +372,7 @@ uint32 FNetObjectAttachmentSendQueue::SerializeUnreliable(FNetSerializationConte
 		Attachment->SerializeCreationInfo(Context, Attachment->GetCreationInfo());
 		if (bSerializeWithObject)
 		{
-			Attachment->SerializeWithObject(Context, NetHandle);
+			Attachment->SerializeWithObject(Context, RefHandle);
 		}
 		else
 		{
@@ -523,7 +523,7 @@ void FNetObjectAttachmentsWriter::DropUnreliableAttachments(ENetObjectAttachment
 	}
 }
 
-EAttachmentWriteStatus FNetObjectAttachmentsWriter::Serialize(FNetSerializationContext& Context, ENetObjectAttachmentType Type, uint32 ObjectIndex, const FNetHandle NetHandle,  FNetObjectAttachmentsWriter::ReplicationRecord& OutRecord, bool& bOutHasUnsentAttachments)
+EAttachmentWriteStatus FNetObjectAttachmentsWriter::Serialize(FNetSerializationContext& Context, ENetObjectAttachmentType Type, uint32 ObjectIndex, const FNetRefHandle RefHandle,  FNetObjectAttachmentsWriter::ReplicationRecord& OutRecord, bool& bOutHasUnsentAttachments)
 {
 	FNetObjectAttachmentSendQueue* Queue = GetQueue(Type, ObjectIndex);
 	// If this ensure fires we have bad logic for keeping track of whether there are attachments or not
@@ -534,7 +534,7 @@ EAttachmentWriteStatus FNetObjectAttachmentsWriter::Serialize(FNetSerializationC
 		return EAttachmentWriteStatus::NoAttachments;
 	}
 
-	return Queue->Serialize(Context, NetHandle, OutRecord, bOutHasUnsentAttachments);
+	return Queue->Serialize(Context, RefHandle, OutRecord, bOutHasUnsentAttachments);
 }
 
 void FNetObjectAttachmentsWriter::CommitReplicationRecord(ENetObjectAttachmentType Type, uint32 ObjectIndex, FNetObjectAttachmentsWriter::ReplicationRecord Record)
@@ -632,7 +632,7 @@ public:
 	bool HasUnprocessed() const { return !Queue.IsEmpty(); }
 	bool IsSafeToDestroy() const { return Queue.IsEmpty() && !NetBlobAssembler.IsValid(); }
 
-	void Enqueue(FNetSerializationContext& Context, FNetHandle NetHandle, const TRefCountPtr<FNetBlob>& NetBlob, bool bIsPartialNetBlob)
+	void Enqueue(FNetSerializationContext& Context, FNetRefHandle RefHandle, const TRefCountPtr<FNetBlob>& NetBlob, bool bIsPartialNetBlob)
 	{
 		if (bIsPartialNetBlob)
 		{
@@ -644,7 +644,7 @@ public:
 				NetBlobAssembler->Init(InitParams);
 			}
 
-			NetBlobAssembler->AddPartialNetBlob(Context, NetHandle, reinterpret_cast<const TRefCountPtr<FPartialNetBlob>&>(NetBlob));
+			NetBlobAssembler->AddPartialNetBlob(Context, RefHandle, reinterpret_cast<const TRefCountPtr<FPartialNetBlob>&>(NetBlob));
 			if (NetBlobAssembler->IsReadyToAssemble())
 			{
 				const TRefCountPtr<FNetBlob>& AssembledBlob = NetBlobAssembler->Assemble(Context);
@@ -776,7 +776,7 @@ bool FNetObjectAttachmentReceiveQueue::IsPartialNetBlob(const TRefCountPtr<FNetB
 	return Blob.IsValid() && Blob->GetCreationInfo().Type == PartialNetBlobType;
 }
 
-void FNetObjectAttachmentReceiveQueue::Deserialize(FNetSerializationContext& Context, FNetHandle NetHandle)
+void FNetObjectAttachmentReceiveQueue::Deserialize(FNetSerializationContext& Context, FNetRefHandle RefHandle)
 {
 	FNetBitStreamReader& Reader = *Context.GetBitStreamReader();
 
@@ -793,7 +793,7 @@ void FNetObjectAttachmentReceiveQueue::Deserialize(FNetSerializationContext& Con
 	if (bHasReliableAttachments)
 	{
 		UE_NET_TRACE_SCOPE(Reliable, Reader, Context.GetTraceCollector(), ENetTraceVerbosity::Trace);
-		DeserializeReliable(Context, NetHandle);
+		DeserializeReliable(Context, RefHandle);
 		if (Context.HasErrorOrOverflow())
 		{
 			return;
@@ -803,7 +803,7 @@ void FNetObjectAttachmentReceiveQueue::Deserialize(FNetSerializationContext& Con
 	if (bHasUnreliableAttachments)
 	{
 		UE_NET_TRACE_SCOPE(Unreliable, Reader, Context.GetTraceCollector(), ENetTraceVerbosity::Trace);
-		DeserializeUnreliable(Context, NetHandle);
+		DeserializeUnreliable(Context, RefHandle);
 		if (Context.HasErrorOrOverflow())
 		{
 			return;
@@ -811,7 +811,7 @@ void FNetObjectAttachmentReceiveQueue::Deserialize(FNetSerializationContext& Con
 	}
 }
 
-uint32 FNetObjectAttachmentReceiveQueue::DeserializeReliable(FNetSerializationContext& Context, FNetHandle NetHandle)
+uint32 FNetObjectAttachmentReceiveQueue::DeserializeReliable(FNetSerializationContext& Context, FNetRefHandle RefHandle)
 {
 	if (ReliableQueue == nullptr)
 	{
@@ -824,9 +824,9 @@ uint32 FNetObjectAttachmentReceiveQueue::DeserializeReliable(FNetSerializationCo
 	}
 
 	uint32 DeserializedReliableCount = 0;
-	if (NetHandle.IsValid())
+	if (RefHandle.IsValid())
 	{
-		DeserializedReliableCount = ReliableQueue->DeserializeWithObject(Context, NetHandle);
+		DeserializedReliableCount = ReliableQueue->DeserializeWithObject(Context, RefHandle);
 	}
 	else
 	{
@@ -840,7 +840,7 @@ uint32 FNetObjectAttachmentReceiveQueue::DeserializeReliable(FNetSerializationCo
 
 	while (const TRefCountPtr<FNetBlob>* Attachment = ReliableQueue->Peek())
 	{
-		DeferredProcessingQueue->Enqueue(Context, NetHandle, *Attachment, IsPartialNetBlob(*Attachment));
+		DeferredProcessingQueue->Enqueue(Context, RefHandle, *Attachment, IsPartialNetBlob(*Attachment));
 		ReliableQueue->Pop();
 
 		if (Context.HasErrorOrOverflow())
@@ -852,7 +852,7 @@ uint32 FNetObjectAttachmentReceiveQueue::DeserializeReliable(FNetSerializationCo
 	return DeserializedReliableCount;
 }
 
-uint32 FNetObjectAttachmentReceiveQueue::DeserializeUnreliable(FNetSerializationContext& Context, FNetHandle NetHandle)
+uint32 FNetObjectAttachmentReceiveQueue::DeserializeUnreliable(FNetSerializationContext& Context, FNetRefHandle RefHandle)
 {
 	INetBlobReceiver* BlobReceiver = Context.GetNetBlobReceiver();
 	checkSlow(BlobReceiver != nullptr);
@@ -860,14 +860,14 @@ uint32 FNetObjectAttachmentReceiveQueue::DeserializeUnreliable(FNetSerialization
 	FNetBitStreamReader& Reader = *Context.GetBitStreamReader();
 
 	uint32 DeserializedUnreliableCount = 0;
-	const bool bDeserializeWithObject = NetHandle.IsValid();
+	const bool bDeserializeWithObject = RefHandle.IsValid();
 	bool bHasMoreAttachments = false;
 	do
 	{
 		// If the unreliable queue overflows this could be a carefully crafted malicious packet.
 		if (UnreliableQueue.Count() == MaxUnreliableCount)
 		{
-			UE_LOG(LogIris, Error, TEXT("Unreliable queue is full for %s"), *NetHandle.ToString());
+			UE_LOG(LogIris, Error, TEXT("Unreliable queue is full for %s"), *RefHandle.ToString());
 			Context.SetError(NetError_UnreliableQueueFull);
 			break;
 		}
@@ -885,7 +885,7 @@ uint32 FNetObjectAttachmentReceiveQueue::DeserializeUnreliable(FNetSerialization
 
 		if (bDeserializeWithObject)
 		{
-			Attachment->DeserializeWithObject(Context, NetHandle);
+			Attachment->DeserializeWithObject(Context, RefHandle);
 		}
 		else
 		{
@@ -897,7 +897,7 @@ uint32 FNetObjectAttachmentReceiveQueue::DeserializeUnreliable(FNetSerialization
 		// Break if something went wrong with deserialization.
 		if (Context.HasErrorOrOverflow())
 		{
-			UE_LOG(LogIris, Error, TEXT("Failed to deserialize unreliable attachments for %s"), *NetHandle.ToString());
+			UE_LOG(LogIris, Error, TEXT("Failed to deserialize unreliable attachments for %s"), *RefHandle.ToString());
 			break;
 		}
 
@@ -937,7 +937,7 @@ void FNetObjectAttachmentsReader::DropAllAttachments(ENetObjectAttachmentType Ty
 	ObjectToQueue.Remove(ObjectIndex);
 }
 
-void FNetObjectAttachmentsReader::Deserialize(FNetSerializationContext& Context, ENetObjectAttachmentType Type, uint32 ObjectIndex, FNetHandle NetHandle)
+void FNetObjectAttachmentsReader::Deserialize(FNetSerializationContext& Context, ENetObjectAttachmentType Type, uint32 ObjectIndex, FNetRefHandle RefHandle)
 {
 	FNetObjectAttachmentReceiveQueue* Queue = GetOrCreateQueue(Type, ObjectIndex);
 	if (!ensure(Queue != nullptr))
@@ -946,7 +946,7 @@ void FNetObjectAttachmentsReader::Deserialize(FNetSerializationContext& Context,
 		return;
 	}
 
-	Queue->Deserialize(Context, NetHandle);
+	Queue->Deserialize(Context, RefHandle);
 }
 
 FNetObjectAttachmentReceiveQueue* FNetObjectAttachmentsReader::GetQueue(ENetObjectAttachmentType Type, uint32 ObjectIndex)
diff --git a/Engine/Source/Runtime/Experimental/Iris/Core/Private/Iris/ReplicationSystem/AttachmentReplication.h b/Engine/Source/Runtime/Experimental/Iris/Core/Private/Iris/ReplicationSystem/AttachmentReplication.h
--- a/Engine/Source/Runtime/Experimental/Iris/Core/Private/Iris/ReplicationSystem/AttachmentReplication.h
+++ b/Engine/Source/Runtime/Experimental/Iris/Core/Private/Iris/ReplicationSystem/AttachmentReplication.h
@@ -86,9 +86,9 @@ private:
 		FInternalRecord() { CombinedRecord = InvalidReplicationRecord; };
 	};
 
-	EAttachmentWriteStatus Serialize(FNetSerializationContext& Context, FNetHandle NetHandle, ReplicationRecord& OutRecord, bool& bOutHasUnprocessedAttachments);
-	uint32 SerializeReliable(FNetSerializationContext& Context, FNetHandle NetHandle, FReliableNetBlobQueue::ReplicationRecord& OutRecord);
-	uint32 SerializeUnreliable(FNetSerializationContext& Context, FNetHandle NetHandle, uint32& OutRecord);
+	EAttachmentWriteStatus Serialize(FNetSerializationContext& Context, FNetRefHandle RefHandle, ReplicationRecord& OutRecord, bool& bOutHasUnprocessedAttachments);
+	uint32 SerializeReliable(FNetSerializationContext& Context, FNetRefHandle RefHandle, FReliableNetBlobQueue::ReplicationRecord& OutRecord);
+	uint32 SerializeUnreliable(FNetSerializationContext& Context, FNetRefHandle RefHandle, uint32& OutRecord);
 
 	void CommitReplicationRecord(ReplicationRecord Record);
 
@@ -121,7 +121,7 @@ public:
 	void DropAllAttachments(ENetObjectAttachmentType Type, uint32 ObjectIndex);
 	void DropUnreliableAttachments(ENetObjectAttachmentType Type, uint32 ObjectIndex, bool& bOutHasUnsentAttachments);
 
-	EAttachmentWriteStatus Serialize(FNetSerializationContext& Context, ENetObjectAttachmentType Type, uint32 ObjectIndex, FNetHandle NetHandle, ReplicationRecord& OutRecord, bool& bOutHasUnsentAttachments);
+	EAttachmentWriteStatus Serialize(FNetSerializationContext& Context, ENetObjectAttachmentType Type, uint32 ObjectIndex, FNetRefHandle RefHandle, ReplicationRecord& OutRecord, bool& bOutHasUnsentAttachments);
 
 	void CommitReplicationRecord(ENetObjectAttachmentType Type, uint32 ObjectIndex, ReplicationRecord Record);
 
@@ -173,9 +173,9 @@ private:
 	bool HasDeferredProcessingQueueUnprocessed() const;
 	bool IsPartialNetBlob(const TRefCountPtr<FNetBlob>& Blob) const;
 
-	void Deserialize(FNetSerializationContext& Context, FNetHandle NetHandle);
-	uint32 DeserializeReliable(FNetSerializationContext& Context, FNetHandle NetHandle);
-	uint32 DeserializeUnreliable(FNetSerializationContext& Context, FNetHandle NetHandle);
+	void Deserialize(FNetSerializationContext& Context, FNetRefHandle RefHandle);
+	uint32 DeserializeReliable(FNetSerializationContext& Context, FNetRefHandle RefHandle);
+	uint32 DeserializeUnreliable(FNetSerializationContext& Context, FNetRefHandle RefHandle);
 
 	TResizableCircularQueue<TRefCountPtr<FNetBlob>> UnreliableQueue;
 	FReliableNetBlobQueue* ReliableQueue;
@@ -202,7 +202,7 @@ public:
 
 	void DropAllAttachments(ENetObjectAttachmentType Type, uint32 ObjectIndex);
 
-	void Deserialize(FNetSerializationContext& Context, ENetObjectAttachmentType Type, uint32 ObjectIndex, FNetHandle NetHandle);
+	void Deserialize(FNetSerializationContext& Context, ENetObjectAttachmentType Type, uint32 ObjectIndex, FNetRefHandle RefHandle);
 
 	FNetObjectAttachmentReceiveQueue* GetQueue(ENetObjectAttachmentType Type, uint32 ObjectIndex);
 
diff --git a/Engine/Source/Runtime/Experimental/Iris/Core/Private/Iris/ReplicationSystem/Conditionals/ReplicationConditionals.cpp b/Engine/Source/Runtime/Experimental/Iris/Core/Private/Iris/ReplicationSystem/Conditionals/ReplicationConditionals.cpp
--- a/Engine/Source/Runtime/Experimental/Iris/Core/Private/Iris/ReplicationSystem/Conditionals/ReplicationConditionals.cpp
+++ b/Engine/Source/Runtime/Experimental/Iris/Core/Private/Iris/ReplicationSystem/Conditionals/ReplicationConditionals.cpp
@@ -35,7 +35,7 @@ void FReplicationConditionals::Init(FReplicationConditionalsInitParams& Params)
 	}
 #endif
 
-	NetHandleManager = Params.NetHandleManager;
+	NetRefHandleManager = Params.NetRefHandleManager;
 	ReplicationFiltering = Params.ReplicationFiltering;
 	ReplicationConnections = Params.ReplicationConnections;
 	BaselineInvalidationTracker = Params.BaselineInvalidationTracker;
@@ -47,7 +47,7 @@ void FReplicationConditionals::Init(FReplicationConditionalsInitParams& Params)
 	ConnectionInfos.SetNum(MaxConnectionCount + 1U);
 }
 
-bool FReplicationConditionals::SetConditionConnectionFilter(FInternalNetHandle ObjectIndex, EReplicationCondition Condition, uint32 ConnectionId, bool bEnable)
+bool FReplicationConditionals::SetConditionConnectionFilter(FInternalNetRefIndex ObjectIndex, EReplicationCondition Condition, uint32 ConnectionId, bool bEnable)
 {
 	if (ConnectionId >= MaxConnectionCount)
 	{
@@ -87,7 +87,7 @@ void FReplicationConditionals::RemoveConnection(uint32 ConnectionId)
 	ConnectionInfo.ObjectConditionals.Empty();
 }
 
-bool FReplicationConditionals::SetCondition(FInternalNetHandle ObjectIndex, EReplicationCondition Condition, bool bEnable)
+bool FReplicationConditionals::SetCondition(FInternalNetRefIndex ObjectIndex, EReplicationCondition Condition, bool bEnable)
 {
 	if (!ensure(Condition != EReplicationCondition::RoleAutonomous))
 	{
@@ -111,14 +111,14 @@ bool FReplicationConditionals::SetCondition(FInternalNetHandle ObjectIndex, ERep
 	return false;
 }
 
-void FReplicationConditionals::InitPropertyCustomConditions(FInternalNetHandle ObjectIndex)
+void FReplicationConditionals::InitPropertyCustomConditions(FInternalNetRefIndex ObjectIndex)
 {
 	IRIS_PROFILER_SCOPE(FReplicationConditionals_InitPropertyCustomConditions);
 
-	const FNetHandleManager::FReplicatedObjectData& ReplicatedObjectData = NetHandleManager->GetReplicatedObjectDataNoCheck(ObjectIndex);
+	const FNetRefHandleManager::FReplicatedObjectData& ReplicatedObjectData = NetRefHandleManager->GetReplicatedObjectDataNoCheck(ObjectIndex);
 	const FReplicationProtocol* Protocol = ReplicatedObjectData.Protocol;
 
-	if (NetHandleManager->GetReplicatedObjectStateBufferNoCheck(ObjectIndex) == nullptr || !EnumHasAnyFlags(Protocol->ProtocolTraits, EReplicationProtocolTraits::HasLifetimeConditionals))
+	if (NetRefHandleManager->GetReplicatedObjectStateBufferNoCheck(ObjectIndex) == nullptr || !EnumHasAnyFlags(Protocol->ProtocolTraits, EReplicationProtocolTraits::HasLifetimeConditionals))
 	{
 		return;
 	}
@@ -189,14 +189,14 @@ void FReplicationConditionals::InitPropertyCustomConditions(FInternalNetHandle O
 }
 
 // N.B. Calls can come for properties that have been disabled. We must handle such cases gracefully.
-bool FReplicationConditionals::SetPropertyCustomCondition(FInternalNetHandle ObjectIndex, const void* Owner, uint16 RepIndex, bool bIsActive)
+bool FReplicationConditionals::SetPropertyCustomCondition(FInternalNetRefIndex ObjectIndex, const void* Owner, uint16 RepIndex, bool bIsActive)
 {
 	IRIS_PROFILER_SCOPE(FReplicationConditionals_SetPropertyCustomCondition);
 
-	const FNetHandleManager::FReplicatedObjectData& ReplicatedObjectData = NetHandleManager->GetReplicatedObjectDataNoCheck(ObjectIndex);
+	const FNetRefHandleManager::FReplicatedObjectData& ReplicatedObjectData = NetRefHandleManager->GetReplicatedObjectDataNoCheck(ObjectIndex);
 	const FReplicationProtocol* Protocol = ReplicatedObjectData.Protocol;
 
-	if (NetHandleManager->GetReplicatedObjectStateBufferNoCheck(ObjectIndex) == nullptr || !EnumHasAnyFlags(Protocol->ProtocolTraits, EReplicationProtocolTraits::HasLifetimeConditionals))
+	if (NetRefHandleManager->GetReplicatedObjectStateBufferNoCheck(ObjectIndex) == nullptr || !EnumHasAnyFlags(Protocol->ProtocolTraits, EReplicationProtocolTraits::HasLifetimeConditionals))
 	{
 		return false;
 	}
@@ -324,17 +324,17 @@ void FReplicationConditionals::Update()
 	UpdateObjectsInScope();
 }
 
-void FReplicationConditionals::GetChildSubObjectsToReplicate(uint32 ReplicatingConnectionId, const FConditionalsMask& LifetimeConditionals,  const FInternalNetHandle ParentObjectIndex, FSubObjectsToReplicateArray& OutSubObjectsToReplicate)
+void FReplicationConditionals::GetChildSubObjectsToReplicate(uint32 ReplicatingConnectionId, const FConditionalsMask& LifetimeConditionals,  const FInternalNetRefIndex ParentObjectIndex, FSubObjectsToReplicateArray& OutSubObjectsToReplicate)
 {
 	// To mimic old system we use a weird replication order based on hierarchy, SubSubObjects are replicated before the parent
 	FChildSubObjectsInfo SubObjectsInfo;
-	if (NetHandleManager->GetChildSubObjects(ParentObjectIndex, SubObjectsInfo))
+	if (NetRefHandleManager->GetChildSubObjects(ParentObjectIndex, SubObjectsInfo))
 	{
 		if (SubObjectsInfo.SubObjectLifeTimeConditions == nullptr)
 		{
 			for (uint32 ArrayIndex = 0; ArrayIndex < SubObjectsInfo.NumSubObjects; ++ArrayIndex)
 			{
-				FInternalNetHandle SubObjectIndex = SubObjectsInfo.ChildSubObjects[ArrayIndex];
+				FInternalNetRefIndex SubObjectIndex = SubObjectsInfo.ChildSubObjects[ArrayIndex];
 				GetChildSubObjectsToReplicate(ReplicatingConnectionId, LifetimeConditionals, SubObjectIndex, OutSubObjectsToReplicate);
 				OutSubObjectsToReplicate.Add(SubObjectIndex);
 			}
@@ -347,7 +347,7 @@ void FReplicationConditionals::GetChildSubObjectsToReplicate(uint32 ReplicatingC
 				const ELifetimeCondition LifeTimeCondition = (ELifetimeCondition)SubObjectsInfo.SubObjectLifeTimeConditions[ArrayIndex];
 				if (LifeTimeCondition == COND_NetGroup)
 				{
-					const FInternalNetHandle SubObjectIndex = SubObjectsInfo.ChildSubObjects[ArrayIndex];
+					const FInternalNetRefIndex SubObjectIndex = SubObjectsInfo.ChildSubObjects[ArrayIndex];
 
 					uint32 GroupCount = 0U;
 					if (const FNetObjectGroupHandle* GroupMemberships = NetObjectGroups->GetGroupMemberships(SubObjectIndex, GroupCount))
@@ -382,7 +382,7 @@ void FReplicationConditionals::GetChildSubObjectsToReplicate(uint32 ReplicatingC
 				}
 				else if (LifetimeConditionals.IsConditionEnabled(LifeTimeCondition))
 				{
-					const FInternalNetHandle SubObjectIndex = SubObjectsInfo.ChildSubObjects[ArrayIndex];
+					const FInternalNetRefIndex SubObjectIndex = SubObjectsInfo.ChildSubObjects[ArrayIndex];
 					GetChildSubObjectsToReplicate(ReplicatingConnectionId, LifetimeConditionals, SubObjectIndex, OutSubObjectsToReplicate);
 					OutSubObjectsToReplicate.Add(SubObjectIndex);
 				}
@@ -391,7 +391,7 @@ void FReplicationConditionals::GetChildSubObjectsToReplicate(uint32 ReplicatingC
 	}
 }
 
-void FReplicationConditionals::GetSubObjectsToReplicate(uint32 ReplicationConnectionId, FInternalNetHandle RootObjectIndex, FSubObjectsToReplicateArray& OutSubObjectsToReplicate)
+void FReplicationConditionals::GetSubObjectsToReplicate(uint32 ReplicationConnectionId, FInternalNetRefIndex RootObjectIndex, FSubObjectsToReplicateArray& OutSubObjectsToReplicate)
 {
 	IRIS_PROFILER_SCOPE(FReplicationConditionals_GetSubObjectsToReplicate);
 
@@ -402,7 +402,7 @@ void FReplicationConditionals::GetSubObjectsToReplicate(uint32 ReplicationConnec
 	GetChildSubObjectsToReplicate(ReplicationConnectionId, LifetimeConditionals, RootObjectIndex, OutSubObjectsToReplicate);
 }
 
-bool FReplicationConditionals::ApplyConditionalsToChangeMask(uint32 ReplicatingConnectionId, FInternalNetHandle ParentObjectIndex, FInternalNetHandle ObjectIndex, uint32* ChangeMaskData, const uint32* ConditionalChangeMaskData, const FReplicationProtocol* Protocol)
+bool FReplicationConditionals::ApplyConditionalsToChangeMask(uint32 ReplicatingConnectionId, FInternalNetRefIndex ParentObjectIndex, FInternalNetRefIndex ObjectIndex, uint32* ChangeMaskData, const uint32* ConditionalChangeMaskData, const FReplicationProtocol* Protocol)
 {
 	IRIS_PROFILER_SCOPE(FReplicationConditionals_ApplyConditionalsToChangeMask);
 
@@ -535,8 +535,8 @@ void FReplicationConditionals::UpdateObjectsInScope()
 {
 	IRIS_PROFILER_SCOPE(FReplicationConditionals_UpdateObjectsInScope);
 
-	const FNetBitArray& ObjectsInScope = NetHandleManager->GetScopableInternalIndices();
-	const FNetBitArray& PrevObjectsInScope = NetHandleManager->GetPrevFrameScopableInternalIndices();
+	const FNetBitArray& ObjectsInScope = NetRefHandleManager->GetScopableInternalIndices();
+	const FNetBitArray& PrevObjectsInScope = NetRefHandleManager->GetPrevFrameScopableInternalIndices();
 
 	const uint32 WordCountForModifiedWords = Align(FPlatformMath::Max(MaxObjectCount, 1U), 32U)/32U;
 	TArray<uint32> ModifiedWords;
@@ -592,7 +592,7 @@ void FReplicationConditionals::UpdateObjectsInScope()
 	}
 }
 
-FReplicationConditionals::FConditionalsMask FReplicationConditionals::GetLifetimeConditionals(uint32 ReplicatingConnectionId, FInternalNetHandle ParentObjectIndex) const
+FReplicationConditionals::FConditionalsMask FReplicationConditionals::GetLifetimeConditionals(uint32 ReplicatingConnectionId, FInternalNetRefIndex ParentObjectIndex) const
 {
 	FConditionalsMask ConditionalsMask{0};
 
@@ -619,13 +619,13 @@ FReplicationConditionals::FConditionalsMask FReplicationConditionals::GetLifetim
 	return ConditionalsMask;
 }
 
-void FReplicationConditionals::ClearPerObjectInfo(FInternalNetHandle ObjectIndex)
+void FReplicationConditionals::ClearPerObjectInfo(FInternalNetRefIndex ObjectIndex)
 {
 	FPerObjectInfo& PerObjectInfo = PerObjectInfos.GetData()[ObjectIndex];
 	PerObjectInfo = {};
 }
 
-void FReplicationConditionals::ClearConnectionInfosForObject(const FNetBitArray& ValidConnections, uint32 MaxConnectionId, FInternalNetHandle ObjectIndex)
+void FReplicationConditionals::ClearConnectionInfosForObject(const FNetBitArray& ValidConnections, uint32 MaxConnectionId, FInternalNetRefIndex ObjectIndex)
 {
 	IRIS_PROFILER_SCOPE(FReplicationConditionals_ClearConnectionInfosForObject);
 
diff --git a/Engine/Source/Runtime/Experimental/Iris/Core/Private/Iris/ReplicationSystem/Conditionals/ReplicationConditionals.h b/Engine/Source/Runtime/Experimental/Iris/Core/Private/Iris/ReplicationSystem/Conditionals/ReplicationConditionals.h
--- a/Engine/Source/Runtime/Experimental/Iris/Core/Private/Iris/ReplicationSystem/Conditionals/ReplicationConditionals.h
+++ b/Engine/Source/Runtime/Experimental/Iris/Core/Private/Iris/ReplicationSystem/Conditionals/ReplicationConditionals.h
@@ -13,8 +13,8 @@ namespace UE::Net
 	namespace Private
 	{
 		class FDeltaCompressionBaselineInvalidationTracker;
-		typedef uint32 FInternalNetHandle;
-		class FNetHandleManager;
+		typedef uint32 FInternalNetRefIndex;
+		class FNetRefHandleManager;
 		class FReplicationConnections;
 		class FReplicationFiltering;
 		class FNetObjectGroups;
@@ -26,7 +26,7 @@ namespace UE::Net::Private
 
 struct FReplicationConditionalsInitParams
 {
-	const FNetHandleManager* NetHandleManager = nullptr;
+	const FNetRefHandleManager* NetRefHandleManager = nullptr;
 	const FReplicationFiltering* ReplicationFiltering = nullptr;
 	const FReplicationConnections* ReplicationConnections = nullptr;
 	const FNetObjectGroups* NetObjectGroups = nullptr;
@@ -53,19 +53,19 @@ public:
 	void AddConnection(uint32 ConnectionId);
 	void RemoveConnection(uint32 ConnectionId);
 
-	bool SetConditionConnectionFilter(FInternalNetHandle ObjectIndex, EReplicationCondition Condition, uint32 ConnectionId, bool bEnable);
-	bool SetCondition(FInternalNetHandle ObjectIndex, EReplicationCondition Condition, bool bEnable);
+	bool SetConditionConnectionFilter(FInternalNetRefIndex ObjectIndex, EReplicationCondition Condition, uint32 ConnectionId, bool bEnable);
+	bool SetCondition(FInternalNetRefIndex ObjectIndex, EReplicationCondition Condition, bool bEnable);
 
 	// For property custom conditions only
-	void InitPropertyCustomConditions(FInternalNetHandle ObjectIndex);
-	bool SetPropertyCustomCondition(FInternalNetHandle ObjectIndex, const void* Owner, uint16 RepIndex, bool bIsActive);
+	void InitPropertyCustomConditions(FInternalNetRefIndex ObjectIndex);
+	bool SetPropertyCustomCondition(FInternalNetRefIndex ObjectIndex, const void* Owner, uint16 RepIndex, bool bIsActive);
 
 	void Update();
 
-	bool ApplyConditionalsToChangeMask(uint32 ReplicatingConnectionId, FInternalNetHandle ParentObjectIndex, FInternalNetHandle ObjectIndex, uint32* ChangeMaskData, const uint32* ConditionalChangeMaskData, const FReplicationProtocol* Protocol);
+	bool ApplyConditionalsToChangeMask(uint32 ReplicatingConnectionId, FInternalNetRefIndex ParentObjectIndex, FInternalNetRefIndex ObjectIndex, uint32* ChangeMaskData, const uint32* ConditionalChangeMaskData, const FReplicationProtocol* Protocol);
 
-	using FSubObjectsToReplicateArray = TArray<FInternalNetHandle, TInlineAllocator<32>>;
-	void GetSubObjectsToReplicate(uint32 ReplicationConnectionId, FInternalNetHandle ParentObjectIndex, FSubObjectsToReplicateArray& OutSubObjectsToReplicate);
+	using FSubObjectsToReplicateArray = TArray<FInternalNetRefIndex, TInlineAllocator<32>>;
+	void GetSubObjectsToReplicate(uint32 ReplicationConnectionId, FInternalNetRefIndex ParentObjectIndex, FSubObjectsToReplicateArray& OutSubObjectsToReplicate);
 
 private:
 	struct FPerObjectInfo
@@ -99,18 +99,18 @@ private:
 private:
 	void UpdateObjectsInScope();
 
-	FConditionalsMask GetLifetimeConditionals(uint32 ReplicatingConnectionId, FInternalNetHandle ParentObjectIndex) const;
-	bool ObjectHasLifetimeConditionals(FInternalNetHandle ObjectIndex) const;	
+	FConditionalsMask GetLifetimeConditionals(uint32 ReplicatingConnectionId, FInternalNetRefIndex ParentObjectIndex) const;
+	bool ObjectHasLifetimeConditionals(FInternalNetRefIndex ObjectIndex) const;	
 
-	FPerObjectInfo* GetPerObjectInfo(FInternalNetHandle ObjectIndex);
-	const FPerObjectInfo* GetPerObjectInfo(FInternalNetHandle ObjectIndex) const;
-	void ClearPerObjectInfo(FInternalNetHandle ObjectIndex);
-	void ClearConnectionInfosForObject(const FNetBitArray& ValidConnections, uint32 MaxConnectionId, FInternalNetHandle ObjectIndex);
+	FPerObjectInfo* GetPerObjectInfo(FInternalNetRefIndex ObjectIndex);
+	const FPerObjectInfo* GetPerObjectInfo(FInternalNetRefIndex ObjectIndex) const;
+	void ClearPerObjectInfo(FInternalNetRefIndex ObjectIndex);
+	void ClearConnectionInfosForObject(const FNetBitArray& ValidConnections, uint32 MaxConnectionId, FInternalNetRefIndex ObjectIndex);
 
-	void GetChildSubObjectsToReplicate(uint32 ReplicatingConnectionId, const FConditionalsMask& LifetimeConditionals,  const FInternalNetHandle ParentObjectIndex, FSubObjectsToReplicateArray& OutSubObjectsToReplicate);
+	void GetChildSubObjectsToReplicate(uint32 ReplicatingConnectionId, const FConditionalsMask& LifetimeConditionals,  const FInternalNetRefIndex ParentObjectIndex, FSubObjectsToReplicateArray& OutSubObjectsToReplicate);
 
 private:
-	const FNetHandleManager* NetHandleManager = nullptr;
+	const FNetRefHandleManager* NetRefHandleManager = nullptr;
 	const FReplicationFiltering* ReplicationFiltering = nullptr;
 	const FReplicationConnections* ReplicationConnections = nullptr;
 	FDeltaCompressionBaselineInvalidationTracker* BaselineInvalidationTracker = nullptr;
@@ -123,12 +123,12 @@ private:
 	TArray<FPerConnectionInfo> ConnectionInfos;
 };
 
-inline FReplicationConditionals::FPerObjectInfo* FReplicationConditionals::GetPerObjectInfo(FInternalNetHandle ObjectIndex)
+inline FReplicationConditionals::FPerObjectInfo* FReplicationConditionals::GetPerObjectInfo(FInternalNetRefIndex ObjectIndex)
 {
 	return PerObjectInfos.GetData() + ObjectIndex;
 }
 
-inline const FReplicationConditionals::FPerObjectInfo* FReplicationConditionals::GetPerObjectInfo(FInternalNetHandle ObjectIndex) const
+inline const FReplicationConditionals::FPerObjectInfo* FReplicationConditionals::GetPerObjectInfo(FInternalNetRefIndex ObjectIndex) const
 {
 	return PerObjectInfos.GetData() + ObjectIndex;
 }
diff --git a/Engine/Source/Runtime/Experimental/Iris/Core/Private/Iris/ReplicationSystem/DeltaCompression/DeltaCompressionBaselineInvalidationTracker.cpp b/Engine/Source/Runtime/Experimental/Iris/Core/Private/Iris/ReplicationSystem/DeltaCompression/DeltaCompressionBaselineInvalidationTracker.cpp
--- a/Engine/Source/Runtime/Experimental/Iris/Core/Private/Iris/ReplicationSystem/DeltaCompression/DeltaCompressionBaselineInvalidationTracker.cpp
+++ b/Engine/Source/Runtime/Experimental/Iris/Core/Private/Iris/ReplicationSystem/DeltaCompression/DeltaCompressionBaselineInvalidationTracker.cpp
@@ -21,7 +21,7 @@ void FDeltaCompressionBaselineInvalidationTracker::Init(FDeltaCompressionBaselin
 	InvalidatedObjects.Init(InitParams.MaxObjectCount);
 }
 
-void FDeltaCompressionBaselineInvalidationTracker::InvalidateBaselines(FInternalNetHandle ObjectIndex, uint32 ConnId)
+void FDeltaCompressionBaselineInvalidationTracker::InvalidateBaselines(FInternalNetRefIndex ObjectIndex, uint32 ConnId)
 {
 	if (InvalidatedObjects.GetBit(ObjectIndex))
 	{
diff --git a/Engine/Source/Runtime/Experimental/Iris/Core/Private/Iris/ReplicationSystem/DeltaCompression/DeltaCompressionBaselineInvalidationTracker.h b/Engine/Source/Runtime/Experimental/Iris/Core/Private/Iris/ReplicationSystem/DeltaCompression/DeltaCompressionBaselineInvalidationTracker.h
--- a/Engine/Source/Runtime/Experimental/Iris/Core/Private/Iris/ReplicationSystem/DeltaCompression/DeltaCompressionBaselineInvalidationTracker.h
+++ b/Engine/Source/Runtime/Experimental/Iris/Core/Private/Iris/ReplicationSystem/DeltaCompression/DeltaCompressionBaselineInvalidationTracker.h
@@ -12,7 +12,7 @@
 namespace UE::Net::Private
 {
 	class FDeltaCompressionBaselineManager;
-	typedef uint32 FInternalNetHandle;
+	typedef uint32 FInternalNetRefIndex;
 }
 
 namespace UE::Net::Private
@@ -35,7 +35,7 @@ public:
 	struct FInvalidationInfo
 	{
 		uint32 ConnId = InvalidateBaselineForAllConnections;
-		FInternalNetHandle ObjectIndex = 0U;
+		FInternalNetRefIndex ObjectIndex = 0U;
 	};
 
 public:
@@ -44,7 +44,7 @@ public:
 
 	void Init(FDeltaCompressionBaselineInvalidationTrackerInitParams& InitParams);
 
-	void InvalidateBaselines(FInternalNetHandle ObjectIndex, uint32 ConnId);
+	void InvalidateBaselines(FInternalNetRefIndex ObjectIndex, uint32 ConnId);
 
 	// Returns an array of objects with enabled conditions such that DC baselines need to be invalidated.
 	TArrayView<const FInvalidationInfo> GetBaselineInvalidationInfos() const;
diff --git a/Engine/Source/Runtime/Experimental/Iris/Core/Private/Iris/ReplicationSystem/DeltaCompression/DeltaCompressionBaselineManager.cpp b/Engine/Source/Runtime/Experimental/Iris/Core/Private/Iris/ReplicationSystem/DeltaCompression/DeltaCompressionBaselineManager.cpp
--- a/Engine/Source/Runtime/Experimental/Iris/Core/Private/Iris/ReplicationSystem/DeltaCompression/DeltaCompressionBaselineManager.cpp
+++ b/Engine/Source/Runtime/Experimental/Iris/Core/Private/Iris/ReplicationSystem/DeltaCompression/DeltaCompressionBaselineManager.cpp
@@ -55,7 +55,7 @@ FDeltaCompressionBaselineManager::FDeltaCompressionBaselineManager()
 {
 	// If alignment of storage type for FPerObjectInfo is too small then it's totally unsafe to cast such pointers to FPerObjectInfo.
 	static_assert(InvalidObjectInfoIndex == 0, "InvalidObjectInfoIndex needs to be zero in order for current code to function properly.");
-	static_assert(InvalidInternalIndex == FNetHandleManager::InvalidInternalIndex, "InvalidInternalIndex differs from FNetHandleManager::InvalidInternalIndex");
+	static_assert(InvalidInternalIndex == FNetRefHandleManager::InvalidInternalIndex, "InvalidInternalIndex differs from FNetRefHandleManager::InvalidInternalIndex");
 }
 
 FDeltaCompressionBaselineManager::~FDeltaCompressionBaselineManager()
@@ -66,7 +66,7 @@ FDeltaCompressionBaselineManager::~FDeltaCompressionBaselineManager()
 void FDeltaCompressionBaselineManager::Init(FDeltaCompressionBaselineManagerInitParams& InitParams)
 {
 	Connections = InitParams.Connections;
-	NetHandleManager = InitParams.NetHandleManager;
+	NetRefHandleManager = InitParams.NetRefHandleManager;
 	BaselineInvalidationTracker = InitParams.BaselineInvalidationTracker;
 	ReplicationSystem = InitParams.ReplicationSystem;
 	MaxConnectionCount = InitParams.Connections->GetMaxConnectionCount();
@@ -135,12 +135,12 @@ void FDeltaCompressionBaselineManager::RemoveConnection(uint32 ConnectionId)
 	ReleaseBaselinesForConnection(ConnectionId);
 }
 
-ENetObjectDeltaCompressionStatus FDeltaCompressionBaselineManager::GetDeltaCompressionStatus(FInternalNetHandle Index) const
+ENetObjectDeltaCompressionStatus FDeltaCompressionBaselineManager::GetDeltaCompressionStatus(FInternalNetRefIndex Index) const
 {
 	return DeltaCompressionEnabledObjects.GetBit(Index) ? ENetObjectDeltaCompressionStatus::Allow : ENetObjectDeltaCompressionStatus::Disallow;
 }
 
-void FDeltaCompressionBaselineManager::SetDeltaCompressionStatus(FInternalNetHandle ObjectIndex, ENetObjectDeltaCompressionStatus Status)
+void FDeltaCompressionBaselineManager::SetDeltaCompressionStatus(FInternalNetRefIndex ObjectIndex, ENetObjectDeltaCompressionStatus Status)
 {
 	// If DC is disabled or the object doesn't support it then mark the object as not DC enabled, regardless of the user's wishes.
 	bool bEnableDCForObject = (Status == ENetObjectDeltaCompressionStatus::Allow ? true : false);
@@ -158,8 +158,8 @@ void FDeltaCompressionBaselineManager::UpdateScope()
 	using SignedWordType = TSignedIntType<sizeof(WordType)>::Type;
 	constexpr SIZE_T BitsPerWord = sizeof(WordType)*8U;
 
-	const FNetBitArray& ObjectsInScope = NetHandleManager->GetScopableInternalIndices();
-	const FNetBitArray& PrevObjectsInScope = NetHandleManager->GetPrevFrameScopableInternalIndices();
+	const FNetBitArray& ObjectsInScope = NetRefHandleManager->GetScopableInternalIndices();
+	const FNetBitArray& PrevObjectsInScope = NetRefHandleManager->GetPrevFrameScopableInternalIndices();
 
 	const WordType* ObjectsInScopeStorage = ObjectsInScope.GetData();
 	const WordType* PrevObjectsInScopeStorage = PrevObjectsInScope.GetData();
@@ -288,7 +288,7 @@ FDeltaCompressionBaselineManager::FPerObjectInfo* FDeltaCompressionBaselineManag
 		ObjectInfo->ObjectIndex = ObjectIndex;
 		
 		// Allocate storage for changemasks for all connections
-		const FNetHandleManager::FReplicatedObjectData& ObjectData = NetHandleManager->GetReplicatedObjectDataNoCheck(ObjectIndex);
+		const FNetRefHandleManager::FReplicatedObjectData& ObjectData = NetRefHandleManager->GetReplicatedObjectDataNoCheck(ObjectIndex);
 		const FReplicationProtocol* Protocol = ObjectData.Protocol;
 		const uint32 ChangeMaskStride = Align(Protocol->ChangeMaskBitCount, sizeof(ChangeMaskStorageType)*8U)/(sizeof(ChangeMaskStorageType)*8U);
 		if (ChangeMaskStride > 0)
@@ -588,7 +588,7 @@ void FDeltaCompressionBaselineManager::DestructBaselineSharingContext(FBaselineS
 
 bool FDeltaCompressionBaselineManager::DoesObjectSupportDeltaCompression(uint32 ObjectIndex) const
 {
-	const FNetHandleManager::FReplicatedObjectData& ReplicatedObjectData = NetHandleManager->GetReplicatedObjectDataNoCheck(ObjectIndex);
+	const FNetRefHandleManager::FReplicatedObjectData& ReplicatedObjectData = NetRefHandleManager->GetReplicatedObjectDataNoCheck(ObjectIndex);
 	const FReplicationProtocol* Protocol = ReplicatedObjectData.Protocol;
 	return EnumHasAnyFlags(Protocol->ProtocolTraits, EReplicationProtocolTraits::SupportsDeltaCompression);
 }
diff --git a/Engine/Source/Runtime/Experimental/Iris/Core/Private/Iris/ReplicationSystem/DeltaCompression/DeltaCompressionBaselineManager.h b/Engine/Source/Runtime/Experimental/Iris/Core/Private/Iris/ReplicationSystem/DeltaCompression/DeltaCompressionBaselineManager.h
--- a/Engine/Source/Runtime/Experimental/Iris/Core/Private/Iris/ReplicationSystem/DeltaCompression/DeltaCompressionBaselineManager.h
+++ b/Engine/Source/Runtime/Experimental/Iris/Core/Private/Iris/ReplicationSystem/DeltaCompression/DeltaCompressionBaselineManager.h
@@ -18,8 +18,8 @@ namespace UE::Net::Private
 {
 	struct FChangeMaskCache;
 	class FDeltaCompressionBaselineInvalidationTracker;
-	typedef uint32 FInternalNetHandle;
-	class FNetHandleManager;
+	typedef uint32 FInternalNetRefIndex;
+	class FNetRefHandleManager;
 	class FReplicationConnections;
 }
 
@@ -33,7 +33,7 @@ struct FDeltaCompressionBaselineManagerInitParams
 {
 	const FDeltaCompressionBaselineInvalidationTracker* BaselineInvalidationTracker = nullptr;
 	FReplicationConnections* Connections = nullptr;
-	const FNetHandleManager* NetHandleManager = nullptr;
+	const FNetRefHandleManager* NetRefHandleManager = nullptr;
 	FReplicationStateStorage* ReplicationStateStorage = nullptr;
 	UReplicationSystem* ReplicationSystem = nullptr;
 	uint32 MaxObjectCount = 0;
@@ -67,8 +67,8 @@ public:
 	void AddConnection(uint32 ConnectionId);
 	void RemoveConnection(uint32 ConnectionId);
 
-	void SetDeltaCompressionStatus(FInternalNetHandle Index, ENetObjectDeltaCompressionStatus Status);
-	ENetObjectDeltaCompressionStatus GetDeltaCompressionStatus(FInternalNetHandle Index) const;
+	void SetDeltaCompressionStatus(FInternalNetRefIndex Index, ENetObjectDeltaCompressionStatus Status);
+	ENetObjectDeltaCompressionStatus GetDeltaCompressionStatus(FInternalNetRefIndex Index) const;
 
 	uint32 GetMaxDeltaCompressedObjectCount() const;
 
@@ -103,7 +103,7 @@ private:
 		InvalidObjectInfoIndex = 0,
 		InvalidBaselineStateInfoIndex = 0,
 
-		// Must match NetHandleManager::InvalidInternalIndex
+		// Must match NetRefHandleManager::InvalidInternalIndex
 		InvalidInternalIndex = 0,
 
 		ObjectInfoGrowCount = 128,
@@ -227,7 +227,7 @@ private:
 	FGlobalChangeMaskAllocator ChangeMaskAllocator;
 
 	FReplicationConnections* Connections = nullptr;
-	const FNetHandleManager* NetHandleManager = nullptr;
+	const FNetRefHandleManager* NetRefHandleManager = nullptr;
 	UReplicationSystem* ReplicationSystem = nullptr;
 	const FDeltaCompressionBaselineInvalidationTracker* BaselineInvalidationTracker = nullptr;
 	uint32 MaxConnectionCount = 0;
