[PCG] Improved collapse support in spawn actor node, adding support for multiple meshes and avoiding a crash when the template actor contained ISMC's.

[Engine] Made ISMComponentDescriptor availible in game code to support runtime PCG generation.

diff --git a/Engine/Plugins/Experimental/PCG/Source/PCG/Private/Elements/PCGSpawnActor.cpp b/Engine/Plugins/Experimental/PCG/Source/PCG/Private/Elements/PCGSpawnActor.cpp
--- a/Engine/Plugins/Experimental/PCG/Source/PCG/Private/Elements/PCGSpawnActor.cpp
+++ b/Engine/Plugins/Experimental/PCG/Source/PCG/Private/Elements/PCGSpawnActor.cpp
@@ -15,6 +15,7 @@
 #include "Components/InstancedStaticMeshComponent.h"
 #include "GameFramework/Actor.h"
 #include "Engine/World.h"
+#include "ISMPartition/ISMComponentDescriptor.h"
 
 UPCGNode* UPCGSpawnActorSettings::CreateNode() const
 {
@@ -155,54 +156,79 @@ bool FPCGSpawnActorElement::ExecuteInternal(FPCGContext* Context) const
 
 		// Spawn actors/populate ISM
 		{
-			UInstancedStaticMeshComponent* ISMC = nullptr;
-			TArray<FTransform> Instances;
-
-			// If we are collapsing actors, we need to get the mesh & prep the ISMC
 			if (Settings->Option == EPCGSpawnActorOption::CollapseActors)
 			{
 				TArray<UActorComponent*> Components;
 				UPCGActorHelpers::GetActorClassDefaultComponents(Settings->TemplateActorClass, Components, UStaticMeshComponent::StaticClass());
-				UStaticMeshComponent* FirstSMC = nullptr;
-				UStaticMesh* Mesh = nullptr;
+				
+				TMap<FISMComponentDescriptor, TArray<FTransform>> MeshDescriptorTransforms;
 
 				for (UActorComponent* Component : Components)
 				{
-					if (UStaticMeshComponent* SMC = Cast<UStaticMeshComponent>(Component))
+					if (UStaticMeshComponent* StaticMeshComponent = Cast<UStaticMeshComponent>(Component))
 					{
-						FirstSMC = SMC;
-						Mesh = SMC->GetStaticMesh();
-						if (Mesh)
+						FISMComponentDescriptor ISMDescriptor;
+						ISMDescriptor.InitFrom(StaticMeshComponent);
+						ISMDescriptor.ComputeHash();
+
+						TArray<FTransform>& Transforms = MeshDescriptorTransforms.FindOrAdd(ISMDescriptor);
+
+						if (UInstancedStaticMeshComponent* InstancedStaticMeshComponent = Cast<UInstancedStaticMeshComponent>(StaticMeshComponent))
 						{
-							break;
+							const int32 NumInstances = InstancedStaticMeshComponent->GetInstanceCount();
+							Transforms.Reserve(Transforms.Num() + NumInstances);
+							
+							for (int32 InstanceIndex = 0; InstanceIndex < NumInstances; InstanceIndex++)
+							{
+								FTransform InstanceTransform;
+								if (InstancedStaticMeshComponent->GetInstanceTransform(InstanceIndex, InstanceTransform))
+								{
+									Transforms.Add(InstanceTransform);
+								}
+							}
+						}
+						else
+						{
+							Transforms.Add(StaticMeshComponent->GetRelativeTransform());
 						}
 					}
 				}
-
-				if (Mesh)
+				
+				for (const TPair<FISMComponentDescriptor, TArray<FTransform>>& ISMCDescriptorTransforms : MeshDescriptorTransforms)
 				{
-					FPCGISMCBuilderParameters Params;
-					Params.Mesh = Mesh;
-
-					ISMC = UPCGActorHelpers::GetOrCreateISMC(TargetActor, Context->SourceComponent.Get(), Params);
-					UEngine::CopyPropertiesForUnrelatedObjects(FirstSMC, ISMC);
-				}
-				else
-				{
-					PCGE_LOG(Error, "No supported mesh found");
+					const FISMComponentDescriptor& ISMCDescriptor = ISMCDescriptorTransforms.Key;
+					
+					UInstancedStaticMeshComponent* ISMC = UPCGActorHelpers::GetOrCreateISMC(TargetActor, Context->SourceComponent.Get(), ISMCDescriptor);
+					if (!ISMC)
+					{
+						continue;
+					}
+										
+					const TArray<FTransform>& ISMCTransforms = ISMCDescriptorTransforms.Value;
+					
+					TArray<FTransform> Transforms;
+					Transforms.Reserve(Points.Num() * ISMCTransforms.Num());
+					for (int32 PointIndex = 0; PointIndex < Points.Num(); PointIndex++)
+					{
+						const FPCGPoint& Point = Points[PointIndex];
+						for (int32 TransformIndex = 0; TransformIndex < ISMCTransforms.Num(); TransformIndex++)
+						{
+							const FTransform& Transform = ISMCTransforms[TransformIndex];
+							Transforms.Add(Transform * Point.Transform);
+						}
+					}
+					
+					ISMC->NumCustomDataFloats = 0;
+					ISMC->AddInstances(Transforms, false, true);
+					ISMC->UpdateBounds();
+					
+					PCGE_LOG(Verbose, "Added %d instances of %s to %s in %s", Transforms.Num(), *ISMC->GetStaticMesh().GetName(), *ISMC->GetName(), *TargetActor->GetActorNameOrLabel());
 				}
 			}
 
 			{
 				TRACE_CPUPROFILER_EVENT_SCOPE(FPCGSpawnActorElement::ExecuteInternal::SpawnActors);
-				if (Settings->Option == EPCGSpawnActorOption::CollapseActors && ISMC)
-				{
-					for (const FPCGPoint& Point : Points)
-					{
-						Instances.Add(Point.Transform);
-					}
-				}
-				else if (Settings->Option != EPCGSpawnActorOption::CollapseActors && (bHasAuthority || !bSpawnedActorsRequireAuthority))
+				if (Settings->Option != EPCGSpawnActorOption::CollapseActors && (bHasAuthority || !bSpawnedActorsRequireAuthority))
 				{
 					FActorSpawnParameters SpawnParams;
 					SpawnParams.Owner = TargetActor;
@@ -266,16 +292,6 @@ bool FPCGSpawnActorElement::ExecuteInternal(FPCGContext* Context) const
 					PCGE_LOG(Verbose, "Generated %d actors", Points.Num());
 				}
 			}
-
-			// Finalize
-			if (ISMC && Instances.Num() > 0)
-			{
-				ISMC->NumCustomDataFloats = 0;
-				ISMC->AddInstances(Instances, false, true);
-				ISMC->UpdateBounds();
-
-				PCGE_LOG(Verbose, "Added %d ISM instances", Instances.Num());
-			}
 		}
 
 		// Finally, pass through the input, in all cases: 
@@ -327,4 +343,4 @@ EPCGSettingsType UPCGSpawnActorSettings::GetType() const
 {
 	return GetSubgraph() ? EPCGSettingsType::Subgraph : EPCGSettingsType::Spawner;
 }
-#endif // WITH_EDITOR
\ No newline at end of file
+#endif // WITH_EDITOR
diff --git a/Engine/Plugins/Experimental/PCG/Source/PCG/Private/Helpers/PCGActorHelpers.cpp b/Engine/Plugins/Experimental/PCG/Source/PCG/Private/Helpers/PCGActorHelpers.cpp
--- a/Engine/Plugins/Experimental/PCG/Source/PCG/Private/Helpers/PCGActorHelpers.cpp
+++ b/Engine/Plugins/Experimental/PCG/Source/PCG/Private/Helpers/PCGActorHelpers.cpp
@@ -16,6 +16,7 @@
 #include "Engine/StaticMesh.h"
 #include "Engine/World.h"
 #include "WorldPartition/WorldPartition.h"
+#include "ISMPartition/ISMComponentDescriptor.h"
 
 #if WITH_EDITOR
 #include "Editor.h"
@@ -25,6 +26,73 @@
 #include "ObjectTools.h"
 #endif
 
+UInstancedStaticMeshComponent* UPCGActorHelpers::GetOrCreateISMC(AActor* InTargetActor, UPCGComponent* InSourceComponent, const FISMComponentDescriptor& InISMCDescriptor)
+{
+	check(InTargetActor);
+	check(InSourceComponent);
+
+	const UStaticMesh* StaticMesh = InISMCDescriptor.StaticMesh;
+	if (!StaticMesh)
+	{
+		return nullptr;
+	}
+	
+	TArray<UPCGManagedISMComponent*> MISMCs;
+	InSourceComponent->ForEachManagedResource([&MISMCs](UPCGManagedResource* InResource)
+	{
+		if (UPCGManagedISMComponent* Resource = Cast<UPCGManagedISMComponent>(InResource))
+		{
+			MISMCs.Add(Resource);
+		}
+	});
+
+	for (UPCGManagedISMComponent* MISMC : MISMCs)
+	{
+		UInstancedStaticMeshComponent* ISMC = MISMC->GetComponent();
+		if (ISMC)
+		{
+			FISMComponentDescriptor ManagedDescriptor;
+			ManagedDescriptor.InitFrom(MISMC->GetComponent());
+			ManagedDescriptor.ComputeHash();
+			
+			if (ManagedDescriptor == InISMCDescriptor)
+			{
+				MISMC->MarkAsUsed();
+				return ISMC;
+			}
+		}
+	}
+
+	InTargetActor->Modify();
+	
+	// Done as in InstancedStaticMesh.cpp
+#if WITH_EDITOR
+	const bool bMeshHasNaniteData = StaticMesh->NaniteSettings.bEnabled;
+#else
+	const bool bMeshHasNaniteData = StaticMesh->GetRenderData()->NaniteResources.PageStreamingStates.Num() > 0;
+#endif
+
+	const TSubclassOf<UInstancedStaticMeshComponent> ComponentClass = bMeshHasNaniteData ? UInstancedStaticMeshComponent::StaticClass() : UHierarchicalInstancedStaticMeshComponent::StaticClass();
+	
+	UInstancedStaticMeshComponent* ISMC = NewObject<UInstancedStaticMeshComponent>(InTargetActor, ComponentClass);
+	InISMCDescriptor.InitComponent(ISMC);
+		
+	ISMC->RegisterComponent();
+
+	InTargetActor->AddInstanceComponent(ISMC);
+	
+	ISMC->AttachToComponent(InTargetActor->GetRootComponent(), FAttachmentTransformRules(EAttachmentRule::KeepRelative, EAttachmentRule::KeepWorld, EAttachmentRule::KeepWorld, false));
+	ISMC->ComponentTags.Add(InSourceComponent->GetFName());
+	ISMC->ComponentTags.Add(PCGHelpers::DefaultPCGTag);
+
+	// Create managed resource on source component
+	UPCGManagedISMComponent* Resource = NewObject<UPCGManagedISMComponent>(InSourceComponent);
+	Resource->GeneratedComponent = ISMC;
+	InSourceComponent->AddToManagedResources(Resource);
+
+	return ISMC;
+}
+
 UInstancedStaticMeshComponent* UPCGActorHelpers::GetOrCreateISMC(AActor* InTargetActor, UPCGComponent* InSourceComponent, const FPCGISMCBuilderParameters& InParams)
 {
 	UStaticMesh* InMesh = InParams.Mesh;
@@ -34,11 +102,6 @@ UInstancedStaticMeshComponent* UPCGActorHelpers::GetOrCreateISMC(AActor* InTarge
 	check(InTargetActor != nullptr && InMesh != nullptr);
 	check(InSourceComponent);
 
-	if (!InSourceComponent)
-	{
-		return nullptr;
-	}
-
 	TArray<UPCGManagedISMComponent*> MISMCs;
 	InSourceComponent->ForEachManagedResource([&MISMCs](UPCGManagedResource* InResource)
 	{
diff --git a/Engine/Plugins/Experimental/PCG/Source/PCG/Public/Helpers/PCGActorHelpers.h b/Engine/Plugins/Experimental/PCG/Source/PCG/Public/Helpers/PCGActorHelpers.h
--- a/Engine/Plugins/Experimental/PCG/Source/PCG/Public/Helpers/PCGActorHelpers.h
+++ b/Engine/Plugins/Experimental/PCG/Source/PCG/Public/Helpers/PCGActorHelpers.h
@@ -6,20 +6,20 @@
 #include "Kismet/BlueprintFunctionLibrary.h"
 #include "Templates/SubclassOf.h"
 #include "Engine/EngineTypes.h"
-#include "Engine/CollisionProfile.h"
 
 #include <type_traits>
 
 #include "PCGActorHelpers.generated.h"
 
 class AActor;
-class UInstancedStaticMeshComponent;
-class UStaticMesh;
-class UPCGComponent;
-class UMaterialInterface;
 class UActorComponent;
+class UInstancedStaticMeshComponent;
 class ULevel;
+class UMaterialInterface;
+class UPCGComponent;
+class UStaticMesh;
 class UWorld;
+struct FISMComponentDescriptor;
 
 struct FPCGISMCBuilderParameters
 {
@@ -38,7 +38,8 @@ class PCG_API UPCGActorHelpers : public UBlueprintFunctionLibrary
 	GENERATED_BODY()
 
 public:
-	static UInstancedStaticMeshComponent* GetOrCreateISMC(AActor* InActor, UPCGComponent* SourceComponent, const FPCGISMCBuilderParameters& Params);
+	static UInstancedStaticMeshComponent* GetOrCreateISMC(AActor* InTargetActor, UPCGComponent* InSourceComponent, const FISMComponentDescriptor& InISMCDescriptor);
+	static UInstancedStaticMeshComponent* GetOrCreateISMC(AActor* InTargetActor, UPCGComponent* SourceComponent, const FPCGISMCBuilderParameters& Params);
 	static bool DeleteActors(UWorld* World, const TArray<TSoftObjectPtr<AActor>>& ActorsToDelete);
 
 	/**
diff --git a/Engine/Source/Runtime/Engine/Private/ISMPartition/ISMComponentDescriptor.cpp b/Engine/Source/Runtime/Engine/Private/ISMPartition/ISMComponentDescriptor.cpp
--- a/Engine/Source/Runtime/Engine/Private/ISMPartition/ISMComponentDescriptor.cpp
+++ b/Engine/Source/Runtime/Engine/Private/ISMPartition/ISMComponentDescriptor.cpp
@@ -4,24 +4,16 @@
 
 #include UE_INLINE_GENERATED_CPP_BY_NAME(ISMComponentDescriptor)
 
-#if WITH_EDITOR
-
 #include "Serialization/ArchiveCrc32.h"
 #include "Components/InstancedStaticMeshComponent.h"
 #include "Components/HierarchicalInstancedStaticMeshComponent.h"
 
-#endif
-
 FISMComponentDescriptor::FISMComponentDescriptor()
 {
-#if WITH_EDITORONLY_DATA
 	// Make sure we have proper defaults
 	InitFrom(UHierarchicalInstancedStaticMeshComponent::StaticClass()->GetDefaultObject<UHierarchicalInstancedStaticMeshComponent>());
-#endif
 }
 
-#if WITH_EDITOR
-
 FISMComponentDescriptor FISMComponentDescriptor::CreateFrom(const TSubclassOf<UStaticMeshComponent>& From)
 {
 	FISMComponentDescriptor ComponentDescriptor;
@@ -51,7 +43,6 @@ void FISMComponentDescriptor::InitFrom(const UStaticMeshComponent* Template, boo
 	TranslucencySortPriority = Template->TranslucencySortPriority;
 	OverriddenLightMapRes = Template->OverriddenLightMapRes;
 	CustomDepthStencilValue = Template->CustomDepthStencilValue;
-	HLODBatchingPolicy = Template->HLODBatchingPolicy;
 	bCastShadow = Template->CastShadow;
 	bCastStaticShadow = Template->bCastStaticShadow;
 	bCastDynamicShadow = Template->bCastDynamicShadow;
@@ -63,15 +54,19 @@ void FISMComponentDescriptor::InitFrom(const UStaticMeshComponent* Template, boo
 	bOverrideLightMapRes = Template->bOverrideLightMapRes;
 	bUseAsOccluder = Template->bUseAsOccluder;
 	bRenderCustomDepth = Template->bRenderCustomDepth;
-	bIncludeInHLOD = Template->bEnableAutoLODGeneration;
 	bHiddenInGame = Template->bHiddenInGame;
 	bIsEditorOnly = Template->bIsEditorOnly;
 	bVisible = Template->GetVisibleFlag();
 	bVisibleInRayTracing = Template->bVisibleInRayTracing;
-	bConsiderForActorPlacementWhenHidden = Template->bConsiderForActorPlacementWhenHidden;
 	bEvaluateWorldPositionOffset = Template->bEvaluateWorldPositionOffset;
 	bIsLocalToWorldDeterminantNegative = Template->GetRenderMatrix().Determinant() < 0;
 
+#if WITH_EDITORONLY_DATA
+	HLODBatchingPolicy = Template->HLODBatchingPolicy;
+	bIncludeInHLOD = Template->bEnableAutoLODGeneration;
+	bConsiderForActorPlacementWhenHidden = Template->bConsiderForActorPlacementWhenHidden;
+#endif
+
 	if (const UInstancedStaticMeshComponent* ISMTemplate = Cast<UInstancedStaticMeshComponent>(Template))
 	{
 		InstanceStartCullDistance = ISMTemplate->InstanceStartCullDistance;
@@ -116,7 +111,6 @@ bool FISMComponentDescriptor::operator==(const FISMComponentDescriptor& Other) c
 	TranslucencySortPriority == Other.TranslucencySortPriority &&
 	OverriddenLightMapRes == Other.OverriddenLightMapRes &&
 	CustomDepthStencilValue == Other.CustomDepthStencilValue &&
-	HLODBatchingPolicy == Other.HLODBatchingPolicy &&
 	bCastShadow == Other.bCastShadow &&
 	bCastStaticShadow == Other.bCastStaticShadow &&
 	bCastDynamicShadow == Other.bCastDynamicShadow &&
@@ -128,15 +122,18 @@ bool FISMComponentDescriptor::operator==(const FISMComponentDescriptor& Other) c
 	bOverrideLightMapRes == Other.bOverrideLightMapRes &&
 	bUseAsOccluder == Other.bUseAsOccluder &&
 	bRenderCustomDepth == Other.bRenderCustomDepth &&
-	bIncludeInHLOD == Other.bIncludeInHLOD &&
 	bEnableDiscardOnLoad == Other.bEnableDiscardOnLoad &&
 	bHiddenInGame == Other.bHiddenInGame &&
 	bIsEditorOnly == Other.bIsEditorOnly &&
 	bVisible == Other.bVisible &&
 	bVisibleInRayTracing == Other.bVisibleInRayTracing &&
-	bConsiderForActorPlacementWhenHidden == Other.bConsiderForActorPlacementWhenHidden &&
 	bEvaluateWorldPositionOffset == Other.bEvaluateWorldPositionOffset &&
 	bIsLocalToWorldDeterminantNegative == Other.bIsLocalToWorldDeterminantNegative &&
+#if WITH_EDITORONLY_DATA
+	HLODBatchingPolicy == Other.HLODBatchingPolicy &&
+	bIncludeInHLOD == Other.bIncludeInHLOD &&
+	bConsiderForActorPlacementWhenHidden == Other.bConsiderForActorPlacementWhenHidden &&
+#endif // WITH_EDITORONLY_DATA
 	BodyInstance.GetCollisionEnabled() == Other.BodyInstance.GetCollisionEnabled() && 
 	BodyInstance.GetCollisionResponse() == Other.BodyInstance.GetCollisionResponse() &&
 	BodyInstance.DoesUseCollisionProfile() == Other.BodyInstance.DoesUseCollisionProfile() &&
@@ -196,7 +193,6 @@ void FISMComponentDescriptor::InitComponent(UInstancedStaticMeshComponent* ISMCo
 	ISMComponent->TranslucencySortPriority = TranslucencySortPriority;
 	ISMComponent->OverriddenLightMapRes = OverriddenLightMapRes;
 	ISMComponent->CustomDepthStencilValue = CustomDepthStencilValue;
-	ISMComponent->HLODBatchingPolicy = HLODBatchingPolicy;
 	ISMComponent->CastShadow = bCastShadow;
 	ISMComponent->bCastStaticShadow = bCastStaticShadow;
 	ISMComponent->bCastDynamicShadow = bCastDynamicShadow;
@@ -208,14 +204,18 @@ void FISMComponentDescriptor::InitComponent(UInstancedStaticMeshComponent* ISMCo
 	ISMComponent->bOverrideLightMapRes = bOverrideLightMapRes;
 	ISMComponent->bUseAsOccluder = bUseAsOccluder;
 	ISMComponent->bRenderCustomDepth = bRenderCustomDepth;
-	ISMComponent->bEnableAutoLODGeneration = bIncludeInHLOD;;
 	ISMComponent->bHiddenInGame = bHiddenInGame;
 	ISMComponent->bIsEditorOnly = bIsEditorOnly;
 	ISMComponent->SetVisibleFlag(bVisible);
 	ISMComponent->bVisibleInRayTracing = bVisibleInRayTracing;
-	ISMComponent->bConsiderForActorPlacementWhenHidden = bConsiderForActorPlacementWhenHidden;
 	ISMComponent->bEvaluateWorldPositionOffset = bEvaluateWorldPositionOffset;
 	ISMComponent->bReverseCulling = bIsLocalToWorldDeterminantNegative;
+	
+#if WITH_EDITORONLY_DATA
+	ISMComponent->HLODBatchingPolicy = HLODBatchingPolicy;
+	ISMComponent->bEnableAutoLODGeneration = bIncludeInHLOD;
+	ISMComponent->bConsiderForActorPlacementWhenHidden = bConsiderForActorPlacementWhenHidden;
+#endif // WITH_EDITORONLY_DATA
 
 	// HISM Specific
 	if (UHierarchicalInstancedStaticMeshComponent* HISMComponent = Cast<UHierarchicalInstancedStaticMeshComponent>(ISMComponent))
@@ -223,6 +223,3 @@ void FISMComponentDescriptor::InitComponent(UInstancedStaticMeshComponent* ISMCo
 		HISMComponent->bEnableDensityScaling = bEnableDensityScaling;
 	}
 }
-
-#endif
-
diff --git a/Engine/Source/Runtime/Engine/Public/ISMPartition/ISMComponentDescriptor.h b/Engine/Source/Runtime/Engine/Public/ISMPartition/ISMComponentDescriptor.h
--- a/Engine/Source/Runtime/Engine/Public/ISMPartition/ISMComponentDescriptor.h
+++ b/Engine/Source/Runtime/Engine/Public/ISMPartition/ISMComponentDescriptor.h
@@ -20,7 +20,6 @@ struct ENGINE_API FISMComponentDescriptor
 	GENERATED_USTRUCT_BODY()
 
 	FISMComponentDescriptor();
-#if WITH_EDITOR
 	static FISMComponentDescriptor CreateFrom(const TSubclassOf<UStaticMeshComponent>& ComponentClass);
 	void InitFrom(const UStaticMeshComponent* Component, bool bInitBodyInstance = true);
 
@@ -44,11 +43,8 @@ struct ENGINE_API FISMComponentDescriptor
 	{
 		return Lhs.Hash < Rhs.Hash;
 	}
-#endif
 
 public:
-
-#if WITH_EDITORONLY_DATA
 	UPROPERTY()
 	mutable uint32 Hash = 0;
 
@@ -109,8 +105,10 @@ public:
 	UPROPERTY(EditAnywhere, Category = "Component Settings")
 	int32 CustomDepthStencilValue;
 
+#if WITH_EDITORONLY_DATA
 	UPROPERTY(EditAnywhere, Category = "Component Settings")
 	EHLODBatchingPolicy HLODBatchingPolicy;
+#endif
 	
 	UPROPERTY(EditAnywhere, Category = "Component Settings")
 	uint8 bCastShadow : 1;
@@ -151,9 +149,6 @@ public:
 	UPROPERTY(EditAnywhere, Category = "Component Settings")
 	uint8 bRenderCustomDepth : 1;
 
-	UPROPERTY(EditAnywhere, Category = "Component Settings")
-	uint8 bIncludeInHLOD : 1;
-
 	UPROPERTY(EditAnywhere, Category = "Component Settings")
 	uint8 bVisibleInRayTracing : 1;
 
@@ -165,14 +160,18 @@ public:
 
 	UPROPERTY(EditAnywhere, Category = "Component Settings")
 	uint8 bVisible : 1;
+	
+	UPROPERTY(EditAnywhere, Category = "Component Settings")
+	uint8 bEvaluateWorldPositionOffset : 1;
 
 	UPROPERTY(EditAnywhere, Category = "Component Settings")
-	uint8 bConsiderForActorPlacementWhenHidden : 1;
+	uint8 bIsLocalToWorldDeterminantNegative : 1;
 
+#if WITH_EDITORONLY_DATA
 	UPROPERTY(EditAnywhere, Category = "Component Settings")
-	uint8 bEvaluateWorldPositionOffset : 1;
+	uint8 bIncludeInHLOD : 1;
 
 	UPROPERTY(EditAnywhere, Category = "Component Settings")
-	uint8 bIsLocalToWorldDeterminantNegative : 1;
+	uint8 bConsiderForActorPlacementWhenHidden : 1;
 #endif
 };
