MLDeformer: Allow mesh deformer on skinned mesh component to be unset.

diff --git a/Engine/Source/Runtime/Engine/Classes/Components/SkinnedMeshComponent.h b/Engine/Source/Runtime/Engine/Classes/Components/SkinnedMeshComponent.h
--- a/Engine/Source/Runtime/Engine/Classes/Components/SkinnedMeshComponent.h
+++ b/Engine/Source/Runtime/Engine/Classes/Components/SkinnedMeshComponent.h
@@ -247,6 +247,9 @@ protected:
 	UPROPERTY(EditAnywhere, BlueprintReadOnly, Category = "Deformer", meta = (editcondition = "bSetMeshDeformer"))
 	TObjectPtr<UMeshDeformer> MeshDeformer;
 
+	/** Set the MeshDeformer and update the internal MeshDeformerInstance. */
+	void SetMeshDeformer(bool bInSetMeshDeformer, UMeshDeformer* InMeshDeformer);
+
 	/** Get the currently active MeshDeformer. This may come from the SkeletalMesh default or the Component override. */
 	UMeshDeformer* GetActiveMeshDeformer() const;
 
@@ -934,6 +937,14 @@ public:
 	UFUNCTION(BlueprintCallable, Category = "Components|SkinnedMesh")
 	void SetMeshDeformer(UMeshDeformer* InMeshDeformer);
 
+	/**
+	 * Unset any MeshDeformer applied to this Component.
+	 *
+	 * @param InMeshDeformer New mesh deformer to set for this component
+	 */
+	UFUNCTION(BlueprintCallable, Category = "Components|SkinnedMesh")
+	void UnsetMeshDeformer();
+
 	/** 
 	 * Get Parent Bone of the input bone
 	 * 
diff --git a/Engine/Source/Runtime/Engine/Private/Components/SkinnedMeshComponent.cpp b/Engine/Source/Runtime/Engine/Private/Components/SkinnedMeshComponent.cpp
--- a/Engine/Source/Runtime/Engine/Private/Components/SkinnedMeshComponent.cpp
+++ b/Engine/Source/Runtime/Engine/Private/Components/SkinnedMeshComponent.cpp
@@ -2119,11 +2119,11 @@ UMeshDeformer* USkinnedMeshComponent::GetActiveMeshDeformer() const
 	return nullptr;
 }
 
-void USkinnedMeshComponent::SetMeshDeformer(UMeshDeformer* InMeshDeformer)
+void USkinnedMeshComponent::SetMeshDeformer(bool bInSetMeshDeformer, UMeshDeformer* InMeshDeformer)
 {
-	bSetMeshDeformer = true;
+	bSetMeshDeformer = bInSetMeshDeformer;
 	MeshDeformer = InMeshDeformer;
-	
+
 	UMeshDeformer* ActiveMeshDeformer = GetActiveMeshDeformer();
 	MeshDeformerInstanceSettings = ActiveMeshDeformer ? ActiveMeshDeformer->CreateSettingsInstance(this) : nullptr;
 	MeshDeformerInstance = ActiveMeshDeformer ? ActiveMeshDeformer->CreateInstance(this, MeshDeformerInstanceSettings) : nullptr;
@@ -2131,6 +2131,16 @@ void USkinnedMeshComponent::SetMeshDeformer(UMeshDeformer* InMeshDeformer)
 	MarkRenderDynamicDataDirty();
 }
 
+void USkinnedMeshComponent::SetMeshDeformer(UMeshDeformer* InMeshDeformer)
+{
+	SetMeshDeformer(true, InMeshDeformer);
+}
+
+void USkinnedMeshComponent::UnsetMeshDeformer()
+{
+	SetMeshDeformer(false, nullptr);
+}
+
 FSkeletalMeshRenderData* USkinnedMeshComponent::GetSkeletalMeshRenderData() const
 {
 	if (MeshObject)
