Provide a default console key (section) for slovak keyboard layout

diff --git a/Engine/Source/Runtime/Engine/Private/UserInterface/InputSettings.cpp b/Engine/Source/Runtime/Engine/Private/UserInterface/InputSettings.cpp
--- a/Engine/Source/Runtime/Engine/Private/UserInterface/InputSettings.cpp
+++ b/Engine/Source/Runtime/Engine/Private/UserInterface/InputSettings.cpp
@@ -104,6 +104,7 @@ void UInputSettings::PostInitProperties()
 			DefaultConsoleKey = FInputKeyManager::Get().GetKeyFromCodes(VK_OEM_5, 0);
 			break;
 
+		case LANG_SLOVAK:
 		case LANG_SWEDISH:
 			DefaultConsoleKey = EKeys::Section;
 			break;
@@ -479,7 +480,7 @@ void UInputSettings::SetDefaultPlayerInputClass(TSubclassOf<UPlayerInput> NewDef
 	if(ensure(NewDefaultPlayerInputClass))
 	{
 		UInputSettings* InputSettings = Cast<UInputSettings>(UInputSettings::StaticClass()->GetDefaultObject());
-		InputSettings->DefaultPlayerInputClass = NewDefaultPlayerInputClass;	
+		InputSettings->DefaultPlayerInputClass = NewDefaultPlayerInputClass;
 	}
 }
 
@@ -488,7 +489,7 @@ void UInputSettings::SetDefaultInputComponentClass(TSubclassOf<UInputComponent>
 	if(ensure(NewDefaultInputComponentClass))
 	{
 		UInputSettings* InputSettings = Cast<UInputSettings>(UInputSettings::StaticClass()->GetDefaultObject());
-		InputSettings->DefaultInputComponentClass = NewDefaultInputComponentClass;	
+		InputSettings->DefaultInputComponentClass = NewDefaultInputComponentClass;
 	}
 }
 
@@ -498,7 +499,7 @@ FHardwareDeviceIdentifier::FHardwareDeviceIdentifier(FName InClassName, FString
 	: InputClassName(InClassName)
 	, HardwareDeviceIdentifier(InHardwareDeviceIdentifier)
 {
-	
+
 }
 
 bool FHardwareDeviceIdentifier::IsValid() const
@@ -543,7 +544,7 @@ const TArray<FString>& UInputPlatformSettings::GetAllHardwareDeviceNames()
 			}
 		}
 	}
-	
+
 	return HardwareDevices;
 }
-#endif	// WITH_EDITOR
\ No newline at end of file
+#endif	// WITH_EDITOR
