Bindless Resources prerequisites
- Making SetUAVParameter handle bindless UAVs.
- Changed uses of RHICommandList::SetUAVParameter to use the global SetUAVParameter.
- Changed uses of RHICommandList::SetShaderResourceViewParameter to use the global SetSRVParameter.

diff --git a/Engine/Plugins/FX/Niagara/Source/NiagaraVertexFactories/Private/NiagaraDrawIndirect.cpp b/Engine/Plugins/FX/Niagara/Source/NiagaraVertexFactories/Private/NiagaraDrawIndirect.cpp
--- a/Engine/Plugins/FX/Niagara/Source/NiagaraVertexFactories/Private/NiagaraDrawIndirect.cpp
+++ b/Engine/Plugins/FX/Niagara/Source/NiagaraVertexFactories/Private/NiagaraDrawIndirect.cpp
@@ -43,14 +43,8 @@ FNiagaraDrawIndirectArgsGenCS::FNiagaraDrawIndirectArgsGenCS(const ShaderMetaTyp
 void FNiagaraDrawIndirectArgsGenCS::SetOutput(FRHICommandList& RHICmdList, FRHIUnorderedAccessView* DrawIndirectArgsUAV, FRHIUnorderedAccessView* InstanceCountsUAV)
 {
 	FRHIComputeShader* ComputeShaderRHI = RHICmdList.GetBoundComputeShader();
-	if (DrawIndirectArgsParam.IsBound())
-	{
-		RHICmdList.SetUAVParameter(ComputeShaderRHI, DrawIndirectArgsParam.GetUAVIndex(), DrawIndirectArgsUAV);
-	}
-	if (InstanceCountsParam.IsBound())
-	{
-		RHICmdList.SetUAVParameter(ComputeShaderRHI, InstanceCountsParam.GetUAVIndex(), InstanceCountsUAV);
-	}
+	DrawIndirectArgsParam.SetUAV(RHICmdList, ComputeShaderRHI, DrawIndirectArgsUAV);
+	InstanceCountsParam.SetUAV(RHICmdList, ComputeShaderRHI, DrawIndirectArgsUAV);
 }
 
 void FNiagaraDrawIndirectArgsGenCS::SetParameters(FRHICommandList& RHICmdList, FRHIShaderResourceView* TaskInfosBuffer, FRHIShaderResourceView* CulledInstanceCountsBuffer, int32 ArgGenTaskOffset, int32 NumArgGenTasks, int32 NumInstanceCountClearTasks)
@@ -69,14 +63,8 @@ void FNiagaraDrawIndirectArgsGenCS::UnbindBuffers(FRHICommandList& RHICmdList)
 	FRHIComputeShader* ComputeShaderRHI = RHICmdList.GetBoundComputeShader();
 
 	SetSRVParameter(RHICmdList, ComputeShaderRHI, TaskInfosParam, nullptr);
-	if (DrawIndirectArgsParam.IsBound())
-	{
-		RHICmdList.SetUAVParameter(ComputeShaderRHI, DrawIndirectArgsParam.GetUAVIndex(), nullptr);
-	}
-	if (InstanceCountsParam.IsBound())
-	{
-		RHICmdList.SetUAVParameter(ComputeShaderRHI, InstanceCountsParam.GetUAVIndex(), nullptr);
-	}
+	DrawIndirectArgsParam.UnsetUAV(RHICmdList, ComputeShaderRHI);
+	InstanceCountsParam.UnsetUAV(RHICmdList, ComputeShaderRHI);
 }
 void FNiagaraDrawIndirectResetCountsCS::ModifyCompilationEnvironment(const FGlobalShaderPermutationParameters& Parameters, FShaderCompilerEnvironment& OutEnvironment)
 {
@@ -97,11 +85,7 @@ FNiagaraDrawIndirectResetCountsCS::FNiagaraDrawIndirectResetCountsCS(const Shade
 void FNiagaraDrawIndirectResetCountsCS::SetOutput(FRHICommandList& RHICmdList, FRHIUnorderedAccessView* InstanceCountsUAV)
 {
 	FRHIComputeShader* ComputeShaderRHI = RHICmdList.GetBoundComputeShader();
-
-	if (InstanceCountsParam.IsBound())
-	{
-		RHICmdList.SetUAVParameter(ComputeShaderRHI, InstanceCountsParam.GetUAVIndex(), InstanceCountsUAV);
-	}
+	InstanceCountsParam.SetUAV(RHICmdList, ComputeShaderRHI, InstanceCountsUAV);
 }
 
 void FNiagaraDrawIndirectResetCountsCS::SetParameters(FRHICommandList& RHICmdList, FRHIShaderResourceView* TaskInfosBuffer, int32 NumArgGenTasks, int32 NumInstanceCountClearTasks)
@@ -119,9 +103,5 @@ void FNiagaraDrawIndirectResetCountsCS::UnbindBuffers(FRHICommandList& RHICmdLis
 	FRHIComputeShader* ComputeShaderRHI = RHICmdList.GetBoundComputeShader();
 
 	SetSRVParameter(RHICmdList, ComputeShaderRHI, TaskInfosParam, nullptr);
-
-	if (InstanceCountsParam.IsBound())
-	{
-		RHICmdList.SetUAVParameter(ComputeShaderRHI, InstanceCountsParam.GetUAVIndex(), nullptr);
-	}
+	InstanceCountsParam.UnsetUAV(RHICmdList, ComputeShaderRHI);
 }
\ No newline at end of file
diff --git a/Engine/Source/Runtime/Engine/Private/GPUSort.cpp b/Engine/Source/Runtime/Engine/Private/GPUSort.cpp
--- a/Engine/Source/Runtime/Engine/Private/GPUSort.cpp
+++ b/Engine/Source/Runtime/Engine/Private/GPUSort.cpp
@@ -255,10 +255,7 @@ public:
 	void SetOutput(FRHICommandList& RHICmdList, FRHIUnorderedAccessView* OutOffsetsUAV)
 	{
 		FRHIComputeShader* ComputeShaderRHI = RHICmdList.GetBoundComputeShader();
-		if ( OutOffsets.IsBound() )
-		{
-			RHICmdList.SetUAVParameter(ComputeShaderRHI, OutOffsets.GetBaseIndex(), OutOffsetsUAV);
-		}
+		SetUAVParameter(RHICmdList, ComputeShaderRHI, OutOffsets, OutOffsetsUAV);
 	}
 
 	/**
@@ -267,10 +264,7 @@ public:
 	void UnbindBuffers(FRHICommandList& RHICmdList)
 	{
 		FRHIComputeShader* ComputeShaderRHI = RHICmdList.GetBoundComputeShader();
-		if ( OutOffsets.IsBound() )
-		{
-			RHICmdList.SetUAVParameter(ComputeShaderRHI, OutOffsets.GetBaseIndex(), nullptr);
-		}
+		SetUAVParameter(RHICmdList, ComputeShaderRHI, OutOffsets, nullptr);
 	}
 
 private:
@@ -329,14 +323,8 @@ public:
 	{
 		FRHIComputeShader* ComputeShaderRHI = RHICmdList.GetBoundComputeShader();
 		SetUniformBufferParameter(RHICmdList, ComputeShaderRHI, GetUniformBufferParameter<FRadixSortParameters>(), RadixSortUniformBuffer );
-		if ( InKeys.IsBound() )
-		{
-			RHICmdList.SetShaderResourceViewParameter(ComputeShaderRHI, InKeys.GetBaseIndex(), InKeysSRV);
-		}
-		if ( RadixSortParameterBuffer.IsBound() )
-		{
-			RHICmdList.SetShaderResourceViewParameter(ComputeShaderRHI, RadixSortParameterBuffer.GetBaseIndex(), RadixSortParameterBufferSRV);
-		}
+		SetSRVParameter(RHICmdList, ComputeShaderRHI, InKeys, InKeysSRV);
+		SetSRVParameter(RHICmdList, ComputeShaderRHI, RadixSortParameterBuffer, RadixSortParameterBufferSRV);
 	}
 
 	/**
@@ -345,10 +333,7 @@ public:
 	void SetOutput(FRHICommandList& RHICmdList, FRHIUnorderedAccessView* OutOffsetsUAV)
 	{
 		FRHIComputeShader* ComputeShaderRHI = RHICmdList.GetBoundComputeShader();
-		if ( OutOffsets.IsBound() )
-		{
-			RHICmdList.SetUAVParameter(ComputeShaderRHI, OutOffsets.GetBaseIndex(), OutOffsetsUAV);
-		}
+		SetUAVParameter(RHICmdList, ComputeShaderRHI, OutOffsets, OutOffsetsUAV);
 	}
 
 	/**
@@ -357,18 +342,9 @@ public:
 	void UnbindBuffers(FRHICommandList& RHICmdList)
 	{
 		FRHIComputeShader* ComputeShaderRHI = RHICmdList.GetBoundComputeShader();
-		if ( RadixSortParameterBuffer.IsBound() )
-		{
-			RHICmdList.SetShaderResourceViewParameter(ComputeShaderRHI, RadixSortParameterBuffer.GetBaseIndex(), nullptr);
-		}
-		if ( InKeys.IsBound() )
-		{
-			RHICmdList.SetShaderResourceViewParameter(ComputeShaderRHI, InKeys.GetBaseIndex(), nullptr);
-		}
-		if ( OutOffsets.IsBound() )
-		{
-			RHICmdList.SetUAVParameter(ComputeShaderRHI, OutOffsets.GetBaseIndex(), nullptr);
-		}
+		SetSRVParameter(RHICmdList, ComputeShaderRHI, RadixSortParameterBuffer, nullptr);
+		SetSRVParameter(RHICmdList, ComputeShaderRHI, InKeys, nullptr);
+		SetUAVParameter(RHICmdList, ComputeShaderRHI, OutOffsets, nullptr);
 	}
 
 private:
@@ -419,10 +395,7 @@ public:
 	void SetParameters(FRHICommandList& RHICmdList, FRHIShaderResourceView* InOffsetsSRV)
 	{
 		FRHIComputeShader* ComputeShaderRHI = RHICmdList.GetBoundComputeShader();
-		if ( InOffsets.IsBound() )
-		{
-			RHICmdList.SetShaderResourceViewParameter(ComputeShaderRHI, InOffsets.GetBaseIndex(), InOffsetsSRV);
-		}
+		SetSRVParameter(RHICmdList, ComputeShaderRHI, InOffsets, InOffsetsSRV);
 	}
 
 	/**
@@ -431,10 +404,7 @@ public:
 	void SetOutput(FRHICommandList& RHICmdList, FRHIUnorderedAccessView* OutOffsetsUAV)
 	{
 		FRHIComputeShader* ComputeShaderRHI = RHICmdList.GetBoundComputeShader();
-		if ( OutOffsets.IsBound() )
-		{
-			RHICmdList.SetUAVParameter(ComputeShaderRHI, OutOffsets.GetBaseIndex(), OutOffsetsUAV);
-		}
+		SetUAVParameter(RHICmdList, ComputeShaderRHI, OutOffsets, OutOffsetsUAV);
 	}
 
 	/**
@@ -443,14 +413,8 @@ public:
 	void UnbindBuffers(FRHICommandList& RHICmdList)
 	{
 		FRHIComputeShader* ComputeShaderRHI = RHICmdList.GetBoundComputeShader();
-		if ( InOffsets.IsBound() )
-		{
-			RHICmdList.SetShaderResourceViewParameter(ComputeShaderRHI, InOffsets.GetBaseIndex(), nullptr);
-		}
-		if ( OutOffsets.IsBound() )
-		{
-			RHICmdList.SetUAVParameter(ComputeShaderRHI, OutOffsets.GetBaseIndex(), nullptr);
-		}
+		SetSRVParameter(RHICmdList, ComputeShaderRHI, InOffsets, nullptr);
+		SetUAVParameter(RHICmdList, ComputeShaderRHI, OutOffsets, nullptr);
 	}
 
 private:
@@ -522,22 +486,10 @@ public:
 	{
 		FRHIComputeShader* ComputeShaderRHI = RHICmdList.GetBoundComputeShader();
 		SetUniformBufferParameter(RHICmdList, ComputeShaderRHI, GetUniformBufferParameter<FRadixSortParameters>(), RadixSortUniformBuffer );
-		if ( RadixSortParameterBuffer.IsBound() )
-		{
-			RHICmdList.SetShaderResourceViewParameter(ComputeShaderRHI, RadixSortParameterBuffer.GetBaseIndex(), RadixSortParameterBufferSRV);
-		}
-		if ( InKeys.IsBound() )
-		{
-			RHICmdList.SetShaderResourceViewParameter(ComputeShaderRHI, InKeys.GetBaseIndex(), InKeysSRV);
-		}
-		if ( InValues.IsBound() )
-		{
-			RHICmdList.SetShaderResourceViewParameter(ComputeShaderRHI, InValues.GetBaseIndex(), InValuesSRV);
-		}
-		if ( InOffsets.IsBound() )
-		{
-			RHICmdList.SetShaderResourceViewParameter(ComputeShaderRHI, InOffsets.GetBaseIndex(), InOffsetsSRV);
-		}
+		SetSRVParameter(RHICmdList, ComputeShaderRHI, RadixSortParameterBuffer, RadixSortParameterBufferSRV);
+		SetSRVParameter(RHICmdList, ComputeShaderRHI, InKeys, InKeysSRV);
+		SetSRVParameter(RHICmdList, ComputeShaderRHI, InValues, InValuesSRV);
+		SetSRVParameter(RHICmdList, ComputeShaderRHI, InOffsets, InOffsetsSRV);
 	}
 
 	/**
@@ -546,14 +498,8 @@ public:
 	void SetOutput(FRHICommandList& RHICmdList, FRHIUnorderedAccessView* OutKeysUAV, FRHIUnorderedAccessView* OutValuesUAV)
 	{
 		FRHIComputeShader* ComputeShaderRHI = RHICmdList.GetBoundComputeShader();
-		if ( OutKeys.IsBound() )
-		{
-			RHICmdList.SetUAVParameter(ComputeShaderRHI, OutKeys.GetBaseIndex(), OutKeysUAV);
-		}
-		if ( OutValues.IsBound() )
-		{
-			RHICmdList.SetUAVParameter(ComputeShaderRHI, OutValues.GetBaseIndex(), OutValuesUAV);
-		}
+		SetUAVParameter(RHICmdList, ComputeShaderRHI, OutKeys, OutKeysUAV);
+		SetUAVParameter(RHICmdList, ComputeShaderRHI, OutValues, OutValuesUAV);
 	}
 
 	/**
@@ -562,30 +508,12 @@ public:
 	void UnbindBuffers(FRHICommandList& RHICmdList)
 	{
 		FRHIComputeShader* ComputeShaderRHI = RHICmdList.GetBoundComputeShader();
-		if ( RadixSortParameterBuffer.IsBound() )
-		{
-			RHICmdList.SetShaderResourceViewParameter(ComputeShaderRHI, RadixSortParameterBuffer.GetBaseIndex(), nullptr);
-		}
-		if ( InKeys.IsBound() )
-		{
-			RHICmdList.SetShaderResourceViewParameter(ComputeShaderRHI, InKeys.GetBaseIndex(), nullptr);
-		}
-		if ( InValues.IsBound() )
-		{
-			RHICmdList.SetShaderResourceViewParameter(ComputeShaderRHI, InValues.GetBaseIndex(), nullptr);
-		}
-		if ( InOffsets.IsBound() )
-		{
-			RHICmdList.SetShaderResourceViewParameter(ComputeShaderRHI, InOffsets.GetBaseIndex(), nullptr);
-		}
-		if ( OutKeys.IsBound() )
-		{
-			RHICmdList.SetUAVParameter(ComputeShaderRHI, OutKeys.GetBaseIndex(), nullptr);
-		}
-		if ( OutValues.IsBound() )
-		{
-			RHICmdList.SetUAVParameter(ComputeShaderRHI, OutValues.GetBaseIndex(), nullptr);
-		}
+		SetSRVParameter(RHICmdList, ComputeShaderRHI, RadixSortParameterBuffer, nullptr);
+		SetSRVParameter(RHICmdList, ComputeShaderRHI, InKeys, nullptr);
+		SetSRVParameter(RHICmdList, ComputeShaderRHI, InValues, nullptr);
+		SetSRVParameter(RHICmdList, ComputeShaderRHI, InOffsets, nullptr);
+		SetUAVParameter(RHICmdList, ComputeShaderRHI, OutKeys, nullptr);
+		SetUAVParameter(RHICmdList, ComputeShaderRHI, OutValues, nullptr);
 	}
 
 private:
diff --git a/Engine/Source/Runtime/Engine/Private/GPUSortManager.cpp b/Engine/Source/Runtime/Engine/Private/GPUSortManager.cpp
--- a/Engine/Source/Runtime/Engine/Private/GPUSortManager.cpp
+++ b/Engine/Source/Runtime/Engine/Private/GPUSortManager.cpp
@@ -154,7 +154,7 @@ void FCopyUIntBufferCS::SetParameters(
 	FRHIComputeShader* ComputeShaderRHI = RHICmdList.GetBoundComputeShader();
 	check(DestCount > 0 && DestCount <= COPYUINTCS_BUFFER_COUNT);
 
-	RHICmdList.SetShaderResourceViewParameter(ComputeShaderRHI, SourceData.GetBaseIndex(), InSourceData);
+	SetSRVParameter(RHICmdList, ComputeShaderRHI, SourceData, InSourceData);
 
 	FUintVector4 CopyParamsValue(StartingIndex, 0, 0, 0);
 	for (int32 Index = 0; Index < DestCount; ++Index)
diff --git a/Engine/Source/Runtime/Engine/Private/Particles/ParticleGpuSimulation.cpp b/Engine/Source/Runtime/Engine/Private/Particles/ParticleGpuSimulation.cpp
--- a/Engine/Source/Runtime/Engine/Private/Particles/ParticleGpuSimulation.cpp
+++ b/Engine/Source/Runtime/Engine/Private/Particles/ParticleGpuSimulation.cpp
@@ -948,10 +948,7 @@ public:
 	void SetParameters(FRHICommandList& RHICmdList, FParticleShaderParamRef TileOffsetsRef)
 	{
 		FRHIVertexShader* VertexShaderRHI = RHICmdList.GetBoundVertexShader();
-		if (TileOffsets.IsBound())
-		{
-			RHICmdList.SetShaderResourceViewParameter(VertexShaderRHI, TileOffsets.GetBaseIndex(), TileOffsetsRef);
-		}
+		SetSRVParameter(RHICmdList, VertexShaderRHI, TileOffsets, TileOffsetsRef);
 		SetShaderValue(RHICmdList, VertexShaderRHI, TileSizeX, (float)GParticleSimulationTileSize / (float)GParticleSimulationTextureSizeX);
 		SetShaderValue(RHICmdList, VertexShaderRHI, TileSizeY, (float)GParticleSimulationTileSize / (float)GParticleSimulationTextureSizeY);
 	}
@@ -1078,9 +1075,9 @@ public:
 	void SetVectorFieldParameters(FRHICommandList& RHICmdList, const FVectorFieldUniformBufferRef& UniformBuffer, FRHITexture3D* const* VolumeTexturesRHI)
 	{
 		FRHIPixelShader* PixelShaderRHI = RHICmdList.GetBoundPixelShader();
-			SetUniformBufferParameter(RHICmdList, PixelShaderRHI, GetUniformBufferParameter<FVectorFieldUniformParameters>(), UniformBuffer);
+		SetUniformBufferParameter(RHICmdList, PixelShaderRHI, GetUniformBufferParameter<FVectorFieldUniformParameters>(), UniformBuffer);
 		
-			FRHISamplerState* SamplerStateLinear = TStaticSamplerState<SF_Bilinear,AM_Clamp,AM_Clamp,AM_Clamp>::GetRHI();
+		FRHISamplerState* SamplerStateLinear = TStaticSamplerState<SF_Bilinear,AM_Clamp,AM_Clamp,AM_Clamp>::GetRHI();
 
 		for (int32 i = 0; i < MAX_VECTOR_FIELDS; ++i)
 		{
@@ -1108,10 +1105,7 @@ public:
 		FRHIShaderResourceView* NullSRV = nullptr;
 		for (int32 i = 0; i < MAX_VECTOR_FIELDS; ++i)
 		{
-			if (VectorFieldTextures[i].IsBound())
-			{
-				RHICmdList.SetShaderResourceViewParameter(PixelShaderRHI, VectorFieldTextures[i].GetBaseIndex(), NullSRV);
-			}
+			SetSRVParameter(RHICmdList, PixelShaderRHI, VectorFieldTextures[i], NullSRV);
 		}
 	}
 
@@ -2107,10 +2101,7 @@ public:
 	void SetOutput(FRHICommandList& RHICmdList, FRHIUnorderedAccessView* OutBoundsUAV )
 	{
 		FRHIComputeShader* ComputeShaderRHI = RHICmdList.GetBoundComputeShader();
-		if ( OutBounds.IsBound() )
-		{
-			RHICmdList.SetUAVParameter(ComputeShaderRHI, OutBounds.GetBaseIndex(), OutBoundsUAV);
-		}
+		SetUAVParameter(RHICmdList, ComputeShaderRHI, OutBounds, OutBoundsUAV);
 	}
 
 	/**
@@ -2125,14 +2116,9 @@ public:
 	{
 		FRHIComputeShader* ComputeShaderRHI = RHICmdList.GetBoundComputeShader();
 		SetUniformBufferParameter(RHICmdList, ComputeShaderRHI, GetUniformBufferParameter<FParticleBoundsParameters>(), UniformBuffer );
-		if ( InParticleIndices.IsBound() )
-		{
-			RHICmdList.SetShaderResourceViewParameter(ComputeShaderRHI, InParticleIndices.GetBaseIndex(), InIndicesSRV);
-		}
-		if ( PositionTexture.IsBound() )
-		{
-			RHICmdList.SetShaderTexture(ComputeShaderRHI, PositionTexture.GetBaseIndex(), PositionTextureRHI);
-		}
+
+		SetSRVParameter(RHICmdList, ComputeShaderRHI, InParticleIndices, InIndicesSRV);
+		SetTextureParameter(RHICmdList, ComputeShaderRHI, PositionTexture, PositionTextureRHI);
 
 		SetShaderValue(RHICmdList, ComputeShaderRHI, TextureSizeX, GParticleSimulationTextureSizeX);
 		SetShaderValue(RHICmdList, ComputeShaderRHI, TextureSizeY, GParticleSimulationTextureSizeY);
@@ -2144,14 +2130,8 @@ public:
 	void UnbindBuffers(FRHICommandList& RHICmdList)
 	{
 		FRHIComputeShader* ComputeShaderRHI = RHICmdList.GetBoundComputeShader();
-		if ( InParticleIndices.IsBound() )
-		{
-			RHICmdList.SetShaderResourceViewParameter(ComputeShaderRHI, InParticleIndices.GetBaseIndex(), nullptr);
-		}
-		if ( OutBounds.IsBound() )
-		{
-			RHICmdList.SetUAVParameter(ComputeShaderRHI, OutBounds.GetBaseIndex(), nullptr);
-		}
+		SetSRVParameter(RHICmdList, ComputeShaderRHI, InParticleIndices, nullptr);
+		SetUAVParameter(RHICmdList, ComputeShaderRHI, OutBounds, nullptr);
 	}
 
 private:
diff --git a/Engine/Source/Runtime/Engine/Private/Particles/ParticleSortingGPU.cpp b/Engine/Source/Runtime/Engine/Private/Particles/ParticleSortingGPU.cpp
--- a/Engine/Source/Runtime/Engine/Private/Particles/ParticleSortingGPU.cpp
+++ b/Engine/Source/Runtime/Engine/Private/Particles/ParticleSortingGPU.cpp
@@ -83,14 +83,8 @@ public:
 	void SetOutput(FRHICommandList& RHICmdList, FRHIUnorderedAccessView* OutKeysUAV, FRHIUnorderedAccessView* OutIndicesUAV )
 	{
 		FRHIComputeShader* ComputeShaderRHI = RHICmdList.GetBoundComputeShader();
-		if ( OutKeys.IsBound() )
-		{
-			RHICmdList.SetUAVParameter(ComputeShaderRHI, OutKeys.GetBaseIndex(), OutKeysUAV);
-		}
-		if ( OutParticleIndices.IsBound() )
-		{
-			RHICmdList.SetUAVParameter(ComputeShaderRHI, OutParticleIndices.GetBaseIndex(), OutIndicesUAV);
-		}
+		SetUAVParameter(RHICmdList, ComputeShaderRHI, OutKeys, OutKeysUAV);
+		SetUAVParameter(RHICmdList, ComputeShaderRHI, OutParticleIndices, OutIndicesUAV);
 	}
 
 	/**
@@ -104,10 +98,7 @@ public:
 	{
 		FRHIComputeShader* ComputeShaderRHI = RHICmdList.GetBoundComputeShader();
 		SetUniformBufferParameter(RHICmdList, ComputeShaderRHI, GetUniformBufferParameter<FParticleKeyGenParameters>(), UniformBuffer );
-		if ( InParticleIndices.IsBound() )
-		{
-			RHICmdList.SetShaderResourceViewParameter(ComputeShaderRHI, InParticleIndices.GetBaseIndex(), InIndicesSRV);
-		}
+		SetSRVParameter(RHICmdList, ComputeShaderRHI, InParticleIndices, InIndicesSRV);
 	}
 
 	/**
@@ -116,10 +107,7 @@ public:
 	void SetPositionTextures(FRHICommandList& RHICmdList, FRHITexture2D* PositionTextureRHI)
 	{
 		FRHIComputeShader* ComputeShaderRHI = RHICmdList.GetBoundComputeShader();
-		if (PositionTexture.IsBound())
-		{
-			RHICmdList.SetShaderTexture(ComputeShaderRHI, PositionTexture.GetBaseIndex(), PositionTextureRHI);
-		}
+		SetTextureParameter(RHICmdList, ComputeShaderRHI, PositionTexture, PositionTextureRHI);
 	}
 
 	/**
@@ -128,18 +116,10 @@ public:
 	void UnbindBuffers(FRHICommandList& RHICmdList)
 	{
 		FRHIComputeShader* ComputeShaderRHI = RHICmdList.GetBoundComputeShader();
-		if ( InParticleIndices.IsBound() )
-		{
-			RHICmdList.SetShaderResourceViewParameter(ComputeShaderRHI, InParticleIndices.GetBaseIndex(), nullptr);
-		}
-		if ( OutKeys.IsBound() )
-		{
-			RHICmdList.SetUAVParameter(ComputeShaderRHI, OutKeys.GetBaseIndex(), nullptr);
-		}
-		if ( OutParticleIndices.IsBound() )
-		{
-			RHICmdList.SetUAVParameter(ComputeShaderRHI, OutParticleIndices.GetBaseIndex(), nullptr);
-		}
+
+		SetSRVParameter(RHICmdList, ComputeShaderRHI, InParticleIndices, nullptr);
+		SetUAVParameter(RHICmdList, ComputeShaderRHI, OutKeys, nullptr);
+		SetUAVParameter(RHICmdList, ComputeShaderRHI, OutParticleIndices, nullptr);
 	}
 
 private:
diff --git a/Engine/Source/Runtime/Engine/Private/VectorField.cpp b/Engine/Source/Runtime/Engine/Private/VectorField.cpp
--- a/Engine/Source/Runtime/Engine/Private/VectorField.cpp
+++ b/Engine/Source/Runtime/Engine/Private/VectorField.cpp
@@ -1017,10 +1017,7 @@ public:
 	void SetOutput(FRHICommandList& RHICmdList, FRHIUnorderedAccessView* VolumeTextureUAV)
 	{
 		FRHIComputeShader* ComputeShaderRHI = RHICmdList.GetBoundComputeShader();
-		if ( OutVolumeTexture.IsBound() )
-		{
-			RHICmdList.SetUAVParameter(ComputeShaderRHI, OutVolumeTexture.GetBaseIndex(), VolumeTextureUAV);
-		}
+		SetUAVParameter(RHICmdList, ComputeShaderRHI, OutVolumeTexture, VolumeTextureUAV);
 	}
 
 	/**
@@ -1029,10 +1026,7 @@ public:
 	void UnbindBuffers(FRHICommandList& RHICmdList)
 	{
 		FRHIComputeShader* ComputeShaderRHI = RHICmdList.GetBoundComputeShader();
-		if ( OutVolumeTexture.IsBound() )
-		{
-			RHICmdList.SetUAVParameter(ComputeShaderRHI, OutVolumeTexture.GetBaseIndex(), nullptr);
-		}
+		SetUAVParameter(RHICmdList, ComputeShaderRHI, OutVolumeTexture, nullptr);
 	}
 
 private:
diff --git a/Engine/Source/Runtime/RenderCore/Public/ShaderParameterUtils.h b/Engine/Source/Runtime/RenderCore/Public/ShaderParameterUtils.h
--- a/Engine/Source/Runtime/RenderCore/Public/ShaderParameterUtils.h
+++ b/Engine/Source/Runtime/RenderCore/Public/ShaderParameterUtils.h
@@ -270,55 +270,53 @@ FORCEINLINE void SetSRVParameter(
 {
 	if (Parameter.IsBound())
 	{
-		RHICmdList.SetShaderResourceViewParameter(Shader, Parameter.GetBaseIndex(), NewShaderResourceViewRHI);
+#if PLATFORM_SUPPORTS_BINDLESS_RENDERING
+		if (Parameter.GetType() == EShaderParameterType::BindlessResourceIndex)
+		{
+			const FRHIDescriptorHandle Handle = NewShaderResourceViewRHI ? NewShaderResourceViewRHI->GetBindlessHandle() : FRHIDescriptorHandle();
+			SetBindlessParameter(RHICmdList, Shader, Parameter, Handle);
+		}
+		else
+#endif
+		{
+			RHICmdList.SetShaderResourceViewParameter(Shader, Parameter.GetBaseIndex(), NewShaderResourceViewRHI);
+		}
 	}
 }
 
+template<typename TRHIShader, typename TRHICmdList>
+FORCEINLINE void SetSRVParameter(TRHICmdList& RHICmdList, const TRefCountPtr<TRHIShader>& Shader, const FShaderResourceParameter& Parameter, FRHIShaderResourceView* NewShaderResourceViewRHI)
+{
+	SetSRVParameter(RHICmdList, Shader.GetReference(), Parameter, NewShaderResourceViewRHI);
+}
 
 template<typename TRHIShader, typename TRHICmdList>
-FORCEINLINE void SetSRVParameter(
-	TRHICmdList& RHICmdList,
-	const TRefCountPtr<TRHIShader>& Shader,
-	const FShaderResourceParameter& Parameter,
-	FRHIShaderResourceView* NewShaderResourceViewRHI
-)
+FORCEINLINE void SetUAVParameterSafeShader(TRHICmdList& RHICmdList, TRHIShader* Shader, const FShaderResourceParameter& Parameter, FRHIUnorderedAccessView* UAV)
 {
 	if (Parameter.IsBound())
 	{
-		RHICmdList.SetShaderResourceViewParameter(Shader.GetReference(), Parameter.GetBaseIndex(), NewShaderResourceViewRHI);
+#if PLATFORM_SUPPORTS_BINDLESS_RENDERING
+		if (Parameter.GetType() == EShaderParameterType::BindlessResourceIndex)
+		{
+			const FRHIDescriptorHandle Handle = UAV ? UAV->GetBindlessHandle() : FRHIDescriptorHandle();
+			SetBindlessParameter(RHICmdList, Shader, Parameter, Handle);
+		}
+		else
+#endif
+		{
+			RHICmdList.SetUAVParameter(Shader, Parameter.GetBaseIndex(), UAV);
+		}
 	}
 }
 
-/**
- * Sets the value of a unordered access view parameter
- */
-FORCEINLINE void SetUAVParameter(
-	FRHIComputeCommandList& RHICmdList,
-	FRHIComputeShader* ComputeShader,
-	const FShaderResourceParameter& Parameter,
-	FRHIUnorderedAccessView* NewUnorderedAccessViewRHI
-	)
+FORCEINLINE void SetUAVParameter(FRHIComputeCommandList& RHICmdList, FRHIComputeShader* Shader, const FShaderResourceParameter& Parameter, FRHIUnorderedAccessView* UAV)
 {
-	if (Parameter.IsBound())
-	{
-		RHICmdList.SetUAVParameter(ComputeShader, Parameter.GetBaseIndex(), NewUnorderedAccessViewRHI);
-	}
+	SetUAVParameterSafeShader(RHICmdList, Shader, Parameter, UAV);
 }
 
-/**
- * Sets the value of a unordered access view parameter
- */
-FORCEINLINE void SetUAVParameter(
-	FRHICommandList& RHICmdList,
-	FRHIPixelShader* PixelShader,
-	const FShaderResourceParameter& Parameter,
-	FRHIUnorderedAccessView* NewUnorderedAccessViewRHI
-)
+FORCEINLINE void SetUAVParameter(FRHICommandList& RHICmdList, FRHIPixelShader* Shader, const FShaderResourceParameter& Parameter, FRHIUnorderedAccessView* UAV)
 {
-	if (Parameter.IsBound())
-	{
-		RHICmdList.SetUAVParameter(PixelShader, Parameter.GetBaseIndex(), NewUnorderedAccessViewRHI);
-	}
+	SetUAVParameterSafeShader(RHICmdList, Shader, Parameter, UAV);
 }
 
 
@@ -375,6 +373,12 @@ inline void FRWShaderParameter::SetTexture(TRHICmdList& RHICmdList, const TShade
 	}
 }
 
+template<typename TRHICmdList>
+inline void FRWShaderParameter::SetUAV(TRHICmdList& RHICmdList, FRHIComputeShader* ComputeShader, FRHIUnorderedAccessView* UAV) const
+{
+	SetUAVParameter(RHICmdList, ComputeShader, UAVParameter, UAV);
+}
+
 template<typename TRHICmdList>
 inline void FRWShaderParameter::UnsetUAV(TRHICmdList& RHICmdList, FRHIComputeShader* ComputeShader) const
 {
diff --git a/Engine/Source/Runtime/RenderCore/Public/ShaderParameters.h b/Engine/Source/Runtime/RenderCore/Public/ShaderParameters.h
--- a/Engine/Source/Runtime/RenderCore/Public/ShaderParameters.h
+++ b/Engine/Source/Runtime/RenderCore/Public/ShaderParameters.h
@@ -158,6 +158,9 @@ public:
 	template<typename TShaderRHIRef, typename TRHICmdList>
 	inline void SetTexture(TRHICmdList& RHICmdList, const TShaderRHIRef& Shader, FRHITexture* Texture, FRHIUnorderedAccessView* UAV) const;
 
+	template<typename TRHICmdList>
+	inline void SetUAV(TRHICmdList& RHICmdList, FRHIComputeShader* ComputeShader, FRHIUnorderedAccessView* UAV) const;
+
 	template<typename TRHICmdList>
 	inline void UnsetUAV(TRHICmdList& RHICmdList, FRHIComputeShader* ComputeShader) const;
 
