Squelch warnings due to slow shader compiles in cooking, since this will block an autosubmit on Horde.

diff --git a/Engine/Source/Runtime/Engine/Private/ShaderCompiler/ShaderCompiler.cpp b/Engine/Source/Runtime/Engine/Private/ShaderCompiler/ShaderCompiler.cpp
--- a/Engine/Source/Runtime/Engine/Private/ShaderCompiler/ShaderCompiler.cpp
+++ b/Engine/Source/Runtime/Engine/Private/ShaderCompiler/ShaderCompiler.cpp
@@ -5421,7 +5421,7 @@ void FShaderCompilingManager::ProcessAsyncResults(bool bLimitExecutionTime, bool
 	}
 
 
-	if (bBlockOnGlobalShaderCompletion && !bLimitExecutionTime)
+	if (bBlockOnGlobalShaderCompletion && !bLimitExecutionTime && !IsRunningCookCommandlet())
 	{
 		check(PendingFinalizeShaderMaps.Num() == 0);
 
