Log a Warning if we ever re-generate the GUID of a material function.

- This will cause any material using that material function to have a dirty DDC key that changes each time we load the material.
- This in turn will always cause the material to compile shaders until the material function is resaved.

diff --git a/Engine/Source/Runtime/Engine/Private/Materials/MaterialExpressions.cpp b/Engine/Source/Runtime/Engine/Private/Materials/MaterialExpressions.cpp
--- a/Engine/Source/Runtime/Engine/Private/Materials/MaterialExpressions.cpp
+++ b/Engine/Source/Runtime/Engine/Private/Materials/MaterialExpressions.cpp
@@ -13890,8 +13890,8 @@ void UMaterialFunction::PostLoad()
 		if (EditorOnly->ExpressionCollection.Expressions.Remove(nullptr) != 0)
 		{
 			// Force this function to recompile because its expressions have changed
-			// Warning: any content taking this path will recompile every load until saved!
 			// Which means removing an expression class will cause the need for a resave of all materials affected
+			UE_LOG(LogMaterial, Warning, TEXT("Please resave %s.  It is missing a material expression and this will cause any material using this function to recopmile shaders each time it loads."), *GetFullName());
 			StateId = FGuid::NewGuid();
 		}
 	}
