Add a FInputDevicePropertyHandle return type to the device property subsystem.

- Change from using a TSubclassOf to a raw TObjectPtr because you don't get the default values of an instanced object with TSubclassOf, which would mean you cant set the triggers properly.

- Rename "bResetUponCompletion" to "bRemoveAfterEvaluationTime" for a clearer description of what the flag will do.

- Expose more things with ENGINE_API so that we can make better debug tools in the InputDebugging plugin

#rnx

diff --git a/Engine/Source/Runtime/Engine/Classes/GameFramework/InputDeviceSubsystem.h b/Engine/Source/Runtime/Engine/Classes/GameFramework/InputDeviceSubsystem.h
--- a/Engine/Source/Runtime/Engine/Classes/GameFramework/InputDeviceSubsystem.h
+++ b/Engine/Source/Runtime/Engine/Classes/GameFramework/InputDeviceSubsystem.h
@@ -8,10 +8,47 @@
 #include "Templates/SubclassOf.h"
 #include "InputDeviceSubsystem.generated.h"
 
-DECLARE_LOG_CATEGORY_EXTERN(LogInputDeviceProperties, Log, All);
+ENGINE_API DECLARE_LOG_CATEGORY_EXTERN(LogInputDeviceProperties, Log, All);
 
 class UInputDeviceProperty;
 
+/** A handle to an active input device property that is being used by the InputDeviceSubsytem. */
+USTRUCT(BlueprintType)
+struct ENGINE_API FInputDevicePropertyHandle
+{
+	friend class UInputDeviceSubsystem;
+
+	GENERATED_BODY()
+
+	FInputDevicePropertyHandle();
+	~FInputDevicePropertyHandle() = default;
+
+	/** Returns true if this handle is valid */
+	bool IsValid() const;
+
+	/** An invalid Input Device Property handle */
+	static FInputDevicePropertyHandle InvalidHandle;
+
+	friend uint32 GetTypeHash(const FInputDevicePropertyHandle& InHandle);
+
+	bool operator==(const FInputDevicePropertyHandle& Other) const;
+	bool operator!=(const FInputDevicePropertyHandle& Other) const;
+
+	FString ToString() const;
+
+private:
+
+	// Private constructor because we don't want any other users to make a valid device property handle.
+	FInputDevicePropertyHandle(uint32 InternalID);
+
+	/** Static function to get a unique device handle. */
+	static FInputDevicePropertyHandle AcquireValidHandle();
+
+	/** The internal ID of this handle. Populated by the private constructor in AcquireValidHandle */
+	UPROPERTY()
+	uint32 InternalId;
+};
+
 /** Contains a pointer to an active device property and keeps track of how long it has been evaluated for */
 USTRUCT()
 struct ENGINE_API FActiveDeviceProperty
@@ -33,7 +70,10 @@ struct ENGINE_API FActiveDeviceProperty
 	* it will have it's Reset function called. Set this to false if you want your device property to stay 
 	* in it's ending state. 
 	*/
-	bool bResetUponCompletion = true;
+	bool bRemoveAfterEvaluationTime = true;
+
+	/** The handle of this active property. */
+	FInputDevicePropertyHandle PropertyHandle = FInputDevicePropertyHandle::InvalidHandle;
 };
 
 /** Parameters for the UInputDeviceSubsystem::SetDeviceProperty function */
@@ -44,11 +84,11 @@ struct ENGINE_API FSetDevicePropertyParams
 
 	FSetDevicePropertyParams();
 
-	FSetDevicePropertyParams(TSubclassOf<UInputDeviceProperty> InClass, const FPlatformUserId InUserId, const bool bInResetUponCompletion = true);
+	FSetDevicePropertyParams(TObjectPtr<UInputDeviceProperty> InProperty, const FPlatformUserId InUserId, const bool bInRemoveAfterEvaluationTime = true);
 	
-	/** The device Property to set */
-	UPROPERTY(EditAnywhere, BlueprintReadWrite, Category = "Input Devices")
-	TSubclassOf<UInputDeviceProperty> DevicePropertyClass;
+	/** The device property to set */
+	UPROPERTY(EditInstanceOnly, BlueprintReadWrite, Category = "Input Devices")
+	TObjectPtr<UInputDeviceProperty> DeviceProperty;
 	
 	/** The Platform User whose device's should receive the device property */
 	UPROPERTY(EditAnywhere, BlueprintReadWrite, Category = "Input Devices")
@@ -56,11 +96,12 @@ struct ENGINE_API FSetDevicePropertyParams
 
 	/**
 	* If true, then when the InputDeviceProperty is done being evaluated and has fulfilled its duration, then
-	* it will have it's Reset function called. Set this to false if you want your device property to stay
-	* in it's ending state.
+	* it will have it's Reset function called and be removed as an active property. 
+	* 
+	* Set this to false if you want your device property to keep being applied without regards to it's duration.
 	*/
 	UPROPERTY(EditAnywhere, BlueprintReadWrite, Category = "Input Devices")
-	bool bResetUponCompletion = true;
+	bool bRemoveAfterEvaluationTime = true;
 };
 
 /**
@@ -78,6 +119,7 @@ UCLASS(BlueprintType)
 class ENGINE_API UInputDeviceSubsystem : public UEngineSubsystem, public FTickableGameObject
 {
 	friend class FInputDeviceSubsystemProcessor;
+	friend class FInputDeviceDebugVisualizer;
 	
 	GENERATED_BODY()
 
@@ -100,7 +142,7 @@ public:
 
 	/** Set the given device property, which will start the evaluation and application of it to the given platform user. */
 	UFUNCTION(BlueprintCallable, Category = "Input Devices", meta = (AutoCreateRefTerm = "Params"))
-	void SetDeviceProperty(const FSetDevicePropertyParams& Params);
+	FInputDevicePropertyHandle SetDeviceProperty(const FSetDevicePropertyParams& Params);
 
 	/** 
 	* Remove any active device properties that have the same class as the one given.
@@ -111,12 +153,26 @@ public:
 	* @return						The number of removed device properties. 
 	*/
 	UFUNCTION(BlueprintCallable, Category = "Input Devices", meta = (ReturnDisplayName = "Num Removed"))
-	int32 RemoveDeviceProperty(const FPlatformUserId UserId, TSubclassOf<UInputDeviceProperty> DevicePropertyClass);
+	int32 RemoveDevicePropertiesOfClass(const FPlatformUserId UserId, TSubclassOf<UInputDeviceProperty> DevicePropertyClass);
+
+	/**
+	* Remove a single device property based on it's handle
+	*
+	* @param FInputDevicePropertyHandle		Device property handle to be removed
+	*
+	* @return								The number of removed device properties.
+	*/
+	UFUNCTION(BlueprintCallable, Category = "Input Devices", meta = (ReturnDisplayName = "Num Removed"))
+	int32 RemoveDevicePropertyByHandle(const FInputDevicePropertyHandle HandleToRemove);
 
 	/** Removes all the current Input Device Properties that are active, regardless of the Platform User */
 	UFUNCTION(BlueprintCallable, Category = "Input Devices")
 	void RemoveAllDeviceProperties();
 
+	/** Returns true if the given handle is valid */
+	UFUNCTION(BlueprintCallable, Category = "Input Devices")
+	static bool IsDevicePropertyHandleValid(const FInputDevicePropertyHandle& InHandle);
+
 	/** Gets the most recently used hardware input device for the given platform user */
 	FHardwareDeviceIdentifier GetMostRecentlyUsedHardwareDevice(const FPlatformUserId InUserId) const;
 
@@ -134,12 +190,12 @@ protected:
 	void SetMostRecentlyUsedHardwareDevice(const FInputDeviceId InDeviceId, const FHardwareDeviceIdentifier& InHardwareId);
 	
 	/**
-	* Array of the active device properties that are currently being evaluated on Tick.
+	* Array of currently active input device properties that will be evaluated on tick
 	*/
 	UPROPERTY(Transient)
 	TArray<FActiveDeviceProperty> ActiveProperties;
 	
-	/** A map of platform user's to their most recent hardwave device identifier */
+	/** A map of platform user's to their most recent hardware device identifier */
 	TMap<FPlatformUserId, FHardwareDeviceIdentifier> LatestInputDeviceIdentifiers;
 
 	/** An input processor that is used to determine the current hardware input device */
diff --git a/Engine/Source/Runtime/Engine/Private/GameFramework/InputDeviceSubsystem.cpp b/Engine/Source/Runtime/Engine/Private/GameFramework/InputDeviceSubsystem.cpp
--- a/Engine/Source/Runtime/Engine/Private/GameFramework/InputDeviceSubsystem.cpp
+++ b/Engine/Source/Runtime/Engine/Private/GameFramework/InputDeviceSubsystem.cpp
@@ -74,7 +74,57 @@ public:
 };
 
 ////////////////////////////////////////////////////////
-// UInputDeviceSubsystem
+// FInputDevicePropertyHandle
+
+FInputDevicePropertyHandle FInputDevicePropertyHandle::InvalidHandle(0);
+
+FInputDevicePropertyHandle::FInputDevicePropertyHandle()
+	: InternalId(0)
+{
+
+}
+
+FInputDevicePropertyHandle::FInputDevicePropertyHandle(uint32 InInternalID)
+	: InternalId(InInternalID)
+{
+
+}
+
+bool FInputDevicePropertyHandle::operator==(const FInputDevicePropertyHandle& Other) const
+{
+	return InternalId == Other.InternalId;
+}
+
+bool FInputDevicePropertyHandle::operator!=(const FInputDevicePropertyHandle& Other) const
+{
+	return InternalId != Other.InternalId;
+}
+
+uint32 GetTypeHash(const FInputDevicePropertyHandle& InHandle)
+{
+	return GetTypeHash(InHandle.InternalId);
+}
+
+FString FInputDevicePropertyHandle::ToString() const
+{
+	return IsValid() ? FString::FromInt(InternalId) : TEXT("Invalid");
+}
+
+bool FInputDevicePropertyHandle::IsValid() const
+{
+	return InternalId != FInputDevicePropertyHandle::InvalidHandle.InternalId;
+}
+
+FInputDevicePropertyHandle FInputDevicePropertyHandle::AcquireValidHandle()
+{
+	// 0 is the "Invalid" index for these handles. Start them at 1
+	static uint32 GHandleIndex = 1;
+
+	return FInputDevicePropertyHandle(++GHandleIndex);
+}
+
+////////////////////////////////////////////////////////
+// FSetDevicePropertyParams
 
 FSetDevicePropertyParams::FSetDevicePropertyParams()
 	: UserId(FSlateApplicationBase::SlateAppPrimaryPlatformUser)
@@ -82,14 +132,17 @@ FSetDevicePropertyParams::FSetDevicePropertyParams()
 
 }
 
-FSetDevicePropertyParams::FSetDevicePropertyParams(TSubclassOf<UInputDeviceProperty> InClass, const FPlatformUserId InUserId, const bool bInResetUponCompletion /* = true */)
-	: DevicePropertyClass(InClass)
+FSetDevicePropertyParams::FSetDevicePropertyParams(TObjectPtr<UInputDeviceProperty> InProperty, const FPlatformUserId InUserId, const bool bInRemoveAfterEvaluationTime /* = true */)
+	: DeviceProperty(InProperty)
 	, UserId(InUserId)
-	, bResetUponCompletion(bInResetUponCompletion)
+	, bRemoveAfterEvaluationTime(bInRemoveAfterEvaluationTime)
 {
 	
 }
 
+////////////////////////////////////////////////////////
+// UInputDeviceSubsystem
+
 UInputDeviceSubsystem* UInputDeviceSubsystem::Get()
 {
 	return GEngine ? GEngine->GetEngineSubsystem<UInputDeviceSubsystem>() : nullptr;
@@ -203,14 +256,12 @@ void UInputDeviceSubsystem::Tick(float DeltaTime)
 		ActiveProperties[Index].EvaluatedDuration += DeltaTime;
 
 		// If the property has run past it's duration, reset it and remove it from our active properties
-		if (ActiveProperties[Index].EvaluatedDuration > ActiveProperties[Index].Property->GetDuration())
+		// Only do this if it is marked as 'bRemoveAfterEvaluationTime' so that you can keep device properties set
+		// without having to worry about duration. 
+		if (ActiveProperties[Index].bRemoveAfterEvaluationTime && ActiveProperties[Index].EvaluatedDuration > ActiveProperties[Index].Property->GetDuration())
 		{
-			// Evaluation complete, mark this one for removal
-			if (ActiveProperties[Index].bResetUponCompletion)
-			{
-				ActiveProperties[Index].Property->ResetDeviceProperty(ActiveProperties[Index].PlatformUser);
-			}			
-			ActiveProperties.RemoveAtSwap(Index);
+			ActiveProperties[Index].Property->ResetDeviceProperty(ActiveProperties[Index].PlatformUser);
+			ActiveProperties.RemoveAtSwap(Index);			
 		}
 		// Otherwise, we can evaluate and apply it as normal
 		else
@@ -221,27 +272,35 @@ void UInputDeviceSubsystem::Tick(float DeltaTime)
 	}
 }
 
-void UInputDeviceSubsystem::SetDeviceProperty(const FSetDevicePropertyParams& Params)
+FInputDevicePropertyHandle UInputDeviceSubsystem::SetDeviceProperty(const FSetDevicePropertyParams& Params)
 {
-	if (!Params.DevicePropertyClass)
+	if (!Params.DeviceProperty)
 	{
 		UE_LOG(LogInputDeviceProperties, Error, TEXT("Invalid DevicePropertyClass passed into SetDeviceProperty! Nothing will happen."));
-		return;
+		return FInputDevicePropertyHandle::InvalidHandle;
 	}
-	
-	FActiveDeviceProperty ActiveProp = {};
+		
+	FInputDevicePropertyHandle OutHandle = FInputDevicePropertyHandle::AcquireValidHandle();
+
+	if (ensureAlwaysMsgf(OutHandle.IsValid(), TEXT("Unable to acquire a valid input device property handle!")))
+	{
+		FActiveDeviceProperty ActiveProp = {};
 
-	// Spawn an instance of this device property
-	ActiveProp.Property = NewObject<UInputDeviceProperty>(/* Outer = */ this, /* Class */ Params.DevicePropertyClass);
-	ensure(ActiveProp.Property);
+		// Spawn an instance of this device property
+		ActiveProp.Property = DuplicateObject<UInputDeviceProperty>(Params.DeviceProperty, /* Outer = */ this);
+		ensure(ActiveProp.Property);
 
-	ActiveProp.PlatformUser = Params.UserId;
-	ActiveProp.bResetUponCompletion = Params.bResetUponCompletion;
-	
-	ActiveProperties.Emplace(ActiveProp);
+		ActiveProp.PlatformUser = Params.UserId;
+		ActiveProp.bRemoveAfterEvaluationTime = Params.bRemoveAfterEvaluationTime;
+		ActiveProp.PropertyHandle = OutHandle;
+
+		ActiveProperties.Add(ActiveProp);
+	}
+
+	return OutHandle;
 }
 
-int32 UInputDeviceSubsystem::RemoveDeviceProperty(const FPlatformUserId UserId, TSubclassOf<UInputDeviceProperty> DevicePropertyClass)
+int32 UInputDeviceSubsystem::RemoveDevicePropertiesOfClass(const FPlatformUserId UserId, TSubclassOf<UInputDeviceProperty> DevicePropertyClass)
 {
 	int32 NumRemoved = 0;
 
@@ -266,6 +325,34 @@ int32 UInputDeviceSubsystem::RemoveDeviceProperty(const FPlatformUserId UserId,
 	return NumRemoved;
 }
 
+int32 UInputDeviceSubsystem::RemoveDevicePropertyByHandle(const FInputDevicePropertyHandle HandleToRemove)
+{
+	int32 NumRemoved = 0;
+
+	for (int32 Index = ActiveProperties.Num() - 1; Index >= 0; --Index)
+	{
+		if (ActiveProperties[Index].PropertyHandle == HandleToRemove)
+		{
+			ActiveProperties[Index].Property->ResetDeviceProperty(ActiveProperties[Index].PlatformUser);
+			ActiveProperties.RemoveAtSwap(Index);
+			++NumRemoved;
+			break;
+		}
+	}
+
+	if (NumRemoved <= 0)
+	{
+		UE_LOG(LogInputDeviceProperties, Warning, TEXT("Unable to remove a device property with handle '%s'"), *HandleToRemove.ToString());
+	}
+
+	return NumRemoved;
+}
+
+bool UInputDeviceSubsystem::IsDevicePropertyHandleValid(const FInputDevicePropertyHandle& InHandle)
+{
+	return InHandle.IsValid();
+}
+
 void UInputDeviceSubsystem::RemoveAllDeviceProperties()
 {
 	ActiveProperties.Empty();
