Fix UStreamableRenderAsset::IsFullyStreamedIn never returning true when we are using cooked content in an editor and the cooked texture has optional mips

diff --git a/Engine/Source/Runtime/Engine/Private/StreamableRenderAsset.cpp b/Engine/Source/Runtime/Engine/Private/StreamableRenderAsset.cpp
--- a/Engine/Source/Runtime/Engine/Private/StreamableRenderAsset.cpp
+++ b/Engine/Source/Runtime/Engine/Private/StreamableRenderAsset.cpp
@@ -345,6 +345,17 @@ bool UStreamableRenderAsset::IsFullyStreamedIn()
 		return true;
 	}
 
+#if WITH_EDITOR
+	UPackage* Package = GetOutermost();
+	if (Package
+		&& Package->bIsCookedForEditor
+		&& CachedSRRState.NumNonOptionalLODs < CachedSRRState.MaxNumLODs
+		&& IStreamingManager::Get().IsRenderAssetStreamingEnabled())
+	{
+		return IStreamingManager::Get().GetRenderAssetStreamingManager().IsFullyStreamedIn(this);
+	}
+#endif	
+
 	// IsFullyStreamedIn() might be used incorrectly if any logic waits on it to be true.
 	// there could be optional mips which are not available to be loaded, so waiting on IsFullyStreamedIn would never finish
 	ensureMsgf(CachedSRRState.NumResidentLODs != CachedSRRState.NumNonOptionalLODs, TEXT("IsFullyStreamedIn() is being called on (%s) which might not have optional LODs mounted."), *GetFName().ToString());
diff --git a/Engine/Source/Runtime/Engine/Private/Streaming/StreamingManagerTexture.cpp b/Engine/Source/Runtime/Engine/Private/Streaming/StreamingManagerTexture.cpp
--- a/Engine/Source/Runtime/Engine/Private/Streaming/StreamingManagerTexture.cpp
+++ b/Engine/Source/Runtime/Engine/Private/Streaming/StreamingManagerTexture.cpp
@@ -869,6 +869,29 @@ void FRenderAssetStreamingManager::RemoveStreamingRenderAsset( UStreamableRender
 	STAT(GatheredStats.CallbacksCycles += FPlatformTime::Cycles();)
 }
 
+bool FRenderAssetStreamingManager::IsFullyStreamedIn(UStreamableRenderAsset* RenderAsset)
+{
+	check(RenderAsset);
+
+	const FStreamingRenderAsset* StreamingAsset = GetStreamingRenderAsset(RenderAsset);
+	const FStreamableRenderResourceState& AssetState = RenderAsset->GetStreamableResourceState();
+	
+	if (StreamingAsset)
+	{
+		const FStreamingRenderAsset::EOptionalMipsState OptionalLODState = StreamingAsset->OptionalMipsState;
+		int32 NumLODsToConsiderAsFull = AssetState.MaxNumLODs - RenderAsset->GetCachedLODBias();
+		
+		if (OptionalLODState == FStreamingRenderAsset::OMS_NoOptionalMips)
+		{
+			NumLODsToConsiderAsFull = FMath::Min(NumLODsToConsiderAsFull, (int32)AssetState.NumNonOptionalLODs);
+		}
+
+		return AssetState.NumResidentLODs >= NumLODsToConsiderAsFull;
+	}
+
+	return false;
+}
+
 /** Called when a spawned primitive is deleted, or when an actor is destroyed in the editor. */
 void FRenderAssetStreamingManager::NotifyActorDestroyed( AActor* Actor )
 {
diff --git a/Engine/Source/Runtime/Engine/Private/Streaming/StreamingManagerTexture.h b/Engine/Source/Runtime/Engine/Private/Streaming/StreamingManagerTexture.h
--- a/Engine/Source/Runtime/Engine/Private/Streaming/StreamingManagerTexture.h
+++ b/Engine/Source/Runtime/Engine/Private/Streaming/StreamingManagerTexture.h
@@ -139,6 +139,9 @@ struct FRenderAssetStreamingManager final : public IRenderAssetStreamingManager
 	/** Removes a texture/mesh from the streaming manager. */
 	virtual void RemoveStreamingRenderAsset( UStreamableRenderAsset* RenderAsset ) override;
 
+	/** Only call on the game thread. */
+	virtual bool IsFullyStreamedIn(UStreamableRenderAsset* RenderAsset) override;
+
 	/** Adds a ULevel to the streaming manager. */
 	virtual void AddLevel( class ULevel* Level ) override;
 
diff --git a/Engine/Source/Runtime/Engine/Public/ContentStreaming.h b/Engine/Source/Runtime/Engine/Public/ContentStreaming.h
--- a/Engine/Source/Runtime/Engine/Public/ContentStreaming.h
+++ b/Engine/Source/Runtime/Engine/Public/ContentStreaming.h
@@ -454,6 +454,9 @@ struct IRenderAssetStreamingManager : public IStreamingManager
 	/** Removes a texture/mesh from the streaming manager. */
 	virtual void RemoveStreamingRenderAsset(UStreamableRenderAsset* RenderAsset) = 0;
 
+	/** Check whether all runtime-allowed LODs have been loaded. */
+	virtual bool IsFullyStreamedIn(UStreamableRenderAsset* RenderAsset) = 0;
+
 	virtual int64 GetMemoryOverBudget() const = 0;
 
 	/** Pool size for streaming. */
