Fix a compilation error

diff --git a/Engine/Source/Runtime/Engine/Private/StreamableRenderAsset.cpp b/Engine/Source/Runtime/Engine/Private/StreamableRenderAsset.cpp
--- a/Engine/Source/Runtime/Engine/Private/StreamableRenderAsset.cpp
+++ b/Engine/Source/Runtime/Engine/Private/StreamableRenderAsset.cpp
@@ -350,7 +350,7 @@ bool UStreamableRenderAsset::IsFullyStreamedIn()
 	if (Package
 		&& Package->bIsCookedForEditor
 		&& CachedSRRState.NumNonOptionalLODs < CachedSRRState.MaxNumLODs
-		&& IStreamingManager::Get().IsRenderAssetStreamingEnabled())
+		&& IStreamingManager::Get().IsRenderAssetStreamingEnabled(EStreamableRenderAssetType::None))
 	{
 		return IStreamingManager::Get().GetRenderAssetStreamingManager().IsFullyStreamedIn(this);
 	}
