Include reinstancer created objects in the CPFUO ReferenceReplacementMap. When we encounter an instanced subobject, and it's already in the ReferenceReplacementMap, make sure that it is outered correctly

diff --git a/Engine/Source/Editor/UnrealEd/Private/Kismet2/KismetReinstanceUtilities.cpp b/Engine/Source/Editor/UnrealEd/Private/Kismet2/KismetReinstanceUtilities.cpp
--- a/Engine/Source/Editor/UnrealEd/Private/Kismet2/KismetReinstanceUtilities.cpp
+++ b/Engine/Source/Editor/UnrealEd/Private/Kismet2/KismetReinstanceUtilities.cpp
@@ -1921,6 +1921,7 @@ static void ReplaceObjectHelper(UObject*& OldObject, UClass* OldClass, UObject*&
 	// Copy property values
 	Options.bNotifyObjectReplacement = true;
 	Options.bSkipCompilerGeneratedDefaults = true;
+	Options.OptionalReplacementMappings = &OldToNewInstanceMap;
 	UEditorEngine::CopyPropertiesForUnrelatedObjects(OldObject, NewUObject, Options);
 	// Generate new subobjects
 	InstancedPropertyUtils::FArchiveInsertInstancedSubObjects InstancedSubObjSpawner(NewUObject, InstancedPropertyMap);
diff --git a/Engine/Source/Runtime/Engine/Classes/Engine/Engine.h b/Engine/Source/Runtime/Engine/Classes/Engine/Engine.h
--- a/Engine/Source/Runtime/Engine/Classes/Engine/Engine.h
+++ b/Engine/Source/Runtime/Engine/Classes/Engine/Engine.h
@@ -2867,6 +2867,7 @@ public:
 
 		// In cases where the SourceObject will no longer be able to look up its correct Archetype, it can be supplied
 		UObject* SourceObjectArchetype;
+		TMap<UObject*, UObject*>* OptionalReplacementMappings;
 
 		ENGINE_API FCopyPropertiesForUnrelatedObjectsParams();
 		ENGINE_API FCopyPropertiesForUnrelatedObjectsParams(const FCopyPropertiesForUnrelatedObjectsParams&);
diff --git a/Engine/Source/Runtime/Engine/Private/InstancedReferenceSubobjectHelper.cpp b/Engine/Source/Runtime/Engine/Private/InstancedReferenceSubobjectHelper.cpp
--- a/Engine/Source/Runtime/Engine/Private/InstancedReferenceSubobjectHelper.cpp
+++ b/Engine/Source/Runtime/Engine/Private/InstancedReferenceSubobjectHelper.cpp
@@ -249,30 +249,55 @@ void FFindInstancedReferenceSubobjectHelper::Duplicate(UObject* OldObject, UObje
 				{
 					const bool bKeptByOld = OldInstancedSubObjects.Contains(Obj);
 					const bool bNotHandledYet = !ReferenceReplacementMap.Contains(Obj);
-					if (bKeptByOld && bNotHandledYet)
+					if (bKeptByOld)
 					{
-						// This name may have been taken by an instanced subobject, since we're copying from
-						// an object we want to recreate it here with a duplicate of the old object. I'm 
-						// using a rename call here to reserve the name for ourself. These objects come from
-						// FObjectInitializer::InstanceSubobjects()
-						UObject* ExistingObject = StaticFindObjectFast(UObject::StaticClass(), NewObject, Obj->GetFName());
-						if (ExistingObject)
+						if (bNotHandledYet)
 						{
-							ExistingObject->Rename(
-								nullptr,
-								GetTransientPackage(),
-								REN_DoNotDirty | REN_DontCreateRedirectors | REN_ForceNoResetLoaders | REN_NonTransactional);
+							// This name may have been taken by an instanced subobject, since we're copying from
+							// an object we want to recreate it here with a duplicate of the old object. I'm 
+							// using a rename call here to reserve the name for ourself. These objects come from
+							// FObjectInitializer::InstanceSubobjects()
+							UObject* ExistingObject = StaticFindObjectFast(UObject::StaticClass(), NewObject, Obj->GetFName());
+							if (ExistingObject)
+							{
+								ExistingObject->Rename(
+									nullptr,
+									GetTransientPackage(),
+									REN_DoNotDirty | REN_DontCreateRedirectors | REN_ForceNoResetLoaders | REN_NonTransactional);
+							}
+
+							UObject* NewEditInlineSubobject = StaticDuplicateObject(Obj, NewObject, Obj->GetFName());
+							ReferenceReplacementMap.Add(Obj, NewEditInlineSubobject);
+
+							// NOTE: we cannot patch OldObject's linker table here, since we don't 
+							//       know the relation between the two objects (one could be of a 
+							//       super class, and the other a child)
+
+							// We also need to make sure to fixup any properties here
+							DuplicatedObjects.Add(NewEditInlineSubobject);
+						}
+						else
+						{
+							// make sure the object is outered correctly:
+							check(bNewObjectHasOldOuter);
+							UObject* RealNewSubobject = ReferenceReplacementMap.FindChecked(Obj);
+							if (RealNewSubobject->GetOuter() != NewObject)
+							{
+								UObject* ExistingObject = StaticFindObjectFast(UObject::StaticClass(), NewObject, RealNewSubobject->GetFName());
+								if (ExistingObject)
+								{
+									ExistingObject->Rename(
+										nullptr,
+										GetTransientPackage(),
+										REN_DoNotDirty | REN_DontCreateRedirectors | REN_ForceNoResetLoaders | REN_NonTransactional);
+									ReferenceReplacementMap.Add(ExistingObject, RealNewSubobject);
+								}
+								RealNewSubobject->Rename(
+									*(RealNewSubobject->GetName()),
+									NewObject,
+									REN_DoNotDirty | REN_DontCreateRedirectors | REN_ForceNoResetLoaders | REN_NonTransactional);
+							}
 						}
-
-						UObject* NewEditInlineSubobject = StaticDuplicateObject(Obj, NewObject, Obj->GetFName());
-						ReferenceReplacementMap.Add(Obj, NewEditInlineSubobject);
-
-						// NOTE: we cannot patch OldObject's linker table here, since we don't 
-						//       know the relation between the two objects (one could be of a 
-						//       super class, and the other a child)
-
-						// We also need to make sure to fixup any properties here
-						DuplicatedObjects.Add(NewEditInlineSubobject);
 					}
 				}
 			}
diff --git a/Engine/Source/Runtime/Engine/Private/UnrealEngine.cpp b/Engine/Source/Runtime/Engine/Private/UnrealEngine.cpp
--- a/Engine/Source/Runtime/Engine/Private/UnrealEngine.cpp
+++ b/Engine/Source/Runtime/Engine/Private/UnrealEngine.cpp
@@ -16638,6 +16638,7 @@ UEngine::FCopyPropertiesForUnrelatedObjectsParams::FCopyPropertiesForUnrelatedOb
 	, bNotifyObjectReplacement(false)
 	, bClearReferences(true)
 	, SourceObjectArchetype(nullptr)
+	, OptionalReplacementMappings(nullptr)
 {}
 UEngine::FCopyPropertiesForUnrelatedObjectsParams::FCopyPropertiesForUnrelatedObjectsParams(const FCopyPropertiesForUnrelatedObjectsParams&) = default;
 PRAGMA_ENABLE_DEPRECATION_WARNINGS
@@ -16770,6 +16771,10 @@ void UEngine::CopyPropertiesForUnrelatedObjects(UObject* OldObject, UObject* New
 		ReferenceReplacementMap.Add(OldObject->GetClass(), NewObject->GetClass());
 	}
 	ReferenceReplacementMap.Add(OldObject->GetClass()->GetDefaultObject(), NewObject->GetClass()->GetDefaultObject());
+	if (Params.OptionalReplacementMappings)
+	{
+		ReferenceReplacementMap.Append(*Params.OptionalReplacementMappings);
+	}
 
 	TArray<UObject*> ComponentsOnNewObject;
 	TMap<UObject*, UObject*> AggressiveReplaceReferences;
