Fix for metasound failing to play when soundwave loading behavior improperly overriden

diff --git a/Engine/Plugins/Runtime/Metasound/Source/MetasoundEngine/Private/MetasoundSource.cpp b/Engine/Plugins/Runtime/Metasound/Source/MetasoundEngine/Private/MetasoundSource.cpp
--- a/Engine/Plugins/Runtime/Metasound/Source/MetasoundEngine/Private/MetasoundSource.cpp
+++ b/Engine/Plugins/Runtime/Metasound/Source/MetasoundEngine/Private/MetasoundSource.cpp
@@ -171,6 +171,7 @@ UMetaSoundSource::UMetaSoundSource(const FObjectInitializer& ObjectInitializer)
 	: Super(ObjectInitializer)
 	, FMetasoundAssetBase()
 {
+	bProcedural = true;
 	bRequiresStopFade = true;
 	NumChannels = 1;
 
diff --git a/Engine/Source/Runtime/Engine/Private/SoundWave.cpp b/Engine/Source/Runtime/Engine/Private/SoundWave.cpp
--- a/Engine/Source/Runtime/Engine/Private/SoundWave.cpp
+++ b/Engine/Source/Runtime/Engine/Private/SoundWave.cpp
@@ -1692,12 +1692,13 @@ void USoundWave::PostLoad()
 	}
 #endif
 
-	// Don't need to do anything in post load if this is a source bus
-	if (this->IsA(USoundSourceBus::StaticClass()))
+	// Don't need to do anything in post load if this is a source bus or procedural audio
+	if (this->IsA(USoundSourceBus::StaticClass()) || bProcedural)
 	{
 		return;
 	}
 
+
 	CacheInheritedLoadingBehavior();
 	
 	// If our loading behavior is defined by a sound class, we need to update whether this sound wave actually needs to retain its audio data or not.
@@ -3152,19 +3153,21 @@ float USoundWave::GetDuration() const
 
 bool USoundWave::IsStreaming(const TCHAR* PlatformName/* = nullptr */) const
 {
-	bool Result = false;
-	if (GIsEditor && ForceNonStreamingInEditorCVar != 0)
+	const bool bIsForceInline = LoadingBehavior == ESoundWaveLoadingBehavior::ForceInline;
+	const bool bIsForcedNonStreamingInEditor = (GIsEditor && ForceNonStreamingInEditorCVar != 0);
+
+	const bool bIsForcedNonStreaming = (bProcedural || bIsForceInline || bIsForcedNonStreamingInEditor);
+
+	if (bIsForcedNonStreaming)
 	{
-		Result = false;
+		SoundWaveDataPtr->bIsStreaming = false;
 	}
 	else
 	{
-		Result = IsStreaming(*FPlatformCompressionUtilities::GetCookOverrides(PlatformName));
+		SoundWaveDataPtr->bIsStreaming = IsStreaming(*FPlatformCompressionUtilities::GetCookOverrides(PlatformName));
 	}
-
-	SoundWaveDataPtr->bIsStreaming = Result; // update shared flags
 	
-	return Result;
+	return SoundWaveDataPtr->bIsStreaming;
 }
 
 bool USoundWave::IsStreaming(const FPlatformAudioCookOverrides& Overrides) const
@@ -3172,8 +3175,11 @@ bool USoundWave::IsStreaming(const FPlatformAudioCookOverrides& Overrides) const
 	// We stream if (A) bStreaming is set to true, (B) bForceInline is false and either bUseLoadOnDemand was set to true in
 	// our cook overrides, or the AutoStreamingThreshold was set and this sound is longer than the auto streaming threshold.
 	const bool bIsForceInline = LoadingBehavior == ESoundWaveLoadingBehavior::ForceInline;
+	const bool bIsForcedNonStreamingInEditor = (GIsEditor && ForceNonStreamingInEditorCVar != 0);
 
-	if (bIsForceInline || bProcedural)
+	const bool bIsForcedNonStreaming = (bProcedural || bIsForceInline || bIsForcedNonStreamingInEditor);
+
+	if (bIsForcedNonStreaming)
 	{
 		SoundWaveDataPtr->bIsStreaming = false;
 	}
@@ -3181,6 +3187,7 @@ bool USoundWave::IsStreaming(const FPlatformAudioCookOverrides& Overrides) const
 	{
 		SoundWaveDataPtr->bIsStreaming = true;
 	}
+
 	return SoundWaveDataPtr->bIsStreaming;
 }
 
@@ -3829,60 +3836,65 @@ void USoundWave::OverrideLoadingBehavior(ESoundWaveLoadingBehavior InLoadingBeha
 {
 	TRACE_CPUPROFILER_EVENT_SCOPE(USoundWave::OverrideLoadingBehavior);
 
-	const ESoundWaveLoadingBehavior OldBehavior = GetLoadingBehavior(false);
-	const bool bAlreadySetToRetained = (OldBehavior == ESoundWaveLoadingBehavior::RetainOnLoad);
-	const bool bAlreadyLoaded = !HasAnyFlags(RF_NeedLoad);
+	const bool bCanOverrideLoadingBehavior = (LoadingBehavior != ESoundWaveLoadingBehavior::ForceInline);
 
-	// already set to the most aggressive (non-inline) option
-	if (bAlreadySetToRetained)
+	if (bCanOverrideLoadingBehavior)
 	{
-		return;
-	}
-
-	// we don't want to retain in editor, we prime instead
-	if (GIsEditor && InLoadingBehavior == ESoundWaveLoadingBehavior::RetainOnLoad)
-	{
-		InLoadingBehavior = ESoundWaveLoadingBehavior::PrimeOnLoad;
-	}
-
-	// record the new loading behavior
-	// (if this soundwave isn't loaded yet, 
-	// SoundWaveDataPtr->LoadingBehavior will take precedence when it does load)
-	SoundWaveDataPtr->LoadingBehavior = InLoadingBehavior;
-	SoundWaveDataPtr->bLoadingBehaviorOverridden = true;
+		const ESoundWaveLoadingBehavior OldBehavior = GetLoadingBehavior(false);
+		const bool bAlreadySetToRetained = (OldBehavior == ESoundWaveLoadingBehavior::RetainOnLoad);
+		const bool bAlreadyLoaded = !HasAnyFlags(RF_NeedLoad);
 
-	// If we're loading for the cook commandlet, we don't have streamed audio chunks to load.
-	const bool bHasBuiltStreamedAudio = !GetOutermost()->HasAnyPackageFlags(PKG_ReloadingForCooker) && FApp::CanEverRenderAudio();
-
-	// Manually perform prime/retain on already loaded sound waves
-	if (bHasBuiltStreamedAudio && bAlreadyLoaded && IsStreaming())
-	{
-		if (InLoadingBehavior == ESoundWaveLoadingBehavior::RetainOnLoad)
+		// already set to the most aggressive (non-inline) option
+		if (bAlreadySetToRetained)
 		{
-			ConditionalPostLoad();
-			RetainCompressedAudio();
+			return;
 		}
-		else if (InLoadingBehavior == ESoundWaveLoadingBehavior::PrimeOnLoad)
+
+		// we don't want to retain in editor, we prime instead
+		if (GIsEditor && InLoadingBehavior == ESoundWaveLoadingBehavior::RetainOnLoad)
 		{
-#if WITH_EDITOR
-			// In editor, just make sure that data is available in the local DDC cache for quick access.
-			CachePlatformData(true /* bAsyncCache */);
+			InLoadingBehavior = ESoundWaveLoadingBehavior::PrimeOnLoad;
+		}
 
-			if (!InternalProxy.IsValid())
+		// record the new loading behavior
+		// (if this soundwave isn't loaded yet, 
+		// SoundWaveDataPtr->LoadingBehavior will take precedence when it does load)
+		SoundWaveDataPtr->LoadingBehavior = InLoadingBehavior;
+		SoundWaveDataPtr->bLoadingBehaviorOverridden = true;
+
+		// If we're loading for the cook commandlet, we don't have streamed audio chunks to load.
+		const bool bHasBuiltStreamedAudio = !GetOutermost()->HasAnyPackageFlags(PKG_ReloadingForCooker) && FApp::CanEverRenderAudio();
+
+		// Manually perform prime/retain on already loaded sound waves
+		if (bHasBuiltStreamedAudio && bAlreadyLoaded && IsStreaming())
+		{
+			if (InLoadingBehavior == ESoundWaveLoadingBehavior::RetainOnLoad)
 			{
-				InternalProxy = CreateSoundWaveProxy();
+				ConditionalPostLoad();
+				RetainCompressedAudio();
 			}
-#else
-			if (GetNumChunks() > 1)
+			else if (InLoadingBehavior == ESoundWaveLoadingBehavior::PrimeOnLoad)
 			{
+#if WITH_EDITOR
+				// In editor, just make sure that data is available in the local DDC cache for quick access.
+				CachePlatformData(true /* bAsyncCache */);
+
 				if (!InternalProxy.IsValid())
 				{
 					InternalProxy = CreateSoundWaveProxy();
 				}
+#else
+				if (GetNumChunks() > 1)
+				{
+					if (!InternalProxy.IsValid())
+					{
+						InternalProxy = CreateSoundWaveProxy();
+					}
 
-				IStreamingManager::Get().GetAudioStreamingManager().RequestChunk(InternalProxy, 1, [](EAudioChunkLoadResult) {});
-			}
+					IStreamingManager::Get().GetAudioStreamingManager().RequestChunk(InternalProxy, 1, [](EAudioChunkLoadResult) {});
+				}
 #endif
+			}
 		}
 	}
 }
@@ -4062,13 +4074,10 @@ FSoundWaveProxy::FSoundWaveProxy(USoundWave* InWave)
 	// this should have been allocated by the USoundWave and should always be valid
 	check(SoundWaveDataPtr);
 
-	// if we are force inline, we need to make sure the shared data is pulled from the DDC
-	// before we attempt to use a decoder on the proxy (not using stream caching)
-	if (SoundWaveDataPtr->GetLoadingBehavior() == ESoundWaveLoadingBehavior::ForceInline && !InWave->GetResourceData())
-	{
-		InWave->InitAudioResource(SoundWaveDataPtr->GetRuntimeFormat());
-		check(SoundWaveDataPtr->GetResourceSize() > 0);
-	}
+	// non-streaming sources need resource data initialized before the FSoundWaveProxy
+	// can be used. 
+	InWave->InitAudioResource(SoundWaveDataPtr->GetRuntimeFormat());
+	check((InWave->IsStreaming(nullptr)) || (SoundWaveDataPtr->GetResourceSize() > 0));
 }
 
 FSoundWaveProxy::~FSoundWaveProxy()
