Fix a non-unity compile error

diff --git a/Engine/Source/Runtime/Engine/Private/StreamableRenderAsset.cpp b/Engine/Source/Runtime/Engine/Private/StreamableRenderAsset.cpp
--- a/Engine/Source/Runtime/Engine/Private/StreamableRenderAsset.cpp
+++ b/Engine/Source/Runtime/Engine/Private/StreamableRenderAsset.cpp
@@ -6,6 +6,7 @@
 #include "ContentStreaming.h"
 #include "RenderAssetUpdate.h"
 #include "Rendering/NaniteCoarseMeshStreamingManager.h"
+#include "UObject/Package.h"
 
 #include UE_INLINE_GENERATED_CPP_BY_NAME(StreamableRenderAsset)
 
