Add a static Input Device color property that makes it easier to just set a single static color on a device without needing a curve.

#rnx

diff --git a/Engine/Source/Runtime/Engine/Classes/GameFramework/InputDeviceProperties.h b/Engine/Source/Runtime/Engine/Classes/GameFramework/InputDeviceProperties.h
--- a/Engine/Source/Runtime/Engine/Classes/GameFramework/InputDeviceProperties.h
+++ b/Engine/Source/Runtime/Engine/Classes/GameFramework/InputDeviceProperties.h
@@ -100,7 +100,7 @@ public:
 	
 	/**
 	 * Recalculates this device property's duration. This should be called whenever there are changes made
-	 * to things like curves, or other time sensative properties.
+	 * to things like curves, or other time sensitive properties.
 	 */
 	virtual float RecalculateDuration();
 
@@ -119,8 +119,8 @@ protected:
 	float PropertyDuration = 0.1f;
 };
 
-/** 
-* A property that can be used to change the color of an input device's light 
+/**
+* Set the color of an Input Device to a static color. 
 * 
 * NOTE: This property has platform specific implementations and may behave differently per platform.
 * See the docs for more details on each platform.
@@ -130,6 +130,37 @@ class UColorInputDeviceProperty : public UInputDeviceProperty
 {
 	GENERATED_BODY()
 
+public:
+
+	virtual void EvaluateDeviceProperty_Implementation(const FPlatformUserId PlatformUser, const float DeltaTime, const float Duration) override;
+	virtual void ResetDeviceProperty_Implementation(const FPlatformUserId PlatformUser) override;
+	virtual FInputDeviceProperty* GetInternalDeviceProperty() override;
+
+	/** True if the light should be enabled at all */
+	UPROPERTY(EditAnywhere, BlueprintReadWrite, Category = "DeviceProperty")
+	bool bEnable = true;
+
+	/** The color to set the light on  */
+	UPROPERTY(EditAnywhere, BlueprintReadWrite, Category = "DeviceProperty")
+	FColor LightColor = FColor::White;
+
+private:
+
+	/** The internal light color property that this represents; */
+	FInputDeviceLightColorProperty InternalProperty;
+};
+
+/** 
+* A property that can be used to change the color of an input device's light over time with a curve
+* 
+* NOTE: This property has platform specific implementations and may behave differently per platform.
+* See the docs for more details on each platform.
+*/
+UCLASS(Blueprintable, BlueprintType)
+class UColorInputDeviceCurveProperty : public UInputDeviceProperty
+{
+	GENERATED_BODY()
+
 public:
 
 	virtual void EvaluateDeviceProperty_Implementation(const FPlatformUserId PlatformUser, const float DeltaTime, const float Duration) override;
diff --git a/Engine/Source/Runtime/Engine/Private/GameFramework/InputDeviceProperties.cpp b/Engine/Source/Runtime/Engine/Private/GameFramework/InputDeviceProperties.cpp
--- a/Engine/Source/Runtime/Engine/Private/GameFramework/InputDeviceProperties.cpp
+++ b/Engine/Source/Runtime/Engine/Private/GameFramework/InputDeviceProperties.cpp
@@ -67,6 +67,27 @@ void UInputDeviceProperty::ResetDeviceProperty_Implementation(const FPlatformUse
 // UColorInputDeviceProperty
 
 void UColorInputDeviceProperty::EvaluateDeviceProperty_Implementation(const FPlatformUserId PlatformUser, const float DeltaTime, const float Duration)
+{
+	InternalProperty.bEnable = bEnable;
+	InternalProperty.Color = LightColor;
+}
+
+void UColorInputDeviceProperty::ResetDeviceProperty_Implementation(const FPlatformUserId PlatformUser)
+{
+	// Disabling the light will reset the color
+	InternalProperty.bEnable = false;
+	ApplyDeviceProperty(PlatformUser);
+}
+
+FInputDeviceProperty* UColorInputDeviceProperty::GetInternalDeviceProperty()
+{
+	return &InternalProperty;
+}
+
+///////////////////////////////////////////////////////////////////////
+// UColorInputDeviceCurveProperty
+
+void UColorInputDeviceCurveProperty::EvaluateDeviceProperty_Implementation(const FPlatformUserId PlatformUser, const float DeltaTime, const float Duration)
 {
 	InternalProperty.bEnable = bEnable;
 
@@ -77,19 +98,19 @@ void UColorInputDeviceProperty::EvaluateDeviceProperty_Implementation(const FPla
 	}
 }
 
-void UColorInputDeviceProperty::ResetDeviceProperty_Implementation(const FPlatformUserId PlatformUser)
+void UColorInputDeviceCurveProperty::ResetDeviceProperty_Implementation(const FPlatformUserId PlatformUser)
 {
 	// Disabling the light will reset the color
 	InternalProperty.bEnable = false;
 	ApplyDeviceProperty(PlatformUser);
 }
 
-FInputDeviceProperty* UColorInputDeviceProperty::GetInternalDeviceProperty()
+FInputDeviceProperty* UColorInputDeviceCurveProperty::GetInternalDeviceProperty()
 {
 	return &InternalProperty;
 }
 
-float UColorInputDeviceProperty::RecalculateDuration()
+float UColorInputDeviceCurveProperty::RecalculateDuration()
 {
 	float MinTime, MaxTime;
 	if (DeviceColorCurve)
