Add a "Get Platform User Id" function to the Pawn
This will make it a little easier to get the platform user in blueprints for setting device properties or doing UI stuff.

#rnx

diff --git a/Engine/Source/Runtime/Engine/Classes/GameFramework/Pawn.h b/Engine/Source/Runtime/Engine/Classes/GameFramework/Pawn.h
--- a/Engine/Source/Runtime/Engine/Classes/GameFramework/Pawn.h
+++ b/Engine/Source/Runtime/Engine/Classes/GameFramework/Pawn.h
@@ -389,6 +389,14 @@ public:
 	/** Returns true if controlled by a local (not network) Controller.	 */
 	UFUNCTION(BlueprintPure, Category=Pawn)
 	virtual bool IsLocallyControlled() const;
+
+	/**
+	 * Returns the Platform User ID of the PlayerController that is controlling this character.
+	 *
+	 * Returns an invalid Platform User ID if this character is not controlled by a local player.
+	 */
+	UFUNCTION(BlueprintCallable, BlueprintPure, Category = Pawn)
+	FPlatformUserId GetPlatformUserId() const;
   
 	/** Returns true if controlled by a human player (possessed by a PlayerController).	This returns true for players controlled by remote clients */
 	UFUNCTION(BlueprintPure, Category=Pawn)
diff --git a/Engine/Source/Runtime/Engine/Private/Pawn.cpp b/Engine/Source/Runtime/Engine/Private/Pawn.cpp
--- a/Engine/Source/Runtime/Engine/Private/Pawn.cpp
+++ b/Engine/Source/Runtime/Engine/Private/Pawn.cpp
@@ -248,6 +248,16 @@ bool APawn::IsLocallyControlled() const
 {
 	return ( Controller && Controller->IsLocalController() );
 }
+
+FPlatformUserId APawn::GetPlatformUserId() const
+{
+	if (const APlayerController* PC = Cast<APlayerController>(Controller))
+	{
+		return PC->GetPlatformUserId();
+	}
+	return PLATFORMUSERID_NONE;
+}
+
 bool APawn::IsPlayerControlled() const
 {
 	return PlayerState && !PlayerState->IsABot();
