Added UHLODSubsystem::SetHLODAlwaysLoadedCullDistance to allow culling of always loaded HLODs.

diff --git a/Engine/Source/Runtime/Engine/Private/WorldPartition/HLOD/HLODSubsystem.cpp b/Engine/Source/Runtime/Engine/Private/WorldPartition/HLOD/HLODSubsystem.cpp
--- a/Engine/Source/Runtime/Engine/Private/WorldPartition/HLOD/HLODSubsystem.cpp
+++ b/Engine/Source/Runtime/Engine/Private/WorldPartition/HLOD/HLODSubsystem.cpp
@@ -113,6 +113,7 @@ public:
 
 UHLODSubsystem::UHLODSubsystem()
 	: UWorldSubsystem()
+	, HLODAlwaysLoadedCullDistance(-1)
 {
 }
 
@@ -120,6 +121,109 @@ UHLODSubsystem::~UHLODSubsystem()
 {
 }
 
+void UHLODSubsystem::SetHLODAlwaysLoadedCullDistance(int32 InCullDistance)
+{
+	HLODAlwaysLoadedCullDistance = InCullDistance;
+}
+
+void UHLODSubsystem::OnWorldBeginPlay(UWorld& InWorld)
+{
+	Super::OnWorldBeginPlay(InWorld);
+
+	check(InWorld.IsGameWorld());
+
+	HLODAlwaysLoadedCullDistance = -1;
+	AlwaysLoadedHLODActors.Reset();
+	CulledAlwaysLoadedHLODActors.Reset();
+
+	// Cache always loaded HLODs with their Bounds
+	for (AActor* Actor : InWorld.PersistentLevel->Actors)
+	{
+		if (AWorldPartitionHLOD* HLODActor = Cast<AWorldPartitionHLOD>(Actor))
+		{
+			FBox Bounds = HLODActor->GetComponentsBoundingBox(true);
+			AlwaysLoadedHLODActors.Add(HLODActor, Bounds);
+		}
+	}
+}
+
+void UHLODSubsystem::UpdateStreamingState()
+{
+	Super::UpdateStreamingState();
+
+	UWorld* World = GetWorld();
+	if (!World || !World->IsGameWorld())
+	{
+		return;
+	}
+
+	UWorldPartition* WorldPartition = World->GetWorldPartition();
+	FWorldPartitionHLODRuntimeData* WorldPartitionHLODRuntimeData = (WorldPartition && WorldPartition->IsStreamingEnabled()) ? WorldPartitionsHLODRuntimeData.Find(WorldPartition) : nullptr;
+	if (!WorldPartitionHLODRuntimeData)
+	{
+		return;
+	}
+
+	TSet<AWorldPartitionHLOD*> FrameCulled;
+	TSet<AWorldPartitionHLOD*> FrameAlreadyHidden;
+	if (HLODAlwaysLoadedCullDistance > 0)
+	{
+		const TArray<FWorldPartitionStreamingSource>& StreamingSources = WorldPartition->GetStreamingSources();
+		double SquaredCullDistance = (double)HLODAlwaysLoadedCullDistance * (double)HLODAlwaysLoadedCullDistance;
+		for (const auto& Pair : AlwaysLoadedHLODActors)
+		{
+			if (AWorldPartitionHLOD* HLODActor = Pair.Key)
+			{
+				const FName CellName = HLODActor->GetSourceCellName();
+				if (FCellData* CellData = WorldPartitionHLODRuntimeData->CellsData.Find(CellName))
+				{
+					// If source cell is visible, hlod actor is already hidden
+					if (!CellData->bIsCellVisible)
+					{
+						if (!StreamingSources.IsEmpty())
+						{
+							bool bIsCulled = true;
+							for (const FWorldPartitionStreamingSource& Source : StreamingSources)
+							{
+								const FBox& Bounds = Pair.Value;
+								double Distance = Bounds.ComputeSquaredDistanceToPoint(Source.Location);
+								if (Distance < SquaredCullDistance)
+								{
+									bIsCulled = false;
+									break;
+								}
+							}
+							if (bIsCulled)
+							{
+								FrameCulled.Add(HLODActor);
+							}
+						}
+					}
+					else
+					{
+						FrameAlreadyHidden.Add(HLODActor);
+					}
+				}
+			}
+		}
+	}
+
+	TSet<AWorldPartitionHLOD*> ToCull = FrameCulled.Difference(CulledAlwaysLoadedHLODActors);
+	TSet<AWorldPartitionHLOD*> ToShow = CulledAlwaysLoadedHLODActors.Difference(FrameCulled.Union(FrameAlreadyHidden));
+
+	for (AWorldPartitionHLOD* HLODActor : ToCull)
+	{
+		HLODActor->SetVisibility(false);
+	}
+
+	for (AWorldPartitionHLOD* HLODActor : ToShow)
+	{
+		HLODActor->SetVisibility(UHLODSubsystem::WorldPartitionHLODEnabled);
+	}
+
+	CulledAlwaysLoadedHLODActors = MoveTemp(FrameCulled);
+}
+
 bool UHLODSubsystem::WorldPartitionHLODEnabled = true;
 
 FAutoConsoleCommand UHLODSubsystem::EnableHLODCommand(
diff --git a/Engine/Source/Runtime/Engine/Public/WorldPartition/HLOD/HLODSubsystem.h b/Engine/Source/Runtime/Engine/Public/WorldPartition/HLOD/HLODSubsystem.h
--- a/Engine/Source/Runtime/Engine/Public/WorldPartition/HLOD/HLODSubsystem.h
+++ b/Engine/Source/Runtime/Engine/Public/WorldPartition/HLOD/HLODSubsystem.h
@@ -12,7 +12,7 @@ class FSceneViewFamily;
 class FHLODResourcesResidencySceneViewExtension;
 class UWorldPartition;
 class UWorldPartitionRuntimeCell;
-
+class UWorld;
 
 DECLARE_MULTICAST_DELEGATE_OneParam(FWorldPartitionHLODActorRegisteredEvent, AWorldPartitionHLOD* /* InHLODActor */);
 DECLARE_MULTICAST_DELEGATE_OneParam(FWorldPartitionHLODActorUnregisteredEvent, AWorldPartitionHLOD* /* InHLODActor */);
@@ -37,6 +37,8 @@ public:
 
 	//~ Begin UWorldSubsystem Interface.
 	virtual bool DoesSupportWorldType(const EWorldType::Type WorldType) const override;
+	virtual void UpdateStreamingState() override;
+	virtual void OnWorldBeginPlay(UWorld& InWorld) override;
 	//~ End UWorldSubsystem Interface.
 
 	void RegisterHLODActor(AWorldPartitionHLOD* InWorldPartitionHLOD);
@@ -53,6 +55,8 @@ public:
 
 	FWorldPartitionHLODActorRegisteredEvent& OnHLODActorRegisteredEvent() { return HLODActorRegisteredEvent; }
 	FWorldPartitionHLODActorUnregisteredEvent& OnHLODActorUnregisteredEvent() { return HLODActorUnregisteredEvent; }
+
+	void SetHLODAlwaysLoadedCullDistance(int32 InCullDistance);
 	
 private:
 	struct FCellData
@@ -77,7 +81,11 @@ private:
 	};
 	
 	TMap<TObjectPtr<UWorldPartition>, FWorldPartitionHLODRuntimeData> WorldPartitionsHLODRuntimeData;
-		
+
+	TMap<AWorldPartitionHLOD*, FBox> AlwaysLoadedHLODActors;
+	TSet<AWorldPartitionHLOD*> CulledAlwaysLoadedHLODActors;
+	int32 HLODAlwaysLoadedCullDistance;
+
 	void OnWorldPartitionInitialized(UWorldPartition* InWorldPartition);
 	void OnWorldPartitionUninitialized(UWorldPartition* InWorldPartition);
 	void OnBeginRenderViews(const FSceneViewFamily& InViewFamily);
