Add a force dynamic data update flag to SkinnedMeshComponent. When set to true, it is the equivalent of r.DeferSkeletalDynamicDataUpdateUntilGDME=0 and r.RHICmdDeferSkeletalLockAndFillToRHIThread=0.
This allows sk mesh component to bypass r.DeferSkeletalDynamicDataUpdateUntilGDME=1 and update at end-of-frame skin cache batch update.

diff --git a/Engine/Source/Runtime/Engine/Classes/Components/SkinnedMeshComponent.h b/Engine/Source/Runtime/Engine/Classes/Components/SkinnedMeshComponent.h
--- a/Engine/Source/Runtime/Engine/Classes/Components/SkinnedMeshComponent.h
+++ b/Engine/Source/Runtime/Engine/Classes/Components/SkinnedMeshComponent.h
@@ -715,6 +715,14 @@ private:
 	UPROPERTY(transient)
 	uint8 bForceMeshObjectUpdate:1;
 
+	/**
+	   Whether to update dynamic bone & cloth sim data immediately, not to wait until GDME or defer update to RHIThread.
+	   When set to true, it is the equivalent of r.DeferSkeletalDynamicDataUpdateUntilGDME=0 and r.RHICmdDeferSkeletalLockAndFillToRHIThread=0.
+	   When set to false, r.DeferSkeletalDynamicDataUpdateUntilGDME and r.RHICmdDeferSkeletalLockAndFillToRHIThread values are respected.
+	 */
+	UPROPERTY(transient)
+	uint8 bForceUpdateDynamicDataImmediately : 1;
+
 protected:
 	/** Whether we are externally controlling tick rate */
 	uint8 bExternalTickRateControlled:1;
@@ -986,6 +994,9 @@ public:
 
 	bool HasMeshDeformer() const { return GetActiveMeshDeformer() != nullptr; }
 
+	bool GetForceUpdateDynamicDataImmediately() const { return bForceUpdateDynamicDataImmediately; }
+	void SetForceUpdateDynamicDataImmediately(bool bForceUpdateImmediately) { bForceUpdateDynamicDataImmediately = bForceUpdateImmediately; }
+
 	/**
 	 *	Compute SkeletalMesh MinLOD that will be used by this component
 	 */
diff --git a/Engine/Source/Runtime/Engine/Private/Components/SkinnedMeshComponent.cpp b/Engine/Source/Runtime/Engine/Private/Components/SkinnedMeshComponent.cpp
--- a/Engine/Source/Runtime/Engine/Private/Components/SkinnedMeshComponent.cpp
+++ b/Engine/Source/Runtime/Engine/Private/Components/SkinnedMeshComponent.cpp
@@ -412,6 +412,8 @@ USkinnedMeshComponent::USkinnedMeshComponent(const FObjectInitializer& ObjectIni
 	bNeedToFlipSpaceBaseBuffers = false;
 	bBoneVisibilityDirty = false;
 
+	bForceUpdateDynamicDataImmediately = false;
+
 	bCanEverAffectNavigation = false;
 	LeaderBoneMapCacheCount = 0;
 	bSyncAttachParentLOD = true;
diff --git a/Engine/Source/Runtime/Engine/Private/GPUSkinVertexFactory.cpp b/Engine/Source/Runtime/Engine/Private/GPUSkinVertexFactory.cpp
--- a/Engine/Source/Runtime/Engine/Private/GPUSkinVertexFactory.cpp
+++ b/Engine/Source/Runtime/Engine/Private/GPUSkinVertexFactory.cpp
@@ -247,7 +247,7 @@ static bool DeferSkeletalLockAndFillToRHIThread()
 }
 
 bool FGPUBaseSkinVertexFactory::FShaderDataType::UpdateBoneData(FRHICommandListImmediate& RHICmdList, const TArray<FMatrix44f>& ReferenceToLocalMatrices,
-	const TArray<FBoneIndexType>& BoneMap, uint32 RevisionNumber, bool bPrevious, ERHIFeatureLevel::Type InFeatureLevel, bool bUseSkinCache)
+	const TArray<FBoneIndexType>& BoneMap, uint32 RevisionNumber, bool bPrevious, ERHIFeatureLevel::Type InFeatureLevel, bool bUseSkinCache, bool bForceUpdateImmediately)
 {
 	// stat disabled by default due to low-value/high-frequency
 	//QUICK_SCOPE_CYCLE_COUNTER(STAT_FGPUBaseSkinVertexFactory_UpdateBoneData);
@@ -284,7 +284,7 @@ bool FGPUBaseSkinVertexFactory::FShaderDataType::UpdateBoneData(FRHICommandListI
 		}
 		if(NumBones)
 		{
-			if (!bUseSkinCache && DeferSkeletalLockAndFillToRHIThread())
+			if (!bUseSkinCache && !bForceUpdateImmediately && DeferSkeletalLockAndFillToRHIThread())
 			{
 				FRHIBuffer* VertexBuffer = CurrentBoneBuffer->VertexBufferRHI;
 				RHICmdList.EnqueueLambda([VertexBuffer, VectorArraySize, &ReferenceToLocalMatrices, &BoneMap](FRHICommandListImmediate& InRHICmdList)
@@ -1124,7 +1124,7 @@ IMPLEMENT_TYPE_LAYOUT(TGPUSkinAPEXClothVertexFactoryShaderParameters);
 -----------------------------------------------------------------------------*/
 
 bool FGPUBaseSkinAPEXClothVertexFactory::ClothShaderType::UpdateClothSimulData(FRHICommandListImmediate& RHICmdList, const TArray<FVector3f>& InSimulPositions,
-	const TArray<FVector3f>& InSimulNormals, uint32 RevisionNumber, ERHIFeatureLevel::Type FeatureLevel)
+	const TArray<FVector3f>& InSimulNormals, uint32 RevisionNumber, ERHIFeatureLevel::Type FeatureLevel, bool bForceUpdateImmediately)
 {
 	QUICK_SCOPE_CYCLE_COUNTER(STAT_FGPUBaseSkinAPEXClothVertexFactory_UpdateClothSimulData);
 
@@ -1151,7 +1151,7 @@ bool FGPUBaseSkinAPEXClothVertexFactory::ClothShaderType::UpdateClothSimulData(F
 
 	if(NumSimulVerts)
 	{
-		if (DeferSkeletalLockAndFillToRHIThread())
+		if (!bForceUpdateImmediately && DeferSkeletalLockAndFillToRHIThread())
 		{
 			FRHIBuffer* VertexBuffer = CurrentClothBuffer->VertexBufferRHI;
 			RHICmdList.EnqueueLambda([VertexBuffer, VectorArraySize, &InSimulPositions, &InSimulNormals](FRHICommandListImmediate& InRHICmdList)
diff --git a/Engine/Source/Runtime/Engine/Private/SkeletalRenderGPUSkin.cpp b/Engine/Source/Runtime/Engine/Private/SkeletalRenderGPUSkin.cpp
--- a/Engine/Source/Runtime/Engine/Private/SkeletalRenderGPUSkin.cpp
+++ b/Engine/Source/Runtime/Engine/Private/SkeletalRenderGPUSkin.cpp
@@ -507,7 +507,7 @@ void FSkeletalMeshObjectGPUSkin::UpdateDynamicData_RenderThread(FGPUSkinCache* G
 	DynamicData = InDynamicData;
 	LastBoneTransformRevisionNumber = RevisionNumber;
 
-	if (IsDeferredSkeletalDynamicDataUpdateEnabled())
+	if (IsDeferredSkeletalDynamicDataUpdateEnabled() && !DynamicData->bForceUpdateDynamicDataImmediately)
 	{
 		bMorphNeedsUpdateDeferred = bMorphNeedsUpdate;
 		bNeedsUpdateDeferred = true;
@@ -687,12 +687,12 @@ void FSkeletalMeshObjectGPUSkin::ProcessUpdatedDynamicData(EGPUSkinCacheEntryMod
 			if (DynamicData->PreviousReferenceToLocal.Num() > 0)
 			{
 				TArray<FMatrix44f>& PreviousReferenceToLocalMatrices = bShouldUseSeparateMatricesForRayTracing ? DynamicData->PreviousReferenceToLocalForRayTracing : DynamicData->PreviousReferenceToLocal;
-				ShaderData.UpdateBoneData(RHICmdList, PreviousReferenceToLocalMatrices, Section.BoneMap, RevisionNumber, true, FeatureLevel, bUseSkinCache);
+				ShaderData.UpdateBoneData(RHICmdList, PreviousReferenceToLocalMatrices, Section.BoneMap, RevisionNumber, true, FeatureLevel, bUseSkinCache, DynamicData->bForceUpdateDynamicDataImmediately);
 			}
 
 			// Create a uniform buffer from the bone transforms.
 			TArray<FMatrix44f>& ReferenceToLocalMatrices = bShouldUseSeparateMatricesForRayTracing ? DynamicData->ReferenceToLocalForRayTracing : DynamicData->ReferenceToLocal;
-			bool bNeedFence = ShaderData.UpdateBoneData(RHICmdList, ReferenceToLocalMatrices, Section.BoneMap, RevisionNumber, false, FeatureLevel, bUseSkinCache);
+			bool bNeedFence = ShaderData.UpdateBoneData(RHICmdList, ReferenceToLocalMatrices, Section.BoneMap, RevisionNumber, false, FeatureLevel, bUseSkinCache, DynamicData->bForceUpdateDynamicDataImmediately);
 			ShaderData.UpdatedFrameNumber = FrameNumberToPrepare;
 
 			// Update uniform buffer for APEX cloth simulation mesh positions and normals
@@ -701,7 +701,7 @@ void FSkeletalMeshObjectGPUSkin::ProcessUpdatedDynamicData(EGPUSkinCacheEntryMod
 				FGPUBaseSkinAPEXClothVertexFactory::ClothShaderType& ClothShaderData = VertexFactoryData.ClothVertexFactories[SectionIdx]->GetClothShaderData();
 				ClothShaderData.ClothBlendWeight = DynamicData->ClothBlendWeight;
 
-				bNeedFence = ClothShaderData.UpdateClothSimulData(RHICmdList, SimData->Positions, SimData->Normals, RevisionNumber, FeatureLevel) || bNeedFence;
+				bNeedFence = ClothShaderData.UpdateClothSimulData(RHICmdList, SimData->Positions, SimData->Normals, RevisionNumber, FeatureLevel, DynamicData->bForceUpdateDynamicDataImmediately) || bNeedFence;
 				// Transform from cloth space to local space. Cloth space is relative to cloth root bone, local space is component space.
 				ClothShaderData.GetClothToLocalForWriting() = FMatrix44f(SimData->ComponentRelativeTransform.ToMatrixWithScale());
 			}
@@ -2313,6 +2313,7 @@ void FDynamicSkelMeshObjectDataGPUSkin::Clear()
 	ClothBlendWeight = 0.0f;
 	bIsSkinCacheAllowed = false;
 	bHasMeshDeformer = false;
+	bForceUpdateDynamicDataImmediately = false;
 #if RHI_RAYTRACING
 	bAnySegmentUsesWorldPositionOffset = false;
 #endif
@@ -2537,6 +2538,7 @@ void FDynamicSkelMeshObjectDataGPUSkin::InitDynamicSkelMeshObjectDataGPUSkin(
 
 	bIsSkinCacheAllowed = InMeshComponent ? InMeshComponent->IsSkinCacheAllowed(InLODIndex) : false;
 	bHasMeshDeformer = InMeshComponent ? InMeshComponent->GetMeshDeformerInstance() != nullptr : false;
+	bForceUpdateDynamicDataImmediately = InMeshComponent ? InMeshComponent->GetForceUpdateDynamicDataImmediately() : false;
 
 	if (bIsSkinCacheAllowed && InMeshObject->FeatureLevel == ERHIFeatureLevel::ES3_1)
 	{
diff --git a/Engine/Source/Runtime/Engine/Private/SkeletalRenderGPUSkin.h b/Engine/Source/Runtime/Engine/Private/SkeletalRenderGPUSkin.h
--- a/Engine/Source/Runtime/Engine/Private/SkeletalRenderGPUSkin.h
+++ b/Engine/Source/Runtime/Engine/Private/SkeletalRenderGPUSkin.h
@@ -144,6 +144,11 @@ public:
 	// Whether animation is done with a mesh deformer.
 	uint8 bHasMeshDeformer : 1;
 
+	// Whether to update dynamic bone & cloth sim data immediately, not to wait until GDME or defer update to RHIThread.
+	// When set to true, it is the equivalent of r.DeferSkeletalDynamicDataUpdateUntilGDME=0 and r.RHICmdDeferSkeletalLockAndFillToRHIThread=0.
+	// When set to false, r.DeferSkeletalDynamicDataUpdateUntilGDME and r.RHICmdDeferSkeletalLockAndFillToRHIThread values are respected.
+	uint8 bForceUpdateDynamicDataImmediately : 1;
+
 #if RHI_RAYTRACING
 	uint8 bAnySegmentUsesWorldPositionOffset : 1;
 #endif
diff --git a/Engine/Source/Runtime/Engine/Public/GPUSkinVertexFactory.h b/Engine/Source/Runtime/Engine/Public/GPUSkinVertexFactory.h
--- a/Engine/Source/Runtime/Engine/Public/GPUSkinVertexFactory.h
+++ b/Engine/Source/Runtime/Engine/Public/GPUSkinVertexFactory.h
@@ -214,7 +214,7 @@ public:
 
 		// @param FrameTime from GFrameTime
 		bool UpdateBoneData(FRHICommandListImmediate& RHICmdList, const TArray<FMatrix44f>& ReferenceToLocalMatrices,
-			const TArray<FBoneIndexType>& BoneMap, uint32 RevisionNumber, bool bPrevious, ERHIFeatureLevel::Type FeatureLevel, bool bUseSkinCache);
+			const TArray<FBoneIndexType>& BoneMap, uint32 RevisionNumber, bool bPrevious, ERHIFeatureLevel::Type FeatureLevel, bool bUseSkinCache, bool bForceUpdateImmediately);
 
 		void ReleaseBoneData()
 		{
@@ -611,7 +611,8 @@ public:
 			Reset();
 		}
 
-		bool UpdateClothSimulData(FRHICommandListImmediate& RHICmdList, const TArray<FVector3f>& InSimulPositions, const TArray<FVector3f>& InSimulNormals, uint32 RevisionNumber, ERHIFeatureLevel::Type FeatureLevel);
+		bool UpdateClothSimulData(FRHICommandListImmediate& RHICmdList, const TArray<FVector3f>& InSimulPositions, const TArray<FVector3f>& InSimulNormals, uint32 RevisionNumber, 
+									ERHIFeatureLevel::Type FeatureLevel, bool bForceUpdateImmediately);
 
 		void ReleaseClothSimulData()
 		{
