Iris
- Register rpc functions for subobjects with COND_Never to be able to call them even if the subobject does not replicate.
- Fixed ReplicationstateDescriptorBuilder accessing null property

diff --git a/Engine/Source/Runtime/Engine/Private/ActorReplication.cpp b/Engine/Source/Runtime/Engine/Private/ActorReplication.cpp
--- a/Engine/Source/Runtime/Engine/Private/ActorReplication.cpp
+++ b/Engine/Source/Runtime/Engine/Private/ActorReplication.cpp
@@ -882,7 +882,7 @@ void AActor::RegisterReplicationFragments(UE::Net::FFragmentRegistrationContext&
 
 		if (ActorComp->HasAnyFlags(RF_DefaultSubObject) || ActorComp->IsDefaultSubobject())
 		{
-			if (!ActorComp->GetIsReplicated())
+			if (!ActorComp->GetIsReplicated() || ActorComp->GetReplicationCondition() == COND_Never)
 			{
 				// Register RPC functions for not replicated default subobjects
 				ActorComp->RegisterReplicationFragments(Context, EFragmentRegistrationFlags::RegisterRPCsOnly);
