Added character movement option to send base-relative moves when walking on dynamic bases, rather than worldspace.  This greatly improves synchronization while moving on the base, at the risk of worse corrections when jumping on/off.  Also useful if on a large platform (like a boat) that should be the player's visual frame of reference for movement.

diff --git a/Engine/Source/Runtime/Engine/Classes/GameFramework/Character.h b/Engine/Source/Runtime/Engine/Classes/GameFramework/Character.h
--- a/Engine/Source/Runtime/Engine/Classes/GameFramework/Character.h
+++ b/Engine/Source/Runtime/Engine/Classes/GameFramework/Character.h
@@ -154,6 +154,12 @@ namespace MovementBaseUtility
 
 	/** Get the local space location for a given MovementBase, optionally at the location of a bone. Returns false if MovementBase is nullptr, or if BoneName is not a valid bone. Scaling is ignored. */
 	ENGINE_API bool GetLocalMovementBaseLocation(const UPrimitiveComponent* MovementBase, const FName BoneName, const FVector& WorldSpaceLocation, FVector& OutLocalLocation);
+
+	/** Get the world space acceleration from a local space acceleration for a given MovementBase. Returns false if MovementBase is nullptr, or if BoneName is not a valid bone. Scaling is ignored. */
+	ENGINE_API bool GetLocalMovementBaseAccelerationInWorldSpace(const UPrimitiveComponent* MovementBase, const FName BoneName, const FVector& LocalAccel, FVector& OutAccelWorldSpace);
+
+	/** Get the local space acceleration for a given MovementBase, optionally relative to the orientation of a bone. Returns false if MovementBase is nullptr, or if BoneName is not a valid bone. Scaling is ignored. */
+	ENGINE_API bool GetLocalMovementBaseAcceleration(const UPrimitiveComponent* MovementBase, const FName BoneName, const FVector& WorldSpaceAcceleration, FVector& OutLocalAcceleration);
 }
 
 /** Struct to hold information about the "base" object the character is standing on. */
diff --git a/Engine/Source/Runtime/Engine/Classes/GameFramework/CharacterMovementComponent.h b/Engine/Source/Runtime/Engine/Classes/GameFramework/CharacterMovementComponent.h
--- a/Engine/Source/Runtime/Engine/Classes/GameFramework/CharacterMovementComponent.h
+++ b/Engine/Source/Runtime/Engine/Classes/GameFramework/CharacterMovementComponent.h
@@ -750,6 +750,15 @@ public:
 	UPROPERTY(Category="Character Movement (General Settings)", EditAnywhere, BlueprintReadWrite, AdvancedDisplay, meta=(ClampMin="0", UIMin="0", ForceUnits=cm))
 	float MaxDepenetrationWithPawnAsProxy;
 
+	/**
+	 * Option to handle moves as relative to the current movement base. If disabled, moves are handled in world space. 
+	 * Enabling this can be produce better motion sync when walking on highly dynamic / unpredictable movement bases, such as vehicles. However, there is the potential for larger corrections when landing on or jumping off if the base is not well-sync'd.
+	 * This also may be a good option if using dynamic movement bases that become the player's primary frame of visual reference, such as when walking on a large boat or airship.
+	 * Not compatible with deprecated move RPCs. @see NetUsePackedMovementRPCs
+	 */
+	UPROPERTY(Category = "Character Movement (General Settings)", EditDefaultsOnly)
+	bool bUseMovementBaseRelativeMoves = false;
+
 	/**
 	 * How long to take to smoothly interpolate from the old pawn position on the client to the corrected one sent by the server. Not used by Linear smoothing.
 	 */
@@ -2836,6 +2845,7 @@ public:
 	FRotator SavedRotation;
 	FVector SavedVelocity;
 	FVector SavedRelativeLocation;
+	FVector SavedRelativeAcceleration;
 	FRotator SavedControlRotation;
 	TWeakObjectPtr<UPrimitiveComponent> EndBase;
 	FName EndBoneName;
diff --git a/Engine/Source/Runtime/Engine/Private/Character.cpp b/Engine/Source/Runtime/Engine/Private/Character.cpp
--- a/Engine/Source/Runtime/Engine/Private/Character.cpp
+++ b/Engine/Source/Runtime/Engine/Private/Character.cpp
@@ -709,6 +709,25 @@ namespace MovementBaseUtility
 		OutLocalLocation = FTransform(OutQuat, OutLocation).InverseTransformPositionNoScale(WorldSpaceLocation);
 		return bResult;
 	}
+
+	bool GetLocalMovementBaseAccelerationInWorldSpace(const UPrimitiveComponent* MovementBase, const FName BoneName, const FVector& LocalAccel, FVector& OutAccelWorldSpace)
+	{
+		FVector IgnoredLocation;
+		FQuat OutQuat;
+		const bool bResult = GetMovementBaseTransform(MovementBase, BoneName, IgnoredLocation, OutQuat);
+		OutAccelWorldSpace = OutQuat.RotateVector(LocalAccel);
+		return bResult;
+	}
+
+	bool GetLocalMovementBaseAcceleration(const UPrimitiveComponent* MovementBase, const FName BoneName, const FVector& WorldSpaceAcceleration, FVector& OutLocalAcceleration)
+	{
+		FVector IgnoredLocation;
+		FQuat OutQuat;
+		const bool bResult = GetMovementBaseTransform(MovementBase, BoneName, IgnoredLocation, OutQuat);
+		OutLocalAcceleration = OutQuat.Inverse().RotateVector(WorldSpaceAcceleration);
+		return bResult;
+	}
+
 }
 
 
diff --git a/Engine/Source/Runtime/Engine/Private/Components/CharacterMovementComponent.cpp b/Engine/Source/Runtime/Engine/Private/Components/CharacterMovementComponent.cpp
--- a/Engine/Source/Runtime/Engine/Private/Components/CharacterMovementComponent.cpp
+++ b/Engine/Source/Runtime/Engine/Private/Components/CharacterMovementComponent.cpp
@@ -250,6 +250,12 @@ namespace CharacterMovementCVars
 		TEXT("p.BasedMovementMode"),
 		BasedMovementMode, TEXT("0 means always on regular tick (default); 1 means only if not deferring updates; 2 means update and save based movement both on regular ticks and post physics when on a physics base."));
 
+	static int32 AllowBaseRelativeMovement = 1;
+	FAutoConsoleVariableRef CVarAllowBaseRelativeMovement(
+		TEXT("p.AllowBaseRelativeMovement"),
+		AllowBaseRelativeMovement, 
+		TEXT("0=disabled, otherwise enabled. If enabled, resepct the bUseMovementBaseRelativeMoves option on each character. If disabled, all movement will be handled in worldspace."));
+
 	static int32 UseTargetVelocityOnImpact = 1;
 	FAutoConsoleVariableRef CVarUseTargetVelocityOnImpact(
 		TEXT("p.UseTargetVelocityOnImpact"),
@@ -9003,7 +9009,6 @@ void FCharacterNetworkMoveData::ClientFillNetworkMoveData(const FSavedMove_Chara
 	NetworkMoveType = MoveType;
 
 	TimeStamp = ClientMove.TimeStamp;
-	Acceleration = ClientMove.Acceleration;
 	ControlRotation = ClientMove.SavedControlRotation;
 	CompressedMoveFlags = ClientMove.GetCompressedFlags();
 	MovementMode = ClientMove.EndPackedMovementMode;
@@ -9013,16 +9018,33 @@ void FCharacterNetworkMoveData::ClientFillNetworkMoveData(const FSavedMove_Chara
 	{
 		// Determine if we send absolute or relative location
 		UPrimitiveComponent* ClientMovementBase = ClientMove.EndBase.Get();
-		const bool bDynamicBase = MovementBaseUtility::UseRelativeLocation(ClientMovementBase);
-		const FVector SendLocation = bDynamicBase ? ClientMove.SavedRelativeLocation : FRepMovement::RebaseOntoZeroOrigin(ClientMove.SavedLocation, ClientMove.CharacterOwner->GetCharacterMovement());
+
+		const bool bSendBaseRelativeLocation     = MovementBaseUtility::UseRelativeLocation(ClientMovementBase);
+		const bool bSendBaseRelativeAcceleration = CharacterMovementCVars::AllowBaseRelativeMovement 
+													&& bSendBaseRelativeLocation 
+													&& ClientMove.CharacterOwner->GetCharacterMovement()->bUseMovementBaseRelativeMoves;
+
+		const FVector SendLocation     = bSendBaseRelativeLocation ? ClientMove.SavedRelativeLocation : FRepMovement::RebaseOntoZeroOrigin(ClientMove.SavedLocation, ClientMove.CharacterOwner->GetCharacterMovement());
+		const FVector SendAcceleration = bSendBaseRelativeAcceleration ? ClientMove.SavedRelativeAcceleration : ClientMove.Acceleration;
 
 		Location = SendLocation;
-		MovementBase = bDynamicBase ? ClientMovementBase : nullptr;
-		MovementBaseBoneName = bDynamicBase ? ClientMove.EndBoneName : NAME_None;
+		Acceleration = SendAcceleration;
+
+		if (bSendBaseRelativeLocation || bSendBaseRelativeAcceleration)
+		{
+			MovementBase = ClientMovementBase;
+			MovementBaseBoneName = ClientMove.EndBoneName;
+		}
+		else
+		{
+			MovementBase = nullptr;
+			MovementBaseBoneName = NAME_None;
+		}
 	}
 	else
 	{
 		Location = ClientMove.SavedLocation;
+		Acceleration = ClientMove.Acceleration;
 		MovementBase = nullptr;
 		MovementBaseBoneName = NAME_None;
 	}
@@ -9157,7 +9179,17 @@ void UCharacterMovementComponent::ServerMove_PerformMovement(const FCharacterNet
 	}	
 
 	const float ClientTimeStamp = MoveData.TimeStamp;
-	FVector_NetQuantize10 ClientAccel = MoveData.Acceleration;
+	
+	FVector ClientAccel = MoveData.Acceleration;
+
+	// Convert the move's acceleration to worldspace if necessary
+	if (CharacterMovementCVars::AllowBaseRelativeMovement 
+			&& bUseMovementBaseRelativeMoves 
+			&& MovementBaseUtility::IsDynamicBase(MoveData.MovementBase))
+	{
+		MovementBaseUtility::GetLocalMovementBaseAccelerationInWorldSpace(MoveData.MovementBase, MoveData.MovementBaseBoneName, MoveData.Acceleration, ClientAccel);
+	}
+
 	const uint8 ClientMoveFlags = MoveData.CompressedMoveFlags;
 	const FRotator ClientControlRotation = MoveData.ControlRotation;
 
@@ -11634,6 +11666,7 @@ void FSavedMove_Character::Clear()
 	SavedRotation = FRotator::ZeroRotator;
 	SavedRelativeLocation = FVector::ZeroVector;
 	SavedControlRotation = FRotator::ZeroRotator;
+	SavedRelativeAcceleration = FVector::ZeroVector;
 	Acceleration = FVector::ZeroVector;
 	MaxSpeed = 0.0f;
 	AccelMag = 0.0f;
@@ -11773,6 +11806,7 @@ void FSavedMove_Character::PostUpdate(ACharacter* Character, FSavedMove_Characte
 			}
 		}
 #endif
+		// Movement base and base-relative movement
 		UPrimitiveComponent* const MovementBase = Character->GetMovementBase();
 		EndBase = MovementBase;
 		EndBoneName = Character->GetBasedMovement().BoneName;
@@ -11781,6 +11815,14 @@ void FSavedMove_Character::PostUpdate(ACharacter* Character, FSavedMove_Characte
 			SavedRelativeLocation = Character->GetBasedMovement().Location;
 		}
 
+		// Save off movement base-relative acceleration if needed
+		if (CharacterMovementCVars::AllowBaseRelativeMovement 
+				&& Character->GetCharacterMovement()->bUseMovementBaseRelativeMoves
+				&& MovementBaseUtility::IsDynamicBase(MovementBase))
+		{
+			MovementBaseUtility::GetLocalMovementBaseAcceleration(EndBase.Get(), EndBoneName, Acceleration, SavedRelativeAcceleration);
+		}
+
 		// Attachment state
 		if (const USceneComponent* UpdatedComponent = Character->GetCharacterMovement()->UpdatedComponent)
 		{
@@ -11800,9 +11842,9 @@ void FSavedMove_Character::PostUpdate(ACharacter* Character, FSavedMove_Characte
 		}
 	}
 
-	// Only save RootMotion params when initially recording
 	if (PostUpdateMode == PostUpdate_Record)
 	{
+		// Only save RootMotion params when initially recording
 		const FAnimMontageInstance* RootMotionMontageInstance = Character->GetRootMotionAnimMontageInstance();
 		if (RootMotionMontageInstance)
 		{
