Deprecated VisualizeMotionVector and instead expose it as VisualizeReprojection

diff --git a/Engine/Source/Runtime/Engine/Public/ShowFlags.h b/Engine/Source/Runtime/Engine/Public/ShowFlags.h
--- a/Engine/Source/Runtime/Engine/Public/ShowFlags.h
+++ b/Engine/Source/Runtime/Engine/Public/ShowFlags.h
@@ -388,6 +388,7 @@ private:
 		SetPhysicalMaterialMasks(false);
 		SetVisualizeMotionBlur(false);
 		SetVisualizeMotionVectors(false);
+		SetVisualizeReprojection(false);
 		SetEditingLevelInstance(false);
 		SetSelectionOutline(false);
 		SetDebugAI(false);
diff --git a/Engine/Source/Runtime/Engine/Public/ShowFlagsValues.inl b/Engine/Source/Runtime/Engine/Public/ShowFlagsValues.inl
--- a/Engine/Source/Runtime/Engine/Public/ShowFlagsValues.inl
+++ b/Engine/Source/Runtime/Engine/Public/ShowFlagsValues.inl
@@ -116,8 +116,10 @@ SHOWFLAG_ALWAYS_ACCESSIBLE(SeparateTranslucency, SFG_Advanced, NSLOCTEXT("Unreal
 SHOWFLAG_ALWAYS_ACCESSIBLE(ScreenPercentage, SFG_Hidden, NSLOCTEXT("UnrealEd", "ScreenPercentageSF", "Screen Percentage"))
 /** Helper to tweak motion blur settings */
 SHOWFLAG_FIXED_IN_SHIPPING(0, VisualizeMotionBlur, SFG_Visualize, NSLOCTEXT("UnrealEd", "VisualizeMotionBlurSF", "Motion Blur"))
-/** Helper to tweak motion vectors */
-SHOWFLAG_FIXED_IN_SHIPPING(0, VisualizeMotionVectors, SFG_Visualize, NSLOCTEXT("UnrealEd", "VisualizeMotionVectorsSF", "Motion Vectors"))
+/** DEPRECATED in 5.2: Use VisualizeReprojection instead. */
+SHOWFLAG_FIXED_IN_SHIPPING(0, VisualizeMotionVectors, SFG_Hidden, NSLOCTEXT("UnrealEd", "VisualizeMotionVectorsSF", "Motion Vectors"))
+/** Helper to ensure previous frame's temporal history reprojections using motion vectors are working correctly. */
+SHOWFLAG_FIXED_IN_SHIPPING(0, VisualizeReprojection, SFG_Visualize, NSLOCTEXT("UnrealEd", "VisualizeReprojectionSF", "Previous frame's reprojection"))
 /** Whether to display the Reflection Environment feature, which has local reflections from Reflection Capture actors, for now SHOWFLAG_ALWAYS_ACCESSIBLE because it's exposed in SceneCapture */
 SHOWFLAG_ALWAYS_ACCESSIBLE(ReflectionEnvironment, SFG_LightingFeatures, NSLOCTEXT("UnrealEd", "ReflectionEnvironmentSF", "Reflection Environment"))
 /** Visualize pixels that are outside of their object's bounding box (content error). */
diff --git a/Engine/Source/Runtime/Renderer/Private/PostProcess/PostProcessing.cpp b/Engine/Source/Runtime/Renderer/Private/PostProcess/PostProcessing.cpp
--- a/Engine/Source/Runtime/Renderer/Private/PostProcess/PostProcessing.cpp
+++ b/Engine/Source/Runtime/Renderer/Private/PostProcess/PostProcessing.cpp
@@ -379,7 +379,7 @@ void AddPostProcessingPasses(
 	PassSequence.SetEnabled(EPass::VisualizeGBufferOverview, bVisualizeGBufferOverview || bVisualizeGBufferDumpToFile || bVisualizeGBufferDumpToPIpe);
 	PassSequence.SetEnabled(EPass::VisualizeLumenSceneOverview, LumenVisualizeMode == VISUALIZE_MODE_OVERVIEW && bPostProcessingEnabled);
 	PassSequence.SetEnabled(EPass::VisualizeHDR, EngineShowFlags.VisualizeHDR);
-	PassSequence.SetEnabled(EPass::VisualizeMotionVectors, EngineShowFlags.VisualizeMotionVectors);
+	PassSequence.SetEnabled(EPass::VisualizeMotionVectors, EngineShowFlags.VisualizeMotionVectors || EngineShowFlags.VisualizeReprojection);
 #if WITH_EDITOR
 	PassSequence.SetEnabled(EPass::PixelInspector, View.bUsePixelInspector);
 #else
