[Engine]
* AnimNodeMessages - Fixed compile errors. FMessageEntry has a default ctor which tries to create a IGraphMessage but type is pure virtual. Changing to TSharedPtr (realized afterwards that I could have added a ctor with message and info to get around the problem too)

diff --git a/Engine/Source/Runtime/Engine/Private/Animation/AnimNodeMessages.cpp b/Engine/Source/Runtime/Engine/Private/Animation/AnimNodeMessages.cpp
--- a/Engine/Source/Runtime/Engine/Private/Animation/AnimNodeMessages.cpp
+++ b/Engine/Source/Runtime/Engine/Private/Animation/AnimNodeMessages.cpp
@@ -91,7 +91,7 @@ void FMessageStack::ForEachMessage(FName InMessageType, TFunctionRef<EEnumerate(
 			const FMessageEntry& MessageEntry = StackEntry[MessageIndex];
 			if(MessageEntry.Message->GetTypeName() == InMessageType)
 			{
-				if(InFunction(MessageEntry.Message.Get()) == EEnumerate::Stop)
+				if(InFunction(*MessageEntry.Message) == EEnumerate::Stop)
 				{
 					return;
 				}
@@ -110,7 +110,7 @@ void FMessageStack::TopMessage(FName InMessageType, TFunctionRef<void(IGraphMess
 			const FMessageEntry& MessageEntry = StackEntry.Top();
 			if(MessageEntry.Message->GetTypeName() == InMessageType)
 			{
-				InFunction(MessageEntry.Message.Get());
+				InFunction(*MessageEntry.Message);
 			}
 		}
 	}
@@ -126,7 +126,7 @@ void FMessageStack::BottomMessage(FName InMessageType, TFunctionRef<void(IGraphM
 			const FMessageEntry& MessageEntry = StackEntry[0];
 			if(MessageEntry.Message->GetTypeName() == InMessageType)
 			{
-				InFunction(MessageEntry.Message.Get());
+				InFunction(*MessageEntry.Message);
 			}
 		}
 	}
diff --git a/Engine/Source/Runtime/Engine/Public/Animation/AnimNodeMessages.h b/Engine/Source/Runtime/Engine/Public/Animation/AnimNodeMessages.h
--- a/Engine/Source/Runtime/Engine/Public/Animation/AnimNodeMessages.h
+++ b/Engine/Source/Runtime/Engine/Public/Animation/AnimNodeMessages.h
@@ -316,7 +316,7 @@ private:
 	// referenced by multiple cached poses (pesky DAG)
 	struct FMessageEntry
 	{
-		TSharedRef<IGraphMessage> Message;
+		TSharedPtr<IGraphMessage> Message;
 		FNodeInfo NodeInfo;
 	};
 
