texture fix GetDefaultTextureFormatName logic
CompressionSettings TC should be the primary driver of output format
if that is already an uncompressed format, don't change it
only do the non-multiple-of-4 force to uncompressed when needed

diff --git a/Engine/Source/Runtime/Engine/Private/Texture.cpp b/Engine/Source/Runtime/Engine/Private/Texture.cpp
--- a/Engine/Source/Runtime/Engine/Private/Texture.cpp
+++ b/Engine/Source/Runtime/Engine/Private/Texture.cpp
@@ -3191,10 +3191,7 @@ FName GetDefaultTextureFormatName( const ITargetPlatform* TargetPlatform, const
 
 	const ETextureSourceFormat SourceFormat = Texture->Source.GetFormat(LayerIndex);
 
-	// Determine the pixel format of the (un/)compressed texture
-	
-	// Identify TC groups that will map to uncompressed :		
-	bool bIsTCThatMapsToUncompressed = UE::TextureDefines::IsUncompressed(FormatSettings.CompressionSettings);
+	// output format is primarily determined from the CompressionSettings (TC)
 
 	// see if compression needs to be forced off even if requested :
 	bool bNoCompression = FormatSettings.CompressionNone				// Code wants the texture uncompressed.
@@ -3206,11 +3203,6 @@ FName GetDefaultTextureFormatName( const ITargetPlatform* TargetPlatform, const
 		|| ((Texture->GetMaterialType() == MCT_VolumeTexture) && !bSupportCompressedVolumeTexture)
 		|| FormatSettings.CompressionSettings == TC_EncodedReflectionCapture;
 
-	// note that bNoCompression is not the same as bIsTCThatMapsToUncompressed
-	// bIsTCThatMapsToUncompressed is a TC_ that will map to an uncompressed format
-	//	 but does not set bNoCompression
-	// bNoCompression does different mappings than bIsTCThatMapsToUncompressed
-
 	if (!bNoCompression)
 	{
 		int32 SizeX = Texture->Source.GetSizeX();
@@ -3261,30 +3253,11 @@ FName GetDefaultTextureFormatName( const ITargetPlatform* TargetPlatform, const
 		// note that even if top mip is a multiple of 4, lower may not be
 		// we can only choose compression if it's supported by all platforms/RHI's (else check TargetPlatform->SupportsFeature)
 		// note: does not use the passed-in "BlockSize" parameter, hard coded to 4
+		//	 that is correct because ASTC does not require block alignment, only DXTC does which is always a 4-size block
 		if ( (SizeX < 4) || (SizeY < 4) || (SizeX % 4 != 0) || (SizeY % 4 != 0) )
 		{
-			// don't log if TC was going to map to uncompressed anyway
-			if ( ! bIsTCThatMapsToUncompressed )
-			{
-				// this means TC_ or TEXTUREGROUP_ was set wrong
-				// usually it means UI textures that were just left with TC_Default
-				// they should be TEXTUREGROUP_ColorLookupTable or TC_EditorIcon "UserInterface2D"
-
-				// NOTE: this function is called like 10 times for each texture
-				//	it makes any log here a bit annoying
-
-				UE_LOG(LogTexture, Verbose, TEXT("Texture forced to uncompressed because size is not a multiple of 4 : %dx%d: was %s : %s"), SizeX,SizeY, 
-					*(StaticEnum<TextureCompressionSettings>()->GetNameStringByValue(FormatSettings.CompressionSettings)),
-					*Texture->GetPathName());
-			}
-
-			// note if bIsTCThatMapsToUncompressed , this might change the mapping :
-			// @todo Oodle : should not do this if bIsTCThatMapsToUncompressed already
 			bNoCompression = true;
 		}
-
-		// @todo Oodle : the conditions that trigger bNoCompression should be detected earlier
-		//   and shown in the GUI so the artist can see it
 	}
 
 	bool bUseDXT5NormalMap = false;
@@ -3298,75 +3271,7 @@ FName GetDefaultTextureFormatName( const ITargetPlatform* TargetPlatform, const
 
 	// Determine the pixel format of the (un/)compressed texture
 
-	// bNoCompression overrides TC format mapping
-	// even if bIsTCThatMapsToUncompressed would map us to uncompressed anyway
-	if (bNoCompression)
-	{
-		// TC_EditorIcon & TC_EncodedReflectionCapture set bNoCompression
-
-		// NOTE: strange change in behavior from bNoCompression vs bIsTCThatMapsToUncompressed
-		// does not look for TC_HDR , goes to HasHDRSource
-		// others go straight to SourceFormat, ignore TC_
-		// eg. if you were in bIsTCThatMapsToUncompressed , and for some reason you get bNoCompression set
-		//	 then your output format can change from one uncompressed format to another
-
-		// someday: if this was redone, I would make two changes
-		// 1. first of all detect if you already have a TC that maps to uncompressed with a specific format,
-		//		and if so, stick to that
-		//		eg. currently if you choose TC_HDR and your source image is G8 you will get RGBA16F output
-		//		but if the source image is not a multiple of 4, bNoCompression turns on, and you get G8 output
-		//		this can be fixed by putting the bNoCompression check after the standard TC maps
-		// 2. second if you do the bNoCompression, fix TSF formats not mapping to the closest possible
-		//		uncompressed renderable format
-
-		if (Texture->HasHDRSource(LayerIndex))
-		{
-			// if CompressionSettings was already an uncompressed HDR format, use it
-			//  this is a conservative/partial fix of the more general problem of bIsTCThatMapsToUncompressed (see 1. above)
-			if (FormatSettings.CompressionSettings == TC_HDR)
-			{
-				TextureFormatName = NameRGBA16F;
-			}
-			else if (FormatSettings.CompressionSettings == TC_HDR_F32)
-			{
-				TextureFormatName = NameRGBA32F;
-			}
-			else if (FormatSettings.CompressionSettings == TC_HalfFloat)
-			{
-				TextureFormatName = NameR16F;
-			}
-			else if (FormatSettings.CompressionSettings == TC_SingleFloat)
-			{
-				TextureFormatName = NameR32F;
-			}
-			else
-			{
-				// @todo Oodle : this is not the best possible selection of output format from SourceFormat
-			// R16F and R32F is available but not used here even if their TC_ would have chosen them!
-			TextureFormatName = NameRGBA16F;
-		}
-		}
-		else if (SourceFormat == TSF_G16)
-		{
-			// maps G16 but not RGBA16, it will go to BGRA8
-			TextureFormatName = NameG16;
-		}
-		else if (SourceFormat == TSF_G8 || FormatSettings.CompressionSettings == TC_Grayscale)
-		{
-			TextureFormatName = NameG8;
-		}
-		else if (FormatSettings.CompressionSettings == TC_Normalmap && bUseDXT5NormalMap)
-		{
-			// move R to A like we do for DXT5 normal maps :
-			TextureFormatName = NameXGXR8;
-		}
-		else
-		{
-			// note CompressionNoAlpha no longer kills alpha if it's forced to uncompressed (eg. because size is not multiple of 4)
-			TextureFormatName = NameBGRA8;
-		}
-	}
-	else if (FormatSettings.CompressionSettings == TC_LQ)
+	if (FormatSettings.CompressionSettings == TC_LQ)
 	{
 		bool bLQCompressionSupported = TargetPlatform->SupportsLQCompressionTextureFormat();
 		if(bLQCompressionSupported)
@@ -3397,8 +3302,8 @@ FName GetDefaultTextureFormatName( const ITargetPlatform* TargetPlatform, const
 	else if (FormatSettings.CompressionSettings == TC_Grayscale || 
 		FormatSettings.CompressionSettings == TC_Displacementmap)
 	{
-		// @todo Oodle : this is not the best possible selection of output format from SourceFormat
-		// eg. doesn't use G16 if source is RGBA16 or float
+		// Grayscale is G8 output, unless source is specifically G16
+		//	(eg. RGBA16 source still uses G8 output, not G16)
 		if (SourceFormat == TSF_G16)
 		{
 			TextureFormatName = NameG16;
@@ -3407,6 +3312,19 @@ FName GetDefaultTextureFormatName( const ITargetPlatform* TargetPlatform, const
 		{
 			TextureFormatName = NameG8;
 		}
+
+		/*
+		// @todo Oodle: consider alternatively, use G16 for all 16-bit and floating point sources
+		if (SourceFormat == TSF_G8 || SourceFormat == TSF_BGRA8)
+		{
+			TextureFormatName = NameG8;
+		}
+		else
+		{
+			// 16 bit or float sources
+			TextureFormatName = NameG16;
+		}
+		*/
 	}
 	else if ( FormatSettings.CompressionSettings == TC_Alpha)
 	{
@@ -3432,11 +3350,9 @@ FName GetDefaultTextureFormatName( const ITargetPlatform* TargetPlatform, const
 	{
 		TextureFormatName = NameR32F;
 	}
-	else
+	else if ( FormatSettings.CompressionSettings == TC_Default ||
+			FormatSettings.CompressionSettings == TC_Masks )
 	{
-		check( FormatSettings.CompressionSettings == TC_Default ||
-			FormatSettings.CompressionSettings == TC_Masks ); 
-
 		if (FormatSettings.CompressionNoAlpha)
 		{
 			// CompressionNoAlpha changes AutoDXT to DXT1 early
@@ -3452,6 +3368,60 @@ FName GetDefaultTextureFormatName( const ITargetPlatform* TargetPlatform, const
 			TextureFormatName = NameAutoDXT;
 		}
 	}
+	else
+	{
+		// un-handled CompressionSettings cases will have TextureFormatName == none and go into the bNoCompression branch below
+		// alternatively, should TC_EditorIcon be an explicit branch rather than relying on bNoCompression?
+		check( TextureFormatName == NAME_None );
+	}
+
+	bool bTextureFormatNameIsCompressed =
+		(TextureFormatName == NameDXT1) ||
+		(TextureFormatName == NameAutoDXT) ||
+		(TextureFormatName == NameDXT5) ||
+		(TextureFormatName == NameDXT5n) ||
+		(TextureFormatName == NameBC4) ||
+		(TextureFormatName == NameBC5) ||
+		(TextureFormatName == NameBC6H) ||
+		(TextureFormatName == NameBC7);
+	
+	// if !bTextureFormatNameIsCompressed , we already picked an uncompressed format from TC, leave it alone
+	if ( (bNoCompression && bTextureFormatNameIsCompressed) || TextureFormatName == NAME_None )
+	{
+		// TC_EditorIcon & TC_EncodedReflectionCapture weren't handled in the CompressionSettings branches above
+		//	so will have FormatName == None and come in here
+		
+		if (FormatSettings.CompressionSettings == TC_Normalmap && bUseDXT5NormalMap)
+		{
+			// move R to A like we do for DXT5 normal maps : (NameDXT5n)
+			TextureFormatName = NameXGXR8;
+		}
+		else if (FormatSettings.CompressionSettings == TC_HDR_Compressed )
+		{
+			TextureFormatName = NameRGBA16F;
+		}
+		else if (Texture->HasHDRSource(LayerIndex))
+		{
+			// note that if user actually selected an HDR TC we do not come in here
+			// @todo Oodle : consider removing HasHDRSource ; user did not pick an HDR TC output format
+			TextureFormatName = NameRGBA16F;
+		}
+		else if (SourceFormat == TSF_G16)
+		{
+			TextureFormatName = NameG16;
+		}
+		else if (SourceFormat == TSF_G8)
+		{
+			TextureFormatName = NameG8;
+		}
+		else
+		{
+			// note CompressionNoAlpha no longer kills alpha if it's forced to uncompressed (eg. because size is not multiple of 4)
+			TextureFormatName = NameBGRA8;
+		}
+	}
+
+	// fix up stage :
 
 	// Some PC GPUs don't support sRGB read from G8 textures (e.g. AMD DX10 cards on ShaderModel3.0)
 	// This solution requires 4x more memory but a lot of PC HW emulate the format anyway
@@ -3494,7 +3464,9 @@ FName GetDefaultTextureFormatName( const ITargetPlatform* TargetPlatform, const
 	bOodleTextureSdkVersionIsNone = Texture->OodleTextureSdkVersion.IsNone();
 #endif //WITH_EDITOR
 
-	return ConditionalGetPrefixedFormat(TextureFormatName, TargetPlatform, bOodleTextureSdkVersionIsNone);
+	FName Result = ConditionalGetPrefixedFormat(TextureFormatName, TargetPlatform, bOodleTextureSdkVersionIsNone);
+
+	return Result;
 }
 
 void GetDefaultTextureFormatNamePerLayer(TArray<FName>& OutFormatNames, const class ITargetPlatform* TargetPlatform, const class UTexture* Texture, 
