Update the texture build to route new information out:
1. Whether the source mips detects alpha
2. Several hashes to track the state of data at various stages of the build process.

Move the r.SharedLinearTextureEncoding out of the ConsoleManager as that's The Old Way.
Rename TextureDerivedDataVersion for uniqueness.

diff --git a/Engine/Source/Developer/TextureBuild/Private/TextureBuildFunction.cpp b/Engine/Source/Developer/TextureBuild/Private/TextureBuildFunction.cpp
--- a/Engine/Source/Developer/TextureBuild/Private/TextureBuildFunction.cpp
+++ b/Engine/Source/Developer/TextureBuild/Private/TextureBuildFunction.cpp
@@ -11,6 +11,7 @@
 #include "ImageCoreUtils.h"
 #include "Interfaces/ITextureFormat.h"
 #include "Interfaces/ITextureFormatModule.h"
+#include "IO/IoHash.h"
 #include "Memory/CompositeBuffer.h"
 #include "Memory/SharedBuffer.h"
 #include "Modules/ModuleManager.h"
@@ -28,9 +29,12 @@ DEFINE_LOG_CATEGORY_STATIC(LogTextureBuildFunction, Log, All);
 // Any edits to the texture compressor or this file that will change the output of texture builds
 // MUST have a corresponding change to this version. Individual texture formats have a version to
 // change that is specific to the format. A merge conflict affecting the version MUST be resolved
-// by generating a new version.
-static const FGuid TextureDerivedDataVersion(TEXT("9C3A0F67-3973-4F4D-85FB-35B44FEA8089"));
-
+// by generating a new version. This also includes the addition of new outputs to the build, as
+// this will cause a DDC verification failure unless a new version is created.
+// A reminder that for DDC invalidation, running a ddc fill job or the ddc commandlet is a friendly
+// thing to do! -run=DerivedDataCache -Fill -TargetPlatform=Platform1,Platform...N
+//
+static const FGuid TextureBuildFunctionVersion(TEXT("B20676CE-A786-43EE-96F0-2620A4C38ACA"));
 
 static void ReadCbField(FCbFieldView Field, bool& OutValue) { OutValue = Field.AsBool(OutValue); }
 static void ReadCbField(FCbFieldView Field, int32& OutValue) { OutValue = Field.AsInt32(OutValue); }
@@ -275,7 +279,7 @@ static bool TryReadTextureSourceFromCompactBinary(FCbFieldView Source, UE::Deriv
 FGuid FTextureBuildFunction::GetVersion() const
 {
 	UE::DerivedData::FBuildVersionBuilder Builder;
-	Builder << TextureDerivedDataVersion;
+	Builder << TextureBuildFunctionVersion;
 	ITextureFormat* TextureFormat = nullptr;
 	GetVersion(Builder, TextureFormat);
 	if (TextureFormat)
@@ -362,17 +366,7 @@ void FTextureBuildFunction::Build(UE::DerivedData::FBuildContext& Context) const
 	TArray<FCompressedImage2D> CompressedMips;
 	uint32 NumMipsInTail;
 	uint32 ExtData;
-	bool bImageHasAlpha = false;
-	
-	
-	/*
-	for( const FImage & Image : SourceMips )
-	{
-		UE_LOG(LogInit,Display,TEXT("FTextureBuildFunction: SourceMips : %dx%dx%d"),
-			Image.SizeX,Image.SizeY,Image.NumSlices);
-	}
-	*/
-
+	UE::TextureBuildUtilities::FTextureBuildMetadata BuildMetadata;
 	bool bBuildSucceeded = TextureCompressorModule.BuildTexture(
 		SourceMips,
 		AssociatedNormalSourceMips,
@@ -381,7 +375,8 @@ void FTextureBuildFunction::Build(UE::DerivedData::FBuildContext& Context) const
 		CompressedMips,
 		NumMipsInTail,
 		ExtData,
-		&bImageHasAlpha);
+		&BuildMetadata
+		);
 	if (!bBuildSucceeded)
 	{
 		return;
@@ -394,7 +389,7 @@ void FTextureBuildFunction::Build(UE::DerivedData::FBuildContext& Context) const
 	{
 		int32 CalculatedMip0SizeX = 0, CalculatedMip0SizeY = 0, CalculatedMip0NumSlices = 0;
 		int32 CalculatedMipCount = TextureCompressorModule.GetMipCountForBuildSettings(SourceMips[0].SizeX, SourceMips[0].SizeY, SourceMips[0].NumSlices, SourceMips.Num(), BuildSettings, CalculatedMip0SizeX, CalculatedMip0SizeY, CalculatedMip0NumSlices);
-		BuildSettings.GetEncodedTextureDescription(&TextureDescription, TextureFormat, CalculatedMip0SizeX, CalculatedMip0SizeY, CalculatedMip0NumSlices, CalculatedMipCount, bImageHasAlpha);
+		BuildSettings.GetEncodedTextureDescriptionWithPixelFormat(&TextureDescription, (EPixelFormat)CompressedMips[0].PixelFormat, CalculatedMip0SizeX, CalculatedMip0SizeY, CalculatedMip0NumSlices, CalculatedMipCount);
 	}
 
 	FEncodedTextureExtendedData ExtendedData;
@@ -418,12 +413,9 @@ void FTextureBuildFunction::Build(UE::DerivedData::FBuildContext& Context) const
 
 	// Long term, this will be supplied to the build and this would only be called to verify.
 	int32 NumStreamingMips = TextureDescription.GetNumStreamingMips(&ExtendedData, EngineParameters);
-
-	// \todo save bImageHasAlphaChannel back out so that if textures get edited after a build then saved,
-	// we can capture the update.
-
 	
 	{
+		Context.AddValue(UE::DerivedData::FValueId::FromName(ANSITEXTVIEW("TextureBuildMetadata")), BuildMetadata.ToCompactBinaryWithDefaults());
 		Context.AddValue(UE::DerivedData::FValueId::FromName(ANSITEXTVIEW("EncodedTextureDescription")), UE::TextureBuildUtilities::EncodedTextureDescription::ToCompactBinary(TextureDescription));
 		Context.AddValue(UE::DerivedData::FValueId::FromName(ANSITEXTVIEW("EncodedTextureExtendedData")), UE::TextureBuildUtilities::EncodedTextureExtendedData::ToCompactBinary(ExtendedData));
 
@@ -491,6 +483,10 @@ void GenericTextureTilingBuildFunction(UE::DerivedData::FBuildContext& Context,
 		UE::TextureBuildUtilities::TextureEngineParameters::FromCompactBinary(EngineParameters, EngineParametersCb);
 	}
 
+	UE::TextureBuildUtilities::FTextureBuildMetadata BuildMetadata(FCbObject(Context.FindInput(ANSITEXTVIEW("TextureBuildMetadata"))));
+
+	UE_LOG(LogTextureBuildFunction, Display, TEXT("Tiling %d source mip(s) with a tail of %d..."), TextureDescription.NumMips, TextureExtendedData.NumMipsInTail);
+
 	//
 	// Careful - the linear build might have a different streaming mip count than we output due to mip tail
 	// packing.
@@ -542,8 +538,8 @@ void GenericTextureTilingBuildFunction(UE::DerivedData::FBuildContext& Context,
 		FirstMipTailIndex = TextureDescription.NumMips - MipTailCount;
 	}
 
-
 	// Process the mips
+	FIoHashBuilder PostTileHashBuilder;
 	TArray<FSharedBuffer> MipTailBuffers;
 	for (int32 MipIndex = 0; MipIndex < FirstMipTailIndex + 1; MipIndex++)
 	{
@@ -557,6 +553,8 @@ void GenericTextureTilingBuildFunction(UE::DerivedData::FBuildContext& Context,
 		}
 		FSharedBuffer MipData = Tiler->ProcessMipLevel(TextureDescription, TextureExtendedData, MipsForLevel, MipIndex);
 
+		PostTileHashBuilder.Update(MipData.GetData(), MipData.GetSize());
+
 		// Make sure we got the size we advertised prior to the build. If this ever fires then we
 		// have a critical mismatch!
 		check(TextureExtendedData.MipSizesInBytes[MipIndex] == MipData.GetSize());
@@ -572,6 +570,8 @@ void GenericTextureTilingBuildFunction(UE::DerivedData::FBuildContext& Context,
 		}
 	} // end for each mip
 
+	BuildMetadata.PostTileMipsHash = PostTileHashBuilder.Finalize();
+
 	// The mip tail is a bunch of mips all together in one "Value", so assemble them here.
 	FCompositeBuffer MipTail(MipTailBuffers);
 	if (MipTail.GetSize() > 0)
@@ -581,4 +581,5 @@ void GenericTextureTilingBuildFunction(UE::DerivedData::FBuildContext& Context,
 
 	Context.AddValue(UE::DerivedData::FValueId::FromName(UTF8TEXTVIEW("EncodedTextureDescription")), UE::TextureBuildUtilities::EncodedTextureDescription::ToCompactBinary(TextureDescription));
 	Context.AddValue(UE::DerivedData::FValueId::FromName(UTF8TEXTVIEW("EncodedTextureExtendedData")), UE::TextureBuildUtilities::EncodedTextureExtendedData::ToCompactBinary(TextureExtendedData));
+	Context.AddValue(UE::DerivedData::FValueId::FromName(UTF8TEXTVIEW("TextureBuildMetadata")), BuildMetadata.ToCompactBinaryWithDefaults());
 }
diff --git a/Engine/Source/Developer/TextureBuildUtilities/Private/TextureBuildUtilities.cpp b/Engine/Source/Developer/TextureBuildUtilities/Private/TextureBuildUtilities.cpp
--- a/Engine/Source/Developer/TextureBuildUtilities/Private/TextureBuildUtilities.cpp
+++ b/Engine/Source/Developer/TextureBuildUtilities/Private/TextureBuildUtilities.cpp
@@ -1,6 +1,7 @@
 // Copyright Epic Games, Inc. All Rights Reserved.
 
 #include "TextureBuildUtilities.h"
+#include "TextureCompressorModule.h" // for FTextureBuildSettings
 #include "Misc/DataDrivenPlatformInfoRegistry.h"
 #include "Serialization/CompactBinary.h"
 #include "Serialization/CompactBinaryWriter.h"
@@ -208,6 +209,39 @@ namespace TextureEngineParameters
 	}
 } // namespace EncodedTextureDescription
 
+FCbObject FTextureBuildMetadata::ToCompactBinaryWithDefaults() const
+{
+	FTextureBuildMetadata Defaults;
+
+	FCbWriter Writer;
+	Writer.BeginObject();
+	if (bSourceMipsAlphaDetected != Defaults.bSourceMipsAlphaDetected)
+	{
+		Writer.AddBool("bSourceMipsAlphaDetected", bSourceMipsAlphaDetected);
+	}
+	if (PreEncodeMipsHash != Defaults.PreEncodeMipsHash)
+	{
+		Writer.AddHash("PreEncodeMipsHash", PreEncodeMipsHash);
+	}
+	if (PostEncodeMipsHash != Defaults.PostEncodeMipsHash)
+	{
+		Writer.AddHash("PostEncodeMipsHash", PostEncodeMipsHash);
+	}
+	if (PostTileMipsHash != Defaults.PostTileMipsHash)
+	{
+		Writer.AddHash("PostTileMipsHash", PostTileMipsHash);
+	}
+	Writer.EndObject();
+	return Writer.Save().AsObject();
+}
+
+FTextureBuildMetadata::FTextureBuildMetadata(FCbObject InCbObject)
+{
+	bSourceMipsAlphaDetected = InCbObject["bSourceMipsAlphaDetected"].AsBool(bSourceMipsAlphaDetected);
+	PreEncodeMipsHash = InCbObject["PreEncodeMipsHash"].AsHash(PreEncodeMipsHash);
+	PostEncodeMipsHash = InCbObject["PostEncodeMipsHash"].AsHash(PostEncodeMipsHash);
+	PostTileMipsHash = InCbObject["PostTileMipsHash"].AsHash(PostTileMipsHash);
+}
 
 
 } // namespace
diff --git a/Engine/Source/Developer/TextureBuildUtilities/Public/TextureBuildUtilities.h b/Engine/Source/Developer/TextureBuildUtilities/Public/TextureBuildUtilities.h
--- a/Engine/Source/Developer/TextureBuildUtilities/Public/TextureBuildUtilities.h
+++ b/Engine/Source/Developer/TextureBuildUtilities/Public/TextureBuildUtilities.h
@@ -3,10 +3,12 @@
 #pragma once
 
 #include "CoreMinimal.h"
+#include "IO/IoHash.h"
 #include "ImageCore.h"
-#include "TextureCompressorModule.h" // for FTextureBuildSettings
 #include "Interfaces/ITextureFormat.h"
 
+struct FTextureBuildSettings;
+
 /***
 
 TextureBuildUtilities is a utility module for shared code that Engine and TextureBuildWorker (no Engine) can both see
@@ -40,6 +42,24 @@ namespace TextureEngineParameters
 	TEXTUREBUILDUTILITIES_API bool FromCompactBinary(FTextureEngineParameters& OutEngineParameters, FCbObject InCbObject);
 }
 
+// Carries information out of the build that we don't want to cook or save off in the runtime
+struct TEXTUREBUILDUTILITIES_API FTextureBuildMetadata
+{
+	// Whether or not the first mip level of the source has alpha detected. This is used e.g.
+	// picking between BC1 and BC3
+	bool bSourceMipsAlphaDetected = false;
+	
+	// Digests of the data at various processing stages so we can track down determinism issues
+	// that arise.
+	FIoHash PreEncodeMipsHash;
+	FIoHash PostEncodeMipsHash;
+	FIoHash PostTileMipsHash;
+
+	FCbObject ToCompactBinaryWithDefaults() const;
+	FTextureBuildMetadata(FCbObject InCbObject);
+	FTextureBuildMetadata() = default;
+};
+
 TEXTUREBUILDUTILITIES_API bool TextureFormatIsHdr(FName const& InName);
 
 // Removes platform and other custom prefixes from the name.
diff --git a/Engine/Source/Developer/TextureBuildUtilities/TextureBuildUtilities.Build.cs b/Engine/Source/Developer/TextureBuildUtilities/TextureBuildUtilities.Build.cs
--- a/Engine/Source/Developer/TextureBuildUtilities/TextureBuildUtilities.Build.cs
+++ b/Engine/Source/Developer/TextureBuildUtilities/TextureBuildUtilities.Build.cs
@@ -13,6 +13,7 @@ public class TextureBuildUtilities : ModuleRules
 		{
 			"CoreUObject",
 			"Engine",
+			"TextureCompressor"
 		});
 
 		PrivateDependencyModuleNames.AddRange(
@@ -21,7 +22,6 @@ public class TextureBuildUtilities : ModuleRules
 				"DerivedDataCache",
 				"ImageCore",
 				"ImageWrapper",
-				"TextureCompressor",
 				"TextureFormat",
 			}
 			);
diff --git a/Engine/Source/Developer/TextureCompressor/Private/TextureCompressorModule.cpp b/Engine/Source/Developer/TextureCompressor/Private/TextureCompressorModule.cpp
--- a/Engine/Source/Developer/TextureCompressor/Private/TextureCompressorModule.cpp
+++ b/Engine/Source/Developer/TextureCompressor/Private/TextureCompressorModule.cpp
@@ -10,6 +10,7 @@
 #include "Modules/ModuleManager.h"
 #include "Engine/TextureDefines.h"
 #include "TextureFormatManager.h"
+#include "TextureBuildUtilities.h"
 #include "Interfaces/ITextureFormat.h"
 #include "Misc/Paths.h"
 #include "Tasks/Task.h"
@@ -3462,9 +3463,9 @@ public:
 		TArray<FCompressedImage2D>& OutTextureMips,
 		uint32& OutNumMipsInTail,
 		uint32& OutExtData,
-		bool* bOutImageHasAlpha
+		UE::TextureBuildUtilities::FTextureBuildMetadata* OutMetadata
 	)
-	{		
+	{
 		//TRACE_CPUPROFILER_EVENT_SCOPE(Texture.BuildTexture);
 
 		const ITextureFormat* TextureFormat = nullptr;
@@ -3491,6 +3492,11 @@ public:
 
 		TArray<FImage> IntermediateMipChain;
 
+		bool bSourceMipsAlphaDetected = false;
+		if (OutMetadata)
+		{
+			bSourceMipsAlphaDetected = FImageCore::DetectAlphaChannel(SourceMips[0]);
+		}
 
 		// allow to leave texture in sRGB in case compressor accepts other than non-F32 input source
 		// otherwise linearizing will force format to be RGBA32F
@@ -3546,18 +3552,49 @@ public:
 		}
 
 		
-		// DetectAlphaChannel on the top mip of the generated mip chain
-		//	BuildSettings could have programatically introduced alpha that was not in the source
-		// note the order of operations in bForceAlphaChannel and bForceNoAlphaChannel ( ForceNo takes precedence )
-		const bool bImageHasAlphaChannel = !BuildSettings.bForceNoAlphaChannel  && (BuildSettings.bForceAlphaChannel || FImageCore::DetectAlphaChannel(IntermediateMipChain[0]));
-
-		if (bOutImageHasAlpha)
+		// DetectAlphaChannel on the top mip of the generated mip chain. SoonTM this will use the source mip chain. Testing has
+		// shown this to be 99.9% the same, and allows us to get the pixel format earlier.
+		const bool bImageHasAlphaChannel = BuildSettings.GetTextureExpectsAlphaInPixelFormat(FImageCore::DetectAlphaChannel(IntermediateMipChain[0]));
+		
+		if (OutMetadata)
 		{
-			*bOutImageHasAlpha = bImageHasAlphaChannel;
+			// The metadata is about trying to determine bImageHasAlphaChannel _before_ we launch a task, which 
+			// means it must be on the actual source mips, not the post-processed mips. At some point enough textures
+			// will have this saved as part of the creation process and we can rely on it - but for now we route it
+			// back out so that when textures happen to get saved, it goes with it.
+			OutMetadata->bSourceMipsAlphaDetected = bSourceMipsAlphaDetected;
+
+			// Hash the mips before we compress them. This gets saved as part of the derived data and then added to the
+			// diff tags during cook so we can catch determinism issues.
+			FIoHashBuilder MipHashBuilder;
+			for (const FImage& Mip : IntermediateMipChain)
+			{
+				const FImageInfo* Info = &Mip;
+				// Can't just hash FImageInfo due to struct padding RNG making the hash non-deterministic.
+				MipHashBuilder.Update(&Info->SizeX, sizeof(Info->SizeX));
+				MipHashBuilder.Update(&Info->SizeY, sizeof(Info->SizeY));
+				MipHashBuilder.Update(&Info->NumSlices, sizeof(Info->NumSlices));
+				MipHashBuilder.Update(&Info->Format, sizeof(Info->Format));
+				MipHashBuilder.Update(&Info->GammaSpace, sizeof(Info->GammaSpace));
+				MipHashBuilder.Update(Mip.RawData.GetData(), Mip.RawData.Num());
+			}
+			OutMetadata->PreEncodeMipsHash = MipHashBuilder.Finalize();
 		}
 		
-		return CompressMipChain(TextureFormat, IntermediateMipChain, BuildSettings, bImageHasAlphaChannel, DebugTexturePathName,
+		bool bCompressSucceeded = CompressMipChain(TextureFormat, IntermediateMipChain, BuildSettings, bImageHasAlphaChannel, DebugTexturePathName,
 					OutTextureMips, OutNumMipsInTail, OutExtData);
+
+		if (OutMetadata)
+		{
+			FIoHashBuilder MipHashBuilder;
+			for (const FCompressedImage2D& Mip : OutTextureMips)
+			{
+				MipHashBuilder.Update(Mip.RawData.GetData(), Mip.RawData.Num());
+			}
+			OutMetadata->PostEncodeMipsHash = MipHashBuilder.Finalize();
+		}
+
+		return bCompressSucceeded;
 	}
 
 	// IModuleInterface implementation.
diff --git a/Engine/Source/Developer/TextureCompressor/Public/TextureCompressorModule.h b/Engine/Source/Developer/TextureCompressor/Public/TextureCompressorModule.h
--- a/Engine/Source/Developer/TextureCompressor/Public/TextureCompressorModule.h
+++ b/Engine/Source/Developer/TextureCompressor/Public/TextureCompressorModule.h
@@ -13,6 +13,7 @@ class ITextureFormat;
 class ITextureTiler;
 struct FTextureEngineParameters;
 enum EPixelFormat : uint8;
+namespace UE::TextureBuildUtilities { struct FTextureBuildMetadata; };
 
 /**
  * Compressed image data.
@@ -317,6 +318,23 @@ struct FTextureBuildSettings
 		return GetDestGammaSpace();
 	}
 
+	// If there is a choice to be had between two formats: one with alpha and one without, this returns which
+	// one the texture expects.
+	bool GetTextureExpectsAlphaInPixelFormat(bool bInSourceMipsAlphaDetected) const
+	{
+		// note the order of operations! ( ForceNo takes precedence )
+		if (bForceNoAlphaChannel)
+		{
+			return false;
+		}
+		if (bForceAlphaChannel)
+		{
+			return true;
+		}
+		return bInSourceMipsAlphaDetected;
+	}
+
+
 	/*
 	* Convert the build settings to an actual texture description containing enough information to describe the texture
 	* to hardware APIs.
@@ -350,7 +368,7 @@ public:
 		TArray<FCompressedImage2D>& OutTextureMips,
 		uint32& OutNumMipsInTail,
 		uint32& OutExtData,
-		bool* bOutImageHasAlpha // If desired, this will report whether the mip processing determined an alpha channel is necessary in the encoded texture.		
+		UE::TextureBuildUtilities::FTextureBuildMetadata* OutMetadata
 		) = 0;
 
 	
diff --git a/Engine/Source/Developer/TextureCompressor/TextureCompressor.Build.cs b/Engine/Source/Developer/TextureCompressor/TextureCompressor.Build.cs
--- a/Engine/Source/Developer/TextureCompressor/TextureCompressor.Build.cs
+++ b/Engine/Source/Developer/TextureCompressor/TextureCompressor.Build.cs
@@ -16,6 +16,7 @@ public class TextureCompressor : ModuleRules
 			new string[] {
 				"Core",
 				"ImageCore",
+				"TextureBuildUtilities",
 				"TextureFormat",
 			}
 			);
diff --git a/Engine/Source/Editor/TextureEditor/Private/TextureEditorToolkit.cpp b/Engine/Source/Editor/TextureEditor/Private/TextureEditorToolkit.cpp
--- a/Engine/Source/Editor/TextureEditor/Private/TextureEditorToolkit.cpp
+++ b/Engine/Source/Editor/TextureEditor/Private/TextureEditorToolkit.cpp
@@ -510,7 +510,8 @@ void FTextureEditorToolkit::PopulateQuickInfo( )
 	FTexturePlatformData** PlatformDataPtr = Texture->GetRunningPlatformData();
 	if (PlatformDataPtr && PlatformDataPtr[0]) // Can be null if we haven't had a chance to call CachePlatformData on the texture (brand new)
 	{
-		FTexturePlatformData::FTextureEncodeResultMetadata const& ResultMetadata = PlatformDataPtr[0]->ResultMetadata;
+		FTexturePlatformData* PlatformData = PlatformDataPtr[0];
+		FTexturePlatformData::FTextureEncodeResultMetadata const& ResultMetadata = PlatformData->ResultMetadata;
 		if (ResultMetadata.bIsValid == false)
 		{
 			EncodeSpeedText->SetText(NSLOCTEXT("TextureEditor", "QuickInfo_EncodeSpeed_NA", "Encode Speed: N/A"));
@@ -539,8 +540,6 @@ void FTextureEditorToolkit::PopulateQuickInfo( )
 			// Check if we need to compress new Oodle preview once we know we have
 			// valid results.
 			//
-			FTexturePlatformData* PlatformData = PlatformDataPtr[0];
-
 			bool AlreadyHaveResults = false;
 			if (PlatformData->DerivedDataKey.GetIndex() == OodleCompressedPreviewDDCKey.GetIndex())
 			{
@@ -703,6 +702,9 @@ void FTextureEditorToolkit::PopulateQuickInfo( )
 				OodleTilingText->SetText(FText::AsCultureInvariant(UniversalTilingEnum->GetNameStringByValue(ResultMetadata.OodleUniversalTiling)));
 			} // end if encode speed supported
 		} // end if results metadata valid
+
+		SourceMipsAlphaDetectedText->SetText(FText::Format(NSLOCTEXT("TextureEditor", "QuickInfo_SourceAlphaDetected", "Source Alpha Detected: {0}"),
+			PlatformData->bSourceMipsAlphaDetectedValid ? (PlatformData->bSourceMipsAlphaDetected ? NSLOCTEXT("TextureEditor", "True", "True") : NSLOCTEXT("TextureEditor", "False", "False")) : NSLOCTEXT("TextureEditor", "Unknown", "Unknown")));
 	} // end if valid platform data
 
 	UTexture2D* Texture2D = Cast<UTexture2D>(Texture);
@@ -1684,6 +1686,14 @@ void FTextureEditorToolkit::CreateInternalWidgets()
 				SAssignNew(HasAlphaChannelText, STextBlock)
 			]
 
+			+ SVerticalBox::Slot()
+			.AutoHeight()
+			.VAlign(VAlign_Center)
+			.Padding(4.0f)
+			[
+				SAssignNew(SourceMipsAlphaDetectedText, STextBlock)
+			]
+
 			+ SVerticalBox::Slot()
 			.AutoHeight()
 			.VAlign(VAlign_Center)
diff --git a/Engine/Source/Editor/TextureEditor/Private/TextureEditorToolkit.h b/Engine/Source/Editor/TextureEditor/Private/TextureEditorToolkit.h
--- a/Engine/Source/Editor/TextureEditor/Private/TextureEditorToolkit.h
+++ b/Engine/Source/Editor/TextureEditor/Private/TextureEditorToolkit.h
@@ -413,6 +413,7 @@ private:
 	TSharedPtr<STextBlock> FormatText;
 	TSharedPtr<STextBlock> LODBiasText;
 	TSharedPtr<STextBlock> HasAlphaChannelText;
+	TSharedPtr<STextBlock> SourceMipsAlphaDetectedText;
 	TSharedPtr<STextBlock> NumMipsText;
 	TSharedPtr<STextBlock> MipLevelTextBlock;
 	TSharedPtr<STextBlock> EncodeSpeedText;
diff --git a/Engine/Source/Runtime/Core/Private/HAL/ConsoleManager.cpp b/Engine/Source/Runtime/Core/Private/HAL/ConsoleManager.cpp
--- a/Engine/Source/Runtime/Core/Private/HAL/ConsoleManager.cpp
+++ b/Engine/Source/Runtime/Core/Private/HAL/ConsoleManager.cpp
@@ -3202,12 +3202,6 @@ static TAutoConsoleVariable<float> CVarSkeletalMeshLODRadiusScale(
 	TEXT("Scale factor for the screen radius used in computing discrete LOD for skeletal meshes. (0.25-1)"),
 	ECVF_Scalability | ECVF_RenderThreadSafe);
 
-static TAutoConsoleVariable<int32> CVarSharedLinearTextureEncoding(
-	TEXT("r.SharedLinearTextureEncoding"),
-	0,
-	TEXT("If set to 1, textures for platforms that tile will reuse a host linear texture instead of reencoding."),
-	ECVF_RenderThreadSafe);
-
 static TAutoConsoleVariable<int32> CVarPreTileTextures(
 	TEXT("r.PreTileTextures"),
 	1,
diff --git a/Engine/Source/Runtime/Engine/Classes/Engine/Texture.h b/Engine/Source/Runtime/Engine/Classes/Engine/Texture.h
--- a/Engine/Source/Runtime/Engine/Classes/Engine/Texture.h
+++ b/Engine/Source/Runtime/Engine/Classes/Engine/Texture.h
@@ -740,6 +740,21 @@ struct FTexturePlatformData
 	struct FVirtualTextureBuiltData* VTData;
 
 #if WITH_EDITORONLY_DATA
+
+	// When in the editor we have some data that is stored in the derived data
+	// that we don't want to save in the runtime cooked data:
+	// 1. bSourceMipsAlphaDetected - this is part of the calculation for whether to choose
+	//		e.g. BC1 or BC3.
+	// 2. Interim hash data for cook diff tags - this is used for diff tracking for hunting down
+	//		determinism issues.
+	// We only have this data if the textures were rebuilt after adding it (I didn't invalidate the ddc for this),
+	// hence bSourceMipsAlphaDetectedValid. The Hashes are zero if the data isn't present. 
+	FIoHash PreEncodeMipsHash;
+	FIoHash PostEncodeMipsHash;
+	FIoHash PostTileMipsHash;
+	bool bSourceMipsAlphaDetectedValid=false;
+	bool bSourceMipsAlphaDetected=false;
+
 	/** The key associated with this derived data. */
 	TVariant<FString, UE::DerivedData::FCacheKeyProxy> DerivedDataKey;
 
diff --git a/Engine/Source/Runtime/Engine/Private/TextureDerivedData.cpp b/Engine/Source/Runtime/Engine/Private/TextureDerivedData.cpp
--- a/Engine/Source/Runtime/Engine/Private/TextureDerivedData.cpp
+++ b/Engine/Source/Runtime/Engine/Private/TextureDerivedData.cpp
@@ -492,6 +492,18 @@ void GetTextureDerivedMipKey(
 		);
 }
 
+// Get the ddc key for the texture metadata (old build flow).
+UE::DerivedData::FCacheKey GetTextureDerivedMetadataKeyFromSuffix(
+	const FString& KeySuffix
+)
+{
+	FString Key = FDerivedDataCacheInterface::BuildCacheKey(
+		TEXT("TEXTURE"),
+		TEXTURE_DERIVEDDATA_VER,
+		*(KeySuffix + FString(TEXT("_METADATA"))));
+	return UE::DerivedData::ConvertLegacyCacheKey(Key);
+}
+
 /**
  * Computes the derived data key for a texture with the specified compression settings.
  * @param Texture - The texture for which to compute the derived data key.
@@ -1191,6 +1203,8 @@ static void GetBuildSettingsPerFormat(
 	}
 }
 
+void UnpackTextureBuildMetadataFromPlatformData(UE::TextureBuildUtilities::FTextureBuildMetadata* BuildMetadata, const FTexturePlatformData* PlatformData);
+
 /**
  * Stores derived data in the DDC.
  * After this returns, all bulk data from streaming (non-inline) mips will be sent separately to the DDC and the BulkData for those mips removed.
@@ -1208,6 +1222,20 @@ int64 PutDerivedDataInCache(FTexturePlatformData* DerivedData, const FString& De
 	// Build the key with which to cache derived data.
 	GetTextureDerivedDataKeyFromSuffix(DerivedDataKeySuffix, DerivedDataKey);
 
+	{
+		// Store the metadata.
+		UE::TextureBuildUtilities::FTextureBuildMetadata BuildMetadata;
+		UnpackTextureBuildMetadataFromPlatformData(&BuildMetadata, DerivedData);
+		FCbObject MetadataObject = BuildMetadata.ToCompactBinaryWithDefaults();
+		UE::DerivedData::FValue Value = UE::DerivedData::FValue::Compress(MetadataObject.GetBuffer());
+
+		const UE::DerivedData::FSharedString Name = TextureName;
+		UE::DerivedData::FRequestOwner AsyncOwner(UE::DerivedData::EPriority::Normal);
+		const UE::DerivedData::ECachePolicy Policy = bReplaceExistingDDC ? UE::DerivedData::ECachePolicy::Store : UE::DerivedData::ECachePolicy::Default;
+		UE::DerivedData::GetCache().PutValue({ {Name, GetTextureDerivedMetadataKeyFromSuffix(DerivedDataKeySuffix), MoveTemp(Value), Policy} }, AsyncOwner);
+		AsyncOwner.KeepAlive();
+	}
+
 	FString LogString;
 
 	// Write out individual mips to the derived data cache.
@@ -2560,6 +2588,22 @@ static void SerializePlatformData(
 			{
 				CookTags->Add(Texture, "Diff_20_Tex2D_DDK", FString(*DDK));
 			}
+
+			// Did something in the image processing change?
+			// We haven't yet forced a rebuild of textures, so this hash might not exist in the
+			// platform data.
+			if (PlatformData->PreEncodeMipsHash.IsZero() == false)
+			{
+				CookTags->Add(Texture, "Diff_30_Tex2D_PreEncodeHash", LexToString(PlatformData->PreEncodeMipsHash));
+			}
+			if (PlatformData->PostEncodeMipsHash.IsZero() == false)
+			{
+				CookTags->Add(Texture, "Diff_40_Tex2D_PostEncodeHash", LexToString(PlatformData->PostEncodeMipsHash));
+			}
+			if (PlatformData->PostTileMipsHash.IsZero() == false) // won't have this one if not a tiled texture
+			{
+				CookTags->Add(Texture, "Diff_50_Tex2D_PostTileHash", LexToString(PlatformData->PostTileMipsHash));
+			}
 		}
 	}
 #endif
diff --git a/Engine/Source/Runtime/Engine/Private/TextureDerivedDataTask.cpp b/Engine/Source/Runtime/Engine/Private/TextureDerivedDataTask.cpp
--- a/Engine/Source/Runtime/Engine/Private/TextureDerivedDataTask.cpp
+++ b/Engine/Source/Runtime/Engine/Private/TextureDerivedDataTask.cpp
@@ -83,7 +83,37 @@ static TAutoConsoleVariable<int32> CVarForceRetileTextures(
 	TEXT("however the linear texture is allowed to fetch from cache.")
 );
 
+static TAutoConsoleVariable<int32> CVarSharedLinearTextureEncoding(
+	TEXT("r.SharedLinearTextureEncoding"),
+	0,
+	TEXT("If set to 1, textures for platforms that tile will reuse a host linear texture instead of reencoding.")
+);
+
+
 void GetTextureDerivedDataKeyFromSuffix(const FString& KeySuffix, FString& OutKey);
+UE::DerivedData::FCacheKey GetTextureDerivedMetadataKeyFromSuffix(const FString& KeySuffix);
+static void PackTextureBuildMetadataInPlatformData(FTexturePlatformData* PlatformData, const UE::TextureBuildUtilities::FTextureBuildMetadata& BuildMetadata)
+{
+	PlatformData->bSourceMipsAlphaDetectedValid = true;
+	PlatformData->bSourceMipsAlphaDetected = BuildMetadata.bSourceMipsAlphaDetected;
+	PlatformData->PreEncodeMipsHash = BuildMetadata.PreEncodeMipsHash;
+	PlatformData->PostEncodeMipsHash = BuildMetadata.PostEncodeMipsHash;
+	PlatformData->PostTileMipsHash = BuildMetadata.PostTileMipsHash;
+}
+void UnpackTextureBuildMetadataFromPlatformData(UE::TextureBuildUtilities::FTextureBuildMetadata* BuildMetadata, const FTexturePlatformData* PlatformData)
+{
+	if (PlatformData->bSourceMipsAlphaDetectedValid)
+	{
+		BuildMetadata->bSourceMipsAlphaDetected = PlatformData->bSourceMipsAlphaDetected;
+	}
+	else
+	{
+		BuildMetadata->bSourceMipsAlphaDetected = false;
+	}
+	BuildMetadata->PreEncodeMipsHash = PlatformData->PreEncodeMipsHash;
+	BuildMetadata->PostEncodeMipsHash = PlatformData->PostEncodeMipsHash;
+	BuildMetadata->PostTileMipsHash = PlatformData->PostTileMipsHash;
+}
 
 static FTextureEngineParameters GenerateTextureEngineParameters()
 {
@@ -721,8 +751,9 @@ static void DDC1_BuildTexture(
 		uint32 NumMipsInTail;
 		uint32 ExtData;
 
+		UE::TextureBuildUtilities::FTextureBuildMetadata BuildMetadata;
+
 		// Compress the texture by calling texture compressor directly.
-		bool bImageHasAlphaChannel = false;
 		TArray<FCompressedImage2D> CompressedMips;
 		if (Compressor->BuildTexture(TextureData.Blocks[0].MipsPerLayer[0],
 			((bool)Texture.CompositeTexture && CompositeTextureData.Blocks.Num() && CompositeTextureData.Blocks[0].MipsPerLayer.Num()) ? CompositeTextureData.Blocks[0].MipsPerLayer[0] : TArray<FImage>(),
@@ -731,10 +762,12 @@ static void DDC1_BuildTexture(
 			CompressedMips,
 			NumMipsInTail,
 			ExtData,
-			&bImageHasAlphaChannel))
+			&BuildMetadata))
 		{
 			check(CompressedMips.Num());
 
+			PackTextureBuildMetadataInPlatformData(DerivedData, BuildMetadata);
+
 			DDC1_StoreClassicTextureInDerivedData(
 				CompressedMips, DerivedData, InBuildSettingsPerLayer[0].bVolume, InBuildSettingsPerLayer[0].bTextureArray, InBuildSettingsPerLayer[0].bCubemap, 
 				NumMipsInTail, ExtData, bReplaceExistingDDC, TexturePathName, KeySuffix, BytesCached);
@@ -1237,6 +1270,7 @@ FTextureCacheDerivedDataWorker::FTextureCacheDerivedDataWorker(
 		delete DerivedData->VTData;
 		DerivedData->VTData = nullptr;
 	}
+	DerivedData->bSourceMipsAlphaDetectedValid = false;
 	UTexture::GetPixelFormatEnum();
 		
 	const bool bAllowAsyncBuild = EnumHasAnyFlags(CacheFlags, ETextureCacheFlags::AllowAsyncBuild);
@@ -1344,6 +1378,7 @@ static void DDC1_FetchAndFillDerivedData(
 	bool bForVirtualTextureStreamingBuild = EnumHasAnyFlags(CacheFlags, ETextureCacheFlags::ForVirtualTextureStreamingBuild);
 
 	FSharedBuffer RawDerivedData;
+	const FSharedString SharedTexturePathName(TexturePathName);
 
 	FString LocalDerivedDataKeySuffix;
 	FString LocalDerivedDataKey;
@@ -1363,16 +1398,34 @@ static void DDC1_FetchAndFillDerivedData(
 			FString FetchFirstKey;
 			GetTextureDerivedDataKeyFromSuffix(FetchFirstKeySuffix, FetchFirstKey);
 
+			TArray<FCacheGetValueRequest, TInlineAllocator<2>> Requests;
+			Requests.Add({ SharedTexturePathName, ConvertLegacyCacheKey(FetchFirstKey), ECachePolicy::Default, 0 /* UserData */});
+			Requests.Add({ SharedTexturePathName, GetTextureDerivedMetadataKeyFromSuffix(FetchFirstKeySuffix), ECachePolicy::Default, 1 /* UserData */});
+
 			FRequestOwner BlockingOwner(EPriority::Blocking);
-			GetCache().GetValue({ {{TexturePathName}, ConvertLegacyCacheKey(FetchFirstKey)} }, BlockingOwner, [&RawDerivedData](FCacheGetValueResponse&& Response)
+			FSharedBuffer MetadataBuffer;
+
+			GetCache().GetValue(Requests, BlockingOwner, [&MetadataBuffer, &RawDerivedData](FCacheGetValueResponse&& Response)
 			{
-				RawDerivedData = Response.Value.GetData().Decompress();
+				if (Response.UserData == 0)
+				{
+					RawDerivedData = Response.Value.GetData().Decompress();
+				}
+				else
+				{
+					MetadataBuffer = Response.Value.GetData().Decompress();
+				}
 			});
 			BlockingOwner.Wait();
 
 			bLoadedFromDDC = !RawDerivedData.IsNull();
 			if (bLoadedFromDDC)
 			{
+				// We don't necessarily get the metadata until we force a texture rebuild.
+				if (MetadataBuffer.IsNull() == false)
+				{
+					PackTextureBuildMetadataInPlatformData(DerivedData, FCbObject(MetadataBuffer));
+				}
 				bUsedFetchFirst = true;
 				LocalDerivedDataKey = MoveTemp(FetchFirstKey);
 				LocalDerivedDataKeySuffix = MoveTemp(FetchFirstKeySuffix);
@@ -1386,17 +1439,36 @@ static void DDC1_FetchAndFillDerivedData(
 		LocalDerivedDataKeySuffix = MoveTemp(FetchOrBuildKeySuffix);
 		GetTextureDerivedDataKeyFromSuffix(LocalDerivedDataKeySuffix, LocalDerivedDataKey);
 
-		if (bForceRebuild == false)
+		TArray<FCacheGetValueRequest, TInlineAllocator<2>> Requests;
+		Requests.Add({ SharedTexturePathName, ConvertLegacyCacheKey(LocalDerivedDataKey), ECachePolicy::Default, 0 /* UserData */ });
+		Requests.Add({ SharedTexturePathName, GetTextureDerivedMetadataKeyFromSuffix(LocalDerivedDataKeySuffix), ECachePolicy::Default, 1 /* UserData */ });
+
+		FRequestOwner BlockingOwner(EPriority::Blocking);
+		FSharedBuffer MetadataBuffer;
+
+		GetCache().GetValue(Requests, BlockingOwner, [&MetadataBuffer, &RawDerivedData](FCacheGetValueResponse&& Response)
 		{
-			FRequestOwner BlockingOwner(EPriority::Blocking);
-			GetCache().GetValue({ {{TexturePathName}, ConvertLegacyCacheKey(LocalDerivedDataKey)} }, BlockingOwner, [&RawDerivedData](FCacheGetValueResponse&& Response)
+			if (Response.UserData == 0)
 			{
 				RawDerivedData = Response.Value.GetData().Decompress();
-			});
-			BlockingOwner.Wait();
-		}
+			}
+			else
+			{
+				MetadataBuffer = Response.Value.GetData().Decompress();
+			}
+		});
+		BlockingOwner.Wait();
 
 		bLoadedFromDDC = !RawDerivedData.IsNull();
+		if (bLoadedFromDDC)
+		{
+			// Only read the metadata if we actually got the main data.
+			// We don't necessarily get the metadata until we force a texture rebuild.
+			if (MetadataBuffer.IsNull() == false)
+			{
+				PackTextureBuildMetadataInPlatformData(DerivedData, FCbObject(MetadataBuffer));
+			}
+		}
 	}
 
 	KeySuffix = LocalDerivedDataKeySuffix;
@@ -1691,6 +1763,9 @@ bool DDC1_BuildTiledClassicTexture(
 
 	check(LinearDerivedData.GetNumMipsInTail() == 0);
 
+	UE::TextureBuildUtilities::FTextureBuildMetadata BuildMetadata;
+	UnpackTextureBuildMetadataFromPlatformData(&BuildMetadata, &LinearDerivedData);
+
 	// The linear bits are built - need to fetch
 	void* LinearMipData[MAX_TEXTURE_MIP_COUNT] = {};
 	int64 LinearMipSizes[MAX_TEXTURE_MIP_COUNT];
@@ -1748,6 +1823,7 @@ bool DDC1_BuildTiledClassicTexture(
 	TextureDescription.GetEncodedMipIterators(&TextureExtendedData, MipTailIndex, MipsInTail);
 
 	UE_LOG(LogTexture, Display, TEXT("Tiling %s"), *TexturePathName);
+	FIoHashBuilder TiledMipsHash;
 
 	// Do the actual tiling.
 	for (int32 EncodedMipIndex = 0; EncodedMipIndex < MipTailIndex + 1; EncodedMipIndex++)
@@ -1776,9 +1852,13 @@ bool DDC1_BuildTiledClassicTexture(
 		// \todo try and Move this data rather than copying? We use FSharedBuffer as that's the future way,
 		// but we're interacting with older systems that didn't have it, and we can't Move() from an FSharedBuffer.
 		TiledMip.RawData.AddUninitialized(MipData.GetSize());
+		TiledMipsHash.Update(MipData.GetData(), MipData.GetSize());
 		FMemory::Memcpy(TiledMip.RawData.GetData(), MipData.GetData(), MipData.GetSize());
 	} // end for each mip
 
+	BuildMetadata.PostTileMipsHash = TiledMipsHash.Finalize();
+	PackTextureBuildMetadataInPlatformData(DerivedData, BuildMetadata);
+
 	for (int32 MipIndex = 0; MipIndex < TextureDescription.NumMips; ++MipIndex)
 	{
 		FMemory::Free(LinearMipData[MipIndex]);
@@ -1839,7 +1919,7 @@ void FTextureCacheDerivedDataWorker::DoWork()
 	{
 		if (CVarForceRetileTextures.GetValueOnAnyThread())
 		{
-			// We do this after the fetch so it can fill out the metadata nad key suffix that gets used.
+			// We do this after the fetch so it can fill out the metadata and key suffix that gets used.
 			bSucceeded = false;
 			bLoadedFromDDC = false;
 
@@ -1969,6 +2049,11 @@ static bool UnpackPlatformDataFromBuild(FTexturePlatformData& OutPlatformData, U
 	using namespace UE::DerivedData;
 	UE::DerivedData::FBuildOutput& BuildOutput = InBuildCompleteParams.Output;
 
+	{
+		const FValueWithId& Value = BuildOutput.GetValue(FValueId::FromName(UTF8TEXTVIEW("TextureBuildMetadata")));
+		PackTextureBuildMetadataInPlatformData(&OutPlatformData, FCbObject(Value.GetData().Decompress()));
+	}
+
 	// We take this as a build output, however in ideal (future) situations, this is generated prior to build launch and
 	// just routed through the build. Since we currently handle several varying situations, we just always consume it from
 	// the build no matter where it came from. (Both TextureDescription and ExtendedData)
@@ -2123,7 +2208,7 @@ static bool UnpackPlatformDataFromBuild(FTexturePlatformData& OutPlatformData, U
 	return true;
 }
 
-void HandleBuildOutputThenUnpack(FTexturePlatformData& OutPlatformData, UE::DerivedData::FBuildCompleteParams&& InBuildCompleteParams, FBuildResultOptions InBuildResultOptions)
+static void HandleBuildOutputThenUnpack(FTexturePlatformData& OutPlatformData, UE::DerivedData::FBuildCompleteParams&& InBuildCompleteParams, FBuildResultOptions InBuildResultOptions)
 {
 	using namespace UE::DerivedData;
 
@@ -2550,6 +2635,7 @@ private:
 			FetchFirstBuildPolicyBuilder.AddValuePolicy(FValueId::FromName(ANSITEXTVIEW("EncodedTextureDescription")), EBuildPolicy::Cache);
 			FetchFirstBuildPolicyBuilder.AddValuePolicy(FValueId::FromName(ANSITEXTVIEW("EncodedTextureExtendedData")), EBuildPolicy::Cache);
 			FetchFirstBuildPolicyBuilder.AddValuePolicy(FValueId::FromName(ANSITEXTVIEW("MipTail")), EBuildPolicy::Cache);
+			FetchFirstBuildPolicyBuilder.AddValuePolicy(FValueId::FromName(ANSITEXTVIEW("TextureBuildMetadata")), EBuildPolicy::Cache);
 			return FetchFirstBuildPolicyBuilder.Build();
 		}		
 	}
@@ -2572,6 +2658,7 @@ private:
 			BuildPolicyBuilder.AddValuePolicy(FValueId::FromName(ANSITEXTVIEW("EncodedTextureDescription")), EBuildPolicy::Cache);
 			BuildPolicyBuilder.AddValuePolicy(FValueId::FromName(ANSITEXTVIEW("EncodedTextureExtendedData")), EBuildPolicy::Cache);
 			BuildPolicyBuilder.AddValuePolicy(FValueId::FromName(ANSITEXTVIEW("MipTail")), EBuildPolicy::Default);
+			BuildPolicyBuilder.AddValuePolicy(FValueId::FromName(ANSITEXTVIEW("TextureBuildMetadata")), EBuildPolicy::Cache);
 			return BuildPolicyBuilder.Build();
 		}
 	}
@@ -2779,6 +2866,8 @@ public:
 			DefinitionBuilder.AddConstant(UTF8TEXTVIEW("EncodedTextureExtendedDataConstant"), UE::TextureBuildUtilities::EncodedTextureExtendedData::ToCompactBinary(*InTextureExtendedData));
 		}
 		
+		DefinitionBuilder.AddInputBuild(UTF8TEXTVIEW("TextureBuildMetadata"), { InParentBuildKey, UE::DerivedData::FValueId::FromName(UTF8TEXTVIEW("TextureBuildMetadata")) });
+
 		if (InputTextureNumMips > InputTextureNumStreamingMips)
 		{
 			DefinitionBuilder.AddInputBuild(UTF8TEXTVIEW("MipTail"), { InParentBuildKey, UE::DerivedData::FValueId::FromName(UTF8TEXTVIEW("MipTail")) });
