Track each individual FInputDeviceId's hardware info as well as the most recent one for the platform user. Remove the "Native" hardware event changing because it's unnecessary, you can just bind to the same one as BP. Also pass in the input device ID when that delegate is called.

#rnx

diff --git a/Engine/Source/Runtime/Engine/Classes/GameFramework/InputDeviceSubsystem.h b/Engine/Source/Runtime/Engine/Classes/GameFramework/InputDeviceSubsystem.h
--- a/Engine/Source/Runtime/Engine/Classes/GameFramework/InputDeviceSubsystem.h
+++ b/Engine/Source/Runtime/Engine/Classes/GameFramework/InputDeviceSubsystem.h
@@ -109,7 +109,7 @@ struct ENGINE_API FSetDevicePropertyParams
  *
  * @param UserId		The Platform user whose device has changed
  */
-DECLARE_DYNAMIC_MULTICAST_DELEGATE_OneParam(FHardwareInputDeviceChanged, const FPlatformUserId, UserId);
+DECLARE_DYNAMIC_MULTICAST_DELEGATE_TwoParams(FHardwareInputDeviceChanged, const FPlatformUserId, UserId, const FInputDeviceId, DeviceId);
 
 /**
 * The input device subsystem provides an interface to allow users to set Input Device Properties
@@ -144,6 +144,9 @@ public:
 	UFUNCTION(BlueprintCallable, Category = "Input Devices", meta = (AutoCreateRefTerm = "Params"))
 	FInputDevicePropertyHandle SetDeviceProperty(const FSetDevicePropertyParams& Params);
 
+	/** Returns a pointer to the active input device property with the given handle. Returns null if the property doesn't exist */
+	TObjectPtr<UInputDeviceProperty> GetActiveDeviceProperty(const FInputDevicePropertyHandle Handle);
+
 	/** 
 	* Remove any active device properties that have the same class as the one given.
 	* 
@@ -174,12 +177,12 @@ public:
 	static bool IsDevicePropertyHandleValid(const FInputDevicePropertyHandle& InHandle);
 
 	/** Gets the most recently used hardware input device for the given platform user */
+	UFUNCTION(BlueprintCallable, Category = "Input Devices")
 	FHardwareDeviceIdentifier GetMostRecentlyUsedHardwareDevice(const FPlatformUserId InUserId) const;
 
-	/** A delegate that is fired when a platform user changes what Hardware Input device they are using */
-	DECLARE_EVENT_OneParam(UInputDeviceSubsystem, FHardwareInputDeviceChanged, const FPlatformUserId);
-	FHardwareInputDeviceChanged OnInputHardwareDeviceChangedNative;
-	
+	UFUNCTION(BlueprintCallable, Category = "Input Devices")
+	FHardwareDeviceIdentifier GetInputDeviceHardwareIdentifier(const FInputDeviceId InputDevice) const;
+
 	/** A delegate that is fired when a platform user changes what Hardware Input device they are using */
 	UPROPERTY(BlueprintAssignable, Category = "Input Device")
 	FHardwareInputDeviceChanged OnInputHardwareDeviceChanged;
@@ -195,8 +198,11 @@ protected:
 	UPROPERTY(Transient)
 	TArray<FActiveDeviceProperty> ActiveProperties;
 	
+	/** A map of an input device to it's most recent hardware device identifier */
+	TMap<FInputDeviceId, FHardwareDeviceIdentifier> LatestInputDeviceIdentifiers;
+
 	/** A map of platform user's to their most recent hardware device identifier */
-	TMap<FPlatformUserId, FHardwareDeviceIdentifier> LatestInputDeviceIdentifiers;
+	TMap<FPlatformUserId, FHardwareDeviceIdentifier> LatestUserDeviceIdentifiers;
 
 	/** An input processor that is used to determine the current hardware input device */
 	TSharedPtr<class FInputDeviceSubsystemProcessor> InputPreprocessor;
diff --git a/Engine/Source/Runtime/Engine/Private/GameFramework/InputDeviceSubsystem.cpp b/Engine/Source/Runtime/Engine/Private/GameFramework/InputDeviceSubsystem.cpp
--- a/Engine/Source/Runtime/Engine/Private/GameFramework/InputDeviceSubsystem.cpp
+++ b/Engine/Source/Runtime/Engine/Private/GameFramework/InputDeviceSubsystem.cpp
@@ -300,6 +300,18 @@ FInputDevicePropertyHandle UInputDeviceSubsystem::SetDeviceProperty(const FSetDe
 	return OutHandle;
 }
 
+TObjectPtr<UInputDeviceProperty> UInputDeviceSubsystem::GetActiveDeviceProperty(const FInputDevicePropertyHandle Handle)
+{
+	for (const FActiveDeviceProperty& ActiveProp : ActiveProperties)
+	{
+		if (ActiveProp.PropertyHandle == Handle)
+		{
+			return ActiveProp.Property;
+		}
+	}
+	return nullptr;
+}
+
 int32 UInputDeviceSubsystem::RemoveDevicePropertiesOfClass(const FPlatformUserId UserId, TSubclassOf<UInputDeviceProperty> DevicePropertyClass)
 {
 	int32 NumRemoved = 0;
@@ -360,7 +372,7 @@ void UInputDeviceSubsystem::RemoveAllDeviceProperties()
 
 FHardwareDeviceIdentifier UInputDeviceSubsystem::GetMostRecentlyUsedHardwareDevice(const FPlatformUserId InUserId) const
 {
-	if (const FHardwareDeviceIdentifier* FoundDevice = LatestInputDeviceIdentifiers.Find(InUserId))
+	if (const FHardwareDeviceIdentifier* FoundDevice = LatestUserDeviceIdentifiers.Find(InUserId))
 	{
 		return *FoundDevice;
 	}
@@ -368,19 +380,28 @@ FHardwareDeviceIdentifier UInputDeviceSubsystem::GetMostRecentlyUsedHardwareDevi
 	return FHardwareDeviceIdentifier::Invalid;
 }
 
+FHardwareDeviceIdentifier UInputDeviceSubsystem::GetInputDeviceHardwareIdentifier(const FInputDeviceId InputDevice) const
+{
+	if (const FHardwareDeviceIdentifier* FoundDevice = LatestInputDeviceIdentifiers.Find(InputDevice))
+	{
+		return *FoundDevice;
+	}
+
+	return FHardwareDeviceIdentifier::Invalid;
+}
+
 void UInputDeviceSubsystem::SetMostRecentlyUsedHardwareDevice(const FInputDeviceId InDeviceId, const FHardwareDeviceIdentifier& InHardwareId)
 {
-	FPlatformUserId OwningUserId = IPlatformInputDeviceMapper::Get().GetUserForInputDevice(InDeviceId); 
+	FPlatformUserId OwningUserId = IPlatformInputDeviceMapper::Get().GetUserForInputDevice(InDeviceId);
 
-	LatestInputDeviceIdentifiers.Add(OwningUserId, InHardwareId);
+	// Keep track of each input device's latest hardware id
+	LatestInputDeviceIdentifiers.Add(InDeviceId, InHardwareId);
 
-	if (OnInputHardwareDeviceChanged.IsBound())
-	{
-		OnInputHardwareDeviceChanged.Broadcast(OwningUserId);
-	}
+	// Keep a map to platform users so that we can easily get their most recent hardware
+	LatestUserDeviceIdentifiers.Add(OwningUserId, InHardwareId);
 
-	if (OnInputHardwareDeviceChangedNative.IsBound())
+	if (OnInputHardwareDeviceChanged.IsBound())
 	{
-		OnInputHardwareDeviceChangedNative.Broadcast(OwningUserId);
+		OnInputHardwareDeviceChanged.Broadcast(OwningUserId, InDeviceId);
 	}
 }
