Expose a "Reset Trigger Device Properties" option that will let you reset any active affects on a trigger effect.

#rnx

diff --git a/Engine/Source/Runtime/Engine/Classes/GameFramework/InputDeviceProperties.h b/Engine/Source/Runtime/Engine/Classes/GameFramework/InputDeviceProperties.h
--- a/Engine/Source/Runtime/Engine/Classes/GameFramework/InputDeviceProperties.h
+++ b/Engine/Source/Runtime/Engine/Classes/GameFramework/InputDeviceProperties.h
@@ -185,13 +185,14 @@ private:
 };
 
 /** A property that effect the triggers on a gamepad */
-UCLASS(Abstract, Blueprintable)
+UCLASS(Blueprintable, meta = (DisplayName = "Reset Trigger Device Properties"))
 class ENGINE_API UInputDeviceTriggerEffect : public UInputDeviceProperty
 {
 	GENERATED_BODY()
 
 public:	
 
+	virtual FInputDeviceProperty* GetInternalDeviceProperty() override;
 	virtual void ResetDeviceProperty_Implementation(const FPlatformUserId PlatformUser) override;
 
 	/** Which trigger this property should effect */
diff --git a/Engine/Source/Runtime/Engine/Private/GameFramework/InputDeviceProperties.cpp b/Engine/Source/Runtime/Engine/Private/GameFramework/InputDeviceProperties.cpp
--- a/Engine/Source/Runtime/Engine/Private/GameFramework/InputDeviceProperties.cpp
+++ b/Engine/Source/Runtime/Engine/Private/GameFramework/InputDeviceProperties.cpp
@@ -128,6 +128,11 @@ float UColorInputDeviceCurveProperty::RecalculateDuration()
 ///////////////////////////////////////////////////////////////////////
 // UInputDeviceTriggerEffect
 
+FInputDeviceProperty* UInputDeviceTriggerEffect::GetInternalDeviceProperty()
+{
+	return &ResetProperty;
+}
+
 void UInputDeviceTriggerEffect::ResetDeviceProperty_Implementation(const FPlatformUserId PlatformUser)
 {
 	if (bResetUponCompletion)
