Removing PostTile and PostEncode hashes as they cause determinism issues when enabling shared linear texture encoding.

diff --git a/Engine/Source/Developer/TextureBuild/Private/TextureBuildFunction.cpp b/Engine/Source/Developer/TextureBuild/Private/TextureBuildFunction.cpp
--- a/Engine/Source/Developer/TextureBuild/Private/TextureBuildFunction.cpp
+++ b/Engine/Source/Developer/TextureBuild/Private/TextureBuildFunction.cpp
@@ -539,7 +539,6 @@ void GenericTextureTilingBuildFunction(UE::DerivedData::FBuildContext& Context,
 	}
 
 	// Process the mips
-	FIoHashBuilder PostTileHashBuilder;
 	TArray<FSharedBuffer> MipTailBuffers;
 	for (int32 MipIndex = 0; MipIndex < FirstMipTailIndex + 1; MipIndex++)
 	{
@@ -553,8 +552,6 @@ void GenericTextureTilingBuildFunction(UE::DerivedData::FBuildContext& Context,
 		}
 		FSharedBuffer MipData = Tiler->ProcessMipLevel(TextureDescription, TextureExtendedData, MipsForLevel, MipIndex);
 
-		PostTileHashBuilder.Update(MipData.GetData(), MipData.GetSize());
-
 		// Make sure we got the size we advertised prior to the build. If this ever fires then we
 		// have a critical mismatch!
 		check(TextureExtendedData.MipSizesInBytes[MipIndex] == MipData.GetSize());
@@ -570,8 +567,6 @@ void GenericTextureTilingBuildFunction(UE::DerivedData::FBuildContext& Context,
 		}
 	} // end for each mip
 
-	BuildMetadata.PostTileMipsHash = PostTileHashBuilder.Finalize();
-
 	// The mip tail is a bunch of mips all together in one "Value", so assemble them here.
 	FCompositeBuffer MipTail(MipTailBuffers);
 	if (MipTail.GetSize() > 0)
diff --git a/Engine/Source/Developer/TextureBuildUtilities/Private/TextureBuildUtilities.cpp b/Engine/Source/Developer/TextureBuildUtilities/Private/TextureBuildUtilities.cpp
--- a/Engine/Source/Developer/TextureBuildUtilities/Private/TextureBuildUtilities.cpp
+++ b/Engine/Source/Developer/TextureBuildUtilities/Private/TextureBuildUtilities.cpp
@@ -223,14 +223,6 @@ FCbObject FTextureBuildMetadata::ToCompactBinaryWithDefaults() const
 	{
 		Writer.AddHash("PreEncodeMipsHash", PreEncodeMipsHash);
 	}
-	if (PostEncodeMipsHash != Defaults.PostEncodeMipsHash)
-	{
-		Writer.AddHash("PostEncodeMipsHash", PostEncodeMipsHash);
-	}
-	if (PostTileMipsHash != Defaults.PostTileMipsHash)
-	{
-		Writer.AddHash("PostTileMipsHash", PostTileMipsHash);
-	}
 	Writer.EndObject();
 	return Writer.Save().AsObject();
 }
@@ -239,8 +231,6 @@ FTextureBuildMetadata::FTextureBuildMetadata(FCbObject InCbObject)
 {
 	bSourceMipsAlphaDetected = InCbObject["bSourceMipsAlphaDetected"].AsBool(bSourceMipsAlphaDetected);
 	PreEncodeMipsHash = InCbObject["PreEncodeMipsHash"].AsHash(PreEncodeMipsHash);
-	PostEncodeMipsHash = InCbObject["PostEncodeMipsHash"].AsHash(PostEncodeMipsHash);
-	PostTileMipsHash = InCbObject["PostTileMipsHash"].AsHash(PostTileMipsHash);
 }
 
 
diff --git a/Engine/Source/Developer/TextureBuildUtilities/Public/TextureBuildUtilities.h b/Engine/Source/Developer/TextureBuildUtilities/Public/TextureBuildUtilities.h
--- a/Engine/Source/Developer/TextureBuildUtilities/Public/TextureBuildUtilities.h
+++ b/Engine/Source/Developer/TextureBuildUtilities/Public/TextureBuildUtilities.h
@@ -52,8 +52,6 @@ struct TEXTUREBUILDUTILITIES_API FTextureBuildMetadata
 	// Digests of the data at various processing stages so we can track down determinism issues
 	// that arise.
 	FIoHash PreEncodeMipsHash;
-	FIoHash PostEncodeMipsHash;
-	FIoHash PostTileMipsHash;
 
 	FCbObject ToCompactBinaryWithDefaults() const;
 	FTextureBuildMetadata(FCbObject InCbObject);
diff --git a/Engine/Source/Developer/TextureCompressor/Private/TextureCompressorModule.cpp b/Engine/Source/Developer/TextureCompressor/Private/TextureCompressorModule.cpp
--- a/Engine/Source/Developer/TextureCompressor/Private/TextureCompressorModule.cpp
+++ b/Engine/Source/Developer/TextureCompressor/Private/TextureCompressorModule.cpp
@@ -3584,16 +3584,6 @@ public:
 		bool bCompressSucceeded = CompressMipChain(TextureFormat, IntermediateMipChain, BuildSettings, bImageHasAlphaChannel, DebugTexturePathName,
 					OutTextureMips, OutNumMipsInTail, OutExtData);
 
-		if (OutMetadata)
-		{
-			FIoHashBuilder MipHashBuilder;
-			for (const FCompressedImage2D& Mip : OutTextureMips)
-			{
-				MipHashBuilder.Update(Mip.RawData.GetData(), Mip.RawData.Num());
-			}
-			OutMetadata->PostEncodeMipsHash = MipHashBuilder.Finalize();
-		}
-
 		return bCompressSucceeded;
 	}
 
diff --git a/Engine/Source/Runtime/Engine/Classes/Engine/Texture.h b/Engine/Source/Runtime/Engine/Classes/Engine/Texture.h
--- a/Engine/Source/Runtime/Engine/Classes/Engine/Texture.h
+++ b/Engine/Source/Runtime/Engine/Classes/Engine/Texture.h
@@ -750,8 +750,6 @@ struct FTexturePlatformData
 	// We only have this data if the textures were rebuilt after adding it (I didn't invalidate the ddc for this),
 	// hence bSourceMipsAlphaDetectedValid. The Hashes are zero if the data isn't present. 
 	FIoHash PreEncodeMipsHash;
-	FIoHash PostEncodeMipsHash;
-	FIoHash PostTileMipsHash;
 	bool bSourceMipsAlphaDetectedValid=false;
 	bool bSourceMipsAlphaDetected=false;
 
diff --git a/Engine/Source/Runtime/Engine/Private/TextureDerivedData.cpp b/Engine/Source/Runtime/Engine/Private/TextureDerivedData.cpp
--- a/Engine/Source/Runtime/Engine/Private/TextureDerivedData.cpp
+++ b/Engine/Source/Runtime/Engine/Private/TextureDerivedData.cpp
@@ -2596,14 +2596,6 @@ static void SerializePlatformData(
 			{
 				CookTags->Add(Texture, "Diff_30_Tex2D_PreEncodeHash", LexToString(PlatformData->PreEncodeMipsHash));
 			}
-			if (PlatformData->PostEncodeMipsHash.IsZero() == false)
-			{
-				CookTags->Add(Texture, "Diff_40_Tex2D_PostEncodeHash", LexToString(PlatformData->PostEncodeMipsHash));
-			}
-			if (PlatformData->PostTileMipsHash.IsZero() == false) // won't have this one if not a tiled texture
-			{
-				CookTags->Add(Texture, "Diff_50_Tex2D_PostTileHash", LexToString(PlatformData->PostTileMipsHash));
-			}
 		}
 	}
 #endif
diff --git a/Engine/Source/Runtime/Engine/Private/TextureDerivedDataTask.cpp b/Engine/Source/Runtime/Engine/Private/TextureDerivedDataTask.cpp
--- a/Engine/Source/Runtime/Engine/Private/TextureDerivedDataTask.cpp
+++ b/Engine/Source/Runtime/Engine/Private/TextureDerivedDataTask.cpp
@@ -97,8 +97,6 @@ static void PackTextureBuildMetadataInPlatformData(FTexturePlatformData* Platfor
 	PlatformData->bSourceMipsAlphaDetectedValid = true;
 	PlatformData->bSourceMipsAlphaDetected = BuildMetadata.bSourceMipsAlphaDetected;
 	PlatformData->PreEncodeMipsHash = BuildMetadata.PreEncodeMipsHash;
-	PlatformData->PostEncodeMipsHash = BuildMetadata.PostEncodeMipsHash;
-	PlatformData->PostTileMipsHash = BuildMetadata.PostTileMipsHash;
 }
 void UnpackTextureBuildMetadataFromPlatformData(UE::TextureBuildUtilities::FTextureBuildMetadata* BuildMetadata, const FTexturePlatformData* PlatformData)
 {
@@ -111,8 +109,6 @@ void UnpackTextureBuildMetadataFromPlatformData(UE::TextureBuildUtilities::FText
 		BuildMetadata->bSourceMipsAlphaDetected = false;
 	}
 	BuildMetadata->PreEncodeMipsHash = PlatformData->PreEncodeMipsHash;
-	BuildMetadata->PostEncodeMipsHash = PlatformData->PostEncodeMipsHash;
-	BuildMetadata->PostTileMipsHash = PlatformData->PostTileMipsHash;
 }
 
 static FTextureEngineParameters GenerateTextureEngineParameters()
@@ -1823,7 +1819,6 @@ bool DDC1_BuildTiledClassicTexture(
 	TextureDescription.GetEncodedMipIterators(&TextureExtendedData, MipTailIndex, MipsInTail);
 
 	UE_LOG(LogTexture, Display, TEXT("Tiling %s"), *TexturePathName);
-	FIoHashBuilder TiledMipsHash;
 
 	// Do the actual tiling.
 	for (int32 EncodedMipIndex = 0; EncodedMipIndex < MipTailIndex + 1; EncodedMipIndex++)
@@ -1852,11 +1847,9 @@ bool DDC1_BuildTiledClassicTexture(
 		// \todo try and Move this data rather than copying? We use FSharedBuffer as that's the future way,
 		// but we're interacting with older systems that didn't have it, and we can't Move() from an FSharedBuffer.
 		TiledMip.RawData.AddUninitialized(MipData.GetSize());
-		TiledMipsHash.Update(MipData.GetData(), MipData.GetSize());
 		FMemory::Memcpy(TiledMip.RawData.GetData(), MipData.GetData(), MipData.GetSize());
 	} // end for each mip
 
-	BuildMetadata.PostTileMipsHash = TiledMipsHash.Finalize();
 	PackTextureBuildMetadataInPlatformData(DerivedData, BuildMetadata);
 
 	for (int32 MipIndex = 0; MipIndex < TextureDescription.NumMips; ++MipIndex)
