Implements r.DynamicRes.DynamicFrameTime to increase dynamic res when CPU or VSync bound on desktop

This to ensure by default dynamic res still maxout the GPU feeding maximum pixel/s possible to TSR when user's frame cap or VSync settings are not optimally set.

diff --git a/Engine/Source/Runtime/Engine/Private/DynamicResolution.cpp b/Engine/Source/Runtime/Engine/Private/DynamicResolution.cpp
--- a/Engine/Source/Runtime/Engine/Private/DynamicResolution.cpp
+++ b/Engine/Source/Runtime/Engine/Private/DynamicResolution.cpp
@@ -13,6 +13,13 @@
 #include "UnrealEngine.h"
 #include "RenderGraphEvent.h"
 
+#if PLATFORM_DESKTOP
+	#define COMPILE_DYNAMIC_FRAME_TIME 1
+#else
+	#define COMPILE_DYNAMIC_FRAME_TIME 0
+#endif
+
+
 static TAutoConsoleVariable<float> CVarDynamicResMinSP(
 	TEXT("r.DynamicRes.MinScreenPercentage"),
 	DynamicRenderScaling::FractionToPercentage(DynamicRenderScaling::FHeuristicSettings::kDefaultMinResolutionFraction),
@@ -32,6 +39,31 @@ static TAutoConsoleVariable<float> CVarFrameTimeBudget(
 	TEXT("Frame's time budget in milliseconds."),
 	ECVF_RenderThreadSafe | ECVF_Default);
 
+#if COMPILE_DYNAMIC_FRAME_TIME
+
+static TAutoConsoleVariable<int32> CVarDynamicFrameTimeEnable(
+	TEXT("r.DynamicRes.DynamicFrameTime"), 1,
+	TEXT("Whether the r.DynamicRes.FrameTimeBudget should automatically increases when frame rate is bound by CPU."),
+	ECVF_RenderThreadSafe | ECVF_Default);
+
+static TAutoConsoleVariable<float> CVarDynamicFrameTimeErrorMarginPercent(
+	TEXT("r.DynamicRes.DynamicFrameTime.ErrorMarginPercent"), 10.0f,
+	TEXT("How much headroom should be left between CPU and GPU."),
+	ECVF_RenderThreadSafe | ECVF_Default);
+
+static TAutoConsoleVariable<int32> CVarDynamicFrameTimeTrack(
+	TEXT("r.DynamicRes.DynamicFrameTime.Track"), 1,
+	TEXT("What to track to control the budget\n")
+	TEXT(" 0: Frametime (but can create feedback loop when GPU bound and VSync);")
+	TEXT(" 1: Threads time (default);"),
+	ECVF_RenderThreadSafe | ECVF_Default);
+
+static TAutoConsoleVariable<float> CVarDynamicFrameTimeRoundUpToVsync(
+	TEXT("r.DynamicRes.DynamicFrameTime.RoundUpToVSyncError"), 10.f,
+	TEXT("Error to use to round up the dynamic frame time to vsync boundaries (default=10%)."),
+	ECVF_RenderThreadSafe | ECVF_Default);
+
+#endif
 
 
 static TAutoConsoleVariable<float> CVarTargetedGPUHeadRoomPercentage(
@@ -161,13 +193,15 @@ void FDynamicResolutionHeuristicProxy::ResetInternal()
 	NumberOfFramesSinceScreenPercentageChange = 0;
 	CurrentFrameResolutionFractions.SetAll(1.0f);
 	CurrentFrameMaxResolutionFractions.SetAll(1.0f);
+	DynamicFrameTimeBudgetMs = 0.0f;
+	DynamicFrameTimeVSyncFactor = 0;
 
 	// Ignore previous frame timings.
 	IgnoreFrameRemainingCount = 1;
 }
 
 uint64 FDynamicResolutionHeuristicProxy::CreateNewPreviousFrameTimings_RenderThread(
-	float GameThreadTimeMs, float RenderThreadTimeMs)
+	float FrameTimeMs, float GameThreadTimeMs, float RenderThreadTimeMs, float RHIThreadTimeMs)
 {
 	check(IsInRenderingThread());
 
@@ -185,8 +219,10 @@ uint64 FDynamicResolutionHeuristicProxy::CreateNewPreviousFrameTimings_RenderThr
 		int32 NewHistoryEntryIndex = (PreviousFrameIndex + 1) % History.Num();
 		History[NewHistoryEntryIndex] = FrameHistoryEntry();
 		History[NewHistoryEntryIndex].ResolutionFractions = CurrentFrameResolutionFractions;
+		History[NewHistoryEntryIndex].FrameTimeMs = FrameTimeMs;
 		History[NewHistoryEntryIndex].GameThreadTimeMs = GameThreadTimeMs;
 		History[NewHistoryEntryIndex].RenderThreadTimeMs = RenderThreadTimeMs;
+		History[NewHistoryEntryIndex].RHIThreadTimeMs = RHIThreadTimeMs;
 		PreviousFrameIndex = NewHistoryEntryIndex;
 		HistorySize = FMath::Min(HistorySize + 1, History.Num());
 
@@ -234,6 +270,96 @@ void FDynamicResolutionHeuristicProxy::RefreshCurrentFrameResolutionFraction_Ren
 
 	const bool bCanChangeResolution = NumberOfFramesSinceScreenPercentageChange >= CVarFrameChangePeriod.GetValueOnRenderThread();
 
+#if COMPILE_DYNAMIC_FRAME_TIME
+	float MinGlobalFrameTime = 0.0f;
+	if (CVarDynamicFrameTimeEnable.GetValueOnRenderThread())
+	{
+		const float VSyncFrameTime = 16.6;
+		static const IConsoleVariable* VSyncCVar = IConsoleManager::Get().FindConsoleVariable(TEXT("r.VSync"));
+		check(VSyncCVar);
+
+		const int32 FrameTimeTrack = CVarDynamicFrameTimeTrack.GetValueOnRenderThread();
+
+		// Find out the frame time that should be used.
+		TArray<float, TInlineAllocator<16>> SortedFrameTime;
+		SortedFrameTime.Reserve(History.Num());
+		for (int32 BrowsingFrameId = 0; BrowsingFrameId < History.Num(); BrowsingFrameId++)
+		{
+			const FrameHistoryEntry& FrameEntry = History[(History.Num() + PreviousFrameIndex - BrowsingFrameId) % History.Num()];
+			float FrameTimeMs = 0.0f;
+			if (FrameTimeTrack == 0)
+			{
+				FrameTimeMs = FrameEntry.FrameTimeMs;
+			}
+			else
+			{
+				FrameTimeMs = FMath::Max(FrameTimeMs, FrameEntry.GameThreadTimeMs);
+				FrameTimeMs = FMath::Max(FrameTimeMs, FrameEntry.RenderThreadTimeMs);
+				FrameTimeMs = FMath::Max(FrameTimeMs, FrameEntry.RHIThreadTimeMs);
+			}
+
+
+			// Ignores invalid frame time yet. 
+			if (FrameTimeMs < 0)
+			{
+				continue;
+			}
+
+			SortedFrameTime.Add(FrameTimeMs);
+		}
+
+		if (SortedFrameTime.Num() > 0)
+		{
+			SortedFrameTime.Sort([](const float& A, const float& B) {
+				return  A < B;
+			});
+			float MedianFrameTime = SortedFrameTime[SortedFrameTime.Num() / 2];
+
+			// When CPU bound and vsync is enabled, the GPU may end up underused due to vsync wait, so we can crank up the available GPU time a bit more.
+			if (float RoundUpToVSyncError = CVarDynamicFrameTimeRoundUpToVsync.GetValueOnGameThread() > 0.0f && VSyncCVar->GetInt())
+			{
+				float VSyncRateChangeMultiplier = 1.0f + RoundUpToVSyncError / 100.f;
+
+				int32 CurrentVSyncFactor = FMath::CeilToInt(MedianFrameTime / VSyncFrameTime);
+
+				float PreviousSyncedFrameTime = DynamicFrameTimeVSyncFactor * VSyncFrameTime;
+
+				if (DynamicFrameTimeVSyncFactor == 0)
+				{
+					DynamicFrameTimeVSyncFactor = CurrentVSyncFactor;
+				}
+				else if (CurrentVSyncFactor > DynamicFrameTimeVSyncFactor)
+				{
+					if (MedianFrameTime > PreviousSyncedFrameTime * VSyncRateChangeMultiplier)
+					{
+						DynamicFrameTimeVSyncFactor = CurrentVSyncFactor;
+					}
+				}
+				else if (CurrentVSyncFactor < DynamicFrameTimeVSyncFactor)
+				{
+					if (MedianFrameTime < PreviousSyncedFrameTime / VSyncRateChangeMultiplier)
+					{
+						DynamicFrameTimeVSyncFactor = CurrentVSyncFactor;
+					}
+				}
+
+				MedianFrameTime = DynamicFrameTimeVSyncFactor * VSyncFrameTime;
+			}
+			else
+			{
+				DynamicFrameTimeVSyncFactor = 0;
+			}
+
+			const float MaxWithFrameFraction = (CVarDynamicFrameTimeErrorMarginPercent.GetValueOnRenderThread() / 100.0f);
+
+			float TargetFrameTime = MedianFrameTime * (1.0f - MaxWithFrameFraction);
+
+			DynamicFrameTimeBudgetMs = FMath::Lerp(TargetFrameTime, DynamicFrameTimeBudgetMs, FrameWeightExponent);
+			MinGlobalFrameTime = DynamicFrameTimeBudgetMs;
+		}
+	}
+#endif // COMPILE_DYNAMIC_FRAME_TIME
+
 	// New ResolutionFraction to use for this frame.
 	DynamicRenderScaling::TMap<float> NewFrameResolutionFractions = CurrentFrameResolutionFractions;
 
@@ -249,9 +375,12 @@ void FDynamicResolutionHeuristicProxy::RefreshCurrentFrameResolutionFraction_Ren
 
 		bool bEarlyReturn = BudgetHistorySize == 0;
 
+		float BudgetBudgetMs = BudgetSettings.BudgetMs;
 		if (Budget == GDynamicPrimaryResolutionFraction)
 		{
-			// NOP
+			#if COMPILE_DYNAMIC_FRAME_TIME
+				BudgetBudgetMs = FMath::Max(BudgetBudgetMs, MinGlobalFrameTime);
+			#endif
 		}
 		else if (!BudgetSettings.IsEnabled())
 		{
@@ -264,6 +393,8 @@ void FDynamicResolutionHeuristicProxy::RefreshCurrentFrameResolutionFraction_Ren
 			continue;
 		}
 
+		const float BudgetTargetMs = BudgetSettings.GetTargetedMs(BudgetBudgetMs);
+
 		float NewFrameResolutionFraction = 0.0f;
 
 		// Total weight of NewFrameResolutionFraction.
@@ -303,13 +434,13 @@ void FDynamicResolutionHeuristicProxy::RefreshCurrentFrameResolutionFraction_Ren
 			const bool bIsGameThreadBound = FrameEntry.GameThreadTimeMs > GDynamicPrimaryResolutionFraction.GetSettings().BudgetMs;
 
 			// Whether bound by render thread.
-			const bool bIsRenderThreadBound = FrameEntry.RenderThreadTimeMs > GDynamicPrimaryResolutionFraction.GetSettings().GetTargetedMs();
+			const bool bIsRenderThreadBound = FrameEntry.RenderThreadTimeMs > GDynamicPrimaryResolutionFraction.GetSettings().GetTargetedMs(GDynamicPrimaryResolutionFraction.GetSettings().BudgetMs);
 
 			// Whether the frame is CPU bound.
 			const bool bIsCPUBound = bIsGameThreadBound || bIsRenderThreadBound;
 
 			// Whether GPU is over budget, when not CPU bound.
-			const bool bHasOverBudgetGPU = !bIsCPUBound && TotalFrameGPUBusyTimeMs > BudgetSettings.BudgetMs;
+			const bool bHasOverBudgetGPU = !bIsCPUBound && TotalFrameGPUBusyTimeMs > BudgetBudgetMs;
 
 			// Look if this is multiple consecutive GPU over budget frames.
 			if (bHasOverBudgetGPU)
@@ -344,7 +475,7 @@ void FDynamicResolutionHeuristicProxy::RefreshCurrentFrameResolutionFraction_Ren
 				// At resolution change that happen every N frames, amortized over time and scaled down as the
 				// standard variation of the GPU timing over non resolution changing frames increases.
 				SuggestedResolutionFraction = (
-					BudgetSettings.EstimateResolutionFactor(TotalFrameGPUBusyTimeMs)
+					BudgetSettings.EstimateResolutionFactor(BudgetTargetMs, TotalFrameGPUBusyTimeMs)
 					* FrameEntry.ResolutionFractions[Budget]);
 			}
 
@@ -468,7 +599,7 @@ void FDynamicResolutionHeuristicProxy::RefreshHeuristicStats_RenderThread()
 
 		if (HeuristicSettings.IsEnabled())
 		{
-			SET_FLOAT_STAT_FName(Budget.GetStatId_TargetMs().GetName(), HeuristicSettings.GetTargetedMs());
+			SET_FLOAT_STAT_FName(Budget.GetStatId_TargetMs().GetName(), HeuristicSettings.GetTargetedMs(HeuristicSettings.BudgetMs));
 			// MeasuredMs is set in RefreshCurrentFrameResolutionFraction_RenderThread()
 			SET_FLOAT_STAT_FName(Budget.GetStatId_MinScaling().GetName(), HeuristicSettings.MinResolutionFraction);
 			SET_FLOAT_STAT_FName(Budget.GetStatId_MaxScaling().GetName(), HeuristicSettings.MaxResolutionFraction);
@@ -567,7 +698,7 @@ public:
 		}
 	}
 
-	void BeginFrame(FRHICommandList& RHICmdList, float PrevGameThreadTimeMs)
+	void BeginFrame(FRHICommandList& RHICmdList, float PrevFrameTimeMs, float PrevGameThreadTimeMs)
 	{
 		check(IsInRenderingThread());
 
@@ -589,6 +720,7 @@ public:
 
 		// Query render thread time Ms.
 		float PrevRenderThreadTimeMs = FPlatformTime::ToMilliseconds(GRenderThreadTime);
+		float PrevRHIThreadTimeMs = FPlatformTime::ToMilliseconds(GRHIThreadTime);
 
 		bUseTimeQueriesThisFrame = GSupportsTimestampRenderQueries && CVarTimingMeasureModel.GetValueOnRenderThread() == 1;
 
@@ -608,7 +740,7 @@ public:
 
 			// Feed the thread timings to the heuristic.
 			InFlightFrame.HeuristicHistoryEntry = Heuristic.CreateNewPreviousFrameTimings_RenderThread(
-				PrevGameThreadTimeMs, PrevRenderThreadTimeMs);
+				PrevFrameTimeMs, PrevGameThreadTimeMs, PrevRenderThreadTimeMs, PrevRHIThreadTimeMs);
 
 			check(QueryPool.IsValid());
 			InFlightFrame.BeginFrameQuery = QueryPool->AllocateQuery();
@@ -621,7 +753,7 @@ public:
 			float PrevFrameGPUTimeMs = FPlatformTime::ToMilliseconds(RHIGetGPUFrameCycles());
 
 			uint64 HistoryEntryId = Heuristic.CreateNewPreviousFrameTimings_RenderThread(
-				PrevGameThreadTimeMs, PrevRenderThreadTimeMs);
+				PrevFrameTimeMs, PrevGameThreadTimeMs, PrevRenderThreadTimeMs, PrevRHIThreadTimeMs);
 
 			const DynamicRenderScaling::TMap<uint64>& LattestTimings = DynamicRenderScaling::GetLastestTimings();
 			DynamicRenderScaling::TMap<float> BudgetTimingMs;
@@ -1005,13 +1137,14 @@ public:
 		if (Event == EDynamicResolutionStateEvent::BeginFrame)
 		{
 			// Query game thread time in milliseconds.
+			float PrevFrameTimeMs = (FApp::GetCurrentTime() - FApp::GetLastTime()) * 1000.0f;
 			float PrevGameThreadTimeMs = FPlatformTime::ToMilliseconds(GGameThreadTime);
 
 			FDefaultDynamicResolutionStateProxy* P = Proxy;
 			ENQUEUE_RENDER_COMMAND(DynamicResolutionBeginFrame)(
-				[PrevGameThreadTimeMs, P](class FRHICommandList& RHICmdList)
+				[PrevFrameTimeMs, PrevGameThreadTimeMs, P](class FRHICommandList& RHICmdList)
 			{
-				P->BeginFrame(RHICmdList, PrevGameThreadTimeMs);
+				P->BeginFrame(RHICmdList, PrevFrameTimeMs, PrevGameThreadTimeMs);
 			});
 		}
 		else
diff --git a/Engine/Source/Runtime/Engine/Public/DynamicResolutionProxy.h b/Engine/Source/Runtime/Engine/Public/DynamicResolutionProxy.h
--- a/Engine/Source/Runtime/Engine/Public/DynamicResolutionProxy.h
+++ b/Engine/Source/Runtime/Engine/Public/DynamicResolutionProxy.h
@@ -19,7 +19,11 @@ public:
 	void Reset_RenderThread();
 
 	/** Create a new previous frame and feeds its timings and returns history's unique id. */
-	uint64 CreateNewPreviousFrameTimings_RenderThread(float GameThreadTimeMs, float RenderThreadTimeMs);
+	uint64 CreateNewPreviousFrameTimings_RenderThread(float FrameTime, float GameThreadTimeMs, float RenderThreadTimeMs, float RHIThreadTime);
+	FORCEINLINE uint64 CreateNewPreviousFrameTimings_RenderThread(float GameThreadTimeMs, float RenderThreadTimeMs)
+	{
+		return CreateNewPreviousFrameTimings_RenderThread(/* FrameTime = */ 0.0f, GameThreadTimeMs, RenderThreadTimeMs, /* RHIThreadTime = */ 0.0f);
+	}
 
 	/** Commit GPU busy times */
 	void CommitPreviousFrameGPUTimings_RenderThread(
@@ -55,14 +59,16 @@ private:
 	struct FrameHistoryEntry
 	{
 		// Thread timings in milliseconds.
-		float GameThreadTimeMs;
-		float RenderThreadTimeMs;
+		float FrameTimeMs = -1.0f;
+		float GameThreadTimeMs = -1.0f;
+		float RenderThreadTimeMs = -1.0f;
+		float RHIThreadTimeMs = -1.0f;
 
 		// Total GPU busy time for the entire frame in milliseconds.
-		float TotalFrameGPUBusyTimeMs;
+		float TotalFrameGPUBusyTimeMs = -1.0f;
 
 		// Total GPU busy time for the render thread commands that perform dynamic resolutions.
-		float GlobalDynamicResolutionTimeMs;
+		float GlobalDynamicResolutionTimeMs = -1.0f;
 
 		// Time for each individual timings
 		DynamicRenderScaling::TMap<float> BudgetTimingMs;
@@ -71,10 +77,6 @@ private:
 		DynamicRenderScaling::TMap<float> ResolutionFractions;
 
 		inline FrameHistoryEntry()
-			: GameThreadTimeMs(-1.0f)
-			, RenderThreadTimeMs(-1.0f)
-			, TotalFrameGPUBusyTimeMs(-1.0f)
-			, GlobalDynamicResolutionTimeMs(-1.0f)
 		{
 			ResolutionFractions.SetAll(1.0f);
 			BudgetTimingMs.SetAll(-1.0f);
@@ -93,6 +95,12 @@ private:
 	int32 PreviousFrameIndex;
 	int32 HistorySize;
 
+	// Exponential average of the CPU frame time.
+	float DynamicFrameTimeBudgetMs;
+
+	// Multiple of the VSync used.
+	int32 DynamicFrameTimeVSyncFactor = 0.0f;
+
 	// Counts the number of frame since the last screen percentage change.
 	int32 NumberOfFramesSinceScreenPercentageChange;
 
diff --git a/Engine/Source/Runtime/RenderCore/Private/DynamicRenderScaling.cpp b/Engine/Source/Runtime/RenderCore/Private/DynamicRenderScaling.cpp
--- a/Engine/Source/Runtime/RenderCore/Private/DynamicRenderScaling.cpp
+++ b/Engine/Source/Runtime/RenderCore/Private/DynamicRenderScaling.cpp
@@ -18,10 +18,10 @@ bool FHeuristicSettings::IsEnabled() const
 	return BudgetMs > 0.0f;
 }
 
-float FHeuristicSettings::GetTargetedMs() const
+float FHeuristicSettings::GetTargetedMs(float InBudgetMs) const
 {
 	check(IsEnabled());
-	return BudgetMs * (1.0f - TargetedHeadRoom);
+	return InBudgetMs * (1.0f - TargetedHeadRoom);
 }
 
 float FHeuristicSettings::EstimateCostScale(float ResolutionFraction) const
@@ -42,10 +42,8 @@ float FHeuristicSettings::EstimateCostScale(float ResolutionFraction) const
 	return Cost;
 }
 
-float FHeuristicSettings::EstimateResolutionFactor(float TimingMs) const
+float FHeuristicSettings::EstimateResolutionFactor(float TargetedMS, float TimingMs) const
 {
-	float TargetedMS = GetTargetedMs();
-
 	float S = 1.0f;
 	if (Model == EHeuristicModel::Linear)
 	{
diff --git a/Engine/Source/Runtime/RenderCore/Public/DynamicRenderScaling.h b/Engine/Source/Runtime/RenderCore/Public/DynamicRenderScaling.h
--- a/Engine/Source/Runtime/RenderCore/Public/DynamicRenderScaling.h
+++ b/Engine/Source/Runtime/RenderCore/Public/DynamicRenderScaling.h
@@ -57,13 +57,13 @@ struct RENDERCORE_API FHeuristicSettings final
 	bool IsEnabled() const;
 
 	/** Returns the desired GPU cost to be targeted to have head room left to not go over budget. */
-	float GetTargetedMs() const;
+	float GetTargetedMs(float BudgetMs) const;
 
 	/** Returns how much the GPU cost scales for a given ResolutionFraction. */
 	float EstimateCostScale(float ResolutionFraction) const;
 
 	/** Returns how much the ResolutionFraction should scale for a GPU timing to fit to target. */
-	float EstimateResolutionFactor(float TimingMs) const;
+	float EstimateResolutionFactor(float TargetMs, float TimingMs) const;
 
 	/** Returns how much the GPU time should scale between two different resolution fraction. */
 	float EstimateTimeFactor(float CurrentResolutionFraction, float NewResolutionFraction) const;
