Flush DeferredTickCBAssests before GC so we don't end up having stale references in the array. Re-enable r.Streaming.OverlapAssetAndLevelTick on a platform

diff --git a/Engine/Source/Runtime/Engine/Private/Streaming/StreamingManagerTexture.cpp b/Engine/Source/Runtime/Engine/Private/Streaming/StreamingManagerTexture.cpp
--- a/Engine/Source/Runtime/Engine/Private/Streaming/StreamingManagerTexture.cpp
+++ b/Engine/Source/Runtime/Engine/Private/Streaming/StreamingManagerTexture.cpp
@@ -78,6 +78,13 @@ static TAutoConsoleVariable<int32> CVarSyncStatesWhenBlocking(
 	TEXT("If true, SyncStates will be called to fully update async states before flushing outstanding streaming requests. Used by Movie Render Queue to ensure all streaming requests are handled each frame to avoid pop-in."),
 	ECVF_Default);
 
+// TODO: Remove once efficacy has been verified
+static TAutoConsoleVariable<int32> CVarFlushDeferredMipLevelChangeCallbacksBeforeGC(
+	TEXT("r.Streaming.FlushDeferredMipLevelChangeCallbacksBeforeGC"),
+	1,
+	TEXT("Whether to flush deferred mip level change callbacks before GC."),
+	ECVF_Default);
+
 bool TrackRenderAsset( const FString& AssetName );
 bool UntrackRenderAsset( const FString& AssetName );
 void ListTrackedRenderAssets( FOutputDevice& Ar, int32 NumTextures );
@@ -218,6 +225,7 @@ void FRenderAssetStreamingManager::OnPreGarbageCollect()
 	TRACE_CPUPROFILER_EVENT_SCOPE(FRenderAssetStreamingManager::OnPreGarbageCollect);
 
 	FScopeLock ScopeLock(&CriticalSection);
+
 	if (StreamingRenderAssetsSyncEvent.IsValid())
 	{
 		StreamingRenderAssetsSyncEvent->Wait(ENamedThreads::GameThread);
@@ -225,6 +233,11 @@ void FRenderAssetStreamingManager::OnPreGarbageCollect()
 
 	QUICK_SCOPE_CYCLE_COUNTER(STAT_FRenderAssetStreamingManager_OnPreGarbageCollect);
 
+	if (CVarFlushDeferredMipLevelChangeCallbacksBeforeGC.GetValueOnGameThread() != 0)
+	{
+		TickDeferredMipLevelChangeCallbacks();
+	}
+
 	FRemovedRenderAssetArray RemovedRenderAssets;
 
 	// Check all levels for pending kills.
