WorldPartition - HLODSubsystem: Use primitive component MaxDrawDistance instead of custom culling logic for persistent HLOD actors culling

diff --git a/Engine/Source/Runtime/Engine/Private/WorldPartition/HLOD/HLODSubsystem.cpp b/Engine/Source/Runtime/Engine/Private/WorldPartition/HLOD/HLODSubsystem.cpp
--- a/Engine/Source/Runtime/Engine/Private/WorldPartition/HLOD/HLODSubsystem.cpp
+++ b/Engine/Source/Runtime/Engine/Private/WorldPartition/HLOD/HLODSubsystem.cpp
@@ -20,6 +20,7 @@
 #include "EngineUtils.h"
 #include "LevelUtils.h"
 #include "Components/StaticMeshComponent.h"
+#include "Components/InstancedStaticMeshComponent.h"
 #include "HAL/FileManager.h"
 #include "Misc/FileHelper.h"
 #include "SceneViewExtension.h"
@@ -113,7 +114,6 @@ public:
 
 UHLODSubsystem::UHLODSubsystem()
 	: UWorldSubsystem()
-	, HLODAlwaysLoadedCullDistance(-1)
 {
 }
 
@@ -123,7 +123,23 @@ UHLODSubsystem::~UHLODSubsystem()
 
 void UHLODSubsystem::SetHLODAlwaysLoadedCullDistance(int32 InCullDistance)
 {
-	HLODAlwaysLoadedCullDistance = InCullDistance;
+	const float MaxDrawDistance = FMath::Max(InCullDistance, 0);
+	const bool bNeverDistanceCull = InCullDistance <= 0;
+
+	for (AWorldPartitionHLOD* HLODActor : AlwaysLoadedHLODActors)
+	{
+		HLODActor->ForEachComponent<UPrimitiveComponent>(false, [MaxDrawDistance, bNeverDistanceCull](UPrimitiveComponent* HLODComponent)
+		{
+			HLODComponent->bNeverDistanceCull = bNeverDistanceCull;
+			HLODComponent->SetCachedMaxDrawDistance(MaxDrawDistance);
+
+			// Enable per instance culling to avoid ISM components dissapearing as a whole
+			if (UInstancedStaticMeshComponent* ISMComponent = Cast<UInstancedStaticMeshComponent>(HLODComponent))
+			{
+				ISMComponent->SetCullDistances(0, MaxDrawDistance);
+			}
+		});
+	}
 }
 
 void UHLODSubsystem::OnWorldBeginPlay(UWorld& InWorld)
@@ -132,96 +148,16 @@ void UHLODSubsystem::OnWorldBeginPlay(UWorld& InWorld)
 
 	check(InWorld.IsGameWorld());
 
-	HLODAlwaysLoadedCullDistance = -1;
 	AlwaysLoadedHLODActors.Reset();
-	CulledAlwaysLoadedHLODActors.Reset();
 
-	// Cache always loaded HLODs with their Bounds
+	// Cache always loaded HLODs
 	for (AActor* Actor : InWorld.PersistentLevel->Actors)
 	{
 		if (AWorldPartitionHLOD* HLODActor = Cast<AWorldPartitionHLOD>(Actor))
 		{
-			FBox Bounds = HLODActor->GetComponentsBoundingBox(true);
-			AlwaysLoadedHLODActors.Add(HLODActor, Bounds);
-		}
-	}
-}
-
-void UHLODSubsystem::UpdateStreamingState()
-{
-	Super::UpdateStreamingState();
-
-	UWorld* World = GetWorld();
-	if (!World || !World->IsGameWorld())
-	{
-		return;
-	}
-
-	UWorldPartition* WorldPartition = World->GetWorldPartition();
-	FWorldPartitionHLODRuntimeData* WorldPartitionHLODRuntimeData = (WorldPartition && WorldPartition->IsStreamingEnabled()) ? WorldPartitionsHLODRuntimeData.Find(WorldPartition) : nullptr;
-	if (!WorldPartitionHLODRuntimeData)
-	{
-		return;
-	}
-
-	TSet<AWorldPartitionHLOD*> FrameCulled;
-	TSet<AWorldPartitionHLOD*> FrameAlreadyHidden;
-	if (HLODAlwaysLoadedCullDistance > 0)
-	{
-		const TArray<FWorldPartitionStreamingSource>& StreamingSources = WorldPartition->GetStreamingSources();
-		double SquaredCullDistance = (double)HLODAlwaysLoadedCullDistance * (double)HLODAlwaysLoadedCullDistance;
-		for (const auto& Pair : AlwaysLoadedHLODActors)
-		{
-			if (AWorldPartitionHLOD* HLODActor = Pair.Key)
-			{
-				const FName CellName = HLODActor->GetSourceCellName();
-				if (FCellData* CellData = WorldPartitionHLODRuntimeData->CellsData.Find(CellName))
-				{
-					// If source cell is visible, hlod actor is already hidden
-					if (!CellData->bIsCellVisible)
-					{
-						if (!StreamingSources.IsEmpty())
-						{
-							bool bIsCulled = true;
-							for (const FWorldPartitionStreamingSource& Source : StreamingSources)
-							{
-								const FBox& Bounds = Pair.Value;
-								double Distance = Bounds.ComputeSquaredDistanceToPoint(Source.Location);
-								if (Distance < SquaredCullDistance)
-								{
-									bIsCulled = false;
-									break;
-								}
-							}
-							if (bIsCulled)
-							{
-								FrameCulled.Add(HLODActor);
-							}
-						}
-					}
-					else
-					{
-						FrameAlreadyHidden.Add(HLODActor);
-					}
-				}
-			}
+			AlwaysLoadedHLODActors.Add(HLODActor);
 		}
 	}
-
-	TSet<AWorldPartitionHLOD*> ToCull = FrameCulled.Difference(CulledAlwaysLoadedHLODActors);
-	TSet<AWorldPartitionHLOD*> ToShow = CulledAlwaysLoadedHLODActors.Difference(FrameCulled.Union(FrameAlreadyHidden));
-
-	for (AWorldPartitionHLOD* HLODActor : ToCull)
-	{
-		HLODActor->SetVisibility(false);
-	}
-
-	for (AWorldPartitionHLOD* HLODActor : ToShow)
-	{
-		HLODActor->SetVisibility(UHLODSubsystem::WorldPartitionHLODEnabled);
-	}
-
-	CulledAlwaysLoadedHLODActors = MoveTemp(FrameCulled);
 }
 
 bool UHLODSubsystem::WorldPartitionHLODEnabled = true;
diff --git a/Engine/Source/Runtime/Engine/Public/WorldPartition/HLOD/HLODSubsystem.h b/Engine/Source/Runtime/Engine/Public/WorldPartition/HLOD/HLODSubsystem.h
--- a/Engine/Source/Runtime/Engine/Public/WorldPartition/HLOD/HLODSubsystem.h
+++ b/Engine/Source/Runtime/Engine/Public/WorldPartition/HLOD/HLODSubsystem.h
@@ -37,7 +37,6 @@ public:
 
 	//~ Begin UWorldSubsystem Interface.
 	virtual bool DoesSupportWorldType(const EWorldType::Type WorldType) const override;
-	virtual void UpdateStreamingState() override;
 	virtual void OnWorldBeginPlay(UWorld& InWorld) override;
 	//~ End UWorldSubsystem Interface.
 
@@ -82,9 +81,7 @@ private:
 	
 	TMap<TObjectPtr<UWorldPartition>, FWorldPartitionHLODRuntimeData> WorldPartitionsHLODRuntimeData;
 
-	TMap<AWorldPartitionHLOD*, FBox> AlwaysLoadedHLODActors;
-	TSet<AWorldPartitionHLOD*> CulledAlwaysLoadedHLODActors;
-	int32 HLODAlwaysLoadedCullDistance;
+	TArray<AWorldPartitionHLOD*> AlwaysLoadedHLODActors;
 
 	void OnWorldPartitionInitialized(UWorldPartition* InWorldPartition);
 	void OnWorldPartitionUninitialized(UWorldPartition* InWorldPartition);
