Add Cvar to enable VolumeTexture and Texture2DArray streaming (on by default as previously).
Add Cvar to force set LastRenderTime on BaseRHI Texture (on by default).

#localization none
#tests PS5

diff --git a/Engine/Source/Runtime/Engine/Classes/Engine/Texture2DArray.cpp b/Engine/Source/Runtime/Engine/Classes/Engine/Texture2DArray.cpp
--- a/Engine/Source/Runtime/Engine/Classes/Engine/Texture2DArray.cpp
+++ b/Engine/Source/Runtime/Engine/Classes/Engine/Texture2DArray.cpp
@@ -32,6 +32,11 @@
 
 // Externed global switch to control whether streaming is enabled for texture2darray. 
 bool GSupportsTexture2DArrayStreaming = true;
+static FAutoConsoleVariableRef CVarSupportsTexture2DArrayStreaming(
+	TEXT("r.SupportsTexture2DArrayStreaming"),
+	GSupportsTexture2DArrayStreaming,
+	TEXT("Enable Support of Texture2DArray Streaming\n")
+);
 
 static TAutoConsoleVariable<int32> CVarAllowTexture2DArrayAssetCreation(
 	TEXT("r.AllowTexture2DArrayCreation"),
diff --git a/Engine/Source/Runtime/Engine/Private/VolumeTexture.cpp b/Engine/Source/Runtime/Engine/Private/VolumeTexture.cpp
--- a/Engine/Source/Runtime/Engine/Private/VolumeTexture.cpp
+++ b/Engine/Source/Runtime/Engine/Private/VolumeTexture.cpp
@@ -31,6 +31,11 @@
 
 // Externed global switch to control whether streaming is enabled for volume texture. 
 bool GSupportsVolumeTextureStreaming = true;
+static FAutoConsoleVariableRef CVarSupportsVolumeTextureStreaming(
+	TEXT("r.SupportsVolumeTextureStreaming"),
+	GSupportsVolumeTextureStreaming,
+	TEXT("Enable Support of VolumeTexture Streaming\n")
+);
 
 // Limit the possible depth of volume texture otherwise when the user converts 2D textures, they can crash the engine.
 const int32 MAX_VOLUME_TEXTURE_DEPTH = 512;
