Fix incorrect SceneTexture look up/unpacking when Strata is enabled.

diff --git a/Engine/Source/Runtime/Engine/Private/Materials/HLSLMaterialTranslator.cpp b/Engine/Source/Runtime/Engine/Private/Materials/HLSLMaterialTranslator.cpp
--- a/Engine/Source/Runtime/Engine/Private/Materials/HLSLMaterialTranslator.cpp
+++ b/Engine/Source/Runtime/Engine/Private/Materials/HLSLMaterialTranslator.cpp
@@ -6630,6 +6630,14 @@ int32 FHLSLMaterialTranslator::SceneTextureLookup(int32 ViewportUV, uint32 InSce
 		LookUp = AddCodeChunk(MCT_Float4, TEXT("MobileSceneTextureLookup(Parameters, %d, %s)"), (int32)SceneTextureId, *CoerceParameter(BufferUV, MCT_Float2));
 	}
 
+	// Strata only
+	// When SceneTexture lookup node is used, single/simple paths are disabled to ensure texture decoding is properly handled.
+	// Reading SceneTexture, when Strata is enabled, implies unpacking material buffer data. The unpacking function exists in different 'flavor' 
+	// for optimization purpose (simple/single/complex). To avoid compiling out single or complex unpacking paths (due to defines set by analyzing 
+	// the current shader, vs. scene texture pixels), we force Simple/Single versions to be disabled
+	bStrataMaterialIsSimple = false;
+	bStrataMaterialIsSingle = false;
+	
 	if (SceneTextureId == PPI_PostProcessInput0 && Material->GetMaterialDomain() == MD_PostProcess && Material->GetBlendableLocation() != BL_AfterTonemapping)
 	{
 		return AddInlinedCodeChunk(MCT_Float4, TEXT("(float4(View.OneOverPreExposure.xxx, 1) * %s)"), *CoerceParameter(LookUp, MCT_Float4));
