[UE-170446] TargetWeightInterpolationSpeedPerSec is not initialized properly in Lyra

diff --git a/Engine/Source/Runtime/Engine/Classes/Animation/BlendSpace.h b/Engine/Source/Runtime/Engine/Classes/Animation/BlendSpace.h
--- a/Engine/Source/Runtime/Engine/Classes/Animation/BlendSpace.h
+++ b/Engine/Source/Runtime/Engine/Classes/Animation/BlendSpace.h
@@ -72,7 +72,7 @@ struct FBlendSpaceBlendProfile
 	TObjectPtr<UBlendProfile> BlendProfile;
 
 	UPROPERTY(EditAnywhere, Category = SampleSmoothing, meta = (DisplayName = "Weight Speed", ClampMin = "0"))
-	float TargetWeightInterpolationSpeedPerSec;
+	float TargetWeightInterpolationSpeedPerSec = 0.0f;
 };
 
 USTRUCT()
@@ -750,7 +750,7 @@ public:
 	* value of zero disables this smoothing entirely.
 	*/
 	UPROPERTY(EditAnywhere, Category = SampleSmoothing, meta = (DisplayName = "Weight Speed", ClampMin = "0"))
-	float TargetWeightInterpolationSpeedPerSec;
+	float TargetWeightInterpolationSpeedPerSec = 0.0f;
 
 	/**
 	 * If set then this eases in/out the sample weight adjustments, using the speed to determine how much smoothing to apply.
