Iris - Fixed issue with the wrong FNetRefHandle being used when updating the owning net connection.

diff --git a/Engine/Source/Runtime/Engine/Private/ActorReplication.cpp b/Engine/Source/Runtime/Engine/Private/ActorReplication.cpp
--- a/Engine/Source/Runtime/Engine/Private/ActorReplication.cpp
+++ b/Engine/Source/Runtime/Engine/Private/ActorReplication.cpp
@@ -962,7 +962,7 @@ void AActor::UpdateOwningNetConnection() const
 
 		if (Actor->GetIsReplicated())
 		{
-			const UE::Net::FNetRefHandle RefHandle = ObjectReplicationBridge->GetReplicatedRefHandle(this);
+			const UE::Net::FNetRefHandle RefHandle = ObjectReplicationBridge->GetReplicatedRefHandle(Actor);
 			if (RefHandle.IsValid())
 			{
 				ReplicationSystem->SetOwningNetConnection(RefHandle, NewOwningNetConnectionId);
