ColorManagement: Working color space support in sky atmosphere.

diff --git a/Engine/Source/Runtime/Engine/Public/Rendering/SkyAtmosphereCommonData.cpp b/Engine/Source/Runtime/Engine/Public/Rendering/SkyAtmosphereCommonData.cpp
--- a/Engine/Source/Runtime/Engine/Public/Rendering/SkyAtmosphereCommonData.cpp
+++ b/Engine/Source/Runtime/Engine/Public/Rendering/SkyAtmosphereCommonData.cpp
@@ -7,6 +7,7 @@
 #include "SkyAtmosphereCommonData.h"
 
 #include "Components/SkyAtmosphereComponent.h"
+#include "ColorSpace.h"
 
 //PRAGMA_DISABLE_OPTIMIZATION
 
@@ -44,17 +45,63 @@ FAtmosphereSetup::FAtmosphereSetup(const USkyAtmosphereComponent& SkyAtmosphereC
 	GroundAlbedo = FLinearColor(SkyAtmosphereComponent.GroundAlbedo);
 	MultiScatteringFactor = FMath::Clamp(SkyAtmosphereComponent.MultiScatteringFactor, 0.0f, 2.0f);
 
-	RayleighDensityExpScale = -1.0f / SkyAtmosphereComponent.RayleighExponentialDistribution;
-	RayleighScattering = (SkyAtmosphereComponent.RayleighScattering * SkyAtmosphereComponent.RayleighScatteringScale).GetClamped(0.0f, 1e38f);
+	auto ConvertCoefficientsFromSRGBToWorkingColorSpace = [](FLinearColor CoeffSRGB)
+	{
+		using namespace UE::Color;
+
+		const FColorSpace& WorkingColorSpace = FColorSpace::GetWorking();
+		if (WorkingColorSpace.IsSRGB())
+		{
+			return CoeffSRGB;
+		}
+		else
+		{
+			// Compute the transmittance color from the coefficients.
+			FLinearColor Transmittance = FLinearColor(
+				FMath::Exp(-CoeffSRGB.R),
+				FMath::Exp(-CoeffSRGB.G),
+				FMath::Exp(-CoeffSRGB.B));
+
+			// Convert transmittance color from sRGB to working color space.
+			Transmittance = FColorSpaceTransform(FColorSpace(EColorSpace::sRGB), WorkingColorSpace).Apply(Transmittance);
+
+			// New we have a transmittance in working color space, convert it back to coefficients for this working color space.
+			return FLinearColor(
+				-FMath::Loge(FMath::Max(0.00001, Transmittance.R)),
+				-FMath::Loge(FMath::Max(0.00001, Transmittance.G)),
+				-FMath::Loge(FMath::Max(0.00001, Transmittance.B)));
+		}
+	};
+
+	// Rayleigh scattering
+	{
+		RayleighScattering = (SkyAtmosphereComponent.RayleighScattering * SkyAtmosphereComponent.RayleighScatteringScale).GetClamped(0.0f, 1e38f);
+		RayleighScattering = ConvertCoefficientsFromSRGBToWorkingColorSpace(RayleighScattering);
+
+		RayleighDensityExpScale = -1.0f / SkyAtmosphereComponent.RayleighExponentialDistribution;
+	}
+
+	// Mie scattering
+	{
+
+		MieScattering = (SkyAtmosphereComponent.MieScattering * SkyAtmosphereComponent.MieScatteringScale).GetClamped(0.0f, 1e38f);
+		MieScattering = ConvertCoefficientsFromSRGBToWorkingColorSpace(MieScattering);
 
-	MieScattering = (SkyAtmosphereComponent.MieScattering * SkyAtmosphereComponent.MieScatteringScale).GetClamped(0.0f, 1e38f);
-	MieAbsorption = (SkyAtmosphereComponent.MieAbsorption * SkyAtmosphereComponent.MieAbsorptionScale).GetClamped(0.0f, 1e38f);
-	MieExtinction = MieScattering + MieAbsorption;
-	MiePhaseG = SkyAtmosphereComponent.MieAnisotropy;
-	MieDensityExpScale = -1.0f / SkyAtmosphereComponent.MieExponentialDistribution;
+		MieAbsorption = (SkyAtmosphereComponent.MieAbsorption * SkyAtmosphereComponent.MieAbsorptionScale).GetClamped(0.0f, 1e38f);
+		MieAbsorption = ConvertCoefficientsFromSRGBToWorkingColorSpace(MieAbsorption);
 
-	AbsorptionExtinction = (SkyAtmosphereComponent.OtherAbsorption * SkyAtmosphereComponent.OtherAbsorptionScale).GetClamped(0.0f, 1e38f);
-	TentToCoefficients(SkyAtmosphereComponent.OtherTentDistribution, AbsorptionDensity0LayerWidth, AbsorptionDensity0LinearTerm, AbsorptionDensity1LinearTerm, AbsorptionDensity0ConstantTerm, AbsorptionDensity1ConstantTerm);
+		MieExtinction = MieScattering + MieAbsorption;
+		MiePhaseG = SkyAtmosphereComponent.MieAnisotropy;
+		MieDensityExpScale = -1.0f / SkyAtmosphereComponent.MieExponentialDistribution;
+	}
+
+	// Ozone
+	{
+		AbsorptionExtinction = (SkyAtmosphereComponent.OtherAbsorption * SkyAtmosphereComponent.OtherAbsorptionScale).GetClamped(0.0f, 1e38f);
+		AbsorptionExtinction = ConvertCoefficientsFromSRGBToWorkingColorSpace(AbsorptionExtinction);
+
+		TentToCoefficients(SkyAtmosphereComponent.OtherTentDistribution, AbsorptionDensity0LayerWidth, AbsorptionDensity0LinearTerm, AbsorptionDensity1LinearTerm, AbsorptionDensity0ConstantTerm, AbsorptionDensity1ConstantTerm);
+	}
 
 	TransmittanceMinLightElevationAngle = SkyAtmosphereComponent.TransmittanceMinLightElevationAngle;
 
