BodySetup: Stop tracking cache misses when the cache query is being skipped

#rnx

diff --git a/Engine/Source/Runtime/Engine/Private/PhysicsEngine/BodySetup.cpp b/Engine/Source/Runtime/Engine/Private/PhysicsEngine/BodySetup.cpp
--- a/Engine/Source/Runtime/Engine/Private/PhysicsEngine/BodySetup.cpp
+++ b/Engine/Source/Runtime/Engine/Private/PhysicsEngine/BodySetup.cpp
@@ -1336,6 +1336,7 @@ void GetDDCBuiltData(FByteBulkData* OutResult, DDCBuilderType& InBuilder, UBodyS
 		{
 			bDataWasBuilt = true;
 			InBuilder.Build(OutData);
+			COOK_STAT(Timer.TrackCyclesOnly());
 		}
 	}
 
