Revising option for base-relative acceleration to be global and defaulted to on.

diff --git a/Engine/Source/Runtime/Engine/Classes/GameFramework/CharacterMovementComponent.h b/Engine/Source/Runtime/Engine/Classes/GameFramework/CharacterMovementComponent.h
--- a/Engine/Source/Runtime/Engine/Classes/GameFramework/CharacterMovementComponent.h
+++ b/Engine/Source/Runtime/Engine/Classes/GameFramework/CharacterMovementComponent.h
@@ -750,15 +750,6 @@ public:
 	UPROPERTY(Category="Character Movement (General Settings)", EditAnywhere, BlueprintReadWrite, AdvancedDisplay, meta=(ClampMin="0", UIMin="0", ForceUnits=cm))
 	float MaxDepenetrationWithPawnAsProxy;
 
-	/**
-	 * Option to handle moves as relative to the current movement base. If disabled, moves are handled in world space. 
-	 * Enabling this can be produce better motion sync when walking on highly dynamic / unpredictable movement bases, such as vehicles. However, there is the potential for larger corrections when landing on or jumping off if the base is not well-sync'd.
-	 * This also may be a good option if using dynamic movement bases that become the player's primary frame of visual reference, such as when walking on a large boat or airship.
-	 * Not compatible with deprecated move RPCs. @see NetUsePackedMovementRPCs
-	 */
-	UPROPERTY(Category = "Character Movement (General Settings)", EditDefaultsOnly)
-	bool bUseMovementBaseRelativeMoves = false;
-
 	/**
 	 * How long to take to smoothly interpolate from the old pawn position on the client to the corrected one sent by the server. Not used by Linear smoothing.
 	 */
diff --git a/Engine/Source/Runtime/Engine/Private/Components/CharacterMovementComponent.cpp b/Engine/Source/Runtime/Engine/Private/Components/CharacterMovementComponent.cpp
--- a/Engine/Source/Runtime/Engine/Private/Components/CharacterMovementComponent.cpp
+++ b/Engine/Source/Runtime/Engine/Private/Components/CharacterMovementComponent.cpp
@@ -250,11 +250,17 @@ namespace CharacterMovementCVars
 		TEXT("p.BasedMovementMode"),
 		BasedMovementMode, TEXT("0 means always on regular tick (default); 1 means only if not deferring updates; 2 means update and save based movement both on regular ticks and post physics when on a physics base."));
 
-	static int32 AllowBaseRelativeMovement = 1;
-	FAutoConsoleVariableRef CVarAllowBaseRelativeMovement(
-		TEXT("p.AllowBaseRelativeMovement"),
-		AllowBaseRelativeMovement, 
-		TEXT("0=disabled, otherwise enabled. If enabled, resepct the bUseMovementBaseRelativeMoves option on each character. If disabled, all movement will be handled in worldspace."));
+	/**
+	 * Option to handle move acceleration as relative to dynamic movement bases. If disabled, moves are handled in world space.
+	 * Enabling this will produce better motion when walking on highly dynamic / unpredictable movement bases, such as vehicles. However, there is the potential for larger corrections when landing on or jumping off if the base's rotation is not well-sync'd.
+	 * This also may be a good option if using dynamic movement bases that become the player's primary frame of visual reference, such as when walking on a large boat or airship.
+	 * Not compatible with deprecated move RPCs. @see NetUsePackedMovementRPCs
+	 */
+	static int32 UseBaseRelativeAcceleration = 1;
+	FAutoConsoleVariableRef CVarUseBaseRelativeAcceleration(
+		TEXT("p.UseBaseRelativeAcceleration"),
+		UseBaseRelativeAcceleration,
+		TEXT("If enabled, character acceleration will be treated as relative to dynamic movement bases."));
 
 	static int32 UseTargetVelocityOnImpact = 1;
 	FAutoConsoleVariableRef CVarUseTargetVelocityOnImpact(
@@ -9020,9 +9026,7 @@ void FCharacterNetworkMoveData::ClientFillNetworkMoveData(const FSavedMove_Chara
 		UPrimitiveComponent* ClientMovementBase = ClientMove.EndBase.Get();
 
 		const bool bSendBaseRelativeLocation     = MovementBaseUtility::UseRelativeLocation(ClientMovementBase);
-		const bool bSendBaseRelativeAcceleration = CharacterMovementCVars::AllowBaseRelativeMovement 
-													&& bSendBaseRelativeLocation 
-													&& ClientMove.CharacterOwner->GetCharacterMovement()->bUseMovementBaseRelativeMoves;
+		const bool bSendBaseRelativeAcceleration = CharacterMovementCVars::UseBaseRelativeAcceleration && bSendBaseRelativeLocation;
 
 		const FVector SendLocation     = bSendBaseRelativeLocation ? ClientMove.SavedRelativeLocation : FRepMovement::RebaseOntoZeroOrigin(ClientMove.SavedLocation, ClientMove.CharacterOwner->GetCharacterMovement());
 		const FVector SendAcceleration = bSendBaseRelativeAcceleration ? ClientMove.SavedRelativeAcceleration : ClientMove.Acceleration;
@@ -9183,9 +9187,7 @@ void UCharacterMovementComponent::ServerMove_PerformMovement(const FCharacterNet
 	FVector ClientAccel = MoveData.Acceleration;
 
 	// Convert the move's acceleration to worldspace if necessary
-	if (CharacterMovementCVars::AllowBaseRelativeMovement 
-			&& bUseMovementBaseRelativeMoves 
-			&& MovementBaseUtility::IsDynamicBase(MoveData.MovementBase))
+	if (CharacterMovementCVars::UseBaseRelativeAcceleration && MovementBaseUtility::IsDynamicBase(MoveData.MovementBase))
 	{
 		MovementBaseUtility::GetLocalMovementBaseAccelerationInWorldSpace(MoveData.MovementBase, MoveData.MovementBaseBoneName, MoveData.Acceleration, ClientAccel);
 	}
@@ -11816,9 +11818,7 @@ void FSavedMove_Character::PostUpdate(ACharacter* Character, FSavedMove_Characte
 		}
 
 		// Save off movement base-relative acceleration if needed
-		if (CharacterMovementCVars::AllowBaseRelativeMovement 
-				&& Character->GetCharacterMovement()->bUseMovementBaseRelativeMoves
-				&& MovementBaseUtility::IsDynamicBase(MovementBase))
+		if (CharacterMovementCVars::UseBaseRelativeAcceleration && MovementBaseUtility::IsDynamicBase(MovementBase))
 		{
 			MovementBaseUtility::GetLocalMovementBaseAcceleration(EndBase.Get(), EndBoneName, Acceleration, SavedRelativeAcceleration);
 		}
