Fix OrientAndScale retargeting in the presence of skeleton remapping

Fixes a crash accessing authored retarget base transforms out of range.
This was always using the 'target' (i.e. the mesh) skeleton indices, rather than the 'source' (i.e. the anim) to index the retarget source transform array. Accounting for the remapping means we access the correct transforms.

diff --git a/Engine/Source/Runtime/Engine/Private/Animation/AnimationRuntime.cpp b/Engine/Source/Runtime/Engine/Private/Animation/AnimationRuntime.cpp
--- a/Engine/Source/Runtime/Engine/Private/Animation/AnimationRuntime.cpp
+++ b/Engine/Source/Runtime/Engine/Private/Animation/AnimationRuntime.cpp
@@ -2655,7 +2655,7 @@ void FAnimationRuntime::RetargetBoneTransform(const USkeleton* SourceSkeleton, c
 			{
 				if (!bIsBakedAdditive)
 				{
-					const FRetargetSourceCachedData& RetargetSourceCachedData = RequiredBones.GetRetargetSourceCachedData(SourceName, RetargetTransforms);
+					const FRetargetSourceCachedData& RetargetSourceCachedData = RequiredBones.GetRetargetSourceCachedData(SourceName, SkeletonRemapping, RetargetTransforms);
 					const TArray<FOrientAndScaleRetargetingCachedData>& OrientAndScaleDataArray = RetargetSourceCachedData.OrientAndScaleData;
 					const TArray<int32>& CompactPoseIndexToOrientAndScaleIndex = RetargetSourceCachedData.CompactPoseIndexToOrientAndScaleIndex;
 
diff --git a/Engine/Source/Runtime/Engine/Private/BoneContainer.cpp b/Engine/Source/Runtime/Engine/Private/BoneContainer.cpp
--- a/Engine/Source/Runtime/Engine/Private/BoneContainer.cpp
+++ b/Engine/Source/Runtime/Engine/Private/BoneContainer.cpp
@@ -433,10 +433,10 @@ const FRetargetSourceCachedData& FBoneContainer::GetRetargetSourceCachedData(con
 {
 	LLM_SCOPE_BYNAME(TEXT("Animation/BoneContainer"));
 	const TArray<FTransform>& RetargetTransforms = AssetSkeleton->GetRefLocalPoses(InRetargetSourceName);
-	return GetRetargetSourceCachedData(InRetargetSourceName, RetargetTransforms);
+	return GetRetargetSourceCachedData(InRetargetSourceName, FSkeletonRemapping(), RetargetTransforms);
 }
 
-const FRetargetSourceCachedData& FBoneContainer::GetRetargetSourceCachedData(const FName& InSourceName, const TArray<FTransform>& InRetargetTransforms) const
+const FRetargetSourceCachedData& FBoneContainer::GetRetargetSourceCachedData(const FName& InSourceName, const FSkeletonRemapping& InRemapping, const TArray<FTransform>& InRetargetTransforms) const
 {
 	LLM_SCOPE_BYNAME(TEXT("Animation/BoneContainer"));
 	FRetargetSourceCachedData* RetargetSourceCachedData = RetargetSourceCachedDataLUT.Find(InSourceName);
@@ -454,11 +454,12 @@ const FRetargetSourceCachedData& FBoneContainer::GetRetargetSourceCachedData(con
 
 		for (int32 CompactBoneIndex = 0; CompactBoneIndex < CompactPoseNumBones; CompactBoneIndex++)
 		{
-			const int32& SkeletonBoneIndex = CompactPoseToSkeletonIndex[CompactBoneIndex];
+			const int32 TargetSkeletonBoneIndex = CompactPoseToSkeletonIndex[CompactBoneIndex];
+			const int32 SourceSkeletonBoneIndex = InRemapping.IsValid() ? InRemapping.GetSourceSkeletonBoneIndex(TargetSkeletonBoneIndex) : TargetSkeletonBoneIndex;
 
-			if (AssetSkeleton->GetBoneTranslationRetargetingMode(SkeletonBoneIndex) == EBoneTranslationRetargetingMode::OrientAndScale)
+			if (SourceSkeletonBoneIndex != INDEX_NONE && AssetSkeleton->GetBoneTranslationRetargetingMode(SourceSkeletonBoneIndex) == EBoneTranslationRetargetingMode::OrientAndScale)
 			{
-				const FVector SourceSkelTrans = AuthoredOnRefSkeleton[SkeletonBoneIndex].GetTranslation();
+				const FVector SourceSkelTrans = AuthoredOnRefSkeleton[SourceSkeletonBoneIndex].GetTranslation();
 				FVector TargetSkelTrans;
 				if (RefPoseOverride.IsValid())
 				{
diff --git a/Engine/Source/Runtime/Engine/Public/BoneContainer.h b/Engine/Source/Runtime/Engine/Public/BoneContainer.h
--- a/Engine/Source/Runtime/Engine/Public/BoneContainer.h
+++ b/Engine/Source/Runtime/Engine/Public/BoneContainer.h
@@ -661,7 +661,7 @@ public:
 	void CacheRequiredAnimCurveUids(const FCurveEvaluationOption& CurveEvalOption);
 
 	const FRetargetSourceCachedData& GetRetargetSourceCachedData(const FName& InRetargetSource) const;
-	const FRetargetSourceCachedData& GetRetargetSourceCachedData(const FName& InSourceName, const TArray<FTransform>& InRetargetTransforms) const;
+	const FRetargetSourceCachedData& GetRetargetSourceCachedData(const FName& InSourceName, const FSkeletonRemapping& InRemapping, const TArray<FTransform>& InRetargetTransforms) const;
 
 #if DO_CHECK
 	/** Get the LOD that we calculated required bones when regenerated */
