Fix non-unity CIS

Post robomerge fixup

diff --git a/Engine/Source/Runtime/Engine/Private/Animation/AnimationDecompression.cpp b/Engine/Source/Runtime/Engine/Private/Animation/AnimationDecompression.cpp
--- a/Engine/Source/Runtime/Engine/Private/Animation/AnimationDecompression.cpp
+++ b/Engine/Source/Runtime/Engine/Private/Animation/AnimationDecompression.cpp
@@ -260,7 +260,7 @@ void DecompressPose(FCompactPose& OutPose,
 	const int32 NumBonesToOrientAndScaleRetarget = OrientAndScaleRetargetingPairs.Num();
 	if (NumBonesToOrientAndScaleRetarget > 0)
 	{
-		const FRetargetSourceCachedData& RetargetSourceCachedData = RequiredBones.GetRetargetSourceCachedData(DecompressionContext.AnimName, RetargetTransforms);
+		const FRetargetSourceCachedData& RetargetSourceCachedData = RequiredBones.GetRetargetSourceCachedData(DecompressionContext.AnimName, SkeletonRemapping, RetargetTransforms);
 		const TArray<FOrientAndScaleRetargetingCachedData>& OrientAndScaleDataArray = RetargetSourceCachedData.OrientAndScaleData;
 		const TArray<int32>& CompactPoseIndexToOrientAndScaleIndex = RetargetSourceCachedData.CompactPoseIndexToOrientAndScaleIndex;
 
