Add PIE state delegates to the Input Device Subsystem so that device properties will not be evaluated while PIE is stopped, and they will be cleared when PIE starts/stops. This makes iteration time for making properties faster for designers and more consistent

#rnx

diff --git a/Engine/Source/Runtime/Engine/Classes/GameFramework/InputDeviceSubsystem.h b/Engine/Source/Runtime/Engine/Classes/GameFramework/InputDeviceSubsystem.h
--- a/Engine/Source/Runtime/Engine/Classes/GameFramework/InputDeviceSubsystem.h
+++ b/Engine/Source/Runtime/Engine/Classes/GameFramework/InputDeviceSubsystem.h
@@ -191,6 +191,17 @@ protected:
 
 	/** Set the most recently used hardware device */
 	void SetMostRecentlyUsedHardwareDevice(const FInputDeviceId InDeviceId, const FHardwareDeviceIdentifier& InHardwareId);
+
+	// Callbacks for when PIE is started/stopped. We will likely want to pause/resume input device properties
+	// when this happens, or just remove all active properties when PIE stops. This will make designer iteration a little easier
+#if WITH_EDITOR
+	void OnPrePIEStarted(bool bSimulating);
+	void OnPIEPaused(bool bSimulating);
+	void OnPIEResumed(bool bSimulating);
+	void OnPIEStopped(bool bSimulating);
+
+	bool bIsPIEPlaying = false;
+#endif
 	
 	/**
 	* Array of currently active input device properties that will be evaluated on tick
diff --git a/Engine/Source/Runtime/Engine/Private/GameFramework/InputDeviceSubsystem.cpp b/Engine/Source/Runtime/Engine/Private/GameFramework/InputDeviceSubsystem.cpp
--- a/Engine/Source/Runtime/Engine/Private/GameFramework/InputDeviceSubsystem.cpp
+++ b/Engine/Source/Runtime/Engine/Private/GameFramework/InputDeviceSubsystem.cpp
@@ -10,6 +10,10 @@
 #include "GenericPlatform/GenericPlatformInputDeviceMapper.h"
 #include "Framework/Application/SlateApplication.h"			// For RegisterInputPreProcessor
 
+#if WITH_EDITOR
+#include "Editor.h"		// For PIE delegates
+#endif
+
 DEFINE_LOG_CATEGORY(LogInputDeviceProperties);
 
 ////////////////////////////////////////////////////////
@@ -157,6 +161,13 @@ void UInputDeviceSubsystem::Initialize(FSubsystemCollectionBase& Collection)
 	
 	InputPreprocessor = MakeShared<FInputDeviceSubsystemProcessor>();
 	FSlateApplication::Get().RegisterInputPreProcessor(InputPreprocessor, 0);
+
+#if WITH_EDITOR
+	FEditorDelegates::PreBeginPIE.AddUObject(this, &UInputDeviceSubsystem::OnPrePIEStarted);
+	FEditorDelegates::PausePIE.AddUObject(this, &UInputDeviceSubsystem::OnPIEPaused);
+	FEditorDelegates::ResumePIE.AddUObject(this, &UInputDeviceSubsystem::OnPIEResumed);
+	FEditorDelegates::EndPIE.AddUObject(this, &UInputDeviceSubsystem::OnPIEStopped);
+#endif	// WITH_EDITOR
 }
 
 void UInputDeviceSubsystem::Deinitialize()
@@ -225,6 +236,14 @@ ETickableTickType UInputDeviceSubsystem::GetTickableTickType() const
 
 bool UInputDeviceSubsystem::IsAllowedToTick() const
 {
+#if WITH_EDITOR
+	// If we are PIE'ing, then check if PIE is paused
+	if (GEditor->bIsSimulatingInEditor || GEditor->PlayWorld)
+	{
+		return bIsPIEPlaying && !ActiveProperties.IsEmpty();
+	}	
+#endif
+
 	// Only tick when there are active device properties
 	return !ActiveProperties.IsEmpty();
 }
@@ -287,6 +306,7 @@ FInputDevicePropertyHandle UInputDeviceSubsystem::SetDeviceProperty(const FSetDe
 		FActiveDeviceProperty ActiveProp = {};
 
 		// Spawn an instance of this device property
+		// TODO: Possible performance problems with DuplicateObject because FDuplicateDataWriter is not very performant
 		ActiveProp.Property = DuplicateObject<UInputDeviceProperty>(Params.DeviceProperty, /* Outer = */ this);
 		ensure(ActiveProp.Property);
 
@@ -405,3 +425,29 @@ void UInputDeviceSubsystem::SetMostRecentlyUsedHardwareDevice(const FInputDevice
 		OnInputHardwareDeviceChanged.Broadcast(OwningUserId, InDeviceId);
 	}
 }
+
+#if WITH_EDITOR
+void UInputDeviceSubsystem::OnPrePIEStarted(bool bSimulating)
+{
+	// Remove all active properties, just in case someone was previewing something in the editor that are still going
+	RemoveAllDeviceProperties();
+	bIsPIEPlaying = true;
+}
+
+void UInputDeviceSubsystem::OnPIEPaused(bool bSimulating)
+{
+	bIsPIEPlaying = false;
+}
+
+void UInputDeviceSubsystem::OnPIEResumed(bool bSimulating)
+{
+	bIsPIEPlaying = true;
+}
+
+void UInputDeviceSubsystem::OnPIEStopped(bool bSimulating)
+{
+	// Remove all active properties when PIE stops
+	RemoveAllDeviceProperties();
+	bIsPIEPlaying = false;
+}
+#endif
\ No newline at end of file
