Fix for Android resetting FPS settings between launches
#tests Android Dev Client

diff --git a/Engine/Source/Runtime/Engine/Private/GameUserSettings.cpp b/Engine/Source/Runtime/Engine/Private/GameUserSettings.cpp
--- a/Engine/Source/Runtime/Engine/Private/GameUserSettings.cpp
+++ b/Engine/Source/Runtime/Engine/Private/GameUserSettings.cpp
@@ -76,7 +76,7 @@ FIntPoint UGameUserSettings::GetDesktopResolution() const
 
 void UGameUserSettings::SetScreenResolution(FIntPoint Resolution)
 {
-	if (!FPlatformProperties::HasFixedResolution() && (ResolutionSizeX != Resolution.X || ResolutionSizeY != Resolution.Y))
+	if (ResolutionSizeX != Resolution.X || ResolutionSizeY != Resolution.Y)
 	{
 		ResolutionSizeX = Resolution.X;
 		ResolutionSizeY = Resolution.Y;
