Fix CIS crash in input device subsystem where GEditor would be null

#rnx

diff --git a/Engine/Source/Runtime/Engine/Private/GameFramework/InputDeviceSubsystem.cpp b/Engine/Source/Runtime/Engine/Private/GameFramework/InputDeviceSubsystem.cpp
--- a/Engine/Source/Runtime/Engine/Private/GameFramework/InputDeviceSubsystem.cpp
+++ b/Engine/Source/Runtime/Engine/Private/GameFramework/InputDeviceSubsystem.cpp
@@ -238,7 +238,7 @@ bool UInputDeviceSubsystem::IsAllowedToTick() const
 {
 #if WITH_EDITOR
 	// If we are PIE'ing, then check if PIE is paused
-	if (GEditor->bIsSimulatingInEditor || GEditor->PlayWorld)
+	if (GEditor && (GEditor->bIsSimulatingInEditor || GEditor->PlayWorld))
 	{
 		return bIsPIEPlaying && !ActiveProperties.IsEmpty();
 	}	
