Fixed crash when running with -server with map containing Actor Folders and fixup of a folder is necessary.

diff --git a/Engine/Source/Runtime/Engine/Private/Level.cpp b/Engine/Source/Runtime/Engine/Private/Level.cpp
--- a/Engine/Source/Runtime/Engine/Private/Level.cpp
+++ b/Engine/Source/Runtime/Engine/Private/Level.cpp
@@ -436,7 +436,7 @@ bool FLevelSimplificationDetails::operator == (const FLevelSimplificationDetails
 static bool IsActorFolderObjectsFeatureAvailable()
 {
 #if WITH_EDITOR
-	return !IsRunningCookCommandlet() && !IsRunningGame();
+	return !IsRunningCookCommandlet() && !IsRunningGame() && GIsEditor;
 #else
 	return false;
 #endif
