Expose FHardwareDeviceIdentifier with ENGINE_API so that it can be used in plugins

#rnx

diff --git a/Engine/Source/Runtime/Engine/Classes/GameFramework/InputSettings.h b/Engine/Source/Runtime/Engine/Classes/GameFramework/InputSettings.h
--- a/Engine/Source/Runtime/Engine/Classes/GameFramework/InputSettings.h
+++ b/Engine/Source/Runtime/Engine/Classes/GameFramework/InputSettings.h
@@ -287,7 +287,7 @@ private:
 * These mappings should match a FInputDeviceScope that is used by an IInputDevice
 */
 USTRUCT(BlueprintType)
-struct FHardwareDeviceIdentifier
+struct ENGINE_API FHardwareDeviceIdentifier
 {
 	GENERATED_BODY()
 
