Add outputs to WorldPosition material expression to directly access XY and Z which are commonly used.

diff --git a/Engine/Source/Runtime/Engine/Private/Materials/MaterialExpressions.cpp b/Engine/Source/Runtime/Engine/Private/Materials/MaterialExpressions.cpp
--- a/Engine/Source/Runtime/Engine/Private/Materials/MaterialExpressions.cpp
+++ b/Engine/Source/Runtime/Engine/Private/Materials/MaterialExpressions.cpp
@@ -12479,6 +12479,12 @@ UMaterialExpressionWorldPosition::UMaterialExpressionWorldPosition(const FObject
 
 	MenuCategories.Add(ConstructorStatics.NAME_Coordinates);
 
+	Outputs.Reset();
+	Outputs.Add(FExpressionOutput(TEXT("XYZ"), 1, 1, 1, 1, 0));
+	Outputs.Add(FExpressionOutput(TEXT("XY"), 1, 1, 1, 0, 0));
+	Outputs.Add(FExpressionOutput(TEXT("Z"), 1, 0, 0, 1, 0));
+
+	bShowOutputNameOnPin = true;
 	bShaderInputData = true;
 #endif
 	WorldPositionShaderOffset = WPT_Default;
