Fix a bug of InstancedStaticMesh missing vertex color on mobile.

diff --git a/Engine/Source/Runtime/Engine/Private/InstancedStaticMesh.cpp b/Engine/Source/Runtime/Engine/Private/InstancedStaticMesh.cpp
--- a/Engine/Source/Runtime/Engine/Private/InstancedStaticMesh.cpp
+++ b/Engine/Source/Runtime/Engine/Private/InstancedStaticMesh.cpp
@@ -1047,6 +1047,7 @@ void FInstancedStaticMeshRenderData::BindBuffersToVertexFactories()
 
 	check(IsInRenderingThread());
 
+	const bool bRHISupportsManualVertexFetch = RHISupportsManualVertexFetch(GShaderPlatformForFeatureLevel[FeatureLevel]);
 	const bool bCanUseGPUScene = UseGPUScene(GMaxRHIShaderPlatform, FeatureLevel);
 	if (!bCanUseGPUScene)
 	{
@@ -1075,7 +1076,7 @@ void FInstancedStaticMeshRenderData::BindBuffersToVertexFactories()
 		}
 		else
 		{
-			FColorVertexBuffer::BindDefaultColorVertexBuffer(&VertexFactory, Data, FColorVertexBuffer::NullBindStride::FColorSizeForComponentOverride);
+			FColorVertexBuffer::BindDefaultColorVertexBuffer(&VertexFactory, Data, bRHISupportsManualVertexFetch ? FColorVertexBuffer::NullBindStride::FColorSizeForComponentOverride : FColorVertexBuffer::NullBindStride::ZeroForDefaultBufferBind);
 		}
 
 		check(PerInstanceRenderData);
