Renamed specular color and edge to F0 and F90 because specular color has a different meanging in legacy. F0/F90 is clearer and less confusing.

diff --git a/Engine/Source/Runtime/Engine/Private/Materials/MaterialExpressions.cpp b/Engine/Source/Runtime/Engine/Private/Materials/MaterialExpressions.cpp
--- a/Engine/Source/Runtime/Engine/Private/Materials/MaterialExpressions.cpp
+++ b/Engine/Source/Runtime/Engine/Private/Materials/MaterialExpressions.cpp
@@ -22682,11 +22682,11 @@ FName UMaterialExpressionStrataSlabBSDF::GetInputName(int32 InputIndex) const
 		}
 		else if (InputIndex == 1)
 		{
-			return TEXT("Specular Color");
+			return TEXT("F0");
 		}
 		else if (InputIndex == 2)
 		{
-			return  TEXT("Edge Specular Color");
+			return  TEXT("F90");
 		}
 	}
 
