Fix material instance blend mode when Strata is enabled.

diff --git a/Engine/Source/Runtime/Engine/Classes/Materials/MaterialInstance.h b/Engine/Source/Runtime/Engine/Classes/Materials/MaterialInstance.h
--- a/Engine/Source/Runtime/Engine/Classes/Materials/MaterialInstance.h
+++ b/Engine/Source/Runtime/Engine/Classes/Materials/MaterialInstance.h
@@ -796,6 +796,7 @@ public:
 	ENGINE_API virtual float GetOpacityMaskClipValue() const override;
 	ENGINE_API virtual bool GetCastDynamicShadowAsMasked() const override;
 	ENGINE_API virtual EBlendMode GetBlendMode() const override;
+	ENGINE_API virtual EStrataBlendMode GetStrataBlendMode() const override;
 	ENGINE_API virtual FMaterialShadingModelField GetShadingModels() const override;
 	ENGINE_API virtual bool IsShadingModelFromMaterialExpression() const override;
 	ENGINE_API virtual bool IsTwoSided() const override;
diff --git a/Engine/Source/Runtime/Engine/Classes/Materials/MaterialInstanceDynamic.h b/Engine/Source/Runtime/Engine/Classes/Materials/MaterialInstanceDynamic.h
--- a/Engine/Source/Runtime/Engine/Classes/Materials/MaterialInstanceDynamic.h
+++ b/Engine/Source/Runtime/Engine/Classes/Materials/MaterialInstanceDynamic.h
@@ -196,6 +196,7 @@ class ENGINE_API UMaterialInstanceDynamic : public UMaterialInstance
 	virtual FMaterialShadingModelField GetShadingModels() const override;
 	virtual bool IsShadingModelFromMaterialExpression() const override;
 	virtual EBlendMode GetBlendMode() const override;
+	virtual EStrataBlendMode GetStrataBlendMode() const override;
 	virtual bool IsTwoSided() const override;
 	virtual bool IsTranslucencyWritingVelocity() const override;
 	virtual bool IsDitheredLODTransition() const override;
diff --git a/Engine/Source/Runtime/Engine/Private/Materials/MaterialInstance.cpp b/Engine/Source/Runtime/Engine/Private/Materials/MaterialInstance.cpp
--- a/Engine/Source/Runtime/Engine/Private/Materials/MaterialInstance.cpp
+++ b/Engine/Source/Runtime/Engine/Private/Materials/MaterialInstance.cpp
@@ -1645,6 +1645,11 @@ EBlendMode UMaterialInstanceDynamic::GetBlendMode() const
 	return Parent ? Parent->GetBlendMode() : BLEND_Opaque;
 }
 
+EStrataBlendMode UMaterialInstanceDynamic::GetStrataBlendMode() const
+{
+	return Parent ? Parent->GetStrataBlendMode() : SBM_Opaque;
+}
+
 bool UMaterialInstanceDynamic::IsTwoSided() const
 {
 	return Parent ? Parent->IsTwoSided() : false;
@@ -4231,6 +4236,11 @@ EBlendMode UMaterialInstance::GetBlendMode() const
 	return BlendMode;
 }
 
+EStrataBlendMode UMaterialInstance::GetStrataBlendMode() const
+{
+	return Parent ? Parent->GetStrataBlendMode() : SBM_Opaque;
+}
+
 FMaterialShadingModelField UMaterialInstance::GetShadingModels() const
 {
 	return ShadingModels;
