Wrap CL23126995 around a cvar and apply the change to all  platforms.

Original CL23126995 description: Use 'ApproximateOcclusion' for more mesh mobility types, to reduce amount of individual queries for a movable meshes. Improves RHI/Render threads and GPU time.

diff --git a/Engine/Source/Runtime/Engine/Private/PrimitiveSceneProxy.cpp b/Engine/Source/Runtime/Engine/Private/PrimitiveSceneProxy.cpp
--- a/Engine/Source/Runtime/Engine/Private/PrimitiveSceneProxy.cpp
+++ b/Engine/Source/Runtime/Engine/Private/PrimitiveSceneProxy.cpp
@@ -81,6 +81,17 @@ static TAutoConsoleVariable<int32> CVarVelocityForceOutput(
 	ECVF_RenderThreadSafe
 	);
 
+static TAutoConsoleVariable<int32> CVarApproximateOcclusionQueries(
+	TEXT("r.ApproximateOcclusionQueries"),
+	0,
+	TEXT("Batch occlusion for a static and skeletal mesh even if there are movable.In general it's more beneficial to batch occlusion queires for all meshes"),
+	ECVF_RenderThreadSafe
+);
+
+bool IsAllowingApproximateOcclusionQueries()
+{
+	return CVarApproximateOcclusionQueries.GetValueOnAnyThread() != 0;
+}
 
 bool IsOptimizedWPO()
 {
diff --git a/Engine/Source/Runtime/Engine/Private/SkeletalMesh.cpp b/Engine/Source/Runtime/Engine/Private/SkeletalMesh.cpp
--- a/Engine/Source/Runtime/Engine/Private/SkeletalMesh.cpp
+++ b/Engine/Source/Runtime/Engine/Private/SkeletalMesh.cpp
@@ -6032,9 +6032,8 @@ FSkeletalMeshSceneProxy::FSkeletalMeshSceneProxy(const USkinnedMeshComponent* Co
 	}
 #endif
 
-	if (FeatureLevel == ERHIFeatureLevel::ES3_1)
-	{
-		// batch occlusion for a "Static" skeletal mesh
+	if (IsAllowingApproximateOcclusionQueries())
+	{		
 		bAllowApproximateOcclusion = (bAllowApproximateOcclusion || bRenderStatic);
 	}
 }
diff --git a/Engine/Source/Runtime/Engine/Private/StaticMeshRender.cpp b/Engine/Source/Runtime/Engine/Private/StaticMeshRender.cpp
--- a/Engine/Source/Runtime/Engine/Private/StaticMeshRender.cpp
+++ b/Engine/Source/Runtime/Engine/Private/StaticMeshRender.cpp
@@ -420,8 +420,7 @@ FStaticMeshSceneProxy::FStaticMeshSceneProxy(UStaticMeshComponent* InComponent,
 	// Enable dynamic triangle reordering to remove/reduce sorting issue when rendered with a translucent material (i.e., order-independent-transparency)
 	bSupportsSortedTriangles = InComponent->bSortTriangles;
 
-	// In general it's more beneficial to batch occlusion queires for all meshes on mobile
-	if (FeatureLevel == ERHIFeatureLevel::ES3_1)
+	if (IsAllowingApproximateOcclusionQueries())
 	{
 		bAllowApproximateOcclusion = true;
 	}
diff --git a/Engine/Source/Runtime/Engine/Public/PrimitiveSceneProxy.h b/Engine/Source/Runtime/Engine/Public/PrimitiveSceneProxy.h
--- a/Engine/Source/Runtime/Engine/Public/PrimitiveSceneProxy.h
+++ b/Engine/Source/Runtime/Engine/Public/PrimitiveSceneProxy.h
@@ -130,7 +130,7 @@ public:
 };
 
 extern bool IsOptimizedWPO();
-
+extern bool IsAllowingApproximateOcclusionQueries();
 extern bool CacheShadowDepthsFromPrimitivesUsingWPO();
 
 enum class ERayTracingPrimitiveFlags : uint8
