Default the Hardware Device ID to an invalid device. Check for if the hardware is "Invalid"

Convert FHardwareDeviceID to be FNames, not FStrings for better performance.

#rnx

diff --git a/Engine/Source/Runtime/Engine/Classes/GameFramework/ForceFeedbackEffect.h b/Engine/Source/Runtime/Engine/Classes/GameFramework/ForceFeedbackEffect.h
--- a/Engine/Source/Runtime/Engine/Classes/GameFramework/ForceFeedbackEffect.h
+++ b/Engine/Source/Runtime/Engine/Classes/GameFramework/ForceFeedbackEffect.h
@@ -105,7 +105,7 @@ class UForceFeedbackEffect : public UObject
 
 	/** A map of platform name -> ForceFeedback channel details */
 	UPROPERTY(EditAnywhere, Category = "ForceFeedbackEffect", Meta = (GetOptions = "Engine.InputPlatformSettings.GetAllHardwareDeviceNames"))
-	TMap<FString, FForceFeedbackEffectOverridenChannelDetails> PerDeviceOverides;
+	TMap<FName, FForceFeedbackEffectOverridenChannelDetails> PerDeviceOverides;
 
 	/** A map of input device properties that we want to set while this effect is playing */
 	UPROPERTY(EditAnywhere, Instanced, Category = "ForceFeedbackEffect")
diff --git a/Engine/Source/Runtime/Engine/Classes/GameFramework/InputSettings.h b/Engine/Source/Runtime/Engine/Classes/GameFramework/InputSettings.h
--- a/Engine/Source/Runtime/Engine/Classes/GameFramework/InputSettings.h
+++ b/Engine/Source/Runtime/Engine/Classes/GameFramework/InputSettings.h
@@ -291,8 +291,8 @@ struct ENGINE_API FHardwareDeviceIdentifier
 {
 	GENERATED_BODY()
 
-	FHardwareDeviceIdentifier() = default;
-	FHardwareDeviceIdentifier(FName InClassName, FString InHardwareDeviceIdentifier);
+	FHardwareDeviceIdentifier();
+	FHardwareDeviceIdentifier(const FName InClassName, const FName InHardwareDeviceIdentifier);
 	
 	/** 
 	* The name of the Input Class that uses this hardware device.
@@ -306,10 +306,12 @@ struct ENGINE_API FHardwareDeviceIdentifier
 	 * This should correspond with a FInputDeviceScope that is used by an IInputDevice
 	 */
 	UPROPERTY(EditDefaultsOnly, BlueprintReadOnly, Category = "Hardware")
-	FString HardwareDeviceIdentifier;
+	FName HardwareDeviceIdentifier;
 
+	/** Returns true if this hardware device Identifier has valid names */
 	bool IsValid() const;
 	
+	/** An Invalid Hardware Device Identifier. */
 	static FHardwareDeviceIdentifier Invalid;
 
 	bool operator==(const FHardwareDeviceIdentifier& Other) const
@@ -341,7 +343,7 @@ public:
 	* and the editor will make a nice drop down for you with all the current options that are in the settings.
 	*/
 	UFUNCTION()
-	static const TArray<FString>& GetAllHardwareDeviceNames();
+	static const TArray<FName>& GetAllHardwareDeviceNames();
 #endif	// WITH_EDITOR
 
 	/** Add the given hardware device identifier to this platform's settings. */
diff --git a/Engine/Source/Runtime/Engine/Private/GameFramework/InputDeviceSubsystem.cpp b/Engine/Source/Runtime/Engine/Private/GameFramework/InputDeviceSubsystem.cpp
--- a/Engine/Source/Runtime/Engine/Private/GameFramework/InputDeviceSubsystem.cpp
+++ b/Engine/Source/Runtime/Engine/Private/GameFramework/InputDeviceSubsystem.cpp
@@ -30,7 +30,8 @@ class FInputDeviceSubsystemProcessor : public IInputProcessor
 		{
 			if (const FInputDeviceScope* Scope = FInputDeviceScope::GetCurrent())
 			{
-				SubSystem->SetMostRecentlyUsedHardwareDevice(InDeviceId, { Scope->InputDeviceName, Scope->HardwareDeviceIdentifier });
+				// TODO: Refactor FInputDeviceScope to use FName's instead of a FString for HardwareDeviceIdentifier
+				SubSystem->SetMostRecentlyUsedHardwareDevice(InDeviceId, { Scope->InputDeviceName, FName(*Scope->HardwareDeviceIdentifier) });
 			}	
 		}
 	}
diff --git a/Engine/Source/Runtime/Engine/Private/UserInterface/InputSettings.cpp b/Engine/Source/Runtime/Engine/Private/UserInterface/InputSettings.cpp
--- a/Engine/Source/Runtime/Engine/Private/UserInterface/InputSettings.cpp
+++ b/Engine/Source/Runtime/Engine/Private/UserInterface/InputSettings.cpp
@@ -493,9 +493,17 @@ void UInputSettings::SetDefaultInputComponentClass(TSubclassOf<UInputComponent>
 	}
 }
 
-FHardwareDeviceIdentifier FHardwareDeviceIdentifier::Invalid = { NAME_None, TEXT("Invalid") };
+FHardwareDeviceIdentifier FHardwareDeviceIdentifier::Invalid = { NAME_None, NAME_None };
 
-FHardwareDeviceIdentifier::FHardwareDeviceIdentifier(FName InClassName, FString InHardwareDeviceIdentifier)
+// Default to the invalid hardware device identifier
+FHardwareDeviceIdentifier::FHardwareDeviceIdentifier()
+	: InputClassName(Invalid.InputClassName)
+	, HardwareDeviceIdentifier(Invalid.HardwareDeviceIdentifier)
+{
+	
+}
+
+FHardwareDeviceIdentifier::FHardwareDeviceIdentifier(const FName InClassName, const FName InHardwareDeviceIdentifier)
 	: InputClassName(InClassName)
 	, HardwareDeviceIdentifier(InHardwareDeviceIdentifier)
 {
@@ -504,7 +512,7 @@ FHardwareDeviceIdentifier::FHardwareDeviceIdentifier(FName InClassName, FString
 
 bool FHardwareDeviceIdentifier::IsValid() const
 {
-	return InputClassName.IsValid() && !HardwareDeviceIdentifier.IsEmpty();
+	return InputClassName.IsValid() && HardwareDeviceIdentifier.IsValid();
 }
 
 UInputPlatformSettings* UInputPlatformSettings::Get()
@@ -526,9 +534,9 @@ const TArray<FHardwareDeviceIdentifier>& UInputPlatformSettings::GetHardwareDevi
 }
 
 #if WITH_EDITOR
-const TArray<FString>& UInputPlatformSettings::GetAllHardwareDeviceNames()
+const TArray<FName>& UInputPlatformSettings::GetAllHardwareDeviceNames()
 {
-	static TArray<FString> HardwareDevices;
+	static TArray<FName> HardwareDevices;
 	HardwareDevices.Reset();
 
 	// Get every known platform's InputPlatformSettings and compile a list of them
