Add Per-Device override support for input device properties! All this change does is break out the UPROPERTY's of each device property subclass to a struct so that they can be used with the per-device "GetOptions" TMap.

#rnx

diff --git a/Engine/Source/Runtime/Engine/Classes/GameFramework/InputDeviceProperties.h b/Engine/Source/Runtime/Engine/Classes/GameFramework/InputDeviceProperties.h
--- a/Engine/Source/Runtime/Engine/Classes/GameFramework/InputDeviceProperties.h
+++ b/Engine/Source/Runtime/Engine/Classes/GameFramework/InputDeviceProperties.h
@@ -5,6 +5,7 @@
 #include "CoreMinimal.h"
 #include "UObject/Object.h"
 #include "GenericPlatform/IInputInterface.h"
+#include "GameFramework/InputDeviceSubsystem.h"
 
 #include "InputDeviceProperties.generated.h"
 
@@ -27,13 +28,13 @@
 * support certain device properties. An older gamepad may not have any advanced trigger haptics for 
 * example. 
 */
-UCLASS(Abstract, Blueprintable, BlueprintType, EditInlineNew)
+UCLASS(Abstract, Blueprintable, BlueprintType, EditInlineNew, CollapseCategories)
 class ENGINE_API UInputDeviceProperty : public UObject
 {
 	GENERATED_BODY()
 public:
 
-	UInputDeviceProperty();
+	UInputDeviceProperty(const FObjectInitializer& ObjectInitializer = FObjectInitializer::Get());
 
 	/**
 	* Evaluate this device property for a given duration. 
@@ -111,6 +112,16 @@
 
 protected:
 
+	/** Returns the device specific data for the given platform user. Returns the default data if none are given */
+	template<class TDataLayout>
+	const TDataLayout* GetDeviceSpecificData(const FPlatformUserId UserId, const TMap<FName, TDataLayout>& InMap) const;
+
+	template<class TDataLayout>
+	TDataLayout* GetDeviceSpecificDataMutable(const FPlatformUserId UserId, TMap<FName, TDataLayout>& InMap)
+	{
+		return GetDeviceSpecificData<TDataLayout>(UserId, InMap);
+	}
+
 	/**
 	* The duration that this device property should last. Override this if your property has any dynamic curves 
 	* to be the max time range.
@@ -119,6 +130,40 @@
 	float PropertyDuration = 0.1f;
 };
 
+template<class TDataLayout>
+const TDataLayout* UInputDeviceProperty::GetDeviceSpecificData(const FPlatformUserId UserId, const TMap<FName, TDataLayout>& InDeviceData) const
+{
+	if (const UInputDeviceSubsystem* SubSystem = UInputDeviceSubsystem::Get())
+	{
+		const FHardwareDeviceIdentifier Hardware = SubSystem->GetMostRecentlyUsedHardwareDevice(UserId);
+		// Check if there are any per-input device overrides available
+		if (const TDataLayout* DeviceDetails = InDeviceData.Find(Hardware.HardwareDeviceIdentifier))
+		{
+			return DeviceDetails;
+		}
+	}
+
+	return nullptr;
+}
+
+///////////////////////////////////////////////////////////////////////
+// UColorInputDeviceProperty
+
+/** Data required for setting the Input Device Color */
+USTRUCT(BlueprintType)
+struct FDeviceColorData
+{
+	GENERATED_BODY()
+
+	/** True if the light should be enabled at all */
+	UPROPERTY(EditAnywhere, BlueprintReadWrite, Category = "Color")
+	bool bEnable = true;
+
+	/** The color to set the light on  */
+	UPROPERTY(EditAnywhere, BlueprintReadWrite, Category = "Color")
+	FColor LightColor = FColor::White;
+};
+
 /**
 * Set the color of an Input Device to a static color. 
 * 
@@ -136,13 +181,13 @@
 	virtual void ResetDeviceProperty_Implementation(const FPlatformUserId PlatformUser) override;
 	virtual FInputDeviceProperty* GetInternalDeviceProperty() override;
 
-	/** True if the light should be enabled at all */
-	UPROPERTY(EditAnywhere, BlueprintReadWrite, Category = "DeviceProperty")
-	bool bEnable = true;
+	/** Default color data that will be used by default. Device Specific overrides will be used when the current input device matches */
+	UPROPERTY(BlueprintReadWrite, EditAnywhere, Category = "Color")
+	FDeviceColorData ColorData;
 
-	/** The color to set the light on  */
-	UPROPERTY(EditAnywhere, BlueprintReadWrite, Category = "DeviceProperty")
-	FColor LightColor = FColor::White;
+	/** A map of device specific color data. If no overrides are specified, the Default hardware data will be used */
+	UPROPERTY(BlueprintReadWrite, EditAnywhere, Category = "Color", meta = (GetOptions = "Engine.InputPlatformSettings.GetAllHardwareDeviceNames"))
+	TMap<FName, FDeviceColorData> DeviceOverrideData;
 
 private:
 
@@ -150,6 +195,25 @@
 	FInputDeviceLightColorProperty InternalProperty;
 };
 
+
+///////////////////////////////////////////////////////////////////////
+// UColorInputDeviceCurveProperty
+
+/** Data required for setting the Input Device Color */
+USTRUCT(BlueprintType)
+struct FDeviceColorCurveData
+{
+	GENERATED_BODY()
+
+	/** True if the light should be enabled at all */
+	UPROPERTY(EditAnywhere, BlueprintReadWrite, Category = "Color")
+	bool bEnable = true;
+
+	/** The color the device light should be */
+	UPROPERTY(EditAnywhere, BlueprintReadWrite, Category = "Color")
+	TObjectPtr<UCurveLinearColor> DeviceColorCurve;
+};
+
 /** 
 * A property that can be used to change the color of an input device's light over time with a curve
 * 
@@ -169,14 +233,13 @@
 	virtual float RecalculateDuration() override;
 
 protected:
+	/** Default color data that will be used by default. Device Specific overrides will be used when the current input device matches */
+	UPROPERTY(BlueprintReadWrite, EditAnywhere, Category = "Color")
+	FDeviceColorCurveData ColorData;
 
-	/** True if the light should be enabled at all */
-	UPROPERTY(EditAnywhere, BlueprintReadWrite, Category = "DeviceProperty")
-	bool bEnable = true;
+	/** A map of device specific color data. If no overrides are specified, the Default hardware data will be used */
+	UPROPERTY(BlueprintReadWrite, EditAnywhere, Category = "Color", meta = (GetOptions = "Engine.InputPlatformSettings.GetAllHardwareDeviceNames"))
+	TMap<FName, FDeviceColorCurveData> DeviceOverrideData;
 
-	/** The color the device light should be */
-	UPROPERTY(EditAnywhere, BlueprintReadWrite, Category = "DeviceProperty")
-	TObjectPtr<UCurveLinearColor> DeviceColorCurve;
-
 private:
 
@@ -184,8 +247,26 @@
 	FInputDeviceLightColorProperty InternalProperty;
 };
 
+
+///////////////////////////////////////////////////////////////////////
+// UInputDeviceTriggerEffect
+
+USTRUCT(BlueprintType)
+struct FDeviceTriggerBaseData
+{
+	GENERATED_BODY()
+
+	/** Which trigger this property should effect */
+	UPROPERTY(EditAnywhere, BlueprintReadWrite, Category = "DeviceProperty")
+	EInputDeviceTriggerMask AffectedTriggers = EInputDeviceTriggerMask::None;
+
+	/** True if the triggers should be reset after the duration of this device property */
+	UPROPERTY(EditAnywhere, BlueprintReadWrite, Category = "DeviceProperty")
+	bool bResetUponCompletion = true;
+};
+
 /** A property that effect the triggers on a gamepad */
-UCLASS(Blueprintable, meta = (DisplayName = "Reset Trigger Device Properties"))
+UCLASS(Abstract, Blueprintable, meta = (DisplayName = "Reset Trigger Device Properties"))
 class ENGINE_API UInputDeviceTriggerEffect : public UInputDeviceProperty
 {
 	GENERATED_BODY()
@@ -195,13 +276,8 @@
 	virtual FInputDeviceProperty* GetInternalDeviceProperty() override;
 	virtual void ResetDeviceProperty_Implementation(const FPlatformUserId PlatformUser) override;
 
-	/** Which trigger this property should effect */
-    UPROPERTY(EditAnywhere, BlueprintReadWrite, Category = "DeviceProperty")
-    EInputDeviceTriggerMask AffectedTriggers = EInputDeviceTriggerMask::None;
-
-	/** True if the triggers should be reset after the duration of this device property */
 	UPROPERTY(EditAnywhere, BlueprintReadWrite, Category = "DeviceProperty")
-	bool bResetUponCompletion = true;
+	FDeviceTriggerBaseData BaseTriggerData;
 
 protected:
 
@@ -209,6 +285,23 @@
 	FInputDeviceTriggerResetProperty ResetProperty = {};
 };
 
+///////////////////////////////////////////////////////////////////////
+// UInputDeviceTriggerFeedbackProperty
+
+USTRUCT(BlueprintType)
+struct FDeviceTriggerFeedbackData
+{
+	GENERATED_BODY()
+
+	/** What position on the trigger that the feedback should be applied to over time (1-9) */
+	UPROPERTY(EditAnywhere, BlueprintReadWrite, Category = "DeviceProperty")
+	TObjectPtr<UCurveFloat> FeedbackPositionCurve;
+
+	/** How strong the feedback is over time (1-8) */
+	UPROPERTY(EditAnywhere, BlueprintReadWrite, Category = "DeviceProperty")
+	TObjectPtr<UCurveFloat> FeedbackStrenghCurve;
+};
+
 /** 
 * Sets simple trigger feedback
 * 
@@ -229,22 +322,45 @@
 	virtual float RecalculateDuration() override;
 
 	/** What position on the trigger that the feedback should be applied to over time (1-9) */
-	UPROPERTY(EditAnywhere, BlueprintReadWrite, Category = "DeviceProperty")
-	TObjectPtr<UCurveFloat> FeedbackPositionCurve;
+	UPROPERTY(EditAnywhere, BlueprintReadWrite, Category = "Trigger")
+	FDeviceTriggerFeedbackData TriggerData;
 
-	/** How strong the feedback is over time (1-8) */
-	UPROPERTY(EditAnywhere, BlueprintReadWrite, Category = "DeviceProperty")
-	TObjectPtr<UCurveFloat> FeedbackStrenghCurve;
+	/** A map of device specific color data. If no overrides are specified, the Default hardware data will be used */
+	UPROPERTY(BlueprintReadWrite, EditAnywhere, Category = "Trigger", meta = (GetOptions = "Engine.InputPlatformSettings.GetAllHardwareDeviceNames"))
+	TMap<FName, FDeviceTriggerFeedbackData> DeviceOverrideData;
 
 private:
 
-	int32 GetPositionValue(const float Duration) const;
-	int32 GetStrengthValue(const float Duration) const;
+	int32 GetPositionValue(const FDeviceTriggerFeedbackData* Data, const float Duration) const;
+	int32 GetStrengthValue(const FDeviceTriggerFeedbackData* Data, const float Duration) const;
 
 	/** The internal property that represents this trigger feedback. */
 	FInputDeviceTriggerFeedbackProperty InternalProperty;
 };
 
+///////////////////////////////////////////////////////////////////////
+// UInputDeviceTriggerResistanceProperty
+
+USTRUCT(BlueprintType)
+struct FDeviceTriggerTriggerResistanceData
+{
+	GENERATED_BODY()
+
+	/** The position that the trigger should start providing resistance */
+	UPROPERTY(EditAnywhere, BlueprintReadWrite, Category = "DeviceProperty")
+	int32 StartPosition = 0;
+
+	/** How strong the resistance is */
+	UPROPERTY(EditAnywhere, BlueprintReadWrite, Category = "DeviceProperty")
+	int32 StartStrengh = 0;
+
+	UPROPERTY(EditAnywhere, BlueprintReadWrite, Category = "DeviceProperty")
+	int32 EndPosition = 0;
+
+	UPROPERTY(EditAnywhere, BlueprintReadWrite, Category = "DeviceProperty")
+	int32 EndStrengh = 0;
+};
+
 /** 
 * Provides resistance to a trigger while it is being pressed between a start and end value
 * 
@@ -265,19 +381,13 @@
 
 protected:
 
-	/** The position that the trigger should start providing resistance */
-	UPROPERTY(EditAnywhere, BlueprintReadWrite, Category = "DeviceProperty")
-	int32 StartPosition = 0;
+	UPROPERTY(EditAnywhere, BlueprintReadWrite, Category = "Trigger")
+	FDeviceTriggerTriggerResistanceData TriggerData;
 
-	/** How strong the resistance is */
-	UPROPERTY(EditAnywhere, BlueprintReadWrite, Category = "DeviceProperty")
-	int32 StartStrengh = 0;
-	
-	UPROPERTY(EditAnywhere, BlueprintReadWrite, Category = "DeviceProperty")
-	int32 EndPosition = 0;
-	
-	UPROPERTY(EditAnywhere, BlueprintReadWrite, Category = "DeviceProperty")
-	int32 EndStrengh = 0;
 
+	/** A map of device specific color data. If no overrides are specified, the Default hardware data will be used */
+	UPROPERTY(BlueprintReadWrite, EditAnywhere, Category = "Trigger", meta = (GetOptions = "Engine.InputPlatformSettings.GetAllHardwareDeviceNames"))
+	TMap<FName, FDeviceTriggerTriggerResistanceData> DeviceOverrideData;
+
 private:
 
@@ -286,6 +396,28 @@
 };
 
 
+///////////////////////////////////////////////////////////////////////
+// UInputDeviceTriggerVibrationProperty
+
+
+USTRUCT(BlueprintType)
+struct FDeviceTriggerTriggerVibrationData
+{
+	GENERATED_BODY()
+
+	/** What position on the trigger that the feedback should be applied to over time (1-9) */
+	UPROPERTY(EditAnywhere, BlueprintReadWrite, Category = "DeviceProperty")
+	TObjectPtr<UCurveFloat> TriggerPositionCurve;
+
+	/** The frequency of the vibration */
+	UPROPERTY(EditAnywhere, BlueprintReadWrite, Category = "DeviceProperty")
+	TObjectPtr<UCurveFloat> VibrationFrequencyCurve;
+
+	/** The amplitude of the vibration */
+	UPROPERTY(EditAnywhere, BlueprintReadWrite, Category = "DeviceProperty")
+	TObjectPtr<UCurveFloat> VibrationAmplitudeCurve;
+};
+
 /**
 * Sets trigger vibration
 *
@@ -305,23 +437,18 @@
 	virtual FInputDeviceProperty* GetInternalDeviceProperty() override;
 	virtual float RecalculateDuration() override;
 
-	/** What position on the trigger that the feedback should be applied to over time (1-9) */
-	UPROPERTY(EditAnywhere, BlueprintReadWrite, Category = "DeviceProperty")
-	TObjectPtr<UCurveFloat> TriggerPositionCurve;
+	UPROPERTY(EditAnywhere, BlueprintReadWrite, Category = "Trigger")
+	FDeviceTriggerTriggerVibrationData TriggerData;
 
-	/** The frequency of the vibration */
-	UPROPERTY(EditAnywhere, BlueprintReadWrite, Category = "DeviceProperty")
-	TObjectPtr<UCurveFloat> VibrationFrequencyCurve;
+	/** A map of device specific color data. If no overrides are specified, the Default hardware data will be used */
+	UPROPERTY(BlueprintReadWrite, EditAnywhere, Category = "Trigger", meta = (GetOptions = "Engine.InputPlatformSettings.GetAllHardwareDeviceNames"))
+	TMap<FName, FDeviceTriggerTriggerVibrationData> DeviceOverrideData;
 
-	/** The amplitude of the vibration */
-	UPROPERTY(EditAnywhere, BlueprintReadWrite, Category = "DeviceProperty")
-	TObjectPtr<UCurveFloat> VibrationAmplitudeCurve;
-
 private:
 
-	int32 GetTriggerPositionValue(const float Duration) const;
-	int32 GetVibrationFrequencyValue(const float Duration) const;
-	int32 GetVibrationAmplitudeValue(const float Duration) const;
+	int32 GetTriggerPositionValue(const FDeviceTriggerTriggerVibrationData* Data, const float Duration) const;
+	int32 GetVibrationFrequencyValue(const FDeviceTriggerTriggerVibrationData* Data, const float Duration) const;
+	int32 GetVibrationAmplitudeValue(const FDeviceTriggerTriggerVibrationData* Data, const float Duration) const;
 
 	/** The internal property that represents this trigger feedback. */
 	FInputDeviceTriggerVibrationProperty InternalProperty;
diff --git a/Engine/Source/Runtime/Engine/Private/GameFramework/InputDeviceProperties.cpp b/Engine/Source/Runtime/Engine/Private/GameFramework/InputDeviceProperties.cpp
--- a/Engine/Source/Runtime/Engine/Private/GameFramework/InputDeviceProperties.cpp
+++ b/Engine/Source/Runtime/Engine/Private/GameFramework/InputDeviceProperties.cpp
@@ -9,7 +9,8 @@
 ///////////////////////////////////////////////////////////////////////
 // UInputDeviceProperty
 
-UInputDeviceProperty::UInputDeviceProperty()
+UInputDeviceProperty::UInputDeviceProperty(const FObjectInitializer& ObjectInitializer)
+	: Super(ObjectInitializer)
 {
 	RecalculateDuration();
 }
@@ -68,8 +69,18 @@ void UInputDeviceProperty::ResetDeviceProperty_Implementation(const FPlatformUse
 
 void UColorInputDeviceProperty::EvaluateDeviceProperty_Implementation(const FPlatformUserId PlatformUser, const float DeltaTime, const float Duration)
 {
-	InternalProperty.bEnable = bEnable;
-	InternalProperty.Color = LightColor;
+	// Check for an override on the current input device
+	if (const FDeviceColorData* Data = GetDeviceSpecificData<FDeviceColorData>(PlatformUser, DeviceOverrideData))
+	{
+		InternalProperty.bEnable = Data->bEnable;
+		InternalProperty.Color = Data->LightColor;
+	}
+	// Otherwise use the default color data
+	else
+	{
+		InternalProperty.bEnable = ColorData.bEnable;
+		InternalProperty.Color = ColorData.LightColor;
+	}
 }
 
 void UColorInputDeviceProperty::ResetDeviceProperty_Implementation(const FPlatformUserId PlatformUser)
@@ -89,12 +100,27 @@ FInputDeviceProperty* UColorInputDeviceProperty::GetInternalDeviceProperty()
 
 void UColorInputDeviceCurveProperty::EvaluateDeviceProperty_Implementation(const FPlatformUserId PlatformUser, const float DeltaTime, const float Duration)
 {
-	InternalProperty.bEnable = bEnable;
+	// Check for an override on the current input device
+	if (const FDeviceColorCurveData* Data = GetDeviceSpecificData<FDeviceColorCurveData>(PlatformUser, DeviceOverrideData))
+	{
+		InternalProperty.bEnable = Data->bEnable;
 
-	if (ensure(DeviceColorCurve))
+		if (ensure(Data->DeviceColorCurve))
+		{
+			FLinearColor CurveColor = Data->DeviceColorCurve->GetLinearColorValue(Duration);
+			InternalProperty.Color = CurveColor.ToFColorSRGB();
+		}
+	}
+	// Otherwise use the default color data
+	else
 	{
-		FLinearColor CurveColor = DeviceColorCurve->GetLinearColorValue(Duration);
-		InternalProperty.Color = CurveColor.ToFColorSRGB();
+		InternalProperty.bEnable = ColorData.bEnable;
+
+		if (ensure(ColorData.DeviceColorCurve))
+		{
+			FLinearColor CurveColor = ColorData.DeviceColorCurve->GetLinearColorValue(Duration);
+			InternalProperty.Color = CurveColor.ToFColorSRGB();
+		}
 	}
 }
 
@@ -112,16 +138,31 @@ FInputDeviceProperty* UColorInputDeviceCurveProperty::GetInternalDeviceProperty(
 
 float UColorInputDeviceCurveProperty::RecalculateDuration()
 {
-	float MinTime, MaxTime;
-	if (DeviceColorCurve)
+	float MinTime = 0.f;
+	float MaxTime = 0.f;
+
+	if (ColorData.DeviceColorCurve)
 	{
-		DeviceColorCurve->GetTimeRange(MinTime, MaxTime);
-		PropertyDuration = MaxTime;
+		ColorData.DeviceColorCurve->GetTimeRange(MinTime, MaxTime);
 	}
-	else
+
+	// Find the max time of any device specific data
+	for (const TPair<FName, FDeviceColorCurveData>& Pair : DeviceOverrideData)
 	{
-		PropertyDuration = 1.0f;
+		if (Pair.Value.DeviceColorCurve)
+		{
+			float DeviceSpecificMinTime = 0.f;
+			float DeviceSpecificMaxTime = 0.f;
+			Pair.Value.DeviceColorCurve->GetTimeRange(DeviceSpecificMinTime, DeviceSpecificMaxTime);
+			if (DeviceSpecificMaxTime > MaxTime)
+			{
+				MaxTime = DeviceSpecificMaxTime;
+			}
+		}		
 	}
+
+	PropertyDuration = MaxTime;
+	
 	return PropertyDuration;
 }
 
@@ -135,10 +176,10 @@ FInputDeviceProperty* UInputDeviceTriggerEffect::GetInternalDeviceProperty()
 
 void UInputDeviceTriggerEffect::ResetDeviceProperty_Implementation(const FPlatformUserId PlatformUser)
 {
-	if (bResetUponCompletion)
+	if (BaseTriggerData.bResetUponCompletion)
 	{
 		// Pass in our reset property
-		ResetProperty.AffectedTriggers = AffectedTriggers;
+		ResetProperty.AffectedTriggers = BaseTriggerData.AffectedTriggers;
 		ApplyDeviceProperty(PlatformUser, &ResetProperty);
 	}	
 }
@@ -149,37 +190,46 @@ void UInputDeviceTriggerEffect::ResetDeviceProperty_Implementation(const FPlatfo
 UInputDeviceTriggerFeedbackProperty::UInputDeviceTriggerFeedbackProperty()
 	: UInputDeviceTriggerEffect()
 {
-	InternalProperty.AffectedTriggers = AffectedTriggers;
+	InternalProperty.AffectedTriggers = BaseTriggerData.AffectedTriggers;
 }
 
-int32 UInputDeviceTriggerFeedbackProperty::GetPositionValue(const float Duration) const
+int32 UInputDeviceTriggerFeedbackProperty::GetPositionValue(const FDeviceTriggerFeedbackData* Data, const float Duration) const
 {
-	if (ensure(FeedbackPositionCurve))
+	if (ensure(Data->FeedbackPositionCurve))
 	{
-		// TODO: Make the max position a cvar
-		int32 Position = FeedbackPositionCurve->GetFloatValue(Duration);
-		return FMath::Clamp(Position, 0, 9);
+		// TODO: Make the max Strength a cvar
+		int32 Pos = Data->FeedbackPositionCurve->GetFloatValue(Duration);
+		return FMath::Clamp(Pos, 0, 8);
 	}
+
 	return 0;
 }
 
-int32 UInputDeviceTriggerFeedbackProperty::GetStrengthValue(const float Duration) const
+int32 UInputDeviceTriggerFeedbackProperty::GetStrengthValue(const FDeviceTriggerFeedbackData* Data, const float Duration) const
 {
-	if (ensure(FeedbackStrenghCurve))
+	if (ensure(Data->FeedbackStrenghCurve))
 	{
 		// TODO: Make the max Strength a cvar
-		int32 Strength = FeedbackPositionCurve->GetFloatValue(Duration);
+		int32 Strength = Data->FeedbackStrenghCurve->GetFloatValue(Duration);
 		return FMath::Clamp(Strength, 0, 8);
 	}
+
 	return 0;
 }
 
 void UInputDeviceTriggerFeedbackProperty::EvaluateDeviceProperty_Implementation(const FPlatformUserId PlatformUser, const float DeltaTime, const float Duration)
 {		
-	InternalProperty.AffectedTriggers = AffectedTriggers;
+	InternalProperty.AffectedTriggers = BaseTriggerData.AffectedTriggers;
+
+	const FDeviceTriggerFeedbackData* DataToUse = &TriggerData;
 
-	InternalProperty.Position = GetPositionValue(Duration);
-	InternalProperty.Strengh = GetStrengthValue(Duration);
+	if (const FDeviceTriggerFeedbackData* OverrideData = GetDeviceSpecificData<FDeviceTriggerFeedbackData>(PlatformUser, DeviceOverrideData))
+	{
+		DataToUse = OverrideData;
+	}
+
+	InternalProperty.Position = GetPositionValue(DataToUse, Duration);
+	InternalProperty.Strengh = GetStrengthValue(DataToUse, Duration);
 }
 
 FInputDeviceProperty* UInputDeviceTriggerFeedbackProperty::GetInternalDeviceProperty()
@@ -191,14 +241,28 @@ float UInputDeviceTriggerFeedbackProperty::RecalculateDuration()
 {
 	// Get the max time from the two curves on this property
 	float MinTime, MaxTime = 0.0f;
-	if (FeedbackPositionCurve)
+	if (TriggerData.FeedbackPositionCurve)
 	{
-		FeedbackPositionCurve->GetTimeRange(MinTime, MaxTime);
+		TriggerData.FeedbackPositionCurve->GetTimeRange(MinTime, MaxTime);
 	}
 	
-	if (FeedbackStrenghCurve)
+	if (TriggerData.FeedbackStrenghCurve)
 	{
-		FeedbackStrenghCurve->GetTimeRange(MinTime, MaxTime);
+		TriggerData.FeedbackStrenghCurve->GetTimeRange(MinTime, MaxTime);
+	}
+
+	// Find the max time of any device specific data
+	for (const TPair<FName, FDeviceTriggerFeedbackData>& Pair : DeviceOverrideData)
+	{
+		if (Pair.Value.FeedbackPositionCurve)
+		{
+			Pair.Value.FeedbackPositionCurve->GetTimeRange(MinTime, MaxTime);
+		}
+
+		if (Pair.Value.FeedbackStrenghCurve)
+		{
+			Pair.Value.FeedbackStrenghCurve->GetTimeRange(MinTime, MaxTime);
+		}
 	}
 	
 	PropertyDuration = MaxTime;
@@ -216,11 +280,22 @@ UInputDeviceTriggerResistanceProperty::UInputDeviceTriggerResistanceProperty()
 
 void UInputDeviceTriggerResistanceProperty::EvaluateDeviceProperty_Implementation(const FPlatformUserId PlatformUser, const float DeltaTime, const float Duration)
 {
-	InternalProperty.AffectedTriggers = AffectedTriggers;
-	InternalProperty.StartPosition = StartPosition;
-	InternalProperty.StartStrengh = StartStrengh;
-	InternalProperty.EndPosition = EndPosition;
-	InternalProperty.EndStrengh = EndStrengh;
+	InternalProperty.AffectedTriggers = BaseTriggerData.AffectedTriggers;
+
+	if (const FDeviceTriggerTriggerResistanceData* Data = GetDeviceSpecificData<FDeviceTriggerTriggerResistanceData>(PlatformUser, DeviceOverrideData))
+	{
+		InternalProperty.StartPosition = Data->StartPosition;
+		InternalProperty.StartStrengh = Data->StartStrengh;
+		InternalProperty.EndPosition = Data->EndPosition;
+		InternalProperty.EndStrengh = Data->EndStrengh;
+	}
+	else
+	{
+		InternalProperty.StartPosition = TriggerData.StartPosition;
+		InternalProperty.StartStrengh = TriggerData.StartStrengh;
+		InternalProperty.EndPosition = TriggerData.EndPosition;
+		InternalProperty.EndStrengh = TriggerData.EndStrengh;
+	}
 }
 
 FInputDeviceProperty* UInputDeviceTriggerResistanceProperty::GetInternalDeviceProperty()
@@ -240,10 +315,17 @@ UInputDeviceTriggerVibrationProperty::UInputDeviceTriggerVibrationProperty()
 
 void UInputDeviceTriggerVibrationProperty::EvaluateDeviceProperty_Implementation(const FPlatformUserId PlatformUser, const float DeltaTime, const float Duration)
 {
-	InternalProperty.AffectedTriggers = AffectedTriggers;
-	InternalProperty.TriggerPosition = GetTriggerPositionValue(Duration);
-	InternalProperty.VibrationFrequency = GetVibrationFrequencyValue(Duration);
-	InternalProperty.VibrationAmplitude = GetVibrationAmplitudeValue(Duration);
+	const FDeviceTriggerTriggerVibrationData* DataToUse = &TriggerData;
+
+	if (const FDeviceTriggerTriggerVibrationData* OverrideData = GetDeviceSpecificData<FDeviceTriggerTriggerVibrationData>(PlatformUser, DeviceOverrideData))
+	{
+		DataToUse = OverrideData;
+	}
+
+	InternalProperty.AffectedTriggers = BaseTriggerData.AffectedTriggers;
+	InternalProperty.TriggerPosition = GetTriggerPositionValue(DataToUse, Duration);
+	InternalProperty.VibrationFrequency = GetVibrationFrequencyValue(DataToUse, Duration);
+	InternalProperty.VibrationAmplitude = GetVibrationAmplitudeValue(DataToUse, Duration);
 }
 
 FInputDeviceProperty* UInputDeviceTriggerVibrationProperty::GetInternalDeviceProperty()
@@ -269,43 +351,53 @@ float UInputDeviceTriggerVibrationProperty::RecalculateDuration()
 		}
 	};
 
-	EvaluateMaxTime(TriggerPositionCurve);
-	EvaluateMaxTime(VibrationFrequencyCurve);
-	EvaluateMaxTime(VibrationAmplitudeCurve);
+	EvaluateMaxTime(TriggerData.TriggerPositionCurve);
+	EvaluateMaxTime(TriggerData.VibrationFrequencyCurve);
+	EvaluateMaxTime(TriggerData.VibrationAmplitudeCurve);
+
+	for (const TPair<FName, FDeviceTriggerTriggerVibrationData>& Pair : DeviceOverrideData)
+	{
+		EvaluateMaxTime(Pair.Value.TriggerPositionCurve);
+		EvaluateMaxTime(Pair.Value.VibrationFrequencyCurve);
+		EvaluateMaxTime(Pair.Value.VibrationAmplitudeCurve);
+	}
 
 	PropertyDuration = MaxTime;
 	return PropertyDuration;
 }
 
-int32 UInputDeviceTriggerVibrationProperty::GetTriggerPositionValue(const float Duration) const
+int32 UInputDeviceTriggerVibrationProperty::GetTriggerPositionValue(const FDeviceTriggerTriggerVibrationData* Data, const float Duration) const
 {
-	if (ensure(TriggerPositionCurve))
+	if (ensure(Data->TriggerPositionCurve))
 	{
 		// TODO: Make the max Strength a cvar
-		int32 Strength = TriggerPositionCurve->GetFloatValue(Duration);
+		int32 Strength = Data->TriggerPositionCurve->GetFloatValue(Duration);
 		return FMath::Clamp(Strength, 0, 9);
 	}
+
 	return 0;
 }
 
-int32 UInputDeviceTriggerVibrationProperty::GetVibrationFrequencyValue(const float Duration) const
+int32 UInputDeviceTriggerVibrationProperty::GetVibrationFrequencyValue(const FDeviceTriggerTriggerVibrationData* Data, const float Duration) const
 {
-	if (ensure(VibrationFrequencyCurve))
+	if (ensure(Data->VibrationFrequencyCurve))
 	{
 		// TODO: Make the max Frequency a cvar
-		int32 Strength = VibrationFrequencyCurve->GetFloatValue(Duration);
+		int32 Strength = Data->VibrationFrequencyCurve->GetFloatValue(Duration);
 		return FMath::Clamp(Strength, 0, 255);
 	}
+	
 	return 0;
 }
 
-int32 UInputDeviceTriggerVibrationProperty::GetVibrationAmplitudeValue(const float Duration) const
+int32 UInputDeviceTriggerVibrationProperty::GetVibrationAmplitudeValue(const FDeviceTriggerTriggerVibrationData* Data, const float Duration) const
 {
-	if (ensure(VibrationAmplitudeCurve))
+	if (ensure(Data->VibrationAmplitudeCurve))
 	{
 		// TODO: Make the max Amplitude a cvar
-		int32 Strength = VibrationAmplitudeCurve->GetFloatValue(Duration);
+		int32 Strength = Data->VibrationAmplitudeCurve->GetFloatValue(Duration);
 		return FMath::Clamp(Strength, 0, 8);
 	}
+	
 	return 0;
 }
