Updated RandomStream to work in Python

This struct has internally mutable state (for the running seed), which meant that trying to use the unreal.MathLibrary functions on it would fail to persist the new seed. These have been changed to be ScriptMethod (so that the self argument can be copied back out) and marked as ScriptMethodMutable where needed.

Functions that weren't taking the RandomStream as the first parameter have been updated to do so, and deprecated overloads left for any C++ code that was directly using them (Blueprints handle paramter ordering changing without issue).

diff --git a/Engine/Source/Runtime/Engine/Classes/Kismet/KismetMathLibrary.h b/Engine/Source/Runtime/Engine/Classes/Kismet/KismetMathLibrary.h
--- a/Engine/Source/Runtime/Engine/Classes/Kismet/KismetMathLibrary.h
+++ b/Engine/Source/Runtime/Engine/Classes/Kismet/KismetMathLibrary.h
@@ -201,8 +201,14 @@ class ENGINE_API UKismetMathLibrary : public UBlueprintFunctionLibrary
 	 * Get a random chance with the specified weight. Range of weight is 0.0 - 1.0 E.g.,
 	*		Weight = .6 return value = True 60% of the time
 	*/
-	UFUNCTION(BlueprintPure, Category = "Math|Random", meta=(Weight = "0.5"))
-	static bool RandomBoolWithWeightFromStream(float Weight, const FRandomStream& RandomStream);
+	UFUNCTION(BlueprintPure, Category = "Math|Random", meta=(Weight = "0.5", ScriptMethod = "RandomBoolWithWeight", ScriptMethodMutable))
+	static bool RandomBoolWithWeightFromStream(const FRandomStream& RandomStream, float Weight);
+
+	UE_DEPRECATED(5.2, "Use RandomBoolWithWeightFromStream taking the FRandomStream as the first argument.")
+	static bool RandomBoolWithWeightFromStream(float Weight, const FRandomStream& RandomStream)
+	{
+		return RandomBoolWithWeightFromStream(RandomStream, Weight);
+	}
 
 	/** Returns the logical complement of the Boolean value (NOT A) */
 	UFUNCTION(BlueprintPure, meta=(DisplayName = "NOT Boolean", CompactNodeTitle = "NOT", Keywords = "! not negate"), Category="Math|Boolean")
@@ -3988,51 +3994,87 @@ class ENGINE_API UKismetMathLibrary : public UBlueprintFunctionLibrary
 	//
 
 	/** Returns a uniformly distributed random number between 0 and Max - 1 */
-	UFUNCTION(BlueprintPure, Category="Math|Random")
-	static int32 RandomIntegerFromStream(int32 Max, const FRandomStream& Stream);
+	UFUNCTION(BlueprintPure, Category="Math|Random", meta=(ScriptMethod="RandomInt", ScriptMethodMutable))
+	static int32 RandomIntegerFromStream(const FRandomStream& Stream, int32 Max);
+
+	UE_DEPRECATED(5.2, "Use RandomIntegerFromStream taking the FRandomStream as the first argument.")
+	static int32 RandomIntegerFromStream(int32 Max, const FRandomStream& Stream)
+	{
+		return RandomIntegerFromStream(Stream, Max);
+	}
 
 	/** Return a random integer between Min and Max (>= Min and <= Max) */
-	UFUNCTION(BlueprintPure, Category="Math|Random")
-	static int32 RandomIntegerInRangeFromStream(int32 Min, int32 Max, const FRandomStream& Stream);
+	UFUNCTION(BlueprintPure, Category="Math|Random", meta=(ScriptMethod="RandomIntInRange", ScriptMethodMutable))
+	static int32 RandomIntegerInRangeFromStream(const FRandomStream& Stream, int32 Min, int32 Max);
+
+	UE_DEPRECATED(5.2, "Use RandomIntegerInRangeFromStream taking the FRandomStream as the first argument.")
+	static int32 RandomIntegerInRangeFromStream(int32 Min, int32 Max, const FRandomStream& Stream)
+	{
+		return RandomIntegerInRangeFromStream(Stream, Min, Max);
+	}
 
 	/** Returns a random bool */
-	UFUNCTION(BlueprintPure, Category="Math|Random")
+	UFUNCTION(BlueprintPure, Category="Math|Random", meta=(ScriptMethod="RandomBool", ScriptMethodMutable))
 	static bool RandomBoolFromStream(const FRandomStream& Stream);
 
 	/** Returns a random float between 0 and 1 */
-	UFUNCTION(BlueprintPure, Category="Math|Random")
+	UFUNCTION(BlueprintPure, Category="Math|Random", meta=(ScriptMethod="RandomFloat", ScriptMethodMutable))
 	static float RandomFloatFromStream(const FRandomStream& Stream);
 
 	/** Generate a random number between Min and Max */
-	UFUNCTION(BlueprintPure, Category="Math|Random")
-	static float RandomFloatInRangeFromStream(float Min, float Max, const FRandomStream& Stream);
+	UFUNCTION(BlueprintPure, Category="Math|Random", meta=(ScriptMethod="RandomFloatInRange", ScriptMethodMutable))
+	static float RandomFloatInRangeFromStream(const FRandomStream& Stream, float Min, float Max);
+
+	UE_DEPRECATED(5.2, "Use RandomFloatInRangeFromStream taking the FRandomStream as the first argument.")
+	static float RandomFloatInRangeFromStream(float Min, float Max, const FRandomStream& Stream)
+	{
+		return RandomFloatInRangeFromStream(Stream, Min, Max);
+	}
 
 	/** Returns a random vector with length of 1.0 */
-	UFUNCTION(BlueprintPure, Category="Math|Random")
+	UFUNCTION(BlueprintPure, Category="Math|Random", meta=(ScriptMethod="RandomUnitVector", ScriptMethodMutable))
 	static FVector RandomUnitVectorFromStream(const FRandomStream& Stream);
 
 	/** Returns a random point within the specified bounding box using the first vector as an origin and the second as the half size of the AABB. */
-	UFUNCTION(BlueprintPure, Category="Math|Random")
-	static FVector RandomPointInBoundingBoxFromStream(const FVector Center, const FVector HalfSize, const FRandomStream& Stream);
+	UFUNCTION(BlueprintPure, Category="Math|Random", meta=(ScriptMethod="RandomPointInBoundedBox", ScriptMethodMutable))
+	static FVector RandomPointInBoundingBoxFromStream(const FRandomStream& Stream, const FVector Center, const FVector HalfSize);
+
+	UE_DEPRECATED(5.2, "Use RandomPointInBoundingBoxFromStream taking the FRandomStream as the first argument.")
+	static FVector RandomPointInBoundingBoxFromStream(const FVector Center, const FVector HalfSize, const FRandomStream& Stream)
+	{
+		return RandomPointInBoundingBoxFromStream(Stream, Center, HalfSize);
+	}
 
 	/** Returns a random point within the specified bounding box. */
-	UFUNCTION(BlueprintPure, Category="Math|Random", meta=(DisplayName="Random Point In Bounding Box From Stream (Box)"))
-	static FVector RandomPointInBoundingBoxFromStream_Box(const FBox Box, const FRandomStream& Stream);
+	UFUNCTION(BlueprintPure, Category="Math|Random", meta=(DisplayName="Random Point In Bounding Box From Stream (Box)"), meta=(ScriptMethod="RandomPointInBox", ScriptMethodMutable))
+	static FVector RandomPointInBoundingBoxFromStream_Box(const FRandomStream& Stream, const FBox Box);
+
+	UE_DEPRECATED(5.2, "Use RandomPointInBoundingBoxFromStream_Box taking the FRandomStream as the first argument.")
+	static FVector RandomPointInBoundingBoxFromStream_Box(const FBox Box, const FRandomStream& Stream)
+	{
+		return RandomPointInBoundingBoxFromStream_Box(Stream, Box);
+	}
 
 	/** Create a random rotation */
-	UFUNCTION(BlueprintPure, Category="Math|Random")
-	static FRotator RandomRotatorFromStream(bool bRoll, const FRandomStream& Stream);
+	UFUNCTION(BlueprintPure, Category="Math|Random", meta=(ScriptMethod="RandomRotator", ScriptMethodMutable))
+	static FRotator RandomRotatorFromStream(const FRandomStream& Stream, bool bRoll);
+
+	UE_DEPRECATED(5.2, "Use RandomRotatorFromStream taking the FRandomStream as the first argument.")
+	static FRotator RandomRotatorFromStream(bool bRoll, const FRandomStream& Stream)
+	{
+		return RandomRotatorFromStream(Stream, bRoll);
+	}
 
 	/** Reset a random stream */
-	UFUNCTION(BlueprintCallable, Category="Math|Random")
+	UFUNCTION(BlueprintCallable, Category="Math|Random", meta=(ScriptMethod="Reset", ScriptMethodMutable))
 	static void ResetRandomStream(const FRandomStream& Stream);
 
 	/** Create a new random seed for a random stream */
-	UFUNCTION(BlueprintCallable, Category="Math|Random")
+	UFUNCTION(BlueprintCallable, Category="Math|Random", meta=(ScriptMethod="GenerateNewSeed"))
 	static void SeedRandomStream(UPARAM(ref) FRandomStream& Stream);
 
 	/** Set the seed of a random stream to a specific number */
-	UFUNCTION(BlueprintCallable, Category="Math|Random")
+	UFUNCTION(BlueprintCallable, Category="Math|Random", meta=(ScriptMethod="SetSeed"))
 	static void SetRandomStreamSeed(UPARAM(ref) FRandomStream& Stream, int32 NewSeed);
 
 	/** 
@@ -4041,8 +4083,14 @@ class ENGINE_API UKismetMathLibrary : public UBlueprintFunctionLibrary
 	 * @param ConeHalfAngleInRadians	The half-angle of the cone (from ConeDir to edge), in radians.
 	 * @param Stream					The random stream from which to obtain the vector.
 	 */
-	UFUNCTION(BlueprintPure, Category="Math|Random", meta = (Keywords = "RandomVector"))
-	static FVector RandomUnitVectorInConeInRadiansFromStream(const FVector& ConeDir, float ConeHalfAngleInRadians, const FRandomStream& Stream);
+	UFUNCTION(BlueprintPure, Category="Math|Random", meta = (Keywords = "RandomVector", ScriptMethod = "RandomUnitVectorInConeInRadians", ScriptMethodMutable))
+	static FVector RandomUnitVectorInConeInRadiansFromStream(const FRandomStream& Stream, const FVector& ConeDir, float ConeHalfAngleInRadians);
+
+	UE_DEPRECATED(5.2, "Use RandomUnitVectorInConeInRadiansFromStream taking the FRandomStream as the first argument.")
+	static FVector RandomUnitVectorInConeInRadiansFromStream(const FVector& ConeDir, float ConeHalfAngleInRadians, const FRandomStream& Stream)
+	{
+		return RandomUnitVectorInConeInRadiansFromStream(Stream, ConeDir, ConeHalfAngleInRadians);
+	}
 
 	/** 
 	 * Returns a random vector with length of 1, within the specified cone, with uniform random distribution.
@@ -4050,10 +4098,16 @@ class ENGINE_API UKismetMathLibrary : public UBlueprintFunctionLibrary
 	 * @param ConeHalfAngleInDegrees	The half-angle of the cone (from ConeDir to edge), in degrees.
 	 * @param Stream					The random stream from which to obtain the vector.
 	 */
-	UFUNCTION(BlueprintPure, Category="Math|Random", meta = (Keywords = "RandomVector"))
+	UFUNCTION(BlueprintPure, Category="Math|Random", meta = (Keywords = "RandomVector", ScriptMethod = "RandomUnitVectorInConeInDegrees", ScriptMethodMutable))
+	static inline FVector RandomUnitVectorInConeInDegreesFromStream(const FRandomStream& Stream, const FVector& ConeDir, float ConeHalfAngleInDegrees)
+	{
+		return RandomUnitVectorInConeInRadiansFromStream(Stream, ConeDir, FMath::DegreesToRadians(ConeHalfAngleInDegrees));
+	}
+
+	UE_DEPRECATED(5.2, "Use RandomUnitVectorInConeInDegreesFromStream taking the FRandomStream as the first argument.")
 	static inline FVector RandomUnitVectorInConeInDegreesFromStream(const FVector& ConeDir, float ConeHalfAngleInDegrees, const FRandomStream& Stream)
 	{
-		return RandomUnitVectorInConeInRadiansFromStream(ConeDir, FMath::DegreesToRadians(ConeHalfAngleInDegrees), Stream);
+		return RandomUnitVectorInConeInDegreesFromStream(Stream, ConeDir, ConeHalfAngleInDegrees);
 	}
 
 	/**
@@ -4064,8 +4118,14 @@ class ENGINE_API UKismetMathLibrary : public UBlueprintFunctionLibrary
 	* @param MaxPitchInRadians	The pitch angle of the cone (from ConeDir to vertical edge), in radians.
 	* @param Stream				The random stream from which to obtain the vector.
 	*/
-	UFUNCTION(BlueprintPure, Category = "Math|Random", meta = (Keywords = "RandomVector"))
-	static FVector RandomUnitVectorInEllipticalConeInRadiansFromStream(const FVector& ConeDir, float MaxYawInRadians, float MaxPitchInRadians, const FRandomStream& Stream);
+	UFUNCTION(BlueprintPure, Category = "Math|Random", meta = (Keywords = "RandomVector", ScriptMethod = "RandomUnitVectorInEllipticalConeInRadians", ScriptMethodMutable))
+	static FVector RandomUnitVectorInEllipticalConeInRadiansFromStream(const FRandomStream& Stream, const FVector& ConeDir, float MaxYawInRadians, float MaxPitchInRadians);
+
+	UE_DEPRECATED(5.2, "Use RandomUnitVectorInEllipticalConeInRadiansFromStream taking the FRandomStream as the first argument.")
+	static FVector RandomUnitVectorInEllipticalConeInRadiansFromStream(const FVector& ConeDir, float MaxYawInRadians, float MaxPitchInRadians, const FRandomStream& Stream)
+	{
+		return RandomUnitVectorInEllipticalConeInRadiansFromStream(Stream, ConeDir, MaxYawInRadians, MaxPitchInRadians);
+	}
 
 	/**
 	* Returns a random vector with length of 1, within the specified cone, with uniform random distribution.
@@ -4075,10 +4135,16 @@ class ENGINE_API UKismetMathLibrary : public UBlueprintFunctionLibrary
 	* @param MaxPitchInDegrees	The pitch angle of the cone (from ConeDir to vertical edge), in degrees.
 	* @param Stream				The random stream from which to obtain the vector.
 	*/
-	UFUNCTION(BlueprintPure, Category = "Math|Random", meta = (Keywords = "RandomVector"))
+	UFUNCTION(BlueprintPure, Category = "Math|Random", meta = (Keywords = "RandomVector", ScriptMethod = "RandomUnitVectorInEllipticalConeInDegrees", ScriptMethodMutable))
+	static inline FVector RandomUnitVectorInEllipticalConeInDegreesFromStream(const FRandomStream& Stream, const FVector& ConeDir, float MaxYawInDegrees, float MaxPitchInDegrees)
+	{
+		return RandomUnitVectorInEllipticalConeInRadiansFromStream(Stream, ConeDir, FMath::DegreesToRadians(MaxYawInDegrees), FMath::DegreesToRadians(MaxPitchInDegrees));
+	}
+
+	UE_DEPRECATED(5.2, "Use RandomUnitVectorInEllipticalConeInDegreesFromStream taking the FRandomStream as the first argument.")
 	static inline FVector RandomUnitVectorInEllipticalConeInDegreesFromStream(const FVector& ConeDir, float MaxYawInDegrees, float MaxPitchInDegrees, const FRandomStream& Stream)
 	{
-		return RandomUnitVectorInEllipticalConeInRadiansFromStream(ConeDir, FMath::DegreesToRadians(MaxYawInDegrees), FMath::DegreesToRadians(MaxPitchInDegrees), Stream);
+		return RandomUnitVectorInEllipticalConeInDegreesFromStream(Stream, ConeDir, MaxYawInDegrees, MaxPitchInDegrees);
 	}
 
 	/**
diff --git a/Engine/Source/Runtime/Engine/Private/KismetMathLibrary.cpp b/Engine/Source/Runtime/Engine/Private/KismetMathLibrary.cpp
--- a/Engine/Source/Runtime/Engine/Private/KismetMathLibrary.cpp
+++ b/Engine/Source/Runtime/Engine/Private/KismetMathLibrary.cpp
@@ -182,7 +182,7 @@ bool UKismetMathLibrary::RandomBoolWithWeight(float Weight)
 
 }
 
-bool UKismetMathLibrary::RandomBoolWithWeightFromStream(float Weight, const FRandomStream& RandomStream)
+bool UKismetMathLibrary::RandomBoolWithWeightFromStream(const FRandomStream& RandomStream, float Weight)
 {
 	//If the Weight equals to 0.0f then always return false
 	if (Weight <= 0.0f)
@@ -1052,12 +1052,12 @@ FRotator UKismetMathLibrary::FindRelativeLookAtRotation(const FTransform& StartT
 	return NormalizedDeltaRotator(FindLookAtRotation(StartTransform.GetLocation(), TargetLocation), StartTransform.GetRotation().Rotator());
 }
 
-int32 UKismetMathLibrary::RandomIntegerFromStream(int32 Max, const FRandomStream& Stream)
+int32 UKismetMathLibrary::RandomIntegerFromStream(const FRandomStream& Stream, int32 Max)
 {
 	return Stream.RandHelper(Max);
 }
 
-int32 UKismetMathLibrary::RandomIntegerInRangeFromStream(int32 Min, int32 Max, const FRandomStream& Stream)
+int32 UKismetMathLibrary::RandomIntegerInRangeFromStream(const FRandomStream& Stream, int32 Min, int32 Max)
 {
 	return Stream.RandRange(Min, Max);
 }
@@ -1082,7 +1082,7 @@ float UKismetMathLibrary::RandomFloatFromStream(const FRandomStream& Stream)
 	return Stream.FRand();
 }
 
-float UKismetMathLibrary::RandomFloatInRangeFromStream(float Min, float Max, const FRandomStream& Stream)
+float UKismetMathLibrary::RandomFloatInRangeFromStream(const FRandomStream& Stream, float Min, float Max)
 {
 	return Min + (Max - Min) * RandomFloatFromStream(Stream);
 }
@@ -1092,19 +1092,19 @@ FVector UKismetMathLibrary::RandomUnitVectorFromStream(const FRandomStream& Stre
 	return Stream.VRand();
 }
 
-FVector UKismetMathLibrary::RandomPointInBoundingBoxFromStream(const FVector Center, const FVector HalfSize, const FRandomStream& Stream)
+FVector UKismetMathLibrary::RandomPointInBoundingBoxFromStream(const FRandomStream& Stream, const FVector Center, const FVector HalfSize)
 {
 	const FVector BoxMin = Center - HalfSize;
 	const FVector BoxMax = Center + HalfSize;
 	return Stream.RandPointInBox(FBox(BoxMin, BoxMax));
 }
 
-FVector UKismetMathLibrary::RandomPointInBoundingBoxFromStream_Box(const FBox Box, const FRandomStream& Stream)
+FVector UKismetMathLibrary::RandomPointInBoundingBoxFromStream_Box(const FRandomStream& Stream, const FBox Box)
 {
 	return Stream.RandPointInBox(Box);
 }
 
-FRotator UKismetMathLibrary::RandomRotatorFromStream(bool bRoll, const FRandomStream& Stream)
+FRotator UKismetMathLibrary::RandomRotatorFromStream(const FRandomStream& Stream, bool bRoll)
 {
 	FRotator RRot;
 	RRot.Yaw = RandomFloatFromStream(Stream) * 360.f;
@@ -1465,12 +1465,12 @@ FRandomStream UKismetMathLibrary::MakeRandomStream(int32 InitialSeed)
 	return FRandomStream(InitialSeed);
 }
 
-FVector UKismetMathLibrary::RandomUnitVectorInConeInRadiansFromStream(const FVector& ConeDir, float ConeHalfAngleInRadians, const FRandomStream& Stream)
+FVector UKismetMathLibrary::RandomUnitVectorInConeInRadiansFromStream(const FRandomStream& Stream, const FVector& ConeDir, float ConeHalfAngleInRadians)
 {
 	return Stream.VRandCone(ConeDir, ConeHalfAngleInRadians);
 }
 
-FVector UKismetMathLibrary::RandomUnitVectorInEllipticalConeInRadiansFromStream(const FVector& ConeDir, float MaxYawInRadians, float MaxPitchInRadians, const FRandomStream& Stream)
+FVector UKismetMathLibrary::RandomUnitVectorInEllipticalConeInRadiansFromStream(const FRandomStream& Stream, const FVector& ConeDir, float MaxYawInRadians, float MaxPitchInRadians)
 {
 	return Stream.VRandCone(ConeDir, MaxYawInRadians, MaxPitchInRadians);
 }
