Add warning and bind pose when mirroring additive animations

diff --git a/Engine/Source/Runtime/AnimGraphRuntime/Private/AnimNodes/AnimNode_Mirror.cpp b/Engine/Source/Runtime/AnimGraphRuntime/Private/AnimNodes/AnimNode_Mirror.cpp
--- a/Engine/Source/Runtime/AnimGraphRuntime/Private/AnimNodes/AnimNode_Mirror.cpp
+++ b/Engine/Source/Runtime/AnimGraphRuntime/Private/AnimNodes/AnimNode_Mirror.cpp
@@ -179,14 +179,26 @@ void FAnimNode_MirrorBase::Evaluate_AnyThread(FPoseContext& Output)
 	{
 		if (GetBoneMirroring())
 		{
-			const FBoneContainer& BoneContainer = Output.Pose.GetBoneContainer();
-			const TArray<FBoneIndexType>& RequiredBoneIndices = BoneContainer.GetBoneIndicesArray();
-			int32 NumReqBones = RequiredBoneIndices.Num();
-			if (CompactPoseMirrorBones.Num() != NumReqBones)
+			if (Output.ExpectsAdditivePose())
 			{
-				FillCompactPoseAndComponentRefRotations(BoneContainer);
+				FText Message = FText::Format(LOCTEXT("AdditiveMirrorWarning", "Trying to mirror an additive animation pose is not supported in anim instance '{0}'"), FText::FromString(Output.AnimInstanceProxy->GetAnimInstanceName()));
+				Output.LogMessage(EMessageSeverity::Warning, Message);
+
+				// Force a bind pose to make it obvious
+				Output.Pose.ResetToAdditiveIdentity();
+			}
+			else
+			{
+				const FBoneContainer& BoneContainer = Output.Pose.GetBoneContainer();
+				const TArray<FBoneIndexType>& RequiredBoneIndices = BoneContainer.GetBoneIndicesArray();
+				int32 NumReqBones = RequiredBoneIndices.Num();
+				if (CompactPoseMirrorBones.Num() != NumReqBones)
+				{
+					FillCompactPoseAndComponentRefRotations(BoneContainer);
+				}
+
+				FAnimationRuntime::MirrorPose(Output.Pose, MirrorDataTable->MirrorAxis, CompactPoseMirrorBones, ComponentSpaceRefRotations);
 			}
-			FAnimationRuntime::MirrorPose(Output.Pose, MirrorDataTable->MirrorAxis, CompactPoseMirrorBones, ComponentSpaceRefRotations);
 		}
 
 		if (GetCurveMirroring())
diff --git a/Engine/Source/Runtime/Engine/Private/Animation/AnimationRuntime.cpp b/Engine/Source/Runtime/Engine/Private/Animation/AnimationRuntime.cpp
--- a/Engine/Source/Runtime/Engine/Private/Animation/AnimationRuntime.cpp
+++ b/Engine/Source/Runtime/Engine/Private/Animation/AnimationRuntime.cpp
@@ -1292,6 +1292,26 @@ void FAnimationRuntime::MirrorPose(FCompactPose& Pose, EAxis::Type MirrorAxis, c
 		return;
 	}
 
+	// Mirroring is authored in object space and as such we must transform the local space transforms in object space in order
+	// to apply the object space mirroring axis. To facilitate this, we use object space transforms for the bind pose which can be cached.
+	// We ignore the translation/scale part of the bind pose as they don't impact mirroring.
+	// 
+	// Rotations, translations, and scales are all treated differently:
+	//    Rotation:
+	//        We transform the local space rotation into object space
+	//        We mirror the rotation axis
+	//        We apply a correction: if the source and target bones are different, we must account for the mirrored delta between the two
+	//        We transform the result back into local space
+	//    Translation:
+	//        We rotate the local space translation into object space
+	//        We mirror the result
+	//        We then rotate it back into local space
+	//    Scale:
+	//        Mirroring does not modify scale
+	// 
+	// This sadly doesn't quite work for additive poses because in order to transform it into the bind pose reference frame,
+	// we need the base pose it is applied on. Worse still, the base pose might not be static, it could be a time scaled sequence.
+
 	auto MirrorTransform = [&ComponentSpaceRefRotations, MirrorAxis](const FTransform& SourceTransform, const FCompactPoseBoneIndex& SourceParentIndex, const FCompactPoseBoneIndex& SourceBoneIndex, const FCompactPoseBoneIndex& TargetParentIndex, const FCompactPoseBoneIndex& TargetBoneIndex) -> FTransform {
 
 		const FQuat TargetParentRefRotation = ComponentSpaceRefRotations[TargetParentIndex];
