World Partition Actor Descriptors:
- Split bounds into editor and runtime parts.
- Don't use fixed bounds anymore for foliage (backout 22179457).

diff --git a/Engine/Source/Editor/SceneOutliner/Private/ActorDescTreeItem.cpp b/Engine/Source/Editor/SceneOutliner/Private/ActorDescTreeItem.cpp
--- a/Engine/Source/Editor/SceneOutliner/Private/ActorDescTreeItem.cpp
+++ b/Engine/Source/Editor/SceneOutliner/Private/ActorDescTreeItem.cpp
@@ -288,7 +288,7 @@ void FActorDescTreeItem::FocusActorBounds() const
 	if (FWorldPartitionActorDesc const* ActorDesc = ActorDescHandle.Get())
 	{
 		const bool bActiveViewportOnly = true;
-		GEditor->MoveViewportCamerasToBox(ActorDesc->GetBounds(), bActiveViewportOnly, 0.5f);
+		GEditor->MoveViewportCamerasToBox(ActorDesc->GetEditorBounds(), bActiveViewportOnly, 0.5f);
 	}
 }
 
diff --git a/Engine/Source/Editor/WorldPartitionEditor/Private/WorldPartition/SWorldPartitionEditorGrid2D.cpp b/Engine/Source/Editor/WorldPartitionEditor/Private/WorldPartition/SWorldPartitionEditorGrid2D.cpp
--- a/Engine/Source/Editor/WorldPartitionEditor/Private/WorldPartition/SWorldPartitionEditorGrid2D.cpp
+++ b/Engine/Source/Editor/WorldPartitionEditor/Private/WorldPartition/SWorldPartitionEditorGrid2D.cpp
@@ -191,7 +191,7 @@ void ForEachIntersectingLoaderAdapters(UWorldPartition* WorldPartition, const FB
 			{
 				if (IWorldPartitionActorLoaderInterface::ILoaderAdapter* LoaderAdapter = Cast<IWorldPartitionActorLoaderInterface>(Actor)->GetLoaderAdapter())
 				{
-					if (IsBoundsSelected(SelectBox, ActorDesc->GetBounds()))
+					if (IsBoundsSelected(SelectBox, ActorDesc->GetEditorBounds()))
 					{
 						if (!Func(Actor))
 						{
@@ -214,7 +214,7 @@ public:
 		, bUseActor(bInUseActor)
 	{}
 
-	FBox GetBounds() const
+	FBox GetEditorBounds() const
 	{
 		if (bUseActor)
 		{
@@ -224,7 +224,7 @@ public:
 			}
 		}
 
-		return ActorDesc->GetBounds();
+		return ActorDesc->GetEditorBounds();
 	}
 
 	AActor* GetActor() const
@@ -1442,7 +1442,7 @@ uint32 SWorldPartitionEditorGrid2D::PaintActors(const FGeometry& AllottedGeometr
 				}
 			};
 
-			const FBox ActorBounds = ActorDescView.GetBounds();
+			const FBox ActorBounds = ActorDescView.GetEditorBounds();
 			const AActor* Actor = ActorDescView.GetActor();
 			const bool bIsSelected = Actor ? Actor->IsSelected() : WorldPartitionSubsystem->SelectedActorDescs.Contains(ActorDescView.GetActorDesc());
 			const bool bIsSpatiallyLoaded = ActorDescView.GetIsSpatiallyLoaded();
@@ -1450,7 +1450,7 @@ uint32 SWorldPartitionEditorGrid2D::PaintActors(const FGeometry& AllottedGeometr
 
 			if (bIsSelected)
 			{
-				const FBox ActorDescBounds = ActorDescView.GetActorDesc()->GetBounds();
+				const FBox ActorDescBounds = ActorDescView.GetActorDesc()->GetEditorBounds();
 				if (!ActorDescBounds.Equals(ActorBounds, 1.0f))
 				{
 					ShowActorBox(ActorDescBounds, false, bIsSpatiallyLoaded, ActorLabel);
@@ -1969,7 +1969,7 @@ FReply SWorldPartitionEditorGrid2D::FocusSelection()
 	UWorldPartitionSubsystem* WorldPartitionSubsystem = UWorld::GetSubsystem<UWorldPartitionSubsystem>(World);
 	for (FWorldPartitionActorDesc* SelectedActorDesc : WorldPartitionSubsystem->SelectedActorDescs)
 	{
-		SelectionBox += SelectedActorDesc->GetBounds();
+		SelectionBox += SelectedActorDesc->GetEditorBounds();
 	}
 
 	if (!SelectionBox.IsValid)
diff --git a/Engine/Source/Runtime/Engine/Private/LevelInstance/LevelInstanceSubsystem.cpp b/Engine/Source/Runtime/Engine/Private/LevelInstance/LevelInstanceSubsystem.cpp
--- a/Engine/Source/Runtime/Engine/Private/LevelInstance/LevelInstanceSubsystem.cpp
+++ b/Engine/Source/Runtime/Engine/Private/LevelInstance/LevelInstanceSubsystem.cpp
@@ -1471,7 +1471,7 @@ void ULevelInstanceSubsystem::FActorDescContainerInstanceManager::FActorDescCont
 		{
 			continue;
 		}
-		Bounds += ActorDescIt->GetBounds();
+		Bounds += ActorDescIt->GetRuntimeBounds();
 	}
 }
 
diff --git a/Engine/Source/Runtime/Engine/Private/WorldPartition/ActorPartition/PartitionActorDesc.cpp b/Engine/Source/Runtime/Engine/Private/WorldPartition/ActorPartition/PartitionActorDesc.cpp
--- a/Engine/Source/Runtime/Engine/Private/WorldPartition/ActorPartition/PartitionActorDesc.cpp
+++ b/Engine/Source/Runtime/Engine/Private/WorldPartition/ActorPartition/PartitionActorDesc.cpp
@@ -1,6 +1,7 @@
 // Copyright Epic Games, Inc. All Rights Reserved.
 
 #include "WorldPartition/ActorPartition/PartitionActorDesc.h"
+#include "ActorPartition/ActorPartitionSubsystem.h"
 #include "ActorPartition/PartitionActor.h"
 
 #if WITH_EDITOR
@@ -52,4 +53,10 @@ bool FPartitionActorDesc::Equals(const FWorldPartitionActorDesc* Other) const
 
 	return false;
 }
+
+FBox FPartitionActorDesc::GetEditorBounds() const
+{
+	const UActorPartitionSubsystem::FCellCoord CellCoord(GridIndexX, GridIndexY, GridIndexZ, 0);
+	return UActorPartitionSubsystem::FCellCoord::GetCellBounds(CellCoord, GridSize);
+}
 #endif
diff --git a/Engine/Source/Runtime/Engine/Private/WorldPartition/Landscape/LandscapeActorDesc.cpp b/Engine/Source/Runtime/Engine/Private/WorldPartition/Landscape/LandscapeActorDesc.cpp
--- a/Engine/Source/Runtime/Engine/Private/WorldPartition/Landscape/LandscapeActorDesc.cpp
+++ b/Engine/Source/Runtime/Engine/Private/WorldPartition/Landscape/LandscapeActorDesc.cpp
@@ -48,6 +48,12 @@ void FLandscapeActorDesc::Serialize(FArchive& Ar)
 	}
 }
 
+FBox FLandscapeActorDesc::GetEditorBounds() const
+{
+	// We need to skip super class since we aren't using grid indices as it should be (it's in Landscape space).
+	return FWorldPartitionActorDesc::GetEditorBounds();
+}
+
 void FLandscapeActorDesc::Unload()
 {
 	if (ALandscapeProxy* LandscapeProxy = Cast<ALandscapeProxy>(GetActor()))
diff --git a/Engine/Source/Runtime/Engine/Private/WorldPartition/LoaderAdapter/LoaderAdapterSpatial.cpp b/Engine/Source/Runtime/Engine/Private/WorldPartition/LoaderAdapter/LoaderAdapterSpatial.cpp
--- a/Engine/Source/Runtime/Engine/Private/WorldPartition/LoaderAdapter/LoaderAdapterSpatial.cpp
+++ b/Engine/Source/Runtime/Engine/Private/WorldPartition/LoaderAdapter/LoaderAdapterSpatial.cpp
@@ -22,7 +22,7 @@ void ILoaderAdapterSpatial::ForEachActor(TFunctionRef<void(const FWorldPartition
 
 		WorldPartition->EditorHash->ForEachIntersectingActor(*GetBoundingBox(), [this, WorldPartition, &InOperation](FWorldPartitionActorDesc* ActorDesc)
 		{
-			if (Intersect(ActorDesc->GetBounds()))
+			if (Intersect(ActorDesc->GetEditorBounds()))
 			{
 				FWorldPartitionHandle ActorHandle(WorldPartition, ActorDesc->GetGuid());
 				InOperation(ActorHandle);
diff --git a/Engine/Source/Runtime/Engine/Private/WorldPartition/RuntimeSpatialHash/RuntimeSpatialHashGridHelper.cpp b/Engine/Source/Runtime/Engine/Private/WorldPartition/RuntimeSpatialHash/RuntimeSpatialHashGridHelper.cpp
--- a/Engine/Source/Runtime/Engine/Private/WorldPartition/RuntimeSpatialHash/RuntimeSpatialHashGridHelper.cpp
+++ b/Engine/Source/Runtime/Engine/Private/WorldPartition/RuntimeSpatialHash/RuntimeSpatialHashGridHelper.cpp
@@ -192,7 +192,7 @@ FSquare2DGridHelper GetPartitionedActors(const FBox& WorldBounds, const FSpatial
 			if (bUseLocation)
 			{
 				// Find grid level that best matches actor set bounds
-				const float MaxLength = InActorSetInstanceBounds.GetExtent().GetMax() * 2.0;
+				const float MaxLength = InActorSetInstanceBounds.GetSize().GetMax();
 				OutGridLevel = FMath::Min(FMath::CeilToInt(FMath::Max<float>(FMath::Log2(MaxLength / CellSize), 0)), GridLevelCount - 1);
 			}
 			return bUseLocation;
diff --git a/Engine/Source/Runtime/Engine/Private/WorldPartition/RuntimeSpatialHash/RuntimeSpatialHashHLOD.cpp b/Engine/Source/Runtime/Engine/Private/WorldPartition/RuntimeSpatialHash/RuntimeSpatialHashHLOD.cpp
--- a/Engine/Source/Runtime/Engine/Private/WorldPartition/RuntimeSpatialHash/RuntimeSpatialHashHLOD.cpp
+++ b/Engine/Source/Runtime/Engine/Private/WorldPartition/RuntimeSpatialHash/RuntimeSpatialHashHLOD.cpp
@@ -193,8 +193,8 @@ static TArray<FGuid> GenerateHLODsForGrid(UWorldPartition* WorldPartition, const
 					{
 						const IStreamingGenerationContext::FActorInstance& ActorInstance = ActorInstances.Emplace_GetRef(ActorGuid, ActorSetInstance);
 
-						const double ActorMinZ = ActorInstance.GetActorDescView().GetBounds().Min.Z;
-						const double ActorMaxZ = ActorInstance.GetActorDescView().GetBounds().Max.Z;
+						const double ActorMinZ = ActorInstance.GetActorDescView().GetRuntimeBounds().Min.Z;
+						const double ActorMaxZ = ActorInstance.GetActorDescView().GetRuntimeBounds().Max.Z;
 
 						CellBounds.Min.Z = FMath::Min(CellBounds.Min.Z, ActorMinZ);
 						CellBounds.Max.Z = FMath::Max(CellBounds.Max.Z, ActorMaxZ);
@@ -524,7 +524,7 @@ bool UWorldPartitionRuntimeSpatialHash::GenerateHLOD(ISourceControlHelper* Sourc
 			
 			IStreamingGenerationContext::FActorSetInstance& NewHLODActorSetInstance = HLODActorSetInstances.Emplace_GetRef();
 			
-			NewHLODActorSetInstance.Bounds = HLODActorDescView.GetBounds();
+			NewHLODActorSetInstance.Bounds = HLODActorDescView.GetRuntimeBounds();
 			NewHLODActorSetInstance.RuntimeGrid = HLODActorDescView.GetRuntimeGrid();
 			NewHLODActorSetInstance.bIsSpatiallyLoaded = HLODActorDescView.GetIsSpatiallyLoaded();
 			NewHLODActorSetInstance.DataLayers = UDataLayerSubsystem::GetRuntimeDataLayerInstances(MainActorSetContainer->ActorDescContainer->GetWorld(), HLODActorDescView.GetRuntimeDataLayers());
diff --git a/Engine/Source/Runtime/Engine/Private/WorldPartition/WorldPartition.cpp b/Engine/Source/Runtime/Engine/Private/WorldPartition/WorldPartition.cpp
--- a/Engine/Source/Runtime/Engine/Private/WorldPartition/WorldPartition.cpp
+++ b/Engine/Source/Runtime/Engine/Private/WorldPartition/WorldPartition.cpp
@@ -1352,7 +1352,7 @@ void UWorldPartition::Tick(float DeltaSeconds)
 			{
 				if (ActorDescIterator->GetIsSpatiallyLoadedRaw())
 				{
-					AllActorsBounds += ActorDescIterator->GetBounds();
+					AllActorsBounds += ActorDescIterator->GetEditorBounds();
 
 					// Warn the user if the world becomes larger that 4km in any axis
 					if (AllActorsBounds.GetSize().GetMax() >= UWorldPartition::WorldExtentToEnableStreaming)
@@ -1600,7 +1600,7 @@ void UWorldPartition::DumpActorDescs(const FString& Path)
 		}
 		ActorDescs.Sort([](const FWorldPartitionActorDesc& A, const FWorldPartitionActorDesc& B)
 		{
-			return A.GetBounds().GetExtent().GetMax() < B.GetBounds().GetExtent().GetMax();
+			return A.GetRuntimeBounds().GetExtent().GetMax() < B.GetRuntimeBounds().GetExtent().GetMax();
 		});
 		for (const FWorldPartitionActorDesc* ActorDescIterator : ActorDescs)
 		{
diff --git a/Engine/Source/Runtime/Engine/Private/WorldPartition/WorldPartitionActorDesc.cpp b/Engine/Source/Runtime/Engine/Private/WorldPartition/WorldPartitionActorDesc.cpp
--- a/Engine/Source/Runtime/Engine/Private/WorldPartition/WorldPartitionActorDesc.cpp
+++ b/Engine/Source/Runtime/Engine/Private/WorldPartition/WorldPartitionActorDesc.cpp
@@ -256,6 +256,18 @@ FString FWorldPartitionActorDesc::ToString(EToStringMode Mode) const
 
 	if (Mode >= EToStringMode::Compact)
 	{
+		FString BoundsStr;
+		const FBox EditorBounds = GetEditorBounds();
+		const FBox RuntimeBounds = GetRuntimeBounds();
+		if (EditorBounds.Equals(RuntimeBounds))
+		{
+			BoundsStr = RuntimeBounds.ToString();
+		}
+		else
+		{
+			BoundsStr = *FString::Printf(TEXT("(Editor:%s Runtime:%s)"), *EditorBounds.ToString(), *RuntimeBounds.ToString());
+		}
+
 		Result.Appendf(
 			TEXT(" BaseClass:%s NativeClass:%s Name:%s Label:%s SpatiallyLoaded:%s Bounds:%s RuntimeGrid:%s EditorOnly:%s RuntimeOnly:%s HLODRelevant:%s"),
 			*BaseClass.ToString(), 
@@ -263,7 +275,7 @@ FString FWorldPartitionActorDesc::ToString(EToStringMode Mode) const
 			*GetActorName().ToString(),
 			*GetActorLabel().ToString(),
 			GetBoolStr(bIsSpatiallyLoaded),
-			*GetBounds().ToString(),
+			*BoundsStr,
 			*RuntimeGrid.ToString(),
 			GetBoolStr(bActorIsEditorOnly),
 			GetBoolStr(bActorIsRuntimeOnly),
@@ -496,7 +508,12 @@ PRAGMA_ENABLE_DEPRECATION_WARNINGS
 	}
 }
 
-FBox FWorldPartitionActorDesc::GetBounds() const
+FBox FWorldPartitionActorDesc::GetEditorBounds() const
+{
+	return FBox(BoundsLocation - BoundsExtent, BoundsLocation + BoundsExtent);
+}
+
+FBox FWorldPartitionActorDesc::GetRuntimeBounds() const
 {
 	return FBox(BoundsLocation - BoundsExtent, BoundsLocation + BoundsExtent);
 }
diff --git a/Engine/Source/Runtime/Engine/Private/WorldPartition/WorldPartitionActorDescView.cpp b/Engine/Source/Runtime/Engine/Private/WorldPartition/WorldPartitionActorDescView.cpp
--- a/Engine/Source/Runtime/Engine/Private/WorldPartition/WorldPartitionActorDescView.cpp
+++ b/Engine/Source/Runtime/Engine/Private/WorldPartition/WorldPartitionActorDescView.cpp
@@ -38,10 +38,12 @@ UClass* FWorldPartitionActorDescView::GetActorNativeClass() const
 	return ActorDesc->GetActorNativeClass();
 }
 
+PRAGMA_DISABLE_DEPRECATION_WARNINGS
 FVector FWorldPartitionActorDescView::GetOrigin() const
 {
 	return ActorDesc->GetOrigin();
 }
+PRAGMA_ENABLE_DEPRECATION_WARNINGS
 
 FName FWorldPartitionActorDescView::GetRuntimeGrid() const
 {
@@ -119,11 +121,22 @@ FName FWorldPartitionActorDescView::GetActorName() const
 	return ActorDesc->GetActorName();
 }
 
-
+PRAGMA_DISABLE_DEPRECATION_WARNINGS
 FBox FWorldPartitionActorDescView::GetBounds() const
 {
 	return ActorDesc->GetBounds();
 }
+PRAGMA_ENABLE_DEPRECATION_WARNINGS
+
+FBox FWorldPartitionActorDescView::GetEditorBounds() const
+{
+	return ActorDesc->GetEditorBounds();
+}
+
+FBox FWorldPartitionActorDescView::GetRuntimeBounds() const
+{
+	return ActorDesc->GetRuntimeBounds();
+}
 
 const TArray<FGuid>& FWorldPartitionActorDescView::GetReferences() const
 {
diff --git a/Engine/Source/Runtime/Engine/Private/WorldPartition/WorldPartitionEditorSpatialHash.cpp b/Engine/Source/Runtime/Engine/Private/WorldPartition/WorldPartitionEditorSpatialHash.cpp
--- a/Engine/Source/Runtime/Engine/Private/WorldPartition/WorldPartitionEditorSpatialHash.cpp
+++ b/Engine/Source/Runtime/Engine/Private/WorldPartition/WorldPartitionEditorSpatialHash.cpp
@@ -71,14 +71,14 @@ void UWorldPartitionEditorSpatialHash::Tick(float DeltaSeconds)
 
 			for (const FWorldPartitionHandle& Actor : Cell->Actors)
 			{
-				RuntimeBounds += Actor->GetBounds();
+				RuntimeBounds += Actor->GetEditorBounds();
 			}
 		}
 
 		NonSpatialBounds.Init();
 		for (const FWorldPartitionHandle& Actor : AlwaysLoadedCell->Actors)
 		{
-			NonSpatialBounds += Actor->GetBounds();
+			NonSpatialBounds += Actor->GetEditorBounds();
 		}
 
 		const int32 OldLevel = GetLevelForBox(EditorBounds);
@@ -129,7 +129,7 @@ void UWorldPartitionEditorSpatialHash::HashActor(FWorldPartitionHandle& InActorH
 	check(InActorHandle.IsValid());
 
 	const bool bConsiderActorSpatiallyLoaded = InActorHandle->GetIsSpatiallyLoaded();
-	const FBox ActorBounds = bConsiderActorSpatiallyLoaded ? InActorHandle->GetBounds() : FBox(ForceInit);
+	const FBox ActorBounds = bConsiderActorSpatiallyLoaded ? InActorHandle->GetEditorBounds() : FBox(ForceInit);
 
 #if DO_CHECK
 	check(!HashedActors.Contains(InActorHandle->GetGuid()));
@@ -139,7 +139,7 @@ void UWorldPartitionEditorSpatialHash::HashActor(FWorldPartitionHandle& InActorH
 	if (!bConsiderActorSpatiallyLoaded)
 	{
 		AlwaysLoadedCell->Actors.Add(InActorHandle);
-		NonSpatialBounds += InActorHandle->GetBounds();
+		NonSpatialBounds += InActorHandle->GetEditorBounds();
 	}
 	else
 	{
@@ -232,7 +232,7 @@ void UWorldPartitionEditorSpatialHash::HashActor(FWorldPartitionHandle& InActorH
 			}
 		}
 
-		RuntimeBounds += InActorHandle->GetBounds();
+		RuntimeBounds += InActorHandle->GetEditorBounds();
 	}
 }
 
@@ -241,7 +241,7 @@ void UWorldPartitionEditorSpatialHash::UnhashActor(FWorldPartitionHandle& InActo
 	check(InActorHandle.IsValid());
 
 	const bool bConsiderActorSpatiallyLoaded = InActorHandle->GetIsSpatiallyLoaded();
-	const FBox ActorBounds = bConsiderActorSpatiallyLoaded ? InActorHandle->GetBounds() : FBox(ForceInit);
+	const FBox ActorBounds = bConsiderActorSpatiallyLoaded ? InActorHandle->GetEditorBounds() : FBox(ForceInit);
 
 #if DO_CHECK
 	const FBox OldActorBounds = HashedActors.FindAndRemoveChecked(InActorHandle->GetGuid());
@@ -318,7 +318,7 @@ int32 UWorldPartitionEditorSpatialHash::ForEachIntersectingActor(const FBox& Box
 
 					if (!bWasAlreadyInSet)
 					{
-						if (Box.Intersect(ActorDesc->GetBounds()))
+						if (Box.Intersect(ActorDesc->GetEditorBounds()))
 						{
 							InOperation(*ActorDesc);
 						}
@@ -334,7 +334,7 @@ int32 UWorldPartitionEditorSpatialHash::ForEachIntersectingActor(const FBox& Box
 		{
 			if (ActorDesc.IsValid())
 			{
-				if (Box.Intersect(ActorDesc->GetBounds()))
+				if (Box.Intersect(ActorDesc->GetEditorBounds()))
 				{
 					bool bWasAlreadyInSet;
 					IntersectedActorDescs.Add(*ActorDesc, &bWasAlreadyInSet);
diff --git a/Engine/Source/Runtime/Engine/Private/WorldPartition/WorldPartitionRuntimeSpatialHash.cpp b/Engine/Source/Runtime/Engine/Private/WorldPartition/WorldPartitionRuntimeSpatialHash.cpp
--- a/Engine/Source/Runtime/Engine/Private/WorldPartition/WorldPartitionRuntimeSpatialHash.cpp
+++ b/Engine/Source/Runtime/Engine/Private/WorldPartition/WorldPartitionRuntimeSpatialHash.cpp
@@ -1444,7 +1444,7 @@ bool UWorldPartitionRuntimeSpatialHash::CreateStreamingGrid(const FSpatialHashRu
 					const FWorldPartitionActorDescView& ActorDescView = ActorInstance.GetActorDescView();
 					StreamingCell->AddActorToCell(ActorDescView, ActorInstance.GetContainerID(), ActorInstance.GetTransform(), ActorInstance.GetActorDescContainer());
 					
-					CellContentBounds += ActorDescView.GetBounds().TransformBy(ActorInstance.GetTransform());
+					CellContentBounds += ActorDescView.GetRuntimeBounds().TransformBy(ActorInstance.GetTransform());
 
 					if (ActorInstance.GetContainerID().IsMainContainer() && StreamingCell->UnsavedActorsContainer)
 					{
diff --git a/Engine/Source/Runtime/Engine/Private/WorldPartition/WorldPartitionStreamingGeneration.cpp b/Engine/Source/Runtime/Engine/Private/WorldPartition/WorldPartitionStreamingGeneration.cpp
--- a/Engine/Source/Runtime/Engine/Private/WorldPartition/WorldPartitionStreamingGeneration.cpp
+++ b/Engine/Source/Runtime/Engine/Private/WorldPartition/WorldPartitionStreamingGeneration.cpp
@@ -160,7 +160,7 @@ class FWorldPartitionStreamingGenerator
 					for (const FGuid& ActorGuid : ActorSetInstance.ActorSet->Actors)
 					{
 						const FWorldPartitionActorDescView& ActorDescView = ActorSetContainer.ActorDescViewMap->FindByGuidChecked(ActorGuid);
-						ActorSetInstance.Bounds += ActorDescView.GetBounds().TransformBy(ContainerInstanceDescriptor.Transform);
+						ActorSetInstance.Bounds += ActorDescView.GetRuntimeBounds().TransformBy(ContainerInstanceDescriptor.Transform);
 					}
 				}
 			}
@@ -464,7 +464,7 @@ class FWorldPartitionStreamingGenerator
 			{
 				if (ActorDescView.GetIsSpatiallyLoaded())
 				{
-					ContainerInstanceDescriptorBounds += ActorDescView.GetBounds().TransformBy(ContainerInstanceDescriptorTransform);
+					ContainerInstanceDescriptorBounds += ActorDescView.GetRuntimeBounds().TransformBy(ContainerInstanceDescriptorTransform);
 				}
 			});
 		}
diff --git a/Engine/Source/Runtime/Engine/Private/WorldPartitionBlueprintLibrary.cpp b/Engine/Source/Runtime/Engine/Private/WorldPartitionBlueprintLibrary.cpp
--- a/Engine/Source/Runtime/Engine/Private/WorldPartitionBlueprintLibrary.cpp
+++ b/Engine/Source/Runtime/Engine/Private/WorldPartitionBlueprintLibrary.cpp
@@ -34,7 +34,7 @@ FActorDesc::FActorDesc(const FWorldPartitionActorDesc& InActorDesc, const FTrans
 
 	Name = InActorDesc.GetActorName();
 	Label = InActorDesc.GetActorLabel();
-	Bounds = InActorDesc.GetBounds().TransformBy(InTransform);
+	Bounds = InActorDesc.GetEditorBounds().TransformBy(InTransform);
 	RuntimeGrid = InActorDesc.GetRuntimeGrid();
 	bIsSpatiallyLoaded = InActorDesc.GetIsSpatiallyLoaded();
 	bActorIsEditorOnly = InActorDesc.GetActorIsEditorOnly();
@@ -158,7 +158,7 @@ bool UWorldPartitionBlueprintLibrary::GetIntersectingActorDescs(const UActorDesc
 
 	for (FActorDescList::TConstIterator<> ActorDescIt(InContainer); ActorDescIt; ++ActorDescIt)
 	{
-		if (ActorDescIt->GetBounds().Intersect(InBox))
+		if (ActorDescIt->GetEditorBounds().Intersect(InBox))
 		{
 			bResult &= HandleIntersectingActorDesc(*ActorDescIt, InBox, InTransform, OutActorDescs);
 		}
diff --git a/Engine/Source/Runtime/Engine/Public/WorldPartition/ActorPartition/PartitionActorDesc.h b/Engine/Source/Runtime/Engine/Public/WorldPartition/ActorPartition/PartitionActorDesc.h
--- a/Engine/Source/Runtime/Engine/Public/WorldPartition/ActorPartition/PartitionActorDesc.h
+++ b/Engine/Source/Runtime/Engine/Public/WorldPartition/ActorPartition/PartitionActorDesc.h
@@ -18,9 +18,11 @@ public:
 	int64 GridIndexY;
 	int64 GridIndexZ;
 	FGuid GridGuid;
+
 protected:
 	virtual void Init(const AActor* InActor) override;
 	virtual bool Equals(const FWorldPartitionActorDesc* Other) const override;
 	virtual void Serialize(FArchive& Ar) override;
+	virtual FBox GetEditorBounds() const override;
 };
 #endif
diff --git a/Engine/Source/Runtime/Engine/Public/WorldPartition/Landscape/LandscapeActorDesc.h b/Engine/Source/Runtime/Engine/Public/WorldPartition/Landscape/LandscapeActorDesc.h
--- a/Engine/Source/Runtime/Engine/Public/WorldPartition/Landscape/LandscapeActorDesc.h
+++ b/Engine/Source/Runtime/Engine/Public/WorldPartition/Landscape/LandscapeActorDesc.h
@@ -20,6 +20,7 @@ public:
 
 protected:
 	virtual void Serialize(FArchive& Ar) override;
+	virtual FBox GetEditorBounds() const override;
 
 private:
 	FGuid LandscapeActorGuid;
diff --git a/Engine/Source/Runtime/Engine/Public/WorldPartition/WorldPartitionActorDesc.h b/Engine/Source/Runtime/Engine/Public/WorldPartition/WorldPartitionActorDesc.h
--- a/Engine/Source/Runtime/Engine/Public/WorldPartition/WorldPartitionActorDesc.h
+++ b/Engine/Source/Runtime/Engine/Public/WorldPartition/WorldPartitionActorDesc.h
@@ -85,7 +85,12 @@ PRAGMA_ENABLE_DEPRECATION_WARNINGS
 	inline FTopLevelAssetPath GetBaseClass() const { return BaseClass; }
 	inline FTopLevelAssetPath GetNativeClass() const { return NativeClass; }
 	inline UClass* GetActorNativeClass() const { return ActorNativeClass; }
+
+PRAGMA_DISABLE_DEPRECATION_WARNINGS
+	UE_DEPRECATED(5.2, "GetOrigin is deprecated.")
 	inline FVector GetOrigin() const { return GetBounds().GetCenter(); }
+PRAGMA_ENABLE_DEPRECATION_WARNINGS
+
 	inline FName GetRuntimeGrid() const { return RuntimeGrid; }
 	inline bool GetIsSpatiallyLoaded() const { return bIsForcedNonSpatiallyLoaded ? false : bIsSpatiallyLoaded; }
 	inline bool GetIsSpatiallyLoadedRaw() const { return bIsSpatiallyLoaded; }
@@ -109,7 +114,14 @@ PRAGMA_ENABLE_DEPRECATION_WARNINGS
 	inline FName GetActorLabel() const { return ActorLabel; }
 	inline FName GetFolderPath() const { return FolderPath; }
 	inline const FGuid& GetFolderGuid() const { return FolderGuid; }
-	FBox GetBounds() const;
+
+PRAGMA_DISABLE_DEPRECATION_WARNINGS
+	UE_DEPRECATED(5.2, "GetBounds is deprecated, GetEditorBounds or GetRuntimeBounds should be used instead.")
+	FBox GetBounds() const { return GetEditorBounds(); }
+PRAGMA_ENABLE_DEPRECATION_WARNINGS
+	virtual FBox GetEditorBounds() const;
+	FBox GetRuntimeBounds() const;
+
 	inline const FGuid& GetParentActor() const { return ParentActor; }
 	inline bool IsUsingDataLayerAsset() const { return bIsUsingDataLayerAsset; }
 	inline void AddProperty(FName PropertyName, FName PropertyValue = NAME_None) { Properties.AddProperty(PropertyName, PropertyValue); }
diff --git a/Engine/Source/Runtime/Engine/Public/WorldPartition/WorldPartitionActorDescView.h b/Engine/Source/Runtime/Engine/Public/WorldPartition/WorldPartitionActorDescView.h
--- a/Engine/Source/Runtime/Engine/Public/WorldPartition/WorldPartitionActorDescView.h
+++ b/Engine/Source/Runtime/Engine/Public/WorldPartition/WorldPartitionActorDescView.h
@@ -30,7 +30,10 @@ public:
 	FTopLevelAssetPath GetBaseClass() const;
 	FTopLevelAssetPath GetNativeClass() const;
 	UClass* GetActorNativeClass() const;
+PRAGMA_DISABLE_DEPRECATION_WARNINGS
+	UE_DEPRECATED(5.2, "GetOrigin is deprecated.")
 	FVector GetOrigin() const;
+PRAGMA_ENABLE_DEPRECATION_WARNINGS
 	FName GetRuntimeGrid() const;
 	bool GetIsSpatiallyLoaded() const;
 	bool GetActorIsEditorOnly() const;
@@ -49,7 +52,15 @@ PRAGMA_DISABLE_DEPRECATION_WARNINGS
 PRAGMA_ENABLE_DEPRECATION_WARNINGS
 	FSoftObjectPath GetActorSoftPath() const;
 	FName GetActorLabel() const;
+
+PRAGMA_DISABLE_DEPRECATION_WARNINGS
+	UE_DEPRECATED(5.2, "GetBounds is deprecated, GetEditorBounds or GetRuntimeBounds should be used instead.")
 	FBox GetBounds() const;
+PRAGMA_ENABLE_DEPRECATION_WARNINGS
+
+	FBox GetEditorBounds() const;
+	FBox GetRuntimeBounds() const;
+
 	const TArray<FGuid>& GetReferences() const;
 	FString ToString() const;
 	const FGuid& GetParentActor() const;
diff --git a/Engine/Source/Runtime/Foliage/Private/InstancedFoliage.cpp b/Engine/Source/Runtime/Foliage/Private/InstancedFoliage.cpp
--- a/Engine/Source/Runtime/Foliage/Private/InstancedFoliage.cpp
+++ b/Engine/Source/Runtime/Foliage/Private/InstancedFoliage.cpp
@@ -4158,17 +4158,6 @@ bool AInstancedFoliageActor::ShouldImport(FStringView ActorPropString, bool IsMo
 	return false;
 }
 
-FBox AInstancedFoliageActor::GetStreamingBounds() const
-{
-	if (UWorld* World = GetWorld(); World && World->IsPartitionedWorld())
-	{
-		const float HalfGridSize = GridSize * 0.5f;
-		return FBox(GetActorLocation() - HalfGridSize, GetActorLocation() + HalfGridSize);
-	}
-
-	return Super::GetStreamingBounds();
-}
-
 void AInstancedFoliageActor::ApplySelection(bool bApply)
 {
 	for (auto& Pair : FoliageInfos)
diff --git a/Engine/Source/Runtime/Foliage/Public/InstancedFoliageActor.h b/Engine/Source/Runtime/Foliage/Public/InstancedFoliageActor.h
--- a/Engine/Source/Runtime/Foliage/Public/InstancedFoliageActor.h
+++ b/Engine/Source/Runtime/Foliage/Public/InstancedFoliageActor.h
@@ -184,7 +184,6 @@ public:
 	virtual void PostDuplicate(bool bDuplicateForPIE) override;
 	virtual bool ShouldExport() override;
 	virtual bool ShouldImport(FStringView ActorPropString, bool IsMovingLevel) override;
-	virtual FBox GetStreamingBounds() const override;
 
 	// Called in response to BSP rebuilds to migrate foliage from obsolete to new components.
 	FOLIAGE_API void MapRebuild();
diff --git a/Engine/Source/Runtime/Landscape/Private/Landscape.cpp b/Engine/Source/Runtime/Landscape/Private/Landscape.cpp
--- a/Engine/Source/Runtime/Landscape/Private/Landscape.cpp
+++ b/Engine/Source/Runtime/Landscape/Private/Landscape.cpp
@@ -4735,7 +4735,7 @@ FBox ULandscapeInfo::GetCompleteBounds() const
 			}
 			else
 			{
-				Bounds += ActorDesc->GetBounds();
+				Bounds += ActorDesc->GetEditorBounds();
 			}
 		}
 
