Fixed RecursiveGetExpressionChain for Make and Set Material Attributes, in the case where we're looking for the MP_MaterialAttributes property : this property needs to be followed up since it's actually a group of several properties.
Also fixed potential issues with mobile material expressions not being properly filtered

robomerge[FN_EngineMerge] -UE5-RES

diff --git a/Engine/Source/Runtime/Engine/Private/Materials/Material.cpp b/Engine/Source/Runtime/Engine/Private/Materials/Material.cpp
--- a/Engine/Source/Runtime/Engine/Private/Materials/Material.cpp
+++ b/Engine/Source/Runtime/Engine/Private/Materials/Material.cpp
@@ -159,6 +159,18 @@ namespace MaterialImpl
 		{ MP_Normal, SF_Pixel },
 		{ MP_OpacityMask, SF_Pixel },
 	};
+
+	bool IsPropertyRelevantForMobile(EMaterialProperty InProperty)
+	{
+		for (FMobileProperty MobileProperty : GMobileRelevantMaterialProperties)
+		{
+			if (MobileProperty.Property == InProperty)
+			{
+				return true;
+			}
+		}
+		return false;
+	}
 }
 
 #if WITH_EDITOR
@@ -6365,20 +6377,6 @@ bool UMaterial::IsPropertySupported(EMaterialProperty InProperty) const
 	return bSupported;
 }
 
-bool UMaterial::IsPropertyRelevantForMobile(EMaterialProperty InProperty)
-{
-	using namespace MaterialImpl;
-
-	for (FMobileProperty MobileProperty : GMobileRelevantMaterialProperties)
-	{
-		if (MobileProperty.Property == InProperty)
-		{
-			return true;
-		}
-	}
-	return false;
-}
-
 static bool IsPropertyActive_Internal(EMaterialProperty InProperty,
 	EMaterialDomain Domain,
 	EBlendMode BlendMode,
