Remove device property functions as bluerpintcallable UFUNCTIONs to make sure that people use the Device subsystem instead of calling the functions directly. Make the Apply function protected.

#rnx
#prelfight skip

diff --git a/Engine/Source/Runtime/Engine/Classes/GameFramework/InputDeviceProperties.h b/Engine/Source/Runtime/Engine/Classes/GameFramework/InputDeviceProperties.h
--- a/Engine/Source/Runtime/Engine/Classes/GameFramework/InputDeviceProperties.h
+++ b/Engine/Source/Runtime/Engine/Classes/GameFramework/InputDeviceProperties.h
@@ -45,7 +45,7 @@ public:
 	* @param Duration			The number of seconds that this property has been active. Use this to get things like curve data over time.
 	* @return					A pointer to the evaluated input device property.
 	*/
-	UFUNCTION(BlueprintCallable, BlueprintNativeEvent, Category = "InputDevice")
+	UFUNCTION(BlueprintNativeEvent, Category = "InputDevice")
 	void EvaluateDeviceProperty(const FPlatformUserId PlatformUser, const float DeltaTime, const float Duration);
 
 	/** 
@@ -62,7 +62,7 @@ public:
 	* 
 	* @param PlatformUser		The platform user that should receive this device property change
 	*/
-	UFUNCTION(BlueprintCallable, BlueprintNativeEvent, Category = "InputDevice")
+	UFUNCTION(BlueprintNativeEvent, Category = "InputDevice")
 	void ResetDeviceProperty(const FPlatformUserId PlatformUser);
 
 	/**
@@ -79,16 +79,8 @@ public:
 	* 
 	* @param UserId		The owning Platform User whose input device this property should be applied to.
 	*/
-	UFUNCTION(BlueprintCallable, Category = "InputDevice")
+	UFUNCTION(Category = "InputDevice")
 	virtual void ApplyDeviceProperty(const FPlatformUserId UserId);
-
-	/** 
-	* Apply the given device property
-	* 
-	* @param UserId			The owning Platform User whose input device this property should be applied to.
-	* @param RawProperty	The internal input device property to apply.
-	*/
-	static void ApplyDeviceProperty(const FPlatformUserId UserId, FInputDeviceProperty* RawProperty);
 	
 	/** Gets a pointer to the current input device property that the IInputInterface can use. */
 	virtual FInputDeviceProperty* GetInternalDeviceProperty() { return nullptr; };
@@ -112,6 +104,14 @@ public:
 
 protected:
 
+	/**
+	* Apply the given device property
+	*
+	* @param UserId			The owning Platform User whose input device this property should be applied to.
+	* @param RawProperty	The internal input device property to apply.
+	*/
+	static void ApplyDeviceProperty_Internal(const FPlatformUserId UserId, FInputDeviceProperty* RawProperty);
+
 	/** Returns the device specific data for the given platform user. Returns the default data if none are given */
 	template<class TDataLayout>
 	const TDataLayout* GetDeviceSpecificData(const FPlatformUserId UserId, const TMap<FName, TDataLayout>& InMap) const;
diff --git a/Engine/Source/Runtime/Engine/Private/GameFramework/InputDeviceProperties.cpp b/Engine/Source/Runtime/Engine/Private/GameFramework/InputDeviceProperties.cpp
--- a/Engine/Source/Runtime/Engine/Private/GameFramework/InputDeviceProperties.cpp
+++ b/Engine/Source/Runtime/Engine/Private/GameFramework/InputDeviceProperties.cpp
@@ -17,10 +17,10 @@ UInputDeviceProperty::UInputDeviceProperty(const FObjectInitializer& ObjectIniti
 
 void UInputDeviceProperty::ApplyDeviceProperty(const FPlatformUserId UserId)
 {
-	UInputDeviceProperty::ApplyDeviceProperty(UserId, GetInternalDeviceProperty());
+	UInputDeviceProperty::ApplyDeviceProperty_Internal(UserId, GetInternalDeviceProperty());
 }
 
-void UInputDeviceProperty::ApplyDeviceProperty(const FPlatformUserId UserId, FInputDeviceProperty* RawProperty)
+void UInputDeviceProperty::ApplyDeviceProperty_Internal(const FPlatformUserId UserId, FInputDeviceProperty* RawProperty)
 {
 	if (ensure(RawProperty))
 	{
@@ -180,7 +180,7 @@ void UInputDeviceTriggerEffect::ResetDeviceProperty_Implementation(const FPlatfo
 	{
 		// Pass in our reset property
 		ResetProperty.AffectedTriggers = BaseTriggerData.AffectedTriggers;
-		ApplyDeviceProperty(PlatformUser, &ResetProperty);
+		ApplyDeviceProperty_Internal(PlatformUser, &ResetProperty);
 	}	
 }
 
