Use NewObject + TSubclassOf instead of DuplicateObject for better performance. This means that people will need to make a blueprint of the device property to use it nicely.Add a way to remove multiple device property handles with one function

#rnx

diff --git a/Engine/Plugins/Runtime/GameplayAbilities/Source/GameplayAbilities/Private/GameplayCueNotifyTypes.cpp b/Engine/Plugins/Runtime/GameplayAbilities/Source/GameplayAbilities/Private/GameplayCueNotifyTypes.cpp
--- a/Engine/Plugins/Runtime/GameplayAbilities/Source/GameplayAbilities/Private/GameplayCueNotifyTypes.cpp
+++ b/Engine/Plugins/Runtime/GameplayAbilities/Source/GameplayAbilities/Private/GameplayCueNotifyTypes.cpp
@@ -865,9 +865,9 @@ bool FGameplayCueNotify_InputDevicePropertyInfo::SetDeviceProperties(const FGame
 			Params.bRemoveAfterEvaluationTime = true;
 			if (ensure(Params.UserId.IsValid()))
 			{
-				for (TObjectPtr<UInputDeviceProperty> Prop : DeviceProperties)
+				for (TSubclassOf<UInputDeviceProperty> PropClass : DeviceProperties)
 				{
-					Params.DeviceProperty = Prop;
+					Params.DevicePropertyClass = PropClass;
 					FInputDevicePropertyHandle Handle = System->SetDeviceProperty(Params);
 					ensure(Handle.IsValid());
 				}
@@ -882,12 +882,12 @@ bool FGameplayCueNotify_InputDevicePropertyInfo::SetDeviceProperties(const FGame
 void FGameplayCueNotify_InputDevicePropertyInfo::ValidateBurstAssets(UObject* ContainingAsset, const FString& Context, TArray<FText>& ValidationErrors) const
 {
 #if WITH_EDITORONLY_DATA
-	for (const TObjectPtr<UInputDeviceProperty> Prop : DeviceProperties)
+	for (const TSubclassOf<UInputDeviceProperty> PropClass : DeviceProperties)
 	{
-		if (!Prop)
+		if (!PropClass)
 		{
 			ValidationErrors.Add(FText::Format(
-				LOCTEXT("InputDeviceProperty_ShouldNotBeNull", "There is a null device property used in slot [{0}] for asset [{1}]."),
+				LOCTEXT("InputDeviceProperty_ShouldNotBeNull", "There is a null device property class used in slot [{0}] for asset [{1}]."),
 				FText::AsCultureInvariant(Context),
 				FText::AsCultureInvariant(ContainingAsset->GetPathName())));
 		}
diff --git a/Engine/Plugins/Runtime/GameplayAbilities/Source/GameplayAbilities/Public/GameplayCueNotifyTypes.h b/Engine/Plugins/Runtime/GameplayAbilities/Source/GameplayAbilities/Public/GameplayCueNotifyTypes.h
--- a/Engine/Plugins/Runtime/GameplayAbilities/Source/GameplayAbilities/Public/GameplayCueNotifyTypes.h
+++ b/Engine/Plugins/Runtime/GameplayAbilities/Source/GameplayAbilities/Public/GameplayCueNotifyTypes.h
@@ -629,8 +629,8 @@ struct FGameplayCueNotify_InputDevicePropertyInfo
 	void ValidateBurstAssets(UObject* ContainingAsset, const FString& Context, TArray<FText>& ValidationErrors) const;
 	
 	/** Input Device properties to apply */
-	UPROPERTY(EditAnywhere, Instanced, BlueprintReadOnly, Category = GameplayCueNotify)
-	TArray<TObjectPtr<UInputDeviceProperty>> DeviceProperties;
+	UPROPERTY(EditAnywhere, BlueprintReadOnly, Category = GameplayCueNotify)
+	TArray<TSubclassOf<UInputDeviceProperty>> DeviceProperties;
 };
 
 /**
diff --git a/Engine/Source/Runtime/Engine/Classes/GameFramework/InputDeviceSubsystem.h b/Engine/Source/Runtime/Engine/Classes/GameFramework/InputDeviceSubsystem.h
--- a/Engine/Source/Runtime/Engine/Classes/GameFramework/InputDeviceSubsystem.h
+++ b/Engine/Source/Runtime/Engine/Classes/GameFramework/InputDeviceSubsystem.h
@@ -84,11 +84,11 @@ struct ENGINE_API FSetDevicePropertyParams
 
 	FSetDevicePropertyParams();
 
-	FSetDevicePropertyParams(TObjectPtr<UInputDeviceProperty> InProperty, const FPlatformUserId InUserId, const bool bInRemoveAfterEvaluationTime = true);
+	FSetDevicePropertyParams(TSubclassOf<UInputDeviceProperty> InPropertyClass, const FPlatformUserId InUserId, const bool bInRemoveAfterEvaluationTime = true);
 	
 	/** The device property to set */
 	UPROPERTY(EditInstanceOnly, BlueprintReadWrite, Category = "Input Devices")
-	TObjectPtr<UInputDeviceProperty> DeviceProperty;
+	TSubclassOf<UInputDeviceProperty> DevicePropertyClass;
 	
 	/** The Platform User whose device's should receive the device property */
 	UPROPERTY(EditAnywhere, BlueprintReadWrite, Category = "Input Devices")
@@ -162,11 +162,23 @@ public:
 	* Remove a single device property based on it's handle
 	*
 	* @param FInputDevicePropertyHandle		Device property handle to be removed
+	* @param bResetOnRemoval				If true, call ResetDeviceProperty before it is removed. Default is true.
 	*
 	* @return								The number of removed device properties.
 	*/
 	UFUNCTION(BlueprintCallable, Category = "Input Devices", meta = (ReturnDisplayName = "Num Removed"))
-	int32 RemoveDevicePropertyByHandle(const FInputDevicePropertyHandle HandleToRemove);
+	int32 RemoveDevicePropertyByHandle(const FInputDevicePropertyHandle HandleToRemove, const bool bResetOnRemoval = true);
+
+	/**
+	* Remove a set of device properties based on their handles. 
+	* 
+	* @param HandlesToRemove	The set of device property handles to remove
+	* @param bResetOnRemoval	If true, call ResetDeviceProperty before it is removed. Default is true.
+	* 
+	* @return					The number of removed device properties
+	*/
+	UFUNCTION(BlueprintCallable, Category = "Input Devices", meta = (ReturnDisplayName = "Num Removed"))
+	int32 RemoveDevicePropertyHandles(const TSet<FInputDevicePropertyHandle>& HandlesToRemove, const bool bResetOnRemoval = true);
 
 	/** Removes all the current Input Device Properties that are active, regardless of the Platform User */
 	UFUNCTION(BlueprintCallable, Category = "Input Devices")
diff --git a/Engine/Source/Runtime/Engine/Private/GameFramework/InputDeviceSubsystem.cpp b/Engine/Source/Runtime/Engine/Private/GameFramework/InputDeviceSubsystem.cpp
--- a/Engine/Source/Runtime/Engine/Private/GameFramework/InputDeviceSubsystem.cpp
+++ b/Engine/Source/Runtime/Engine/Private/GameFramework/InputDeviceSubsystem.cpp
@@ -137,8 +137,8 @@ FSetDevicePropertyParams::FSetDevicePropertyParams()
 
 }
 
-FSetDevicePropertyParams::FSetDevicePropertyParams(TObjectPtr<UInputDeviceProperty> InProperty, const FPlatformUserId InUserId, const bool bInRemoveAfterEvaluationTime /* = true */)
-	: DeviceProperty(InProperty)
+FSetDevicePropertyParams::FSetDevicePropertyParams(TSubclassOf<UInputDeviceProperty> InPropertyClass, const FPlatformUserId InUserId, const bool bInRemoveAfterEvaluationTime /* = true */)
+	: DevicePropertyClass(InPropertyClass)
 	, UserId(InUserId)
 	, bRemoveAfterEvaluationTime(bInRemoveAfterEvaluationTime)
 {
@@ -294,7 +294,7 @@ void UInputDeviceSubsystem::Tick(float DeltaTime)
 
 FInputDevicePropertyHandle UInputDeviceSubsystem::SetDeviceProperty(const FSetDevicePropertyParams& Params)
 {
-	if (!Params.DeviceProperty)
+	if (!Params.DevicePropertyClass)
 	{
 		UE_LOG(LogInputDeviceProperties, Error, TEXT("Invalid DevicePropertyClass passed into SetDeviceProperty! Nothing will happen."));
 		return FInputDevicePropertyHandle::InvalidHandle;
@@ -306,9 +306,8 @@ FInputDevicePropertyHandle UInputDeviceSubsystem::SetDeviceProperty(const FSetDe
 	{
 		FActiveDeviceProperty ActiveProp = {};
 
-		// Spawn an instance of this device property
-		// TODO: Possible performance problems with DuplicateObject because FDuplicateDataWriter is not very performant
-		ActiveProp.Property = DuplicateObject<UInputDeviceProperty>(Params.DeviceProperty, /* Outer = */ this);
+		// Spawn an instance of this device property		
+		ActiveProp.Property = NewObject<UInputDeviceProperty>(/* Outer = */ this, /* Class */ Params.DevicePropertyClass);
 		ensure(ActiveProp.Property);
 
 		ActiveProp.PlatformUser = Params.UserId;
@@ -358,15 +357,23 @@ int32 UInputDeviceSubsystem::RemoveDevicePropertiesOfClass(const FPlatformUserId
 	return NumRemoved;
 }
 
-int32 UInputDeviceSubsystem::RemoveDevicePropertyByHandle(const FInputDevicePropertyHandle HandleToRemove)
+int32 UInputDeviceSubsystem::RemoveDevicePropertyByHandle(const FInputDevicePropertyHandle HandleToRemove, const bool bResetOnRemoval /*= true*/)
+{
+	return RemoveDevicePropertyHandles({ HandleToRemove }, bResetOnRemoval);
+}
+
+int32 UInputDeviceSubsystem::RemoveDevicePropertyHandles(const TSet<FInputDevicePropertyHandle>& HandlesToRemove, const bool bResetOnRemoval /*= true*/)
 {
 	int32 NumRemoved = 0;
 
 	for (int32 Index = ActiveProperties.Num() - 1; Index >= 0; --Index)
 	{
-		if (ActiveProperties[Index].PropertyHandle == HandleToRemove)
+		if (HandlesToRemove.Contains(ActiveProperties[Index].PropertyHandle))
 		{
-			ActiveProperties[Index].Property->ResetDeviceProperty(ActiveProperties[Index].PlatformUser);
+			if (bResetOnRemoval)
+			{
+				ActiveProperties[Index].Property->ResetDeviceProperty(ActiveProperties[Index].PlatformUser);
+			}			
 			ActiveProperties.RemoveAtSwap(Index);
 			++NumRemoved;
 			break;
@@ -375,7 +382,7 @@ int32 UInputDeviceSubsystem::RemoveDevicePropertyByHandle(const FInputDeviceProp
 
 	if (NumRemoved <= 0)
 	{
-		UE_LOG(LogInputDeviceProperties, Warning, TEXT("Unable to remove a device property with handle '%s'"), *HandleToRemove.ToString());
+		UE_LOG(LogInputDeviceProperties, Warning, TEXT("Unable to remove any device property handles!"));
 	}
 
 	return NumRemoved;
