Remove check for transient flag which is triggering.
Will still investigate the root cause.

diff --git a/Engine/Source/Runtime/Engine/Private/Materials/MaterialInstance.cpp b/Engine/Source/Runtime/Engine/Private/Materials/MaterialInstance.cpp
--- a/Engine/Source/Runtime/Engine/Private/Materials/MaterialInstance.cpp
+++ b/Engine/Source/Runtime/Engine/Private/Materials/MaterialInstance.cpp
@@ -1623,10 +1623,6 @@ void UMaterialInstanceDynamic::CopyScalarAndVectorParameters(const UMaterialInte
 
 void UMaterialInstanceDynamic::SetNaniteOverride(UMaterialInterface* InMaterial)
 {
-	// Only want to do this for transient objects.
-	// It's unsafe to poke values directly into NaniteOverrideMaterial if it could be serialized.
-	check(HasAnyFlags(RF_Transient));
-
 	NaniteOverrideMaterial.InitUnsafe(InMaterial);
 }
 
