Don't try and create RVT render resources for CDO or in server/cooker.

robomerge[FN_EngineMerge] -UE5-RES

diff --git a/Engine/Source/Runtime/Engine/Private/VT/RuntimeVirtualTexture.cpp b/Engine/Source/Runtime/Engine/Private/VT/RuntimeVirtualTexture.cpp
--- a/Engine/Source/Runtime/Engine/Private/VT/RuntimeVirtualTexture.cpp
+++ b/Engine/Source/Runtime/Engine/Private/VT/RuntimeVirtualTexture.cpp
@@ -563,11 +563,14 @@ void URuntimeVirtualTexture::InitResource(IVirtualTexture* InProducer, FVTProduc
 
 void URuntimeVirtualTexture::InitNullResource()
 {
-	FNullVirtualTextureProducer* Producer = new FNullVirtualTextureProducer;
-	FVTProducerDescription ProducerDesc;
-	FNullVirtualTextureProducer::GetNullProducerDescription(ProducerDesc);
-	FRuntimeVirtualTextureRenderResource::FResourceInitDesc InitDesc;
-	Resource->Init(Producer, ProducerDesc, InitDesc);
+	if (FApp::CanEverRender() && !HasAnyFlags(RF_ClassDefaultObject))
+	{
+		FNullVirtualTextureProducer* Producer = new FNullVirtualTextureProducer;
+		FVTProducerDescription ProducerDesc;
+		FNullVirtualTextureProducer::GetNullProducerDescription(ProducerDesc);
+		FRuntimeVirtualTextureRenderResource::FResourceInitDesc InitDesc;
+		Resource->Init(Producer, ProducerDesc, InitDesc);
+	}
 }
 
 void URuntimeVirtualTexture::GetAssetRegistryTags(TArray<FAssetRegistryTag>& OutTags) const
