World Partition: fix HLODs always loaded when the editor has streaming disabled.

#rnx

diff --git a/Engine/Source/Runtime/Engine/Private/WorldPartition/LoaderAdapter/LoaderAdapterActorList.cpp b/Engine/Source/Runtime/Engine/Private/WorldPartition/LoaderAdapter/LoaderAdapterActorList.cpp
--- a/Engine/Source/Runtime/Engine/Private/WorldPartition/LoaderAdapter/LoaderAdapterActorList.cpp
+++ b/Engine/Source/Runtime/Engine/Private/WorldPartition/LoaderAdapter/LoaderAdapterActorList.cpp
@@ -83,8 +83,8 @@ bool FLoaderAdapterActorList::ContainsActor(const FWorldPartitionHandle& ActorHa
 	return Actors.Contains(ActorHandle);
 }
 
-bool FLoaderAdapterActorList::ShouldActorBeLoaded(const FWorldPartitionHandle& ActorHandle) const
+bool FLoaderAdapterActorList::PassActorDescFilter(const FWorldPartitionHandle& ActorHandle) const
 {
-	return ActorHandle.IsValid() && !ActorsToRemove.Contains(ActorHandle) && PassDataLayersFilter(ActorHandle);
+	return ActorHandle.IsValid() && !ActorsToRemove.Contains(ActorHandle) && FLoaderAdapterList::PassActorDescFilter(ActorHandle);
 }
 #endif
diff --git a/Engine/Source/Runtime/Engine/Private/WorldPartition/WorldPartition.cpp b/Engine/Source/Runtime/Engine/Private/WorldPartition/WorldPartition.cpp
--- a/Engine/Source/Runtime/Engine/Private/WorldPartition/WorldPartition.cpp
+++ b/Engine/Source/Runtime/Engine/Private/WorldPartition/WorldPartition.cpp
@@ -273,6 +273,22 @@ public:
 		bIncludeNonSpatiallyLoadedActors = true;
 	}
 };
+
+class FLoaderAdapterPinnedActors: public FLoaderAdapterActorList
+{
+public:
+	FLoaderAdapterPinnedActors(UWorld* InWorld)
+		: FLoaderAdapterActorList(InWorld)
+	{}
+
+protected:
+	virtual bool PassActorDescFilter(const FWorldPartitionHandle& ActorHandle) const override
+	{
+		// We want to be able to pin any type of actors (HLODs, etc).
+		return ActorHandle.IsValid() && !ActorsToRemove.Contains(ActorHandle);
+	}
+};
+
 #endif
 
 UWorldPartition::UWorldPartition(const FObjectInitializer& ObjectInitializer)
@@ -495,7 +511,7 @@ void UWorldPartition::Initialize(UWorld* InWorld, const FTransform& InTransform)
 		EditorHash->Initialize();
 
 		AlwaysLoadedActors = new FLoaderAdapterAlwaysLoadedActors(OuterWorld);
-		PinnedActors = new FLoaderAdapterActorList(OuterWorld);
+		PinnedActors = new FLoaderAdapterPinnedActors(OuterWorld);
 		
 		IWorldPartitionEditorModule& WorldPartitionEditorModule = FModuleManager::LoadModuleChecked<IWorldPartitionEditorModule>("WorldPartitionEditor");
 		ForceLoadedActors = WorldPartitionEditorModule.GetDisableLoadingInEditor() ? new FLoaderAdapterActorList(OuterWorld) : nullptr;
diff --git a/Engine/Source/Runtime/Engine/Public/WorldPartition/LoaderAdapter/LoaderAdapterActorList.h b/Engine/Source/Runtime/Engine/Public/WorldPartition/LoaderAdapter/LoaderAdapterActorList.h
--- a/Engine/Source/Runtime/Engine/Public/WorldPartition/LoaderAdapter/LoaderAdapterActorList.h
+++ b/Engine/Source/Runtime/Engine/Public/WorldPartition/LoaderAdapter/LoaderAdapterActorList.h
@@ -22,10 +22,9 @@ public:
 
 protected:
 	//~ Begin IWorldPartitionActorLoaderInterface::ILoaderAdapterList interface
-	virtual bool ShouldActorBeLoaded(const FWorldPartitionHandle& ActorHandle) const override;
+	virtual bool PassActorDescFilter(const FWorldPartitionHandle& ActorHandle) const override;
 	//~ End IWorldPartitionActorLoaderInterface::ILoaderAdapterList interface
 
-private:
 	TSet<FWorldPartitionHandle> ActorsToRemove;
 };
 #endif
\ No newline at end of file
diff --git a/Engine/Source/Runtime/Engine/Public/WorldPartition/WorldPartitionActorLoaderInterface.h b/Engine/Source/Runtime/Engine/Public/WorldPartition/WorldPartitionActorLoaderInterface.h
--- a/Engine/Source/Runtime/Engine/Public/WorldPartition/WorldPartitionActorLoaderInterface.h
+++ b/Engine/Source/Runtime/Engine/Public/WorldPartition/WorldPartitionActorLoaderInterface.h
@@ -47,14 +47,15 @@ public:
 	protected:
 		// Private interface
 		virtual void ForEachActor(TFunctionRef<void(const FWorldPartitionHandle&)> InOperation) const =0;
-		virtual bool ShouldActorBeLoaded(const FWorldPartitionHandle& Actor) const;
+
+		bool ShouldActorBeLoaded(const FWorldPartitionHandle& Actor) const;
 
 		void RegisterDelegates();
 		void UnregisterDelegates();
 
 		// Actors filtering
-		bool PassActorDescFilter(const FWorldPartitionHandle& Actor) const;
-		bool PassDataLayersFilter(const FWorldPartitionHandle& Actor) const;
+		virtual bool PassActorDescFilter(const FWorldPartitionHandle& Actor) const;
+		virtual bool PassDataLayersFilter(const FWorldPartitionHandle& Actor) const;
 		void RefreshLoadedState();
 
 		void PostLoadedStateChanged(int32 NumLoads, int32 NumUnloads);
