Recurse into MaterialFunctions

NOTE: this change pulls in ALL output-connected nodes inside of each MaterialFunction used.

It also assumes ALL inputs to a MaterialFunctionCall are active (no change in behavior there, but could be improved with a more fine-scale analysis).

diff --git a/Engine/Source/Runtime/Engine/Classes/Materials/Material.h b/Engine/Source/Runtime/Engine/Classes/Materials/Material.h
--- a/Engine/Source/Runtime/Engine/Classes/Materials/Material.h
+++ b/Engine/Source/Runtime/Engine/Classes/Materials/Material.h
@@ -1669,33 +1669,35 @@ public:
 	/**
 	 *	Get all referenced expressions (returns the chains for all properties).
 	 *
-	 *	@param	OutExpressions			The array to fill in all of the expressions.
-	 *	@param	InStaticParameterSet	Optional static parameter set - if supplied only walk the StaticSwitch branches according to it.
-	 *	@Param	InFeatureLevel			Optional feature level - if supplied, only walk FeatureLevelSwitch branches according to it.
-	 *	@Param	InQuality				Optional quality switch - if supplied, only walk QualitySwitch branches according to it.
-	 *	@Param	InShadingPath			Optional shading path switch - if supplied, only walk ShadingPathSwitch branches according to it.
+	 *	@param	OutExpressions						The array to fill in all of the expressions.
+	 *	@param	InStaticParameterSet				Optional static parameter set - if supplied only walk the StaticSwitch branches according to it.
+	 *	@Param	InFeatureLevel						Optional feature level - if supplied, only walk FeatureLevelSwitch branches according to it.
+	 *	@Param	InQuality							Optional quality switch - if supplied, only walk QualitySwitch branches according to it.
+	 *	@Param	InShadingPath						Optional shading path switch - if supplied, only walk ShadingPathSwitch branches according to it.
+	 *  @Param	bInRecurseIntoMaterialFunctions		Optional enable recursion into Material Functions - if true, will also return expressions from inside of MaterialFunctions.
 	 *
 	 *	@return	bool					true if successful, false if not.
 	 */
 	ENGINE_API virtual bool GetAllReferencedExpressions(TArray<UMaterialExpression*>& OutExpressions, struct FStaticParameterSet* InStaticParameterSet,
-		ERHIFeatureLevel::Type InFeatureLevel = ERHIFeatureLevel::Num, EMaterialQualityLevel::Type InQuality = EMaterialQualityLevel::Num, ERHIShadingPath::Type InShadingPath = ERHIShadingPath::Num);
+		ERHIFeatureLevel::Type InFeatureLevel = ERHIFeatureLevel::Num, EMaterialQualityLevel::Type InQuality = EMaterialQualityLevel::Num, ERHIShadingPath::Type InShadingPath = ERHIShadingPath::Num, const bool bInRecurseIntoMaterialFunctions = false);
 
 
 	/**
 	 *	Get the expression chain for the given property (ie fill in the given array with all expressions in the chain).
 	 *
-	 *	@param	InProperty				The material property chain to inspect, such as MP_BaseColor.
-	 *	@param	OutExpressions			The array to fill in all of the expressions.
-	 *	@param	InStaticParameterSet	Optional static parameter set - if supplied only walk the StaticSwitch branches according to it.
-	 *	@Param	InFeatureLevel			Optional feature level - if supplied, only walk FeatureLevelSwitch branches according to it.
-	 *	@Param	InQuality				Optional quality switch - if supplied, only walk QualitySwitch branches according to it.
-	 *	@Param	InShadingPath			Optional shading path switch - if supplied, only walk ShadingPathSwitch branches according to it.
+	 *	@param	InProperty							The material property chain to inspect, such as MP_BaseColor.
+	 *	@param	OutExpressions						The array to fill in all of the expressions.
+	 *	@param	InStaticParameterSet				Optional static parameter set - if supplied only walk the StaticSwitch branches according to it.
+	 *	@Param	InFeatureLevel						Optional feature level - if supplied, only walk FeatureLevelSwitch branches according to it.
+	 *	@Param	InQuality							Optional quality switch - if supplied, only walk QualitySwitch branches according to it.
+	 *	@Param	InShadingPath						Optional shading path switch - if supplied, only walk ShadingPathSwitch branches according to it.
+	 *  @Param	bInRecurseIntoMaterialFunctions		Optional enable recursion into Material Functions - if true, will also return expressions from inside of MaterialFunctions.
 	 *
 	 *	@return	bool					true if successful, false if not.
 	 */
 	ENGINE_API virtual bool GetExpressionsInPropertyChain(EMaterialProperty InProperty, 
 		TArray<UMaterialExpression*>& OutExpressions, struct FStaticParameterSet* InStaticParameterSet,
-		ERHIFeatureLevel::Type InFeatureLevel = ERHIFeatureLevel::Num, EMaterialQualityLevel::Type InQuality = EMaterialQualityLevel::Num, ERHIShadingPath::Type InShadingPath = ERHIShadingPath::Num);
+		ERHIFeatureLevel::Type InFeatureLevel = ERHIFeatureLevel::Num, EMaterialQualityLevel::Type InQuality = EMaterialQualityLevel::Num, ERHIShadingPath::Type InShadingPath = ERHIShadingPath::Num, const bool bInRecurseIntoMaterialFunctions = false);
 
 #endif // WITH_EDITOR
 
@@ -1705,14 +1707,15 @@ protected:
 	/**
 	 *	Recursively retrieve the expressions contained in the chain of the given expression.
 	 *
-	 *	@param	InExpression			The expression to start at.
-	 *	@param	InOutProcessedInputs	An array of processed expression inputs. (To avoid circular loops causing infinite recursion)
-	 *	@param	OutExpressions			The array to fill in all of the expressions.
-	 *	@param	InStaticParameterSet	Optional static parameter set - if supplied only walk the StaticSwitch branches according to it.
-	 *	@Param	InFeatureLevel			Optional feature level - if supplied, only walk FeatureLevelSwitch branches according to it.
-	 *	@Param	InQuality				Optional quality switch - if supplied, only walk QualitySwitch branches according to it.
-	 *	@Param	InShadingPath			Optional shading path switch - if supplied, only walk ShadingPathSwitch branches according to it.
-	 *	@Param	InShaderFrequency		Optional shader frequency - if supplied, only walk ShaderFrequencySwitch branches according to it.
+	 *	@param	InExpression						The expression to start at.
+	 *	@param	InOutProcessedInputs				An array of processed expression inputs. (To avoid circular loops causing infinite recursion)
+	 *	@param	OutExpressions						The array to fill in all of the expressions.
+	 *	@param	InStaticParameterSet				Optional static parameter set - if supplied only walk the StaticSwitch branches according to it.
+	 *	@Param	InFeatureLevel						Optional feature level - if supplied, only walk FeatureLevelSwitch branches according to it.
+	 *	@Param	InQuality							Optional quality switch - if supplied, only walk QualitySwitch branches according to it.
+	 *	@Param	InShadingPath						Optional shading path switch - if supplied, only walk ShadingPathSwitch branches according to it.
+	 *	@Param	InShaderFrequency					Optional shader frequency - if supplied, only walk ShaderFrequencySwitch branches according to it.
+	 *  @Param	bInRecurseIntoMaterialFunctions		Optional enable recursion into Material Functions - if true, will also return expressions from inside of MaterialFunctions.
 	 *
 	 *	@return	bool					true if successful, false if not.
 	 */
@@ -1722,7 +1725,8 @@ protected:
 		EMaterialQualityLevel::Type InQuality = EMaterialQualityLevel::Num,
 		ERHIShadingPath::Type InShadingPath = ERHIShadingPath::Num,
 		EShaderFrequency InShaderFrequency = SF_NumFrequencies, 
-		EMaterialProperty InProperty = MP_MAX);
+		EMaterialProperty InProperty = MP_MAX,
+		const bool bInRecurseIntoMaterialFunctions = false);
 
 	/**
 	*	Recursively update the bRealtimePreview for each expression based on whether it is connected to something that is time-varying.
diff --git a/Engine/Source/Runtime/Engine/Private/Materials/Material.cpp b/Engine/Source/Runtime/Engine/Private/Materials/Material.cpp
--- a/Engine/Source/Runtime/Engine/Private/Materials/Material.cpp
+++ b/Engine/Source/Runtime/Engine/Private/Materials/Material.cpp
@@ -5566,7 +5566,7 @@ void UMaterial::GetAllExpressionsForCustomInterpolators(TArray<class UMaterialEx
 
 #if WITH_EDITOR
 bool UMaterial::GetAllReferencedExpressions(TArray<UMaterialExpression*>& OutExpressions, struct FStaticParameterSet* InStaticParameterSet,
-	ERHIFeatureLevel::Type InFeatureLevel, EMaterialQualityLevel::Type InQuality, ERHIShadingPath::Type InShadingPath)
+	ERHIFeatureLevel::Type InFeatureLevel, EMaterialQualityLevel::Type InQuality, ERHIShadingPath::Type InShadingPath, const bool bInRecurseIntoMaterialFunctions)
 {
 	using namespace MaterialImpl;
 
@@ -5578,7 +5578,7 @@ bool UMaterial::GetAllReferencedExpressions(TArray<UMaterialExpression*>& OutExp
 		if (bUseMaterialAttributes)
 		{
 			TArray<UMaterialExpression*> MPRefdExpressions;
-			if (GetExpressionsInPropertyChain(MP_MaterialAttributes, MPRefdExpressions, InStaticParameterSet, InFeatureLevel, InQuality, InShadingPath) == true)
+			if (GetExpressionsInPropertyChain(MP_MaterialAttributes, MPRefdExpressions, InStaticParameterSet, InFeatureLevel, InQuality, InShadingPath, bInRecurseIntoMaterialFunctions) == true)
 			{
 				for (int32 AddIdx = 0; AddIdx < MPRefdExpressions.Num(); AddIdx++)
 				{
@@ -5592,7 +5592,7 @@ bool UMaterial::GetAllReferencedExpressions(TArray<UMaterialExpression*>& OutExp
 			{
 				EMaterialProperty MaterialProp = MobileProperty.Property;
 				TArray<UMaterialExpression*> MPRefdExpressions;
-				if (GetExpressionsInPropertyChain(MaterialProp, MPRefdExpressions, InStaticParameterSet, InFeatureLevel, InQuality, InShadingPath) == true)
+				if (GetExpressionsInPropertyChain(MaterialProp, MPRefdExpressions, InStaticParameterSet, InFeatureLevel, InQuality, InShadingPath, bInRecurseIntoMaterialFunctions) == true)
 				{
 					for (int32 AddIdx = 0; AddIdx < MPRefdExpressions.Num(); AddIdx++)
 					{
@@ -5612,7 +5612,7 @@ bool UMaterial::GetAllReferencedExpressions(TArray<UMaterialExpression*>& OutExp
 				if (Expression->IsA<UMaterialExpressionRuntimeVirtualTextureOutput>())
 				{
 					TArray<FExpressionInput*> ProcessedInputs;
-					RecursiveGetExpressionChain(Expression, ProcessedInputs, OutExpressions, InStaticParameterSet, InFeatureLevel, InQuality, InShadingPath);
+					RecursiveGetExpressionChain(Expression, ProcessedInputs, OutExpressions, InStaticParameterSet, InFeatureLevel, InQuality, InShadingPath, SF_NumFrequencies, MP_MAX, bInRecurseIntoMaterialFunctions);
 				}
 			}
 		}
@@ -5623,7 +5623,7 @@ bool UMaterial::GetAllReferencedExpressions(TArray<UMaterialExpression*>& OutExp
 	    {
 		    EMaterialProperty MaterialProp = EMaterialProperty(MPIdx);
 		    TArray<UMaterialExpression*> MPRefdExpressions;
-			if (GetExpressionsInPropertyChain(MaterialProp, MPRefdExpressions, InStaticParameterSet, InFeatureLevel, InQuality, InShadingPath) == true)
+			if (GetExpressionsInPropertyChain(MaterialProp, MPRefdExpressions, InStaticParameterSet, InFeatureLevel, InQuality, InShadingPath, bInRecurseIntoMaterialFunctions) == true)
 			{
 			    for (int32 AddIdx = 0; AddIdx < MPRefdExpressions.Num(); AddIdx++)
 			    {
@@ -5637,7 +5637,7 @@ bool UMaterial::GetAllReferencedExpressions(TArray<UMaterialExpression*>& OutExp
 	    for (UMaterialExpressionCustomOutput* Expression : CustomOutputExpressions)
 	    {
 		    TArray<FExpressionInput*> ProcessedInputs;
-			RecursiveGetExpressionChain(Expression, ProcessedInputs, OutExpressions, InStaticParameterSet, InFeatureLevel, InQuality, InShadingPath);
+			RecursiveGetExpressionChain(Expression, ProcessedInputs, OutExpressions, InStaticParameterSet, InFeatureLevel, InQuality, InShadingPath, SF_NumFrequencies, MP_MAX, bInRecurseIntoMaterialFunctions);
 		}
 
 		// If this is a material function, we want to also trace function outputs
@@ -5646,7 +5646,7 @@ bool UMaterial::GetAllReferencedExpressions(TArray<UMaterialExpression*>& OutExp
 		for (UMaterialExpressionFunctionOutput* Expression : FunctionOutputExpressions)
 		{
 			TArray<FExpressionInput*> ProcessedInputs;
-			RecursiveGetExpressionChain(Expression, ProcessedInputs, OutExpressions, InStaticParameterSet, InFeatureLevel, InQuality, InShadingPath);
+			RecursiveGetExpressionChain(Expression, ProcessedInputs, OutExpressions, InStaticParameterSet, InFeatureLevel, InQuality, InShadingPath, SF_NumFrequencies, MP_MAX, bInRecurseIntoMaterialFunctions);
 		}
 	}
 
@@ -5656,7 +5656,7 @@ bool UMaterial::GetAllReferencedExpressions(TArray<UMaterialExpression*>& OutExp
 
 bool UMaterial::GetExpressionsInPropertyChain(EMaterialProperty InProperty, 
 	TArray<UMaterialExpression*>& OutExpressions, struct FStaticParameterSet* InStaticParameterSet,
-	ERHIFeatureLevel::Type InFeatureLevel, EMaterialQualityLevel::Type InQuality, ERHIShadingPath::Type InShadingPath)
+	ERHIFeatureLevel::Type InFeatureLevel, EMaterialQualityLevel::Type InQuality, ERHIShadingPath::Type InShadingPath, const bool bInRecurseIntoMaterialFunctions)
 {
 	OutExpressions.Empty();
 	FExpressionInput* StartingExpression = GetExpressionInputForProperty(InProperty);
@@ -5679,7 +5679,7 @@ bool UMaterial::GetExpressionsInPropertyChain(EMaterialProperty InProperty,
 			ShaderFrequency = FMaterialAttributeDefinitionMap::GetShaderFrequency(InProperty);
 		}
 
-		RecursiveGetExpressionChain(StartingExpression->Expression, ProcessedInputs, OutExpressions, InStaticParameterSet, InFeatureLevel, InQuality, InShadingPath, ShaderFrequency, InProperty);
+		RecursiveGetExpressionChain(StartingExpression->Expression, ProcessedInputs, OutExpressions, InStaticParameterSet, InFeatureLevel, InQuality, InShadingPath, ShaderFrequency, InProperty, bInRecurseIntoMaterialFunctions);
 	}
 	return true;
 }
@@ -5747,7 +5747,7 @@ bool UMaterial::GetTexturesInPropertyChain(EMaterialProperty InProperty, TArray<
 bool UMaterial::RecursiveGetExpressionChain(
 	UMaterialExpression* InExpression, TArray<FExpressionInput*>& InOutProcessedInputs, 
 	TArray<UMaterialExpression*>& OutExpressions, struct FStaticParameterSet* InStaticParameterSet,
-	ERHIFeatureLevel::Type InFeatureLevel, EMaterialQualityLevel::Type InQuality, ERHIShadingPath::Type InShadingPath, EShaderFrequency InShaderFrequency, EMaterialProperty InProperty)
+	ERHIFeatureLevel::Type InFeatureLevel, EMaterialQualityLevel::Type InQuality, ERHIShadingPath::Type InShadingPath, EShaderFrequency InShaderFrequency, EMaterialProperty InProperty, const bool bInRecurseIntoMaterialFunctions)
 {
 	using namespace MaterialImpl;
 
@@ -5762,9 +5762,11 @@ bool UMaterial::RecursiveGetExpressionChain(
 	UMaterialExpressionSetMaterialAttributes* SetMaterialAttributesExp;
 	UMaterialExpressionShaderStageSwitch* ShaderStageSwitchExp;
 	UMaterialExpressionNamedRerouteUsage* RerouteUsageExp;
+	UMaterialExpressionMaterialFunctionCall* MaterialFunctionCallExp;
 
 	const bool bMobileOnly = InFeatureLevel <= ERHIFeatureLevel::ES3_1;
 
+	// special case the various switch nodes to only follow the currently active path
 	if (InFeatureLevel != ERHIFeatureLevel::Num && (FeatureLevelSwitchExp = Cast<UMaterialExpressionFeatureLevelSwitch>(InExpression)) != nullptr)
 	{
 		if (FeatureLevelSwitchExp->Inputs[InFeatureLevel].IsConnected())
@@ -5912,7 +5914,40 @@ bool UMaterial::RecursiveGetExpressionChain(
 		// continue searching from the reroute declaration
 		if (RerouteUsageExp->Declaration != nullptr)
 		{
-			RecursiveGetExpressionChain(RerouteUsageExp->Declaration, InOutProcessedInputs, OutExpressions, InStaticParameterSet, InFeatureLevel, InQuality, InShadingPath, InShaderFrequency, InProperty);
+			RecursiveGetExpressionChain(RerouteUsageExp->Declaration, InOutProcessedInputs, OutExpressions, InStaticParameterSet, InFeatureLevel, InQuality, InShadingPath, InShaderFrequency, InProperty, bInRecurseIntoMaterialFunctions);
+		}
+	}
+	else if (bInRecurseIntoMaterialFunctions && (MaterialFunctionCallExp = Cast<UMaterialExpressionMaterialFunctionCall>(InExpression)) != nullptr)
+	{
+		// NOTE: we do the simple thing here, and assume that ALL of the MaterialFunction outputs are connected,
+		// and that all of the inputs feed into each output.  This simplifies the problem enough that we don't need
+		// to build another structure to track processed inputs within each MaterialFunctionCall separately.
+		// However, it does mean that we may add some expression nodes (from inside MaterialFunction) to OutExpressions,
+		// that would not be considered active by an analysis that traces actual connections into the MaterialFunction
+		UMaterialFunctionInterface* MFI = MaterialFunctionCallExp->MaterialFunction.Get();
+		if (MFI != nullptr)
+		{
+			TArray<FFunctionExpressionInput> FuncInputs;
+			TArray<FFunctionExpressionOutput> FuncOutputs;
+			MFI->GetInputsAndOutputs(FuncInputs, FuncOutputs);
+
+			// here we assume ALL outputs of the MaterialFunctionCall are active
+			for (FFunctionExpressionOutput& Out : FuncOutputs)
+			{
+				if (Out.ExpressionOutput)
+				{
+					Inputs.Add(&Out.ExpressionOutput->A);
+					InputsFrequency.Add(InShaderFrequency);
+				}
+			}
+		}
+
+		// here we assume ALL inputs to the MaterialFunctionCall are active
+		auto ExprInputs = InExpression->GetInputs();
+		for (int i = 0; i < ExprInputs.Num(); i++)
+		{
+			Inputs.Add(ExprInputs[i]);
+			InputsFrequency.Add(InShaderFrequency);
 		}
 	}
 	else
@@ -5978,7 +6013,7 @@ bool UMaterial::RecursiveGetExpressionChain(
 					if (bProcessInput == true)
 					{
 						InOutProcessedInputs.Add(InnerInput);
-						RecursiveGetExpressionChain(InnerInput->Expression, InOutProcessedInputs, OutExpressions, InStaticParameterSet, InFeatureLevel, InQuality, InShadingPath, InputsFrequency[InputIdx], InProperty);
+						RecursiveGetExpressionChain(InnerInput->Expression, InOutProcessedInputs, OutExpressions, InStaticParameterSet, InFeatureLevel, InQuality, InShadingPath, InputsFrequency[InputIdx], InProperty, bInRecurseIntoMaterialFunctions);
 					}
 				}
 			}
diff --git a/Engine/Source/Runtime/Landscape/Private/LandscapeEdit.cpp b/Engine/Source/Runtime/Landscape/Private/LandscapeEdit.cpp
--- a/Engine/Source/Runtime/Landscape/Private/LandscapeEdit.cpp
+++ b/Engine/Source/Runtime/Landscape/Private/LandscapeEdit.cpp
@@ -6850,8 +6850,9 @@ static void GetAllMobileRelevantLayerNames(TSet<FName>& OutLayerNames, UMaterial
 {
 	TArray<FName> LayerNames;
 
+	const bool bRecurseIntoMaterialFunctions = true;
 	TArray<UMaterialExpression*> ES31Expressions;
-	InMaterial->GetAllReferencedExpressions(ES31Expressions, nullptr, ERHIFeatureLevel::ES3_1);
+	InMaterial->GetAllReferencedExpressions(ES31Expressions, nullptr, ERHIFeatureLevel::ES3_1, EMaterialQualityLevel::Num, ERHIShadingPath::Num, bRecurseIntoMaterialFunctions );
 
 	TArray<UMaterialExpression*> MobileExpressions = MoveTemp(ES31Expressions);
 	for (UMaterialExpression* Expression : MobileExpressions)
