#UE Add more info to an assert

diff --git a/Engine/Source/Runtime/Engine/Private/Animation/AnimSequenceBase.cpp b/Engine/Source/Runtime/Engine/Private/Animation/AnimSequenceBase.cpp
--- a/Engine/Source/Runtime/Engine/Private/Animation/AnimSequenceBase.cpp
+++ b/Engine/Source/Runtime/Engine/Private/Animation/AnimSequenceBase.cpp
@@ -126,7 +126,7 @@ void UAnimSequenceBase::PostLoad()
 		    
 		    const bool bRequiresModelPopulation = GetLinkerCustomVersion(FUE5MainStreamObjectVersion::GUID) < FUE5MainStreamObjectVersion::IntroducingAnimationDataModel;
 		    PRAGMA_DISABLE_DEPRECATION_WARNINGS
-		    checkf(bRequiresModelPopulation || DataModel != nullptr || DataModelInterface != nullptr, TEXT("Invalid Animation Sequence base state, no data model found past upgrade object version"));
+		    checkf(bRequiresModelPopulation || DataModel != nullptr || DataModelInterface != nullptr, TEXT("Invalid Animation Sequence base state, no data model found past upgrade object version. AnimSequenceBase:%s"), *GetPathName());
 		    PRAGMA_ENABLE_DEPRECATION_WARNINGS
     
 		    // Construct a new IAnimationDataModel instance
