Allow overriding IO priorities for various streaming asset types when the loading screen is visible

diff --git a/Engine/Source/Runtime/Engine/Private/Streaming/SkeletalMeshUpdate.cpp b/Engine/Source/Runtime/Engine/Private/Streaming/SkeletalMeshUpdate.cpp
--- a/Engine/Source/Runtime/Engine/Private/Streaming/SkeletalMeshUpdate.cpp
+++ b/Engine/Source/Runtime/Engine/Private/Streaming/SkeletalMeshUpdate.cpp
@@ -15,6 +15,13 @@ SkeletalMeshUpdate.cpp: Helpers to stream in and out skeletal mesh LODs.
 #include "Components/SkinnedMeshComponent.h"
 #include "Streaming/RenderAssetUpdate.inl"
 
+int32 GStreamingSkeletalMeshIOPriority = (int32)AIOP_Low;
+static FAutoConsoleVariableRef CVarStreamingSkeletalMeshIOPriority(
+	TEXT("r.Streaming.SkeletalMeshIOPriority"),
+	GStreamingSkeletalMeshIOPriority,
+	TEXT("Base I/O priority for loading skeletal mesh LODs"),
+	ECVF_Default);
+
 extern int32 GStreamingMaxReferenceChecks;
 
 template class TRenderAssetUpdate<FSkelMeshUpdateContext>;
@@ -440,7 +447,8 @@ void FSkeletalMeshStreamIn_IO::SetIORequest(const FContext& Context)
 		// but that won't do anything because the tick would not try to acquire the lock since it is already locked.
 		TaskSynchronization.Increment();
 
-		const EAsyncIOPriorityAndFlags Priority = bHighPrioIORequest ? AIOP_BelowNormal : AIOP_Low;
+		const EAsyncIOPriorityAndFlags Priority = (EAsyncIOPriorityAndFlags)FMath::Clamp<int32>(GStreamingSkeletalMeshIOPriority + (bHighPrioIORequest ? 1 : 0), AIOP_Low, AIOP_High);
+
 		Batch.Issue(BulkData, Priority, [this](FBulkDataRequest::EStatus Status)
 		{
 			// At this point task synchronization would hold the number of pending requests.
diff --git a/Engine/Source/Runtime/Engine/Private/Streaming/StaticMeshUpdate.cpp b/Engine/Source/Runtime/Engine/Private/Streaming/StaticMeshUpdate.cpp
--- a/Engine/Source/Runtime/Engine/Private/Streaming/StaticMeshUpdate.cpp
+++ b/Engine/Source/Runtime/Engine/Private/Streaming/StaticMeshUpdate.cpp
@@ -22,6 +22,13 @@ static FAutoConsoleVariableRef CVarStreamingMaxReferenceChecksBeforeStreamOut(
 	ECVF_Default
 );
 
+int32 GStreamingStaticMeshIOPriority = (int32)AIOP_Low;
+static FAutoConsoleVariableRef CVarStreamingStaticMeshIOPriority(
+	TEXT("r.Streaming.StaticMeshIOPriority"),
+	GStreamingStaticMeshIOPriority,
+	TEXT("Base I/O priority for loading static mesh LODs"),
+	ECVF_Default);
+
 // Instantiate TRenderAssetUpdate for FStaticMeshUpdateContext
 template class TRenderAssetUpdate<FStaticMeshUpdateContext>;
 
@@ -479,8 +486,9 @@ void FStaticMeshStreamIn_IO::SetIORequest(const FContext& Context)
 		// Increment as we push the request. If a request complete immediately, then it will call the callback
 		// but that won't do anything because the tick would not try to acquire the lock since it is already locked.
 		TaskSynchronization.Increment();
-		
-		const EAsyncIOPriorityAndFlags Priority = bHighPrioIORequest ? AIOP_BelowNormal : AIOP_Low;
+
+		const EAsyncIOPriorityAndFlags Priority = (EAsyncIOPriorityAndFlags)FMath::Clamp<int32>(GStreamingStaticMeshIOPriority + (bHighPrioIORequest ? 1 : 0), AIOP_Low, AIOP_High);
+
 		Batch.Issue(BulkData, Priority, [this](FBulkDataRequest::EStatus Status)
 		{
 			TaskSynchronization.Decrement();
diff --git a/Engine/Source/Runtime/Engine/Private/Streaming/Texture2DMipDataProvider_IO.cpp b/Engine/Source/Runtime/Engine/Private/Streaming/Texture2DMipDataProvider_IO.cpp
--- a/Engine/Source/Runtime/Engine/Private/Streaming/Texture2DMipDataProvider_IO.cpp
+++ b/Engine/Source/Runtime/Engine/Private/Streaming/Texture2DMipDataProvider_IO.cpp
@@ -13,6 +13,8 @@ Texture2DMipDataProvider_IO.cpp : Implementation of FTextureMipDataProvider usin
 #include "Streaming/TextureStreamingHelpers.h"
 #include "ContentStreaming.h"
 
+extern int32 GStreamingTextureIOPriority;
+
 FTexture2DMipDataProvider_IO::FTexture2DMipDataProvider_IO(const UTexture* InTexture, bool InPrioritizedIORequest)
 	: FTextureMipDataProvider(InTexture, ETickState::Init, ETickThread::Async)
 	, bPrioritizedIORequest(InPrioritizedIORequest)
@@ -107,7 +109,7 @@ int32 FTexture2DMipDataProvider_IO::GetMips(
 			OwnerMip.BulkData.CreateStreamingRequest(
 				0,
 				OwnerMip.BulkData.GetBulkDataSize(),
-				bPrioritizedIORequest ? (AIOP_FLAG_DONTCACHE | AIOP_BelowNormal) : (AIOP_FLAG_DONTCACHE | AIOP_Low),
+				(EAsyncIOPriorityAndFlags)FMath::Clamp<int32>(GStreamingTextureIOPriority + (bPrioritizedIORequest ? 1 : 0), AIOP_Low, AIOP_High) | AIOP_FLAG_DONTCACHE,
 				&AsyncFileCallBack,
 				(uint8*)MipInfo.DestData
 			)
diff --git a/Engine/Source/Runtime/Engine/Private/Streaming/Texture2DStreamIn_IO.cpp b/Engine/Source/Runtime/Engine/Private/Streaming/Texture2DStreamIn_IO.cpp
--- a/Engine/Source/Runtime/Engine/Private/Streaming/Texture2DStreamIn_IO.cpp
+++ b/Engine/Source/Runtime/Engine/Private/Streaming/Texture2DStreamIn_IO.cpp
@@ -15,6 +15,13 @@ Texture2DStreamIn.cpp: Stream in helper for 2D textures using texture streaming
 #include "Streaming/TextureStreamingHelpers.h"
 #include "ContentStreaming.h"
 
+int32 GStreamingTextureIOPriority = (int32)AIOP_Low;
+static FAutoConsoleVariableRef CVarStreamingTextureIOPriority(
+	TEXT("r.Streaming.TextureIOPriority"),
+	GStreamingTextureIOPriority,
+	TEXT("Base I/O priority for loading texture mips"),
+	ECVF_Default);
+
 FTexture2DStreamIn_IO::FTexture2DStreamIn_IO(UTexture2D* InTexture, bool InPrioritizedIORequest)
 	: FTexture2DStreamIn(InTexture)
 	, bPrioritizedIORequest(InPrioritizedIORequest)
@@ -80,7 +87,7 @@ void FTexture2DStreamIn_IO::SetIORequests(const FContext& Context)
 			IORequests[MipIndex] = MipMap.BulkData.CreateStreamingRequest(
 				0,
 				BulkDataSize,
-				bPrioritizedIORequest ? (AIOP_FLAG_DONTCACHE|AIOP_BelowNormal) : (AIOP_FLAG_DONTCACHE|AIOP_Low),
+				(EAsyncIOPriorityAndFlags)FMath::Clamp<int32>(GStreamingTextureIOPriority + (bPrioritizedIORequest ? 1 : 0), AIOP_Low, AIOP_High) | AIOP_FLAG_DONTCACHE,
 				&AsyncFileCallBack,
 				(uint8*)MipData[MipIndex].Data);
 		}
