Optimized CacheAudioCookOverrides by updating it to use the already loaded ini file instead of loading it again.

Before: 268ms
After: 0.2ms

diff --git a/Engine/Source/Runtime/Engine/Private/AudioCompressionSettingsUtils.cpp b/Engine/Source/Runtime/Engine/Private/AudioCompressionSettingsUtils.cpp
--- a/Engine/Source/Runtime/Engine/Private/AudioCompressionSettingsUtils.cpp
+++ b/Engine/Source/Runtime/Engine/Private/AudioCompressionSettingsUtils.cpp
@@ -99,13 +99,12 @@ void CacheAudioCookOverrides(FPlatformAudioCookOverrides& OutOverrides, const TC
 		OutOverrides = FPlatformAudioCookOverrides();
 		return;
 	}
-	
-	FConfigFile PlatformFile;
-	FConfigCacheIni::LoadLocalIniFile(PlatformFile, TEXT("Engine"), true, *PlatformName);
 
+	FConfigFile LocalPlatformEngineIni;
+	const FConfigFile* PlatformFile = FConfigCacheIni::FindOrLoadPlatformConfig(LocalPlatformEngineIni, TEXT("Engine"), *PlatformName);
 
 	int32 SoundCueQualityIndex = INDEX_NONE;
-	if (PlatformFile.GetInt(*CategoryName, TEXT("SoundCueCookQualityIndex"), SoundCueQualityIndex))
+	if (PlatformFile->GetInt(*CategoryName, TEXT("SoundCueCookQualityIndex"), SoundCueQualityIndex))
 	{
 		OutOverrides.SoundCueCookQualityIndex = SoundCueQualityIndex;
 	}
@@ -118,37 +117,36 @@ void CacheAudioCookOverrides(FPlatformAudioCookOverrides& OutOverrides, const TC
 	int32 RetrievedCacheSize = DefaultCacheSize;
 	int32 RetrievedChunkSizeOverride = INDEX_NONE;
 
-	PlatformFile.GetInt(*CategoryName, TEXT("CacheSizeKB"), RetrievedCacheSize);
+	PlatformFile->GetInt(*CategoryName, TEXT("CacheSizeKB"), RetrievedCacheSize);
 	if (!RetrievedCacheSize)
 	{
-		PlatformFile.SetInt64(*CategoryName, TEXT("CacheSizeKB"), DefaultCacheSize);
 		RetrievedCacheSize = DefaultCacheSize;
 	}
 
 	OutOverrides.StreamCachingSettings.CacheSizeKB = RetrievedCacheSize;
 
-	PlatformFile.GetInt(*CategoryName, TEXT("MaxChunkSizeOverrideKB"), RetrievedChunkSizeOverride);
+	PlatformFile->GetInt(*CategoryName, TEXT("MaxChunkSizeOverrideKB"), RetrievedChunkSizeOverride);
 	OutOverrides.StreamCachingSettings.MaxChunkSizeOverrideKB = RetrievedChunkSizeOverride;
 
 	bool bForceLegacyStreamChunking = false;
-	PlatformFile.GetBool(*CategoryName, TEXT("bForceLegacyStreamChunking"), bForceLegacyStreamChunking);
+	PlatformFile->GetBool(*CategoryName, TEXT("bForceLegacyStreamChunking"), bForceLegacyStreamChunking);
 	OutOverrides.StreamCachingSettings.bForceLegacyStreamChunking = bForceLegacyStreamChunking;
 
 	int32 ZerothChunkSizeForLegacyStreamChunking = 0;
-	PlatformFile.GetInt(*CategoryName, TEXT("ZerothChunkSizeForLegacyStreamChunking"), ZerothChunkSizeForLegacyStreamChunking);
+	PlatformFile->GetInt(*CategoryName, TEXT("ZerothChunkSizeForLegacyStreamChunking"), ZerothChunkSizeForLegacyStreamChunking);
 	OutOverrides.StreamCachingSettings.ZerothChunkSizeForLegacyStreamChunkingKB = ZerothChunkSizeForLegacyStreamChunking;
 
-	PlatformFile.GetBool(*CategoryName, TEXT("bResampleForDevice"), OutOverrides.bResampleForDevice);
+	PlatformFile->GetBool(*CategoryName, TEXT("bResampleForDevice"), OutOverrides.bResampleForDevice);
 
-	PlatformFile.GetFloat(*CategoryName, TEXT("CompressionQualityModifier"), OutOverrides.CompressionQualityModifier);
+	PlatformFile->GetFloat(*CategoryName, TEXT("CompressionQualityModifier"), OutOverrides.CompressionQualityModifier);
 
-	PlatformFile.GetFloat(*CategoryName, TEXT("AutoStreamingThreshold"), OutOverrides.AutoStreamingThreshold);
+	PlatformFile->GetFloat(*CategoryName, TEXT("AutoStreamingThreshold"), OutOverrides.AutoStreamingThreshold);
 
 #if 1
 	//Cache sample rate map:
 	float RetrievedSampleRate = -1.0f;
 
-	PlatformFile.GetFloat(*CategoryName, TEXT("MaxSampleRate"), RetrievedSampleRate);
+	PlatformFile->GetFloat(*CategoryName, TEXT("MaxSampleRate"), RetrievedSampleRate);
 	float* FoundSampleRate = OutOverrides.PlatformSampleRates.Find(ESoundwaveSampleRateSettings::Max);
 
 	if (FoundSampleRate)
@@ -166,7 +164,7 @@ void CacheAudioCookOverrides(FPlatformAudioCookOverrides& OutOverrides, const TC
 
 	RetrievedSampleRate = -1.0f;
 
-	PlatformFile.GetFloat(*CategoryName, TEXT("HighSampleRate"), RetrievedSampleRate);
+	PlatformFile->GetFloat(*CategoryName, TEXT("HighSampleRate"), RetrievedSampleRate);
 	FoundSampleRate = OutOverrides.PlatformSampleRates.Find(ESoundwaveSampleRateSettings::High);
 
 	if (FoundSampleRate)
@@ -185,7 +183,7 @@ void CacheAudioCookOverrides(FPlatformAudioCookOverrides& OutOverrides, const TC
 
 	RetrievedSampleRate = -1.0f;
 
-	PlatformFile.GetFloat(*CategoryName, TEXT("MedSampleRate"), RetrievedSampleRate);
+	PlatformFile->GetFloat(*CategoryName, TEXT("MedSampleRate"), RetrievedSampleRate);
 	FoundSampleRate = OutOverrides.PlatformSampleRates.Find(ESoundwaveSampleRateSettings::Medium);
 
 	if (FoundSampleRate)
@@ -202,7 +200,7 @@ void CacheAudioCookOverrides(FPlatformAudioCookOverrides& OutOverrides, const TC
 
 	RetrievedSampleRate = -1.0f;
 
-	PlatformFile.GetFloat(*CategoryName, TEXT("LowSampleRate"), RetrievedSampleRate);
+	PlatformFile->GetFloat(*CategoryName, TEXT("LowSampleRate"), RetrievedSampleRate);
 	FoundSampleRate = OutOverrides.PlatformSampleRates.Find(ESoundwaveSampleRateSettings::Low);
 
 	if (FoundSampleRate)
@@ -219,7 +217,7 @@ void CacheAudioCookOverrides(FPlatformAudioCookOverrides& OutOverrides, const TC
 
 	RetrievedSampleRate = -1.0f;
 
-	PlatformFile.GetFloat(*CategoryName, TEXT("MinSampleRate"), RetrievedSampleRate);
+	PlatformFile->GetFloat(*CategoryName, TEXT("MinSampleRate"), RetrievedSampleRate);
 	FoundSampleRate = OutOverrides.PlatformSampleRates.Find(ESoundwaveSampleRateSettings::Min);
 
 	if (FoundSampleRate)
