LWC shader optimization - If both sides are LWC and we are subtracting, assume relative space going forward and demote to floating point precision. Also implemented a TruncateLWC node that can allow for targeted optimizations of random expressions in the material graph. r.MaterialEditor.LWCTruncateMode is now defined to control 0: off, 1: node only, 2: auto trunc  (default is 2)

diff --git a/Engine/Config/BaseEngine.ini b/Engine/Config/BaseEngine.ini
--- a/Engine/Config/BaseEngine.ini
+++ b/Engine/Config/BaseEngine.ini
@@ -2804,6 +2804,7 @@ DefaultASTCQualityBySize=3
 +ClassesExcludedOnDedicatedClient=InterpCurveEdSetup
 
 +VersionedIntRValues=r.AllowStaticLighting
++VersionedIntRValues=r.MaterialEditor.LWCTruncateMode
 +VersionedIntRValues=r.GBuffer
 +VersionedIntRValues=r.VelocityOutputPass
 +VersionedIntRValues=r.SelectiveBasePassOutputs
diff --git a/Engine/Source/Runtime/Core/Private/HAL/ConsoleManager.cpp b/Engine/Source/Runtime/Core/Private/HAL/ConsoleManager.cpp
--- a/Engine/Source/Runtime/Core/Private/HAL/ConsoleManager.cpp
+++ b/Engine/Source/Runtime/Core/Private/HAL/ConsoleManager.cpp
@@ -3040,6 +3040,15 @@ static TAutoConsoleVariable<int32> CVarMobileDesiredResY(
 	TEXT("Desired mobile Y resolution (shortest axis) (non-zero == use for Y, calculate X to retain aspect ratio)"),
 	ECVF_Default);
 
+static TAutoConsoleVariable<int32> CVarLWCTruncateMode(
+	TEXT("r.MaterialEditor.LWCTruncateMode"),
+	2,
+	TEXT("Whether or not the material compiler respects the truncate LWC node or automatic transforms.\n"
+		"0: no truncate (LWC always used even if asked to truncate)\n"
+		"1: respect the truncate LWC node\n"
+		"2: respect the truncate LWC node and automatic transforms"),
+	ECVF_ReadOnly);
+
 // this cvar can be removed in shipping to not compile shaders for development (faster)
 static TAutoConsoleVariable<int32> CVarCompileShadersForDevelopment(
 	TEXT("r.CompileShadersForDevelopment"),
diff --git a/Engine/Source/Runtime/Engine/Classes/Materials/MaterialExpressionTruncateLWC.h b/Engine/Source/Runtime/Engine/Classes/Materials/MaterialExpressionTruncateLWC.h
new file mode 100644
--- /dev/null
+++ b/Engine/Source/Runtime/Engine/Classes/Materials/MaterialExpressionTruncateLWC.h
@@ -0,0 +1,31 @@
+// Copyright Epic Games, Inc. All Rights Reserved.
+
+
+#pragma once
+
+#include "CoreMinimal.h"
+#include "UObject/ObjectMacros.h"
+#include "MaterialExpressionIO.h"
+#include "Materials/MaterialExpression.h"
+#include "MaterialExpressionTruncateLWC.generated.h"
+
+UCLASS(MinimalAPI)
+class UMaterialExpressionTruncateLWC : public UMaterialExpression
+{
+	GENERATED_UCLASS_BODY()
+
+	UPROPERTY()
+	FExpressionInput Input;
+
+	//~ Begin UMaterialExpression Interface
+#if WITH_EDITOR
+	virtual bool GenerateHLSLExpression(FMaterialHLSLGenerator& Generator, UE::HLSLTree::FScope& Scope, int32 OutputIndex, UE::HLSLTree::FExpression const*& OutExpression) const override;
+	virtual int32 Compile(class FMaterialCompiler* Compiler, int32 OutputIndex) override;
+	virtual void GetCaption(TArray<FString>& OutCaptions) const override { OutCaptions.Add(TEXT("TruncateLWC")); }
+	virtual FText GetKeywords() const override { return FText::FromString(TEXT("truncate lwc")); }
+#endif
+	//~ End UMaterialExpression Interface
+};
+
+
+
diff --git a/Engine/Source/Runtime/Engine/Private/Materials/HLSLMaterialDerivativeAutogen.cpp b/Engine/Source/Runtime/Engine/Private/Materials/HLSLMaterialDerivativeAutogen.cpp
--- a/Engine/Source/Runtime/Engine/Private/Materials/HLSLMaterialDerivativeAutogen.cpp
+++ b/Engine/Source/Runtime/Engine/Private/Materials/HLSLMaterialDerivativeAutogen.cpp
@@ -857,7 +857,7 @@ int32 FMaterialDerivativeAutogen::GenerateExpressionFunc2(FHLSLMaterialTranslato
 		// we will be calculating a valid derivative, so this value doesn't matter, but make it false for consistency
 		bIsDerivValidZero = false;
 
-		// if the lhs has a derivitive of zero, and rhs is non-zero, convert lhs from a scalar type to deriv type
+		// if the lhs has a derivative of zero, and rhs is non-zero, convert lhs from a scalar type to deriv type
 		if (LhsDerivInfo.DerivativeStatus == EDerivativeStatus::Zero)
 		{
 			bMakeDerivLhs = true;
diff --git a/Engine/Source/Runtime/Engine/Private/Materials/HLSLMaterialTranslator.cpp b/Engine/Source/Runtime/Engine/Private/Materials/HLSLMaterialTranslator.cpp
--- a/Engine/Source/Runtime/Engine/Private/Materials/HLSLMaterialTranslator.cpp
+++ b/Engine/Source/Runtime/Engine/Private/Materials/HLSLMaterialTranslator.cpp
@@ -10,6 +10,13 @@
 #include "Engine/RendererSettings.h"
 #include <functional>
 
+static int32 GetLWCTruncateMode()
+{
+	static const auto CVar = IConsoleManager::Get().FindTConsoleVariableDataInt(TEXT("r.MaterialEditor.LWCTruncateMode"));
+	const int32 LWCTruncateValue = CVar ? CVar->GetValueOnAnyThread() : 0;
+	return LWCTruncateValue;
+}
+
 bool IsExpressionClassPermitted(const UClass* const Class)
 {
 	if (Class)
@@ -4076,6 +4083,32 @@ int32 FHLSLMaterialTranslator::ForceCast(int32 Code, EMaterialValueType DestType
 	}
 }
 
+int32 FHLSLMaterialTranslator::TruncateLWC(int32 Code)
+{
+	const int32 LWCTruncateMode = GetLWCTruncateMode();
+	const bool bAllowLWCTruncate = LWCTruncateMode == 1 || LWCTruncateMode == 2;
+
+	if (!bAllowLWCTruncate)
+	{
+		return Code;
+	}
+
+	const EMaterialValueType TruncateType = GetParameterType(Code);
+	if (IsLWCType(TruncateType))
+	{
+		EMaterialValueType TruncatedType = MakeNonLWCType(TruncateType);
+		
+		if (TruncateType == MCT_LWCScalar)
+		{
+			TruncatedType = MCT_Float; // Remap from MCT_Float1 to MCT_Float to support replication
+		}
+
+		return ValidCast(Code, TruncatedType);
+	}
+
+	return Code;
+}
+
 int32 FHLSLMaterialTranslator::CastToNonLWCIfDisabled(int32 Code)
 {
 	int32 Result = Code;
@@ -5298,7 +5331,6 @@ int32 FHLSLMaterialTranslator::WorldPosition(EWorldPositionIncludedOffsets World
 	if (ShaderFrequency == SF_Pixel)
 	{
 		// No material offset only available in the vertex shader.
-		// TODO: should also be available in the tesselation shader
 		FunctionNamePattern.ReplaceInline(TEXT("<NO_MATERIAL_OFFSETS>"), TEXT("_NoMaterialOffsets"));
 	}
 	else
@@ -7331,21 +7363,39 @@ int32 FHLSLMaterialTranslator::Sub(int32 A, int32 B)
 	}
 	else
 	{
+		int32 ResultCode = INDEX_NONE;
+
 		if (IsAnalyticDerivEnabled())
 		{
-			return DerivativeAutogen.GenerateExpressionFunc2(*this, FMaterialDerivativeAutogen::EFunc2::Sub, A, B);
+			ResultCode = DerivativeAutogen.GenerateExpressionFunc2(*this, FMaterialDerivativeAutogen::EFunc2::Sub, A, B);
 		}
 		else
 		{
 			if (ResultType & MCT_LWCType)
 			{
-				return AddCodeChunk(ResultType, TEXT("LWCSubtract(%s, %s)"), *GetParameterCode(A), *GetParameterCode(B));
+				ResultCode = AddCodeChunk(ResultType, TEXT("LWCSubtract(%s, %s)"), *GetParameterCode(A), *GetParameterCode(B));
 			}
 			else
 			{
-				return AddCodeChunk(ResultType, TEXT("(%s - %s)"), *GetParameterCode(A), *GetParameterCode(B));
+				ResultCode = AddCodeChunk(ResultType, TEXT("(%s - %s)"), *GetParameterCode(A), *GetParameterCode(B));
 			}
 		}
+
+		// If both sides are LWC and we are subtracting, assume relative space going forward and truncate to single precision float.
+		const EMaterialValueType AType = GetParameterType(A);
+		const EMaterialValueType BType = GetParameterType(B);
+
+		const bool bTruncateToFloat = (AType == BType) && IsLWCType(AType) && IsLWCType(BType);
+		if (bTruncateToFloat)
+		{
+			const bool bAllowLWCTruncate = GetLWCTruncateMode() == 2; // Allow automatic truncate?
+			if (bAllowLWCTruncate)
+			{
+				return TruncateLWC(ResultCode);
+			}
+		}
+
+		return ResultCode;
 	}
 }
 
diff --git a/Engine/Source/Runtime/Engine/Private/Materials/HLSLMaterialTranslator.h b/Engine/Source/Runtime/Engine/Private/Materials/HLSLMaterialTranslator.h
--- a/Engine/Source/Runtime/Engine/Private/Materials/HLSLMaterialTranslator.h
+++ b/Engine/Source/Runtime/Engine/Private/Materials/HLSLMaterialTranslator.h
@@ -730,6 +730,8 @@ protected:
 	virtual int32 ValidCast(int32 Code, EMaterialValueType DestType) override;
 	virtual int32 ForceCast(int32 Code, EMaterialValueType DestType, uint32 ForceCastFlags = 0) override;
 
+	virtual int32 TruncateLWC(int32 Code) override;
+
 	/** Pushes a function onto the compiler's function stack, which indicates that compilation is entering a function. */
 	virtual void PushFunction(FMaterialFunctionCompileState* FunctionState) override;
 
diff --git a/Engine/Source/Runtime/Engine/Private/Materials/MaterialExpressionHLSL.cpp b/Engine/Source/Runtime/Engine/Private/Materials/MaterialExpressionHLSL.cpp
--- a/Engine/Source/Runtime/Engine/Private/Materials/MaterialExpressionHLSL.cpp
+++ b/Engine/Source/Runtime/Engine/Private/Materials/MaterialExpressionHLSL.cpp
@@ -77,6 +77,7 @@
 #include "Materials/MaterialExpressionCeil.h"
 #include "Materials/MaterialExpressionRound.h"
 #include "Materials/MaterialExpressionTruncate.h"
+#include "Materials/MaterialExpressionTruncateLWC.h"
 #include "Materials/MaterialExpressionSaturate.h"
 #include "Materials/MaterialExpressionSign.h"
 #include "Materials/MaterialExpressionSine.h"
@@ -1778,6 +1779,18 @@ bool UMaterialExpressionBlendMaterialAttributes::GenerateHLSLExpression(FMateria
 	return true;
 }
 
+bool UMaterialExpressionTruncateLWC::GenerateHLSLExpression(FMaterialHLSLGenerator& Generator, UE::HLSLTree::FScope& Scope, int32 OutputIndex, UE::HLSLTree::FExpression const*& OutExpression) const
+{
+	const UE::HLSLTree::FExpression* ExpressionInput = Input.AcquireHLSLExpression(Generator, Scope);
+	if (!ExpressionInput)
+	{
+		return false;
+	}
+
+	OutExpression = Generator.GetTree().NewTruncateLWC(ExpressionInput);
+	return true;
+}
+
 static const UE::HLSLTree::FExpression* TransformBase(UE::HLSLTree::FTree& Tree,
 	EMaterialCommonBasis SourceCoordBasis,
 	EMaterialCommonBasis DestCoordBasis,
diff --git a/Engine/Source/Runtime/Engine/Private/Materials/MaterialExpressions.cpp b/Engine/Source/Runtime/Engine/Private/Materials/MaterialExpressions.cpp
--- a/Engine/Source/Runtime/Engine/Private/Materials/MaterialExpressions.cpp
+++ b/Engine/Source/Runtime/Engine/Private/Materials/MaterialExpressions.cpp
@@ -234,6 +234,7 @@
 #include "Materials/MaterialExpressionTransform.h"
 #include "Materials/MaterialExpressionTransformPosition.h"
 #include "Materials/MaterialExpressionTruncate.h"
+#include "Materials/MaterialExpressionTruncateLWC.h"
 #include "Materials/MaterialExpressionTwoSidedSign.h"
 #include "Materials/MaterialExpressionVectorNoise.h"
 #include "Materials/MaterialExpressionVertexColor.h"
@@ -9984,6 +9985,41 @@ int32 UMaterialExpressionNormalize::Compile(class FMaterialCompiler* Compiler, i
 }
 #endif // WITH_EDITOR
 
+//
+//	UMaterialExpressionTruncateLWC
+//
+UMaterialExpressionTruncateLWC::UMaterialExpressionTruncateLWC(const FObjectInitializer& ObjectInitializer)
+	: Super(ObjectInitializer)
+{
+#if WITH_EDITORONLY_DATA
+	// Structure to hold one-time initialization
+	struct FConstructorStatics
+	{
+		FText NAME_Coordinates;
+		FConstructorStatics()
+			: NAME_Coordinates(LOCTEXT("Coordinates", "Coordinates"))
+		{
+		}
+	};
+	static FConstructorStatics ConstructorStatics;
+
+	MenuCategories.Add(ConstructorStatics.NAME_Coordinates);
+#endif
+}
+
+#if WITH_EDITOR
+int32 UMaterialExpressionTruncateLWC::Compile(class FMaterialCompiler* Compiler, int32 OutputIndex)
+{
+	if (!Input.GetTracedInput().Expression)
+	{
+		return Compiler->Errorf(TEXT("Missing input"));
+	}
+
+	int32 CodeIndex = Input.Compile(Compiler);
+	return Compiler->TruncateLWC(CodeIndex);
+}
+#endif // WITH_EDITOR
+
 UMaterialExpressionVertexColor::UMaterialExpressionVertexColor(const FObjectInitializer& ObjectInitializer)
 	: Super(ObjectInitializer)
 {
diff --git a/Engine/Source/Runtime/Engine/Public/HLSLTree/HLSLTree.h b/Engine/Source/Runtime/Engine/Public/HLSLTree/HLSLTree.h
--- a/Engine/Source/Runtime/Engine/Public/HLSLTree/HLSLTree.h
+++ b/Engine/Source/Runtime/Engine/Public/HLSLTree/HLSLTree.h
@@ -664,6 +664,8 @@ public:
 	const FExpression* NewDot(const FExpression* Lhs, const FExpression* Rhs) { return NewSum(NewMul(Lhs, Rhs)); }
 	const FExpression* NewLerp(const FExpression* A, const FExpression* B, const FExpression* T) { return NewAdd(A, NewMul(NewSub(B, A), T)); }
 
+	const FExpression* NewTruncateLWC(const FExpression* Input) { return NewUnaryOp(EOperation::TruncateLWC, Input); }
+
 private:
 	template<typename T, typename... ArgTypes>
 	inline T* NewNode(ArgTypes&&... Args)
diff --git a/Engine/Source/Runtime/Engine/Public/HLSLTree/HLSLTreeTypes.h b/Engine/Source/Runtime/Engine/Public/HLSLTree/HLSLTreeTypes.h
--- a/Engine/Source/Runtime/Engine/Public/HLSLTree/HLSLTreeTypes.h
+++ b/Engine/Source/Runtime/Engine/Public/HLSLTree/HLSLTreeTypes.h
@@ -102,6 +102,8 @@ enum class EOperation : uint8
 	Atan,
 	AtanFast,
 
+	TruncateLWC,
+
 	// Binary Ops
 	Add,
 	Sub,
diff --git a/Engine/Source/Runtime/Engine/Public/MaterialCompiler.h b/Engine/Source/Runtime/Engine/Public/MaterialCompiler.h
--- a/Engine/Source/Runtime/Engine/Public/MaterialCompiler.h
+++ b/Engine/Source/Runtime/Engine/Public/MaterialCompiler.h
@@ -269,6 +269,8 @@ public:
 	virtual int32 ValidCast(int32 Code,EMaterialValueType DestType) = 0;
 	virtual int32 ForceCast(int32 Code,EMaterialValueType DestType,uint32 ForceCastFlags = 0) = 0;
 
+	virtual int32 TruncateLWC(int32 Code) = 0;
+
 	/** Pushes a function onto the compiler's function stack, which indicates that compilation is entering a function. */
 	virtual void PushFunction(FMaterialFunctionCompileState* FunctionState) = 0;
 
@@ -791,6 +793,7 @@ public:
 	virtual bool IsMaterialPropertyUsed(EMaterialProperty Property, int32 CodeChunkIdx) const override { return Compiler->IsMaterialPropertyUsed(Property, CodeChunkIdx); }
 	virtual int32 ValidCast(int32 Code, EMaterialValueType DestType) override { return Compiler->ValidCast(Code, DestType); }
 	virtual int32 ForceCast(int32 Code, EMaterialValueType DestType, uint32 ForceCastFlags = 0) override { return Compiler->ForceCast(Code, DestType, ForceCastFlags); }
+	virtual int32 TruncateLWC(int32 Code) override { return Compiler->TruncateLWC(Code); }
 
 	virtual int32 AccessCollectionParameter(UMaterialParameterCollection* ParameterCollection, int32 ParameterIndex, int32 ComponentIndex) override { return Compiler->AccessCollectionParameter(ParameterCollection, ParameterIndex, ComponentIndex); }
 	virtual int32 NumericParameter(EMaterialParameterType ParameterType, FName ParameterName, const UE::Shader::FValue& DefaultValue) override { return Compiler->NumericParameter(ParameterType, ParameterName, DefaultValue); }
diff --git a/Engine/Source/Runtime/RenderCore/Private/Shader.cpp b/Engine/Source/Runtime/RenderCore/Private/Shader.cpp
--- a/Engine/Source/Runtime/RenderCore/Private/Shader.cpp
+++ b/Engine/Source/Runtime/RenderCore/Private/Shader.cpp
@@ -1481,6 +1481,19 @@ void ShaderMapAppendKeyString(EShaderPlatform Platform, FString& KeyString)
 		KeyString += bValue ? TEXT("_SL") : TEXT("_NoSL");
 	}
 
+	{
+		static const auto CVar = IConsoleManager::Get().FindTConsoleVariableDataInt(TEXT("r.MaterialEditor.LWCTruncateMode"));
+		const int32 LWCTruncateValue = CVar ? CVar->GetValueOnAnyThread() : 0;
+		if (LWCTruncateValue == 1)
+		{
+			KeyString += TEXT("_LWC1");
+		}
+		else if (LWCTruncateValue == 2)
+		{
+			KeyString += TEXT("_LWC2");
+		}
+	}
+
 	{
 		KeyString += IsUsingBasePassVelocity(Platform) ? TEXT("_GV") : TEXT("");
 	}
