Strata - thin film becomes are helper node.

diff --git a/Engine/Shaders/Private/Strata/Strata.ush b/Engine/Shaders/Private/Strata/Strata.ush
--- a/Engine/Shaders/Private/Strata/Strata.ush
+++ b/Engine/Shaders/Private/Strata/Strata.ush
@@ -1063,10 +1063,6 @@ struct FStrataOperator
 	float				 OpaqueRoughRefractThicknessCm;	// The thickness of stack material from the top layer to the bottom layer.
 	float				 OpaqueRoughRefractTopRoughness;// The roughness of the top layer. This is enough because we currently assume all the bottom layers have a constant IOR, so eta will after that always be 1 and no refraction will occur after.
 
-	// Thin film
-	float ThinFilmThickness;							// Thickness of thin-film coating layer. In micro-meters (1um == 1000nm).  Merge with regular ThicknessCm?
-	float ThinFilmIOR;									// IOR of the thin-film coating layer
-
 	// Translucent refractions
 	FStrataLobeStatistic VerticalTop_TopDownRefractionLobeStat;	// Only used by vertical layering operators. It represent the lobe statistics of the top layer at its interface.
 	FStrataLobeStatistic VerticalTop_BottomUpRefractionLobeStat;
@@ -2081,30 +2077,6 @@ FStrataData StrataWeight(FStrataData A, float Weight, inout FStrataTree StrataTr
 	return StrataData;
 }
 
-FStrataData StrataThinFilm(FStrataData A, float Thickness, float IOR, inout FStrataTree StrataTree, int OperatorIndex, uint MaxDistanceFromLeaves)
-{
-	const float SafeThickness = saturate(Thickness);
-	const float SafeIOR = clamp(IOR, 1.01f, 3.0f);
-
-	FStrataData StrataData = GetInitialisedStrataData();
-	StrataData.OperatorIndex = OperatorIndex;
-#if USE_DEVELOPMENT_SHADERS
-	StrataData.PreviewColor = A.PreviewColor;
-#endif
-
-	StrataTree.OperatorCount++;
-	StrataTree.Operators[OperatorIndex] = GetInitialisedStrataOperator();
-	StrataTree.Operators[OperatorIndex].MaxDistanceFromLeaves = MaxDistanceFromLeaves;
-	StrataTree.Operators[OperatorIndex].Type = STRATA_OPERATOR_THINFILM;
-	StrataTree.Operators[OperatorIndex].ThinFilmThickness = SafeThickness;
-	StrataTree.Operators[OperatorIndex].ThinFilmIOR = SafeIOR;
-	StrataTree.Operators[OperatorIndex].LeftIndex = A.OperatorIndex;
-
-	StrataTree.Operators[A.OperatorIndex].ParentIndex = OperatorIndex;
-
-	return StrataData;
-}
-
 FStrataData StrataHorizontalMixing(FStrataData Background, FStrataData Foreground, float Mix, inout FStrataTree StrataTree, int OperatorIndex, uint MaxDistanceFromLeaves)
 {
 	const float ForegroundMixFactor = saturate(Mix);
@@ -2539,11 +2511,18 @@ void StrataGetThinFilmF0F90(float NoV, float ThinFilmThickness, float ThinFilmIO
 	OutF90 = max(F0_Thin, F90);
 }
 
-FStrataData StrataThinFilmParameterBlending(in FStrataData A, float ThinFilmThickness, float ThinFilmIOR, float NoV)
+struct FStrataThinFilmOutput
 {
-	FStrataData StrataData = A;
-	StrataGetThinFilmF0F90(NoV, ThinFilmThickness, ThinFilmIOR, SLAB_F0(StrataData.InlinedBSDF), SLAB_F90(StrataData.InlinedBSDF));
-	return StrataData;
+	float3 F0;
+	float3 F90;
+};
+FStrataThinFilmOutput StrataGetThinFilmF0F90(float NoV, float3 F0, float3 F90, float ThinFilmThickness, float ThinFilmIOR)
+{
+	FStrataThinFilmOutput Output;
+	Output.F0 = F0;
+	Output.F90 = F90;
+	StrataGetThinFilmF0F90(NoV, ThinFilmThickness, ThinFilmIOR, Output.F0, Output.F90);
+	return Output;
 }
 
 ///////////////////////////////////////////////////////////////////////////////
diff --git a/Engine/Shaders/Private/Strata/StrataTree.ush b/Engine/Shaders/Private/Strata/StrataTree.ush
--- a/Engine/Shaders/Private/Strata/StrataTree.ush
+++ b/Engine/Shaders/Private/Strata/StrataTree.ush
@@ -317,21 +317,6 @@ void UpdateSingleOperatorCoverageTransmittance(inout FStrataTree StrataTree, int
 		break;
 	}
 
-	case STRATA_OPERATOR_THINFILM:
-	{
-		// Thin-film does not affect these property
-		Op.Coverage = OpA.Coverage;
-		Op.ThroughputAlongV = OpA.ThroughputAlongV;
-		Op.TransmittanceAlongN = OpA.TransmittanceAlongN;
-		Op.TopDownRefractionLobeStat = OpA.TopDownRefractionLobeStat;
-		Op.TopDownRefractionWorldNormal = OpA.TopDownRefractionWorldNormal;
-		Op.BottomUpRefractionLobeStat = OpA.BottomUpRefractionLobeStat;
-		Op.OpaqueRoughRefractThicknessCm = OpA.OpaqueRoughRefractThicknessCm;
-		Op.OpaqueRoughRefractCoverage = OpA.OpaqueRoughRefractCoverage;
-		Op.OpaqueRoughRefractTopRoughness = OpA.OpaqueRoughRefractTopRoughness;
-		break;
-	}
-
 	}
 
 #undef Op
@@ -424,51 +409,6 @@ void UpdateAllBSDFWithBottomUpOperatorVisit_Vertical(
 #endif // STRATA_INLINE_SHADING
 }
 
-// Pre-update visit operators
-void PreUpdateAllBSDFWithBottomUpOperatorVisit_Vertical(
-	FStrataPixelHeader StrataPixelHeader,
-	inout FStrataTree StrataTree,
-	inout FStrataBSDF CurrentBSDF,
-	FStrataAddressing NullStrataAddressing,
-	inout bool bCanReceiveThinFilm,
-	float3 V,
-	int OpIndex,
-	int PreviousIsInputA)
-{
-#if STRATA_INLINE_SHADING
-
-	// If the BSDF is coated with a vertical layer, it can no longer receive thin film-coating
-	if (PreviousIsInputA <= 0)
-	{
-		bCanReceiveThinFilm = false;
-	}
-
-#endif // STRATA_INLINE_SHADING
-}
-
-void PreUpdateAllBSDFWithBottomUpOperatorVisit_ThinFilm(
-	FStrataPixelHeader StrataPixelHeader,
-	inout FStrataTree StrataTree,
-	inout FStrataBSDF CurrentBSDF,
-	FStrataAddressing NullStrataAddressing,
-	inout bool bCanReceiveThinFilm,
-	float3 V,
-	int OpIndex,
-	int PreviousIsInputA)
-{
-#if STRATA_INLINE_SHADING
-
-	if (bCanReceiveThinFilm && BSDF_GETTYPE(CurrentBSDF) == STRATA_BSDF_TYPE_SLAB)
-	{
-		const float3x3 TangentBasis = StrataGetBSDFSharedBasis_InlineShading(StrataPixelHeader, CurrentBSDF, NullStrataAddressing);
-		const float3 N = TangentBasis[2];
-		const float NoV = saturate(dot(N, V));
-		StrataGetThinFilmF0F90(NoV, Op.ThinFilmThickness, Op.ThinFilmIOR, SLAB_F0(CurrentBSDF), SLAB_F90(CurrentBSDF));
-	}
-
-#endif // STRATA_INLINE_SHADING
-}
-
 #undef Op
 
 
diff --git a/Engine/Shaders/Private/Strata/StrataVisualize.usf b/Engine/Shaders/Private/Strata/StrataVisualize.usf
--- a/Engine/Shaders/Private/Strata/StrataVisualize.usf
+++ b/Engine/Shaders/Private/Strata/StrataVisualize.usf
@@ -633,7 +633,7 @@ void VisualizeMaterialPS(
 
 float2 GetBranchOffset(FStrataOperator Op, bool bIsLeft)
 {
-	if (Op.Type == STRATA_OPERATOR_WEIGHT || Op.Type == STRATA_OPERATOR_THINFILM)
+	if (Op.Type == STRATA_OPERATOR_WEIGHT)
 	{
 		return float2(0, TREE_BRANCH.y);
 	}
@@ -739,10 +739,6 @@ void MaterialDebugStrataTreeCS(uint3 DispatchThreadId : SV_DispatchThreadID)
 				{
 					Print(Context, TEXT("BSDF_LEGACY"), FontWhite);
 				}
-				else if (Op.Type == STRATA_OPERATOR_THINFILM)
-				{
-					Print(Context, TEXT("THINFILM"), FontWhite);
-				}
 				else
 				{
 					Print(Context, TEXT("UNKOWN"), FontWhite);
@@ -885,10 +881,6 @@ void MaterialDebugStrataTreeCS(uint3 DispatchThreadId : SV_DispatchThreadID)
 			{
 				Print(Context, TEXT("BL"), FontWhite);
 			}
-			else if (Op.Type == STRATA_OPERATOR_THINFILM)
-			{
-				Print(Context, TEXT("TF"), FontWhite);
-			}
 			else
 			{
 				Print(Context, TEXT("UNKOWN"), FontWhite);
diff --git a/Engine/Shaders/Shared/StrataDefinitions.h b/Engine/Shaders/Shared/StrataDefinitions.h
--- a/Engine/Shaders/Shared/StrataDefinitions.h
+++ b/Engine/Shaders/Shared/StrataDefinitions.h
@@ -49,7 +49,6 @@
 #define STRATA_OPERATOR_ADD				3
 #define STRATA_OPERATOR_BSDF			4
 #define STRATA_OPERATOR_BSDF_LEGACY		5
-#define STRATA_OPERATOR_THINFILM		6
 
 // This must map directly to EStrataTileMaterialType
 #define STRATA_TILE_TYPE_SIMPLE						0
diff --git a/Engine/Source/Runtime/Engine/Classes/Materials/MaterialExpressionStrata.h b/Engine/Source/Runtime/Engine/Classes/Materials/MaterialExpressionStrata.h
--- a/Engine/Source/Runtime/Engine/Classes/Materials/MaterialExpressionStrata.h
+++ b/Engine/Source/Runtime/Engine/Classes/Materials/MaterialExpressionStrata.h
@@ -934,42 +934,6 @@ class UMaterialExpressionStrataWeight : public UMaterialExpressionStrataBSDF
 	//~ End UMaterialExpression Interface
 };
 
-UCLASS(MinimalAPI, collapsecategories, hidecategories = Object, DisplayName = "Strata Thin-Film")
-class UMaterialExpressionStrataThinFilm : public UMaterialExpressionStrataBSDF
-{
-	GENERATED_UCLASS_BODY()
-
-	/**
-	 * Strata material
-	 */
-	UPROPERTY()
-	FExpressionInput A;
-	
-	/**
-	 * Thin film controls the thin film layer coating the current slab. 0 means disabled and 1 means a coating layer of 10 micrometer. (type = float, unitless, default = 0)
-	 */
-	UPROPERTY()
-	FExpressionInput Thickness;
-
-	/**
-	 * Thin film IOR
-	 */
-	UPROPERTY()
-	FExpressionInput IOR;
-
-	//~ Begin UMaterialExpression Interface
-#if WITH_EDITOR
-	virtual int32 Compile(class FMaterialCompiler* Compiler, int32 OutputIndex) override;
-	virtual void GetCaption(TArray<FString>& OutCaptions) const override;
-	virtual uint32 GetOutputType(int32 OutputIndex) override;
-	virtual uint32 GetInputType(int32 InputIndex) override;
-	virtual bool IsResultStrataMaterial(int32 OutputIndex) override;
-	virtual void GatherStrataMaterialInfo(FStrataMaterialInfo& StrataMaterialInfo, int32 OutputIndex) override;
-	virtual FStrataOperator* StrataGenerateMaterialTopologyTree(class FMaterialCompiler* Compiler, class UMaterialExpression* Parent, int32 OutputIndex) override;
-#endif
-	//~ End UMaterialExpression Interface
-};
-
 
 ///////////////////////////////////////////////////////////////////////////////
 // Utilities
@@ -1074,3 +1038,50 @@ class UMaterialExpressionStrataHazinessToSecondaryRoughness : public UMaterialEx
 	//~ End UMaterialExpression Interface
 };
 
+UCLASS(MinimalAPI, collapsecategories, hidecategories = Object, DisplayName = "Strata Thin-Film")
+class UMaterialExpressionStrataThinFilm : public UMaterialExpressionStrataUtilityBase
+{
+	GENERATED_UCLASS_BODY()
+
+	/**
+	 * The normal of the surface to consider. This input respects the normal space setup on the root node (tangent or world)
+	 */
+	UPROPERTY()
+	FExpressionInput Normal;
+
+	/**
+	 * Defines F0, the percentage of light reflected as specular from a surface when the view is perpendicular to the surface. (type = float3, unit = unitless, defaults to plastic 0.04)
+	 */
+	UPROPERTY()
+	FExpressionInput F0;
+
+	/**
+	 * Defines F90, the percentage of light reflected as specular from a surface when the view is tangent to the surface. (type = float3, unit = unitless, defaults to 1.0f)
+	 */
+	UPROPERTY()
+	FExpressionInput F90;
+	
+	/**
+	 * Thin film controls the thin film layer coating the current slab. 0 means disabled and 1 means a coating layer of 10 micrometer. (type = float, unitless, default = 0)
+	 */
+	UPROPERTY()
+	FExpressionInput Thickness;
+
+	/**
+	 * Thin film IOR
+	 */
+	UPROPERTY()
+	FExpressionInput IOR;
+
+	//~ Begin UMaterialExpression Interface
+#if WITH_EDITOR
+	virtual int32 Compile(class FMaterialCompiler* Compiler, int32 OutputIndex) override;
+	virtual void GetCaption(TArray<FString>& OutCaptions) const override;
+	virtual uint32 GetOutputType(int32 OutputIndex) override;
+	virtual uint32 GetInputType(int32 InputIndex) override;
+	virtual void GetConnectorToolTip(int32 InputIndex, int32 OutputIndex, TArray<FString>& OutToolTip) override;
+	virtual void GetExpressionToolTip(TArray<FString>& OutToolTip) override;
+#endif
+	//~ End UMaterialExpression Interface
+};
+
diff --git a/Engine/Source/Runtime/Engine/Private/Materials/HLSLMaterialTranslator.cpp b/Engine/Source/Runtime/Engine/Private/Materials/HLSLMaterialTranslator.cpp
--- a/Engine/Source/Runtime/Engine/Private/Materials/HLSLMaterialTranslator.cpp
+++ b/Engine/Source/Runtime/Engine/Private/Materials/HLSLMaterialTranslator.cpp
@@ -219,10 +219,6 @@ static inline const TCHAR* GetStrataOperatorStr(int32 OperatorType)
 	{
 		return TEXT("BSDFLEGACY");
 	}
-	case STRATA_OPERATOR_THINFILM:
-	{
-		return TEXT("THINFILM");
-	}
 	}
 	return TEXT("UNKNOWN   ");
 };
@@ -1402,7 +1398,6 @@ bool FHLSLMaterialTranslator::Translate()
 					for (uint32 BSDFIndex = 0; BSDFIndex < StrataMaterialBSDFCount; ++BSDFIndex)
 					{
 						ResourcesString += "\t{\n";
-						ResourcesString += "\t\tbool bCanReceiveThinFilm = true;\n";
 						for (auto& It : StrataMaterialExpressionRegisteredOperators)
 						{
 							if (!It.IsDiscarded() && It.BSDFIndex == BSDFIndex)
@@ -1422,18 +1417,13 @@ bool FHLSLMaterialTranslator::Translate()
 									}
 									case STRATA_OPERATOR_VERTICAL:
 									{
-										ResourcesString += FString::Printf(TEXT("\t PreUpdateAllBSDFWithBottomUpOperatorVisit_Vertical(StrataPixelHeader, StrataTree, StrataTree.BSDFs[%d], NullStrataAddressing, bCanReceiveThinFilm, V, %d /*Op index*/, %d /*PreviousIsInputA*/);\n"), BSDFIndex, CurrentOperator.Index, CurrentOperator.LeftIndex == PreviousOperatorIndex ? 1 : 0);
-										break;
+										// example ResourcesString += FString::Printf(TEXT("\t PreUpdateAllBSDFWithBottomUpOperatorVisit_Vertical(StrataPixelHeader, StrataTree, StrataTree.BSDFs[%d], NullStrataAddressing, V, %d /*Op index*/, %d /*PreviousIsInputA*/);\n"), BSDFIndex, CurrentOperator.Index, CurrentOperator.LeftIndex == PreviousOperatorIndex ? 1 : 0);
+										break; // NOP
 									}
 									case STRATA_OPERATOR_ADD:
 									{
 										break; // NOP
 									}
-									case STRATA_OPERATOR_THINFILM:
-									{
-										ResourcesString += FString::Printf(TEXT("\t PreUpdateAllBSDFWithBottomUpOperatorVisit_ThinFilm(StrataPixelHeader, StrataTree, StrataTree.BSDFs[%d], NullStrataAddressing, bCanReceiveThinFilm, V,  %d /*Op index*/, %d /*PreviousIsInputA*/);\n"), BSDFIndex, CurrentOperator.Index, CurrentOperator.LeftIndex == PreviousOperatorIndex ? 1 : 0);
-										break;
-									}
 									default:
 									case STRATA_OPERATOR_BSDF:
 									{
@@ -1533,10 +1523,6 @@ bool FHLSLMaterialTranslator::Translate()
 									{
 										break; // NOP
 									}
-									case STRATA_OPERATOR_THINFILM:
-									{
-										break; // NOP
-									}
 									default:
 									case STRATA_OPERATOR_BSDF:
 									{
@@ -10072,7 +10058,6 @@ bool FHLSLMaterialTranslator::StrataGenerateDerivedMaterialOperatorData()
 
 		// Operators with a single child
 		case STRATA_OPERATOR_WEIGHT:
-		case STRATA_OPERATOR_THINFILM:
 		{
 			bMustHaveLeftChild = true;
 		}
@@ -10142,7 +10127,6 @@ bool FHLSLMaterialTranslator::StrataGenerateDerivedMaterialOperatorData()
 					break;
 				}
 				case STRATA_OPERATOR_WEIGHT:
-				case STRATA_OPERATOR_THINFILM:
 				{
 					WalkOperators(StrataMaterialExpressionRegisteredOperators[CurrentOperator.LeftIndex], bUseParameterBlending);
 					bOperatorEncountered = true;
@@ -10243,7 +10227,6 @@ bool FHLSLMaterialTranslator::StrataGenerateDerivedMaterialOperatorData()
 
 			// Operators with a single child
 			case STRATA_OPERATOR_WEIGHT:
-			case STRATA_OPERATOR_THINFILM:
 			{
 				check(It.LeftIndex != INDEX_NONE);
 			}
@@ -10260,7 +10243,6 @@ bool FHLSLMaterialTranslator::StrataGenerateDerivedMaterialOperatorData()
 				switch (CurrentOperator.OperatorType)
 				{
 				case STRATA_OPERATOR_WEIGHT:
-				case STRATA_OPERATOR_THINFILM:
 				{
 					CurrentOperator.MaxDistanceFromLeaves = StrataMaterialExpressionRegisteredOperators[CurrentOperator.LeftIndex].MaxDistanceFromLeaves + 1;
 					break;
@@ -10331,7 +10313,6 @@ bool FHLSLMaterialTranslator::StrataGenerateDerivedMaterialOperatorData()
 					break;
 				}
 				case STRATA_OPERATOR_WEIGHT:
-				case STRATA_OPERATOR_THINFILM:
 				{
 					WalkOperators(StrataMaterialExpressionRegisteredOperators[CurrentOperator.LeftIndex]);
 					break;
@@ -11183,57 +11164,6 @@ int32 FHLSLMaterialTranslator::StrataWeightParameterBlending(int32 A, int32 Weig
 	);
 }
 
-int32 FHLSLMaterialTranslator::StrataThinFilm(int32 A, int32 Thickness, int32 IOR, int OperatorIndex, uint32 MaxDistanceFromLeaves)
-{
-	if (A == INDEX_NONE || Thickness == INDEX_NONE || IOR == INDEX_NONE)
-	{
-		return INDEX_NONE;
-	}
-	return AddCodeChunk(
-		MCT_Strata, TEXT("StrataThinFilm(%s, %s, %s, Parameters.StrataTree, %u, %u)"),
-		*GetParameterCode(A),
-		*GetParameterCode(Thickness),
-		*GetParameterCode(IOR),
-		OperatorIndex,
-		MaxDistanceFromLeaves
-	);
-}
-
-int32 FHLSLMaterialTranslator::StrataThinFilmParameterBlending(int32 A, int32 Thickness, int32 IOR, int32 BSDFNormalCodeChunk, FStrataOperator* PromoteToOperator)
-{
-	if (A == INDEX_NONE || Thickness == INDEX_NONE || IOR == INDEX_NONE)
-	{
-		return INDEX_NONE;
-	}
-
-	if (PromoteToOperator)
-	{
-		check(PromoteToOperator->Index != INDEX_NONE);
-		check(PromoteToOperator->BSDFIndex != INDEX_NONE);
-		return AddCodeChunk(
-			MCT_Strata, TEXT("PromoteParameterBlendedBSDFToOperator(StrataThinFilmParameterBlending(%s, %s, %s, dot(%s, %s)), Parameters.StrataTree, %u, %u, %u, %u)"),
-			*GetParameterCode(A),
-			*GetParameterCode(Thickness),
-			*GetParameterCode(IOR),
-			*GetParameterCode(BSDFNormalCodeChunk),
-			*GetParameterCode(CameraVector()),
-			PromoteToOperator->Index,
-			PromoteToOperator->BSDFIndex,
-			PromoteToOperator->LayerDepth,
-			PromoteToOperator->bIsBottom ? 1 : 0
-		);
-	}
-
-	return AddCodeChunk(
-		MCT_Strata, TEXT("StrataThinFilmParameterBlending(%s, %s, %s, dot(%s, %s))"),
-		*GetParameterCode(A),
-		*GetParameterCode(Thickness),
-		*GetParameterCode(IOR),
-		*GetParameterCode(BSDFNormalCodeChunk),
-		*GetParameterCode(CameraVector())
-	);
-}
-
 int32 FHLSLMaterialTranslator::StrataTransmittanceToMFP(int32 TransmittanceColor, int32 DesiredThickness, int32 OutputIndex)
 {
 	if (OutputIndex == INDEX_NONE)
@@ -11311,6 +11241,30 @@ int32 FHLSLMaterialTranslator::StrataHazinessToSecondaryRoughness(int32 BaseRoug
 	return INDEX_NONE;
 }
 
+int32 FHLSLMaterialTranslator::StrataThinFilm(int32 NormalCodeChunk, int32 SpecularColorCodeChunk, int32 EdgeSpecularColorCodeChunk, int32 ThicknessCodeChunk, int32 IORCodeChunk, int32 OutputIndex)
+{
+	if (NormalCodeChunk == INDEX_NONE || SpecularColorCodeChunk == INDEX_NONE || EdgeSpecularColorCodeChunk == INDEX_NONE
+		|| ThicknessCodeChunk == INDEX_NONE || IORCodeChunk == INDEX_NONE)
+	{
+		Errorf(TEXT("Strata error: one of StrataThinFilm node input is invalid (asset: %s).\r\n"), *Material->GetDebugName(), *Material->GetAssetPath().ToString());
+		return INDEX_NONE;
+	}
+	if (OutputIndex < 0 || OutputIndex > 1)
+	{
+		Errorf(TEXT("Strata error: StrataThinFilm output index is invalid (asset: %s).\r\n"), *Material->GetDebugName(), *Material->GetAssetPath().ToString());
+		return INDEX_NONE;
+	}
+	return AddCodeChunk(
+		MCT_Float3, TEXT("StrataGetThinFilmF0F90(%s, %s, %s, %s, %s)%s"),
+		*GetParameterCode(Dot(NormalCodeChunk, CameraVector())),
+		*GetParameterCode(SpecularColorCodeChunk),
+		*GetParameterCode(EdgeSpecularColorCodeChunk),
+		*GetParameterCode(ThicknessCodeChunk),
+		*GetParameterCode(IORCodeChunk),
+		OutputIndex == 0 ? TEXT(".F0") : TEXT(".F90")
+	);
+}
+
 int32 FHLSLMaterialTranslator::StrataCompilePreview(int32 StrataDataCodeChunk)
 {
 	if (StrataDataCodeChunk == INDEX_NONE)
diff --git a/Engine/Source/Runtime/Engine/Private/Materials/HLSLMaterialTranslator.h b/Engine/Source/Runtime/Engine/Private/Materials/HLSLMaterialTranslator.h
--- a/Engine/Source/Runtime/Engine/Private/Materials/HLSLMaterialTranslator.h
+++ b/Engine/Source/Runtime/Engine/Private/Materials/HLSLMaterialTranslator.h
@@ -1105,11 +1105,10 @@ protected:
 	virtual int32 StrataAddParameterBlending(int32 A, int32 B, int32 AMixWeight, const FString& SharedLocalBasisIndexMacro, FStrataOperator* PromoteToOperator) override;
 	virtual int32 StrataWeight(int32 A, int32 Weight, int OperatorIndex, uint32 MaxDistanceFromLeaves) override;
 	virtual int32 StrataWeightParameterBlending(int32 A, int32 Weight, FStrataOperator* PromoteToOperator) override;
-	virtual int32 StrataThinFilm(int32 A, int32 Thickness, int32 IOR, int OperatorIndex, uint32 MaxDistanceFromLeaves) override;
-	virtual int32 StrataThinFilmParameterBlending(int32 A, int32 Thickness, int32 IOR, int32 BSDFNormalCodeChunk, FStrataOperator* PromoteToOperator) override;
 	virtual int32 StrataTransmittanceToMFP(int32 TransmittanceColor, int32 DesiredThickness, int32 OutputIndex) override;
 	virtual int32 StrataMetalnessToDiffuseAlbedoF0(int32 BaseColor, int32 Specular, int32 Metallic, int32 OutputIndex) override;
 	virtual int32 StrataHazinessToSecondaryRoughness(int32 BaseRoughness, int32 Haziness, int32 OutputIndex) override;
+	virtual int32 StrataThinFilm(int32 NormalCodeChunk, int32 SpecularColorCodeChunk, int32 EdgeSpecularColorCodeChunk, int32 ThicknessCodeChunk, int32 IORCodeChunk, int32 OutputIndex) override;
 	virtual int32 StrataCompilePreview(int32 StrataDataCodeChunk) override;
 	virtual bool StrataSkipsOpacityEvaluation() override;
 
diff --git a/Engine/Source/Runtime/Engine/Private/Materials/MaterialExpressions.cpp b/Engine/Source/Runtime/Engine/Private/Materials/MaterialExpressions.cpp
--- a/Engine/Source/Runtime/Engine/Private/Materials/MaterialExpressions.cpp
+++ b/Engine/Source/Runtime/Engine/Private/Materials/MaterialExpressions.cpp
@@ -24524,109 +24524,6 @@ FStrataOperator* UMaterialExpressionStrataWeight::StrataGenerateMaterialTopology
 }
 #endif // WITH_EDITOR
 
-
-
-UMaterialExpressionStrataThinFilm::UMaterialExpressionStrataThinFilm(const FObjectInitializer& ObjectInitializer)
-	: Super(ObjectInitializer)
-{
-	struct FConstructorStatics
-	{
-		FText NAME_Strata;
-		FConstructorStatics() : NAME_Strata(LOCTEXT("Strata Ops", "Strata Operators")) { }
-	};
-	static FConstructorStatics ConstructorStatics;
-#if WITH_EDITORONLY_DATA
-	MenuCategories.Add(ConstructorStatics.NAME_Strata);
-#endif
-}
-
-#if WITH_EDITOR
-int32 UMaterialExpressionStrataThinFilm::Compile(class FMaterialCompiler* Compiler, int32 OutputIndex)
-{
-	if (!A.GetTracedInput().Expression)
-	{
-		return Compiler->Errorf(TEXT("Missing A input"));
-	}
-
-	Compiler->StrataTreeStackPush(this, 0);
-	int32 ACodeChunk = A.Compile(Compiler);
-	Compiler->StrataTreeStackPop();
-	int32 ThicknessCodeChunk = Thickness.GetTracedInput().Expression ? Thickness.Compile(Compiler) : Compiler->Constant(1.0f);
-	int32 IORCodeChunk = IOR.GetTracedInput().Expression ? IOR.Compile(Compiler) : Compiler->Constant(1.44f);
-
-	int32 OutputCodeChunk = INDEX_NONE;
-	FStrataOperator& StrataOperator = Compiler->StrataCompilationGetOperator(Compiler->StrataTreeStackGetPathUniqueId());
-	if (StrataOperator.bUseParameterBlending)
-	{
-		// Propagate the parameter blended normal
-		FStrataOperator* Operator = Compiler->StrataCompilationGetOperatorFromIndex(StrataOperator.LeftIndex);
-		if (!Operator)
-		{
-			return Compiler->Errorf(TEXT("Missing input on ThinFilm node."));
-		}
-		StrataOperator.BSDFRegisteredSharedLocalBasis = Operator->BSDFRegisteredSharedLocalBasis;
-
-		int32 NormalCodeChunk = Operator->BSDFRegisteredSharedLocalBasis.NormalCodeChunk;
-		OutputCodeChunk = Compiler->StrataThinFilmParameterBlending(ACodeChunk, ThicknessCodeChunk, IORCodeChunk, NormalCodeChunk, StrataOperator.bRootOfParameterBlendingSubTree ? &StrataOperator : nullptr);
-	}
-	else
-	{
-		OutputCodeChunk = Compiler->StrataThinFilm(ACodeChunk, ThicknessCodeChunk, IORCodeChunk, StrataOperator.Index, StrataOperator.MaxDistanceFromLeaves);
-	}
-
-	return OutputCodeChunk;
-}
-
-void UMaterialExpressionStrataThinFilm::GetCaption(TArray<FString>& OutCaptions) const
-{
-	OutCaptions.Add(TEXT("Strata Thin-Film"));
-}
-
-uint32 UMaterialExpressionStrataThinFilm::GetOutputType(int32 OutputIndex)
-{
-	return MCT_Strata;
-}
-
-uint32 UMaterialExpressionStrataThinFilm::GetInputType(int32 InputIndex)
-{
-	if (InputIndex == 0) { return MCT_Strata; } // Material
-	if (InputIndex == 1) { return MCT_Float1; } // Thickness
-	if (InputIndex == 2) { return MCT_Float1; } // IOR
-	return MCT_Float1;
-}
-
-bool UMaterialExpressionStrataThinFilm::IsResultStrataMaterial(int32 OutputIndex)
-{
-	return true;
-}
-
-void UMaterialExpressionStrataThinFilm::GatherStrataMaterialInfo(FStrataMaterialInfo& StrataMaterialInfo, int32 OutputIndex)
-{
-	if (A.GetTracedInput().Expression)
-	{
-		A.GetTracedInput().Expression->GatherStrataMaterialInfo(StrataMaterialInfo, A.OutputIndex);
-	}
-}
-
-FStrataOperator* UMaterialExpressionStrataThinFilm::StrataGenerateMaterialTopologyTree(class FMaterialCompiler* Compiler, class UMaterialExpression* Parent, int32 OutputIndex)
-{
-	FStrataOperator& StrataOperator = Compiler->StrataCompilationRegisterOperator(STRATA_OPERATOR_THINFILM, Compiler->StrataTreeStackGetPathUniqueId(), Parent, Compiler->StrataTreeStackGetParentPathUniqueId());
-
-	UMaterialExpression* ChildAExpression = A.GetTracedInput().Expression;
-	FStrataOperator* OpA = nullptr;
-	if (ChildAExpression)
-	{
-		Compiler->StrataTreeStackPush(this, 0);
-		OpA = ChildAExpression->StrataGenerateMaterialTopologyTree(Compiler, this, A.OutputIndex);
-		Compiler->StrataTreeStackPop();
-		AssignOperatorIndexIfNotNull(StrataOperator.LeftIndex, OpA);
-	}
-	CombineFlagForParameterBlending(StrataOperator, OpA);
-
-	return &StrataOperator;
-}
-#endif // WITH_EDITOR
-
 UMaterialExpressionStrataUtilityBase::UMaterialExpressionStrataUtilityBase(const FObjectInitializer& ObjectInitializer)
 	: Super(ObjectInitializer)
 {
@@ -24911,6 +24808,100 @@ void UMaterialExpressionStrataHazinessToSecondaryRoughness::GetExpressionToolTip
 #endif // WITH_EDITOR
 
 
+
+UMaterialExpressionStrataThinFilm::UMaterialExpressionStrataThinFilm(const FObjectInitializer& ObjectInitializer)
+	: Super(ObjectInitializer)
+{
+	struct FConstructorStatics
+	{
+		FText NAME_Strata;
+		FConstructorStatics() : NAME_Strata(LOCTEXT("Strata Helpers", "Strata Helpers")) { }
+	};
+	static FConstructorStatics ConstructorStatics;
+#if WITH_EDITORONLY_DATA
+	MenuCategories.Add(ConstructorStatics.NAME_Strata);
+#endif
+
+	bShowOutputNameOnPin = true;
+
+	Outputs.Reset();
+	Outputs.Add(FExpressionOutput(TEXT("Specular Color")));
+	Outputs.Add(FExpressionOutput(TEXT("Edge Specular Color")));
+}
+
+#if WITH_EDITOR
+int32 UMaterialExpressionStrataThinFilm::Compile(class FMaterialCompiler* Compiler, int32 OutputIndex)
+{
+	int32 NormalCodeChunk		= CompileWithDefaultNormalWS(Compiler, Normal);
+
+	int32 F0CodeChunk			= F0.GetTracedInput().Expression		? F0.Compile(Compiler)			: Compiler->Constant3(0.04f, 0.04f, 0.04f);
+	int32 F90CodeChunk			= F90.GetTracedInput().Expression		? F90.Compile(Compiler)			: Compiler->Constant3(1.0f, 1.0f, 1.0f);
+
+	int32 ThicknessCodeChunk	= Thickness.GetTracedInput().Expression	? Thickness.Compile(Compiler)	: Compiler->Constant(1.0f);
+	int32 IORCodeChunk			= IOR.GetTracedInput().Expression		? IOR.Compile(Compiler)			: Compiler->Constant(1.44f);
+
+	return Compiler->StrataThinFilm(NormalCodeChunk, F0CodeChunk, F90CodeChunk, ThicknessCodeChunk, IORCodeChunk, OutputIndex);
+}
+
+void UMaterialExpressionStrataThinFilm::GetCaption(TArray<FString>& OutCaptions) const
+{
+	OutCaptions.Add(TEXT("Strata Thin-Film"));
+}
+
+uint32 UMaterialExpressionStrataThinFilm::GetOutputType(int32 OutputIndex)
+{
+	switch (OutputIndex)
+	{
+	case 0:
+		return MCT_Float3; // F0
+		break;
+	case 1:
+		return MCT_Float3; // F90
+		break;
+	}
+
+	check(false);
+	return MCT_Float1;
+}
+
+uint32 UMaterialExpressionStrataThinFilm::GetInputType(int32 InputIndex)
+{
+	if (InputIndex == 0) { return MCT_Float3; } // Normal
+	if (InputIndex == 1) { return MCT_Float3; } // F0
+	if (InputIndex == 2) { return MCT_Float3; } // F90
+	if (InputIndex == 3) { return MCT_Float1; } // Thickness
+	if (InputIndex == 4) { return MCT_Float1; } // IOR
+
+	check(false);
+	return MCT_Float1;
+}
+void UMaterialExpressionStrataThinFilm::GetConnectorToolTip(int32 InputIndex, int32 OutputIndex, TArray<FString>& OutToolTip)
+{
+	if (OutputIndex != INDEX_NONE)
+	{
+		switch (OutputIndex)
+		{
+		case 0:
+			ConvertToMultilineToolTip(TEXT("F0 accounting for thin film interferences. This is percentage of light reflected as specular from a surface when the view is perpendicular to the surface. (type = float3, unit = unitless, defaults to plastic 0.04)"), 80, OutToolTip);
+			break;
+		case 1:
+			ConvertToMultilineToolTip(TEXT("F90 accounting for thin film interferences. the percentage of light reflected as specular from a surface when the view is tangent to the surface. (type = float3, unit = unitless, defaults to 1.0f)."), 80, OutToolTip);
+			break;
+		}
+		return;
+	}
+
+	// Else use the default input tooltip
+	Super::GetConnectorToolTip(InputIndex, OutputIndex, OutToolTip);
+}
+
+void UMaterialExpressionStrataThinFilm::GetExpressionToolTip(TArray<FString>& OutToolTip)
+{
+	ConvertToMultilineToolTip(TEXT("Compute the resulting material specular parameter F0 and F90 according to input surface properties as well as the thin film parameters."), 80, OutToolTip);
+}
+#endif // WITH_EDITOR
+
+
 UMaterialExpressionExecBegin::UMaterialExpressionExecBegin(const FObjectInitializer& ObjectInitializer)
 	: Super(ObjectInitializer)
 {
diff --git a/Engine/Source/Runtime/Engine/Public/MaterialCompiler.h b/Engine/Source/Runtime/Engine/Public/MaterialCompiler.h
--- a/Engine/Source/Runtime/Engine/Public/MaterialCompiler.h
+++ b/Engine/Source/Runtime/Engine/Public/MaterialCompiler.h
@@ -674,11 +674,10 @@ public:
 	virtual int32 StrataAddParameterBlending(int32 A, int32 B, int32 AMixWeight, const FString& SharedLocalBasisIndexMacro, FStrataOperator* PromoteToOperator) = 0;
 	virtual int32 StrataWeight(int32 A, int32 Weight, int OperatorIndex, uint32 MaxDistanceFromLeaves) = 0;
 	virtual int32 StrataWeightParameterBlending(int32 A, int32 Weight, FStrataOperator* PromoteToOperator) = 0;
-	virtual int32 StrataThinFilm(int32 A, int32 Thickness, int32 IOR, int OperatorIndex, uint32 MaxDistanceFromLeaves) = 0;
-	virtual int32 StrataThinFilmParameterBlending(int32 A, int32 Thickness, int32 IOR, int32 BSDFNormalCodeChunk, FStrataOperator* PromoteToOperator) = 0;
 	virtual int32 StrataTransmittanceToMFP(int32 TransmittanceColor, int32 DesiredThickness, int32 OutputIndex) = 0;
 	virtual int32 StrataMetalnessToDiffuseAlbedoF0(int32 BaseColor, int32 Specular, int32 Metallic, int32 OutputIndex) = 0;
 	virtual int32 StrataHazinessToSecondaryRoughness(int32 BaseRoughness, int32 Haziness, int32 OutputIndex) = 0;
+	virtual int32 StrataThinFilm(int32 NormalCodeChunk, int32 SpecularColorCodeChunk, int32 EdgeSpecularColorCodeChunk, int32 ThicknessCodeChunk, int32 IORCodeChunk, int32 OutputIndex) = 0;
 	virtual int32 StrataCompilePreview(int32 StrataDataCodeChunk) = 0;
 
 	/**
@@ -1349,15 +1348,9 @@ public:
 	{
 		return Compiler->StrataWeight(A, Weight, OperatorIndex, MaxDistanceFromLeaves);
 	}
-
-	virtual int32 StrataThinFilm(int32 A, int32 Thickness, int32 IOR, int OperatorIndex, uint32 MaxDistanceFromLeaves) override
-	{
-		return Compiler->StrataThinFilm(A, Thickness, IOR, OperatorIndex, MaxDistanceFromLeaves);
-	}
-
-	virtual int32 StrataThinFilmParameterBlending(int32 A, int32 Thickness, int32 IOR, int32 BSDFNormalCodeChunk, FStrataOperator* PromoteToOperator) override
+	virtual int32 StrataThinFilm(int32 NormalCodeChunk, int32 SpecularColorCodeChunk, int32 EdgeSpecularColorCodeChunk, int32 ThicknessCodeChunk, int32 IORCodeChunk, int32 OutputIndex) override
 	{
-		return Compiler->StrataThinFilmParameterBlending(A, Thickness, IOR, BSDFNormalCodeChunk, PromoteToOperator);
+		return Compiler->StrataThinFilm(NormalCodeChunk, SpecularColorCodeChunk, EdgeSpecularColorCodeChunk, ThicknessCodeChunk, IORCodeChunk, OutputIndex);
 	}
 
 	virtual int32 StrataWeightParameterBlending(int32 A, int32 Weight, FStrataOperator* PromoteToOperator) override
