Fix AnimationSequence outer package not being marked as dirty, when applying modifications inside of a Controller Bracket

diff --git a/Engine/Source/Runtime/Engine/Classes/Animation/AnimData/AnimDataModelNotifyCollector.h b/Engine/Source/Runtime/Engine/Classes/Animation/AnimData/AnimDataModelNotifyCollector.h
--- a/Engine/Source/Runtime/Engine/Classes/Animation/AnimData/AnimDataModelNotifyCollector.h
+++ b/Engine/Source/Runtime/Engine/Classes/Animation/AnimData/AnimDataModelNotifyCollector.h
@@ -17,7 +17,7 @@ namespace Anim {
 between top-level EAnimDataModelNotifyType::BracketOpened and EAnimDataModelNotifyType::BracketClosed notifies */
 struct ENGINE_API FAnimDataModelNotifyCollector
 {
-	FAnimDataModelNotifyCollector() : BracketDepth(0) {}
+	FAnimDataModelNotifyCollector() : BracketDepth(0), bDataModified(false) {}
 
 	/** Handle a broadcasted notify, reset if we are opening a new top-level bracket*/
 	void Handle(EAnimDataModelNotifyType NotifyType)
@@ -25,6 +25,7 @@ struct ENGINE_API FAnimDataModelNotifyCollector
 		if (BracketDepth == 0)
 		{
 			Reset();
+			bDataModified = false;
 		}
 		
 		NotifyTypes.Add(NotifyType);
@@ -64,6 +65,9 @@ struct ENGINE_API FAnimDataModelNotifyCollector
 
 	/** Returns whether or not all brackets have been closed */
 	bool IsNotWithinBracket() const { return BracketDepth == 0; }
+
+	void MarkDataModified() { bDataModified = true; }
+	bool WasDataModified() { return bDataModified; }
 protected:
 	void Reset()
 	{
@@ -72,6 +76,7 @@ protected:
 protected:
 	TSet<EAnimDataModelNotifyType> NotifyTypes;
 	int32 BracketDepth;
+	bool bDataModified;
 };
 
 #endif // WITH_EDITOR
diff --git a/Engine/Source/Runtime/Engine/Private/Animation/AnimSequence.cpp b/Engine/Source/Runtime/Engine/Private/Animation/AnimSequence.cpp
--- a/Engine/Source/Runtime/Engine/Private/Animation/AnimSequence.cpp
+++ b/Engine/Source/Runtime/Engine/Private/Animation/AnimSequence.cpp
@@ -5386,8 +5386,7 @@ void UAnimSequence::OnModelModified(const EAnimDataModelNotifyType& NotifyType,
 		}
 	};
 
-	bool bShouldMarkPackageDirty = !FUObjectThreadContext::Get().IsRoutingPostLoad &&
-		NotifyType != EAnimDataModelNotifyType::BracketOpened && NotifyType != EAnimDataModelNotifyType::BracketClosed;
+	bool bShouldMarkPackageDirty = !FUObjectThreadContext::Get().IsRoutingPostLoad && NotifyType != EAnimDataModelNotifyType::BracketOpened;
 
 	switch (NotifyType)
 	{
@@ -5438,6 +5437,8 @@ void UAnimSequence::OnModelModified(const EAnimDataModelNotifyType& NotifyType,
 				const auto ResamplingNotifies = { EAnimDataModelNotifyType::TrackAdded, EAnimDataModelNotifyType::TrackChanged, EAnimDataModelNotifyType::TrackRemoved,  EAnimDataModelNotifyType::Populated };
 				const auto RecompressNotifies = { EAnimDataModelNotifyType::CurveAdded, EAnimDataModelNotifyType::CurveChanged, EAnimDataModelNotifyType::CurveRemoved, EAnimDataModelNotifyType::CurveFlagsChanged, EAnimDataModelNotifyType::CurveScaled,
 				EAnimDataModelNotifyType::AttributeAdded, EAnimDataModelNotifyType::AttributeChanged, EAnimDataModelNotifyType::AttributeRemoved };
+
+				bShouldMarkPackageDirty = NotifyCollector.WasDataModified();
 				
 				if (NotifyCollector.Contains(LengthChangingNotifies) || NotifyCollector.Contains(ResamplingNotifies))
 				{
@@ -5537,9 +5538,16 @@ void UAnimSequence::OnModelModified(const EAnimDataModelNotifyType& NotifyType,
 		}
 	}
 
-	if (NotifyCollector.IsNotWithinBracket() && NotifyType != EAnimDataModelNotifyType::BracketClosed && bShouldMarkPackageDirty)
+	if (NotifyCollector.IsNotWithinBracket())
 	{
-		MarkPackageDirty();
+		if (bShouldMarkPackageDirty)
+		{
+			MarkPackageDirty();
+		}
+	}
+	else if (bShouldMarkPackageDirty)
+	{
+		NotifyCollector.MarkDataModified();
 	}
 }
 
