Make Scalability Menu in Editor work per Preview Shader Platform

diff --git a/Engine/Source/Editor/UnrealEd/Classes/Editor/EditorEngine.h b/Engine/Source/Editor/UnrealEd/Classes/Editor/EditorEngine.h
--- a/Engine/Source/Editor/UnrealEd/Classes/Editor/EditorEngine.h
+++ b/Engine/Source/Editor/UnrealEd/Classes/Editor/EditorEngine.h
@@ -3205,6 +3205,9 @@ public:
 	/** Return whether the feature level preview enabled via Settings->Preview Rendering Level is currently active/displaying. */
 	bool IsFeatureLevelPreviewActive() const;
 
+	/** Return the current active shader platform in the editor*/
+	EShaderPlatform GetActiveShaderPlatform() const;
+
 	/**
 	 * Return the active feature level. This will be the chosen Settings->Preview Rendering Level if the Preview Mode button
 	 * is highlighted or SM5 if the Preview Mode button has been clicked and isn't highlighted.
diff --git a/Engine/Source/Editor/UnrealEd/Private/EditorEngine.cpp b/Engine/Source/Editor/UnrealEd/Private/EditorEngine.cpp
--- a/Engine/Source/Editor/UnrealEd/Private/EditorEngine.cpp
+++ b/Engine/Source/Editor/UnrealEd/Private/EditorEngine.cpp
@@ -7822,21 +7822,26 @@ void UEditorEngine::SetPreviewPlatform(const FPreviewPlatformInfo& NewPreviewPla
 		PreviewFeatureLevelChanged.Broadcast(EffectiveFeatureLevel);
 	}
 
-	Scalability::ChangeScalabilityPreviewPlatform(PreviewPlatform.GetEffectivePreviewPlatformName());
+	Scalability::ChangeScalabilityPreviewPlatform(PreviewPlatform.GetEffectivePreviewPlatformName(), GetActiveShaderPlatform());
 
 	UDeviceProfileManager::Get().RestorePreviewDeviceProfile();
 
 	UStaticMesh::OnLodStrippingQualityLevelChanged(nullptr);
 
-	//Override the current device profile.
-	if (PreviewPlatform.DeviceProfileName != NAME_None)
+	if (PreviewPlatform.bPreviewFeatureLevelActive)
 	{
-		if (UDeviceProfile* DP = UDeviceProfileManager::Get().FindProfile(PreviewPlatform.DeviceProfileName.ToString(), false))
+		//Override the current device profile.
+		if (PreviewPlatform.DeviceProfileName != NAME_None)
 		{
-			UDeviceProfileManager::Get().SetPreviewDeviceProfile(DP);
+			if (UDeviceProfile* DP = UDeviceProfileManager::Get().FindProfile(PreviewPlatform.DeviceProfileName.ToString(), false))
+			{
+				UDeviceProfileManager::Get().SetPreviewDeviceProfile(DP);
+			}
 		}
 	}
 
+	Scalability::ApplyCachedQualityLevelForShaderPlatform(GetActiveShaderPlatform());
+
 	PreviewPlatformChanged.Broadcast();
 
 	if (bSaveSettings)
@@ -7854,7 +7859,7 @@ void UEditorEngine::ToggleFeatureLevelPreview()
 	DefaultWorldFeatureLevel = NewPreviewFeatureLevel;
 	PreviewFeatureLevelChanged.Broadcast(NewPreviewFeatureLevel);
 
-	Scalability::ChangeScalabilityPreviewPlatform(PreviewPlatform.GetEffectivePreviewPlatformName());
+	Scalability::ChangeScalabilityPreviewPlatform(PreviewPlatform.GetEffectivePreviewPlatformName(), GetActiveShaderPlatform());
 
 	if (PreviewPlatform.bPreviewFeatureLevelActive)
 	{
@@ -7882,6 +7887,8 @@ void UEditorEngine::ToggleFeatureLevelPreview()
 		UDeviceProfileManager::Get().RestorePreviewDeviceProfile();
 	}
 
+	Scalability::ApplyCachedQualityLevelForShaderPlatform(GetActiveShaderPlatform());
+
 	PreviewPlatformChanged.Broadcast();
 
 	UStaticMesh::OnLodStrippingQualityLevelChanged(nullptr);
@@ -7901,6 +7908,17 @@ bool UEditorEngine::IsFeatureLevelPreviewActive() const
  	return PreviewPlatform.bPreviewFeatureLevelActive;
 }
 
+EShaderPlatform UEditorEngine::GetActiveShaderPlatform() const
+{
+	EShaderPlatform ActiveShaderPlatform = GShaderPlatformForFeatureLevel[GMaxRHIFeatureLevel];
+	if (PreviewPlatform.bPreviewFeatureLevelActive)
+	{
+		ActiveShaderPlatform = GShaderPlatformForFeatureLevel[PreviewPlatform.PreviewFeatureLevel];
+	}
+
+	return ActiveShaderPlatform;
+}
+
 ERHIFeatureLevel::Type UEditorEngine::GetActiveFeatureLevelPreviewType() const
 {
 	return PreviewPlatform.bPreviewFeatureLevelActive ? PreviewPlatform.PreviewFeatureLevel : GMaxRHIFeatureLevel;
diff --git a/Engine/Source/Runtime/Engine/Private/Scalability.cpp b/Engine/Source/Runtime/Engine/Private/Scalability.cpp
--- a/Engine/Source/Runtime/Engine/Private/Scalability.cpp
+++ b/Engine/Source/Runtime/Engine/Private/Scalability.cpp
@@ -337,10 +337,13 @@ FString GetScalabilitySectionString(const TCHAR* InGroupName, int32 InQualityLev
 
 #if WITH_EDITOR
 FName PlatformScalabilityName; // The name of the current platform scalability, or NAME_None if none is active
+EShaderPlatform ScalabilityShaderPlatform;
+bool bScalabilityShaderPlatformHasBeenChanged;
 FString PlatformScalabilityIniFilename;
 TMap<IConsoleVariable*, FString> PlatformScalabilityCVarBackup;
 TSet<const IConsoleVariable*> PlatformScalabilityCVarAllowList;
 TSet<const IConsoleVariable*> PlatformScalabilityCVarDenyList;
+TMap<EShaderPlatform, Scalability::FQualityLevels> PreviewShaderPlatformToCachedQualityLevels;
 
 void UndoPlatformScalability()
 {
@@ -376,9 +379,11 @@ void ApplyScalabilityGroupFromPlatformIni(const TCHAR* InSectionName, const TCHA
 	UE::ConfigUtilities::ForEachCVarInSectionFromIni(InSectionName, InIniFilename, Func);
 }
 
-void ChangeScalabilityPreviewPlatform(FName NewPlatformScalabilityName)
+void ChangeScalabilityPreviewPlatform(FName NewPlatformScalabilityName, EShaderPlatform ShaderPlatform)
 {
 	PlatformScalabilityName = NewPlatformScalabilityName;
+	bScalabilityShaderPlatformHasBeenChanged = ScalabilityShaderPlatform != ShaderPlatform;
+	ScalabilityShaderPlatform = ShaderPlatform;
 }
 #endif
 
@@ -418,6 +423,20 @@ static void SetGroupQualityLevel(const TCHAR* InGroupName, int32 InQualityLevel,
 	}
 }
 
+#if WITH_EDITOR
+void ApplyCachedQualityLevelForShaderPlatform(EShaderPlatform ShaderPlatform)
+{
+	if (ScalabilityShaderPlatform != EShaderPlatform::SP_NumPlatforms)
+	{
+		if (Scalability::FQualityLevels* CachedQualityLevel = PreviewShaderPlatformToCachedQualityLevels.Find(ScalabilityShaderPlatform))
+		{
+			Scalability::SetQualityLevels(*CachedQualityLevel);
+			Scalability::SaveState(GEditorSettingsIni);
+		}
+	}
+}
+#endif
+
 float GetResolutionScreenPercentage()
 {
 	static IConsoleVariable* CVar = IConsoleManager::Get().FindConsoleVariable(TEXT("r.ScreenPercentage"));
@@ -548,6 +567,10 @@ void InitScalabilitySystem()
 	CVarFoliageQuality.AsVariable()->SetOnChangedCallback(FConsoleVariableDelegate::CreateStatic(&OnChangeFoliageQuality));
 	CVarShadingQuality.AsVariable()->SetOnChangedCallback(FConsoleVariableDelegate::CreateStatic(&OnChangeShadingQuality));
 
+#if WITH_EDITOR
+	ScalabilityShaderPlatform = GMaxRHIShaderPlatform;
+	bScalabilityShaderPlatformHasBeenChanged = false;
+#endif
 	// Set defaults
 	SetQualityLevels(FQualityLevels(), true);
 	GScalabilityBackupQualityLevels = FQualityLevels();
@@ -782,6 +805,10 @@ void SetQualityLevels(const FQualityLevels& QualityLevels, bool bForce/* = false
 	ClampedLevels.SetFoliageQuality(QualityLevels.FoliageQuality);
 	ClampedLevels.SetShadingQuality(QualityLevels.ShadingQuality);
 
+#if WITH_EDITOR
+	bForce = bForce || bScalabilityShaderPlatformHasBeenChanged;
+#endif
+
 	if (GScalabilityUsingTemporaryQualityLevels && !bForce)
 	{
 		// When temporary scalability is active, non-temporary sets are
@@ -806,6 +833,14 @@ void SetQualityLevels(const FQualityLevels& QualityLevels, bool bForce/* = false
 
 		OnScalabilitySettingsChanged.Broadcast(ClampedLevels);
 	}
+
+#if WITH_EDITOR
+	if (ScalabilityShaderPlatform != EShaderPlatform::SP_NumPlatforms)
+	{
+		PreviewShaderPlatformToCachedQualityLevels.FindOrAdd(ScalabilityShaderPlatform) = ClampedLevels;
+	}
+	bScalabilityShaderPlatformHasBeenChanged = false;
+#endif
 }
 
 FQualityLevels GetQualityLevels()
diff --git a/Engine/Source/Runtime/Engine/Public/Scalability.h b/Engine/Source/Runtime/Engine/Public/Scalability.h
--- a/Engine/Source/Runtime/Engine/Public/Scalability.h
+++ b/Engine/Source/Runtime/Engine/Public/Scalability.h
@@ -155,6 +155,10 @@ namespace Scalability
 	/** This is the only suggested way to set the current state - don't set CVars directly **/
 	ENGINE_API void SetQualityLevels(const FQualityLevels& QualityLevels, bool bForce = false);
 
+#if WITH_EDITOR
+	ENGINE_API void ApplyCachedQualityLevelForShaderPlatform(EShaderPlatform ShaderPlatform);
+#endif
+
 	/** This is the only suggested way to get the current state - don't get CVars directly */
 	ENGINE_API FQualityLevels GetQualityLevels();
 
@@ -210,7 +214,7 @@ namespace Scalability
 
 #if WITH_EDITOR
 	/** Set an Editor preview scalability platform */
-	void ENGINE_API ChangeScalabilityPreviewPlatform(FName NewPlatformScalabilityName);
+	void ENGINE_API ChangeScalabilityPreviewPlatform(FName NewPlatformScalabilityName, EShaderPlatform ShaderPlatform);
 #endif
 
 	ENGINE_API FText GetQualityLevelText(int32 Value, int32 NumLevels);
diff --git a/Engine/Source/Runtime/Launch/Private/LaunchEngineLoop.cpp b/Engine/Source/Runtime/Launch/Private/LaunchEngineLoop.cpp
--- a/Engine/Source/Runtime/Launch/Private/LaunchEngineLoop.cpp
+++ b/Engine/Source/Runtime/Launch/Private/LaunchEngineLoop.cpp
@@ -6365,7 +6365,7 @@ bool FEngineLoop::AppInit( )
 		FPIEPreviewDeviceModule* PIEPreviewDeviceProfileSelectorModule = FModuleManager::LoadModulePtr<FPIEPreviewDeviceModule>("PIEPreviewDeviceProfileSelector");
 		if (PIEPreviewDeviceProfileSelectorModule)
 		{
-			Scalability::ChangeScalabilityPreviewPlatform(PIEPreviewDeviceProfileSelectorModule->GetPreviewPlatformName());
+			Scalability::ChangeScalabilityPreviewPlatform(PIEPreviewDeviceProfileSelectorModule->GetPreviewPlatformName(), GEditor->GetActiveShaderPlatform());
 		}
 	}
 #endif
