Fixed FStreamLevelAction::ActivateLevel not updating client's ULevelStreaming::bShouldBlockOnUnload.
#9743

diff --git a/Engine/Source/Runtime/Engine/Classes/GameFramework/PlayerController.h b/Engine/Source/Runtime/Engine/Classes/GameFramework/PlayerController.h
--- a/Engine/Source/Runtime/Engine/Classes/GameFramework/PlayerController.h
+++ b/Engine/Source/Runtime/Engine/Classes/GameFramework/PlayerController.h
@@ -165,6 +165,10 @@ struct ENGINE_API FUpdateLevelStreamingLevelStatus
 	/** Whether we want to force a blocking load */
 	UPROPERTY()
 	bool bNewShouldBlockOnLoad = false;
+
+	/** Whether we want to force a blocking unload */
+	UPROPERTY()
+	bool bNewShouldBlockOnUnload = false;
 };
 
 /** Data structure used to setup an input mode that allows the UI to respond to user input, and if the UI doesn't handle it player input / player controller gets a chance. */
@@ -835,6 +839,7 @@ public:
 
 	/** Notify player of change to level */
 	void LevelStreamingStatusChanged(class ULevelStreaming* LevelObject, bool bNewShouldBeLoaded, bool bNewShouldBeVisible, bool bNewShouldBlockOnLoad, int32 LODIndex);
+	void LevelStreamingStatusChanged(class ULevelStreaming* LevelObject, bool bNewShouldBeLoaded, bool bNewShouldBeVisible, bool bNewShouldBlockOnLoad, bool bNewShouldBlockOnUnload, int32 LODIndex);
 
 	/** Used to wait until a map change can be prepared when one was already in progress */
 	virtual void DelayedPrepareMapChange();
@@ -1408,10 +1413,10 @@ public:
 	 * @param bNewShouldBlockOnLoad - Whether we want to force a blocking load
 	 * @param LODIndex				- Current LOD index for a streaming level
 	 * @param TransactionId			- Optional parameter used when communicating LevelVisibility changes between server and client
+	 * @param bNewShouldBlockOnUnload - Optional parameter used to force a blocking unload or not
 	 */
 	UFUNCTION(Reliable, Client)
-	void ClientUpdateLevelStreamingStatus(FName PackageName, bool bNewShouldBeLoaded, bool bNewShouldBeVisible, bool bNewShouldBlockOnLoad, int32 LODIndex, FNetLevelVisibilityTransactionId TransactionId);
-	void ClientUpdateLevelStreamingStatus(FName PackageName, bool bNewShouldBeLoaded, bool bNewShouldBeVisible, bool bNewShouldBlockOnLoad, int32 LODIndex) { ClientUpdateLevelStreamingStatus(PackageName, bNewShouldBeLoaded, bNewShouldBeVisible, bNewShouldBlockOnLoad, LODIndex, FNetLevelVisibilityTransactionId()); }
+	void ClientUpdateLevelStreamingStatus(FName PackageName, bool bNewShouldBeLoaded, bool bNewShouldBeVisible, bool bNewShouldBlockOnLoad, int32 LODIndex, FNetLevelVisibilityTransactionId TransactionId = FNetLevelVisibilityTransactionId(), bool bNewShouldBlockOnUnload = false);
 
 	/**
 	 * Replicated Update streaming status.  This version allows for the streaming state of many levels to be sent in a single RPC.
diff --git a/Engine/Source/Runtime/Engine/Private/CheatManager.cpp b/Engine/Source/Runtime/Engine/Private/CheatManager.cpp
--- a/Engine/Source/Runtime/Engine/Private/CheatManager.cpp
+++ b/Engine/Source/Runtime/Engine/Private/CheatManager.cpp
@@ -605,6 +605,7 @@ void UCheatManager::SetLevelStreamingStatus(FName PackageName, bool bShouldBeLoa
 					LevelStatus.bNewShouldBeLoaded = bShouldBeLoaded;
 					LevelStatus.bNewShouldBeVisible = bShouldBeVisible;
 					LevelStatus.bNewShouldBlockOnLoad = false;
+					LevelStatus.bNewShouldBlockOnUnload = false;
 					LevelStatus.LODIndex = INDEX_NONE;
 				}
 				PC->ClientUpdateMultipleLevelsStreamingStatus(LevelStatuses);
diff --git a/Engine/Source/Runtime/Engine/Private/GameModeBase.cpp b/Engine/Source/Runtime/Engine/Private/GameModeBase.cpp
--- a/Engine/Source/Runtime/Engine/Private/GameModeBase.cpp
+++ b/Engine/Source/Runtime/Engine/Private/GameModeBase.cpp
@@ -941,6 +941,7 @@ void AGameModeBase::ReplicateStreamingStatus(APlayerController* PC)
 					LevelStatus.bNewShouldBeLoaded = bTheLevelShouldBeLoaded;
 					LevelStatus.bNewShouldBeVisible = bTheLevelShouldBeVisible;
 					LevelStatus.bNewShouldBlockOnLoad = TheLevel->bShouldBlockOnLoad;
+					LevelStatus.bNewShouldBlockOnUnload = TheLevel->bShouldBlockOnUnload;
 					LevelStatus.LODIndex = TheLevel->GetLevelLODIndex();
 				}
 			}
diff --git a/Engine/Source/Runtime/Engine/Private/LevelStreaming.cpp b/Engine/Source/Runtime/Engine/Private/LevelStreaming.cpp
--- a/Engine/Source/Runtime/Engine/Private/LevelStreaming.cpp
+++ b/Engine/Source/Runtime/Engine/Private/LevelStreaming.cpp
@@ -322,9 +322,10 @@ void FStreamLevelAction::ActivateLevel( ULevelStreaming* LevelStreamingObject )
 					PlayerController->LevelStreamingStatusChanged(
 						LevelStreamingObject,
 						bShouldBeLoaded,
-					bShouldBeVisible,
+						bShouldBeVisible,
 						bShouldBlock,
-					INDEX_NONE);
+						bShouldBlock,
+						INDEX_NONE);
 				}
 			}
 		}
diff --git a/Engine/Source/Runtime/Engine/Private/LevelTick.cpp b/Engine/Source/Runtime/Engine/Private/LevelTick.cpp
--- a/Engine/Source/Runtime/Engine/Private/LevelTick.cpp
+++ b/Engine/Source/Runtime/Engine/Private/LevelTick.cpp
@@ -774,6 +774,7 @@ void UWorld::ProcessLevelStreamingVolumes(FVector* OverrideViewLocation)
 									bNewShouldBeLoaded, 
 									bNewShouldBeVisible,
 									LevelStreamingObject->bShouldBlockOnLoad,
+									LevelStreamingObject->bShouldBlockOnUnload,
 									LevelStreamingObject->GetLevelLODIndex());
 						}
 					}
diff --git a/Engine/Source/Runtime/Engine/Private/PlayerController.cpp b/Engine/Source/Runtime/Engine/Private/PlayerController.cpp
--- a/Engine/Source/Runtime/Engine/Private/PlayerController.cpp
+++ b/Engine/Source/Runtime/Engine/Private/PlayerController.cpp
@@ -310,7 +310,7 @@ FName APlayerController::NetworkRemapPath(FName InPackageName, bool bReading)
 
 /// @cond DOXYGEN_WARNINGS
 
-void APlayerController::ClientUpdateLevelStreamingStatus_Implementation(FName PackageName, bool bNewShouldBeLoaded, bool bNewShouldBeVisible, bool bNewShouldBlockOnLoad, int32 LODIndex, FNetLevelVisibilityTransactionId TransactionId)
+void APlayerController::ClientUpdateLevelStreamingStatus_Implementation(FName PackageName, bool bNewShouldBeLoaded, bool bNewShouldBeVisible, bool bNewShouldBlockOnLoad, int32 LODIndex, FNetLevelVisibilityTransactionId TransactionId, bool bNewShouldBlockOnUnload)
 {
 	PackageName = NetworkRemapPath(PackageName, true);
 	
@@ -350,6 +350,7 @@ void APlayerController::ClientUpdateLevelStreamingStatus_Implementation(FName Pa
 		LevelStreamingObject->SetShouldBeLoaded(bNewShouldBeLoaded);
 		LevelStreamingObject->SetShouldBeVisible(bNewShouldBeVisible);
 		LevelStreamingObject->bShouldBlockOnLoad = bNewShouldBlockOnLoad;
+		LevelStreamingObject->bShouldBlockOnUnload = bNewShouldBlockOnUnload;
 		LevelStreamingObject->SetLevelLODIndex(LODIndex);
 		LevelStreamingObject->UpdateNetVisibilityTransactionState(bNewShouldBeVisible, TransactionId);
 	}
@@ -363,7 +364,7 @@ void APlayerController::ClientUpdateMultipleLevelsStreamingStatus_Implementation
 {
 	for( const FUpdateLevelStreamingLevelStatus& LevelStatus : LevelStatuses )
 	{
-		ClientUpdateLevelStreamingStatus_Implementation(LevelStatus.PackageName, LevelStatus.bNewShouldBeLoaded, LevelStatus.bNewShouldBeVisible, LevelStatus.bNewShouldBlockOnLoad, LevelStatus.LODIndex, FNetLevelVisibilityTransactionId());
+		ClientUpdateLevelStreamingStatus_Implementation(LevelStatus.PackageName, LevelStatus.bNewShouldBeLoaded, LevelStatus.bNewShouldBeVisible, LevelStatus.bNewShouldBlockOnLoad, LevelStatus.LODIndex, FNetLevelVisibilityTransactionId(), LevelStatus.bNewShouldBlockOnUnload);
 	}
 }
 
@@ -3460,7 +3461,13 @@ void APlayerController::ClientForceGarbageCollection_Implementation()
 
 /// @endcond
 
-void APlayerController::LevelStreamingStatusChanged(ULevelStreaming* LevelObject, bool bNewShouldBeLoaded, bool bNewShouldBeVisible, bool bNewShouldBlockOnLoad, int32 LODIndex )
+void APlayerController::LevelStreamingStatusChanged(ULevelStreaming* LevelObject, bool bNewShouldBeLoaded, bool bNewShouldBeVisible, bool bNewShouldBlockOnLoad, int32 LODIndex)
+{ 
+	const bool bNewShouldBlockOnUnload = false;
+	LevelStreamingStatusChanged(LevelObject, bNewShouldBeLoaded, bNewShouldBeVisible, bNewShouldBlockOnLoad, bNewShouldBlockOnUnload, LODIndex);
+}
+
+void APlayerController::LevelStreamingStatusChanged(ULevelStreaming* LevelObject, bool bNewShouldBeLoaded, bool bNewShouldBeVisible, bool bNewShouldBlockOnLoad, bool bNewShouldBlockOnUnload, int32 LODIndex)
 {
 	FNetLevelVisibilityTransactionId TransactionId;
 	if (GetNetMode() == NM_Client)
@@ -3475,7 +3482,7 @@ void APlayerController::LevelStreamingStatusChanged(ULevelStreaming* LevelObject
 		TransactionId = NetConnection->UpdateLevelStreamStatusChangedTransactionId(LevelObject, PackageName, bNewShouldBeVisible);
 	}
 
-	ClientUpdateLevelStreamingStatus(NetworkRemapPath(LevelObject->GetWorldAssetPackageFName(), false), bNewShouldBeLoaded, bNewShouldBeVisible, bNewShouldBlockOnLoad, LODIndex, TransactionId);
+	ClientUpdateLevelStreamingStatus(NetworkRemapPath(LevelObject->GetWorldAssetPackageFName(), false), bNewShouldBeLoaded, bNewShouldBeVisible, bNewShouldBlockOnLoad, LODIndex, TransactionId, bNewShouldBlockOnUnload);
 }
 
 /// @cond DOXYGEN_WARNINGS
diff --git a/Engine/Source/Runtime/Engine/Private/UnrealEngine.cpp b/Engine/Source/Runtime/Engine/Private/UnrealEngine.cpp
--- a/Engine/Source/Runtime/Engine/Private/UnrealEngine.cpp
+++ b/Engine/Source/Runtime/Engine/Private/UnrealEngine.cpp
@@ -16327,6 +16327,7 @@ bool UEngine::CommitMapChange( FWorldContext &Context )
 									bShouldBeLoaded,
 									bShouldBeVisible,
 									StreamingLevel->bShouldBlockOnLoad,
+									StreamingLevel->bShouldBlockOnUnload,
 									StreamingLevel->GetLevelLODIndex());
 							}
 						}
