Fix Non-Unity build

diff --git a/Engine/Source/Runtime/Engine/Public/Scalability.h b/Engine/Source/Runtime/Engine/Public/Scalability.h
--- a/Engine/Source/Runtime/Engine/Public/Scalability.h
+++ b/Engine/Source/Runtime/Engine/Public/Scalability.h
@@ -7,6 +7,9 @@
 #pragma once
 
 #include "CoreMinimal.h"
+#if WITH_EDITOR
+#include "RHIDefinitions.h"
+#endif
 
 namespace Scalability
 { 
