Optimus: Add per platform support setting so that we can directly handle per platform shader compilation etc.
Picked apart the way that Passthrough Vertex Factory for mesh deformer relied on putting mesh deformer condition inside skin cache condition. Now the mesh deformer doesn't need to force the skin cache to appear enabled to external systems to make things work.

diff --git a/Engine/Plugins/Animation/DeformerGraph/Source/OptimusCore/OptimusCore.Build.cs b/Engine/Plugins/Animation/DeformerGraph/Source/OptimusCore/OptimusCore.Build.cs
--- a/Engine/Plugins/Animation/DeformerGraph/Source/OptimusCore/OptimusCore.Build.cs
+++ b/Engine/Plugins/Animation/DeformerGraph/Source/OptimusCore/OptimusCore.Build.cs
@@ -27,6 +27,7 @@ namespace UnrealBuildTool.Rules
 					"Core",
 					"CoreUObject",
 					"Engine",
+					"OptimusSettings",
 					"Projects",
 					"RenderCore",
 					"Renderer",
diff --git a/Engine/Plugins/Animation/DeformerGraph/Source/OptimusCore/Private/OptimusDeformer.cpp b/Engine/Plugins/Animation/DeformerGraph/Source/OptimusCore/Private/OptimusDeformer.cpp
--- a/Engine/Plugins/Animation/DeformerGraph/Source/OptimusCore/Private/OptimusDeformer.cpp
+++ b/Engine/Plugins/Animation/DeformerGraph/Source/OptimusCore/Private/OptimusDeformer.cpp
@@ -7,7 +7,6 @@
 #include "Actions/OptimusResourceActions.h"
 #include "Actions/OptimusVariableActions.h"
 #include "Components/MeshComponent.h"
-#include "ComputeFramework/ComputeFramework.h"
 #include "ComputeFramework/ComputeKernel.h"
 #include "Containers/Queue.h"
 #include "DataInterfaces/OptimusDataInterfaceGraph.h"
@@ -29,6 +28,7 @@
 #include "OptimusNodePin.h"
 #include "OptimusObjectVersion.h"
 #include "OptimusResourceDescription.h"
+#include "OptimusSettings.h"
 #include "OptimusVariableDescription.h"
 #include "RenderingThread.h"
 #include "SceneInterface.h"
@@ -2306,9 +2306,9 @@ UMeshDeformerInstance* UOptimusDeformer::CreateInstance(
 		return nullptr;
 	}
 
-	// Return nullptr if ComputeFramework is disabled. Clients can then fallback to some other behaviour.
+	// Return nullptr if deformers are disabled. Clients can then fallback to some other behaviour.
 	EShaderPlatform Platform = InMeshComponent->GetScene() != nullptr ? InMeshComponent->GetScene()->GetShaderPlatform() : GMaxRHIShaderPlatform;
-	if (!ComputeFramework::IsEnabled() || !ComputeFramework::IsSupported(Platform))
+	if (!Optimus::IsEnabled() || !Optimus::IsSupported(Platform))
 	{
 		return nullptr;
 	}
diff --git a/Engine/Plugins/Animation/DeformerGraph/Source/OptimusSettings/OptimusSettings.Build.cs b/Engine/Plugins/Animation/DeformerGraph/Source/OptimusSettings/OptimusSettings.Build.cs
--- a/Engine/Plugins/Animation/DeformerGraph/Source/OptimusSettings/OptimusSettings.Build.cs
+++ b/Engine/Plugins/Animation/DeformerGraph/Source/OptimusSettings/OptimusSettings.Build.cs
@@ -21,10 +21,12 @@ namespace UnrealBuildTool.Rules
             PrivateDependencyModuleNames.AddRange(
                 new string[]
                 {
+					"ComputeFramework",
 					"Core",
 					"CoreUObject",
 					"DeveloperSettings",
 					"Engine",
+					"RenderCore",
 				}
 			);
         }
diff --git a/Engine/Plugins/Animation/DeformerGraph/Source/OptimusSettings/Private/OptimusSettingsModule.cpp b/Engine/Plugins/Animation/DeformerGraph/Source/OptimusSettings/Private/OptimusSettingsModule.cpp
--- a/Engine/Plugins/Animation/DeformerGraph/Source/OptimusSettings/Private/OptimusSettingsModule.cpp
+++ b/Engine/Plugins/Animation/DeformerGraph/Source/OptimusSettings/Private/OptimusSettingsModule.cpp
@@ -3,10 +3,34 @@
 #include "OptimusSettingsModule.h"
 
 #include "Animation/MeshDeformer.h"
+#include "ComputeFramework/ComputeFramework.h"
 #include "Features/IModularFeatures.h"
 #include "Misc/CoreDelegates.h"
 #include "Modules/ModuleManager.h"
 #include "OptimusSettings.h"
+#include "RenderUtils.h"
+
+
+namespace Optimus
+{
+	/** ReadOnly CVar intended for enabling/disabling of DeformerGraph support per platform. */
+	static TAutoConsoleVariable<int32> CVarDeformerGraphEnable(
+		TEXT("r.DeformerGraph.Enable"),
+		1,
+		TEXT("Set to 0 to disable DeformerGraph support on a platform.\n"),
+		ECVF_ReadOnly);
+
+	bool IsSupported(EShaderPlatform Platform)
+	{
+		static FShaderPlatformCachedIniValue<int32> PerPlatformCVar(TEXT("r.DeformerGraph.Enable"));
+		return PerPlatformCVar.Get(Platform) != 0 && ComputeFramework::IsSupported(Platform);
+	}
+
+	bool IsEnabled()
+	{
+		return ComputeFramework::IsEnabled();
+	}
+}
 
 
 UOptimusSettings::UOptimusSettings(const FObjectInitializer& ObjectInitlaizer)
@@ -14,6 +38,11 @@ UOptimusSettings::UOptimusSettings(const FObjectInitializer& ObjectInitlaizer)
 {
 }
 
+bool FOptimusSettingsModule::IsSupported(EShaderPlatform Platform) const
+{
+	return Optimus::IsSupported(Platform);
+}
+
 #if WITH_EDITOR
 
 UOptimusSettings::FOnUpdateSettings UOptimusSettings::OnSettingsChange;
diff --git a/Engine/Plugins/Animation/DeformerGraph/Source/OptimusSettings/Private/OptimusSettingsModule.h b/Engine/Plugins/Animation/DeformerGraph/Source/OptimusSettings/Private/OptimusSettingsModule.h
--- a/Engine/Plugins/Animation/DeformerGraph/Source/OptimusSettings/Private/OptimusSettingsModule.h
+++ b/Engine/Plugins/Animation/DeformerGraph/Source/OptimusSettings/Private/OptimusSettingsModule.h
@@ -14,6 +14,7 @@ class FOptimusSettingsModule : public IModuleInterface, public IMeshDeformerProv
 	void ShutdownModule() override;
 
 	/** IMeshDeformerProvider implementation */
+	bool IsSupported(EShaderPlatform Platform) const override;
 	TObjectPtr<UMeshDeformer> GetDefaultMeshDeformer(FDefaultMeshDeformerSetup const& InSetup) override;
 
 private:
diff --git a/Engine/Plugins/Animation/DeformerGraph/Source/OptimusSettings/Public/OptimusSettings.h b/Engine/Plugins/Animation/DeformerGraph/Source/OptimusSettings/Public/OptimusSettings.h
--- a/Engine/Plugins/Animation/DeformerGraph/Source/OptimusSettings/Public/OptimusSettings.h
+++ b/Engine/Plugins/Animation/DeformerGraph/Source/OptimusSettings/Public/OptimusSettings.h
@@ -43,3 +43,12 @@ class OPTIMUSSETTINGS_API UOptimusSettings : public UDeveloperSettings
 	void PostEditChangeProperty(struct FPropertyChangedEvent& PropertyChangedEvent) override;
 #endif // WITH_EDITOR
 };
+
+namespace Optimus
+{
+	/** Returns true if DeformerGraph is supported on a platform. */
+	OPTIMUSSETTINGS_API bool IsSupported(EShaderPlatform Platform);
+
+	/** Returns true if DeformerGraph is currently enabled. */
+	OPTIMUSSETTINGS_API bool IsEnabled();
+}
diff --git a/Engine/Source/Runtime/Engine/Private/GPUSkinCache.cpp b/Engine/Source/Runtime/Engine/Private/GPUSkinCache.cpp
--- a/Engine/Source/Runtime/Engine/Private/GPUSkinCache.cpp
+++ b/Engine/Source/Runtime/Engine/Private/GPUSkinCache.cpp
@@ -24,7 +24,6 @@ GPUSkinCache.cpp: Performs skinning on a compute shader into a buffer to avoid v
 #include "RayTracingSkinnedGeometry.h"
 #include "GPUSkinCacheVisualizationData.h"
 #include "Internationalization/Internationalization.h"
-#include "Animation/MeshDeformerProvider.h"
 
 DEFINE_STAT(STAT_GPUSkinCache_TotalNumChunks);
 DEFINE_STAT(STAT_GPUSkinCache_TotalNumVertices);
@@ -207,11 +206,7 @@ static inline bool IsGPUSkinCacheAllowed(EShaderPlatform Platform)
 
 ENGINE_API bool IsGPUSkinCacheAvailable(EShaderPlatform Platform)
 {
-	// Enable skin cache shaders if there is a mesh deformer provider.
-	// Store in static because it needs to be consistent and available on all threads.
-	static bool bMeshDeformersAvailable = IMeshDeformerProvider::IsAvailable();
-
-	return (bMeshDeformersAvailable || AreSkinCacheShadersEnabled(Platform) != 0) && IsGPUSkinCacheAllowed(Platform);
+	return AreSkinCacheShadersEnabled(Platform) != 0 && IsGPUSkinCacheAllowed(Platform);
 }
 
 static inline bool IsGPUSkinCacheEnable(EShaderPlatform Platform)
diff --git a/Engine/Source/Runtime/Engine/Private/LocalVertexFactory.cpp b/Engine/Source/Runtime/Engine/Private/LocalVertexFactory.cpp
--- a/Engine/Source/Runtime/Engine/Private/LocalVertexFactory.cpp
+++ b/Engine/Source/Runtime/Engine/Private/LocalVertexFactory.cpp
@@ -14,6 +14,7 @@
 #include "ProfilingDebugging/LoadTimeTracker.h"
 #include "GPUSkinCache.h"
 #include "GPUSkinVertexFactory.h"
+#include "Animation/MeshDeformerProvider.h"
 
 IMPLEMENT_TYPE_LAYOUT(FLocalVertexFactoryShaderParametersBase);
 IMPLEMENT_TYPE_LAYOUT(FLocalVertexFactoryShaderParameters);
@@ -214,6 +215,14 @@ void FLocalVertexFactoryShaderParameters::GetElementShaderBindings(
 	ShaderBindings.Add(Shader->GetUniformBufferParameter<FLocalVertexFactoryLooseParameters>(), LocalVertexFactory->LooseParametersUniformBuffer);
 }
 
+bool IsGPUSkinPassThroughSupported(EShaderPlatform Platform)
+{
+	// Enable the GPUSkin passthrough path if we might use the GPUSkinCache or MeshDeformers.
+	static IMeshDeformerProvider* MeshDeformerProvider = IMeshDeformerProvider::Get();
+	bool bMeshDeformersAvailable = MeshDeformerProvider && MeshDeformerProvider->IsSupported(Platform);
+	return bMeshDeformersAvailable || IsGPUSkinCacheAvailable(Platform);
+}
+
 void FLocalVertexFactoryShaderParameters::GetElementShaderBindingsGPUSkinPassThrough(
 	const FSceneInterface* Scene,
 	const FSceneView* View,
@@ -311,7 +320,7 @@ void FLocalVertexFactory::ModifyCompilationEnvironment(const FVertexFactoryShade
 
 	if (Parameters.VertexFactoryType->SupportsGPUSkinPassThrough())
 	{
-		OutEnvironment.SetDefine(TEXT("SUPPORT_GPUSKIN_PASSTHROUGH"), IsGPUSkinCacheAvailable(Parameters.Platform));
+		OutEnvironment.SetDefine(TEXT("SUPPORT_GPUSKIN_PASSTHROUGH"), IsGPUSkinPassThroughSupported(Parameters.Platform));
 	}
 }
 
diff --git a/Engine/Source/Runtime/Engine/Public/Animation/MeshDeformerProvider.h b/Engine/Source/Runtime/Engine/Public/Animation/MeshDeformerProvider.h
--- a/Engine/Source/Runtime/Engine/Public/Animation/MeshDeformerProvider.h
+++ b/Engine/Source/Runtime/Engine/Public/Animation/MeshDeformerProvider.h
@@ -22,6 +22,9 @@ public:
 	static bool IsAvailable();
 	static IMeshDeformerProvider* Get();
 
+	/** Returns true if the platform is supported. */
+	virtual bool IsSupported(EShaderPlatform Platform) const = 0;
+
 	/** Structure for passing to GetDefaultMeshDeformer(). */
 	struct FDefaultMeshDeformerSetup
 	{
