Fixed missing PIEInstanceID on external packages when running PIE.

diff --git a/Engine/Source/Runtime/CoreUObject/Private/UObject/LinkerLoad.cpp b/Engine/Source/Runtime/CoreUObject/Private/UObject/LinkerLoad.cpp
--- a/Engine/Source/Runtime/CoreUObject/Private/UObject/LinkerLoad.cpp
+++ b/Engine/Source/Runtime/CoreUObject/Private/UObject/LinkerLoad.cpp
@@ -5158,10 +5158,11 @@ UObject* FLinkerLoad::CreateExport( int32 Index )
 		// if our outer is actually an import, then the package we are an export of is not in our outer chain, set our package in that case
 		Params.ExternalPackage = Export.OuterIndex.IsImport() ? LinkerRoot : nullptr;
 
-		// Propagate relevant flags from the outer package to the external package
+		// Propagate relevant properties from the outer package to the external package
 		if (Params.ExternalPackage)
 		{
 			Params.ExternalPackage->SetPackageFlags(ThisParent->GetPackage()->GetPackageFlags() & PKG_PlayInEditor);
+			Params.ExternalPackage->SetPIEInstanceID(ThisParent->GetPackage()->GetPIEInstanceID());
 		}
 
 		{
diff --git a/Engine/Source/Runtime/Engine/Private/PieFixupSerializer.cpp b/Engine/Source/Runtime/Engine/Private/PieFixupSerializer.cpp
--- a/Engine/Source/Runtime/Engine/Private/PieFixupSerializer.cpp
+++ b/Engine/Source/Runtime/Engine/Private/PieFixupSerializer.cpp
@@ -3,6 +3,7 @@
 #include "PieFixupSerializer.h"
 
 #include "Components/InstancedStaticMeshComponent.h"
+#include "UObject/Package.h"
 
 class FMulticastDelegateProperty;
 
@@ -33,6 +34,15 @@ FArchive& FPIEFixupSerializer::operator<<(UObject*& Object)
 	{
 		VisitedObjects.Add(Object);
 
+#if WITH_EDITOR
+		if (UPackage* ExternalPackage = Object->GetExternalPackage())
+		{
+			check(Object->IsPackageExternal());
+			check(ExternalPackage->HasAnyPackageFlags(PKG_PlayInEditor));
+			ExternalPackage->SetPIEInstanceID(PIEInstanceID);
+		}
+#endif
+
 		// Skip instanced static mesh component as their impact on serialization is enormous and they don't contain lazy ptrs.
 		if (!Cast<UInstancedStaticMeshComponent>(Object))
 		{
