Fix NonUnity again

diff --git a/Engine/Source/Runtime/Engine/Private/Scalability.cpp b/Engine/Source/Runtime/Engine/Private/Scalability.cpp
--- a/Engine/Source/Runtime/Engine/Private/Scalability.cpp
+++ b/Engine/Source/Runtime/Engine/Private/Scalability.cpp
@@ -11,6 +11,10 @@
 #include "Interfaces/IAnalyticsProvider.h"
 #include "Interfaces/IProjectManager.h"
 
+#if WITH_EDITOR
+#include "RHI.h"
+#endif
+
 Scalability::FOnScalabilitySettingsChanged Scalability::OnScalabilitySettingsChanged;
 
 #if !UE_BUILD_SHIPPING
