Fix blend profile crash when duplicate indices are introduced due to a virtual bone rename operation.
Fix blend profile bone indices going out of sync when virtual bones are deleted.

diff --git a/Engine/Source/Runtime/Engine/Private/Animation/Skeleton.cpp b/Engine/Source/Runtime/Engine/Private/Animation/Skeleton.cpp
--- a/Engine/Source/Runtime/Engine/Private/Animation/Skeleton.cpp
+++ b/Engine/Source/Runtime/Engine/Private/Animation/Skeleton.cpp
@@ -1973,6 +1973,17 @@ void USkeleton::RemoveVirtualBones(const TArray<FName>& BonesToRemove)
 
 	RegenerateVirtualBoneGuid();
 	HandleVirtualBoneChanges();
+
+	// Blend profiles cache bone names and indices, make sure they remain in sync when the indices change
+	for (UBlendProfile* Profile : BlendProfiles)
+	{
+		// TEMPORARY FIX FOR 5.1.1
+
+		for (FBlendProfileBoneEntry& Entry : Profile->ProfileEntries)
+		{
+			Entry.BoneReference.Initialize(Profile->OwningSkeleton);
+		}
+	}
 }
 
 void USkeleton::RenameVirtualBone(const FName OriginalBoneName, const FName NewBoneName)
@@ -2007,6 +2018,27 @@ void USkeleton::RenameVirtualBone(const FName OriginalBoneName, const FName NewB
 	{
 		RegenerateVirtualBoneGuid();
 		HandleVirtualBoneChanges();
+
+		// @todo: This might be a slow operation if there's a large amount of blend profiles and entries
+		int32 BoneIdx = GetReferenceSkeleton().FindBoneIndex(NewBoneName);
+		if (BoneIdx != INDEX_NONE)
+		{
+			for (UBlendProfile* Profile : BlendProfiles)
+			{
+				// TEMPORARY FIX FOR 5.1.1
+
+				FBlendProfileBoneEntry* FoundEntry = Profile->ProfileEntries.FindByPredicate([BoneIdx](const FBlendProfileBoneEntry& Entry)
+					{
+						return Entry.BoneReference.BoneIndex == BoneIdx;
+					});
+
+				if (FoundEntry)
+				{
+					FoundEntry->BoneReference.BoneName = Profile->OwningSkeleton->GetReferenceSkeleton().GetBoneName(BoneIdx);
+					FoundEntry->BoneReference.Initialize(Profile->OwningSkeleton);
+				}
+			}
+		}
 	}
 }
 
