Strata - Added better positioning of nodes when converting from legacy.

diff --git a/Engine/Source/Runtime/Engine/Private/Materials/Material.cpp b/Engine/Source/Runtime/Engine/Private/Materials/Material.cpp
--- a/Engine/Source/Runtime/Engine/Private/Materials/Material.cpp
+++ b/Engine/Source/Runtime/Engine/Private/Materials/Material.cpp
@@ -3033,6 +3033,22 @@ void UMaterial::ConvertMaterialToStrataMaterial()
 	{
 		return;
 	}
+	
+	// Store current node post from the root node.
+	int32 CurrentNodePosX = EditorX;
+	const int32 TranslationOffsetX = 350.0f;
+	auto SetPosXAndMoveReferenceToTheRight = [&](UMaterialExpression* Node)
+	{
+		Node->MaterialExpressionEditorX = CurrentNodePosX;
+		CurrentNodePosX += TranslationOffsetX;
+	};
+	auto ReplaceNodeAndMoveToTheRight = [&](UMaterialExpression* NodeToReplace, UMaterialExpression* NewNode)
+	{
+		NewNode->MaterialExpressionEditorX = NodeToReplace->MaterialExpressionEditorX;
+		NewNode->MaterialExpressionEditorY = NodeToReplace->MaterialExpressionEditorY;
+		NodeToReplace->MaterialExpressionEditorX = NewNode->MaterialExpressionEditorX + TranslationOffsetX;
+		CurrentNodePosX = FMath::Max(NodeToReplace->MaterialExpressionEditorX + TranslationOffsetX, CurrentNodePosX);
+	};
 
 	auto MoveConnectionTo = [](auto& OldNodeInput, UMaterialExpression* NewNode, uint32 NewInputIndex)
 	{
@@ -3143,6 +3159,7 @@ void UMaterial::ConvertMaterialToStrataMaterial()
 	if (bUseMaterialAttributes && EditorOnly->MaterialAttributes.Expression && !EditorOnly->MaterialAttributes.Expression->IsResultStrataMaterial(EditorOnly->MaterialAttributes.OutputIndex)) // M_Rifle cause issues there
 	{
 		UMaterialExpressionBreakMaterialAttributes* BreakMatAtt = NewObject<UMaterialExpressionBreakMaterialAttributes>(this);
+		SetPosXAndMoveReferenceToTheRight(BreakMatAtt);
 		MoveConnectionTo(EditorOnly->MaterialAttributes, BreakMatAtt, 0);
 
 		// Check if Anisotropy input is actually used/connected
@@ -3190,6 +3207,7 @@ void UMaterial::ConvertMaterialToStrataMaterial()
 		}
 
 		ConvertNode = NewObject<UMaterialExpressionStrataLegacyConversion>(this);
+		SetPosXAndMoveReferenceToTheRight(ConvertNode);
 		ConvertNode->BaseColor.Connect(0, BreakMatAtt);
 		ConvertNode->Metallic.Connect(1, BreakMatAtt);
 		ConvertNode->Specular.Connect(2, BreakMatAtt);
@@ -3233,6 +3251,7 @@ void UMaterial::ConvertMaterialToStrataMaterial()
 
 			// Add shading model node
 			UMaterialExpressionShadingModel* ShadingModelNode = NewObject<UMaterialExpressionShadingModel>(this);
+			ReplaceNodeAndMoveToTheRight(ConvertNode, ShadingModelNode);
 			ShadingModelNode->ShadingModel = ShadingModel;
 			ConvertNode->ShadingModel.Connect(0, ShadingModelNode);
 
@@ -3248,6 +3267,7 @@ void UMaterial::ConvertMaterialToStrataMaterial()
 
 			// Now pass through the convert to decal node, which flag the material as SSM_Decal, which will set the domain to Decal.
 			UMaterialExpressionStrataConvertToDecal* ConvertToDecalNode = NewObject<UMaterialExpressionStrataConvertToDecal>(this);
+			ReplaceNodeAndMoveToTheRight(ConvertNode, ConvertToDecalNode);
 			ConvertToDecalNode->DecalMaterial.Connect(0, ConvertNode);
 
 			EditorOnly->FrontMaterial.Connect(0, ConvertToDecalNode);
@@ -3259,10 +3279,9 @@ void UMaterial::ConvertMaterialToStrataMaterial()
 	else if (!bUseMaterialAttributes && !EditorOnly->FrontMaterial.IsConnected() && GetExpressions().IsEmpty())
 	{
 		// Empty material: Create by default a slab node
-		const int32 NewNodeOffsetX = -300;
 		UMaterialExpressionStrataSlabBSDF* SlabNode = NewObject<UMaterialExpressionStrataSlabBSDF>(this);
+		SetPosXAndMoveReferenceToTheRight(SlabNode);
 		SlabNode->bUseMetalness = false;
-		SlabNode->MaterialExpressionEditorX = EditorX + NewNodeOffsetX;
 		EditorOnly->FrontMaterial.Connect(0, SlabNode);
 		bRelinkCustomOutputNodes = false;
 		bInvalidateShader = true;
@@ -3281,6 +3300,7 @@ void UMaterial::ConvertMaterialToStrataMaterial()
 
 					// Top slab BSDF as a simple Disney material
 					UMaterialExpressionStrataSlabBSDF* BottomSlabBSDF = NewObject<UMaterialExpressionStrataSlabBSDF>(this);
+					SetPosXAndMoveReferenceToTheRight(BottomSlabBSDF);
 					MoveConnectionTo(EditorOnly->BaseColor, BottomSlabBSDF, 0);					// BaseColor
 					MoveConnectionTo(EditorOnly->Metallic, BottomSlabBSDF, 2);					// Metallic
 					MoveConnectionTo(EditorOnly->Specular, BottomSlabBSDF, 3);					// Specular
@@ -3293,12 +3313,15 @@ void UMaterial::ConvertMaterialToStrataMaterial()
 
 					// Now weight the top base material by opacity.
 					UMaterialExpressionStrataSlabBSDF* TopSlabBSDF = NewObject<UMaterialExpressionStrataSlabBSDF>(this);
+					TopSlabBSDF->MaterialExpressionEditorX = BottomSlabBSDF->MaterialExpressionEditorX;
+					TopSlabBSDF->MaterialExpressionEditorY = BottomSlabBSDF->MaterialExpressionEditorY + 650;
 					MoveConnectionTo(EditorOnly->EmissiveColor, TopSlabBSDF, 14);				// Emissive
 					MoveConnectionTo(EditorOnly->ClearCoatRoughness, TopSlabBSDF, 7);			// ClearCoatRoughness => Roughness
 					MoveConnectionTo(EditorOnly->Normal, TopSlabBSDF, 9);						// Normal
 
 					//  The top layer has a hard coded specular value of 0.5 (F0 = 0.04)
 					UMaterialExpressionConstant* ConstantHalf = NewObject<UMaterialExpressionConstant>(this);
+					ReplaceNodeAndMoveToTheRight(TopSlabBSDF, ConstantHalf);
 					ConstantHalf->R = 0.5f;
 					TopSlabBSDF->GetInput(3)->Connect(0, ConstantHalf);
 
@@ -3306,11 +3329,13 @@ void UMaterial::ConvertMaterialToStrataMaterial()
 					// To simplify things, we set the top slab BSDF as having a constant Grey scale transmittance.
 					// As for the original, this is achieved with coverage so both transmittance and specular contribution vanishes
 					UMaterialExpressionConstant* ConstantZero = NewObject<UMaterialExpressionConstant>(this);
+					ReplaceNodeAndMoveToTheRight(TopSlabBSDF, ConstantZero);
 					ConstantZero->R = 0.0f;
 					TopSlabBSDF->GetInput(0)->Connect(0, ConstantZero);							// BaseColor = 0 to only feature absorption, no scattering
 
 					// Now setup the mean free path with a hard coded transmittance of 0.75 when viewing the surface perpendicularly
 					UMaterialExpressionConstant* Constant075 = NewObject<UMaterialExpressionConstant>(this);
+					ReplaceNodeAndMoveToTheRight(TopSlabBSDF, Constant075);
 					Constant075->R = 0.75f;
 					UMaterialExpressionStrataTransmittanceToMFP* TransToMDFP = NewObject<UMaterialExpressionStrataTransmittanceToMFP>(this);
 					TransToMDFP->GetInput(0)->Connect(0, Constant075);
@@ -3319,10 +3344,12 @@ void UMaterial::ConvertMaterialToStrataMaterial()
 
 					// Now weight the top base material by ClearCoat
 					UMaterialExpressionStrataWeight* TopSlabBSDFWithCoverage = NewObject<UMaterialExpressionStrataWeight>(this);
+					SetPosXAndMoveReferenceToTheRight(TopSlabBSDFWithCoverage);
 					TopSlabBSDFWithCoverage->GetInput(0)->Connect(0, TopSlabBSDF);				// TopSlabBSDF -> A
 					MoveConnectionTo(EditorOnly->ClearCoat, TopSlabBSDFWithCoverage, 1);		// ClearCoat -> Weight
 
 					UMaterialExpressionStrataVerticalLayering* VerticalLayering = NewObject<UMaterialExpressionStrataVerticalLayering>(this);
+					SetPosXAndMoveReferenceToTheRight(VerticalLayering);
 					VerticalLayering->GetInput(0)->Connect(0, TopSlabBSDFWithCoverage);			// Top -> Top
 					VerticalLayering->GetInput(1)->Connect(0, BottomSlabBSDF);					// Bottom -> Base
 
@@ -3335,6 +3362,7 @@ void UMaterial::ConvertMaterialToStrataMaterial()
 			if (!bClearCoatConversionDone)
 			{
 				ConvertNode = NewObject<UMaterialExpressionStrataLegacyConversion>(this);
+				SetPosXAndMoveReferenceToTheRight(ConvertNode);
 				ConvertNode->SubsurfaceProfile = bRequireNoSubsurfaceProfile ? nullptr : SubsurfaceProfile;
 				MoveConnectionTo(EditorOnly->BaseColor, ConvertNode, 0);
 				MoveConnectionTo(EditorOnly->Metallic, ConvertNode, 1);
@@ -3361,6 +3389,7 @@ void UMaterial::ConvertMaterialToStrataMaterial()
 					if (!EditorOnly->ShadingModelFromMaterialExpression.IsConnected())
 					{
 						UMaterialExpressionShadingModel* ShadingModelNode = NewObject<UMaterialExpressionShadingModel>(this);
+						ReplaceNodeAndMoveToTheRight(ConvertNode, ShadingModelNode);
 						ShadingModelNode->ShadingModel = MSM_DefaultLit;
 						ConvertNode->ShadingModel.Connect(0, ShadingModelNode);
 					}
@@ -3387,6 +3416,7 @@ void UMaterial::ConvertMaterialToStrataMaterial()
 					check(!bHasShadingModelMixture);
 
 					UMaterialExpressionShadingModel* ShadingModelNode = NewObject<UMaterialExpressionShadingModel>(this);
+					ReplaceNodeAndMoveToTheRight(ConvertNode, ShadingModelNode);
 					ShadingModelNode->ShadingModel = ShadingModel;
 					ConvertNode->ShadingModel.Connect(0, ShadingModelNode);
 
@@ -3402,6 +3432,7 @@ void UMaterial::ConvertMaterialToStrataMaterial()
 		else if (MaterialDomain == MD_Volume)
 		{
 			UMaterialExpressionStrataVolumetricFogCloudBSDF* VolBSDF = NewObject<UMaterialExpressionStrataVolumetricFogCloudBSDF>(this);
+			SetPosXAndMoveReferenceToTheRight(VolBSDF);
 			MoveConnectionTo(EditorOnly->BaseColor, VolBSDF, 0);		// Albedo
 			MoveConnectionTo(EditorOnly->SubsurfaceColor, VolBSDF, 1);	// Extinction
 			MoveConnectionTo(EditorOnly->EmissiveColor, VolBSDF, 2);	// EmissiveColor
@@ -3420,6 +3451,7 @@ void UMaterial::ConvertMaterialToStrataMaterial()
 
 			// Only Emissive & Opacity are valid input for PostProcess material
 			UMaterialExpressionStrataLightFunction* LightFunctionNode = NewObject<UMaterialExpressionStrataLightFunction>(this);
+			SetPosXAndMoveReferenceToTheRight(LightFunctionNode);
 			MoveConnectionTo(EditorOnly->EmissiveColor, LightFunctionNode, 0);
 
 			EditorOnly->FrontMaterial.Connect(0, LightFunctionNode);
@@ -3438,10 +3470,12 @@ void UMaterial::ConvertMaterialToStrataMaterial()
 			}
 
 			UMaterialExpressionStrataPostProcess* PostProcNode = NewObject<UMaterialExpressionStrataPostProcess>(this);
+			SetPosXAndMoveReferenceToTheRight(PostProcNode);
 			if (GetBlendMode() == BLEND_Translucent)
 			{
 				// If the blending mode was translucent, we must multiply the color by the opacity before the output node
 				UMaterialExpressionMultiply* ColorMultiplyOpacityNode = NewObject<UMaterialExpressionMultiply>(this);
+				ReplaceNodeAndMoveToTheRight(PostProcNode, ColorMultiplyOpacityNode);
 				CopyConnectionTo(EditorOnly->Opacity, ColorMultiplyOpacityNode, 0);
 				MoveConnectionTo(EditorOnly->EmissiveColor, ColorMultiplyOpacityNode, 1);
 
@@ -3469,6 +3503,7 @@ void UMaterial::ConvertMaterialToStrataMaterial()
 			ShadingModels.AddShadingModel(MSM_DefaultLit);
 
 			ConvertNode = NewObject<UMaterialExpressionStrataLegacyConversion>(this);
+			SetPosXAndMoveReferenceToTheRight(ConvertNode);
 			MoveConnectionTo(EditorOnly->BaseColor, ConvertNode, 0);
 			MoveConnectionTo(EditorOnly->Metallic, ConvertNode, 1);
 			MoveConnectionTo(EditorOnly->Specular, ConvertNode, 2);
@@ -3484,6 +3519,7 @@ void UMaterial::ConvertMaterialToStrataMaterial()
 
 			// Add constant for the Unlit shading model
 			UMaterialExpressionShadingModel* ShadingModelNode = NewObject<UMaterialExpressionShadingModel>(this);
+			ReplaceNodeAndMoveToTheRight(ConvertNode, ShadingModelNode);
 			ShadingModelNode->ShadingModel = ShadingModel;
 			ConvertNode->ShadingModel.Connect(0, ShadingModelNode);
 
@@ -3491,7 +3527,8 @@ void UMaterial::ConvertMaterialToStrataMaterial()
 			check(ConvertNode->ConvertedStrataMaterialInfo.CountShadingModels() == 1);
 
 			// Now pass through the convert to decal node, which flag the material as SSM_Decal, which will set the domain to Decal.
-			UMaterialExpressionStrataConvertToDecal* ConvertToDecalNode= NewObject<UMaterialExpressionStrataConvertToDecal>(this);
+			UMaterialExpressionStrataConvertToDecal* ConvertToDecalNode = NewObject<UMaterialExpressionStrataConvertToDecal>(this);
+			ReplaceNodeAndMoveToTheRight(ConvertNode, ConvertToDecalNode);
 			ConvertToDecalNode->DecalMaterial.Connect(0, ConvertNode);
 
 			EditorOnly->FrontMaterial.Connect(0, ConvertToDecalNode);
@@ -3540,6 +3577,9 @@ void UMaterial::ConvertMaterialToStrataMaterial()
 
 	if (bInvalidateShader)
 	{
+		// Set the root node position.
+		EditorX = CurrentNodePosX;
+
 		// We might have moved connections above so update the CachedExpressionData from the EditorOnly connection data (ground truth).
 		UpdateCachedExpressionData();
 	}
