Fix mask operation for hair strands

diff --git a/Engine/Source/Runtime/Engine/Private/RayTracingInstance.cpp b/Engine/Source/Runtime/Engine/Private/RayTracingInstance.cpp
--- a/Engine/Source/Runtime/Engine/Private/RayTracingInstance.cpp
+++ b/Engine/Source/Runtime/Engine/Private/RayTracingInstance.cpp
@@ -60,7 +60,7 @@ FRayTracingMaskAndFlags BuildRayTracingInstanceMaskAndFlags(TArrayView<const FMe
 		// For hair strands, opaque/translucent mask should be cleared to make sure geometry is only in the hair group. 
 		// If any segment receives shadow, it should receive only thin shadow instead of shadow.
 
-		Result.Mask &= (~RAY_TRACING_MASK_SHADOW | ~RAY_TRACING_MASK_THIN_SHADOW | ~RAY_TRACING_MASK_TRANSLUCENT | ~RAY_TRACING_MASK_OPAQUE);
+		Result.Mask &= ~(RAY_TRACING_MASK_SHADOW | RAY_TRACING_MASK_THIN_SHADOW | RAY_TRACING_MASK_TRANSLUCENT | RAY_TRACING_MASK_OPAQUE);
 
 		if (bAnySegmentsCastShadow)
 		{
