Fix static mesh dynamic raytracing instances using cached materials of incorrect LOD.

diff --git a/Engine/Source/Runtime/Engine/Private/StaticMeshRender.cpp b/Engine/Source/Runtime/Engine/Private/StaticMeshRender.cpp
--- a/Engine/Source/Runtime/Engine/Private/StaticMeshRender.cpp
+++ b/Engine/Source/Runtime/Engine/Private/StaticMeshRender.cpp
@@ -1807,25 +1807,7 @@ void FStaticMeshSceneProxy::GetDynamicRayTracingInstances(FRayTracingMaterialGat
 		return;
 	}
 
-	ESceneDepthPriorityGroup PrimitiveDPG = GetStaticDepthPriorityGroup();
-	const int32 LODIndex = FMath::Max(GetLOD(Context.ReferenceView), (int32)GetCurrentFirstLODIdx_RenderThread());
-
-	int32 RayTracingLODIndex = INDEX_NONE;
-	for (int32 LODIdx = LODIndex; LODIdx < RenderData->LODResources.Num(); LODIdx++)
-	{
-		if (RenderData->LODResources[LODIdx].RayTracingGeometry.IsValid())
-		{
-			RayTracingLODIndex = LODIdx;
-			break;
-		}
-	}
-
-	if (RayTracingLODIndex == INDEX_NONE)
-	{
-		return;
-	}
-
-	const FStaticMeshLODResources& LODModel = RenderData->LODResources[RayTracingLODIndex];
+	const ESceneDepthPriorityGroup PrimitiveDPG = GetStaticDepthPriorityGroup();
 
 	bool bEvaluateWPO = bDynamicRayTracingGeometry;
 
@@ -1842,8 +1824,31 @@ void FStaticMeshSceneProxy::GetDynamicRayTracingInstances(FRayTracingMaterialGat
 		}
 	}
 
+	const int32 NumLODs = RenderData->LODResources.Num();
+
+	int32 LODIndex = FMath::Max(GetLOD(Context.ReferenceView), (int32)GetCurrentFirstLODIdx_RenderThread());
+	const int32 OriginalLODIndex = LODIndex;
+
+	if (!bEvaluateWPO)
+	{
+		for (; LODIndex < NumLODs; LODIndex++)
+		{
+			if (RenderData->LODResources[LODIndex].RayTracingGeometry.IsValid())
+			{
+				break;
+			}
+		}
+	}
+
+	if (LODIndex == INDEX_NONE || LODIndex >= NumLODs)
+	{
+		return;
+	}
+
+	const FStaticMeshLODResources& LODModel = RenderData->LODResources[LODIndex];
+
 	// TODO: Need to validate that the DynamicRayTracingGeometries are still valid - they could contain streamed out IndexBuffers from the shared StaticMesh (UE-139474)
-	FRayTracingGeometry& Geometry = bEvaluateWPO? DynamicRayTracingGeometries[LODIndex] : RenderData->LODResources[RayTracingLODIndex].RayTracingGeometry;
+	FRayTracingGeometry& Geometry = bEvaluateWPO ? DynamicRayTracingGeometries[LODIndex] : RenderData->LODResources[LODIndex].RayTracingGeometry;
 	
 	if (LODModel.GetNumVertices() <= 0 || Geometry.Initializer.TotalPrimitiveCount <= 0)
 	{
