enum EShaderPlatform now has uint16 type in order to be forward declared
Forward declared EShaderPlatform in Scalability.h

diff --git a/Engine/Source/Runtime/Engine/Private/Scalability.cpp b/Engine/Source/Runtime/Engine/Private/Scalability.cpp
--- a/Engine/Source/Runtime/Engine/Private/Scalability.cpp
+++ b/Engine/Source/Runtime/Engine/Private/Scalability.cpp
@@ -13,6 +13,7 @@
 
 #if WITH_EDITOR
 #include "RHI.h"
+#include "RHIDefinitions.h"
 #endif
 
 Scalability::FOnScalabilitySettingsChanged Scalability::OnScalabilitySettingsChanged;
@@ -383,7 +384,7 @@ void ApplyScalabilityGroupFromPlatformIni(const TCHAR* InSectionName, const TCHA
 	UE::ConfigUtilities::ForEachCVarInSectionFromIni(InSectionName, InIniFilename, Func);
 }
 
-void ChangeScalabilityPreviewPlatform(FName NewPlatformScalabilityName, EShaderPlatform ShaderPlatform)
+void ChangeScalabilityPreviewPlatform(FName NewPlatformScalabilityName, const EShaderPlatform& ShaderPlatform)
 {
 	PlatformScalabilityName = NewPlatformScalabilityName;
 	bScalabilityShaderPlatformHasBeenChanged = ScalabilityShaderPlatform != ShaderPlatform;
@@ -428,7 +429,7 @@ static void SetGroupQualityLevel(const TCHAR* InGroupName, int32 InQualityLevel,
 }
 
 #if WITH_EDITOR
-void ApplyCachedQualityLevelForShaderPlatform(EShaderPlatform ShaderPlatform)
+void ApplyCachedQualityLevelForShaderPlatform(const EShaderPlatform& ShaderPlatform)
 {
 	if (ScalabilityShaderPlatform != EShaderPlatform::SP_NumPlatforms)
 	{
diff --git a/Engine/Source/Runtime/Engine/Public/Scalability.h b/Engine/Source/Runtime/Engine/Public/Scalability.h
--- a/Engine/Source/Runtime/Engine/Public/Scalability.h
+++ b/Engine/Source/Runtime/Engine/Public/Scalability.h
@@ -7,9 +7,8 @@
 #pragma once
 
 #include "CoreMinimal.h"
-#if WITH_EDITOR
-#include "RHIDefinitions.h"
-#endif
+
+enum EShaderPlatform : uint16;
 
 namespace Scalability
 { 
@@ -159,7 +158,7 @@ namespace Scalability
 	ENGINE_API void SetQualityLevels(const FQualityLevels& QualityLevels, bool bForce = false);
 
 #if WITH_EDITOR
-	ENGINE_API void ApplyCachedQualityLevelForShaderPlatform(EShaderPlatform ShaderPlatform);
+	ENGINE_API void ApplyCachedQualityLevelForShaderPlatform(const EShaderPlatform& ShaderPlatform);
 #endif
 
 	/** This is the only suggested way to get the current state - don't get CVars directly */
@@ -217,7 +216,7 @@ namespace Scalability
 
 #if WITH_EDITOR
 	/** Set an Editor preview scalability platform */
-	void ENGINE_API ChangeScalabilityPreviewPlatform(FName NewPlatformScalabilityName, EShaderPlatform ShaderPlatform);
+	void ENGINE_API ChangeScalabilityPreviewPlatform(FName NewPlatformScalabilityName, const EShaderPlatform& ShaderPlatform);
 #endif
 
 	ENGINE_API FText GetQualityLevelText(int32 Value, int32 NumLevels);
diff --git a/Engine/Source/Runtime/RHI/Public/RHIDefinitions.h b/Engine/Source/Runtime/RHI/Public/RHIDefinitions.h
--- a/Engine/Source/Runtime/RHI/Public/RHIDefinitions.h
+++ b/Engine/Source/Runtime/RHI/Public/RHIDefinitions.h
@@ -89,7 +89,7 @@ enum EShaderFrequency : uint8
 static_assert(SF_NumFrequencies <= (1 << SF_NumBits), "SF_NumFrequencies will not fit on SF_NumBits");
 
 /** @warning: update *LegacyShaderPlatform* when the below changes */
-enum EShaderPlatform
+enum EShaderPlatform : uint16
 {
 	SP_PCD3D_SM5					= 0,
 	SP_METAL						= 11,
