Iris
- Added missing include

diff --git a/Engine/Source/Runtime/Engine/Private/Net/Experimental/Iris/DataStreamChannel.cpp b/Engine/Source/Runtime/Engine/Private/Net/Experimental/Iris/DataStreamChannel.cpp
--- a/Engine/Source/Runtime/Engine/Private/Net/Experimental/Iris/DataStreamChannel.cpp
+++ b/Engine/Source/Runtime/Engine/Private/Net/Experimental/Iris/DataStreamChannel.cpp
@@ -20,6 +20,7 @@
 #include "Iris/Core/IrisProfiler.h"
 #include "Iris/Core/IrisMemoryTracker.h"
 #include "ProfilingDebugging/CsvProfiler.h"
+#include "PacketHandler.h"
 #endif // UE_WITH_IRIS
 
 namespace UE::Net::Private
