Keep duplicated vertex data for mesh deformers if the section is marked for RecomputeNormal support.
In future we will use a more specific setting to control this (and other connectivity data we would like to add).
But this is simplest approach to allow some control while maintaining 5.1 data compatibility.

diff --git a/Engine/Source/Runtime/Engine/Private/SkeletalMeshLODRenderData.cpp b/Engine/Source/Runtime/Engine/Private/SkeletalMeshLODRenderData.cpp
--- a/Engine/Source/Runtime/Engine/Private/SkeletalMeshLODRenderData.cpp
+++ b/Engine/Source/Runtime/Engine/Private/SkeletalMeshLODRenderData.cpp
@@ -5,6 +5,7 @@
 #include "Engine/SkinnedAsset.h"
 #include "Engine/SkeletalMesh.h"	// for FSkeletalMeshLODInfo
 #include "Engine/SkinnedAssetCommon.h"
+#include "Animation/MeshDeformerProvider.h"
 #include "Animation/MorphTarget.h"
 #include "Misc/ConfigCacheIni.h"
 #include "EngineLogs.h"
@@ -84,14 +85,28 @@ namespace
 	};
 }
 
-static bool IsGPUSkinCacheAvailable(const ITargetPlatform& TargetPlatform)
+/** 
+ * This function returns true if the duplicate vertices should be cooked. 
+ * The data is used to deal with seams along split vertices when recomputing normals at runtime.
+ */
+static bool RequiresDuplicateVerticesInCook(EShaderPlatform InPlatform, FSkelMeshRenderSection const& RenderSection)
+{
+	// DuplicatedVertices are cooked if we have GPUSkinCache or MeshDeformer systems for this platform.
+	// We always cook when GPUSkinCache is available so that we can support runtime switch of r.SkinCache.RecomputeTangents.
+	// For MeshDeformer we only cook if the section was marked for bRecomputeTangent.
+	static IMeshDeformerProvider* MeshDeformerProvider = IMeshDeformerProvider::Get();
+	bool bMeshDeformersAvailable = MeshDeformerProvider && MeshDeformerProvider->IsSupported(InPlatform);
+	return (RenderSection.bRecomputeTangent && bMeshDeformersAvailable) || IsGPUSkinCacheAvailable(InPlatform);
+}
+
+static bool RequiresDuplicateVerticesInCook(const ITargetPlatform& InTargetPlatform, FSkelMeshRenderSection const& RenderSection)
 {
 	TArray<FName> TargetedShaderFormats;
-	TargetPlatform.GetAllTargetedShaderFormats(TargetedShaderFormats);
+	InTargetPlatform.GetAllTargetedShaderFormats(TargetedShaderFormats);
 	for (int32 FormatIndex = 0; FormatIndex < TargetedShaderFormats.Num(); ++FormatIndex)
 	{
 		const EShaderPlatform LegacyShaderPlatform = ShaderFormatToLegacyShaderPlatform(TargetedShaderFormats[FormatIndex]);
-		if (IsGPUSkinCacheAvailable(LegacyShaderPlatform))
+		if (RequiresDuplicateVerticesInCook(LegacyShaderPlatform, RenderSection))
 		{
 			return true;
 		}
@@ -99,6 +114,24 @@ static bool IsGPUSkinCacheAvailable(const ITargetPlatform& TargetPlatform)
 	return false;
 }
 
+/** 
+ * This function returns true if the duplicate vertices should be kept at runtime initialization. 
+ * We may cook the data, but use this runtime check to discard it and save memory.
+ */
+static bool RequiresDuplicateVertices()
+{
+	// Never drop at runtime if the data was cooked for MeshDeformers.
+	static IMeshDeformerProvider* MeshDeformerProvider = IMeshDeformerProvider::Get();
+	if (MeshDeformerProvider && MeshDeformerProvider->IsSupported(GMaxRHIShaderPlatform))
+	{
+		return true;
+	}
+
+	// Can assume data is for GPUSkinCache. Respect the CVar option to discard all data at runtime.
+	return GPUSkinCacheNeedsDuplicatedVertices();
+}
+
+
 // Serialization.
 FArchive& operator<<(FArchive& Ar, FSkelMeshRenderSection& S)
 {
@@ -108,10 +141,9 @@ FArchive& operator<<(FArchive& Ar, FSkelMeshRenderSection& S)
 	Ar.UsingCustomVersion(FRecomputeTangentCustomVersion::GUID);
 	Ar.UsingCustomVersion(FUE5ReleaseStreamObjectVersion::GUID);
 
-	// DuplicatedVerticesBuffer is used only for SkinCache and Editor features which is SM5 only
+	// Strip DuplicatedVerticesBuffer from cook if it won't be needed.
 	uint8 ClassDataStripFlags = 0;
-	if (Ar.IsCooking() && 
-		!(Ar.CookingTarget()->SupportsFeature(ETargetPlatformFeatures::DeferredRendering) || IsGPUSkinCacheAvailable(*Ar.CookingTarget())))
+	if (Ar.IsCooking() && !RequiresDuplicateVerticesInCook(*Ar.CookingTarget(), S))
 	{
 		ClassDataStripFlags |= DuplicatedVertices;
 	}
@@ -205,26 +237,23 @@ void FSkeletalMeshLODRenderData::InitResources(bool bNeedsVertexColors, int32 LO
 		BeginInitResource(&ClothVertexBuffer);
 	}
 
-	// DuplicatedVerticesBuffer is used only for SkinCache and Editor features which is SM5 only
-    if (IsGPUSkinCacheAvailable(GMaxRHIShaderPlatform) || IsFeatureLevelSupported(GMaxRHIShaderPlatform, ERHIFeatureLevel::SM5))
+	// We can discard any the cooked DuplicatedVertices here based on runtime settings.
+	bool bNeedDuplicatedVertices = RequiresDuplicateVertices();
+	for (FSkelMeshRenderSection& RenderSection : RenderSections)
 	{
-		const bool bSkinCacheNeedsDuplicatedVertices = GPUSkinCacheNeedsDuplicatedVertices();
-		for (auto& RenderSection : RenderSections)
+		if (bNeedDuplicatedVertices)
+		{
+			// No need to discard CPU data in cooked builds as bNeedsCPUAccess is false (see FDuplicatedVerticesBuffer constructor), 
+			// so it'd be auto-discarded after the RHI has copied the resource data. Keep CPU data when in the editor for geometry operations.
+			check(RenderSection.DuplicatedVerticesBuffer.DupVertData.Num());
+			BeginInitResource(&RenderSection.DuplicatedVerticesBuffer);
+		}
+		else
 		{
-			if (bSkinCacheNeedsDuplicatedVertices)
-			{
-				// No need to discard CPU data in cooked builds as bNeedsCPUAccess is false (see FDuplicatedVerticesBuffer constructor), 
-				// so it'd be auto-discarded after the RHI has copied the resource data. Keep CPU data when in the editor for geometry operations.
-				check(RenderSection.DuplicatedVerticesBuffer.DupVertData.Num());
-				BeginInitResource(&RenderSection.DuplicatedVerticesBuffer);
-			}
-			else
-			{
 #if !WITH_EDITOR
-				// Discard CPU data in cooked builds. Keep CPU data when in the editor for geometry operations.
-				RenderSection.DuplicatedVerticesBuffer.ReleaseCPUResources();
+			// Discard CPU data in cooked builds. Keep CPU data when in the editor for geometry operations.
+			RenderSection.DuplicatedVerticesBuffer.ReleaseCPUResources();
 #endif
-			}
 		}
 	}
 
@@ -263,19 +292,9 @@ void FSkeletalMeshLODRenderData::ReleaseResources()
 	SkinWeightVertexBuffer.BeginReleaseResources();
 	BeginReleaseResource(&StaticVertexBuffers.ColorVertexBuffer);
 	BeginReleaseResource(&ClothVertexBuffer);
-	// DuplicatedVerticesBuffer is used only for SkinCache and Editor features which is SM5 only
-    if (IsGPUSkinCacheAvailable(GMaxRHIShaderPlatform) || IsFeatureLevelSupported(GMaxRHIShaderPlatform, ERHIFeatureLevel::SM5))
+	for (FSkelMeshRenderSection& RenderSection : RenderSections)
 	{
-		if (GPUSkinCacheNeedsDuplicatedVertices())
-		{
-			for (auto& RenderSection : RenderSections)
-			{
-#if WITH_EDITOR
-				check(RenderSection.DuplicatedVerticesBuffer.DupVertData.Num());
-#endif
-				BeginReleaseResource(&RenderSection.DuplicatedVerticesBuffer);
-			}
-		}
+		BeginReleaseResource(&RenderSection.DuplicatedVerticesBuffer);
 	}
 	BeginReleaseResource(&MorphTargetVertexInfoBuffers);
 
