Add cvar setting to disable mesh deformers below a fixed LOD.
This is temporary before full implementation of per LOD deformer settings.

diff --git a/Engine/Source/Runtime/Engine/Classes/Components/SkinnedMeshComponent.h b/Engine/Source/Runtime/Engine/Classes/Components/SkinnedMeshComponent.h
--- a/Engine/Source/Runtime/Engine/Classes/Components/SkinnedMeshComponent.h
+++ b/Engine/Source/Runtime/Engine/Classes/Components/SkinnedMeshComponent.h
@@ -263,7 +263,10 @@ protected:
 
 public:
 	/** Get the currently active MeshDeformer Instance. */
-	virtual UMeshDeformerInstance const* GetMeshDeformerInstance() const override { return MeshDeformerInstance; }
+	UMeshDeformerInstance const* GetMeshDeformerInstance() const override { return MeshDeformerInstance; }
+
+	/** Max LOD at which to update or apply the MeshDeformer. */
+	int32 GetMeshDeformerMaxLOD() const;
 
 	/** const getters for previous transform idea */
 	const TArray<uint8>& GetPreviousBoneVisibilityStates() const
diff --git a/Engine/Source/Runtime/Engine/Private/Components/SkinnedMeshComponent.cpp b/Engine/Source/Runtime/Engine/Private/Components/SkinnedMeshComponent.cpp
--- a/Engine/Source/Runtime/Engine/Private/Components/SkinnedMeshComponent.cpp
+++ b/Engine/Source/Runtime/Engine/Private/Components/SkinnedMeshComponent.cpp
@@ -1077,7 +1077,7 @@ void USkinnedMeshComponent::SendRenderDynamicData_Concurrent()
 			UpdateMorphMaterialUsageOnProxy();
 		}
 
-		if (MeshDeformerInstance != nullptr)
+		if (MeshDeformerInstance != nullptr && UseLOD <= GetMeshDeformerMaxLOD())
 		{
 			if (MeshDeformerInstance->IsActive())
 			{
@@ -2143,6 +2143,19 @@ void USkinnedMeshComponent::UnsetMeshDeformer()
 	SetMeshDeformer(false, nullptr);
 }
 
+static TAutoConsoleVariable<int32> CVarMeshDeformerMaxLod(
+	TEXT("r.MeshDeformerMaxLOD"),
+	-1,
+	TEXT("Don't apply MeshDeformers to SkinnedMeshs above this LOD.\n")
+	TEXT("Default is -1 (disabled).\n"),
+	ECVF_Default);
+
+int32 USkinnedMeshComponent::GetMeshDeformerMaxLOD() const
+{
+	int32 MaxLod = CVarMeshDeformerMaxLod.GetValueOnGameThread();
+	return MaxLod >= 0 ? MaxLod : GetNumLODs() - 1;
+}
+
 FSkeletalMeshRenderData* USkinnedMeshComponent::GetSkeletalMeshRenderData() const
 {
 	if (MeshObject)
diff --git a/Engine/Source/Runtime/Engine/Private/SkeletalRenderGPUSkin.cpp b/Engine/Source/Runtime/Engine/Private/SkeletalRenderGPUSkin.cpp
--- a/Engine/Source/Runtime/Engine/Private/SkeletalRenderGPUSkin.cpp
+++ b/Engine/Source/Runtime/Engine/Private/SkeletalRenderGPUSkin.cpp
@@ -2537,7 +2537,7 @@ void FDynamicSkelMeshObjectDataGPUSkin::InitDynamicSkelMeshObjectDataGPUSkin(
 	}
 
 	bIsSkinCacheAllowed = InMeshComponent ? InMeshComponent->IsSkinCacheAllowed(InLODIndex) : false;
-	bHasMeshDeformer = InMeshComponent ? InMeshComponent->GetMeshDeformerInstance() != nullptr : false;
+	bHasMeshDeformer = InMeshComponent ? InMeshComponent->GetMeshDeformerInstance() != nullptr && LODIndex <= InMeshComponent->GetMeshDeformerMaxLOD() : false;
 	bForceUpdateDynamicDataImmediately = InMeshComponent ? InMeshComponent->GetForceUpdateDynamicDataImmediately() : false;
 
 	if (bIsSkinCacheAllowed && InMeshObject->FeatureLevel == ERHIFeatureLevel::ES3_1)
