USkeleton::Modify() should not be call outside of the main thread, when we build the skeletalmesh asynchronously. This is a temporary change until we refactor the morph target build.

#rnx

diff --git a/Engine/Source/Runtime/Engine/Private/Animation/Skeleton.cpp b/Engine/Source/Runtime/Engine/Private/Animation/Skeleton.cpp
--- a/Engine/Source/Runtime/Engine/Private/Animation/Skeleton.cpp
+++ b/Engine/Source/Runtime/Engine/Private/Animation/Skeleton.cpp
@@ -1733,7 +1733,10 @@ bool USkeleton::VerifySmartNameInternal(const FName&  ContainerName, FSmartName&
 		if (!Mapping->FindSmartName(InOutSmartName.DisplayName, InOutSmartName))
 		{
 #if WITH_EDITOR
-			Modify();
+			if (IsInGameThread())
+			{
+				Modify();
+			}
 #endif
 			InOutSmartName = Mapping->AddName(InOutSmartName.DisplayName);
 			return true;
@@ -1763,7 +1766,12 @@ FSmartNameMapping* USkeleton::GetOrAddSmartNameContainer(const FName& ContainerN
 	FSmartNameMapping* Mapping = SmartNames.GetContainerInternal(ContainerName);
 	if (Mapping == nullptr)
 	{
-		Modify();
+#if WITH_EDITOR
+		if (IsInGameThread())
+		{
+			Modify();
+		}
+#endif
 		IncreaseAnimCurveUidVersion();
 		SmartNames.AddContainer(ContainerName);
 		AnimCurveMapping = SmartNames.GetContainerInternal(USkeleton::AnimCurveMappingName);
@@ -1892,8 +1900,8 @@ void USkeleton::AccumulateCurveMetaData(FName CurveName, bool bMaterialSet, bool
 				CurveMetaData->Type.bMaterial |= bMaterialSet;
 				CurveMetaData->Type.bMorphtarget |= bMorphtargetSet;
 
-				if (bOldMaterial != CurveMetaData->Type.bMaterial 
-					|| bOldMorphtarget != CurveMetaData->Type.bMorphtarget)
+				if (IsInGameThread() && (bOldMaterial != CurveMetaData->Type.bMaterial
+					|| bOldMorphtarget != CurveMetaData->Type.bMorphtarget))
 				{
 					MarkPackageDirty();
 				}
