Fix blend profile crash when duplicate indices are introduced due to a virtual bone rename operation.
Fix blend profile bone indices going out of sync when virtual bones are deleted.

diff --git a/Engine/Source/Runtime/Engine/Classes/Animation/BlendProfile.h b/Engine/Source/Runtime/Engine/Classes/Animation/BlendProfile.h
--- a/Engine/Source/Runtime/Engine/Classes/Animation/BlendProfile.h
+++ b/Engine/Source/Runtime/Engine/Classes/Animation/BlendProfile.h
@@ -83,10 +83,18 @@ public:
 	void SetBoneBlendScale(const FName& InBoneName, float InScale, bool bRecurse = false, bool bCreate = false);
 
 	/** Removes the entry for the specified bone index (does nothing if it doesn't exist) 
-	*  @param InBoneIdx Index of the bone to remove from this blend profile
-	*/
+	 *  @param InBoneIdx Index of the bone to remove from this blend profile
+	 */
 	void RemoveEntry(int32 InBoneIdx);
 
+	/** Ensures the bone name of the specified entry matches the skeleton index (does nothing if it doesn't exist)
+	 *  @param InBoneIdx Index of the bone to refresh
+	 */
+	void RefreshBoneEntry(int32 InBoneIndex);
+
+	/** Ensures the bone names match the skeleton indices by using the bone name as our lookup key. */
+	void RefreshBoneEntriesFromName();
+
 	/**
 	 * Get the bone entry by entry index.
 	 * @param[in] InEntryIdx The index to the bone entry in range [0, GetNumBlendEntries()-1].
diff --git a/Engine/Source/Runtime/Engine/Private/Animation/BlendProfile.cpp b/Engine/Source/Runtime/Engine/Private/Animation/BlendProfile.cpp
--- a/Engine/Source/Runtime/Engine/Private/Animation/BlendProfile.cpp
+++ b/Engine/Source/Runtime/Engine/Private/Animation/BlendProfile.cpp
@@ -51,6 +51,28 @@ void UBlendProfile::RemoveEntry(int32 InBoneIdx)
 		});
 }
 
+void UBlendProfile::RefreshBoneEntry(int32 InBoneIndex)
+{
+	FBlendProfileBoneEntry* FoundEntry = ProfileEntries.FindByPredicate([InBoneIndex](const FBlendProfileBoneEntry& Entry)
+		{
+			return Entry.BoneReference.BoneIndex == InBoneIndex;
+		});
+
+	if (FoundEntry)
+	{
+		FoundEntry->BoneReference.BoneName = OwningSkeleton->GetReferenceSkeleton().GetBoneName(InBoneIndex);
+		FoundEntry->BoneReference.Initialize(OwningSkeleton);
+	}
+}
+
+void UBlendProfile::RefreshBoneEntriesFromName()
+{
+	for (FBlendProfileBoneEntry& Entry : ProfileEntries)
+	{
+		Entry.BoneReference.Initialize(OwningSkeleton);
+	}
+}
+
 const FBlendProfileBoneEntry& UBlendProfile::GetEntry(const int32 InEntryIdx) const
 {
 	return ProfileEntries[InEntryIdx];
diff --git a/Engine/Source/Runtime/Engine/Private/Animation/Skeleton.cpp b/Engine/Source/Runtime/Engine/Private/Animation/Skeleton.cpp
--- a/Engine/Source/Runtime/Engine/Private/Animation/Skeleton.cpp
+++ b/Engine/Source/Runtime/Engine/Private/Animation/Skeleton.cpp
@@ -1985,12 +1985,7 @@ void USkeleton::RemoveVirtualBones(const TArray<FName>& BonesToRemove)
 	// Blend profiles cache bone names and indices, make sure they remain in sync when the indices change
 	for (UBlendProfile* Profile : BlendProfiles)
 	{
-		// TEMPORARY FIX FOR 5.1.1
-
-		for (FBlendProfileBoneEntry& Entry : Profile->ProfileEntries)
-		{
-			Entry.BoneReference.Initialize(Profile->OwningSkeleton);
-		}
+		Profile->RefreshBoneEntriesFromName();
 	}
 }
 
@@ -2033,18 +2028,7 @@ void USkeleton::RenameVirtualBone(const FName OriginalBoneName, const FName NewB
 		{
 			for (UBlendProfile* Profile : BlendProfiles)
 			{
-				// TEMPORARY FIX FOR 5.1.1
-
-				FBlendProfileBoneEntry* FoundEntry = Profile->ProfileEntries.FindByPredicate([BoneIdx](const FBlendProfileBoneEntry& Entry)
-					{
-						return Entry.BoneReference.BoneIndex == BoneIdx;
-					});
-
-				if (FoundEntry)
-				{
-					FoundEntry->BoneReference.BoneName = Profile->OwningSkeleton->GetReferenceSkeleton().GetBoneName(BoneIdx);
-					FoundEntry->BoneReference.Initialize(Profile->OwningSkeleton);
-				}
+				Profile->RefreshBoneEntry(BoneIdx);
 			}
 		}
 	}
