Small change to Physics BodyInstance to support the non UPrimitiveComponent

#rnx

diff --git a/Engine/Source/Runtime/Engine/Classes/PhysicsEngine/BodyInstance.h b/Engine/Source/Runtime/Engine/Classes/PhysicsEngine/BodyInstance.h
--- a/Engine/Source/Runtime/Engine/Classes/PhysicsEngine/BodyInstance.h
+++ b/Engine/Source/Runtime/Engine/Classes/PhysicsEngine/BodyInstance.h
@@ -56,7 +56,7 @@ namespace EDOFMode
 	};
 }
 
-struct FBodyInstnace;
+struct FBodyInstance;
 
 #ifndef CHAOS_DEBUG_NAME
 #define CHAOS_DEBUG_NAME 0
@@ -68,6 +68,7 @@ struct FBodyInstnace;
 struct FInitBodySpawnParams
 {
 	ENGINE_API FInitBodySpawnParams(const UPrimitiveComponent* PrimComp);
+	ENGINE_API FInitBodySpawnParams(bool bInStaticPhysics, bool bInPhysicsTypeDeterminesSimulation);
 
 	/** Whether the created physics actor will be static */
 	bool bStaticPhysics;
diff --git a/Engine/Source/Runtime/Engine/Private/PhysicsEngine/BodyInstance.cpp b/Engine/Source/Runtime/Engine/Private/PhysicsEngine/BodyInstance.cpp
--- a/Engine/Source/Runtime/Engine/Private/PhysicsEngine/BodyInstance.cpp
+++ b/Engine/Source/Runtime/Engine/Private/PhysicsEngine/BodyInstance.cpp
@@ -1463,6 +1463,13 @@ FInitBodySpawnParams::FInitBodySpawnParams(const UPrimitiveComponent* PrimComp)
 	}
 }
 
+FInitBodySpawnParams::FInitBodySpawnParams(bool bInStaticPhysics, bool bInPhysicsTypeDeterminesSimulation)
+	: bStaticPhysics(bInStaticPhysics)
+	, bPhysicsTypeDeterminesSimulation(bInPhysicsTypeDeterminesSimulation)
+{
+}
+
+
 // Chaos addition
 static TAutoConsoleVariable<int32> CVarAllowCreatePhysxBodies(
 	TEXT("p.chaos.AllowCreatePhysxBodies"),
