Added FISMComponentBatcher which can be used to batch SMC & ISMC
Now used in place of duplicated code in:
* UHLODBuilder::BatchInstances()
* FMeshMergeUtilities::MergeComponentsToInstances
* FPackedLevelActorISMBuilder::PackActors

diff --git a/Engine/Source/Developer/MeshMergeUtilities/Private/MeshMergeUtilities.cpp b/Engine/Source/Developer/MeshMergeUtilities/Private/MeshMergeUtilities.cpp
--- a/Engine/Source/Developer/MeshMergeUtilities/Private/MeshMergeUtilities.cpp
+++ b/Engine/Source/Developer/MeshMergeUtilities/Private/MeshMergeUtilities.cpp
@@ -80,6 +80,7 @@
 #include "Widgets/Notifications/SNotificationList.h"
 #include "Framework/Notifications/NotificationManager.h"
 
+#include "ISMPartition/ISMComponentBatcher.h"
 #include "ISMPartition/ISMComponentDescriptor.h"
 
 #define LOCTEXT_NAMESPACE "MeshMergeUtils"
@@ -3546,10 +3547,10 @@ void FMeshMergeUtilities::MergeComponentsToInstances(const TArray<UPrimitiveComp
 						NewComponent->SetCollisionProfileName(ComponentEntry.CollisionProfileName);
 						NewComponent->SetCollisionEnabled(ComponentEntry.CollisionEnabled);
 						NewComponent->SetMobility(EComponentMobility::Static);
-						for(UStaticMeshComponent* OriginalComponent : ComponentEntry.OriginalComponents)
-						{
-							NewComponent->AddInstance(OriginalComponent->GetComponentTransform());
-						}
+
+						FISMComponentBatcher ISMComponentBatcher;
+						ISMComponentBatcher.Append(ComponentEntry.OriginalComponents);
+						ISMComponentBatcher.InitComponent(NewComponent);						
 
 						NewComponent->RegisterComponent();
 					}
diff --git a/Engine/Source/Runtime/Engine/Private/ISMPartition/ISMComponentBatcher.cpp b/Engine/Source/Runtime/Engine/Private/ISMPartition/ISMComponentBatcher.cpp
new file mode 100644
--- /dev/null
+++ b/Engine/Source/Runtime/Engine/Private/ISMPartition/ISMComponentBatcher.cpp
@@ -0,0 +1,115 @@
+// Copyright Epic Games, Inc. All Rights Reserved.
+
+#include "ISMPartition/ISMComponentBatcher.h"
+
+void FISMComponentBatcher::Add(const UActorComponent* InComponent)
+{
+	AddInternal(InComponent, TOptional<TFunctionRef<FTransform(const FTransform&)>>());
+}
+
+void FISMComponentBatcher::Add(const UActorComponent* InComponent, TFunctionRef<FTransform(const FTransform&)> InTransformFunc)
+{
+	AddInternal(InComponent, InTransformFunc);
+}
+
+void FISMComponentBatcher::AddInternal(const UActorComponent* InComponent, TOptional<TFunctionRef<FTransform(const FTransform&)>> InTransformFunc)
+{
+	int32 NewNumCustomDataFloats = 0;
+	int32 NewNumInstances = 0;
+
+	// Compute number of instances & custom data float to add
+	if (const UInstancedStaticMeshComponent* ISMC = Cast<UInstancedStaticMeshComponent>(InComponent))
+	{
+		NewNumCustomDataFloats = ISMC->NumCustomDataFloats;
+		NewNumInstances = ISMC->GetInstanceCount();
+	}
+	else if (const UStaticMeshComponent* SMC = Cast<UStaticMeshComponent>(InComponent))
+	{
+		NewNumCustomDataFloats = SMC->GetCustomPrimitiveData().Data.Num();
+		NewNumInstances = 1;
+	}
+
+	if (NewNumCustomDataFloats > NumCustomDataFloats)
+	{
+		TArray<float> NewInstancesCustomData;
+		NewInstancesCustomData.AddZeroed(NewNumCustomDataFloats * NumInstances);
+
+		for (int32 InstanceIdx = 0; InstanceIdx < NumInstances; ++InstanceIdx)
+		{
+			int32 SrcCustomDataOffset = NumCustomDataFloats * InstanceIdx;
+			int32 DestCustomDataOffset = NewNumCustomDataFloats * InstanceIdx;
+			for (int32 CustomDataIdx = 0; CustomDataIdx < NumCustomDataFloats; ++CustomDataIdx)
+			{
+				NewInstancesCustomData[DestCustomDataOffset + CustomDataIdx] = InstancesCustomData[SrcCustomDataOffset + CustomDataIdx];
+			}
+		}
+
+		InstancesCustomData = MoveTemp(NewInstancesCustomData);
+		NumCustomDataFloats = NewNumCustomDataFloats;
+	}
+
+	NumInstances += NewNumInstances;
+	InstancesTransformsWS.Reserve(NumInstances);
+	InstancesCustomData.Reserve(NumInstances * NumCustomDataFloats);
+
+	// Add instances
+	if (const UInstancedStaticMeshComponent* ISMC = Cast<UInstancedStaticMeshComponent>(InComponent))
+	{
+		// Add random seeds
+		RandomSeeds.Add({ NumInstances, ISMC->InstancingRandomSeed });
+		for (const FInstancedStaticMeshRandomSeed& AdditionalRandomSeed : ISMC->AdditionalRandomSeeds)
+		{
+			RandomSeeds.Add({ NumInstances + AdditionalRandomSeed.StartInstanceIndex, AdditionalRandomSeed.RandomSeed });
+		}
+
+		// Add each instance
+		for (int32 InstanceIdx = 0; InstanceIdx < ISMC->GetInstanceCount(); InstanceIdx++)
+		{
+			// Add instance transform
+			FTransform InstanceTransformWS;
+			ISMC->GetInstanceTransform(InstanceIdx, InstanceTransformWS, /*bWorldSpace*/ true);
+			if (InTransformFunc.IsSet())
+			{
+				InstanceTransformWS = InTransformFunc->operator()(InstanceTransformWS);
+			}
+			InstancesTransformsWS.Add(InstanceTransformWS);
+
+			// Add per instance custom data
+			TConstArrayView<float> InstanceCustomData(&ISMC->PerInstanceSMCustomData[InstanceIdx * ISMC->NumCustomDataFloats], ISMC->NumCustomDataFloats);
+			InstancesCustomData.Append(InstanceCustomData);
+			InstancesCustomData.AddDefaulted(NumCustomDataFloats - InstanceCustomData.Num());
+		}
+	}
+	else if (const UStaticMeshComponent* SMC = Cast<UStaticMeshComponent>(InComponent))
+	{
+		FTransform InstanceTransformWS = SMC->GetComponentTransform();
+		if (InTransformFunc.IsSet())
+		{
+			InstanceTransformWS = InTransformFunc->operator()(InstanceTransformWS);
+		}
+
+		// Add transform
+		InstancesTransformsWS.Add(InstanceTransformWS);
+
+		// Add custom data
+		InstancesCustomData.Append(SMC->GetCustomPrimitiveData().Data);
+		InstancesCustomData.AddDefaulted(NumCustomDataFloats - SMC->GetCustomPrimitiveData().Data.Num());
+	}
+}
+
+void FISMComponentBatcher::InitComponent(UInstancedStaticMeshComponent* ISMComponent) const
+{
+	ISMComponent->NumCustomDataFloats = NumCustomDataFloats;
+	ISMComponent->AddInstances(InstancesTransformsWS, /*bShouldReturnIndices*/false, /*bWorldSpace*/true);
+	ISMComponent->PerInstanceSMCustomData = InstancesCustomData;
+
+	if (!RandomSeeds.IsEmpty())
+	{
+		ISMComponent->InstancingRandomSeed = RandomSeeds[0].RandomSeed;
+
+		if (RandomSeeds.Num() > 1)
+		{
+			ISMComponent->AdditionalRandomSeeds = TArrayView<FInstancedStaticMeshRandomSeed>(&const_cast<FISMComponentBatcher*>(this)->RandomSeeds[1], RandomSeeds.Num() - 1);
+		}
+	}		
+}
\ No newline at end of file
diff --git a/Engine/Source/Runtime/Engine/Private/PackedLevelActor/PackedLevelActorISMBuilder.cpp b/Engine/Source/Runtime/Engine/Private/PackedLevelActor/PackedLevelActorISMBuilder.cpp
--- a/Engine/Source/Runtime/Engine/Private/PackedLevelActor/PackedLevelActorISMBuilder.cpp
+++ b/Engine/Source/Runtime/Engine/Private/PackedLevelActor/PackedLevelActorISMBuilder.cpp
@@ -14,6 +14,8 @@
 #include "Components/HierarchicalInstancedStaticMeshComponent.h"
 #include "Engine/StaticMesh.h"
 
+#include "ISMPartition/ISMComponentBatcher.h"
+
 FPackedLevelActorBuilderID FPackedLevelActorISMBuilder::BuilderID = 'ISMP';
 
 FPackedLevelActorBuilderID FPackedLevelActorISMBuilder::GetID() const
@@ -42,49 +44,28 @@ void FPackedLevelActorISMBuilder::PackActors(FPackedLevelActorBuilderContext& In
 	check(InClusterID.GetBuilderID() == GetID());
 	FTransform ActorTransform = InPackingActor->GetActorTransform();
 
-	TArray<FTransform> InstanceTransforms;
-	for (UActorComponent* Component : InComponents)
-	{
-		// If we have a ISM we need to add all instances
-		if (UInstancedStaticMeshComponent* ISMComponent = Cast<UInstancedStaticMeshComponent>(Component))
-		{
-			for (int32 InstanceIndex = 0; InstanceIndex < ISMComponent->GetInstanceCount(); ++InstanceIndex)
-			{
-				FTransform InstanceTransform;
-
-				if (ensure(ISMComponent->GetInstanceTransform(InstanceIndex, InstanceTransform, /*bWorldSpace=*/ true)))
-				{
-					// WorldSpace -> ActorSpace -> Apply pivot change
-					InstanceTransforms.Add(InstanceTransform.GetRelativeTransform(ActorTransform) * InContext.GetRelativePivotTransform());
-				}
-			}
-		}
-		else // other subclasses are processed like regular UStaticMeshComponent
-		{
-			UStaticMeshComponent* StaticMeshComponent = CastChecked<UStaticMeshComponent>(Component);
-
-			// WorldSpace -> ActorSpace -> Apply pivot change
-			InstanceTransforms.Add(StaticMeshComponent->GetComponentTransform().GetRelativeTransform(ActorTransform) * InContext.GetRelativePivotTransform());
-		}
-	}
-
 	FPackedLevelActorISMBuilderCluster* ISMCluster = (FPackedLevelActorISMBuilderCluster*)InClusterID.GetData();
 	check(ISMCluster);
 
+	FISMComponentBatcher ISMComponentBatcher;
+	ISMComponentBatcher.Append(InComponents, [&ActorTransform, &InContext](const FTransform& InTransform) { return InTransform.GetRelativeTransform(ActorTransform) * InContext.GetRelativePivotTransform() * ActorTransform; });
+
 	TSubclassOf<UInstancedStaticMeshComponent> ComponentClass = UInstancedStaticMeshComponent::StaticClass();
-	if (InstanceTransforms.Num() > 1 && ISMCluster->ISMDescriptor.StaticMesh && !ISMCluster->ISMDescriptor.StaticMesh->NaniteSettings.bEnabled)
+	if (ISMComponentBatcher.GetNumInstances() > 1 && ISMCluster->ISMDescriptor.StaticMesh && !ISMCluster->ISMDescriptor.StaticMesh->NaniteSettings.bEnabled)
 	{
 		// Use HISM for non-nanite when there is more than one transform (no use in using HISM cpu occlusion for a single instance)
 		ComponentClass = UHierarchicalInstancedStaticMeshComponent::StaticClass();
 	}
-		
+
 	UInstancedStaticMeshComponent* PackComponent = InPackingActor->AddPackedComponent<UInstancedStaticMeshComponent>(ComponentClass);
-		
 	PackComponent->AttachToComponent(InPackingActor->GetRootComponent(), FAttachmentTransformRules::SnapToTargetIncludingScale);
 
+	// Initialize the ISM properties using the ISM descriptor
 	ISMCluster->ISMDescriptor.InitComponent(PackComponent);
 
-	PackComponent->AddInstances(InstanceTransforms, /*bShouldReturnIndices*/false, /*bWorldSpace*/false);
+	// Initialize the ISM instances using the ISM batcher
+	ISMComponentBatcher.InitComponent(PackComponent);
+
 	PackComponent->RegisterComponent();
 }
 
diff --git a/Engine/Source/Runtime/Engine/Private/WorldPartition/HLOD/HLODBuilder.cpp b/Engine/Source/Runtime/Engine/Private/WorldPartition/HLOD/HLODBuilder.cpp
--- a/Engine/Source/Runtime/Engine/Private/WorldPartition/HLOD/HLODBuilder.cpp
+++ b/Engine/Source/Runtime/Engine/Private/WorldPartition/HLOD/HLODBuilder.cpp
@@ -7,6 +7,7 @@
 #include "Engine/StaticMesh.h"
 #include "Engine/Texture.h"
 #include "Components/InstancedStaticMeshComponent.h"
+#include "ISMPartition/ISMComponentBatcher.h"
 #include "ISMPartition/ISMComponentDescriptor.h"
 #include "Materials/MaterialInterface.h"
 #include "UObject/Package.h"
@@ -171,19 +172,6 @@ namespace
 			ComputeHash();
 		}
 	};
-
-	// Store batched instances data
-	struct FInstancingData
-	{
-		int32							NumInstances = 0;
-
-		TArray<FTransform>				InstancesTransforms;
-
-		int32							NumCustomDataFloats = 0;
-		TArray<float>					InstancesCustomData;
-
-		TArray<FInstancedStaticMeshRandomSeed> RandomSeeds;
-	};
 }
 
 TArray<UActorComponent*> UHLODBuilder::BatchInstances(const TArray<UActorComponent*>& InSourceComponents)
@@ -193,64 +181,12 @@ TArray<UActorComponent*> UHLODBuilder::BatchInstances(const TArray<UActorCompone
 	TArray<UStaticMeshComponent*> SourceStaticMeshComponents = FilterComponents<UStaticMeshComponent>(InSourceComponents);
 
 	// Prepare instance batches
-	TMap<FISMComponentDescriptor, FInstancingData> InstancesData;
+	TMap<FISMComponentDescriptor, FISMComponentBatcher> InstancesData;
 	for (UStaticMeshComponent* SMC : SourceStaticMeshComponents)
 	{
 		FCustomISMComponentDescriptor ISMComponentDescriptor(SMC);
-		FInstancingData& InstancingData = InstancesData.FindOrAdd(ISMComponentDescriptor);
-
-		if (UInstancedStaticMeshComponent* ISMC = Cast<UInstancedStaticMeshComponent>(SMC))
-		{
-			InstancingData.NumCustomDataFloats = FMath::Max(InstancingData.NumCustomDataFloats, ISMC->NumCustomDataFloats);
-			InstancingData.RandomSeeds.Add({ InstancingData.NumInstances, ISMC->InstancingRandomSeed });
-			InstancingData.NumInstances += ISMC->GetInstanceCount();
-		}
-		else
-		{
-			InstancingData.NumCustomDataFloats = FMath::Max(InstancingData.NumCustomDataFloats, SMC->GetCustomPrimitiveData().Data.Num());
-			InstancingData.NumInstances++;
-		}
-	}
-
-	// Resize arrays
-	for (auto& Entry : InstancesData)
-	{
-		const FISMComponentDescriptor& ISMComponentDescriptor = Entry.Key;
-		FInstancingData& EntryInstancingData = Entry.Value;
-
-		EntryInstancingData.InstancesTransforms.Reset(EntryInstancingData.NumInstances);
-		EntryInstancingData.InstancesCustomData.Reset(EntryInstancingData.NumInstances * EntryInstancingData.NumCustomDataFloats);
-	}
-
-	// Append all transforms & per instance custom data
-	for (UStaticMeshComponent* SMC : SourceStaticMeshComponents)
-	{
-		FCustomISMComponentDescriptor ISMComponentDescriptor(SMC);
-		FInstancingData& InstancingData = InstancesData.FindChecked(ISMComponentDescriptor);
-
-		if (UInstancedStaticMeshComponent* ISMC = Cast<UInstancedStaticMeshComponent>(SMC))
-		{
-			// Add transforms
-			for (int32 InstanceIdx = 0; InstanceIdx < ISMC->GetInstanceCount(); InstanceIdx++)
-			{
-				FTransform& InstanceTransform = InstancingData.InstancesTransforms.AddDefaulted_GetRef();
-				ISMC->GetInstanceTransform(InstanceIdx, InstanceTransform, true);
-			}
-
-			// Add per instance custom data
-			int32 NumCustomDataFloatToAdd = ISMC->GetInstanceCount() * InstancingData.NumCustomDataFloats;
-			InstancingData.InstancesCustomData.Append(ISMC->PerInstanceSMCustomData);
-			InstancingData.InstancesCustomData.AddDefaulted(NumCustomDataFloatToAdd - ISMC->PerInstanceSMCustomData.Num());
-		}
-		else
-		{
-			// Add transform
-			InstancingData.InstancesTransforms.Add(SMC->GetComponentTransform());
-
-			// Add custom data
-			InstancingData.InstancesCustomData.Append(SMC->GetCustomPrimitiveData().Data);
-			InstancingData.InstancesCustomData.AddDefaulted(InstancingData.NumCustomDataFloats - SMC->GetCustomPrimitiveData().Data.Num());
-		}
+		FISMComponentBatcher& ISMComponentBatcher = InstancesData.FindOrAdd(ISMComponentDescriptor);
+		ISMComponentBatcher.Add(SMC);
 	}
 
 	// Create an ISMC for each SM asset we found
@@ -258,27 +194,14 @@ TArray<UActorComponent*> UHLODBuilder::BatchInstances(const TArray<UActorCompone
 	for (auto& Entry : InstancesData)
 	{
 		const FISMComponentDescriptor& ISMComponentDescriptor = Entry.Key;
-		FInstancingData& EntryInstancingData = Entry.Value;
+		const FISMComponentBatcher& ISMComponentBatcher = Entry.Value;
 
-		check(EntryInstancingData.InstancesTransforms.Num() * EntryInstancingData.NumCustomDataFloats == EntryInstancingData.InstancesCustomData.Num());
+		UInstancedStaticMeshComponent* ISMComponent = ISMComponentDescriptor.CreateComponent(GetTransientPackage());
+		ISMComponentBatcher.InitComponent(ISMComponent);
 
-		UInstancedStaticMeshComponent* Component = ISMComponentDescriptor.CreateComponent(GetTransientPackage());
-		Component->SetForcedLodModel(Component->GetStaticMesh()->GetNumLODs());
-		Component->NumCustomDataFloats = EntryInstancingData.NumCustomDataFloats;
-		Component->AddInstances(EntryInstancingData.InstancesTransforms, /*bShouldReturnIndices*/false, /*bWorldSpace*/true);
-		Component->PerInstanceSMCustomData = MoveTemp(EntryInstancingData.InstancesCustomData);
-
-		if (!EntryInstancingData.RandomSeeds.IsEmpty())
-		{
-			Component->InstancingRandomSeed = EntryInstancingData.RandomSeeds[0].RandomSeed;
-		}
-
-		if (EntryInstancingData.RandomSeeds.Num() > 1)
-		{
-			Component->AdditionalRandomSeeds = TArrayView<FInstancedStaticMeshRandomSeed>(&EntryInstancingData.RandomSeeds[1], EntryInstancingData.RandomSeeds.Num() - 1);
-		}
+		ISMComponent->SetForcedLodModel(ISMComponent->GetStaticMesh()->GetNumLODs());
 
-		HLODComponents.Add(Component);
+		HLODComponents.Add(ISMComponent);
 	};
 
 	return HLODComponents;
diff --git a/Engine/Source/Runtime/Engine/Public/ISMPartition/ISMComponentBatcher.h b/Engine/Source/Runtime/Engine/Public/ISMPartition/ISMComponentBatcher.h
new file mode 100644
--- /dev/null
+++ b/Engine/Source/Runtime/Engine/Public/ISMPartition/ISMComponentBatcher.h
@@ -0,0 +1,75 @@
+// Copyright Epic Games, Inc. All Rights Reserved.
+#pragma once
+
+#include "CoreMinimal.h"
+#include "Components/InstancedStaticMeshComponent.h"
+#include "Engine/EngineTypes.h"
+#include "Templates/EnableIf.h"
+#include "Templates/UnrealTypeTraits.h"
+
+
+/** Struct that allows batching of transforms and custom data of multiple (possibly instanced) static mesh components */
+struct ENGINE_API FISMComponentBatcher
+{
+public:
+	/**
+	 * Add a single component to be batched
+	 * @param	InActorComponent	Component to be batched
+	 */
+	void Add(const UActorComponent* InActorComponent);
+
+	/**
+	 * Add a single component to be batched
+	 * @param	InActorComponent	Component to be batched
+	 * @param	InTransformFunc		Function that takes the world space transform of an instance and modifies it. Must return a world space transform.
+	 */
+	void Add(const UActorComponent* InActorComponent, TFunctionRef<FTransform(const FTransform&)> InTransformFunc);
+	
+	/**
+	 * Add an array of component to be batched
+	 * @param	InComponents		Components to be batched
+	 */
+	template<typename TComponentClass = UActorComponent, typename = typename TEnableIf<TIsDerivedFrom<TComponentClass, UActorComponent>::IsDerived>::Type>
+	void Append(const TArray<TComponentClass*>& InComponents)
+	{
+		for (const UActorComponent* InComponent : InComponents)
+		{
+			AddInternal(InComponent, TOptional<TFunctionRef<FTransform(const FTransform&)>>());
+		}
+	}
+
+	/**
+	 * Add an array of components to be batched
+	 * @param	InComponents		Components to be batched
+	 * @param	InTransformFunc		Function that takes the world space transform of an instance and modifies it. Must return a world space transform.
+	 */
+	template<typename TComponentClass = UActorComponent, typename = typename TEnableIf<TIsDerivedFrom<TComponentClass, UActorComponent>::IsDerived>::Type>
+	void Append(const TArray<TComponentClass*>& InComponents, TFunctionRef<FTransform(const FTransform&)> InTransformFunc)
+	{
+		for (const UActorComponent* InComponent : InComponents)
+		{
+			AddInternal(InComponent, InTransformFunc);
+		}
+	}
+
+	/**
+	 * Return the number of instances batched so far.
+	 */
+	int32 GetNumInstances() const { return NumInstances; }
+
+	/**
+	 * Initialize the instances of the provided ISM component using the batched data stored in this class.
+	 * @param	ISMComponent	Instanced static mesh component which will be modified.
+	 */
+	void InitComponent(UInstancedStaticMeshComponent* ISMComponent) const;
+
+private:
+	void AddInternal(const UActorComponent* InComponent, TOptional<TFunctionRef<FTransform(const FTransform&)>> InTransformFunc);
+
+	int32 NumInstances = 0;
+	int32 NumCustomDataFloats = 0;
+
+	TArray<FTransform> InstancesTransformsWS;
+	TArray<float> InstancesCustomData;
+	TArray<FInstancedStaticMeshRandomSeed> RandomSeeds;
+};
