Add profile tags to ContentBundle heavier operations

diff --git a/Engine/Source/Runtime/Engine/Private/WorldPartition/ContentBundle/ContentBundleBase.cpp b/Engine/Source/Runtime/Engine/Private/WorldPartition/ContentBundle/ContentBundleBase.cpp
--- a/Engine/Source/Runtime/Engine/Private/WorldPartition/ContentBundle/ContentBundleBase.cpp
+++ b/Engine/Source/Runtime/Engine/Private/WorldPartition/ContentBundle/ContentBundleBase.cpp
@@ -22,6 +22,8 @@ FContentBundleBase::~FContentBundleBase()
 
 void FContentBundleBase::Initialize()
 {
+	TRACE_CPUPROFILER_EVENT_SCOPE(FContentBundleBase::Initialize);
+
 	check(GetStatus() == EContentBundleStatus::Unknown);
 
 	DoInitialize();
@@ -31,6 +33,8 @@ void FContentBundleBase::Initialize()
 
 void FContentBundleBase::Uninitialize()
 {
+	TRACE_CPUPROFILER_EVENT_SCOPE(FContentBundleBase::Uninitialize);
+
 	check(GetStatus() != EContentBundleStatus::Unknown);
 
 	if (GetStatus() == EContentBundleStatus::ReadyToInject || GetStatus() == EContentBundleStatus::ContentInjected)
@@ -45,6 +49,8 @@ void FContentBundleBase::Uninitialize()
 
 void FContentBundleBase::InjectContent()
 {
+	TRACE_CPUPROFILER_EVENT_SCOPE(FContentBundleBase::InjectContent);
+
 	check(GetStatus() == EContentBundleStatus::Registered);
 
 	DoInjectContent();
@@ -59,6 +65,8 @@ void FContentBundleBase::InjectContent()
 
 void FContentBundleBase::RemoveContent()
 {
+	TRACE_CPUPROFILER_EVENT_SCOPE(FContentBundleBase::RemoveContent);
+
 	check(GetStatus() == EContentBundleStatus::ReadyToInject || GetStatus() == EContentBundleStatus::ContentInjected || GetStatus() == EContentBundleStatus::FailedToInject);
 
 	if (GetStatus() == EContentBundleStatus::ContentInjected || GetStatus() == EContentBundleStatus::ReadyToInject)
diff --git a/Engine/Source/Runtime/Engine/Private/WorldPartition/ContentBundle/ContentBundleContainer.cpp b/Engine/Source/Runtime/Engine/Private/WorldPartition/ContentBundle/ContentBundleContainer.cpp
--- a/Engine/Source/Runtime/Engine/Private/WorldPartition/ContentBundle/ContentBundleContainer.cpp
+++ b/Engine/Source/Runtime/Engine/Private/WorldPartition/ContentBundle/ContentBundleContainer.cpp
@@ -62,6 +62,8 @@ void FContentBundleContainer::Initialize()
 
 void FContentBundleContainer::Deinitialize()
 {
+	TRACE_CPUPROFILER_EVENT_SCOPE(FContentBundleContainer::Deinitialize);
+
 	if (GetInjectedWorld())
 	{
 		UE_LOG(LogContentBundle, Log, TEXT("[Container: %s] Deleting container."), *GetInjectedWorld()->GetName());
@@ -268,19 +270,16 @@ void FContentBundleContainer::DeinitializeContentBundle(FContentBundleBase& Cont
 
 void FContentBundleContainer::InjectContentBundle(FContentBundleBase& ContentBundle)
 {
-	TRACE_CPUPROFILER_EVENT_SCOPE(InjectContentBundle);
 	ContentBundle.InjectContent();
 }
 
 void FContentBundleContainer::RemoveContentBundle(FContentBundleBase& ContentBundle)
 {
-	TRACE_CPUPROFILER_EVENT_SCOPE(RemoveContentBundle);
 	ContentBundle.RemoveContent();
 }
 
 void FContentBundleContainer::OnContentBundleClientRegistered(TSharedPtr<FContentBundleClient>& ContentBundleClient)
 {
-	TRACE_CPUPROFILER_EVENT_SCOPE(InitializeContentBundle);
 	InitializeContentBundle(ContentBundleClient);
 }
 
diff --git a/Engine/Source/Runtime/Engine/Private/WorldPartition/ContentBundle/ContentBundleEditor.cpp b/Engine/Source/Runtime/Engine/Private/WorldPartition/ContentBundle/ContentBundleEditor.cpp
--- a/Engine/Source/Runtime/Engine/Private/WorldPartition/ContentBundle/ContentBundleEditor.cpp
+++ b/Engine/Source/Runtime/Engine/Private/WorldPartition/ContentBundle/ContentBundleEditor.cpp
@@ -306,6 +306,8 @@ void FContentBundleEditor::AddReferencedObjects(FReferenceCollector& Collector)
 
 void FContentBundleEditor::GenerateStreaming(TArray<FString>* OutPackageToGenerate)
 {
+	TRACE_CPUPROFILER_EVENT_SCOPE(FContentBundleEditor::GenerateStreaming);
+
 	if (GetStatus() != EContentBundleStatus::ContentInjected)
 	{
 		UE_LOG(LogContentBundle, Log, TEXT("[CB: %s] Skipping streaming generation. It's status is: %s."), *GetDescriptor()->GetDisplayName(), *UEnum::GetDisplayValueAsText(GetStatus()).ToString());
@@ -384,7 +386,6 @@ bool FContentBundleEditor::PopulateGeneratorPackageForCook(class IWorldPartition
 {
 	bool bIsSuccess = true;
 
-
 	if (HasCookedContent())
 	{
 		UE_LOG(LogContentBundle, Log, TEXT("[CB: %s][Cook] Populating Generator Package. %u Packages"), *GetDescriptor()->GetDisplayName(), PackagesToCook.Num());
