Enabled logging in Shipping/Test in KismetSystemLibrary.cpp when USE_LOGGING_IN_SHIPPING is defined to 1.

diff --git a/Engine/Source/Runtime/Engine/Private/KismetSystemLibrary.cpp b/Engine/Source/Runtime/Engine/Private/KismetSystemLibrary.cpp
--- a/Engine/Source/Runtime/Engine/Private/KismetSystemLibrary.cpp
+++ b/Engine/Source/Runtime/Engine/Private/KismetSystemLibrary.cpp
@@ -319,7 +319,7 @@ UObject* UKismetSystemLibrary::Conv_InterfaceToObject(const FScriptInterface& In
 
 void UKismetSystemLibrary::LogString(const FString& InString, bool bPrintToLog)
 {
-#if !(UE_BUILD_SHIPPING || UE_BUILD_TEST) // Do not Print in Shipping or Test
+#if !(UE_BUILD_SHIPPING || UE_BUILD_TEST) || USE_LOGGING_IN_SHIPPING // Do not Print in Shipping or Test unless explictly enabled.
 	if(bPrintToLog)
 	{
 		UE_LOG(LogBlueprintUserMessages, Log, TEXT("%s"), *InString);
@@ -333,7 +333,7 @@ void UKismetSystemLibrary::LogString(const FString& InString, bool bPrintToLog)
 
 void UKismetSystemLibrary::PrintString(const UObject* WorldContextObject, const FString& InString, bool bPrintToScreen, bool bPrintToLog, FLinearColor TextColor, float Duration, const FName Key)
 {
-#if !(UE_BUILD_SHIPPING || UE_BUILD_TEST) // Do not Print in Shipping or Test
+#if !(UE_BUILD_SHIPPING || UE_BUILD_TEST) || USE_LOGGING_IN_SHIPPING // Do not Print in Shipping or Test unless explictly enabled.
 
 	UWorld* World = GEngine->GetWorldFromContextObject(WorldContextObject, EGetWorldErrorMode::ReturnNull);
 	FString Prefix;
@@ -419,14 +419,14 @@ void UKismetSystemLibrary::PrintString(const UObject* WorldContextObject, const
 
 void UKismetSystemLibrary::PrintText(const UObject* WorldContextObject, const FText InText, bool bPrintToScreen, bool bPrintToLog, FLinearColor TextColor, float Duration, const FName Key)
 {
-#if !(UE_BUILD_SHIPPING || UE_BUILD_TEST) // Do not Print in Shipping or Test
+#if !(UE_BUILD_SHIPPING || UE_BUILD_TEST) || USE_LOGGING_IN_SHIPPING // Do not Print in Shipping or Test unless explictly enabled.
 	PrintString(WorldContextObject, InText.ToString(), bPrintToScreen, bPrintToLog, TextColor, Duration, Key);
 #endif
 }
 
 void UKismetSystemLibrary::PrintWarning(const FString& InString)
 {
-#if !(UE_BUILD_SHIPPING || UE_BUILD_TEST) // Do not Print in Shipping or Test
+#if !(UE_BUILD_SHIPPING || UE_BUILD_TEST) || USE_LOGGING_IN_SHIPPING // Do not Print in Shipping or Test unless explictly enabled.
 	PrintString(NULL, InString, true, true);
 #endif
 }
