Fixed CIS complaining about usage of FBlueprintContextTracker in shipping.

diff --git a/Engine/Source/Runtime/Engine/Private/KismetSystemLibrary.cpp b/Engine/Source/Runtime/Engine/Private/KismetSystemLibrary.cpp
--- a/Engine/Source/Runtime/Engine/Private/KismetSystemLibrary.cpp
+++ b/Engine/Source/Runtime/Engine/Private/KismetSystemLibrary.cpp
@@ -18,6 +18,7 @@
 #include "GameFramework/Actor.h"
 #include "CollisionQueryParams.h"
 #include "UObject/Package.h"
+#include "UObject/Script.h"
 #include "WorldCollision.h"
 #include "Components/PrimitiveComponent.h"
 #include "Engine/CollisionProfile.h"
@@ -358,6 +359,7 @@ void UKismetSystemLibrary::PrintString(const UObject* WorldContextObject, const
 		}
 	}
 
+#if DO_BLUEPRINT_GUARD
 	if (UE::Blueprint::Private::bBlamePrintString && !FBlueprintContextTracker::Get().GetCurrentScriptStack().IsEmpty())
 	{
 		const TArrayView<const FFrame* const> ScriptStack = FBlueprintContextTracker::Get().GetCurrentScriptStack();
@@ -366,6 +368,7 @@ void UKismetSystemLibrary::PrintString(const UObject* WorldContextObject, const
 			*ScriptStack.Last()->Node->GetName(),
 			*Prefix);
 	}
+#endif
 
 	const FString FinalDisplayString = Prefix + InString;
 	FString FinalLogString = FinalDisplayString;
