World Partition - HLOD: Fixed issue with HLODs being built before construction scripts were fully executed on their source actors

diff --git a/Engine/Plugins/Editor/WorldPartitionHLODUtilities/Source/Private/WorldPartitionHLODUtilities.cpp b/Engine/Plugins/Editor/WorldPartitionHLODUtilities/Source/Private/WorldPartitionHLODUtilities.cpp
--- a/Engine/Plugins/Editor/WorldPartitionHLODUtilities/Source/Private/WorldPartitionHLODUtilities.cpp
+++ b/Engine/Plugins/Editor/WorldPartitionHLODUtilities/Source/Private/WorldPartitionHLODUtilities.cpp
@@ -56,7 +56,15 @@ static UWorldPartitionLevelStreamingDynamic* CreateLevelStreamingFromHLODActor(A
 	UWorldPartitionLevelStreamingDynamic* LevelStreaming = UWorldPartitionLevelStreamingDynamic::LoadInEditor(World, LevelStreamingName, Mappings);
 	check(LevelStreaming);
 
-	if (!LevelStreaming->GetLoadSucceeded())
+	if (LevelStreaming->GetLoadSucceeded())
+	{
+		// Finish assets compilation
+		FAssetCompilingManager::Get().FinishAllCompilation();
+
+		// Ensure all deferred construction scripts are executed
+		FAssetCompilingManager::Get().ProcessAsyncTasks();
+	}
+	else
 	{
 		bOutDirty = true;
 		UE_LOG(LogHLODBuilder, Warning, TEXT("HLOD actor \"%s\" needs to be rebuilt as it didn't succeed in loading all actors."), *InHLODActor->GetActorLabel());
diff --git a/Engine/Plugins/Editor/WorldPartitionHLODUtilities/Source/Public/HLODBuilderInstancing.h b/Engine/Plugins/Editor/WorldPartitionHLODUtilities/Source/Public/HLODBuilderInstancing.h
--- a/Engine/Plugins/Editor/WorldPartitionHLODUtilities/Source/Public/HLODBuilderInstancing.h
+++ b/Engine/Plugins/Editor/WorldPartitionHLODUtilities/Source/Public/HLODBuilderInstancing.h
@@ -31,7 +31,6 @@ class WORLDPARTITIONHLODUTILITIES_API UHLODBuilderInstancing : public UHLODBuild
 	 GENERATED_UCLASS_BODY()
 
 public:
-	virtual bool RequiresCompiledAssets() const override { return false; }
 	virtual bool RequiresWarmup() const override { return false; }
 
 	virtual TSubclassOf<UHLODBuilderSettings> GetSettingsClass() const override;
diff --git a/Engine/Source/Runtime/Engine/Private/WorldPartition/HLOD/HLODBuilder.cpp b/Engine/Source/Runtime/Engine/Private/WorldPartition/HLOD/HLODBuilder.cpp
--- a/Engine/Source/Runtime/Engine/Private/WorldPartition/HLOD/HLODBuilder.cpp
+++ b/Engine/Source/Runtime/Engine/Private/WorldPartition/HLOD/HLODBuilder.cpp
@@ -46,11 +46,6 @@ void UHLODBuilder::SetHLODBuilderSettings(const UHLODBuilderSettings* InHLODBuil
 	HLODBuilderSettings = InHLODBuilderSettings;
 }
 
-bool UHLODBuilder::RequiresCompiledAssets() const
-{
-	return true;
-}
-
 bool UHLODBuilder::RequiresWarmup() const
 {
 	return true;
@@ -264,17 +259,6 @@ TArray<UActorComponent*> UHLODBuilder::Build(const FHLODBuildContext& InHLODBuil
 		HLODBuildersForComponents.FindOrAdd(HLODBuilderClass).Add(SourceComponent);
 	}
 
-	// If any of the builders requires it, wait for assets compilation to finish
-	for (const auto& HLODBuilderPair : HLODBuildersForComponents)
-	{
-		const UHLODBuilder* HLODBuilder = HLODBuilderPair.Key ? HLODBuilderPair.Key->GetDefaultObject<UHLODBuilder>() : this;
-		if (HLODBuilder->RequiresCompiledAssets())
-		{
-			FAssetCompilingManager::Get().FinishAllCompilation();
-			break;
-		}
-	}	
-	
 	// Build HLOD components by sending source components to the individual builders, in batch
 	TArray<UActorComponent*> HLODComponents;
 	for (const auto& HLODBuilderPair : HLODBuildersForComponents)
diff --git a/Engine/Source/Runtime/Engine/Public/WorldPartition/HLOD/HLODBuilder.h b/Engine/Source/Runtime/Engine/Public/WorldPartition/HLOD/HLODBuilder.h
--- a/Engine/Source/Runtime/Engine/Public/WorldPartition/HLOD/HLODBuilder.h
+++ b/Engine/Source/Runtime/Engine/Public/WorldPartition/HLOD/HLODBuilder.h
@@ -87,11 +87,6 @@ public:
 	 */
 	virtual TSubclassOf<UHLODBuilderSettings> GetSettingsClass() const;
 
-	/**
-	 * Should return true if this builder needs to have it's input assets fully compiled.
-	 */
-	virtual bool RequiresCompiledAssets() const;
-
 	/**
 	 * Should return true if components generated from this builder need a warmup phase before being made visible.
 	 * If your components are using virtual textures or Nanite meshes, this should return true, as it will be necessary
@@ -154,7 +149,6 @@ class ENGINE_API UNullHLODBuilder : public UHLODBuilder
 	GENERATED_UCLASS_BODY()
 
 #if WITH_EDITOR
-	virtual bool RequiresCompiledAssets() const { return false; }
 	virtual bool RequiresWarmup() const { return false; }
 	virtual uint32 ComputeHLODHash(const UActorComponent* InSourceComponent) const { return 0; }
 	virtual TArray<UActorComponent*> Build(const FHLODBuildContext& InHLODBuildContext, const TArray<UActorComponent*>& InSourceComponents) const { return {}; }
