Release RayTracingGeometry for skeletal meshes when dynamic ray tracing is enabled.

diff --git a/Engine/Source/Runtime/Engine/Private/SkeletalMeshLODRenderData.cpp b/Engine/Source/Runtime/Engine/Private/SkeletalMeshLODRenderData.cpp
--- a/Engine/Source/Runtime/Engine/Private/SkeletalMeshLODRenderData.cpp
+++ b/Engine/Source/Runtime/Engine/Private/SkeletalMeshLODRenderData.cpp
@@ -271,7 +271,7 @@ void FSkeletalMeshLODRenderData::InitResources(bool bNeedsVertexColors, int32 LO
 	}
 
 #if RHI_RAYTRACING
-	if (IsRayTracingEnabled())
+	if (IsRayTracingAllowed())
 	{
 		if (SourceRayTracingGeometry.RawData.Num() > 0)
 		{
@@ -302,7 +302,7 @@ void FSkeletalMeshLODRenderData::ReleaseResources()
 	SkinWeightProfilesData.ReleaseResources();
 
 #if RHI_RAYTRACING
-	if (IsRayTracingEnabled())
+	if (IsRayTracingAllowed())
 	{
 		BeginReleaseResource(&SourceRayTracingGeometry);
 		BeginReleaseResource(&StaticRayTracingGeometry);
