Iris
- Bumped static priority for PlayerController when using iris-replication

diff --git a/Engine/Source/Runtime/Engine/Private/PlayerController.cpp b/Engine/Source/Runtime/Engine/Private/PlayerController.cpp
--- a/Engine/Source/Runtime/Engine/Private/PlayerController.cpp
+++ b/Engine/Source/Runtime/Engine/Private/PlayerController.cpp
@@ -6015,13 +6015,27 @@ void APlayerController::RemoveFromNetConditionGroup(FName NetGroup)
 #if UE_WITH_IRIS
 void APlayerController::BeginReplication()
 {
+	using namespace UE::Net;
+
 	// Always allow the PlayerController to be replicated as it is required for travel.
 	FActorBeginReplicationParams Params;
 	Params.bIncludeInLevelGroupFilter = false;
 	Super::BeginReplication(Params);
 
+	// Bump prio of playercontroller in order to make sure it replicates really early
+	const FNetHandle Handle = FReplicationSystemUtil::GetNetHandle(this);
+	
+	if (Handle.IsValid())
+	{
+		UReplicationSystem* ReplicationSystem = GetReplicationSystem(Handle.GetReplicationSystemId());
+
+		// $IRIS TODO: Add config support for overriding static prio
+		static constexpr float PlayerControllerStaticPriority = 100.f;
+		ReplicationSystem->SetStaticPriority(Handle, PlayerControllerStaticPriority);
+	}
+
 	// Enable groups once owner is set!!
-	UE::Net::FReplicationSystemUtil::UpdateSubObjectGroupMemberships(this);
+	FReplicationSystemUtil::UpdateSubObjectGroupMemberships(this);
 }
 #endif // UE_WITH_IRIS
 
