Iris
- Fixed compilation error in PlayerController.cpp

diff --git a/Engine/Source/Runtime/Engine/Private/PlayerController.cpp b/Engine/Source/Runtime/Engine/Private/PlayerController.cpp
--- a/Engine/Source/Runtime/Engine/Private/PlayerController.cpp
+++ b/Engine/Source/Runtime/Engine/Private/PlayerController.cpp
@@ -6022,16 +6022,17 @@ void APlayerController::BeginReplication()
 	Params.bIncludeInLevelGroupFilter = false;
 	Super::BeginReplication(Params);
 
-	// Bump prio of playercontroller in order to make sure it replicates really early
-	const FNetHandle Handle = FReplicationSystemUtil::GetNetHandle(this);
-	
-	if (Handle.IsValid())
+	if (UReplicationSystem* ReplicationSystem = FReplicationSystemUtil::GetReplicationSystem(this))
 	{
-		UReplicationSystem* ReplicationSystem = GetReplicationSystem(Handle.GetReplicationSystemId());
+		if (UActorReplicationBridge* Bridge = Cast<UActorReplicationBridge>(ReplicationSystem->GetReplicationBridge()))
+		{
+			// Bump prio of playercontroller in order to make sure it replicates really early
+			const FNetRefHandle RefHandle = Bridge->GetReplicatedRefHandle(this);	
 
-		// $IRIS TODO: Add config support for overriding static prio
-		static constexpr float PlayerControllerStaticPriority = 100.f;
-		ReplicationSystem->SetStaticPriority(Handle, PlayerControllerStaticPriority);
+			// $IRIS TODO: Add config support for overriding static prio
+			static constexpr float PlayerControllerStaticPriority = 100.f;
+			ReplicationSystem->SetStaticPriority(RefHandle, PlayerControllerStaticPriority);
+		}
 	}
 
 	// Enable groups once owner is set!!
