Iris
- Added static helper in FReplicationSystemUtil::SetStaticPriority to forward call to all ReplicationSystems

diff --git a/Engine/Source/Runtime/Engine/Private/Net/Iris/ReplicationSystem/ReplicationSystemUtil.cpp b/Engine/Source/Runtime/Engine/Private/Net/Iris/ReplicationSystem/ReplicationSystemUtil.cpp
--- a/Engine/Source/Runtime/Engine/Private/Net/Iris/ReplicationSystem/ReplicationSystemUtil.cpp
+++ b/Engine/Source/Runtime/Engine/Private/Net/Iris/ReplicationSystem/ReplicationSystemUtil.cpp
@@ -571,6 +571,28 @@ void FReplicationSystemUtil::SetReplicationCondition(FNetHandle NetHandle, ERepl
 	});
 }
 
+void FReplicationSystemUtil::SetStaticPriority(const AActor* Actor, float Priority)
+{
+	FNetHandle ActorHandle = GetNetHandle(Actor);
+	if (!ActorHandle.IsValid())
+	{
+		return;
+	}
+	
+	ReplicationSystemUtil::ForEachReplicationSystem([ActorHandle, Priority](UReplicationSystem* ReplicationSystem)
+	{
+		if (ReplicationSystem->IsServer())
+		{
+			if (UObjectReplicationBridge* Bridge = ReplicationSystem->GetReplicationBridgeAs<UObjectReplicationBridge>())
+			{
+				const FNetRefHandle ActorRefHandle = Bridge->GetReplicatedRefHandle(ActorHandle);
+				ReplicationSystem->SetStaticPriority(ActorRefHandle, Priority);
+			}
+		}
+	});
+}
+
+
 
 }
 
diff --git a/Engine/Source/Runtime/Engine/Private/PlayerController.cpp b/Engine/Source/Runtime/Engine/Private/PlayerController.cpp
--- a/Engine/Source/Runtime/Engine/Private/PlayerController.cpp
+++ b/Engine/Source/Runtime/Engine/Private/PlayerController.cpp
@@ -6022,18 +6022,9 @@ void APlayerController::BeginReplication()
 	Params.bIncludeInLevelGroupFilter = false;
 	Super::BeginReplication(Params);
 
-	if (UReplicationSystem* ReplicationSystem = FReplicationSystemUtil::GetReplicationSystem(this))
-	{
-		if (UActorReplicationBridge* Bridge = Cast<UActorReplicationBridge>(ReplicationSystem->GetReplicationBridge()))
-		{
-			// Bump prio of playercontroller in order to make sure it replicates really early
-			const FNetRefHandle RefHandle = Bridge->GetReplicatedRefHandle(this);	
-
-			// $IRIS TODO: Add config support for overriding static prio
-			static constexpr float PlayerControllerStaticPriority = 100.f;
-			ReplicationSystem->SetStaticPriority(RefHandle, PlayerControllerStaticPriority);
-		}
-	}
+	// Bump prio of playercontroller in order to make sure it replicates really early
+	static constexpr float PlayerControllerStaticPriority = 100.f;
+	FReplicationSystemUtil::SetStaticPriority(this, PlayerControllerStaticPriority);
 
 	// Enable groups once owner is set!!
 	FReplicationSystemUtil::UpdateSubObjectGroupMemberships(this);
diff --git a/Engine/Source/Runtime/Engine/Public/Net/Iris/ReplicationSystem/ReplicationSystemUtil.h b/Engine/Source/Runtime/Engine/Public/Net/Iris/ReplicationSystem/ReplicationSystemUtil.h
--- a/Engine/Source/Runtime/Engine/Public/Net/Iris/ReplicationSystem/ReplicationSystemUtil.h
+++ b/Engine/Source/Runtime/Engine/Public/Net/Iris/ReplicationSystem/ReplicationSystemUtil.h
@@ -108,6 +108,12 @@ struct FReplicationSystemUtil
 
 	/** Enable or disable a replication condition. This will affect the replication of properties with conditions. */
 	ENGINE_API static void SetReplicationCondition(FNetHandle NetHandle, EReplicationCondition Condition, bool bEnableCondition);
+
+	/**
+	 * Sets a fixed priority for a replicated object 
+	 * @see ReplicationSystem.h for details
+	 */
+	ENGINE_API static void SetStaticPriority(const AActor* Actor, float Priority);
 };
 
 }
