Fix check() when loading skinned mesh without duplicated vertices on deformer graph supported platforms.

diff --git a/Engine/Source/Runtime/Engine/Private/SkeletalMeshLODRenderData.cpp b/Engine/Source/Runtime/Engine/Private/SkeletalMeshLODRenderData.cpp
--- a/Engine/Source/Runtime/Engine/Private/SkeletalMeshLODRenderData.cpp
+++ b/Engine/Source/Runtime/Engine/Private/SkeletalMeshLODRenderData.cpp
@@ -238,22 +238,24 @@ void FSkeletalMeshLODRenderData::InitResources(bool bNeedsVertexColors, int32 LO
 	}
 
 	// We can discard any the cooked DuplicatedVertices here based on runtime settings.
-	bool bNeedDuplicatedVertices = RequiresDuplicateVertices();
+	bool bDiscardDuplicatedVertices = !RequiresDuplicateVertices();
 	for (FSkelMeshRenderSection& RenderSection : RenderSections)
 	{
-		if (bNeedDuplicatedVertices)
-		{
-			// No need to discard CPU data in cooked builds as bNeedsCPUAccess is false (see FDuplicatedVerticesBuffer constructor), 
-			// so it'd be auto-discarded after the RHI has copied the resource data. Keep CPU data when in the editor for geometry operations.
-			check(RenderSection.DuplicatedVerticesBuffer.DupVertData.Num());
-			BeginInitResource(&RenderSection.DuplicatedVerticesBuffer);
-		}
-		else
+		if(RenderSection.DuplicatedVerticesBuffer.DupVertData.Num())
 		{
+			if (bDiscardDuplicatedVertices)
+			{
 #if !WITH_EDITOR
-			// Discard CPU data in cooked builds. Keep CPU data when in the editor for geometry operations.
-			RenderSection.DuplicatedVerticesBuffer.ReleaseCPUResources();
+				// Discard CPU data in cooked builds. Keep CPU data when in the editor for geometry operations.
+				RenderSection.DuplicatedVerticesBuffer.ReleaseCPUResources();
 #endif
+			}
+			else
+			{
+				// No need to discard CPU data in cooked builds as bNeedsCPUAccess is false (see FDuplicatedVerticesBuffer constructor), 
+				// so it'd be auto-discarded after the RHI has copied the resource data. Keep CPU data when in the editor for geometry operations.
+				BeginInitResource(&RenderSection.DuplicatedVerticesBuffer);
+			}
 		}
 	}
 
