Fixed UWorld::FOnActorRemovedFromWorld duplicated events.

#tests VK-Edit debugging tool test for duplicated actorremoval calls within actor memory tracker.

diff --git a/Engine/Source/Runtime/Engine/Private/LevelActor.cpp b/Engine/Source/Runtime/Engine/Private/LevelActor.cpp
--- a/Engine/Source/Runtime/Engine/Private/LevelActor.cpp
+++ b/Engine/Source/Runtime/Engine/Private/LevelActor.cpp
@@ -996,8 +996,6 @@ bool UWorld::DestroyActor( AActor* ThisActor, bool bNetForce, bool bShouldModify
 #endif
 	}
 
-	OnActorRemovedFromWorld.Broadcast(ThisActor);
-
 	// Clean up the actor's components.
 	ThisActor->UnregisterAllComponents();
 
