HLOD generation creates black textures

Some EditorOnlyData objects for Material Function have been saved with an incorrect name
When detecting this case, fixup the EditorOnlyData by pointing it to the correct object at load time.

diff --git a/Engine/Source/Runtime/Engine/Classes/Materials/MaterialFunctionInterface.h b/Engine/Source/Runtime/Engine/Classes/Materials/MaterialFunctionInterface.h
--- a/Engine/Source/Runtime/Engine/Classes/Materials/MaterialFunctionInterface.h
+++ b/Engine/Source/Runtime/Engine/Classes/Materials/MaterialFunctionInterface.h
@@ -37,6 +37,9 @@ class UMaterialFunctionInterfaceEditorOnlyData : public UObject
 	GENERATED_BODY()
 public:
 	
+	//~ Begin UObject Interface.
+	ENGINE_API virtual void Serialize(FArchive& Ar) override;
+	//~ End UObject Interface.
 };
 
 /**
@@ -49,6 +52,8 @@ class UMaterialFunctionInterface : public UObject
 
 #if WITH_EDITORONLY_DATA
 protected:
+	friend class UMaterialFunctionInterfaceEditorOnlyData;
+
 	UPROPERTY()
 	TObjectPtr<UMaterialFunctionInterfaceEditorOnlyData> EditorOnlyData;
 
diff --git a/Engine/Source/Runtime/Engine/Private/Materials/MaterialExpressions.cpp b/Engine/Source/Runtime/Engine/Private/Materials/MaterialExpressions.cpp
--- a/Engine/Source/Runtime/Engine/Private/Materials/MaterialExpressions.cpp
+++ b/Engine/Source/Runtime/Engine/Private/Materials/MaterialExpressions.cpp
@@ -44,6 +44,8 @@
 #include "Engine/TextureRenderTargetCube.h"
 #include "Engine/VolumeTexture.h"
 #include "Engine/SubsurfaceProfile.h"
+#include "Serialization/ObjectWriter.h"
+#include "Serialization/ObjectReader.h"
 #include "Styling/CoreStyle.h"
 #include "VT/RuntimeVirtualTexture.h"
 #include "SparseVolumeTexture/SparseVolumeTexture.h"
@@ -13379,9 +13381,38 @@ void UMaterialExpressionCustom::Serialize(FStructuredArchive::FRecord Record)
 #endif // WITH_EDITORONLY_DATA
 }
 
+///////////////////////////////////////////////////////////////////////////////
+// UMaterialFunctionInterfaceEditorOnlyData
+///////////////////////////////////////////////////////////////////////////////
+void UMaterialFunctionInterfaceEditorOnlyData::Serialize(FArchive& Ar)
+{
+	Super::Serialize(Ar);
+
+#if WITH_EDITORONLY_DATA
+	if (!IsTemplate())
+	{
+		// If our owner material function isn't pointing to this EditorOnlyData it means this object's name
+		// doesn't match the default created object name and we need to fix our pointer into the material function interface
+		UMaterialFunctionInterface* MFInterface = CastChecked<UMaterialFunctionInterface>(GetOuter());
+		if (MFInterface->EditorOnlyData != this)
+		{
+			MFInterface->EditorOnlyData = this;
+		}
+	}
+#endif
+}
+
 ///////////////////////////////////////////////////////////////////////////////
 // UMaterialFunctionInterface
 ///////////////////////////////////////////////////////////////////////////////
+namespace MaterialFunctionInterface
+{
+	FString GetEditorOnlyDataName(const TCHAR* InMaterialName)
+	{
+		return FString::Printf(TEXT("%sEditorOnlyData"), InMaterialName);
+	}
+}
+
 UMaterialFunctionInterface::UMaterialFunctionInterface(const FObjectInitializer& ObjectInitializer)
 	: Super(ObjectInitializer)
 	, MaterialFunctionUsage(EMaterialFunctionUsage::Default)
@@ -13402,6 +13433,30 @@ void UMaterialFunctionInterface::PostInitProperties()
 void UMaterialFunctionInterface::PostLoad()
 {
 	Super::PostLoad();
+
+#if WITH_EDITORONLY_DATA
+	if (EditorOnlyData && !GetPackage()->HasAnyPackageFlags(PKG_Cooked))
+	{
+		// Test for badly named EditorOnlyData objects
+		const FString EditorOnlyDataName = MaterialFunctionInterface::GetEditorOnlyDataName(*GetName());
+		if (EditorOnlyData->GetName() != EditorOnlyDataName)
+		{
+			UMaterialFunctionInterfaceEditorOnlyData* CorrectEditorOnlyDataObj = Cast<UMaterialFunctionInterfaceEditorOnlyData>(StaticFindObject(EditorOnlyData->GetClass(), EditorOnlyData->GetOuter(), *EditorOnlyDataName, true));
+			if (CorrectEditorOnlyDataObj)
+			{
+				// Copy data to correct EditorOnlyObject
+				TArray<uint8> Data;
+				FObjectWriter Ar(EditorOnlyData, Data);
+				FObjectReader(CorrectEditorOnlyDataObj, Data);
+
+				// Point EditorOnlyData to the right object
+				EditorOnlyData = CorrectEditorOnlyDataObj;
+			}
+		}
+	}
+#endif
+
+	
 	if (!StateId.IsValid())
 	{
 		StateId = FGuid::NewGuid();
@@ -13433,14 +13488,6 @@ void UMaterialFunctionInterface::GetAssetRegistryTags(TArray<FAssetRegistryTag>&
 #endif
 }
 
-namespace MaterialFunctionInterface
-{
-	FString GetEditorOnlyDataName(const TCHAR* InMaterialName)
-	{
-		return FString::Printf(TEXT("%sEditorOnlyData"), InMaterialName);
-	}
-}
-
 bool UMaterialFunctionInterface::Rename(const TCHAR* NewName, UObject* NewOuter, ERenameFlags Flags)
 {
 	bool bRenamed = Super::Rename(NewName, NewOuter, Flags);
