Texture: Integrated the derived data build jobs with the asset compiling manager in a two-phase approach

The texture build task queues with the texture compiling manager to block on concurrency limits, or when texture builds are paused. The build scheduler will queue with the asset compiling manager to reserve memory when the build needs to execute locally. With this approach, resuming one texture will complete one texture, even though multiple separate build jobs may be executed.

diff --git a/Engine/Source/Developer/DerivedDataCache/Public/Experimental/DerivedDataBuildSchedulerThreadPoolProvider.h b/Engine/Source/Developer/DerivedDataCache/Public/Experimental/DerivedDataBuildSchedulerThreadPoolProvider.h
deleted file mode 100644
diff --git a/Engine/Source/Runtime/Engine/Private/AssetCompilingManager.cpp b/Engine/Source/Runtime/Engine/Private/AssetCompilingManager.cpp
--- a/Engine/Source/Runtime/Engine/Private/AssetCompilingManager.cpp
+++ b/Engine/Source/Runtime/Engine/Private/AssetCompilingManager.cpp
@@ -27,6 +27,12 @@ LLM_DEFINE_TAG(AssetCompilation, NAME_None, NAME_None, GET_STATFNAME(STAT_AssetC
 #include "Algo/Find.h"
 #include "ProfilingDebugging/CountersTrace.h"
 
+#if WITH_EDITOR
+#include "DerivedDataBuildSchedulerQueue.h"
+#include "DerivedDataThreadPoolTask.h"
+#include "Features/IModularFeatures.h"
+#endif
+
 #include UE_INLINE_GENERATED_CPP_BY_NAME(AssetCompilingManager)
 
 #define LOCTEXT_NAMESPACE "AssetCompilingManager"
@@ -88,7 +94,23 @@ namespace AssetCompilingManagerImpl
 		}
 #endif
 	}
-}
+
+#if WITH_EDITOR
+	using namespace UE::DerivedData;
+
+	class FAssetCompilingManagerMemoryQueue final : public IBuildSchedulerMemoryQueue
+	{
+		void Reserve(uint64 Memory, IRequestOwner& Owner, TUniqueFunction<void ()>&& OnComplete) final
+		{
+			// TODO: This is not thread-safe right now because the thread pool is created on first use.
+			//       That first use is expected to occur on the main thread, which we do not guarantee.
+			LaunchTaskInThreadPool(Memory, Owner, FAssetCompilingManager::Get().GetThreadPool(), MoveTemp(OnComplete));
+		}
+	};
+
+	static FAssetCompilingManagerMemoryQueue GMemoryQueue;
+#endif
+} // AssetCompilingManagerImpl
 
 #if WITH_EDITOR
 
@@ -273,6 +295,15 @@ FAssetCompilingManager::FAssetCompilingManager()
 	RegisterManager(&FTextureCompilingManager::Get());
 	RegisterManager(&FActorDeferredScriptManager::Get());
 	RegisterManager(&FSoundWaveCompilingManager::Get());
+
+	IModularFeatures::Get().RegisterModularFeature(UE::DerivedData::IBuildSchedulerMemoryQueue::FeatureName, &AssetCompilingManagerImpl::GMemoryQueue);
+#endif
+}
+
+FAssetCompilingManager::~FAssetCompilingManager()
+{
+#if WITH_EDITOR
+	IModularFeatures::Get().UnregisterModularFeature(UE::DerivedData::IBuildSchedulerMemoryQueue::FeatureName, &AssetCompilingManagerImpl::GMemoryQueue);
 #endif
 }
 
diff --git a/Engine/Source/Runtime/Engine/Private/TextureCompiler.cpp b/Engine/Source/Runtime/Engine/Private/TextureCompiler.cpp
--- a/Engine/Source/Runtime/Engine/Private/TextureCompiler.cpp
+++ b/Engine/Source/Runtime/Engine/Private/TextureCompiler.cpp
@@ -4,10 +4,7 @@
 
 #if WITH_EDITOR
 
-#include "DerivedDataSharedString.h"
 #include "Engine/Texture.h"
-#include "Experimental/DerivedDataBuildSchedulerThreadPoolProvider.h"
-#include "Features/IModularFeatures.h"
 #include "ObjectCacheContext.h"
 #include "Settings/EditorExperimentalSettings.h"
 #include "Misc/QueuedThreadPoolWrapper.h"
@@ -81,42 +78,10 @@ namespace TextureCompilingManagerImpl
 	}
 }
 
-class UE::Private::FTextureThreadPoolProvider final : public UE::DerivedData::IBuildSchedulerThreadPoolProvider
-{
-public:
-	explicit FTextureThreadPoolProvider(FTextureCompilingManager& InManager)
-		: Manager(InManager)
-	{
-		IModularFeatures::Get().RegisterModularFeature(FeatureName, this);
-	}
-
-	~FTextureThreadPoolProvider() final
-	{
-		IModularFeatures::Get().UnregisterModularFeature(FeatureName, this);
-	}
-
-private:
-	const UE::DerivedData::FUtf8SharedString& GetTypeName() const final
-	{
-		static UE::DerivedData::FUtf8SharedString TypeName(UTF8TEXTVIEW("Texture"));
-		return TypeName;
-	}
-
-	FQueuedThreadPool* GetThreadPool() const final
-	{
-		// TODO: This is not thread-safe right now because the thread pool is created on first use.
-		//       That first use is expected to occur on the main thread, which we do not guarantee.
-		return Manager.GetThreadPool();
-	}
-
-	const FTextureCompilingManager& Manager;
-};
-
 FTextureCompilingManager::FTextureCompilingManager()
 	: Notification(GetAssetNameFormat())
 {
 	TextureCompilingManagerImpl::EnsureInitializedCVars();
-	ThreadPoolProvider = MakePimpl<UE::Private::FTextureThreadPoolProvider>(*this);
 }
 
 FName FTextureCompilingManager::GetStaticAssetTypeName()
diff --git a/Engine/Source/Runtime/Engine/Private/TextureDerivedDataTask.cpp b/Engine/Source/Runtime/Engine/Private/TextureDerivedDataTask.cpp
--- a/Engine/Source/Runtime/Engine/Private/TextureDerivedDataTask.cpp
+++ b/Engine/Source/Runtime/Engine/Private/TextureDerivedDataTask.cpp
@@ -35,6 +35,7 @@
 #include "DerivedDataBuildSession.h"
 #include "DerivedDataCache.h"
 #include "DerivedDataRequestOwner.h"
+#include "DerivedDataThreadPoolTask.h"
 #include "Engine/TextureCube.h"
 #include "Engine/VolumeTexture.h"
 #include "ImageCoreUtils.h"
@@ -55,6 +56,7 @@
 #include "Serialization/BulkDataRegistry.h"
 #include "Serialization/CompactBinaryWriter.h"
 #include "TextureBuildUtilities.h"
+#include "TextureCompiler.h"
 #include "TextureDerivedDataBuildUtils.h"
 #include "TextureEncodingSettings.h"
 #include "TextureFormatManager.h"
@@ -2323,10 +2325,10 @@ struct FBuildInfo
 
 static void LaunchBuildWithFallback(
 	FBuildResults& OutBuildResults,
-	FBuildResultOptions InBuildResultOptions,
-	FBuildInfo& InInitialBuild, 
-	TOptional<FBuildInfo> InFallbackBuild, 
-	UE::DerivedData::FRequestOwner& InRequestOwner	// Owner must be valid for the duration of the build.
+	FBuildResultOptions BuildResultOptions,
+	FBuildInfo& InInitialBuild,
+	TOptional<FBuildInfo> InFallbackBuild,
+	UE::DerivedData::FRequestOwner& InRequestOwner // Owner must be valid for the duration of the build.
 )
 {
 	using namespace UE::DerivedData;
@@ -2336,36 +2338,49 @@ static void LaunchBuildWithFallback(
 		OutBuildResults.PlatformData.ResultMetadata = *InInitialBuild.ResultMetadata;
 	}
 
-	InInitialBuild.BuildSession.Build(InInitialBuild.BuildDefinition, {}, InInitialBuild.BuildPolicy, InRequestOwner,
+	LaunchTaskInThreadPool(InRequestOwner, FTextureCompilingManager::Get().GetThreadPool(),
 		[
 			FallbackBuild = MoveTemp(InFallbackBuild),
 			RequestOwner = &InRequestOwner,
 			OutBuildResults = &OutBuildResults,
-			BuildResultOptions = InBuildResultOptions,
+			BuildResultOptions,
+			BuildDefinition = InInitialBuild.BuildDefinition,
+			BuildPolicy = InInitialBuild.BuildPolicy,
+			BuildSession = &InInitialBuild.BuildSession,
 			InitialBuildTilingInputResolver = InInitialBuild.TilingInputResolver
-		](FBuildCompleteParams&& Params) mutable
+		]() mutable
 	{
-		if (Params.Status == EStatus::Error &&
-			FallbackBuild.IsSet())
+		BuildSession->Build(BuildDefinition, {}, BuildPolicy, *RequestOwner,
+			[
+				FallbackBuild = MoveTemp(FallbackBuild),
+				RequestOwner,
+				OutBuildResults,
+				BuildResultOptions,
+				InitialBuildTilingInputResolver
+			](FBuildCompleteParams&& Params) mutable
 		{
-			if (FallbackBuild->ResultMetadata.IsSet())
+			if (Params.Status == EStatus::Error &&
+				FallbackBuild.IsSet())
 			{
-				OutBuildResults->PlatformData.ResultMetadata = *FallbackBuild->ResultMetadata;
+				if (FallbackBuild->ResultMetadata.IsSet())
+				{
+					OutBuildResults->PlatformData.ResultMetadata = *FallbackBuild->ResultMetadata;
+				}
+				FallbackBuild->BuildSession.Build(FallbackBuild->BuildDefinition, {}, FallbackBuild->BuildPolicy, *RequestOwner,
+					[
+						OutBuildResults = OutBuildResults,
+						BuildResultOptions,
+						FallbackBuildTilingInputResolver = FallbackBuild->TilingInputResolver
+					](FBuildCompleteParams&& Params) mutable
+				{
+					GetBuildResultsFromCompleteParams(*OutBuildResults, BuildResultOptions, MoveTemp(Params), FallbackBuildTilingInputResolver);
+				});
 			}
-			FallbackBuild->BuildSession.Build(FallbackBuild->BuildDefinition, {}, FallbackBuild->BuildPolicy, *RequestOwner,
-				[
-					OutBuildResults = OutBuildResults,
-					BuildResultOptions = BuildResultOptions,
-					FallbackBuildTilingInputResolver = FallbackBuild->TilingInputResolver
-				](FBuildCompleteParams&& Params) mutable
-			{
-				GetBuildResultsFromCompleteParams(*OutBuildResults, BuildResultOptions, MoveTemp(Params), FallbackBuildTilingInputResolver);
-			});
-		}
-		else
-		{			
-			GetBuildResultsFromCompleteParams(*OutBuildResults, BuildResultOptions, MoveTemp(Params), InitialBuildTilingInputResolver);
-		}
+			else
+			{
+				GetBuildResultsFromCompleteParams(*OutBuildResults, BuildResultOptions, MoveTemp(Params), InitialBuildTilingInputResolver);
+			}
+		});
 	});
 };
 
diff --git a/Engine/Source/Runtime/Engine/Public/AssetCompilingManager.h b/Engine/Source/Runtime/Engine/Public/AssetCompilingManager.h
--- a/Engine/Source/Runtime/Engine/Public/AssetCompilingManager.h
+++ b/Engine/Source/Runtime/Engine/Public/AssetCompilingManager.h
@@ -127,6 +127,7 @@ public:
 	FAssetPostCompileEvent& OnAssetPostCompileEvent() { return AssetPostCompileEvent; }
 private:
 	FAssetCompilingManager();
+	~FAssetCompilingManager();
 	friend struct IAssetCompilingManager;
 
 	/** Take some action whenever the number of remaining asset changes. */
diff --git a/Engine/Source/Runtime/Engine/Public/TextureCompiler.h b/Engine/Source/Runtime/Engine/Public/TextureCompiler.h
--- a/Engine/Source/Runtime/Engine/Public/TextureCompiler.h
+++ b/Engine/Source/Runtime/Engine/Public/TextureCompiler.h
@@ -6,7 +6,6 @@
 #include "Containers/Set.h"
 #include "AssetCompilingManager.h"
 #include "AsyncCompilationHelpers.h"
-#include "Templates/PimplPtr.h"
 
 #if WITH_EDITOR
 
@@ -14,8 +13,6 @@ class UTexture;
 class FQueuedThreadPool;
 enum class EQueuedWorkPriority : uint8;
 
-namespace UE::Private { class FTextureThreadPoolProvider; }
-
 DECLARE_MULTICAST_DELEGATE_OneParam(FTexturePostCompileEvent, const TArrayView<UTexture* const>&);
 
 class FTextureCompilingManager : IAssetCompilingManager
@@ -107,8 +104,6 @@ private:
 
 	/** Event issued at the end of the compile process */
 	FTexturePostCompileEvent TexturePostCompileEvent;
-
-	TPimplPtr<UE::Private::FTextureThreadPoolProvider> ThreadPoolProvider;
 };
 
 #endif
\ No newline at end of file
