Optimize StreamableManager::FindInMemory by always selecting the fastest function to find the target object.
SoftObjectPaths are mostly used for top level assets, and very rarely represents a SubObject.
Top level assets are only represented by exactly two FNames and can be found much faster by completely skipping the ToString() conversion and the expensive ResolveName resolution inside StaticFindObject.
SubObjects will also benefit from the faster Asset lookup, and the slower ResolveName processing is only required for the actual SubObjectString part.
When a load is completed use the package pointer provided in the callback instead of finding it by name.

Profiling of console client startup:
7.9X speedup in FindInMemory (inclusive from 214ms to 27ms) for ~52 300 calls distributed between:
3.8X speedup in AsyncLoadCallback (inclusive from 144ms to 38ms)  for ~34 000 calls.
1.7X speedup in StreamInternal (inclusive from 204ms to 119ms) for ~51 200 calls.

#rnx

diff --git a/Engine/Source/Runtime/Engine/Classes/Engine/StreamableManager.h b/Engine/Source/Runtime/Engine/Classes/Engine/StreamableManager.h
--- a/Engine/Source/Runtime/Engine/Classes/Engine/StreamableManager.h
+++ b/Engine/Source/Runtime/Engine/Classes/Engine/StreamableManager.h
@@ -565,7 +565,7 @@ private:
 
 	void RemoveReferencedAsset(const FSoftObjectPath& Target, TSharedRef<FStreamableHandle> Handle);
 	void StartHandleRequests(TSharedRef<FStreamableHandle> Handle);
-	void FindInMemory(FSoftObjectPath& InOutTarget, struct FStreamable* Existing);
+	void FindInMemory(FSoftObjectPath& InOutTarget, struct FStreamable* Existing, UPackage* Package = nullptr);
 	FSoftObjectPath HandleLoadedRedirector(UObjectRedirector* LoadedRedirector, FSoftObjectPath RequestedPath, struct FStreamable* RequestedStreamable);
 	struct FStreamable* FindStreamable(const FSoftObjectPath& Target) const;
 	struct FStreamable* StreamInternal(const FSoftObjectPath& Target, TAsyncLoadPriority Priority, TSharedRef<FStreamableHandle> Handle);
@@ -573,7 +573,7 @@ private:
 	void CheckCompletedRequests(const FSoftObjectPath& Target, struct FStreamable* Existing);
 
 	void OnPreGarbageCollect();
-	void AsyncLoadCallback(FSoftObjectPath Request);
+	void AsyncLoadCallback(FSoftObjectPath Request, UPackage* Package);
 
 	/** Map of paths to streamable objects, this will be the post-redirector name */
 	typedef TMap<FSoftObjectPath, struct FStreamable*> TStreamableMap;
diff --git a/Engine/Source/Runtime/Engine/Private/StreamableManager.cpp b/Engine/Source/Runtime/Engine/Private/StreamableManager.cpp
--- a/Engine/Source/Runtime/Engine/Private/StreamableManager.cpp
+++ b/Engine/Source/Runtime/Engine/Private/StreamableManager.cpp
@@ -853,7 +853,7 @@ void FStreamableHandle::AsyncLoadCallbackWrapper(const FName& PackageName, UPack
 	// Needed so we can bind with a shared pointer for safety
 	if (OwningManager)
 	{
-		OwningManager->AsyncLoadCallback(TargetName);
+		OwningManager->AsyncLoadCallback(TargetName, Package);
 
 		if (!HasLoadCompleted())
 		{
@@ -1480,12 +1480,28 @@ UObject* FStreamableManager::LoadSynchronous(const FSoftObjectPath& Target, bool
 	return nullptr;
 }
 
-void FStreamableManager::FindInMemory( FSoftObjectPath& InOutTargetName, struct FStreamable* Existing )
+void FStreamableManager::FindInMemory(FSoftObjectPath& InOutTargetName, struct FStreamable* Existing, UPackage* Package)
 {
 	check(Existing);
 	check(!Existing->bAsyncLoadRequestOutstanding);
 	UE_LOG(LogStreamableManager, Verbose, TEXT("     Searching in memory for %s"), *InOutTargetName.ToString());
-	Existing->Target = StaticFindObject(UObject::StaticClass(), nullptr, *InOutTargetName.ToString());
+	Existing->Target = nullptr;
+
+	UObject* Asset = Package ? 
+		StaticFindObjectFast(UObject::StaticClass(), Package, InOutTargetName.GetAssetFName()) :
+		StaticFindObject(UObject::StaticClass(), InOutTargetName.GetAssetPath(), /*ExactClass*/false);	
+	if (Asset)
+	{
+		if (InOutTargetName.IsAsset())
+		{
+			Existing->Target = Asset;
+		}
+		else if (InOutTargetName.IsSubobject())
+		{
+			Existing->Target = StaticFindObject(UObject::StaticClass(), Asset, *InOutTargetName.GetSubPathString());
+		}
+	}
+	checkSlow(Existing->Target == StaticFindObject(UObject::StaticClass(), nullptr, *InOutTargetName.ToString()));
 
 	if (Existing->Target && Existing->Target->HasAnyInternalFlags(EInternalObjectFlags::AsyncLoading))
 	{
@@ -1507,7 +1523,7 @@ void FStreamableManager::FindInMemory( FSoftObjectPath& InOutTargetName, struct
 	}
 }
 
-void FStreamableManager::AsyncLoadCallback(FSoftObjectPath TargetName)
+void FStreamableManager::AsyncLoadCallback(FSoftObjectPath TargetName, UPackage* Package)
 {
 	check(IsInGameThread());
 
@@ -1521,7 +1537,7 @@ void FStreamableManager::AsyncLoadCallback(FSoftObjectPath TargetName)
 			Existing->bAsyncLoadRequestOutstanding = false;
 			if (!Existing->Target)
 			{
-				FindInMemory(TargetName, Existing);
+				FindInMemory(TargetName, Existing, Package);
 			}
 
 			CheckCompletedRequests(TargetName, Existing);
