Strata - Removed legacy metalness workflow. Without data conversion.

diff --git a/Engine/Shaders/Private/DebugProbes.usf b/Engine/Shaders/Private/DebugProbes.usf
--- a/Engine/Shaders/Private/DebugProbes.usf
+++ b/Engine/Shaders/Private/DebugProbes.usf
@@ -116,14 +116,9 @@ void WriteOutMaterial(
 	FStrataPixelFootprint PixelFootprint = (FStrataPixelFootprint)0;
 	FStrataData StrataSlabData = GetStrataSlabBSDF(
 		/**FStrataPixelFootprint PixelFootprint */PixelFootprint,
-		/*float UseMetalness					*/true,
-		/*float3 BaseColor						*/BaseColor,
-		/*float3 EdgeColor						*/1.0f,
-		/*float Specular						*/Specular,
-		/*float Metallic						*/Metallic,
-		/*float3 DiffuseAlbedo					*/0.0f,
-		/*float3 F0								*/0.0f,
-		/*float3 F90							*/0.0f,
+		/*float3 DiffuseAlbedo					*/ComputeDiffuseAlbedo(BaseColor, Metallic),
+		/*float3 F0								*/ComputeF0(Specular, BaseColor, Metallic),
+		/*float3 F90							*/1.0f,
 		/*float Roughness						*/Roughness,
 		/*float Anisotropy						*/0.0f,
 		/*float SSSProfileID					*/0.0f,
diff --git a/Engine/Shaders/Private/Strata/Strata.ush b/Engine/Shaders/Private/Strata/Strata.ush
--- a/Engine/Shaders/Private/Strata/Strata.ush
+++ b/Engine/Shaders/Private/Strata/Strata.ush
@@ -1338,9 +1338,7 @@ FStrataPixelFootprint StrataGetPixelFootprint(float3 InWorldPosition)
 
 FStrataData GetStrataSlabBSDF(
 	FStrataPixelFootprint InPixelFootprint,
-	float UseMetalness,
-	float3 BaseColor, float3 EdgeColor, float Specular, float Metallic,	//  Metalness workflow
-	float3 DiffuseAlbedo, float3 F0, float3 F90,						// !Metalness workflow
+	float3 DiffuseAlbedo, float3 F0, float3 F90,
 	float Roughness, float Anisotropy,
 	float SSSProfileID, float3 SSSMFP, float SSSMFPScale, float SSSPhaseAnisotropy, float UseSSSDiffusion,
 	float3 Emissive, 
@@ -1430,29 +1428,14 @@ FStrataData GetStrataSlabBSDF(
 	BSDF_SETTHICKNESSCM		(StrataData.InlinedBSDF,		Thickness);
 #endif
 
-
-	if (UseMetalness > 0.0f)
-	{
-#if !STRATA_FASTPATH
-		BSDF_SETHASF90(StrataData.InlinedBSDF, any(EdgeColor < 1.0f));
-#endif
-
-		const float SafeMetallic = saturate(Metallic);
-		SLAB_DIFFUSEALBEDO(StrataData.InlinedBSDF)			= ComputeDiffuseAlbedo(BaseColor, SafeMetallic);
-		SLAB_F0(StrataData.InlinedBSDF)						= ComputeF0(Specular, BaseColor, SafeMetallic);
-		SLAB_F90(StrataData.InlinedBSDF)					= ComputeF90(F0, EdgeColor, SafeMetallic);
-	}
-	else
-	{
 #if !STRATA_FASTPATH
-		// See how UMaterialExpressionStrataSlabBSDF node interface is changed
-		BSDF_SETHASF90(StrataData.InlinedBSDF, any(F90 < 1.0f));
+	// See how UMaterialExpressionStrataSlabBSDF node interface is changed
+	BSDF_SETHASF90(StrataData.InlinedBSDF, any(F90 < 1.0f));
 #endif
 
-		SLAB_DIFFUSEALBEDO(StrataData.InlinedBSDF)			= DiffuseAlbedo;
-		SLAB_F0(StrataData.InlinedBSDF)						= F0;
-		SLAB_F90(StrataData.InlinedBSDF)					= F90;
-	}
+	SLAB_DIFFUSEALBEDO(StrataData.InlinedBSDF)			= DiffuseAlbedo;
+	SLAB_F0(StrataData.InlinedBSDF)						= F0;
+	SLAB_F90(StrataData.InlinedBSDF)					= F90;
 
 	SLAB_ROUGHNESS			(StrataData.InlinedBSDF)		= Roughness;
 
@@ -1486,14 +1469,7 @@ FStrataData GetStrataSlabBSDF(
 #endif
 
 #if USE_DEVELOPMENT_SHADERS
-	if (UseMetalness > 0.0f)
-	{
-		StrataData.PreviewColor = BaseColor;
-	}
-	else
-	{
-		StrataData.PreviewColor = lerp(DiffuseAlbedo, F0, F0RGBToMetallic(F0));
-	}
+	StrataData.PreviewColor = lerp(DiffuseAlbedo, F0, F0RGBToMetallic(F0));
 #endif
 
 	return StrataData;
@@ -3434,14 +3410,9 @@ void StrataConvertFromDBuffer(
 	const float DefaultThickness = STRATA_LAYER_DEFAULT_THICKNESS_CM;
 	OutData = GetStrataSlabBSDF(
 		PixelFootprint,		// PixelFootprint
-		true,				// UseMetalness
-		In.BaseColor,		// BaseColor		-  Metalness workflow
-		1.f,				// EdgeColor		-  Metalness workflow
-		In.Specular,		// Specular			-  Metalness workflow
-		In.Metallic,		// Metallic			-  Metalness workflow
-		0.f,				// DiffuseAlbedo	- !Metalness workflow
-		0.f,				// F0				- !Metalness workflow
-		0.f,				// F90				- !Metalness workflow
+		ComputeDiffuseAlbedo(In.BaseColor, In.Metallic),
+		ComputeF0(In.Specular, In.BaseColor, In.Metallic),
+		1.0f,				// F90
 		In.Roughness,		// Roughness
 		0.f,				// Anisotropy
 		0.f,				// SSSProfileID
diff --git a/Engine/Shaders/Private/Strata/StrataLegacyConversion.ush b/Engine/Shaders/Private/Strata/StrataLegacyConversion.ush
--- a/Engine/Shaders/Private/Strata/StrataLegacyConversion.ush
+++ b/Engine/Shaders/Private/Strata/StrataLegacyConversion.ush
@@ -88,17 +88,22 @@ FStrataData CreateLegacySlab(
 	const float3 Zeros = float3(0, 0, 0);
 	const float3 Ones  = float3(1, 1, 1);
 
+	float3 DiffuseAlbedo = In.DiffuseAlbedo;
+	float3 F0 = In.F0;
+	float3 F90 = In.F90;
+	if (In.UseMetalness > 0.0f)
+	{
+		DiffuseAlbedo = ComputeDiffuseAlbedo(In.BaseColor, In.Metallic);
+		F0 = ComputeF0(In.Specular, In.BaseColor, In.Metallic);
+		F90 = Ones;
+	}
+
 	// Fixed layer structure for helping compiler to unroll and optimize shader
 	return GetStrataSlabBSDF(
 		In.PixelFootprint,
-		In.UseMetalness,									// UseMetalness
-		In.BaseColor,										// BaseColor		-  Metalness workflow
-		Ones,												// EdgeColor		-  Metalness workflow
-		In.Specular,										// Specular			-  Metalness workflow
-		In.Metallic,										// Metallic			-  Metalness workflow
-		In.DiffuseAlbedo,									// DiffuseAlbedo	- !Metalness workflow
-		In.F0,												// F0				- !Metalness workflow
-		In.F90,												// F90				- !Metalness workflow
+		DiffuseAlbedo,										// DiffuseAlbedo
+		F0,													// F0
+		F90,												// F90
 		In.Roughness,										// Roughness
 		In.Anisotropy,										// Anisotropy
 		In.SSSProfileID,									// SSSProfileID
diff --git a/Engine/Source/Runtime/Engine/Classes/Materials/MaterialExpressionStrata.h b/Engine/Source/Runtime/Engine/Classes/Materials/MaterialExpressionStrata.h
--- a/Engine/Source/Runtime/Engine/Classes/Materials/MaterialExpressionStrata.h
+++ b/Engine/Source/Runtime/Engine/Classes/Materials/MaterialExpressionStrata.h
@@ -199,30 +199,6 @@ class UMaterialExpressionStrataSlabBSDF : public UMaterialExpressionStrataBSDF
 {
 	GENERATED_UCLASS_BODY()
 
-	/**
-	 * Defines the overall color of the Material. (type = float3, unit = unitless, defaults to 0.18)
-	 */
-	UPROPERTY()
-	FExpressionInput BaseColor;
-
-	/**
-	 * Defines the edge color of the Material. This is only applied on metallic material (type = float3, unit = unitless, defaults to 1.0)
-	 */
-	UPROPERTY()
-	FExpressionInput EdgeColor;
-
-	/**
-	 * Controls how \"metal-like\" your surface looks like. 0 means dielectric, 1 means conductor (type = float, unit = unitless, defaults to 0)
-	 */
-	UPROPERTY()
-	FExpressionInput Metallic;
-	
-	/**
-	 * Used to scale the current amount of specularity on non-metallic surfaces and is a value between 0 and 1 (type = float, unit = unitless, defaults to plastic 0.5)
-	 */
-	UPROPERTY()
-	FExpressionInput Specular;
-	
 	/**
 	 * Defines the diffused albedo, the percentage of light reflected as diffuse from the surface. (type = float3, unit = unitless, defaults to 0.18)
 	 */
@@ -323,10 +299,6 @@ class UMaterialExpressionStrataSlabBSDF : public UMaterialExpressionStrataBSDF
 	UPROPERTY(EditAnywhere, BlueprintReadOnly, Category = Material, meta = (DisplayName = "Subsurface Profile"))
 	TObjectPtr<class USubsurfaceProfile> SubsurfaceProfile;
 
-	/** Whether to use the metalness workflow relying on BaseColor, Specular, EdgeColor and Metallic inputs. Or use the DiffuseColor, F0 and F90 specification. */
-	UPROPERTY(EditAnywhere, Category = Mode)
-	uint32 bUseMetalness : 1;
-
 	/** Whether to use light diffusion (i.e., SSS diffusion) or wrap-approximation for material with scattering behavior. This option trades quality over performance and will result into visual differences. */
 	UPROPERTY(EditAnywhere, Category = Mode, meta = (DisplayName = "Use Subsurface Diffusion"))
 	uint32 bUseSSSDiffusion : 1;
@@ -342,7 +314,6 @@ class UMaterialExpressionStrataSlabBSDF : public UMaterialExpressionStrataBSDF
 	virtual FStrataOperator* StrataGenerateMaterialTopologyTree(class FMaterialCompiler* Compiler, class UMaterialExpression* Parent, int32 OutputIndex) override;
 	virtual FName GetInputName(int32 InputIndex) const override;
 	virtual void GetConnectorToolTip(int32 InputIndex, int32 OutputIndex, TArray<FString>& OutToolTip) override;
-	virtual void PostEditChangeProperty(FPropertyChangedEvent& PropertyChangedEvent) override;
 	virtual const TArray<FExpressionInput*> GetInputs() override;
 
 	bool HasEdgeColor() const;
@@ -362,22 +333,16 @@ class UMaterialExpressionStrataSimpleClearCoatBSDF : public UMaterialExpressionS
 	GENERATED_UCLASS_BODY()
 
 	/**
-	 * Defines the overall color of the Material. (type = float3, unit = unitless, defaults to 0.18)
+	 * Defines the diffused albedo, the percentage of light reflected as diffuse from the surface. (type = float3, unit = unitless, defaults to 0.18)
 	 */
 	UPROPERTY()
-	FExpressionInput BaseColor;
+	FExpressionInput DiffuseAlbedo;
 
 	/**
-	 * Controls how \"metal-like\" your surface looks like. 0 means dielectric, 1 means conductor (type = float, unit = unitless, defaults to 0)
-	 */
-	UPROPERTY()
-	FExpressionInput Metallic;
-	
-	/**
-	 * Used to scale the current amount of specularity on non-metallic surfaces and is a value between 0 and 1 (type = float, unit = unitless, defaults to plastic 0.5)
+	 * Defines F0, the percentage of light reflected as specular from a surface when the view is perpendicular to the surface. (type = float3, unit = unitless, defaults to plastic 0.04)
 	 */
 	UPROPERTY()
-	FExpressionInput Specular;
+	FExpressionInput F0;
 
 	/**
 	 * Controls how rough the bottom layer of the material is. Roughness of 0 (smooth) is a mirror reflection and 1 (rough) is completely matte or diffuse. (type = float, unit = unitless, defaults to 0.5)
diff --git a/Engine/Source/Runtime/Engine/Private/Materials/HLSLMaterialTranslator.cpp b/Engine/Source/Runtime/Engine/Private/Materials/HLSLMaterialTranslator.cpp
--- a/Engine/Source/Runtime/Engine/Private/Materials/HLSLMaterialTranslator.cpp
+++ b/Engine/Source/Runtime/Engine/Private/Materials/HLSLMaterialTranslator.cpp
@@ -10690,8 +10690,6 @@ FString FHLSLMaterialTranslator::StrataGetCastParameterCode(int32 Index, EMateri
 }
 
 int32 FHLSLMaterialTranslator::StrataSlabBSDF(
-	int32 UseMetalness,
-	int32 BaseColor, int32 EdgeColor, int32 Specular, int32 Metallic,
 	int32 DiffuseAlbedo, int32 F0, int32 F90,
 	int32 Roughness, int32 Anisotropy,
 	int32 SSSProfileId, int32 SSSMFP, int32 SSSMFPScale, int32 SSSPhaseAniso, int32 UseSSSDiffusion,
@@ -10713,12 +10711,7 @@ int32 FHLSLMaterialTranslator::StrataSlabBSDF(
 			return INDEX_NONE;
 		}
 		return AddCodeChunk(
-			MCT_Strata, TEXT("PromoteParameterBlendedBSDFToOperator(GetStrataSlabBSDF(Parameters.StrataPixelFootprint, %s, %s, %s, %s, %s, %s, %s, %s, %s, %s, %s, %s, %s, %s, %s, %s, %s, %s, %s, %s, %s, %s, %s, Parameters.SharedLocalBases.Types) /* Normal = %s ; Tangent = %s */, Parameters.StrataTree, %u, %u, %u, %u)"),
-			*StrataGetCastParameterCode(UseMetalness,			MCT_Float),
-			*StrataGetCastParameterCode(BaseColor,				MCT_Float3),
-			*StrataGetCastParameterCode(EdgeColor,				MCT_Float3),
-			*StrataGetCastParameterCode(Specular,				MCT_Float),
-			*StrataGetCastParameterCode(Metallic,				MCT_Float),
+			MCT_Strata, TEXT("PromoteParameterBlendedBSDFToOperator(GetStrataSlabBSDF(Parameters.StrataPixelFootprint, %s, %s, %s, %s, %s, %s, %s, %s, %s, %s, %s, %s, %s, %s, %s, %s, %s, %s, Parameters.SharedLocalBases.Types) /* Normal = %s ; Tangent = %s */, Parameters.StrataTree, %u, %u, %u, %u)"),
 			*StrataGetCastParameterCode(DiffuseAlbedo,			MCT_Float3),
 			*StrataGetCastParameterCode(F0,						MCT_Float3),
 			*StrataGetCastParameterCode(F90,					MCT_Float3),
@@ -10747,12 +10740,7 @@ int32 FHLSLMaterialTranslator::StrataSlabBSDF(
 	}
 	
 	return AddCodeChunk(
-		MCT_Strata, TEXT("GetStrataSlabBSDF(Parameters.StrataPixelFootprint, %s, %s, %s, %s, %s, %s, %s, %s, %s, %s, %s, %s, %s, %s, %s, %s, %s, %s, %s, %s, %s, %s, %s, Parameters.SharedLocalBases.Types) /* Normal = %s ; Tangent = %s */"),
-		*StrataGetCastParameterCode(UseMetalness,			MCT_Float),
-		*StrataGetCastParameterCode(BaseColor,				MCT_Float3),
-		*StrataGetCastParameterCode(EdgeColor,				MCT_Float3),
-		*StrataGetCastParameterCode(Specular,				MCT_Float),
-		*StrataGetCastParameterCode(Metallic,				MCT_Float),
+		MCT_Strata, TEXT("GetStrataSlabBSDF(Parameters.StrataPixelFootprint, %s, %s, %s, %s, %s, %s, %s, %s, %s, %s, %s, %s, %s, %s, %s, %s, %s, %s, Parameters.SharedLocalBases.Types) /* Normal = %s ; Tangent = %s */"),
 		*StrataGetCastParameterCode(DiffuseAlbedo,			MCT_Float3),
 		*StrataGetCastParameterCode(F0,						MCT_Float3),
 		*StrataGetCastParameterCode(F90,					MCT_Float3),
diff --git a/Engine/Source/Runtime/Engine/Private/Materials/HLSLMaterialTranslator.h b/Engine/Source/Runtime/Engine/Private/Materials/HLSLMaterialTranslator.h
--- a/Engine/Source/Runtime/Engine/Private/Materials/HLSLMaterialTranslator.h
+++ b/Engine/Source/Runtime/Engine/Private/Materials/HLSLMaterialTranslator.h
@@ -1063,8 +1063,6 @@ protected:
 	// Strata
 	virtual int32 StrataCreateAndRegisterNullMaterial() override;
 	virtual int32 StrataSlabBSDF(
-		int32 UseMetalness,
-		int32 BaseColor, int32 EdgeColor, int32 Specular, int32 Metallic,
 		int32 DiffuseAlbedo, int32 F0, int32 F90,
 		int32 Roughness, int32 Anisotropy,
 		int32 SSSProfileId, int32 SSSMFP, int32 SSSMFPScale, int32 SSSPhaseAniso, int32 bUseSSSDiffusion,
diff --git a/Engine/Source/Runtime/Engine/Private/Materials/Material.cpp b/Engine/Source/Runtime/Engine/Private/Materials/Material.cpp
--- a/Engine/Source/Runtime/Engine/Private/Materials/Material.cpp
+++ b/Engine/Source/Runtime/Engine/Private/Materials/Material.cpp
@@ -3281,7 +3281,6 @@ void UMaterial::ConvertMaterialToStrataMaterial()
 		// Empty material: Create by default a slab node
 		UMaterialExpressionStrataSlabBSDF* SlabNode = NewObject<UMaterialExpressionStrataSlabBSDF>(this);
 		SetPosXAndMoveReferenceToTheRight(SlabNode);
-		SlabNode->bUseMetalness = false;
 		EditorOnly->FrontMaterial.Connect(0, SlabNode);
 		bRelinkCustomOutputNodes = false;
 		bInvalidateShader = true;
@@ -3298,32 +3297,39 @@ void UMaterial::ConvertMaterialToStrataMaterial()
 				{
 					// For this special case, using two slabs to create a clear coat material with separated top and bottom normal. 
 
+					// Create metalness to Slab parameterisation conveersion node
+					UMaterialExpressionStrataMetalnessToDiffuseAlbedoF0* StrataMetalnessToDiffuseAlbedoF0 = NewObject<UMaterialExpressionStrataMetalnessToDiffuseAlbedoF0>(this);
+					SetPosXAndMoveReferenceToTheRight(StrataMetalnessToDiffuseAlbedoF0);
+					MoveConnectionTo(EditorOnly->BaseColor, StrataMetalnessToDiffuseAlbedoF0, 0);					// BaseColor
+					MoveConnectionTo(EditorOnly->Metallic, StrataMetalnessToDiffuseAlbedoF0, 1);					// Metallic
+					MoveConnectionTo(EditorOnly->Specular, StrataMetalnessToDiffuseAlbedoF0, 2);					// Specular
+					
 					// Top slab BSDF as a simple Disney material
 					UMaterialExpressionStrataSlabBSDF* BottomSlabBSDF = NewObject<UMaterialExpressionStrataSlabBSDF>(this);
 					SetPosXAndMoveReferenceToTheRight(BottomSlabBSDF);
-					MoveConnectionTo(EditorOnly->BaseColor, BottomSlabBSDF, 0);					// BaseColor
-					MoveConnectionTo(EditorOnly->Metallic, BottomSlabBSDF, 2);					// Metallic
-					MoveConnectionTo(EditorOnly->Specular, BottomSlabBSDF, 3);					// Specular
-					MoveConnectionTo(EditorOnly->Roughness, BottomSlabBSDF, 7);					// Roughness
-					CopyConnectionTo(EditorOnly->Anisotropy, BottomSlabBSDF, 8);				// Anisotropy
-					MoveConnectionTo(EditorOnly->Tangent, BottomSlabBSDF, 10);					// Tangent
+					BottomSlabBSDF->GetInput(0)->Connect(0, StrataMetalnessToDiffuseAlbedoF0);
+					BottomSlabBSDF->GetInput(1)->Connect(1, StrataMetalnessToDiffuseAlbedoF0);
+					BottomSlabBSDF->GetInput(2)->Connect(2, StrataMetalnessToDiffuseAlbedoF0);
+					MoveConnectionTo(EditorOnly->Roughness, BottomSlabBSDF, 3);					// Roughness
+					CopyConnectionTo(EditorOnly->Anisotropy, BottomSlabBSDF, 4);				// Anisotropy
+					MoveConnectionTo(EditorOnly->Tangent, BottomSlabBSDF, 6);					// Tangent
 
 					check(ClearCoatBottomNormalOutput);
-					CopyConnectionTo(*ClearCoatBottomNormalOutput->GetInput(0), BottomSlabBSDF, 9);// ClearColorBottomNormal -> BottomSlabBSDF.Normal
+					CopyConnectionTo(*ClearCoatBottomNormalOutput->GetInput(0), BottomSlabBSDF, 5);// ClearColorBottomNormal -> BottomSlabBSDF.Normal
 
 					// Now weight the top base material by opacity.
 					UMaterialExpressionStrataSlabBSDF* TopSlabBSDF = NewObject<UMaterialExpressionStrataSlabBSDF>(this);
 					TopSlabBSDF->MaterialExpressionEditorX = BottomSlabBSDF->MaterialExpressionEditorX;
 					TopSlabBSDF->MaterialExpressionEditorY = BottomSlabBSDF->MaterialExpressionEditorY + 650;
-					MoveConnectionTo(EditorOnly->EmissiveColor, TopSlabBSDF, 14);				// Emissive
-					MoveConnectionTo(EditorOnly->ClearCoatRoughness, TopSlabBSDF, 7);			// ClearCoatRoughness => Roughness
-					MoveConnectionTo(EditorOnly->Normal, TopSlabBSDF, 9);						// Normal
+					MoveConnectionTo(EditorOnly->EmissiveColor, TopSlabBSDF, 10);				// Emissive
+					MoveConnectionTo(EditorOnly->ClearCoatRoughness, TopSlabBSDF, 3);			// ClearCoatRoughness => Roughness
+					MoveConnectionTo(EditorOnly->Normal, TopSlabBSDF, 5);						// Normal
 
 					//  The top layer has a hard coded specular value of 0.5 (F0 = 0.04)
 					UMaterialExpressionConstant* ConstantHalf = NewObject<UMaterialExpressionConstant>(this);
 					ReplaceNodeAndMoveToTheRight(TopSlabBSDF, ConstantHalf);
-					ConstantHalf->R = 0.5f;
-					TopSlabBSDF->GetInput(3)->Connect(0, ConstantHalf);
+					ConstantHalf->R = 0.5f * 0.08f;
+					TopSlabBSDF->GetInput(1)->Connect(0, ConstantHalf);
 
 					// The original clear coat is a complex assemblage of arbitrary functions that do not always make sense.
 					// To simplify things, we set the top slab BSDF as having a constant Grey scale transmittance.
@@ -3338,9 +3344,10 @@ void UMaterial::ConvertMaterialToStrataMaterial()
 					ReplaceNodeAndMoveToTheRight(TopSlabBSDF, Constant075);
 					Constant075->R = 0.75f;
 					UMaterialExpressionStrataTransmittanceToMFP* TransToMDFP = NewObject<UMaterialExpressionStrataTransmittanceToMFP>(this);
+					ReplaceNodeAndMoveToTheRight(TopSlabBSDF, TransToMDFP);
 					TransToMDFP->GetInput(0)->Connect(0, Constant075);
-					TopSlabBSDF->GetInput(11)->Connect(0, TransToMDFP);							// MFP -> MFP
-					TopSlabBSDF->GetInput(17)->Connect(1, TransToMDFP);							// Thickness -> Thickness
+					TopSlabBSDF->GetInput(7)->Connect(0, TransToMDFP);							// MFP -> MFP
+					TopSlabBSDF->GetInput(13)->Connect(1, TransToMDFP);							// Thickness -> Thickness
 
 					// Now weight the top base material by ClearCoat
 					UMaterialExpressionStrataWeight* TopSlabBSDFWithCoverage = NewObject<UMaterialExpressionStrataWeight>(this);
diff --git a/Engine/Source/Runtime/Engine/Private/Materials/MaterialExpressions.cpp b/Engine/Source/Runtime/Engine/Private/Materials/MaterialExpressions.cpp
--- a/Engine/Source/Runtime/Engine/Private/Materials/MaterialExpressions.cpp
+++ b/Engine/Source/Runtime/Engine/Private/Materials/MaterialExpressions.cpp
@@ -22464,7 +22464,6 @@ int32 UMaterialExpressionStrataBSDF::CompilePreview(class FMaterialCompiler* Com
 
 UMaterialExpressionStrataSlabBSDF::UMaterialExpressionStrataSlabBSDF(const FObjectInitializer& ObjectInitializer)
 	: Super(ObjectInitializer)
-	, bUseMetalness(true)
 	, bUseSSSDiffusion(true)
 {
 	struct FConstructorStatics
@@ -22489,7 +22488,7 @@ int32 UMaterialExpressionStrataSlabBSDF::Compile(class FMaterialCompiler* Compil
 	// used in StrataGetBSDFSharedBasis cannot be relied on for smooth tangent used for lighting on any mesh.
 	const bool bHasAnisotropy = HasAnisotropy();
 
-	const bool bHasEdgeColor = HasEdgeColor(); // This accounts for EdgeColor and also F90 when the non metalness workflow is selected.
+	const bool bHasEdgeColor = HasEdgeColor();
 	const bool bHasFuzz = HasFuzz();
 	const bool bHasSecondRoughness = HasSecondRoughness();
 	const bool bHasMFPPluggedIn = HasMFPPluggedIn();
@@ -22511,19 +22510,9 @@ int32 UMaterialExpressionStrataSlabBSDF::Compile(class FMaterialCompiler* Compil
 	StrataOperator.BSDFRegisteredSharedLocalBasis = NewRegisteredSharedLocalBasis;
 
 	int32 OutputCodeChunk = Compiler->StrataSlabBSDF(
-		bUseMetalness ? Compiler->Constant(1.0f) : Compiler->Constant(0.0f),
-
-		// Metalness workflow
-		CompileWithDefaultFloat3(Compiler, BaseColor, 0.18f, 0.18f, 0.18f),
-		CompileWithDefaultFloat3(Compiler, EdgeColor, 1.0f, 1.0f, 1.0f),
-		CompileWithDefaultFloat1(Compiler, Specular, DefaultSpecular),
-		CompileWithDefaultFloat1(Compiler, Metallic,  0.0f),
-
-		// !Metalness workflow
 		CompileWithDefaultFloat3(Compiler, DiffuseAlbedo, 0.18f, 0.18f, 0.18f),
 		CompileWithDefaultFloat3(Compiler, F0, DefaultF0, DefaultF0, DefaultF0),
 		CompileWithDefaultFloat3(Compiler, F90, 1.0f, 1.0f, 1.0f),
-
 		RoughnessCodeChunk,
 		AnisotropyCodeChunk,
 		SSSProfileCodeChunk != INDEX_NONE ? SSSProfileCodeChunk : Compiler->Constant(0.0f),	
@@ -22550,19 +22539,9 @@ int32 UMaterialExpressionStrataSlabBSDF::Compile(class FMaterialCompiler* Compil
 const TArray<FExpressionInput*> UMaterialExpressionStrataSlabBSDF::GetInputs()
 {
 	TArray<FExpressionInput*> Result;
-	if (bUseMetalness)
-	{
-		Result.Add(&BaseColor);
-		Result.Add(&EdgeColor);
-		Result.Add(&Metallic);
-		Result.Add(&Specular);
-	}
-	else
-	{
-		Result.Add(&DiffuseAlbedo);
-		Result.Add(&F0);
-		Result.Add(&F90);
-	}
+	Result.Add(&DiffuseAlbedo);
+	Result.Add(&F0);
+	Result.Add(&F90);
 	Result.Add(&Roughness);
 	Result.Add(&Anisotropy);
 	Result.Add(&Normal);
@@ -22579,55 +22558,6 @@ const TArray<FExpressionInput*> UMaterialExpressionStrataSlabBSDF::GetInputs()
 	return Result;
 }
 
-void UMaterialExpressionStrataSlabBSDF::PostEditChangeProperty(FPropertyChangedEvent& PropertyChangedEvent)
-{
-	if (PropertyChangedEvent.Property && (PropertyChangedEvent.Property->GetName() == TEXT("bUseMetalness")))
-	{
-		auto SafeConnect = [](FExpressionInput& Input, int32 OutputIndex, class UMaterialExpression* Expression)
-		{
-			if (Expression)
-			{
-				Input.Connect(OutputIndex, Expression);
-			}
-			else
-			{
-				Input.Expression = nullptr;
-			}
-		};
-
-		// Relink to make the node looks somewhat ok when toggling Metalness
-		if (bUseMetalness)
-		{
-			SafeConnect(BaseColor, DiffuseAlbedo.OutputIndex, DiffuseAlbedo.Expression);
-			SafeConnect(Specular, F0.OutputIndex, F0.Expression);
-			SafeConnect(EdgeColor, F90.OutputIndex, F90.Expression);
-
-			DiffuseAlbedo.Expression = nullptr;
-			F0.Expression = nullptr;
-			F90.Expression = nullptr;
-		}
-		else
-		{
-			SafeConnect(DiffuseAlbedo, BaseColor.OutputIndex, BaseColor.Expression);
-			SafeConnect(F0, Specular.OutputIndex, Specular.Expression);
-			SafeConnect(F90, EdgeColor.OutputIndex, EdgeColor.Expression);
-
-			BaseColor.Expression = nullptr;
-			EdgeColor.Expression = nullptr;
-			Specular.Expression = nullptr;
-			Metallic.Expression = nullptr;
-		}
-
-		// Recreate this node and relink according to updated pins.
-		CastChecked<UMaterialGraphNode>(GraphNode)->RecreateAndLinkNode();
-	}
-	if (GraphNode)
-	{
-		GraphNode->ReconstructNode();
-	}
-	Super::PostEditChangeProperty(PropertyChangedEvent);
-}
-
 void UMaterialExpressionStrataSlabBSDF::GetCaption(TArray<FString>& OutCaptions) const
 {
 	OutCaptions.Add(TEXT("Strata Slab BSDF"));
@@ -22640,92 +22570,67 @@ uint32 UMaterialExpressionStrataSlabBSDF::GetOutputType(int32 OutputIndex)
 
 uint32 UMaterialExpressionStrataSlabBSDF::GetInputType(int32 InputIndex)
 {
-	const int32 SkipDisabledInputOffset = bUseMetalness ? 0 : -1;
-
-	if (bUseMetalness)
+	if (InputIndex == 0)
 	{
-		if (InputIndex == 0)
-		{
-			return MCT_Float3; // BaseColor
-		}
-		else if (InputIndex == 1)
-		{
-			return MCT_Float3; // EdgeColor
-		}
-		else if (InputIndex == 2)
-		{
-			return MCT_Float1; // Metallic
-		}
-		else if (InputIndex == 3)
-		{
-			return MCT_Float1; // Specular
-		}
+		return MCT_Float3; // DiffuseAlbedo
 	}
-	else
+	else if (InputIndex == 1)
 	{
-		if (InputIndex == 0)
-		{
-			return MCT_Float3; // DiffuseAlbedo
-		}
-		else if (InputIndex == 1)
-		{
-			return MCT_Float3; // F0
-		}
-		else if (InputIndex == 2)
-		{
-			return MCT_Float3; // F90
-		}
+		return MCT_Float3; // F0
 	}
-
-	if (InputIndex == (4 + SkipDisabledInputOffset))
+	else if (InputIndex == 2)
+	{
+		return MCT_Float3; // F90
+	}
+	else if (InputIndex == 3)
 	{
 		return MCT_Float1; // Roughness
 	}
-	else if (InputIndex == (5 + SkipDisabledInputOffset))
+	else if (InputIndex == 4)
 	{
 		return MCT_Float1; // Anisotropy
 	}
-	else if (InputIndex == (6 + SkipDisabledInputOffset))
+	else if (InputIndex == 5)
 	{
 		return MCT_Float3; // Normal
 	}
-	else if (InputIndex == (7 + SkipDisabledInputOffset))
+	else if (InputIndex == 6)
 	{
 		return MCT_Float3; // Tangent
 	}
-	else if (InputIndex == (8 + SkipDisabledInputOffset))
+	else if (InputIndex == 7)
 	{
 		return MCT_Float3; // SSSMFP
 	}
-	else if (InputIndex == (9 + SkipDisabledInputOffset))
+	else if (InputIndex == 8)
 	{
 		return MCT_Float1; // SSSMFPScale
 	}
-	else if (InputIndex == (10 + SkipDisabledInputOffset))
+	else if (InputIndex == 9)
 	{
 		return MCT_Float1; // SSSPhaseAniso
 	}
-	else if (InputIndex == (11 + SkipDisabledInputOffset))
+	else if (InputIndex == 10)
 	{
 		return MCT_Float3; // Emissive Color
 	}
-	else if (InputIndex == (12 + SkipDisabledInputOffset))
+	else if (InputIndex == 11)
 	{
 		return MCT_Float1; // SecondRoughness
 	}
-	else if (InputIndex == (13 + SkipDisabledInputOffset))
+	else if (InputIndex == 12)
 	{
 		return MCT_Float1; // SecondRoughnessWeight
 	}
-	else if (InputIndex == (14 + SkipDisabledInputOffset))
+	else if (InputIndex == 13)
 	{
 		return MCT_Float1; // Thickness
 	}
-	else if (InputIndex == (15 + SkipDisabledInputOffset))
+	else if (InputIndex == 14)
 	{
 		return MCT_Float1; // FuzzAmount
 	}
-	else if (InputIndex == (16 + SkipDisabledInputOffset))
+	else if (InputIndex == 15)
 	{
 		return MCT_Float3; // FuzzColor
 	}
@@ -22736,92 +22641,67 @@ uint32 UMaterialExpressionStrataSlabBSDF::GetInputType(int32 InputIndex)
 
 FName UMaterialExpressionStrataSlabBSDF::GetInputName(int32 InputIndex) const
 {
-	const int32 SkipDisabledInputOffset = bUseMetalness ? 0 : -1;
-
-	if (bUseMetalness)
+	if (InputIndex == 0)
 	{
-		if (InputIndex == 0)
-		{
-			return TEXT("Base Color");
-		}
-		else if (InputIndex == 1)
-		{
-			return TEXT("Edge Color");
-		}
-		else if (InputIndex == 2)
-		{
-			return TEXT("Metallic");
-		}
-		else if (InputIndex == 3)
-		{
-			return TEXT("Specular");
-		}
+		return TEXT("Diffuse Albedo");
 	}
-	else
+	else if (InputIndex == 1)
 	{
-		if (InputIndex == 0)
-		{
-			return TEXT("Diffuse Albedo");
-		}
-		else if (InputIndex == 1)
-		{
-			return TEXT("F0");
-		}
-		else if (InputIndex == 2)
-		{
-			return  TEXT("F90");
-		}
+		return TEXT("F0");
 	}
-
-	if (InputIndex == (4 + SkipDisabledInputOffset))
+	else if (InputIndex == 2)
+	{
+		return  TEXT("F90");
+	}
+	else if (InputIndex == 3)
 	{
 		return TEXT("Roughness");
 	}
-	else if (InputIndex == (5 + SkipDisabledInputOffset))
+	else if (InputIndex == 4)
 	{
 		return TEXT("Anisotropy");
 	}
-	else if (InputIndex == (6 + SkipDisabledInputOffset))
+	else if (InputIndex == 5)
 	{
 		return TEXT("Normal");
 	}
-	else if (InputIndex == (7 + SkipDisabledInputOffset))
+	else if (InputIndex == 6)
 	{
 		return TEXT("Tangent");
 	}
-	else if (InputIndex == (8 + SkipDisabledInputOffset))
+	else if (InputIndex == 7)
 	{
 		return TEXT("SSS MFP");
 	}
-	else if (InputIndex == (9 + SkipDisabledInputOffset))
+	else if (InputIndex == 8)
 	{
 		return TEXT("SSS MFP Scale");
 	}
-	else if (InputIndex == (10 + SkipDisabledInputOffset))
+	else if (InputIndex == 9)
 	{
 		return TEXT("SSS Phase Anisotropy");
 	}
-	else if (InputIndex == (11 + SkipDisabledInputOffset))
+	else if (InputIndex == 10)
 	{
 		return TEXT("Emissive Color");
 	}
-	else if (InputIndex == (12 + SkipDisabledInputOffset))
+	else if (InputIndex == 11)
 	{
 		return TEXT("Second Roughness");
 	}
-	else if (InputIndex == (13 + SkipDisabledInputOffset))
+	else if (InputIndex == 12)
 	{
 		return TEXT("Second Roughness Weight");
 	}
-	else if (InputIndex == (14 + SkipDisabledInputOffset))
+	else if (InputIndex == 13)
 	{
 		return TEXT("Thickness");
 	}
-	else if (InputIndex == (15 + SkipDisabledInputOffset))
+	else if (InputIndex == 14)
 	{
 		return TEXT("Fuzz Amount");
 	}
-	else if (InputIndex == (16 + SkipDisabledInputOffset))
+	else if (InputIndex == 15)
 	{
 		return TEXT("Fuzz Color");
 	}
@@ -22836,99 +22716,8 @@ void UMaterialExpressionStrataSlabBSDF::GetConnectorToolTip(int32 InputIndex, in
 		OutToolTip.Add(TEXT("TT Ouput"));
 		return;
 	}
-	const int32 SkipDisabledInputOffset = bUseMetalness ? 0 : -1;
-
-	if (bUseMetalness)
-	{
-		if (InputIndex == 0)
-		{
-			Super::GetConnectorToolTip(0, INDEX_NONE, OutToolTip);
-		}
-		else if (InputIndex == 1)
-		{
-			Super::GetConnectorToolTip(1, INDEX_NONE, OutToolTip);
-		}
-		else if (InputIndex == 2)
-		{
-			Super::GetConnectorToolTip(2, INDEX_NONE, OutToolTip);
-		}
-		else if (InputIndex == 3)
-		{
-			Super::GetConnectorToolTip(3, INDEX_NONE, OutToolTip);
-		}
-	}
-	else
-	{
-		if (InputIndex == 0)
-		{
-			Super::GetConnectorToolTip(4, INDEX_NONE, OutToolTip);
-		}
-		else if (InputIndex == 1)
-		{
-			Super::GetConnectorToolTip(5, INDEX_NONE, OutToolTip);
-		}
-		else if (InputIndex == 2)
-		{
-			Super::GetConnectorToolTip(6, INDEX_NONE, OutToolTip);
-		}
-	}
 
-	if (InputIndex == (4 + SkipDisabledInputOffset))
-	{
-		Super::GetConnectorToolTip(7, INDEX_NONE, OutToolTip);
-	}
-	else if (InputIndex == (5 + SkipDisabledInputOffset))
-	{
-		Super::GetConnectorToolTip(8, INDEX_NONE, OutToolTip);
-	}
-	else if (InputIndex == (6 + SkipDisabledInputOffset))
-	{
-		Super::GetConnectorToolTip(9, INDEX_NONE, OutToolTip);
-	}
-	else if (InputIndex == (7 + SkipDisabledInputOffset))
-	{
-		Super::GetConnectorToolTip(10, INDEX_NONE, OutToolTip);
-	}
-	else if (InputIndex == (8 + SkipDisabledInputOffset))
-	{
-		Super::GetConnectorToolTip(11, INDEX_NONE, OutToolTip);
-	}
-	else if (InputIndex == (9 + SkipDisabledInputOffset))
-	{
-		Super::GetConnectorToolTip(12, INDEX_NONE, OutToolTip);
-	}
-	else if (InputIndex == (10 + SkipDisabledInputOffset))
-	{
-		Super::GetConnectorToolTip(13, INDEX_NONE, OutToolTip);
-	}
-	else if (InputIndex == (11 + SkipDisabledInputOffset))
-	{
-		Super::GetConnectorToolTip(14, INDEX_NONE, OutToolTip);
-	}
-	else if (InputIndex == (12 + SkipDisabledInputOffset))
-	{
-		Super::GetConnectorToolTip(15, INDEX_NONE, OutToolTip);
-	}
-	else if (InputIndex == (13 + SkipDisabledInputOffset))
-	{
-		Super::GetConnectorToolTip(16, INDEX_NONE, OutToolTip);
-	}
-	else if (InputIndex == (14 + SkipDisabledInputOffset))
-	{
-		Super::GetConnectorToolTip(17, INDEX_NONE, OutToolTip);
-	}
-	else if (InputIndex == (15 + SkipDisabledInputOffset))
-	{
-		Super::GetConnectorToolTip(18, INDEX_NONE, OutToolTip);
-	}
-	else if (InputIndex == (16 + SkipDisabledInputOffset))
-	{
-		Super::GetConnectorToolTip(19, INDEX_NONE, OutToolTip);
-	}
-	else if (InputIndex == (17 + SkipDisabledInputOffset))
-	{
-		Super::GetConnectorToolTip(20, INDEX_NONE, OutToolTip);
-	}
+	Super::GetConnectorToolTip(InputIndex, INDEX_NONE, OutToolTip);
 }
 
 bool UMaterialExpressionStrataSlabBSDF::IsResultStrataMaterial(int32 OutputIndex)
@@ -22939,17 +22728,8 @@ bool UMaterialExpressionStrataSlabBSDF::IsResultStrataMaterial(int32 OutputIndex
 void UMaterialExpressionStrataSlabBSDF::GatherStrataMaterialInfo(FStrataMaterialInfo& StrataMaterialInfo, int32 OutputIndex)
 {
 	// Track connected inputs
-	if (bUseMetalness)
-	{
-		if (BaseColor.IsConnected())		{ StrataMaterialInfo.AddPropertyConnected(MP_BaseColor); }
-		if (Metallic.IsConnected())			{ StrataMaterialInfo.AddPropertyConnected(MP_Metallic); }
-		if (Specular.IsConnected())			{ StrataMaterialInfo.AddPropertyConnected(MP_Specular); }
-	}
-	else
-	{
-		if (DiffuseAlbedo.IsConnected())	{ StrataMaterialInfo.AddPropertyConnected(MP_DiffuseColor); }
-		if (F0.IsConnected())				{ StrataMaterialInfo.AddPropertyConnected(MP_SpecularColor); }
-	}
+	if (DiffuseAlbedo.IsConnected())		{ StrataMaterialInfo.AddPropertyConnected(MP_DiffuseColor); }
+	if (F0.IsConnected())					{ StrataMaterialInfo.AddPropertyConnected(MP_SpecularColor); }
 	if (Roughness.IsConnected())			{ StrataMaterialInfo.AddPropertyConnected(MP_Roughness); }
 	if (Anisotropy.IsConnected())			{ StrataMaterialInfo.AddPropertyConnected(MP_Anisotropy); }
 	if (EmissiveColor.IsConnected())		{ StrataMaterialInfo.AddPropertyConnected(MP_EmissiveColor); }
@@ -23002,7 +22782,7 @@ bool UMaterialExpressionStrataSlabBSDF::HasMFPPluggedIn() const
 
 bool UMaterialExpressionStrataSlabBSDF::HasEdgeColor() const
 {
-	return bUseMetalness ? EdgeColor.IsConnected() : F90.IsConnected();
+	return F90.IsConnected();
 }
 
 bool UMaterialExpressionStrataSlabBSDF::HasFuzz() const
@@ -23040,7 +22820,6 @@ UMaterialExpressionStrataSimpleClearCoatBSDF::UMaterialExpressionStrataSimpleCle
 #if WITH_EDITOR
 int32 UMaterialExpressionStrataSimpleClearCoatBSDF::Compile(class FMaterialCompiler* Compiler, int32 OutputIndex)
 {
-	const bool  bUseMetalness = true;
 	const float DefaultSpecular = 0.5f;
 	const float DefaultF0 = DielectricSpecularToF0(DefaultSpecular);
 
@@ -23052,17 +22831,8 @@ int32 UMaterialExpressionStrataSimpleClearCoatBSDF::Compile(class FMaterialCompi
 	StrataOperator.BSDFRegisteredSharedLocalBasis = NewRegisteredSharedLocalBasis;
 
 	int32 OutputCodeChunk = Compiler->StrataSlabBSDF(
-		bUseMetalness ? Compiler->Constant(1.0f) : Compiler->Constant(0.0f),
-
-		// Metalness workflow
-		CompileWithDefaultFloat3(Compiler, BaseColor, 0.18f, 0.18f, 0.18f),
-		Compiler->Constant3(1.0f, 1.0f, 1.0f),					// EdgeColor
-		CompileWithDefaultFloat1(Compiler, Specular, DefaultSpecular),
-		CompileWithDefaultFloat1(Compiler, Metallic, 0.0f),
-
-		// !Metalness workflow unused for this node
-		Compiler->Constant3(0.18f, 0.18f, 0.18f),				// DiffuseAlbedo
-		Compiler->Constant3(DefaultF0, DefaultF0, DefaultF0),	// F0
+		CompileWithDefaultFloat3(Compiler, DiffuseAlbedo, 0.18f, 0.18f, 0.18f),		// DiffuseAlbedo
+		CompileWithDefaultFloat3(Compiler, F0, DefaultF0, DefaultF0, DefaultF0),	// F0
 		Compiler->Constant3(1.0f, 1.0f, 1.0f),					// F90
 		
 		CompileWithDefaultFloat1(Compiler, Roughness, 0.5f),
@@ -23091,9 +22861,8 @@ int32 UMaterialExpressionStrataSimpleClearCoatBSDF::Compile(class FMaterialCompi
 const TArray<FExpressionInput*> UMaterialExpressionStrataSimpleClearCoatBSDF::GetInputs()
 {
 	TArray<FExpressionInput*> Result;
-	Result.Add(&BaseColor);
-	Result.Add(&Metallic);
-	Result.Add(&Specular);
+	Result.Add(&DiffuseAlbedo);
+	Result.Add(&F0);
 	Result.Add(&Roughness);
 	Result.Add(&ClearCoatCoverage);
 	Result.Add(&ClearCoatRoughness);
@@ -23116,33 +22885,29 @@ uint32 UMaterialExpressionStrataSimpleClearCoatBSDF::GetInputType(int32 InputInd
 {
 	if (InputIndex == 0)
 	{
-		return MCT_Float3; // BaseColor
+		return MCT_Float3; // DiffuseAlbedo
 	}
 	else if (InputIndex == 1)
 	{
-		return MCT_Float1; // Metallic
+		return MCT_Float3; // F0
 	}
 	else if (InputIndex == 2)
-	{
-		return MCT_Float1; // Specular
-	}
-	else if (InputIndex == 3)
 	{
 		return MCT_Float1; // Roughness
 	}
-	else if (InputIndex == 4)
+	else if (InputIndex == 3)
 	{
 		return MCT_Float1; // ClearCoatCoverage 
 	}
-	else if (InputIndex == 5)
+	else if (InputIndex == 4)
 	{
 		return MCT_Float1; // ClearCoatRoughness
 	}
-	else if (InputIndex == 6)
+	else if (InputIndex == 5)
 	{
 		return MCT_Float3; // Normal
 	}
-	else if (InputIndex == 7)
+	else if (InputIndex == 6)
 	{
 		return MCT_Float3; // Emissive Color
 	}
@@ -23155,33 +22920,29 @@ FName UMaterialExpressionStrataSimpleClearCoatBSDF::GetInputName(int32 InputInde
 {
 	if (InputIndex == 0)
 	{
-		return TEXT("BaseColor");
+		return TEXT("Diffuse Albedo");
 	}
 	else if (InputIndex == 1)
 	{
-		return TEXT("Metallic");
+		return TEXT("F0");
 	}
 	else if (InputIndex == 2)
-	{
-		return TEXT("Specular");
-	}
-	else if (InputIndex == 3)
 	{
 		return TEXT("Roughness");
 	}
-	else if (InputIndex == 4)
+	else if (InputIndex == 3)
 	{
 		return TEXT("Clear Coat Coverage");
 	}
-	else if (InputIndex == 5)
+	else if (InputIndex == 4)
 	{
 		return TEXT("Clear Coat Roughness");
 	}
-	else if (InputIndex == 6)
+	else if (InputIndex == 5)
 	{
 		return TEXT("Normal");
 	}
-	else if (InputIndex == 7)
+	else if (InputIndex == 6)
 	{
 		return TEXT("Emissive Color");
 	}
@@ -23197,9 +22958,9 @@ bool UMaterialExpressionStrataSimpleClearCoatBSDF::IsResultStrataMaterial(int32
 void UMaterialExpressionStrataSimpleClearCoatBSDF::GatherStrataMaterialInfo(FStrataMaterialInfo& StrataMaterialInfo, int32 OutputIndex)
 {
 	// Track connected inputs
-	if (BaseColor.IsConnected()) { StrataMaterialInfo.AddPropertyConnected(MP_BaseColor); }
-	if (Metallic.IsConnected()) { StrataMaterialInfo.AddPropertyConnected(MP_Metallic); }
-	if (Specular.IsConnected()) { StrataMaterialInfo.AddPropertyConnected(MP_Specular); }
+	if (DiffuseAlbedo.IsConnected()) { StrataMaterialInfo.AddPropertyConnected(MP_BaseColor); }
+	StrataMaterialInfo.AddPropertyConnected(MP_Metallic); // Metallic is always connected with Diffuse/F0 parameterisation
+	if (F0.IsConnected()) { StrataMaterialInfo.AddPropertyConnected(MP_Specular); }
 	if (Roughness.IsConnected()) { StrataMaterialInfo.AddPropertyConnected(MP_Roughness); }
 	if (Normal.IsConnected()) { StrataMaterialInfo.AddPropertyConnected(MP_Normal); }
 	if (EmissiveColor.IsConnected()) { StrataMaterialInfo.AddPropertyConnected(MP_EmissiveColor); }
diff --git a/Engine/Source/Runtime/Engine/Public/MaterialCompiler.h b/Engine/Source/Runtime/Engine/Public/MaterialCompiler.h
--- a/Engine/Source/Runtime/Engine/Public/MaterialCompiler.h
+++ b/Engine/Source/Runtime/Engine/Public/MaterialCompiler.h
@@ -633,8 +633,6 @@ public:
 	// Strata
 	virtual int32 StrataCreateAndRegisterNullMaterial() = 0;
 	virtual int32 StrataSlabBSDF(
-		int32 UseMetalness,
-		int32 BaseColor, int32 EdgeColor, int32 Specular, int32 Metallic,
 		int32 DiffuseAlbedo, int32 F0, int32 F90,
 		int32 Roughness, int32 Anisotropy,
 		int32 SSSProfileId, int32 SSSMFP, int32 SSSMFPScale, int32 SSSPhaseAniso, int32 UseSSSDiffusion,
@@ -1226,8 +1224,6 @@ public:
 	}
 
 	virtual int32 StrataSlabBSDF(
-		int32 UseMetalness,
-		int32 BaseColor, int32 EdgeColor, int32 Specular, int32 Metallic,
 		int32 DiffuseAlbedo, int32 F0, int32 F90,
 		int32 Roughness, int32 Anisotropy,
 		int32 SSSProfileId, int32 SSSMFP, int32 SSSMFPScale, int32 SSSPhaseAniso, int32 UseSSSDiffusion,
@@ -1239,8 +1235,6 @@ public:
 		FStrataOperator* PromoteToOperator) override
 	{
 		return Compiler->StrataSlabBSDF(
-			UseMetalness,
-			BaseColor, EdgeColor, Specular, Metallic,
 			DiffuseAlbedo, F0, F90,
 			Roughness, Anisotropy,
 			SSSProfileId, SSSMFP, SSSMFPScale, SSSPhaseAniso, UseSSSDiffusion,
