Static analysis workaround.

diff --git a/Engine/Source/Runtime/Engine/Private/OnlineReplStructs.cpp b/Engine/Source/Runtime/Engine/Private/OnlineReplStructs.cpp
--- a/Engine/Source/Runtime/Engine/Private/OnlineReplStructs.cpp
+++ b/Engine/Source/Runtime/Engine/Private/OnlineReplStructs.cpp
@@ -904,7 +904,7 @@ void TestUniqueIdRepl(UWorld* InWorld)
 				ExportFailures += bOutSuccess ? 0 : 1;
 			}
 
-			if (ExportFailures > 0)
+			if (ExportFailures > 0) //-V547 Expression 'ExportFailures > 0' is always false - ExportTextItem always returns true
 			{
 				UE_LOG(LogNet, Warning, TEXT("There were %d export failures"), ExportFailures);
 				bTextItemSerializationSuccess = false;
