Content Bundles: don't rename runtime cells as they are already properly outered, to avoid them getting a generic name.
This makes the cook generation for partitioned worlds containing content bundles more deterministic.

#rnx

diff --git a/Engine/Source/Runtime/Engine/Private/WorldPartition/ContentBundle/ContentBundleEditor.cpp b/Engine/Source/Runtime/Engine/Private/WorldPartition/ContentBundle/ContentBundleEditor.cpp
--- a/Engine/Source/Runtime/Engine/Private/WorldPartition/ContentBundle/ContentBundleEditor.cpp
+++ b/Engine/Source/Runtime/Engine/Private/WorldPartition/ContentBundle/ContentBundleEditor.cpp
@@ -397,8 +397,8 @@ bool FContentBundleEditor::PopulateGeneratorPackageForCook(class IWorldPartition
 				UWorldPartitionRuntimeCell** MatchingCell = const_cast<UWorldPartitionRuntimeCell**>(CookPackageIdsToCell.Find(CookPackage->PackageId));
 				if (UWorldPartitionRuntimeCell* Cell = MatchingCell ? *MatchingCell : nullptr)
 				{
-					// Change outer to ExternalStreamingObject so it will be saved in the right package at the end of the cook.
-					Cell->Rename(nullptr, ExternalStreamingObject);
+					// Make sure the cell outer is set to the  ExternalStreamingObject so it will be saved in the right package at the end of the cook.
+					check(Cell->GetOuter() == ExternalStreamingObject);
 
 					if (!Cell->PrepareCellForCook(CookPackage->GetPackage()))
 					{
diff --git a/Engine/Source/Runtime/Engine/Private/WorldPartition/WorldPartitionRuntimeLevelStreamingCell.cpp b/Engine/Source/Runtime/Engine/Private/WorldPartition/WorldPartitionRuntimeLevelStreamingCell.cpp
--- a/Engine/Source/Runtime/Engine/Private/WorldPartition/WorldPartitionRuntimeLevelStreamingCell.cpp
+++ b/Engine/Source/Runtime/Engine/Private/WorldPartition/WorldPartitionRuntimeLevelStreamingCell.cpp
@@ -94,6 +94,7 @@ UWorldPartitionLevelStreamingDynamic* UWorldPartitionRuntimeLevelStreamingCell::
 		// When called by Commandlet (PopulateGeneratedPackageForCook), LevelStreaming's outer is set to Cell/WorldPartition's outer to prevent warnings when saving Cell Levels (Warning: Obj in another map). 
 		// At runtime, LevelStreaming's outer will be properly set to the main world (see UWorldPartitionRuntimeLevelStreamingCell::Activate).
 		UWorld* LevelStreamingOuterWorld = IsRunningCommandlet() ? OuterWorld : OwningWorld;
+		check(!FindObject<UWorldPartitionLevelStreamingDynamic>(LevelStreamingOuterWorld, *LevelStreamingName.ToString()));
 		UWorldPartitionLevelStreamingDynamic* NewLevelStreaming = NewObject<UWorldPartitionLevelStreamingDynamic>(LevelStreamingOuterWorld, UWorldPartitionLevelStreamingDynamic::StaticClass(), LevelStreamingName, RF_NoFlags, NULL);
 
 		FName WorldName = OuterWorld->GetFName();
