Fixed FWorldPartitionStreamingSource default ctor

diff --git a/Engine/Source/Runtime/Engine/Public/WorldPartition/WorldPartitionStreamingSource.h b/Engine/Source/Runtime/Engine/Public/WorldPartition/WorldPartitionStreamingSource.h
--- a/Engine/Source/Runtime/Engine/Public/WorldPartition/WorldPartitionStreamingSource.h
+++ b/Engine/Source/Runtime/Engine/Public/WorldPartition/WorldPartitionStreamingSource.h
@@ -372,7 +372,11 @@ enum class EStreamingSourcePriority : uint8
 struct ENGINE_API FWorldPartitionStreamingSource
 {
 	FWorldPartitionStreamingSource()
-		: bBlockOnSlowLoading(false)
+		: Name(NAME_None)
+		, Location(FVector::ZeroVector)
+		, Rotation(FRotator::ZeroRotator)
+		, TargetState(EStreamingSourceTargetState::Activated)
+		, bBlockOnSlowLoading(false)
 		, Priority(EStreamingSourcePriority::Default)
 		, Velocity(0.f)
 		, DebugColor(ForceInit)
