Fix non unity compile error in StreamableManager (introduced in cl 23285992)

diff --git a/Engine/Source/Runtime/Engine/Private/StreamableManager.cpp b/Engine/Source/Runtime/Engine/Private/StreamableManager.cpp
--- a/Engine/Source/Runtime/Engine/Private/StreamableManager.cpp
+++ b/Engine/Source/Runtime/Engine/Private/StreamableManager.cpp
@@ -3,6 +3,8 @@
 #include "Engine/StreamableManager.h"
 #include "UObject/WeakObjectPtr.h"
 #include "UObject/ObjectRedirector.h"
+#include "UObject/Package.h"
+#include "UObject/UObjectGlobals.h"
 #include "Misc/PackageName.h"
 #include "UObject/UObjectThreadContext.h"
 #include "HAL/IConsoleManager.h"
