Instancing HLOD builder will now handle Mesh Batching Policy on it's own - applying the correct settings, instead of relying on the base class to perform batching

diff --git a/Engine/Plugins/Editor/WorldPartitionHLODUtilities/Source/Public/HLODBuilderInstancing.h b/Engine/Plugins/Editor/WorldPartitionHLODUtilities/Source/Public/HLODBuilderInstancing.h
--- a/Engine/Plugins/Editor/WorldPartitionHLODUtilities/Source/Public/HLODBuilderInstancing.h
+++ b/Engine/Plugins/Editor/WorldPartitionHLODUtilities/Source/Public/HLODBuilderInstancing.h
@@ -32,7 +32,7 @@ class WORLDPARTITIONHLODUTILITIES_API UHLODBuilderInstancing : public UHLODBuild
 
 public:
 	virtual bool RequiresWarmup() const override { return false; }
-
+	virtual bool ShouldIgnoreBatchingPolicy() const override { return true; }
 	virtual TSubclassOf<UHLODBuilderSettings> GetSettingsClass() const override;
 	virtual TArray<UActorComponent*> Build(const FHLODBuildContext& InHLODBuildContext, const TArray<UActorComponent*>& InSourceComponents) const override;
 };
diff --git a/Engine/Source/Runtime/Engine/Private/WorldPartition/HLOD/HLODBuilder.cpp b/Engine/Source/Runtime/Engine/Private/WorldPartition/HLOD/HLODBuilder.cpp
--- a/Engine/Source/Runtime/Engine/Private/WorldPartition/HLOD/HLODBuilder.cpp
+++ b/Engine/Source/Runtime/Engine/Private/WorldPartition/HLOD/HLODBuilder.cpp
@@ -238,17 +238,25 @@ TArray<UActorComponent*> UHLODBuilder::Build(const FHLODBuildContext& InHLODBuil
 	// Handle components using a batching policy separately
 	TArray<UActorComponent*> InputComponents;
 	TArray<UActorComponent*> ComponentsToBatch;
-	for (UActorComponent* SourceComponent : HLODRelevantComponents)
-	{
-		if (ShouldBatchComponent(SourceComponent))
-		{
-			ComponentsToBatch.Add(SourceComponent);
-		}
-		else
+
+	if (!ShouldIgnoreBatchingPolicy())
+	{				
+		for (UActorComponent* SourceComponent : HLODRelevantComponents)
 		{
-			InputComponents.Add(SourceComponent);
+			if (ShouldBatchComponent(SourceComponent))
+			{
+				ComponentsToBatch.Add(SourceComponent);
+			}
+			else
+			{
+				InputComponents.Add(SourceComponent);
+			}
 		}
 	}
+	else
+	{
+		InputComponents = MoveTemp(HLODRelevantComponents);
+	}
 
 	TMap<TSubclassOf<UHLODBuilder>, TArray<UActorComponent*>> HLODBuildersForComponents;
 
diff --git a/Engine/Source/Runtime/Engine/Public/WorldPartition/HLOD/HLODBuilder.h b/Engine/Source/Runtime/Engine/Public/WorldPartition/HLOD/HLODBuilder.h
--- a/Engine/Source/Runtime/Engine/Public/WorldPartition/HLOD/HLODBuilder.h
+++ b/Engine/Source/Runtime/Engine/Public/WorldPartition/HLOD/HLODBuilder.h
@@ -115,6 +115,8 @@ public:
 	static uint32 ComputeHLODHash(const TArray<AActor*>& InSourceActors);
 
 protected:
+	virtual bool ShouldIgnoreBatchingPolicy() const { return false; }
+
 	static TArray<UActorComponent*> BatchInstances(const TArray<UActorComponent*>& InSubComponents);
 
 	template <typename TComponentClass>
