Added Instanced Static Mesh support for sweeps raycasts and overlaps over the primitive component interface
#[fyi] Michael.Lentine, Benn.Gallagher

diff --git a/Engine/Source/Runtime/Engine/Classes/Components/InstancedStaticMeshComponent.h b/Engine/Source/Runtime/Engine/Classes/Components/InstancedStaticMeshComponent.h
--- a/Engine/Source/Runtime/Engine/Classes/Components/InstancedStaticMeshComponent.h
+++ b/Engine/Source/Runtime/Engine/Classes/Components/InstancedStaticMeshComponent.h
@@ -358,13 +358,22 @@ public:
 	virtual void GetLightAndShadowMapMemoryUsage( int32& LightMapMemoryUsage, int32& ShadowMapMemoryUsage ) const override;
 
 	virtual bool DoCustomNavigableGeometryExport(FNavigableGeometryExport& GeomExport) const override;
+
+	virtual bool LineTraceComponent(FHitResult& OutHit, const FVector Start, const FVector End, const FCollisionQueryParams& Params) override;
+	virtual bool SweepComponent(FHitResult& OutHit, const FVector Start, const FVector End, const FQuat& ShapRotation, const FCollisionShape& CollisionShape, bool bTraceComplex = false) override;
+	virtual bool OverlapComponent(const FVector& Pos, const FQuat& Rot, const FCollisionShape& CollisionShape) const override;
+protected:
+	virtual bool ComponentOverlapComponentImpl(class UPrimitiveComponent* PrimComp, const FVector Pos, const FQuat& Quat, const FCollisionQueryParams& Params) override;
+	virtual bool ComponentOverlapMultiImpl(TArray<struct FOverlapResult>& OutOverlaps, const class UWorld* InWorld, const FVector& Pos, const FQuat& Rot, ECollisionChannel TestChannel, const struct FComponentQueryParams& Params, const struct FCollisionObjectQueryParams& ObjectQueryParams = FCollisionObjectQueryParams::DefaultObjectQueryParam) const override;
+public:
+
 	//~ End UPrimitiveComponent Interface
 
 	//~ Begin UNavRelevantInterface Interface
 	virtual void GetNavigationData(FNavigationRelevantData& Data) const override;
 	virtual FBox GetNavigationBounds() const override;
 	virtual bool IsNavigationRelevant() const override;
-	//~ End UPrimitiveComponent Interface
+	//~ End UNavRelevantInterface Interface
 
 	//~ Begin UObject Interface
 	virtual void Serialize(FArchive& Ar) override;
diff --git a/Engine/Source/Runtime/Engine/Private/Components/PrimitiveComponent.cpp b/Engine/Source/Runtime/Engine/Private/Components/PrimitiveComponent.cpp
--- a/Engine/Source/Runtime/Engine/Private/Components/PrimitiveComponent.cpp
+++ b/Engine/Source/Runtime/Engine/Private/Components/PrimitiveComponent.cpp
@@ -17,6 +17,7 @@
 #include "GameFramework/WorldSettings.h"
 #include "Components/SkeletalMeshComponent.h"
 #include "Engine/CollisionProfile.h"
+#include "Engine/InstancedStaticMesh.h"
 #include "Engine/SkeletalMesh.h"
 #include "Materials/MaterialInstanceDynamic.h"
 #include "Engine/Texture2D.h"
@@ -2895,6 +2896,20 @@ bool UPrimitiveComponent::ComponentOverlapComponentImpl(class UPrimitiveComponen
 		return false;
 	}
 
+	// if target is Instanced Static Meshes
+	UInstancedStaticMeshComponent* InstancedStaticMesh = Cast<UInstancedStaticMeshComponent>(PrimComp);
+	if (InstancedStaticMesh)
+	{
+		if (FBodyInstance* ThisBodyInstance = GetBodyInstance())
+		{
+			return ThisBodyInstance->OverlapTestForBodies(Pos, Quat, InstancedStaticMesh->InstanceBodies);
+		}
+		else
+		{
+			return false;
+		}
+	}
+
 	if(FBodyInstance* BI = PrimComp->GetBodyInstance())
 	{
 		if (FBodyInstance* ThisBodyInstance = GetBodyInstance())
diff --git a/Engine/Source/Runtime/Engine/Private/InstancedStaticMesh.cpp b/Engine/Source/Runtime/Engine/Private/InstancedStaticMesh.cpp
--- a/Engine/Source/Runtime/Engine/Private/InstancedStaticMesh.cpp
+++ b/Engine/Source/Runtime/Engine/Private/InstancedStaticMesh.cpp
@@ -34,6 +34,7 @@
 #include "NaniteSceneProxy.h"
 #include "MaterialCachedData.h"
 #include "Collision.h"
+#include "CollisionDebugDrawingPublic.h"
 #include "PipelineStateCache.h"
 
 #include "Elements/Framework/EngineElementsLibrary.h"
@@ -4627,6 +4628,124 @@ bool UInstancedStaticMeshComponent::DoCustomNavigableGeometryExport(FNavigableGe
 	return false;
 }
 
+extern float DebugLineLifetime;
+
+bool UInstancedStaticMeshComponent::LineTraceComponent(FHitResult& OutHit, const FVector Start, const FVector End, const FCollisionQueryParams& Params)
+{
+	bool bHaveHit = false;	
+	float MinTime = UE_MAX_FLT;
+	FHitResult Hit;
+
+	// TODO: use spatial acceleration instead
+	for (FBodyInstance* Body : InstanceBodies)
+	{
+		if (Body->LineTrace(Hit, Start, End, Params.bTraceComplex, Params.bReturnPhysicalMaterial))
+		{
+			bHaveHit = true;
+			if (MinTime > Hit.Time)
+			{
+				MinTime = Hit.Time;
+				OutHit = Hit;
+			}
+		}
+	}
+
+#if !(UE_BUILD_SHIPPING || UE_BUILD_TEST)
+	UWorld* const World = GetWorld();
+	if (World && World->DebugDrawSceneQueries(Params.TraceTag))
+	{
+		TArray<FHitResult> Hits;
+		if (bHaveHit)
+		{
+			Hits.Add(OutHit);
+		}
+		DrawLineTraces(GetWorld(), Start, End, Hits, DebugLineLifetime);
+	}
+#endif //!(UE_BUILD_SHIPPING || UE_BUILD_TEST)
+
+	return bHaveHit;
+}
+
+bool UInstancedStaticMeshComponent::SweepComponent(FHitResult& OutHit, const FVector Start, const FVector End, const FQuat& ShapeWorldRotation, const FCollisionShape& CollisionShape, bool bTraceComplex)
+{
+	bool bHaveHit = false;
+	
+	FHitResult Hit;
+	// TODO: use spatial acceleration instead
+	for (FBodyInstance* Body : InstanceBodies)
+	{
+		if (Body->Sweep(Hit, Start, End, ShapeWorldRotation, CollisionShape, bTraceComplex))
+		{
+			if (!bHaveHit || Hit.Time < OutHit.Time)
+			{
+				OutHit = Hit;
+			}
+			bHaveHit = true;
+		}
+	}
+	return bHaveHit;
+}
+
+bool UInstancedStaticMeshComponent::OverlapComponent(const FVector& Pos, const FQuat& Rot, const FCollisionShape& CollisionShape) const
+{	
+	for (FBodyInstance* Body : InstanceBodies)
+	{
+		if (Body->OverlapTest(Pos, Rot, CollisionShape))
+		{
+			return true;
+		}
+	}
+
+	return false;	
+}
+
+bool UInstancedStaticMeshComponent::ComponentOverlapComponentImpl(class UPrimitiveComponent* PrimComp, const FVector Pos, const FQuat& Quat, const FCollisionQueryParams& Params)
+{
+	//we do not support skeletal mesh vs InstancedStaticMesh overlap test
+	if (PrimComp->IsA<USkeletalMeshComponent>())
+	{
+		UE_LOG(LogCollision, Warning, TEXT("ComponentOverlapComponent : (%s) Does not support InstancedStaticMesh with Physics Asset"), *PrimComp->GetPathName());
+		return false;
+	}
+
+	//We do not support Instanced static meshes vs Instanced static meshes
+	if (PrimComp->IsA<UInstancedStaticMeshComponent>())
+	{
+		UE_LOG(LogCollision, Warning, TEXT("ComponentOverlapComponent : (%s) Does not support InstancedStaticMesh with Physics Asset"), *PrimComp->GetPathName());
+		return false;
+	}
+
+	if (FBodyInstance* BI = PrimComp->GetBodyInstance())
+	{
+		return BI->OverlapTestForBodies(Pos, Quat, InstanceBodies);
+	}
+
+	return false;
+}
+
+bool UInstancedStaticMeshComponent::ComponentOverlapMultiImpl(TArray<struct FOverlapResult>& OutOverlaps, const class UWorld* InWorld, const FVector& Pos, const FQuat& Rot, ECollisionChannel TestChannel, const struct FComponentQueryParams& Params, const struct FCollisionObjectQueryParams& ObjectQueryParams) const
+{
+	OutOverlaps.Reset();
+
+	const FTransform WorldToComponent(GetComponentTransform().Inverse());
+	const FCollisionResponseParams ResponseParams(GetCollisionResponseToChannels());
+
+	FComponentQueryParams ParamsWithSelf = Params;
+	ParamsWithSelf.AddIgnoredComponent(this);
+
+	bool bHaveBlockingHit = false;
+	for (FBodyInstance* Body : InstanceBodies)
+	{
+		checkSlow(Body);
+		if (Body->OverlapMulti(OutOverlaps, InWorld, &WorldToComponent, Pos, Rot, TestChannel, ParamsWithSelf, ResponseParams, ObjectQueryParams))
+		{
+			bHaveBlockingHit = true;
+		}
+	}
+
+	return bHaveBlockingHit;
+}
+
 void UInstancedStaticMeshComponent::GetNavigationData(FNavigationRelevantData& Data) const
 {
 	// Navigation data will get refreshed once async compilation finishes
diff --git a/Engine/Source/Runtime/Engine/Private/SkeletalMeshComponentPhysics.cpp b/Engine/Source/Runtime/Engine/Private/SkeletalMeshComponentPhysics.cpp
--- a/Engine/Source/Runtime/Engine/Private/SkeletalMeshComponentPhysics.cpp
+++ b/Engine/Source/Runtime/Engine/Private/SkeletalMeshComponentPhysics.cpp
@@ -11,6 +11,7 @@
 #include "EngineDefines.h"
 #include "Engine/EngineBaseTypes.h"
 #include "Engine/EngineTypes.h"
+#include "Engine/InstancedStaticMesh.h"
 #include "Components/ActorComponent.h"
 #include "Components/SceneComponent.h"
 #include "CollisionQueryParams.h"
@@ -2624,6 +2625,13 @@ bool USkeletalMeshComponent::ComponentOverlapComponentImpl(class UPrimitiveCompo
 		return false;
 	}
 
+	//We do not support skeletal mesh vs Instanced static meshes
+	if (PrimComp->IsA<UInstancedStaticMeshComponent>())
+	{
+		UE_LOG(LogCollision, Warning, TEXT("ComponentOverlapComponent : (%s) Does not support skeletalmesh with Physics Asset"), *PrimComp->GetPathName());
+		return false;
+	}
+
 	if (FBodyInstance* BI = PrimComp->GetBodyInstance())
 	{
 		return BI->OverlapTestForBodies(Pos, Quat, Bodies);
