Fix typo in skeletal mesh insight tag

#rnx

diff --git a/Engine/Source/Runtime/Engine/Private/SkeletalMesh.cpp b/Engine/Source/Runtime/Engine/Private/SkeletalMesh.cpp
--- a/Engine/Source/Runtime/Engine/Private/SkeletalMesh.cpp
+++ b/Engine/Source/Runtime/Engine/Private/SkeletalMesh.cpp
@@ -2991,7 +2991,7 @@ bool USkeletalMesh::IsPostLoadThreadSafe() const
 void USkeletalMesh::BeginPostLoadInternal(FSkinnedAssetPostLoadContext& Context)
 {
 #if WITH_EDITOR
-	TRACE_CPUPROFILER_EVENT_SCOPE(UStaticMesh::BeginPostLoadInternal);
+	TRACE_CPUPROFILER_EVENT_SCOPE(USkeletalMesh::BeginPostLoadInternal);
 
 	SetInternalFlags(EInternalObjectFlags::Async);
 
