ShaderCompiler MPCook - Disambiguate MaterialStats csv outputfilenames with the multiprocessid to avoid collisions when multiple CookWorker process are writing the file at the same time.
#rnx

diff --git a/Engine/Source/Runtime/Engine/Private/ShaderCompiler/ShaderCompiler.cpp b/Engine/Source/Runtime/Engine/Private/ShaderCompiler/ShaderCompiler.cpp
--- a/Engine/Source/Runtime/Engine/Private/ShaderCompiler/ShaderCompiler.cpp
+++ b/Engine/Source/Runtime/Engine/Private/ShaderCompiler/ShaderCompiler.cpp
@@ -3102,7 +3102,13 @@ void FShaderCompilerStats::WriteStats(FOutputDevice* Ar)
 	{
 		FlushRenderingCommands();
 
-		FString FileName = FPaths::Combine(*FPaths::ProjectSavedDir(), FString::Printf(TEXT("MaterialStats/Stats-%s.csv"), *FDateTime::Now().ToString()));
+		uint32 MultiprocessId = 0;
+		FParse::Value(FCommandLine::Get(), TEXT("-MultiprocessId="), MultiprocessId);
+		FString FileName = FPaths::Combine(*FPaths::ProjectSavedDir(),
+			FString::Printf(TEXT("MaterialStats/Stats-%s%s.csv"),
+				*FDateTime::Now().ToString(),
+				(MultiprocessId == 0 ? TEXT("") : *FString::Printf(TEXT("-%d"), MultiprocessId))
+			));
 		auto DebugWriter = IFileManager::Get().CreateFileWriter(*FileName);
 		FDiagnosticTableWriterCSV StatWriter(DebugWriter);
 		const TSparseArray<ShaderCompilerStats>& PlatformStats = GetShaderCompilerStats();
@@ -3186,7 +3192,12 @@ void FShaderCompilerStats::WriteStats(FOutputDevice* Ar)
 	}
 	{
 
-		FString FileName = FString::Printf(TEXT("%s/MaterialStatsDebug/StatsDebug-%s.csv"), *FPaths::ProjectSavedDir(), *FDateTime::Now().ToString());
+		uint32 MultiprocessId = 0;
+		FParse::Value(FCommandLine::Get(), TEXT("-MultiprocessId="), MultiprocessId);
+		FString FileName = FString::Printf(TEXT("%s/MaterialStatsDebug/StatsDebug-%s%s.csv"),
+			*FPaths::ProjectSavedDir(), *FDateTime::Now().ToString(),
+			(MultiprocessId == 0 ? TEXT("") : *FString::Printf(TEXT("-%d"), MultiprocessId))
+		);
 		auto DebugWriter = IFileManager::Get().CreateFileWriter(*FileName);
 		FDiagnosticTableWriterCSV StatWriter(DebugWriter);
 		const TSparseArray<ShaderCompilerStats>& PlatformStats = GetShaderCompilerStats();
