Refraction/distortion - added a new mode for explicit use specified 2d offset.
Also clarified the refraction data accessible from the PixelInputData.

diff --git a/Engine/Shaders/Private/DistortAccumulatePS.usf b/Engine/Shaders/Private/DistortAccumulatePS.usf
--- a/Engine/Shaders/Private/DistortAccumulatePS.usf
+++ b/Engine/Shaders/Private/DistortAccumulatePS.usf
@@ -81,8 +81,7 @@ void Main(
 	float2 NDC = (MaterialParameters.ScreenPosition.xy / MaterialParameters.ScreenPosition.w);
 	float2 ScreenUV = NDC * ResolvedView.ScreenPositionScaleBias.xy + ResolvedView.ScreenPositionScaleBias.wz;
 
-	// X = IOR, Y = Refraction Bias
-	float2 Refraction = GetMaterialRefraction(PixelMaterialInputs);
+	FMaterialRefractionData RefractionData = GetMaterialRefraction(PixelMaterialInputs);
 
 
 #if STRATA_ENABLED
@@ -124,7 +123,7 @@ void Main(
 	// Compute UV distortion
 	float2 BufferUVDistortion = ComputeBufferUVDistortion(
 		MaterialParameters, PixelMaterialInputs, ResolvedView,
-		RefractionWorldNormal, DistortionParams, ScreenUV, Refraction, TryToClip, EyeIndex);
+		RefractionWorldNormal, DistortionParams, ScreenUV, RefractionData, TryToClip, EyeIndex);
 
 	// Sample depth at distortion offset
 	float2 DistortBufferUV = ScreenUV + BufferUVDistortion;
@@ -135,7 +134,7 @@ void Main(
 	// Compute UV distortion
 	float2 BufferUVDistortion = ComputeBufferUVDistortion(
 		MaterialParameters, PixelMaterialInputs, ResolvedView,
-		MaterialParameters.WorldNormal, DistortionParams, ScreenUV, Refraction, true, EyeIndex);
+		MaterialParameters.WorldNormal, DistortionParams, ScreenUV, RefractionData, true, EyeIndex);
 
 	// Sample depth at distortion offset
 	float2 DistortBufferUV = ScreenUV + BufferUVDistortion;
@@ -146,7 +145,7 @@ void Main(
 	float DistortSceneDepth = CalcSceneDepth(DistortBufferUV);
 
 	// Post process UV distortion according to depth
-	PostProcessUVDistortion(MaterialParameters, PixelMaterialInputs, DistortSceneDepth, BufferUVDistortion, Refraction);
+	PostProcessUVDistortion(MaterialParameters, PixelMaterialInputs, DistortSceneDepth, BufferUVDistortion, RefractionData);
 
 	// store positive and negative offsets separately
 	float2 PosOffset = max(BufferUVDistortion,0);
diff --git a/Engine/Shaders/Private/DistortionCommon.ush b/Engine/Shaders/Private/DistortionCommon.ush
--- a/Engine/Shaders/Private/DistortionCommon.ush
+++ b/Engine/Shaders/Private/DistortionCommon.ush
@@ -6,18 +6,22 @@
 
 float2 ComputeBufferUVDistortion(
 	in FMaterialPixelParameters MaterialParameters, in FPixelMaterialInputs PixelMaterialInputs, in ViewState ResolvedView,
-	in half3 WorldNormal, in float4 DistortionParameters, in float2 ScreenUV, in float2 RefractionIORAndBias, in bool TryToClip, in uint EyeIndex)
+	in half3 WorldNormal, in float4 DistortionParameters, in float2 ScreenUV, in FMaterialRefractionData RefractionData, in bool TryToClip, in uint EyeIndex)
 {
 	half3 ViewNormal = normalize(TransformWorldVectorToView(WorldNormal));
 
 #if REFRACTION_USE_PIXEL_NORMAL_OFFSET
 	half3 ViewVertexNormal = TransformTangentVectorToView(MaterialParameters, float3(0, 0, 1));
 	// Treat refraction of 1 as no refraction, to be consistent with IOR mode
-	float2 ViewportUVDistortion = (ViewVertexNormal.xy - ViewNormal.xy) * (RefractionIORAndBias.x - 1);
-#else
+	float2 ViewportUVDistortion = (ViewVertexNormal.xy - ViewNormal.xy) * (GetMaterialRefractionPixelNormalStrength(RefractionData) - 1);
+#elif REFRACTION_USE_INDEX_OF_REFRACTION
 	// we assume the camera is in air
 	float AirIOR = 1.0f;
-	float2 ViewportUVDistortion = ViewNormal.xy * (RefractionIORAndBias.x - AirIOR);
+	float2 ViewportUVDistortion = ViewNormal.xy * (GetMaterialRefractionIOR(RefractionData) - AirIOR);
+#elif REFRACTION_USE_2D_OFFSET
+	float2 ViewportUVDistortion = GetMaterialRefraction2DOffset(RefractionData);
+#else
+#error Unkown refraction mode.
 #endif
 
 	float2 BufferUVDistortion = ViewportUVDistortion * ResolvedView.ViewSizeAndInvSize.xy * ResolvedView.BufferSizeAndInvSize.zw;
@@ -55,10 +59,10 @@ float2 ComputeBufferUVDistortion(
 
 void PostProcessUVDistortion(
 	in FMaterialPixelParameters MaterialParameters, in FPixelMaterialInputs PixelMaterialInputs, 
-	in float DistortSceneDepth,	inout float2 BufferUVDistortion, in float2 RefractionIORAndBias)
+	in float DistortSceneDepth,	inout float2 BufferUVDistortion, in FMaterialRefractionData RefractionData)
 {
 	// Soft thresholding 
-	float Bias = -RefractionIORAndBias.y;
+	float Bias = -RefractionData.RefractionDepthBias;
 	float Range = clamp(abs(Bias * 0.5f), 0, 50);
 	float Z = DistortSceneDepth;
 	float ZCompare = MaterialParameters.ScreenPosition.w;
diff --git a/Engine/Shaders/Private/MaterialTemplate.ush b/Engine/Shaders/Private/MaterialTemplate.ush
--- a/Engine/Shaders/Private/MaterialTemplate.ush
+++ b/Engine/Shaders/Private/MaterialTemplate.ush
@@ -2912,9 +2912,45 @@ half GetMaterialAmbientOcclusion(FPixelMaterialInputs PixelMaterialInputs)
 	return saturate(GetMaterialAmbientOcclusionRaw(PixelMaterialInputs));
 }
 
-half2 GetMaterialRefraction(FPixelMaterialInputs PixelMaterialInputs)
+struct FMaterialRefractionData
 {
-	return PixelMaterialInputs.Refraction;
+	float2 Data;
+	float  RefractionDepthBias;
+};
+
+float GetMaterialRefractionIOR(in FMaterialRefractionData RefractionData)
+{
+#if REFRACTION_USE_INDEX_OF_REFRACTION
+	return RefractionData.Data.x;
+#else
+	return 1.0f;
+#endif
+}
+
+float GetMaterialRefractionPixelNormalStrength(in FMaterialRefractionData RefractionData)
+{
+#if REFRACTION_USE_PIXEL_NORMAL_OFFSET
+	return RefractionData.Data.x;
+#else
+	return 1.0f;
+#endif
+}
+
+float2 GetMaterialRefraction2DOffset(in FMaterialRefractionData RefractionData)
+{
+#if REFRACTION_USE_2D_OFFSET
+	return RefractionData.Data.xy;
+#else
+	return float2(0.0f, 0.0f);
+#endif
+}
+
+FMaterialRefractionData GetMaterialRefraction(FPixelMaterialInputs PixelMaterialInputs)
+{
+	FMaterialRefractionData Data;
+	Data.Data					= PixelMaterialInputs.Refraction.xy;
+	Data.RefractionDepthBias	= PixelMaterialInputs.Refraction.z;
+	return Data;
 }
 
 #if NUM_TEX_COORD_INTERPOLATORS
diff --git a/Engine/Shaders/Private/PathTracing/PathTracingMaterialHitShader.usf b/Engine/Shaders/Private/PathTracing/PathTracingMaterialHitShader.usf
--- a/Engine/Shaders/Private/PathTracing/PathTracingMaterialHitShader.usf
+++ b/Engine/Shaders/Private/PathTracing/PathTracingMaterialHitShader.usf
@@ -118,7 +118,7 @@ RAY_TRACING_ENTRY_CLOSEST_HIT(PathTracingMaterialCHS,
 	// Store the results in local variables and reuse instead of calling the functions multiple times.
 	half3 BaseColor = GetMaterialBaseColor(PixelMaterialInputs);
 #if MATERIALBLENDING_TRANSLUCENT && REFRACTION_USE_INDEX_OF_REFRACTION
-	Payload.Ior = GetMaterialRefraction(PixelMaterialInputs).x;
+	Payload.Ior = GetMaterialRefractionIOR(GetMaterialRefraction(PixelMaterialInputs));
 	Payload.Ior = (Payload.Ior == 0.0) ? 1.0 : Payload.Ior;
 #else
 	Payload.Ior = 1;
@@ -487,7 +487,7 @@ RAY_TRACING_ENTRY_ANY_HIT(PathTracingMaterialAHS,
 	float3 Transmission = GetThinTranslucentMaterialOutput0(MaterialParameters);
 	float Opacity = GetMaterialOpacity(PixelMaterialInputs);
 #if REFRACTION_USE_INDEX_OF_REFRACTION && !SIMPLIFIED_MATERIAL_SHADER
-	float Ior = GetMaterialRefraction(PixelMaterialInputs).x;
+	float Ior = GetMaterialRefractionIOR(GetMaterialRefraction(PixelMaterialInputs));
 	Ior = Ior == 0 ? 1.0 : Ior;
 #else
 	float Ior = 1.0;
@@ -534,7 +534,7 @@ RAY_TRACING_ENTRY_ANY_HIT(PathTracingMaterialAHS,
 	
 	// Is refraction enabled?
 #if !SIMPLIFIED_MATERIAL_SHADER && MATERIAL_SHADINGMODEL_DEFAULT_LIT && REFRACTION_USE_INDEX_OF_REFRACTION
-	float Ior = GetMaterialRefraction(PixelMaterialInputs).x;
+	float Ior = GetMaterialRefractionIOR(GetMaterialRefraction(PixelMaterialInputs));
 	Ior = Ior == 0.0 ? 1.0 : Ior;
 	int ShadingModelID = GetMaterialShadingModel(PixelMaterialInputs);
 	if (Transparency > 0 && ShadingModelID == SHADINGMODELID_DEFAULT_LIT && Ior != 1.0)
diff --git a/Engine/Shaders/Private/RayTracing/RayTracingMaterialHitShaders.usf b/Engine/Shaders/Private/RayTracing/RayTracingMaterialHitShaders.usf
--- a/Engine/Shaders/Private/RayTracing/RayTracingMaterialHitShaders.usf
+++ b/Engine/Shaders/Private/RayTracing/RayTracingMaterialHitShaders.usf
@@ -313,9 +313,9 @@ void GetMaterialPayload(
 	// Store the results in local variables and reuse instead of calling the functions multiple times.
 	half3 BaseColor = GetMaterialBaseColor(PixelMaterialInputs);
 #if MATERIALBLENDING_TRANSLUCENT && REFRACTION_USE_INDEX_OF_REFRACTION
-	half2 Refraction = GetMaterialRefraction(PixelMaterialInputs);
+	half IOR = GetMaterialRefractionIOR(GetMaterialRefraction(PixelMaterialInputs));
 #else
-	half2 Refraction = 1;
+	half IOR = 1;
 #endif
 	half  Metallic = GetMaterialMetallic(PixelMaterialInputs);
 	half  Specular = GetMaterialSpecular(PixelMaterialInputs);
@@ -341,7 +341,7 @@ void GetMaterialPayload(
 	}
 #endif // MATERIALBLENDING_MASKED
 
-	Payload.Ior = Refraction.x == 0 ? 1.0 : Refraction.x;
+	Payload.Ior = IOR == 0 ? 1.0 : IOR;
 	Payload.CustomData = float4(GetMaterialCustomData0(MaterialParameters), GetMaterialCustomData1(MaterialParameters), 0, 0);
 	
 #if MATERIAL_SHADINGMODEL_CLEAR_COAT
diff --git a/Engine/Shaders/Private/SingleLayerWaterShading.ush b/Engine/Shaders/Private/SingleLayerWaterShading.ush
--- a/Engine/Shaders/Private/SingleLayerWaterShading.ush
+++ b/Engine/Shaders/Private/SingleLayerWaterShading.ush
@@ -89,12 +89,12 @@ WaterVolumeLightingOutput EvaluateWaterVolumeLighting(
 
 		float PixelSceneDepth = ConvertFromDeviceZ(PixelSceneDeviceZ);
 		
-		float2 Refraction = GetMaterialRefraction(PixelMaterialInputs); // X = IOR, Y = Refraction Bias
+		FMaterialRefractionData RefractionData = GetMaterialRefraction(PixelMaterialInputs);
 		float2 BufferUVDistortion = ComputeBufferUVDistortion(
 			MaterialParameters, PixelMaterialInputs, ResolvedView,
 			MaterialParameters.WorldNormal,
-			DistortionParams, ViewportUV, Refraction, false, EyeIndex);
-		PostProcessUVDistortion(MaterialParameters, PixelMaterialInputs, PixelSceneDepth, BufferUVDistortion, Refraction);
+			DistortionParams, ViewportUV, RefractionData, false, EyeIndex);
+		PostProcessUVDistortion(MaterialParameters, PixelMaterialInputs, PixelSceneDepth, BufferUVDistortion, RefractionData);
 		// Also reduce refraction according to water thickness
 		BufferUVDistortion *= saturate((PixelSceneDepth - WaterDepth) * 1.0 / 30.0f);// *(1 - saturate(WaterDepth * 1.0 / 2000.0f));
 
diff --git a/Engine/Source/Runtime/Engine/Classes/Engine/EngineTypes.h b/Engine/Source/Runtime/Engine/Classes/Engine/EngineTypes.h
--- a/Engine/Source/Runtime/Engine/Classes/Engine/EngineTypes.h
+++ b/Engine/Source/Runtime/Engine/Classes/Engine/EngineTypes.h
@@ -383,7 +383,12 @@ enum ERefractionMode
 	 * The refraction material input scales the offset, although a value of 1.0 maps to no refraction, and a value of 2 maps to a scale of 1.0 on the offset.
 	 * This is a non-physical model of refraction but is useful on large refractive surfaces like water, since offsets have to stay small to avoid reading outside scene color.
 	 */
-	RM_PixelNormalOffset UMETA(DisplayName="Pixel Normal Offset")
+	RM_PixelNormalOffset UMETA(DisplayName="Pixel Normal Offset"),
+
+	/**
+	 * Explicit 2D screen offset. This offset is independent of screen resolution and aspect ratio. The user is in charge of any strength and fading.
+	 */
+	RM_2DOffset UMETA(DisplayName = "2D Offset"),
 };
 
 /**
diff --git a/Engine/Source/Runtime/Engine/Private/Materials/HLSLMaterialTranslator.cpp b/Engine/Source/Runtime/Engine/Private/Materials/HLSLMaterialTranslator.cpp
--- a/Engine/Source/Runtime/Engine/Private/Materials/HLSLMaterialTranslator.cpp
+++ b/Engine/Source/Runtime/Engine/Private/Materials/HLSLMaterialTranslator.cpp
@@ -904,7 +904,8 @@ bool FHLSLMaterialTranslator::Translate()
 
 			if (IsTranslucentBlendMode(BlendMode) || MaterialShadingModels.HasShadingModel(MSM_SingleLayerWater))
 			{
-				int32 UserRefraction = ForceCast(Material->CompilePropertyAndSetMaterialProperty(MP_Refraction, this), MCT_Float1);
+				// Cast to exact match is needed for float parameter to be correctly cast to float2.
+				int32 UserRefraction = ForceCast(Material->CompilePropertyAndSetMaterialProperty(MP_Refraction, this), MCT_Float2, MFCF_ExactMatch);
 				int32 RefractionDepthBias = ForceCast(ScalarParameter(FName(TEXT("RefractionDepthBias")), Material->GetRefractionDepthBiasValue()), MCT_Float1);
 
 				Chunk[MP_Refraction] = AppendVector(UserRefraction, RefractionDepthBias);
diff --git a/Engine/Source/Runtime/Engine/Private/Materials/MaterialShared.cpp b/Engine/Source/Runtime/Engine/Private/Materials/MaterialShared.cpp
--- a/Engine/Source/Runtime/Engine/Private/Materials/MaterialShared.cpp
+++ b/Engine/Source/Runtime/Engine/Private/Materials/MaterialShared.cpp
@@ -2325,8 +2325,9 @@ void FMaterial::SetupMaterialEnvironment(
 
 	switch(GetRefractionMode())
 	{
-	case RM_IndexOfRefraction: OutEnvironment.SetDefine(TEXT("REFRACTION_USE_INDEX_OF_REFRACTION"),TEXT("1")); break;
-	case RM_PixelNormalOffset: OutEnvironment.SetDefine(TEXT("REFRACTION_USE_PIXEL_NORMAL_OFFSET"),TEXT("1")); break;
+	case RM_IndexOfRefraction:	OutEnvironment.SetDefine(TEXT("REFRACTION_USE_INDEX_OF_REFRACTION"),	TEXT("1")); break;
+	case RM_PixelNormalOffset:	OutEnvironment.SetDefine(TEXT("REFRACTION_USE_PIXEL_NORMAL_OFFSET"),	TEXT("1")); break;
+	case RM_2DOffset:			OutEnvironment.SetDefine(TEXT("REFRACTION_USE_2D_OFFSET"),				TEXT("1")); break;
 	default: 
 		UE_LOG(LogMaterial, Warning, TEXT("Unknown material refraction mode: %u  Setting to RM_IndexOfRefraction"),(int32)GetRefractionMode());
 		OutEnvironment.SetDefine(TEXT("REFRACTION_USE_INDEX_OF_REFRACTION"),TEXT("1"));
@@ -5421,7 +5422,7 @@ void FMaterialAttributeDefinitionMap::InitializeAttributeMap()
 	Add(FGuid(0x9E502E69, 0x3C8F48FA, 0x94645CFD, 0x28E5428D), TEXT("ClearCoat"),				MP_CustomData0,				MCT_Float,	FVector4(1,0,0,0),	SF_Pixel);
 	Add(FGuid(0xBE4F2FFD, 0x12FC4296, 0xB0124EEA, 0x12C28D92), TEXT("ClearCoatRoughness"),		MP_CustomData1,				MCT_Float,	FVector4(.1,0,0,0),	SF_Pixel);
 	Add(FGuid(0xE8EBD0AD, 0xB1654CBE, 0xB079C3A8, 0xB39B9F15), TEXT("AmbientOcclusion"),		MP_AmbientOcclusion,		MCT_Float,	FVector4(1,0,0,0),	SF_Pixel);
-	Add(FGuid(0xD0B0FA03, 0x14D74455, 0xA851BAC5, 0x81A0788B), TEXT("Refraction"),				MP_Refraction,				MCT_Float2,	FVector4(1,0,0,0),	SF_Pixel);
+	Add(FGuid(0xD0B0FA03, 0x14D74455, 0xA851BAC5, 0x81A0788B), TEXT("Refraction"),				MP_Refraction,				MCT_Float3,	FVector4(1,0,0,0),	SF_Pixel);
 	Add(FGuid(0x0AC97EC3, 0xE3D047BA, 0xB610167D, 0xC4D919FF), TEXT("PixelDepthOffset"),		MP_PixelDepthOffset,		MCT_Float,	FVector4(0,0,0,0),	SF_Pixel);
 	Add(FGuid(0xD9423FFF, 0xD77E4D82, 0x8FF9CF5E, 0x055D1255), TEXT("ShadingModel"),			MP_ShadingModel,			MCT_ShadingModel, FVector4(0, 0, 0, 0), SF_Pixel, INDEX_NONE, false, &CompileShadingModelBlendFunction);
 	Add(FGuid(0x5973A03E, 0x13A74E08, 0x92D0CEDD, 0xF2936CF8), TEXT("FrontMaterial"),			MP_FrontMaterial,			MCT_Strata, FVector4(0,0,0,0),	SF_Pixel, INDEX_NONE, false, &CompileStrataBlendFunction);
diff --git a/Engine/Source/Runtime/Engine/Private/ShaderCompiler/ShaderGenerationUtil.cpp b/Engine/Source/Runtime/Engine/Private/ShaderCompiler/ShaderGenerationUtil.cpp
--- a/Engine/Source/Runtime/Engine/Private/ShaderCompiler/ShaderGenerationUtil.cpp
+++ b/Engine/Source/Runtime/Engine/Private/ShaderCompiler/ShaderGenerationUtil.cpp
@@ -181,6 +181,7 @@ void FShaderCompileUtilities::ApplyFetchEnvironment(FShaderMaterialPropertyDefin
 
 	FETCH_COMPILE_BOOL(REFRACTION_USE_INDEX_OF_REFRACTION);
 	FETCH_COMPILE_BOOL(REFRACTION_USE_PIXEL_NORMAL_OFFSET);
+	FETCH_COMPILE_BOOL(REFRACTION_USE_2D_OFFSET);
 
 	FETCH_COMPILE_BOOL(USE_DITHERED_LOD_TRANSITION_FROM_MATERIAL);
 	FETCH_COMPILE_BOOL(MATERIAL_TWOSIDED);
diff --git a/Engine/Source/Runtime/RenderCore/Public/ShaderMaterial.h b/Engine/Source/Runtime/RenderCore/Public/ShaderMaterial.h
--- a/Engine/Source/Runtime/RenderCore/Public/ShaderMaterial.h
+++ b/Engine/Source/Runtime/RenderCore/Public/ShaderMaterial.h
@@ -135,6 +135,7 @@ struct FShaderMaterialPropertyDefines
 	
 	uint8 REFRACTION_USE_INDEX_OF_REFRACTION : 1;
 	uint8 REFRACTION_USE_PIXEL_NORMAL_OFFSET : 1;
+	uint8 REFRACTION_USE_2D_OFFSET : 1;
 
 	uint8 USE_DITHERED_LOD_TRANSITION_FROM_MATERIAL : 1;
 	uint8 MATERIAL_TWOSIDED : 1;
