Add VirtualLoop Audio Dashboard: Part 1
- Checkpoint: Misc fix-ups and tweaks to support extendible dashboards
#rnx

diff --git a/Engine/Source/Runtime/Engine/Private/AudioDebug.cpp b/Engine/Source/Runtime/Engine/Private/AudioDebug.cpp
--- a/Engine/Source/Runtime/Engine/Private/AudioDebug.cpp
+++ b/Engine/Source/Runtime/Engine/Private/AudioDebug.cpp
@@ -867,6 +867,11 @@ namespace Audio
 		bVisualize3dDebug = !bVisualize3dDebug;
 	}
 
+	bool FAudioDebugger::IsVirtualLoopVisualizeEnabled()
+	{
+		return static_cast<bool>(VirtualLoopsVisualizeEnabledCVar);
+	}
+
 	#if WITH_EDITOR
 	void FAudioDebugger::OnBeginPIE()
 	{
diff --git a/Engine/Source/Runtime/Engine/Private/AudioDevice.cpp b/Engine/Source/Runtime/Engine/Private/AudioDevice.cpp
--- a/Engine/Source/Runtime/Engine/Private/AudioDevice.cpp
+++ b/Engine/Source/Runtime/Engine/Private/AudioDevice.cpp
@@ -9,6 +9,7 @@
 #include "AudioDecompress.h"
 #include "AudioDefines.h"
 #include "AudioEffect.h"
+#include "AudioMixerTrace.h"
 #include "AudioPluginUtilities.h"
 #include "Audio/AudioDebug.h"
 #include "AudioMixerDevice.h"
@@ -41,6 +42,7 @@
 #include "Sound/SoundNode.h"
 #include "Sound/SoundNodeWavePlayer.h"
 #include "Sound/SoundSubmix.h"
+#include "Trace/Trace.h"
 #include "UnrealEngine.h"
 #include "UObject/UObjectHash.h"
 #include "UObject/UObjectIterator.h"
@@ -238,6 +240,13 @@ static FAutoConsoleCommandWithWorldArgsAndOutputDevice GSetCurrentSpatialPluginC
 		})
 	);
 
+#if UE_AUDIO_PROFILERTRACE_ENABLED
+UE_TRACE_EVENT_BEGIN(Audio, VirtualLoopStop)
+	UE_TRACE_EVENT_FIELD(uint64, Timestamp)
+	UE_TRACE_EVENT_FIELD(uint32, PlayOrder)
+UE_TRACE_EVENT_END()
+#endif // UE_AUDIO_PROFILERTRACE_ENABLED
+
 namespace Audio
 {
 	ICompressedAudioInfo* CreateSoundAssetDecoder(const FName& InRuntimeFormat)
@@ -5347,6 +5356,15 @@ bool FAudioDevice::RemoveVirtualLoop(FActiveSound& InActiveSound)
 		}
 #endif // !(UE_BUILD_SHIPPING || UE_BUILD_TEST)
 
+#if UE_AUDIO_PROFILERTRACE_ENABLED
+		const bool bChannelEnabled = UE_TRACE_CHANNELEXPR_IS_ENABLED(AudioChannel);
+		if (bChannelEnabled)
+		{
+			UE_TRACE_LOG(Audio, VirtualLoopStop, AudioChannel)
+				<< VirtualLoopStop.Timestamp(FPlatformTime::Cycles64())
+				<< VirtualLoopStop.PlayOrder(InActiveSound.GetPlayOrder());
+		}
+#endif // UE_AUDIO_PROFILERTRACE_ENABLED
 		VirtualLoops.Remove(&InActiveSound);
 		return true;
 	}
diff --git a/Engine/Source/Runtime/Engine/Private/AudioVirtualLoop.cpp b/Engine/Source/Runtime/Engine/Private/AudioVirtualLoop.cpp
--- a/Engine/Source/Runtime/Engine/Private/AudioVirtualLoop.cpp
+++ b/Engine/Source/Runtime/Engine/Private/AudioVirtualLoop.cpp
@@ -4,6 +4,7 @@
 #include "ActiveSound.h"
 #include "Audio/AudioDebug.h"
 #include "AudioDevice.h"
+#include "AudioMixerTrace.h"
 #include "Sound/SoundBase.h"
 
 
@@ -43,6 +44,35 @@ FAutoConsoleVariableRef CVarVirtualLoopsUpdateRateMax(
 	ECVF_Default);
 
 
+#if UE_AUDIO_PROFILERTRACE_ENABLED
+UE_TRACE_EVENT_BEGIN(Audio, VirtualLoopVirtualize)
+	UE_TRACE_EVENT_FIELD(uint64, Timestamp)
+	UE_TRACE_EVENT_FIELD(uint32, PlayOrder)
+	UE_TRACE_EVENT_FIELD(UE::Trace::WideString, Name)
+UE_TRACE_EVENT_END()
+
+UE_TRACE_EVENT_BEGIN(Audio, VirtualLoopRealize)
+	UE_TRACE_EVENT_FIELD(uint64, Timestamp)
+	UE_TRACE_EVENT_FIELD(uint32, PlayOrder)
+UE_TRACE_EVENT_END()
+
+UE_TRACE_EVENT_BEGIN(Audio, VirtualLoopUpdate)
+	UE_TRACE_EVENT_FIELD(double, Timestamp)
+	UE_TRACE_EVENT_FIELD(uint32, PlayOrder)
+	UE_TRACE_EVENT_FIELD(float, TimeVirtualized)
+	UE_TRACE_EVENT_FIELD(float, PlaybackTime)
+	UE_TRACE_EVENT_FIELD(float, UpdateInterval)
+	UE_TRACE_EVENT_FIELD(double, LocationX)
+	UE_TRACE_EVENT_FIELD(double, LocationY)
+	UE_TRACE_EVENT_FIELD(double, LocationZ)
+	UE_TRACE_EVENT_FIELD(double, RotatorPitch)
+	UE_TRACE_EVENT_FIELD(double, RotatorYaw)
+	UE_TRACE_EVENT_FIELD(double, RotatorRoll)
+
+UE_TRACE_EVENT_END()
+#endif // UE_AUDIO_PROFILERTRACE_ENABLED
+
+
 FAudioVirtualLoop::FAudioVirtualLoop()
 	: TimeSinceLastUpdate(0.0f)
 	, TimeVirtualized(0.0f)
@@ -92,6 +122,17 @@ bool FAudioVirtualLoop::Virtualize(const FActiveSound& InActiveSound, FAudioDevi
 	FActiveSound* ActiveSound = FActiveSound::CreateVirtualCopy(InActiveSound, InAudioDevice);
 	OutVirtualLoop.ActiveSound = ActiveSound;
 	OutVirtualLoop.CalculateUpdateInterval();
+
+#if UE_AUDIO_PROFILERTRACE_ENABLED
+	const bool bChannelEnabled = UE_TRACE_CHANNELEXPR_IS_ENABLED(AudioChannel);
+	if (bChannelEnabled)
+	{
+		UE_TRACE_LOG(Audio, VirtualLoopVirtualize, AudioChannel)
+			<< VirtualLoopVirtualize.Timestamp(FPlatformTime::Cycles64())
+			<< VirtualLoopVirtualize.PlayOrder(ActiveSound->GetPlayOrder())
+			<< VirtualLoopVirtualize.Name(Sound ? *Sound->GetPathName() : TEXT("N/A"));
+	}
+#endif // UE_AUDIO_PROFILERTRACE_ENABLED
 	return true;
 }
 
@@ -195,6 +236,8 @@ void FAudioVirtualLoop::UpdateFocusData(float DeltaTime)
 
 bool FAudioVirtualLoop::Update(float DeltaTime, bool bForceUpdate)
 {
+	check(ActiveSound);
+
 	// Keep playback time up-to-date as it may be used to evaluate whether or
 	// not virtual sound is eligible for playback when compared against
 	// actively playing sounds in concurrency checks.
@@ -221,6 +264,30 @@ bool FAudioVirtualLoop::Update(float DeltaTime, bool bForceUpdate)
 	Audio::FAudioDebugger::DrawDebugInfo(*this);
 #endif // ENABLE_AUDIO_DEBUG
 
+#if UE_AUDIO_PROFILERTRACE_ENABLED
+	const bool bChannelEnabled = UE_TRACE_CHANNELEXPR_IS_ENABLED(AudioChannel);
+	if (bChannelEnabled)
+	{
+		const USoundBase* Sound = ActiveSound->GetSound();
+
+		const FTransform& Transform = ActiveSound->Transform;
+		const FVector& Location = Transform.GetLocation();
+		const FRotator Rotator = Transform.GetRotation().Rotator();
+		UE_TRACE_LOG(Audio, VirtualLoopUpdate, AudioChannel)
+			<< VirtualLoopUpdate.Timestamp(FPlatformTime::Cycles64())
+			<< VirtualLoopUpdate.PlayOrder(ActiveSound->GetPlayOrder())
+			<< VirtualLoopUpdate.TimeVirtualized(TimeVirtualized)
+			<< VirtualLoopUpdate.PlaybackTime(ActiveSound->PlaybackTime)
+			<< VirtualLoopUpdate.UpdateInterval(UpdateInterval)
+			<< VirtualLoopUpdate.LocationX(Location.X)
+			<< VirtualLoopUpdate.LocationY(Location.Y)
+			<< VirtualLoopUpdate.LocationZ(Location.Z)
+			<< VirtualLoopUpdate.RotatorPitch(Rotator.Pitch)
+			<< VirtualLoopUpdate.RotatorYaw(Rotator.Yaw)
+			<< VirtualLoopUpdate.RotatorRoll(Rotator.Roll);
+	}
+#endif // UE_AUDIO_PROFILERTRACE_ENABLED
+
 	UpdateFocusData(UpdateDelta);
 
 	// If not audible, update when will be checked again and return false
@@ -230,6 +297,14 @@ bool FAudioVirtualLoop::Update(float DeltaTime, bool bForceUpdate)
 		return false;
 	}
 
+#if UE_AUDIO_PROFILERTRACE_ENABLED
+	if (bChannelEnabled)
+	{
+		UE_TRACE_LOG(Audio, VirtualLoopRealize, AudioChannel)
+			<< VirtualLoopRealize.Timestamp(FPlatformTime::Cycles64())
+			<< VirtualLoopRealize.PlayOrder(ActiveSound->GetPlayOrder());
+	}
+#endif // UE_AUDIO_PROFILERTRACE_ENABLED
 	return true;
 }
 
diff --git a/Engine/Source/Runtime/Engine/Public/Audio/AudioDebug.h b/Engine/Source/Runtime/Engine/Public/Audio/AudioDebug.h
--- a/Engine/Source/Runtime/Engine/Public/Audio/AudioDebug.h
+++ b/Engine/Source/Runtime/Engine/Public/Audio/AudioDebug.h
@@ -81,16 +81,13 @@ namespace Audio
 		static void RemoveDevice(const FAudioDevice& AudioDevice);
 		static void ResolveDesiredStats(FViewportClient* ViewportClient);
 		static void SendUpdateResultsToGameThread(const FAudioDevice& AudioDevice, const int32 FirstActiveIndex);
-		static bool ToggleStatCues(UWorld* World, FCommonViewportClient* ViewportClient, const TCHAR* Stream);
-		static bool ToggleStatMixes(UWorld* World, FCommonViewportClient* ViewportClient, const TCHAR* Stream);
-		static bool ToggleStatModulators(UWorld* World, FCommonViewportClient* ViewportClient, const TCHAR* Stream);
-		static bool ToggleStatSounds(UWorld* World, FCommonViewportClient* ViewportClient, const TCHAR* Stream);
-		static bool ToggleStatWaves(UWorld* World, FCommonViewportClient* ViewportClient, const TCHAR* Stream);
 		static void UpdateAudibleInactiveSounds(const uint32 FirstIndex, const TArray<FWaveInstance*>& WaveInstances);
 		static void LogSubtitle(const TCHAR* InCmd, USoundWave& InSoundWave);
 		static void ClearStats(const FName StatsToToggle, UWorld* InWorld);
 		static void SetStats(const TSet<FName>& StatsToToggle, UWorld* InWorld);
 
+		static bool IsVirtualLoopVisualizeEnabled();
+
 		void ClearMutesAndSolos();
 		void DumpActiveSounds() const;
 
diff --git a/Engine/Source/Runtime/Engine/Public/AudioVirtualLoop.h b/Engine/Source/Runtime/Engine/Public/AudioVirtualLoop.h
--- a/Engine/Source/Runtime/Engine/Public/AudioVirtualLoop.h
+++ b/Engine/Source/Runtime/Engine/Public/AudioVirtualLoop.h
@@ -1,7 +1,6 @@
 // Copyright Epic Games, Inc. All Rights Reserved.
 #pragma once
 
-#include "CoreMinimal.h"
 #include "ActiveSound.h"
 
 
