If UMaterialInterface objects leak past the time at which the shader code library they depend on is closed, their material shader maps stay in the cache even though those shader maps are unusable even if the shader code library is re-opened. This can also affect newly loaded materials which use the same shader map, even if the shader code library has been re-opened.
This change prevents the cache being polluted by shader maps with forcibly released shaders.

#rnx

diff --git a/Engine/Source/Runtime/Engine/Private/Materials/MaterialShader.cpp b/Engine/Source/Runtime/Engine/Private/Materials/MaterialShader.cpp
--- a/Engine/Source/Runtime/Engine/Private/Materials/MaterialShader.cpp
+++ b/Engine/Source/Runtime/Engine/Private/Materials/MaterialShader.cpp
@@ -23,6 +23,7 @@
 #include "Misc/PathViews.h"
 #include "SceneTexturesConfig.h"
 #include "PSOPrecache.h"
+#include "ShaderCodeLibrary.h"
 
 int32 GMaterialExcludeNonPipelinedShaders = 1;
 static FAutoConsoleVariableRef CVarMaterialExcludeNonPipelinedShaders(
@@ -2788,6 +2789,15 @@ uint32 FMaterialShaderMap::GetShaderNum() const
  */
 void FMaterialShaderMap::Register(EShaderPlatform InShaderPlatform)
 {
+	// Lazy initializer to bind OnSharedShaderMapResourceExplicitRelease to ShaderMapResourceExplicitRelease
+	static struct FMaterialShaderMapInnerLazyInitializer
+	{
+		FMaterialShaderMapInnerLazyInitializer()
+		{
+			OnSharedShaderMapResourceExplicitRelease.BindStatic(&FMaterialShaderMap::ShaderMapResourceExplicitRelease);
+		}
+	} MaterialShaderMapInnerLazyInitializer;
+
 	extern int32 GCreateShadersOnLoad;
 	if (GCreateShadersOnLoad && GetShaderPlatform() == InShaderPlatform)
 	{
@@ -3163,6 +3173,21 @@ void FMaterialShaderMap::SaveShaderStableKeys(EShaderPlatform TargetShaderPlatfo
 }
 #endif // WITH_EDITOR
 
+void FMaterialShaderMap::ShaderMapResourceExplicitRelease(const FShaderMapResource* ShaderMapResource)
+{
+	EShaderPlatform ShaderPlatform = ShaderMapResource->GetPlatform();
+
+	// visit cached Material shader map and remove ones that have been released (possibly due to GC leak) to avoid use after free in RT
+	FScopeLock ScopeLock(&GIdToMaterialShaderMapCS);
+	for (auto MaterialMapIterator = GIdToMaterialShaderMap[ShaderPlatform].CreateIterator(); MaterialMapIterator; ++MaterialMapIterator)
+	{
+		if (MaterialMapIterator->Value->GetResource() == ShaderMapResource)
+		{
+			MaterialMapIterator.RemoveCurrent();
+		}
+	}
+}
+
 /**
  * Dump material stats for a given platform.
  * 
diff --git a/Engine/Source/Runtime/Engine/Public/MaterialShared.h b/Engine/Source/Runtime/Engine/Public/MaterialShared.h
--- a/Engine/Source/Runtime/Engine/Public/MaterialShared.h
+++ b/Engine/Source/Runtime/Engine/Public/MaterialShared.h
@@ -1431,6 +1431,9 @@ private:
 	FShader* ProcessCompilationResultsForSingleJob(class FShaderCompileJob* SingleJob, const FShaderPipelineType* ShaderPipeline, const FSHAHash& MaterialShaderMapHash);
 #endif
 
+	/** Iterate over global MaterialShaderMap cache and remove MaterialShaderMap which reference FShaderMapResource argument */
+	static void ShaderMapResourceExplicitRelease(const FShaderMapResource* ShaderMapResource);
+
 	friend ENGINE_API void DumpMaterialStats( EShaderPlatform Platform );
 	friend class FShaderCompilingManager;
 };
diff --git a/Engine/Source/Runtime/RenderCore/Private/ShaderCodeLibrary.cpp b/Engine/Source/Runtime/RenderCore/Private/ShaderCodeLibrary.cpp
--- a/Engine/Source/Runtime/RenderCore/Private/ShaderCodeLibrary.cpp
+++ b/Engine/Source/Runtime/RenderCore/Private/ShaderCodeLibrary.cpp
@@ -228,6 +228,8 @@ namespace UE
 TSet<UE::ShaderLibrary::Private::FMountedPakFileInfo> UE::ShaderLibrary::Private::FMountedPakFileInfo::KnownPakFiles;
 FCriticalSection UE::ShaderLibrary::Private::FMountedPakFileInfo::KnownPakFilesAccessLock;
 
+FSharedShaderMapResourceExplicitRelease OnSharedShaderMapResourceExplicitRelease;
+
 class FShaderMapResource_SharedCode final : public FShaderMapResource
 {
 public:
@@ -1221,6 +1223,15 @@ void FShaderMapResource_SharedCode::ReleaseRHI()
 
 	// on assumption that we aren't going to get resurrected
 	LibraryInstance = nullptr;
+
+	if (GetNumRefs()> 0)
+	{
+		ensureMsgf(false, TEXT("FShaderMapResource_SharedCode::ReleaseRHI is still referenced (Num of references %d). Invoking OnSharedShaderMapResourceExplicitRelease delegate."), GetNumRefs());
+		UE_LOG(LogShaderLibrary, Warning, TEXT("FShaderMapResource_SharedCode::ReleaseRHI is still referenced (Num of references %d). Invoking OnSharedShaderMapResourceExplicitRelease delegate."), GetNumRefs());
+		
+		// Invoke delegate to notify shader map resource needs to be forced released
+		OnSharedShaderMapResourceExplicitRelease.ExecuteIfBound(this);
+	}
 }
 
 bool FShaderMapResource_SharedCode::TryRelease()
diff --git a/Engine/Source/Runtime/RenderCore/Public/ShaderCodeLibrary.h b/Engine/Source/Runtime/RenderCore/Public/ShaderCodeLibrary.h
--- a/Engine/Source/Runtime/RenderCore/Public/ShaderCodeLibrary.h
+++ b/Engine/Source/Runtime/RenderCore/Public/ShaderCodeLibrary.h
@@ -175,6 +175,9 @@ RENDERCORE_API bool LoadFromCompactBinary(FCbFieldView Field, FStableShaderKeyAn
 #endif
 
 
+DECLARE_DELEGATE_OneParam(FSharedShaderMapResourceExplicitRelease, const FShaderMapResource*);
+RENDERCORE_API extern FSharedShaderMapResourceExplicitRelease OnSharedShaderMapResourceExplicitRelease;
+
 DECLARE_MULTICAST_DELEGATE_TwoParams(FSharedShaderCodeRequest, const FSHAHash&, FArchive*);
 DECLARE_MULTICAST_DELEGATE_OneParam(FSharedShaderCodeRelease, const FSHAHash&);
 
