Non unity compile error fix

diff --git a/Engine/Source/Runtime/Engine/Private/InstancedStaticMesh.cpp b/Engine/Source/Runtime/Engine/Private/InstancedStaticMesh.cpp
--- a/Engine/Source/Runtime/Engine/Private/InstancedStaticMesh.cpp
+++ b/Engine/Source/Runtime/Engine/Private/InstancedStaticMesh.cpp
@@ -10,6 +10,7 @@
 #include "AI/NavigationSystemBase.h"
 #include "Engine/MapBuildDataRegistry.h"
 #include "Components/LightComponent.h"
+#include "Components/SkeletalMeshComponent.h"
 #include "Logging/TokenizedMessage.h"
 #include "Logging/MessageLog.h"
 #include "UnrealEngine.h"
