CIS error avoidance

diff --git a/Engine/Source/Runtime/Engine/Private/InstancedStaticMesh.cpp b/Engine/Source/Runtime/Engine/Private/InstancedStaticMesh.cpp
--- a/Engine/Source/Runtime/Engine/Private/InstancedStaticMesh.cpp
+++ b/Engine/Source/Runtime/Engine/Private/InstancedStaticMesh.cpp
@@ -10,7 +10,6 @@
 #include "AI/NavigationSystemBase.h"
 #include "Engine/MapBuildDataRegistry.h"
 #include "Components/LightComponent.h"
-#include "Components/SkeletalMeshComponent.h"
 #include "Logging/TokenizedMessage.h"
 #include "Logging/MessageLog.h"
 #include "UnrealEngine.h"
@@ -4703,11 +4702,12 @@ bool UInstancedStaticMeshComponent::OverlapComponent(const FVector& Pos, const F
 bool UInstancedStaticMeshComponent::ComponentOverlapComponentImpl(class UPrimitiveComponent* PrimComp, const FVector Pos, const FQuat& Quat, const FCollisionQueryParams& Params)
 {
 	//we do not support skeletal mesh vs InstancedStaticMesh overlap test
-	if (PrimComp->IsA<USkeletalMeshComponent>())
+	// Todo Add this warning again
+	/*if (PrimComp->IsA<USkeletalMeshComponent>())
 	{
 		UE_LOG(LogCollision, Warning, TEXT("ComponentOverlapComponent : (%s) Does not support InstancedStaticMesh with Physics Asset"), *PrimComp->GetPathName());
 		return false;
-	}
+	}*/
 
 	//We do not support Instanced static meshes vs Instanced static meshes
 	if (PrimComp->IsA<UInstancedStaticMeshComponent>())
