DirectionalLight, SkyLight, SkyAtmosphere, ExponentialHeightFog available in Restrictive Mode

diff --git a/Engine/Source/Runtime/Engine/Classes/Components/SkyAtmosphereComponent.h b/Engine/Source/Runtime/Engine/Classes/Components/SkyAtmosphereComponent.h
--- a/Engine/Source/Runtime/Engine/Classes/Components/SkyAtmosphereComponent.h
+++ b/Engine/Source/Runtime/Engine/Classes/Components/SkyAtmosphereComponent.h
@@ -142,7 +142,7 @@ class USkyAtmosphereComponent : public USceneComponent
 	FLinearColor SkyLuminanceFactor;
 
 	/** Makes the aerial perspective look thicker by scaling distances from view to surfaces (opaque and translucent). */
-	UPROPERTY(EditAnywhere, BlueprintReadOnly, interp, Category = "Art Direction", meta = (UIMin = 0.0, UIMax = 3.0, ClampMin = 0.0, SliderExponent = 2.0))
+	UPROPERTY(EditAnywhere, BlueprintReadOnly, interp, Category = "Art Direction", meta = (DisplayName = "Aerial Perspective View Distance Scale", UIMin = 0.0, UIMax = 3.0, ClampMin = 0.0, SliderExponent = 2.0))
 	float AerialPespectiveViewDistanceScale;
 
 	/** Scale the sky and atmosphere lights contribution to the height fog when SupportSkyAtmosphereAffectsHeightFog project setting is true.*/
