#ue5 - Added new 'BaseID' to FBasedMovementInfo used for Character ReplicatedBasedMovement. This can help detect when the base component changes but when it hasn't yet replicated to the client (such as is the case in fast shared replication where object references are null).

diff --git a/Engine/Source/Runtime/Engine/Classes/GameFramework/Character.h b/Engine/Source/Runtime/Engine/Classes/GameFramework/Character.h
--- a/Engine/Source/Runtime/Engine/Classes/GameFramework/Character.h
+++ b/Engine/Source/Runtime/Engine/Classes/GameFramework/Character.h
@@ -168,6 +168,10 @@ struct FBasedMovementInfo
 {
 	GENERATED_USTRUCT_BODY()
 
+	/** Unique (within a reasonable timespan) ID of the base component. Can be used to detect changes in the base when the pointer can't replicate, eg during fast shared replication. */
+	UPROPERTY()
+	uint16 BaseID = 0;
+
 	/** Component we are based on */
 	UPROPERTY()
 	TObjectPtr<UPrimitiveComponent> MovementBase = nullptr;
diff --git a/Engine/Source/Runtime/Engine/Private/Character.cpp b/Engine/Source/Runtime/Engine/Private/Character.cpp
--- a/Engine/Source/Runtime/Engine/Private/Character.cpp
+++ b/Engine/Source/Runtime/Engine/Private/Character.cpp
@@ -26,6 +26,7 @@ DEFINE_LOG_CATEGORY_STATIC(LogCharacter, Log, All);
 
 DECLARE_CYCLE_STAT(TEXT("Char OnNetUpdateSimulatedPosition"), STAT_CharacterOnNetUpdateSimulatedPosition, STATGROUP_Character);
 
+
 FName ACharacter::MeshComponentName(TEXT("CharacterMesh0"));
 FName ACharacter::CharacterMovementComponentName(TEXT("CharMoveComp"));
 FName ACharacter::CapsuleComponentName(TEXT("CollisionCylinder"));
@@ -766,6 +767,11 @@ void ACharacter::SetBase( UPrimitiveComponent* NewBaseComponent, const FName InB
 		UPrimitiveComponent* OldBase = BasedMovement.MovementBase;
 		BasedMovement.MovementBase = NewBaseComponent;
 		BasedMovement.BoneName = BoneName;
+		if (bBaseChanged)
+		{
+			// Increment base ID.
+			BasedMovement.BaseID++;
+		}
 
 		if (CharacterMovement)
 		{
@@ -804,7 +810,7 @@ void ACharacter::SetBase( UPrimitiveComponent* NewBaseComponent, const FName InB
 			const ENetRole LocalRole = GetLocalRole();
 			if (LocalRole == ROLE_Authority || LocalRole == ROLE_AutonomousProxy)
 			{
-				BasedMovement.bServerHasBaseComponent = (BasedMovement.MovementBase != nullptr); // Also set on proxies for nicer debugging.
+				BasedMovement.bServerHasBaseComponent = (BasedMovement.MovementBase != nullptr); // Also set on autonomous proxies for nicer debugging.
 				UE_LOG(LogCharacter, Verbose, TEXT("Setting base on %s for '%s' to '%s'"), LocalRole == ROLE_Authority ? TEXT("Server") : TEXT("AutoProxy"), *GetName(), *GetFullNameSafe(NewBaseComponent));
 			}
 			else
