Expose 'late' level streaming distance as a cvar to share between reporting methods.
Add new LevelStreaming CSV event category.

diff --git a/Engine/Source/Runtime/Engine/Private/ProfilingDebugging/LevelStreamingProfilingSubsystem.cpp b/Engine/Source/Runtime/Engine/Private/ProfilingDebugging/LevelStreamingProfilingSubsystem.cpp
--- a/Engine/Source/Runtime/Engine/Private/ProfilingDebugging/LevelStreamingProfilingSubsystem.cpp
+++ b/Engine/Source/Runtime/Engine/Private/ProfilingDebugging/LevelStreamingProfilingSubsystem.cpp
@@ -18,6 +18,15 @@
 #include UE_INLINE_GENERATED_CPP_BY_NAME(LevelStreamingProfilingSubsystem)
 
 DEFINE_LOG_CATEGORY_STATIC(LogLevelStreamingProfiling, Log, All);
+CSV_DEFINE_CATEGORY_MODULE(ENGINE_API, LevelStreaming, true);
+
+static float GLateStreamingDistanceSquared = 0.0f;
+static FAutoConsoleVariableRef CVarLateStreamingDistanceSquared(
+	TEXT("LevelStreaming.Profiling.LateStreamingDistanceSquared"),
+	GLateStreamingDistanceSquared,
+	TEXT("The squared distance (e.g. from world partition cell bounds) below which a level is considered to have streamed in late."),
+	ECVF_Default
+);
 
 ULevelStreamingProfilingSubsystem::ULevelStreamingProfilingSubsystem(const FObjectInitializer& ObjectInitializer)
 {
@@ -27,6 +36,11 @@ ULevelStreamingProfilingSubsystem::~ULevelStreamingProfilingSubsystem()
 {
 }
 
+double ULevelStreamingProfilingSubsystem::GetLateStreamingDistanceSquared()
+{
+	return GLateStreamingDistanceSquared;
+}
+
 bool ULevelStreamingProfilingSubsystem::DoesSupportWorldType(const EWorldType::Type WorldType) const
 {
 	return WorldType == EWorldType::PIE || WorldType == EWorldType::Game;
diff --git a/Engine/Source/Runtime/Engine/Public/ProfilingDebugging/LevelStreamingProfilingSubsystem.h b/Engine/Source/Runtime/Engine/Public/ProfilingDebugging/LevelStreamingProfilingSubsystem.h
--- a/Engine/Source/Runtime/Engine/Public/ProfilingDebugging/LevelStreamingProfilingSubsystem.h
+++ b/Engine/Source/Runtime/Engine/Public/ProfilingDebugging/LevelStreamingProfilingSubsystem.h
@@ -19,6 +19,8 @@ class ULevelStreaming;
 enum class ELevelStreamingState : uint8;
 enum class ELevelStreamingTargetState : uint8;
 
+CSV_DECLARE_CATEGORY_MODULE_EXTERN(ENGINE_API, LevelStreaming);
+
 /** 
  * This subsystem captures level streaming operations for a specified time and outputs a .tsv (tab separated values) file to the 
  * profiling directory containing the amount of time spent loading levels and the time levels spent queuing for theability to load.
@@ -50,6 +52,9 @@ public:
 	ENGINE_API virtual void PostReport() { }
 #endif
 
+	// Access to tuning values set by cvars for other systems
+	/* Returns the squared distance (e.g. from world partition cell bounds) at which a level is considered to have streamed in too late. */
+	static ENGINE_API double GetLateStreamingDistanceSquared();
 
 protected:
 	// WorldSubsystem interface
