Fix non-unity build

diff --git a/Engine/Source/Runtime/Engine/Public/ProfilingDebugging/LevelStreamingProfilingSubsystem.h b/Engine/Source/Runtime/Engine/Public/ProfilingDebugging/LevelStreamingProfilingSubsystem.h
--- a/Engine/Source/Runtime/Engine/Public/ProfilingDebugging/LevelStreamingProfilingSubsystem.h
+++ b/Engine/Source/Runtime/Engine/Public/ProfilingDebugging/LevelStreamingProfilingSubsystem.h
@@ -5,6 +5,7 @@
 #include "CoreTypes.h"
 #include "Containers/ArrayView.h"
 #include "Containers/Map.h"
+#include "ProfilingDebugging/CsvProfiler.h"
 #include "Subsystems/WorldSubsystem.h"
 #include "Tasks/Task.h"
 #include "UObject/ObjectKey.h"
