Improve deferred construction script processing
  - Add missing asset processing when world begin play is called to avoid construction scripts from being deferred needlessly
  - Level construction scripts are now only concerned by assets still being compiled in their own level
  - Save and continue processing actors where the last processing was aborted due to timeout
  - Only do work when assets we're interested in have actually finished compiling
  - This should resolve issues with construction scripts being called after BeginPlay for some actors

#rnx

diff --git a/Engine/Source/Runtime/Engine/Classes/Engine/Level.h b/Engine/Source/Runtime/Engine/Classes/Engine/Level.h
--- a/Engine/Source/Runtime/Engine/Classes/Engine/Level.h
+++ b/Engine/Source/Runtime/Engine/Classes/Engine/Level.h
@@ -1310,6 +1310,11 @@ public:
 	 * Call on a level that was loaded from disk instead of PIE-duplicating, to fixup actor references
 	 */
 	ENGINE_API void FixupForPIE(int32 PIEInstanceID, TFunctionRef<void(int32, FSoftObjectPath&)> CustomFixupFunction = [](int32, FSoftObjectPath&) {});
+
+	/**
+	 * Returns true if the level contains static meshes that have not finished compiling yet.
+	 */
+	ENGINE_API bool HasStaticMeshCompilationPending();
 #endif
 
 	/** @todo document */
@@ -1361,7 +1366,7 @@ private:
 	bool IncrementalRegisterComponents(bool bPreRegisterComponents, int32 NumComponentsToUpdate, FRegisterComponentContext* Context);
 #if WITH_EDITOR
 	bool IncrementalRunConstructionScripts(bool bProcessAllActors);
-
+	TOptional<bool> bCachedHasStaticMeshCompilationPending;
 private:
 	/**
 	 * Potentially defer the running of an actor's construction script on load
diff --git a/Engine/Source/Runtime/Engine/Private/ActorDeferredScriptManager.cpp b/Engine/Source/Runtime/Engine/Private/ActorDeferredScriptManager.cpp
--- a/Engine/Source/Runtime/Engine/Private/ActorDeferredScriptManager.cpp
+++ b/Engine/Source/Runtime/Engine/Private/ActorDeferredScriptManager.cpp
@@ -7,8 +7,14 @@
 
 #if WITH_EDITOR
 
+#include "Engine/Level.h"
+#include "Engine/StaticMesh.h"
+#include "Logging/LogMacros.h"
+#include "ProfilingDebugging/CountersTrace.h"
+
 #define LOCTEXT_NAMESPACE "ActorDeferredScriptManager"
 
+DEFINE_LOG_CATEGORY_STATIC(LogActorDeferredScriptManager, Log, All);
 
 FActorDeferredScriptManager& FActorDeferredScriptManager::Get()
 {
@@ -18,6 +24,11 @@ FActorDeferredScriptManager& FActorDeferredScriptManager::Get()
 
 void FActorDeferredScriptManager::Shutdown()
 {
+	if (OnAssetChangeDelegateHandle.IsValid())
+	{
+		FAssetCompilingManager::Get().OnAssetPostCompileEvent().Remove(OnAssetChangeDelegateHandle);
+		OnAssetChangeDelegateHandle.Reset();
+	}
 }
 
 FActorDeferredScriptManager::FActorDeferredScriptManager()
@@ -32,7 +43,14 @@ FName FActorDeferredScriptManager::GetStaticAssetTypeName()
 
 void FActorDeferredScriptManager::AddActor(AActor* InActor)
 {
-	PendingConstructionScriptActors.PushFirst(InActor);
+	UE_LOG(LogActorDeferredScriptManager, VeryVerbose, TEXT("Adding actor %s for deferred construction script"), *InActor->GetPathName());
+
+	PendingConstructionScriptActors.Add(InActor);
+
+	// Whenever an actor is added, reset the count so we go over all actors on next ProcessAsyncTasks
+	NumLeftToProcess = PendingConstructionScriptActors.Num();
+
+	UpdateCompilationNotification();
 }
 
 FName FActorDeferredScriptManager::GetAssetTypeName() const
@@ -60,40 +78,153 @@ int32 FActorDeferredScriptManager::GetNumRemainingAssets() const
 	return PendingConstructionScriptActors.Num();
 }
 
+void FActorDeferredScriptManager::OnAssetPostCompile(const TArray<FAssetCompileData>& CompiledAssets)
+{
+	if (!PendingConstructionScriptActors.IsEmpty())
+	{
+		for (const FAssetCompileData& CompileData : CompiledAssets)
+		{
+			if (UObject* Object = CompileData.Asset.Get())
+			{
+				if (UStaticMesh* StaticMesh = Cast<UStaticMesh>(Object))
+				{
+					// Whenever a static mesh finish compiling, process everything again
+					NumLeftToProcess = PendingConstructionScriptActors.Num();
+					break;
+				}
+			}
+		}
+	}
+}
+
+void FActorDeferredScriptManager::EnsureEventRegistered()
+{
+	if (!OnAssetChangeDelegateHandle.IsValid())
+	{
+		OnAssetChangeDelegateHandle = FAssetCompilingManager::Get().OnAssetPostCompileEvent().AddRaw(this, &FActorDeferredScriptManager::OnAssetPostCompile);
+	}
+}
+
 void FActorDeferredScriptManager::ProcessAsyncTasks(bool bLimitExecutionTime)
 {
-	if (!PendingConstructionScriptActors.IsEmpty() && FStaticMeshCompilingManager::Get().GetNumRemainingMeshes() == 0)
+	if (PendingConstructionScriptActors.IsEmpty())
 	{
+		return;
+	}
+
+	EnsureEventRegistered();
+
+	const bool bIsStaticMeshesCompiling = FStaticMeshCompilingManager::Get().GetNumRemainingMeshes() > 0;
+
+	// If this happens, it means OnAssetPostCompile didn't do it's job... better safe than sorry
+	if (NumLeftToProcess == 0 && !bIsStaticMeshesCompiling)
+	{
+		NumLeftToProcess = PendingConstructionScriptActors.Num();
+	}
+
+	if (NumLeftToProcess > 0)
+	{
+		TRACE_CPUPROFILER_EVENT_SCOPE(FActorDeferredScriptManager::ProcessAsyncTasks);
+
+		// Cache outers result to reduce cost
+		TMap<ULevel*, bool> LevelToCompilationLeft;
+		auto DoesLevelHaveCompilationLeft =
+			[&LevelToCompilationLeft](ULevel* Level)
+			{
+				TRACE_CPUPROFILER_EVENT_SCOPE(DoesLevelHaveCompilationLeft);
+
+				if (bool* bLevelHasCompilationLeft = LevelToCompilationLeft.Find(Level))
+				{
+					return *bLevelHasCompilationLeft;
+				}
+
+				bool bLevelHasCompilationLeft = Level->HasStaticMeshCompilationPending();
+				LevelToCompilationLeft.Add(Level, bLevelHasCompilationLeft);
+				return bLevelHasCompilationLeft;
+			};
+
 		double TickStartTime = FPlatformTime::Seconds();
 		const double MaxSecondsPerFrame = 0.016;
-		bool bHasTimeLeft = true;
 
-		// since this deferred run of construction script was supposed to be done during level load
+		// Since this deferred run of construction script was supposed to be done during level load
 		// temporarily set the global flag to prevent dirtying the level package.
 		TGuardValue<bool> IsEditorLoadingPackageGuard(GIsEditorLoadingPackage, true);
-		while (!PendingConstructionScriptActors.IsEmpty() && bHasTimeLeft)
+
+		int32 CurrentIndex = NextIndexToProcess;
+
+		UE_LOG(LogActorDeferredScriptManager, VeryVerbose, TEXT("Starting to process at index %d"), CurrentIndex);
+
+		NextIndexToProcess = 0;
+		NumLeftToProcess = FMath::Min(NumLeftToProcess, PendingConstructionScriptActors.Num());
+
+		// Always process all items if time permits
+		for (; NumLeftToProcess > 0; --NumLeftToProcess, ++CurrentIndex)
 		{
-			TWeakObjectPtr<AActor> WeakActor = PendingConstructionScriptActors.Last();
-			PendingConstructionScriptActors.PopLast();
-			if (AActor* Actor = WeakActor.Get(); Actor && Actor->GetWorld())
+			// Act as a circular buffer so we process all items if we start at an offset.
+			const int32 CircularIndex = CurrentIndex % PendingConstructionScriptActors.Num();
+
+			const bool bHasTimeLeft = bLimitExecutionTime ? ((FPlatformTime::Seconds() - TickStartTime) < MaxSecondsPerFrame) : true;
+			if (!bHasTimeLeft)
+			{
+				// We're not finished, save where we're at for next time
+				UE_LOG(LogActorDeferredScriptManager, VeryVerbose, TEXT("No time left, will continue at index %d"), CircularIndex);
+				NextIndexToProcess = CircularIndex;
+				break;
+			}
+
+			if (AActor* Actor = PendingConstructionScriptActors[CircularIndex].Get(); Actor && Actor->GetWorld())
 			{
+				UE_LOG(LogActorDeferredScriptManager, VeryVerbose, TEXT("Processing index %d (%s)"), CircularIndex, *Actor->GetPathName());
+
+				// Don't do this costly step if we know there is no more meshes compiling
+				if (bIsStaticMeshesCompiling)
+				{
+					if (ULevel* Level = Actor->GetLevel())
+					{
+						if (DoesLevelHaveCompilationLeft(Level))
+						{
+							// Do not rerun construction scripts until all the meshes in the level have finished compiling
+							UE_LOG(LogActorDeferredScriptManager, VeryVerbose, TEXT("Skipping index %d (%s) because it's level outer still has static meshes compilation pending"), CircularIndex, *Actor->GetPathName());
+							continue;
+
+						}
+					}
+				}
+
+				UE_LOG(LogActorDeferredScriptManager, VeryVerbose, TEXT("Rerunning construction script for index %d (%s)"), CircularIndex, *Actor->GetPathName());
+
 				// Temporarily do not consider actor as initialized if they were when running deferred construction scripts
 				FGuardValue_Bitfield(Actor->bActorInitialized, false);
 				Actor->RerunConstructionScripts();
 			}
-			bHasTimeLeft = bLimitExecutionTime ? ((FPlatformTime::Seconds() - TickStartTime) < MaxSecondsPerFrame) : true;
+			else
+			{
+				UE_LOG(LogActorDeferredScriptManager, VeryVerbose, TEXT("Removing index %d because its actor is not valid anymore"), CircularIndex);
+			}
+
+			// Reduce reallocations by only shrinking when removing last item
+			const bool bAllowShrinking = PendingConstructionScriptActors.Num() == 1;
+			UE_LOG(LogActorDeferredScriptManager, VeryVerbose, TEXT("Removing index %d"), CircularIndex);
+			PendingConstructionScriptActors.RemoveAtSwap(CircularIndex, 1, bAllowShrinking);
+			CurrentIndex--;
+
+			if (PendingConstructionScriptActors.IsEmpty())
+			{
+				UE_LOG(LogActorDeferredScriptManager, VeryVerbose, TEXT("Finished processing construction scripts"));
+			}
 		}
 	}
 
 	UpdateCompilationNotification();
 }
 
+TRACE_DECLARE_INT_COUNTER(QueuedConstructionScripts, TEXT("AsyncCompilation/QueuedConstructionScripts"));
 void FActorDeferredScriptManager::UpdateCompilationNotification()
 {
+	TRACE_COUNTER_SET(QueuedConstructionScripts, GetNumRemainingAssets());
 	Notification.Update(GetNumRemainingAssets());
 }
 
-
 #undef LOCTEXT_NAMESPACE
 
 #endif // #if WITH_EDITOR
diff --git a/Engine/Source/Runtime/Engine/Private/Level.cpp b/Engine/Source/Runtime/Engine/Private/Level.cpp
--- a/Engine/Source/Runtime/Engine/Private/Level.cpp
+++ b/Engine/Source/Runtime/Engine/Private/Level.cpp
@@ -1783,6 +1783,7 @@ bool ULevel::IncrementalRunConstructionScripts(bool bProcessAllActors)
 	if (CurrentActorIndexForIncrementalUpdate >= Actors.Num())
 	{
 		CurrentActorIndexForIncrementalUpdate = 0;
+		bCachedHasStaticMeshCompilationPending.Reset();
 		return true;
 	}
 	return false;
@@ -1845,6 +1846,31 @@ void ULevel::MarkLevelComponentsRenderStateDirty()
 
 #if WITH_EDITOR
 
+bool ULevel::HasStaticMeshCompilationPending()
+{
+	TRACE_CPUPROFILER_EVENT_SCOPE(ULevel::HasStaticMeshCompilationPending);
+
+	bool bHasStaticMeshCompilationPending = false;
+	ForEachObjectWithOuterBreakable(
+		this,
+		[&bHasStaticMeshCompilationPending](UObject* InObject)
+		{
+			if (UStaticMeshComponent* Component = Cast<UStaticMeshComponent>(InObject))
+			{
+				if (Component->IsCompiling())
+				{
+					bHasStaticMeshCompilationPending = false;
+					return false;
+				}
+			}
+
+			return true;
+		}
+	);
+
+	return bHasStaticMeshCompilationPending;
+}
+
 bool ULevel::DeferRunningConstructionScripts(AActor* InActor)
 {
 	// if there are outstanding asset (static meshes) compilation when loading actors
@@ -1854,8 +1880,16 @@ bool ULevel::DeferRunningConstructionScripts(AActor* InActor)
 	if (FStaticMeshCompilingManager::Get().GetNumRemainingMeshes() &&
 		InActor->HasNonTrivialUserConstructionScript())
 	{
-		FActorDeferredScriptManager::Get().AddActor(InActor);
-		return true;
+		if (!bCachedHasStaticMeshCompilationPending.IsSet())
+		{
+			bCachedHasStaticMeshCompilationPending = HasStaticMeshCompilationPending();
+		}
+
+		if (bCachedHasStaticMeshCompilationPending)
+		{
+			FActorDeferredScriptManager::Get().AddActor(InActor);
+			return true;
+		}
 	}
 	return false;
 }
diff --git a/Engine/Source/Runtime/Engine/Private/World.cpp b/Engine/Source/Runtime/Engine/Private/World.cpp
--- a/Engine/Source/Runtime/Engine/Private/World.cpp
+++ b/Engine/Source/Runtime/Engine/Private/World.cpp
@@ -3033,6 +3033,12 @@ void UWorld::AddToWorld( ULevel* Level, const FTransform& LevelTransform, bool b
 		bExecuteNextStep					= Level->bAreComponentsCurrentlyRegistered && (!bConsiderTimeLimit || !IsTimeLimitExceeded(TEXT("updating components"), StartTime, Level, TimeLimit));
 	}
 
+#if WITH_EDITOR
+	// Gives a chance to any assets being used for PIE/game to complete before calling
+	// BeginPlay on all actors
+	FAssetCompilingManager::Get().ProcessAsyncTasks();
+#endif
+
 	if( IsGameWorld() && AreActorsInitialized() )
 	{
 		// Initialize all actors and start execution.
@@ -4911,6 +4917,11 @@ void UWorld::BeginPlay()
 		ServerStreamingLevelsVisibility = AServerStreamingLevelsVisibility::SpawnServerActor(this);
 	}
 
+#if WITH_EDITOR
+	// Gives a chance to any assets being used for PIE/game to complete
+	FAssetCompilingManager::Get().ProcessAsyncTasks();
+#endif
+
 	for (UWorldSubsystem* WorldSubsystem : WorldSubsystems)
 	{
 		WorldSubsystem->OnWorldBeginPlay(*this);
diff --git a/Engine/Source/Runtime/Engine/Public/ActorDeferredScriptManager.h b/Engine/Source/Runtime/Engine/Public/ActorDeferredScriptManager.h
--- a/Engine/Source/Runtime/Engine/Public/ActorDeferredScriptManager.h
+++ b/Engine/Source/Runtime/Engine/Public/ActorDeferredScriptManager.h
@@ -28,11 +28,10 @@ public:
 	 */
 	ENGINE_API void FinishAllCompilation() override {}
 
-
 	/** 
 	 * Add an actor for the manager to run the construction of once dependent asset compilation are done (i.e Static Mesh)
 	 */
-	void AddActor(AActor* InActor);
+	ENGINE_API void AddActor(AActor* InActor);
 
 private:
 	friend class FAssetCompilingManager;
@@ -47,11 +46,24 @@ private:
 
 	void UpdateCompilationNotification();
 
+	void OnAssetPostCompile(const TArray<FAssetCompileData>& CompiledAssets);
+
+	void EnsureEventRegistered();
+
 	/** Actors awaiting outstanding asset compilation prior to running their non trivial construction scripts after loading this level. */
-	TDeque<TWeakObjectPtr<AActor>> PendingConstructionScriptActors;
+	TArray<TWeakObjectPtr<AActor>> PendingConstructionScriptActors;
 
 	/** Notification for the amount of pending construction scripts to run */
 	FAsyncCompilationNotification Notification;
+
+	/** We need to monitor when static meshes have finished compiling to only work when necessary */
+	FDelegateHandle OnAssetChangeDelegateHandle;
+
+	/** Where to start on the next ProcessAsyncTasks iteration. */
+	int32 NextIndexToProcess = 0;
+	
+	/** What's left to process if we were aborted because of execution time limit. */
+	int32 NumLeftToProcess = 0;
 };
 
 #endif
\ No newline at end of file
