Add GetPrevInstanceToWorld() variation of GetInstanceToWorld().
Needed for correct implementation of GetMaterialPreviousWorldPositionOffset() if an InstanceToWorld transform is used for WPO.

diff --git a/Engine/Shaders/Private/MaterialTemplate.ush b/Engine/Shaders/Private/MaterialTemplate.ush
--- a/Engine/Shaders/Private/MaterialTemplate.ush
+++ b/Engine/Shaders/Private/MaterialTemplate.ush
@@ -679,6 +679,15 @@ FLWCMatrix GetInstanceToWorld(FMaterialVertexParameters Parameters)
 	#endif
 }
 
+FLWCMatrix GetPrevInstanceToWorld(FMaterialVertexParameters Parameters)
+{
+#if USE_INSTANCING || USE_INSTANCE_CULLING || IS_MESHPARTICLE_FACTORY || IS_NANITE_PASS
+	return Parameters.PrevFrameLocalToWorld;
+#else
+	return GetPrimitiveData(Parameters).PreviousLocalToWorld;
+#endif
+}
+
 FLWCMatrix GetInstanceToWorld(FMaterialPixelParameters Parameters)
 {
 	#if HAS_INSTANCE_LOCAL_TO_WORLD_PS
diff --git a/Engine/Source/Runtime/Engine/Private/Materials/HLSLMaterialTranslator.cpp b/Engine/Source/Runtime/Engine/Private/Materials/HLSLMaterialTranslator.cpp
--- a/Engine/Source/Runtime/Engine/Private/Materials/HLSLMaterialTranslator.cpp
+++ b/Engine/Source/Runtime/Engine/Private/Materials/HLSLMaterialTranslator.cpp
@@ -8641,8 +8641,8 @@ int32 FHLSLMaterialTranslator::TransformBase(EMaterialCommonBasis SourceCoordBas
 		{
 			if (DestCoordBasis == MCB_World)
 			{
-				CodeStr = LWCMultiplyMatrix(TEXT("<A>"), TEXT("GetInstanceToWorld(Parameters)"), AWComponent);
-				CodeDerivStr = LWCMultiplyMatrix(TEXT("<A>"), TEXT("GetInstanceToWorld(Parameters)"), 0);
+				CodeStr = LWCMultiplyMatrix(TEXT("<A>"), TEXT("Get<PREV>InstanceToWorld(Parameters)"), AWComponent);
+				CodeDerivStr = LWCMultiplyMatrix(TEXT("<A>"), TEXT("Get<PREV>InstanceToWorld(Parameters)"), 0);
 				bUsesInstanceLocalToWorldPS |= ShaderFrequency == SF_Pixel;
 			}
 			// use World as an intermediary base
