Backing out 23267598

diff --git a/Engine/Source/Runtime/Engine/Private/LevelActor.cpp b/Engine/Source/Runtime/Engine/Private/LevelActor.cpp
--- a/Engine/Source/Runtime/Engine/Private/LevelActor.cpp
+++ b/Engine/Source/Runtime/Engine/Private/LevelActor.cpp
@@ -996,6 +996,8 @@ bool UWorld::DestroyActor( AActor* ThisActor, bool bNetForce, bool bShouldModify
 #endif
 	}
 
+	OnActorRemovedFromWorld.Broadcast(ThisActor);
+
 	// Clean up the actor's components.
 	ThisActor->UnregisterAllComponents();
 
