Ensure we have enough UV channels reserved when converting a skeletal mesh to FMeshDescription.

#rnx

diff --git a/Engine/Source/Runtime/Engine/Private/SkeletalMeshLODModel.cpp b/Engine/Source/Runtime/Engine/Private/SkeletalMeshLODModel.cpp
--- a/Engine/Source/Runtime/Engine/Private/SkeletalMeshLODModel.cpp
+++ b/Engine/Source/Runtime/Engine/Private/SkeletalMeshLODModel.cpp
@@ -1409,6 +1409,9 @@ void FSkeletalMeshLODModel::GetMeshDescription(FMeshDescription& MeshDescription
 		VertexIDs.Add(MeshDescription.CreateVertex());
 	}
 
+	// Ensure we have enough channels to store all the defined UV coordinates.
+	VertexInstanceUVs.SetNumChannels(static_cast<int32>(NumTexCoords));
+	
 	const TArray<FSkeletalMaterial>& Materials = Owner->GetMaterials();
 	const bool bHasVertexColors = EnumHasAllFlags(Owner->GetVertexBufferFlags(), ESkeletalMeshVertexFlags::HasVertexColors);
 
@@ -1474,7 +1477,7 @@ void FSkeletalMeshLODModel::GetMeshDescription(FMeshDescription& MeshDescription
 					(FVector3f)(SourceVertex.TangentZ.GetSafeNormal()),
 					FVector3f::ZeroVector).Determinant() < 0.0f ? -1.0f : +1.0f);
 
-				for (int32 UVIndex = 0; UVIndex < int32(NumTexCoords); UVIndex++)
+				for (int32 UVIndex = 0; UVIndex < static_cast<int32>(NumTexCoords); UVIndex++)
 				{
 					VertexInstanceUVs.Set(VertexInstanceID, UVIndex, SourceVertex.UVs[UVIndex]);
 				}
