Add tool tips for parameter blending option.

diff --git a/Engine/Source/Runtime/Engine/Classes/Materials/MaterialExpressionStrata.h b/Engine/Source/Runtime/Engine/Classes/Materials/MaterialExpressionStrata.h
--- a/Engine/Source/Runtime/Engine/Classes/Materials/MaterialExpressionStrata.h
+++ b/Engine/Source/Runtime/Engine/Classes/Materials/MaterialExpressionStrata.h
@@ -787,6 +787,9 @@ class UMaterialExpressionStrataHorizontalMixing : public UMaterialExpressionStra
 	UPROPERTY()
 	FExpressionInput Mix;
 
+	/**
+	 * Merge Background and Foreground into a single material by mixing their inputs rather than their evaluation. This makes lighting evaluation cheaper (Default: off)
+	 */
 	UPROPERTY(EditAnywhere, Category = Mode)
 	uint32 bUseParameterBlending : 1;
 
@@ -820,6 +823,9 @@ class UMaterialExpressionStrataVerticalLayering : public UMaterialExpressionStra
 	UPROPERTY()
 	FExpressionInput Base;
 
+	/**
+	 * Merge Top and Base into a single material by mixing their inputs rather than their evaluation. This makes lighting evaluation cheaper (Default: off)
+	 */
 	UPROPERTY(EditAnywhere, Category = Mode)
 	uint32 bUseParameterBlending : 1;
 
@@ -853,6 +859,9 @@ class UMaterialExpressionStrataAdd : public UMaterialExpressionStrataBSDF
 	UPROPERTY()
 	FExpressionInput B;
 
+	/**
+	 * Merge A and B into a single material by mixing their inputs rather than their evaluation. This makes lighting evaluation cheaper (Default: off)
+	 */
 	UPROPERTY(EditAnywhere, Category = Mode)
 	uint32 bUseParameterBlending : 1;
 
