Fix for slow performance when applying a new material in a large scene.  The function UWorld::PropagateLightingScenarioChange was calling PropagateLightingScenarioChange on each component, and several types of components were in turn individually creating an FComponentRecreateRenderStateContext, which generates an UpdateAllPrimitiveSceneInfos command for the entire scene.  So it's more or less an O(N^2) operation on the number of primitives in the scene.

Added a feature to batch all the render state recreate operations together using the FGlobalComponentRecreateRenderStateContext class.  Because this is for a hot fix build, the public virtual API for PropagateLightingScenarioChange can't change, so a utility function was added to the FGlobalComponentRecreateRenderStateContext class, to accumulate batched operations in a global array, to be flushed at the end of the operations.  For some components, certain work needs to occur in the scope of the FGlobalComponentRecreateRenderStateContext (between when the render state is destroyed and re-created), so a callback can be provided to do the in-between re-create work.

#rnx

diff --git a/Engine/Source/Runtime/Engine/Private/MapBuildData.cpp b/Engine/Source/Runtime/Engine/Private/MapBuildData.cpp
--- a/Engine/Source/Runtime/Engine/Private/MapBuildData.cpp
+++ b/Engine/Source/Runtime/Engine/Private/MapBuildData.cpp
@@ -93,11 +93,22 @@ void UWorld::PropagateLightingScenarioChange()
 		}
 	}
 
+	TArray<UActorComponent*> WorldComponents;
 	for (USceneComponent* Component : TObjectRange<USceneComponent>(RF_ClassDefaultObject | RF_ArchetypeObject, true, EInternalObjectFlags::Garbage))
 	{
 		if (Component->GetWorld() == this)
 		{
-			Component->PropagateLightingScenarioChange();
+			WorldComponents.Emplace(Component);
+		}
+	}
+
+	{
+		// Use a global context so UpdateAllPrimitiveSceneInfos only runs once, rather than for each component.  Can save minutes of time.
+		FGlobalComponentRecreateRenderStateContext Context(WorldComponents);
+		
+		for (UActorComponent* Component : WorldComponents)
+		{
+			((USceneComponent*)Component)->PropagateLightingScenarioChange();
 		}
 	}
 
