Ensure that the asset compiling manager thread pool is initialized on the main thread

This is currently guaranteed anyway because both the manager and its thread pool are created by GDistanceFieldAsyncQueue. This change allows the manager to provide that guarantee on its own, which is necessary to avoid a race condition when using new-style derived data builds.

#rnx

diff --git a/Engine/Source/Runtime/Engine/Private/AssetCompilingManager.cpp b/Engine/Source/Runtime/Engine/Private/AssetCompilingManager.cpp
--- a/Engine/Source/Runtime/Engine/Private/AssetCompilingManager.cpp
+++ b/Engine/Source/Runtime/Engine/Private/AssetCompilingManager.cpp
@@ -102,8 +102,6 @@ namespace AssetCompilingManagerImpl
 	{
 		void Reserve(uint64 Memory, IRequestOwner& Owner, TUniqueFunction<void ()>&& OnComplete) final
 		{
-			// TODO: This is not thread-safe right now because the thread pool is created on first use.
-			//       That first use is expected to occur on the main thread, which we do not guarantee.
 			LaunchTaskInThreadPool(Memory, Owner, FAssetCompilingManager::Get().GetThreadPool(), MoveTemp(OnComplete));
 		}
 	};
@@ -296,6 +294,11 @@ FAssetCompilingManager::FAssetCompilingManager()
 	RegisterManager(&FActorDeferredScriptManager::Get());
 	RegisterManager(&FSoundWaveCompilingManager::Get());
 
+	// Ensure that the thread pool is constructed before the memory queue is registered because
+	// the memory queue can call GetThreadPool() from a worker thread, and could lead to a race
+	// to create the thread pool.
+	GetThreadPool();
+
 	IModularFeatures::Get().RegisterModularFeature(UE::DerivedData::IBuildSchedulerMemoryQueue::FeatureName, &AssetCompilingManagerImpl::GMemoryQueue);
 #endif
 }
