Fix crash when linking/unlinking layers in an anim notify

diff --git a/Engine/Source/Runtime/Engine/Private/Components/SkeletalMeshComponent.cpp b/Engine/Source/Runtime/Engine/Private/Components/SkeletalMeshComponent.cpp
--- a/Engine/Source/Runtime/Engine/Private/Components/SkeletalMeshComponent.cpp
+++ b/Engine/Source/Runtime/Engine/Private/Components/SkeletalMeshComponent.cpp
@@ -1594,7 +1594,8 @@ void USkeletalMeshComponent::ConditionallyDispatchQueuedAnimEvents()
 	{
 		bNeedsQueuedAnimEventsDispatched = false;
 
-		for (UAnimInstance* LinkedInstance : LinkedInstances)
+		// Copy linked instances here, as anim notifies could potentially modify this array
+		for (UAnimInstance* LinkedInstance : TArray<TObjectPtr<UAnimInstance>, TInlineAllocator<8>>(LinkedInstances))
 		{
 			LinkedInstance->DispatchQueuedAnimEvents();
 		}
