Fix bad access to empty ISMC->PerInstanceSMCustomData array that occured when ISMC->NumCustomDataFloats is 0

diff --git a/Engine/Source/Runtime/Engine/Private/ISMPartition/ISMComponentBatcher.cpp b/Engine/Source/Runtime/Engine/Private/ISMPartition/ISMComponentBatcher.cpp
--- a/Engine/Source/Runtime/Engine/Private/ISMPartition/ISMComponentBatcher.cpp
+++ b/Engine/Source/Runtime/Engine/Private/ISMPartition/ISMComponentBatcher.cpp
@@ -74,10 +74,17 @@ void FISMComponentBatcher::AddInternal(const UActorComponent* InComponent, TOpti
 			}
 			InstancesTransformsWS.Add(InstanceTransformWS);
 
-			// Add per instance custom data
-			TConstArrayView<float> InstanceCustomData(&ISMC->PerInstanceSMCustomData[InstanceIdx * ISMC->NumCustomDataFloats], ISMC->NumCustomDataFloats);
-			InstancesCustomData.Append(InstanceCustomData);
-			InstancesCustomData.AddDefaulted(NumCustomDataFloats - InstanceCustomData.Num());
+			// Add per instance custom data, if any
+			if (NumCustomDataFloats > 0)
+			{
+				if (ISMC->NumCustomDataFloats > 0)
+				{
+					TConstArrayView<float> InstanceCustomData(&ISMC->PerInstanceSMCustomData[InstanceIdx * ISMC->NumCustomDataFloats], ISMC->NumCustomDataFloats);
+					InstancesCustomData.Append(InstanceCustomData);
+				}
+
+				InstancesCustomData.AddDefaulted(NumCustomDataFloats - ISMC->NumCustomDataFloats);
+			}
 		}
 	}
 	else if (const UStaticMeshComponent* SMC = Cast<UStaticMeshComponent>(InComponent))
@@ -112,4 +119,4 @@ void FISMComponentBatcher::InitComponent(UInstancedStaticMeshComponent* ISMCompo
 			ISMComponent->AdditionalRandomSeeds = TArrayView<FInstancedStaticMeshRandomSeed>(&const_cast<FISMComponentBatcher*>(this)->RandomSeeds[1], RandomSeeds.Num() - 1);
 		}
 	}		
-}
\ No newline at end of file
+}
