UE-150540 - Deprecate or internalize float/double cast nodes (and related nodes) in Blueprint

Several single-precision float conversion functions were deprecated in favor of their double counterparts, which we can also safely redirect. Note that the two float<->double conversion functions are only deprecated for Blueprints. These functions are still used internally by the math expression node.

Additionally:
* Various comments were updated to use the neutral "Float" wording for Blueprints.
* Normalized verb forms of "converts"

diff --git a/Engine/Config/BaseEngine.ini b/Engine/Config/BaseEngine.ini
--- a/Engine/Config/BaseEngine.ini
+++ b/Engine/Config/BaseEngine.ini
@@ -1101,6 +1101,13 @@ DefaultWorldPartitionHLODLayer=/Engine/Maps/Templates/HLODs/HLODLayer_Instanced.
 +PropertyRedirects=(OldName="DecalComponent.SetFadeIn.Duaration", NewName="DecalComponent.SetFadeIn.Duration")
 +PropertyRedirects=(OldName="BlendSpace.PerBoneBlend", NewName="BlendSpace.ManualPerBoneOverrides")
 +ClassRedirects=(OldName="/Script/MovieRenderPipelineCore.MoviePipelineMasterConfig", NewName="/Script/MovieRenderPipelineCore.MoviePipelinePrimaryConfig")
++FunctionRedirects=(OldName="KismetMathLibrary.Conv_FloatToLinearColor", NewName="Conv_DoubleToLinearColor")
++FunctionRedirects=(OldName="KismetMathLibrary.Conv_ByteToFloat", NewName="Conv_ByteToDouble")
++FunctionRedirects=(OldName="KismetMathLibrary.Conv_IntToFloat", NewName="Conv_IntToDouble")
++FunctionRedirects=(OldName="KismetMathLibrary.Conv_BoolToFloat", NewName="Conv_BoolToDouble")
++FunctionRedirects=(OldName="KismetMathLibrary.Conv_FloatToVector", NewName="Conv_DoubleToVector")
++PropertyRedirects=(OldName="KismetMathLibrary.Conv_DoubleToLinearColor.InFloat", NewName="Conv_DoubleToLinearColor.InDouble")
++PropertyRedirects=(OldName="KismetMathLibrary.Conv_DoubleToVector.InFloat", NewName="Conv_DoubleToVector.InDouble")
 
 [CoreUObject.Metadata]
 ; Note: UnrealHeaderTool should be rerun after making changes to MetadataRedirects to catch any keys specified in class headers
diff --git a/Engine/Source/Runtime/Engine/Classes/Kismet/KismetMathLibrary.h b/Engine/Source/Runtime/Engine/Classes/Kismet/KismetMathLibrary.h
--- a/Engine/Source/Runtime/Engine/Classes/Kismet/KismetMathLibrary.h
+++ b/Engine/Source/Runtime/Engine/Classes/Kismet/KismetMathLibrary.h
@@ -541,7 +541,7 @@ class ENGINE_API UKismetMathLibrary : public UBlueprintFunctionLibrary
 		*(double*)RESULT_PARAM = GenericPercent_FloatFloat(A, B);
 	}
 
-	/** Returns the fractional part of a double. */
+	/** Returns the fractional part of a float. */
 	UFUNCTION(BlueprintPure, Category="Math|Float")
 	static double Fraction(double A);
 
@@ -933,7 +933,7 @@ class ENGINE_API UKismetMathLibrary : public UBlueprintFunctionLibrary
 	// IntPoint functions
 	//
 
-	/** Convert an IntPoint to a Vector2D */
+	/** Converts an IntPoint to a Vector2D */
 	UFUNCTION(BlueprintPure, meta = (DisplayName = "To Vector2D (IntPoint)", CompactNodeTitle = "->", ScriptMethod = "Vector2D", Keywords = "cast convert", BlueprintAutocast), Category = "Math|Conversions")
 	static FVector2D Conv_IntPointToVector2D(FIntPoint InIntPoint);
 
@@ -1010,11 +1010,11 @@ class ENGINE_API UKismetMathLibrary : public UBlueprintFunctionLibrary
 	UFUNCTION(BlueprintPure, Category = "Math|Vector2D", meta = (NativeBreakFunc))
 	static void BreakVector2D(FVector2D InVec, double& X, double& Y);
 
-	/** Convert a Vector2D to a Vector */
+	/** Converts a Vector2D to a Vector */
 	UFUNCTION(BlueprintPure, meta = (DisplayName = "To Vector (Vector2D)", CompactNodeTitle = "->", ScriptMethod = "Vector", Keywords = "cast convert", BlueprintAutocast), Category = "Math|Conversions")
 	static FVector Conv_Vector2DToVector(FVector2D InVector2D, float Z = 0);
 
-	/** Convert a Vector2D to an IntPoint */
+	/** Converts a Vector2D to an IntPoint */
 	UFUNCTION(BlueprintPure, meta = (DisplayName = "To IntPoint (Vector2D)", CompactNodeTitle = "->", ScriptMethod = "IntPoint", Keywords = "cast convert", BlueprintAutocast), Category = "Math|Conversions")
 	static FIntPoint Conv_Vector2DToIntPoint(FVector2D InVector2D);
 
@@ -1341,11 +1341,11 @@ class ENGINE_API UKismetMathLibrary : public UBlueprintFunctionLibrary
 	UFUNCTION(BlueprintPure, meta=(DisplayName = "To LinearColor (Vector)", CompactNodeTitle = "->", ScriptMethod = "LinearColor", Keywords="cast convert", BlueprintAutocast), Category="Math|Conversions")
 	static FLinearColor Conv_VectorToLinearColor(FVector InVec);
 
-	/** Convert a vector to a transform. Uses vector as location */
+	/** Converts a vector to a transform. Uses vector as location */
 	UFUNCTION(BlueprintPure, meta=(DisplayName = "To Transform (Vector)", CompactNodeTitle = "->", ScriptMethod = "Transform", Keywords="cast convert", BlueprintAutocast), Category="Math|Conversions")
 	static FTransform Conv_VectorToTransform(FVector InLocation);
 	
-	/** Convert a Vector to a Vector2D using the Vector's (X, Y) coordinates */
+	/** Converts a Vector to a Vector2D using the Vector's (X, Y) coordinates */
 	UFUNCTION(BlueprintPure, meta=(DisplayName = "To Vector2D (Vector)", CompactNodeTitle = "->", ScriptMethod = "Vector2D", Keywords="cast convert", BlueprintAutocast), Category="Math|Conversions")
 	static FVector2D Conv_VectorToVector2D(FVector InVector);
 
@@ -2060,7 +2060,7 @@ class ENGINE_API UKismetMathLibrary : public UBlueprintFunctionLibrary
 	UFUNCTION(BlueprintPure, meta = (NativeBreakFunc), Category = "Math|Vector4")
 	static void BreakVector4(const FVector4& InVec, double& X, double& Y, double& Z, double& W);
 
-	/** Convert a Vector4 to a Vector (dropping the W element) */
+	/** Converts a Vector4 to a Vector (dropping the W element) */
 	UFUNCTION(BlueprintPure, meta = (DisplayName = "To Vector (Vector4)", CompactNodeTitle = "->", ScriptMethod = "Vector", Keywords = "cast convert", BlueprintAutocast), Category = "Math|Conversions")
 	static FVector Conv_Vector4ToVector(const FVector4& InVector4);
 
@@ -2373,7 +2373,7 @@ class ENGINE_API UKismetMathLibrary : public UBlueprintFunctionLibrary
 	UFUNCTION(BlueprintPure, meta=(ScriptMethod = "ToVector", DisplayName = "Get Rotation X Vector", Keywords="rotation rotate cast convert", BlueprintAutocast), Category="Math|Rotator")
 	static FVector Conv_RotatorToVector(FRotator InRot);
 
-	/** Convert Rotator to Transform */
+	/** Converts Rotator to Transform */
 	UFUNCTION(BlueprintPure, meta = (DisplayName = "To Transform (Rotator)", CompactNodeTitle = "->", ScriptMethod = "Transform", Keywords = "cast convert", BlueprintAutocast), Category = "Math|Conversions")
 	static FTransform Conv_RotatorToTransform(const FRotator& InRotator);
 
@@ -2420,13 +2420,13 @@ class ENGINE_API UKismetMathLibrary : public UBlueprintFunctionLibrary
 	// Matrix functions
 	//
 
-	/** Convert a Matrix to a Transform 
+	/** Converts a Matrix to a Transform 
 	* (Assumes Matrix represents a transform) 
 	*/
 	UFUNCTION(BlueprintPure, meta = (DisplayName = "To Transform (Matrix)", CompactNodeTitle = "->", ScriptMethod = "Transform", Keywords = "cast convert"), Category = "Math|Conversions")
 	static FTransform Conv_MatrixToTransform(const FMatrix& InMatrix);
 
-	/** Convert a Matrix to a Rotator 
+	/** Converts a Matrix to a Rotator 
 	* (Assumes Matrix represents a transform) 
 	*/
 	UFUNCTION(BlueprintPure, meta = (DisplayName = "To Rotator (Matrix)", CompactNodeTitle = "->", ScriptMethod = "Rotator", Keywords = "cast convert"), Category = "Math|Conversions")
@@ -2915,11 +2915,11 @@ class ENGINE_API UKismetMathLibrary : public UBlueprintFunctionLibrary
 	UFUNCTION(BlueprintPure, meta = (DisplayName = "Make from Euler (Quat)"), Category = "Math|Quat")
 	static FQuat Quat_MakeFromEuler(const FVector& Euler);
 
-	/** Convert to Rotator representation of this Quaternion. */
+	/** Converts to Rotator representation of this Quaternion. */
 	UFUNCTION(BlueprintPure, meta = (DisplayName = "ToRotator (Quat)", CompactNodeTitle = "->", ScriptMethod = "Rotator", Keywords = "cast convert", BlueprintAutocast), Category = "Math|Conversions")
 	static FRotator Quat_Rotator(const FQuat& Q);
 
-	/** Convert to Quaternion representation of this Rotator. */
+	/** Converts to Quaternion representation of this Rotator. */
 	UFUNCTION(BlueprintPure, meta=(DisplayName = "ToQuaternion (Rotator)", CompactNodeTitle = "->", ScriptMethod = "Quaternion", Keywords="cast convert", BlueprintAutocast), Category="Math|Conversions")
 	static FQuat Conv_RotatorToQuaternion(FRotator InRot);
 
@@ -3041,12 +3041,12 @@ class ENGINE_API UKismetMathLibrary : public UBlueprintFunctionLibrary
 	UFUNCTION(BlueprintCallable, meta = (ScriptMethod = "SetRandomHue"), Category = "Math|Color")
 	static void LinearColor_SetRandomHue(UPARAM(ref) FLinearColor& InOutColor);
 
-	/** Convert a float into a LinearColor, where each element is a float */
-	UFUNCTION(BlueprintPure, meta = (DisplayName = "To LinearColor (Float)", CompactNodeTitle = "->", Keywords = "cast convert", BlueprintAutocast), Category = "Math|Conversions")
+	/** Converts a float into a LinearColor, where each element is a float */
+	UE_DEPRECATED(5.2, "This method has been deprecated and will be removed. Use the double version instead.")
 	static FLinearColor Conv_FloatToLinearColor(float InFloat);
 
-	/** Convert a float into a LinearColor, where each element is a double */
-	UFUNCTION(BlueprintPure, meta = (DisplayName = "To LinearColor (Double)", CompactNodeTitle = "->", Keywords = "cast convert", BlueprintAutocast), Category = "Math|Conversions")
+	/** Converts a float into a LinearColor, where each RGB element is that float */
+	UFUNCTION(BlueprintPure, meta = (DisplayName = "To LinearColor (Float)", CompactNodeTitle = "->", Keywords = "cast convert", BlueprintAutocast), Category = "Math|Conversions")
 	static FLinearColor Conv_DoubleToLinearColor(double InDouble);
 
 	/** Make a color from individual color components (HSV space; Hue is [0..360) while Saturation and Value are 0..1) */
@@ -3077,7 +3077,7 @@ class ENGINE_API UKismetMathLibrary : public UBlueprintFunctionLibrary
 	UFUNCTION(BlueprintPure, meta = (DisplayName = "To Vector (LinearColor)", ScriptMethod = "ToRGBVector", CompactNodeTitle = "->", Keywords = "cast convert", BlueprintAutocast), Category = "Math|Conversions")
 	static FVector Conv_LinearColorToVector(FLinearColor InLinearColor);
 
-	/** Convert from linear to 8-bit RGBE as outlined in Gregory Ward's Real Pixels article, Graphics Gems II, page 80. */
+	/** Converts from linear to 8-bit RGBE as outlined in Gregory Ward's Real Pixels article, Graphics Gems II, page 80. */
 	UFUNCTION(BlueprintPure, meta = (DisplayName = "To RGBE (LinearColor)", ScriptMethod = "ToRGBE", Keywords = "cast convert", BlueprintAutocast), Category = "Math|Color")
 	static FColor LinearColor_ToRGBE(FLinearColor InLinearColor);
 
@@ -3522,19 +3522,19 @@ class ENGINE_API UKismetMathLibrary : public UBlueprintFunctionLibrary
 	// -- Begin K2 utilities
 
 	/** Converts a byte to a float */
-	UFUNCTION(BlueprintPure, meta=(DisplayName = "To Float (Byte)", CompactNodeTitle = "->", Keywords="cast convert", BlueprintAutocast), Category="Math|Conversions")
+	UE_DEPRECATED(5.2, "This method has been deprecated and will be removed. Use the double version instead.")
 	static float Conv_ByteToFloat(uint8 InByte);
 
-	/** Converts a byte to a double */
-	UFUNCTION(BlueprintPure, meta = (DisplayName = "To Double (Byte)", CompactNodeTitle = "->", Keywords = "cast convert", BlueprintAutocast), Category = "Math|Conversions")
+	/** Converts a byte to a float */
+	UFUNCTION(BlueprintPure, meta = (DisplayName = "To Float (Byte)", CompactNodeTitle = "->", Keywords = "cast convert", BlueprintAutocast), Category = "Math|Conversions")
 	static double Conv_ByteToDouble(uint8 InByte);
 
 	/** Converts an integer to a float */
-	UFUNCTION(BlueprintPure, meta=(DisplayName = "To Float (Integer)", CompactNodeTitle = "->", Keywords="cast convert", BlueprintAutocast), Category="Math|Conversions")
+	UE_DEPRECATED(5.2, "This method has been deprecated and will be removed. Use the double version instead.")
 	static float Conv_IntToFloat(int32 InInt);
 
-	/** Converts an integer to a double */
-	UFUNCTION(BlueprintPure, meta = (DisplayName = "To Double (Integer)", CompactNodeTitle = "->", Keywords = "cast convert", BlueprintAutocast), Category = "Math|Conversions")
+	/** Converts an integer to a float */
+	UFUNCTION(BlueprintPure, meta = (DisplayName = "To Float (Integer)", CompactNodeTitle = "->", Keywords = "cast convert", BlueprintAutocast), Category = "Math|Conversions")
 	static double Conv_IntToDouble(int32 InInt);
 
 	/** Converts an integer to a 64 bit integer */
@@ -3550,23 +3550,23 @@ class ENGINE_API UKismetMathLibrary : public UBlueprintFunctionLibrary
 	static int32 Conv_Int64ToInt(int64 InInt);
 
 	/** Converts a 64 bit floating point to a 32 bit floating point (if the float is too large, returns the low 32 bits) */
-	UFUNCTION(BlueprintPure, meta = (DisplayName = "To Float (Double)", CompactNodeTitle = "->", Keywords = "cast convert", BlueprintAutocast), Category = "Math|Conversions")
+	UFUNCTION(BlueprintPure, meta = (BlueprintInternalUseOnly = "true", DeprecatedFunction, DeprecationMessage = "Explicit conversions between floats and doubles are not necessary. Please remove node."))
 	static float Conv_DoubleToFloat(double InDouble);
 
 	/** Converts a 32 bit floating point to a 64 bit floating point */
-	UFUNCTION(BlueprintPure, meta = (DisplayName = "To Double (Float)", CompactNodeTitle = "->", Keywords = "cast convert", BlueprintAutocast), Category = "Math|Conversions")
+	UFUNCTION(BlueprintPure, meta = (BlueprintInternalUseOnly = "true", DeprecatedFunction, DeprecationMessage = "Explicit conversions between floats and doubles are not necessary. Please remove node."))
 	static double Conv_FloatToDouble(float InFloat);
 
 	/** Converts a 64 bit integer to a byte (if the integer is too large, returns the low 8 bits) */
 	UFUNCTION(BlueprintPure, meta=(DisplayName = "To Byte (Integer64)", CompactNodeTitle = "->", Keywords="cast convert", BlueprintAutocast), Category="Math|Conversions")
 	static uint8 Conv_Int64ToByte(int64 InInt);
 
-	/** Converts a 64 bit floating point to a 64 bit integer */
-	UFUNCTION(BlueprintPure, meta = (DisplayName = "To Integer64 (Double)", CompactNodeTitle = "->", Keywords = "cast convert", BlueprintAutocast), Category = "Math|Conversions")
+	/** Converts a float to a 64 bit integer */
+	UFUNCTION(BlueprintPure, meta = (DisplayName = "To Integer64 (Float)", CompactNodeTitle = "->", Keywords = "cast convert", BlueprintAutocast), Category = "Math|Conversions")
 	static int64 Conv_DoubleToInt64(double InDouble);
 
-	/** Converts a 64 bit integer to a 64 bit floating point */
-	UFUNCTION(BlueprintPure, meta = (DisplayName = "To Double (Integer64)", CompactNodeTitle = "->", Keywords = "cast convert", BlueprintAutocast), Category = "Math|Conversions")
+	/** Converts a 64 bit integer to a float */
+	UFUNCTION(BlueprintPure, meta = (DisplayName = "To Float (Integer64)", CompactNodeTitle = "->", Keywords = "cast convert", BlueprintAutocast), Category = "Math|Conversions")
 	static double Conv_Int64ToDouble(int64 InInt);
 
 	/** Converts an integer to an IntVector*/
@@ -3586,9 +3586,13 @@ class ENGINE_API UKismetMathLibrary : public UBlueprintFunctionLibrary
 	static int32 Conv_BoolToInt(bool InBool);
 
 	/** Converts a bool to a float (0.0f or 1.0f) */
-	UFUNCTION(BlueprintPure, meta=(DisplayName = "To Float (Boolean)", CompactNodeTitle = "->", Keywords="cast convert", BlueprintAutocast), Category="Math|Conversions")
+	UE_DEPRECATED(5.2, "This method has been deprecated and will be removed. Use the double version instead.")
 	static float Conv_BoolToFloat(bool InBool);
 
+	/** Converts a bool to a float (0.0 or 1.0) */
+	UFUNCTION(BlueprintPure, meta = (BlueprintInternalUseOnly = "true"))
+	static double Conv_BoolToDouble(bool InBool);
+
 	/** Converts a bool to a byte */
 	UFUNCTION(BlueprintPure, meta=(DisplayName = "To Byte (Boolean)", CompactNodeTitle = "->", Keywords="cast convert", BlueprintAutocast), Category="Math|Conversions")
 	static uint8 Conv_BoolToByte(bool InBool);
@@ -3605,22 +3609,18 @@ class ENGINE_API UKismetMathLibrary : public UBlueprintFunctionLibrary
 	UFUNCTION(BlueprintPure, meta = (DisplayName = "To LinearColor (Color)", CompactNodeTitle = "->", Keywords = "cast convert", BlueprintAutocast), Category = "Math|Conversions")
 	static FLinearColor Conv_ColorToLinearColor(FColor InColor);
 
-	/** Convert an IntVector to a vector */
+	/** Converts an IntVector to a vector */
 	UFUNCTION(BlueprintPure, meta = (DisplayName = "To Vector (IntVector)", CompactNodeTitle = "->", Keywords = "cast convert", BlueprintAutocast), Category = "Math|Conversions")
 	static FVector Conv_IntVectorToVector(const FIntVector& InIntVector);
 
-	/** Convert a float into a vector, where each element is that float */
-	UFUNCTION(BlueprintPure, BlueprintInternalUseOnly, meta=(DisplayName = "Single Precision To Vector", CompactNodeTitle = "->", Keywords="cast convert", BlueprintAutocast), Category="Math|Conversions")
+	/** Converts a float into a vector, where each element is that float */
+	UE_DEPRECATED(5.2, "This method has been deprecated and will be removed. Use the double version instead.")
 	static FVector Conv_FloatToVector(float InFloat);
 
-	/** Convert a double into a vector, where each element is that double */
-	UFUNCTION(BlueprintPure, meta = (DisplayName = "To Vector", CompactNodeTitle = "->", Keywords = "cast convert", BlueprintAutocast), Category = "Math|Conversions")
+	/** Converts a double into a vector, where each element is that float */
+	UFUNCTION(BlueprintPure, meta = (DisplayName = "To Vector (Float)", CompactNodeTitle = "->", Keywords = "cast convert", BlueprintAutocast), Category = "Math|Conversions")
 	static FVector Conv_DoubleToVector(double InDouble);
 
-	/** Convert a float into a vector, where each element is that float */
-	UFUNCTION(BlueprintPure, BlueprintInternalUseOnly, meta=(DisplayName = "Single Precision To Vector2D", CompactNodeTitle = "->", Keywords="cast convert", BlueprintAutocast), Category="Math|Conversions")
-	static FVector2D Conv_FloatToVector2D(float InFloat);
-
 	/** Convert a float into a vector, where each element is that float */
 	UFUNCTION(BlueprintPure, meta=(DisplayName = "To Vector2D", CompactNodeTitle = "->", Keywords="cast convert", BlueprintAutocast), Category="Math|Conversions")
 	static FVector2D Conv_DoubleToVector2D(double InDouble);
@@ -3885,7 +3885,7 @@ class ENGINE_API UKismetMathLibrary : public UBlueprintFunctionLibrary
 	UFUNCTION(BlueprintPure, Category="Math|Transform", meta = (DisplayName = "Determinant", ScriptMethod = "Determinant"))
 	static float Transform_Determinant(const FTransform& Transform);
 
-	/** Convert a Transform to a Matrix with scale */
+	/** Converts a Transform to a Matrix with scale */
 	UFUNCTION(BlueprintPure, Category="Math|Transform", meta = (DisplayName = "To Matrix (Transform)", ScriptMethod = "ToMatrix", CompactNodeTitle = "->", Keywords = "cast convert", BlueprintAutocast))
 	static FMatrix Conv_TransformToMatrix(const FTransform& Transform);
 
diff --git a/Engine/Source/Runtime/Engine/Classes/Kismet/KismetMathLibrary.inl b/Engine/Source/Runtime/Engine/Classes/Kismet/KismetMathLibrary.inl
--- a/Engine/Source/Runtime/Engine/Classes/Kismet/KismetMathLibrary.inl
+++ b/Engine/Source/Runtime/Engine/Classes/Kismet/KismetMathLibrary.inl
@@ -2814,6 +2814,12 @@ float UKismetMathLibrary::Conv_BoolToFloat(bool InBool)
 	return InBool ? 1.0f : 0.0f;
 }
 
+KISMET_MATH_FORCEINLINE
+double UKismetMathLibrary::Conv_BoolToDouble(bool InBool)
+{
+	return InBool ? 1.0 : 0.0;
+}
+
 KISMET_MATH_FORCEINLINE
 int32 UKismetMathLibrary::Conv_ByteToInt(uint8 InByte)
 {
@@ -2862,12 +2868,6 @@ FVector UKismetMathLibrary::Conv_DoubleToVector(double InDouble)
 	return FVector(InDouble);
 }
 
-KISMET_MATH_FORCEINLINE
-FVector2D UKismetMathLibrary::Conv_FloatToVector2D(float InFloat)
-{
-	return FVector2D(InFloat);
-}
-
 KISMET_MATH_FORCEINLINE
 FVector2D UKismetMathLibrary::Conv_DoubleToVector2D(double InDouble)
 {
