UE-152345 - Vector addition wildcard does not cast from integer to float (double-precision) when split apart vector pin is utilized

Promotable operator nodes permit pin links of mismatched types (eg: int to float). Normally, the expansion stage will resolve this discrepancy by creating intermediate cast nodes. However, split pin expansion occurs before this stage. As a result, pin connections of different types will fail.

We can resolve this by making MovePinLinks a bit more flexible. If a connection requires a conversion node during split pin expansion, we can just go ahead and insert the cast node.

diff --git a/Engine/Source/Runtime/Engine/Private/EdGraph/EdGraphSchema.cpp b/Engine/Source/Runtime/Engine/Private/EdGraph/EdGraphSchema.cpp
--- a/Engine/Source/Runtime/Engine/Private/EdGraph/EdGraphSchema.cpp
+++ b/Engine/Source/Runtime/Engine/Private/EdGraph/EdGraphSchema.cpp
@@ -680,9 +680,10 @@ FPinConnectionResponse UEdGraphSchema::MovePinLinks(UEdGraphPin& MoveFromPin, UE
 #endif
 		{
 			FPinConnectionResponse Response = CanCreateConnection(&MoveToPin, NewLink);
+			const bool bCanConnect = Response.CanSafeConnect() || (Response.Response == CONNECT_RESPONSE_MAKE_WITH_CONVERSION_NODE);
 
 #if WITH_EDITOR
-			if (!Response.CanSafeConnect() && bNotifyLinkedNodes)
+			if (!bCanConnect && bNotifyLinkedNodes)
 			{
 				// Connection failed, so notify and try again
 				if (UEdGraphNode* LinkedToNode = NewLink->GetOwningNodeUnchecked())
@@ -692,9 +693,16 @@ FPinConnectionResponse UEdGraphSchema::MovePinLinks(UEdGraphPin& MoveFromPin, UE
 				}
 			}
 #endif
-			if (Response.CanSafeConnect())
+			if (bCanConnect)
 			{
-				MoveToPin.MakeLinkTo(NewLink);
+				if (Response.Response == CONNECT_RESPONSE_MAKE_WITH_CONVERSION_NODE)
+				{
+					CreateAutomaticConversionNodeAndConnections(&MoveToPin, NewLink);
+				}
+				else
+				{
+					MoveToPin.MakeLinkTo(NewLink);
+				}
 			}	
 			else
 			{ 
