Added FWorldPartitionUtils::FSimulateCookedSession to allow to simulate a cooked streaming generation of a partitioned world and to simulate streaming queries.

diff --git a/Engine/Source/Runtime/Engine/Private/WorldPartition/ContentBundle/ContentBundleContainer.cpp b/Engine/Source/Runtime/Engine/Private/WorldPartition/ContentBundle/ContentBundleContainer.cpp
--- a/Engine/Source/Runtime/Engine/Private/WorldPartition/ContentBundle/ContentBundleContainer.cpp
+++ b/Engine/Source/Runtime/Engine/Private/WorldPartition/ContentBundle/ContentBundleContainer.cpp
@@ -384,9 +384,10 @@ void FContentBundleContainer::OnPreGenerateStreaming(TArray<FString>* OutPackage
 
 	GetInjectedWorld()->ContentBundleManager->GetPIEDuplicateHelper()->Clear();
 
+	const bool bIsPIE = true;
 	for (TSharedPtr<FContentBundleEditor>& ContentBundle : GetEditorContentBundles())
 	{
-		ContentBundle->GenerateStreaming(OutPackageToGenerate);
+		ContentBundle->GenerateStreaming(OutPackageToGenerate, bIsPIE);
 	}
 }
 
diff --git a/Engine/Source/Runtime/Engine/Private/WorldPartition/ContentBundle/ContentBundleEditor.cpp b/Engine/Source/Runtime/Engine/Private/WorldPartition/ContentBundle/ContentBundleEditor.cpp
--- a/Engine/Source/Runtime/Engine/Private/WorldPartition/ContentBundle/ContentBundleEditor.cpp
+++ b/Engine/Source/Runtime/Engine/Private/WorldPartition/ContentBundle/ContentBundleEditor.cpp
@@ -304,7 +304,7 @@ void FContentBundleEditor::AddReferencedObjects(FReferenceCollector& Collector)
 	Collector.AddReferencedObject(UnsavedActorMonitor);
 }
 
-void FContentBundleEditor::GenerateStreaming(TArray<FString>* OutPackageToGenerate)
+void FContentBundleEditor::GenerateStreaming(TArray<FString>* OutPackageToGenerate, bool bIsPIE)
 {
 	TRACE_CPUPROFILER_EVENT_SCOPE(FContentBundleEditor::GenerateStreaming);
 
@@ -327,7 +327,7 @@ void FContentBundleEditor::GenerateStreaming(TArray<FString>* OutPackageToGenera
 	});
 	UE_LOG(LogContentBundle, Log, TEXT("[CB: %s] Generated streaming cells. %u cells were generated"), *GetDescriptor()->GetDisplayName(), CellCount);
 
-	if (!IsRunningCookCommandlet())
+	if (bIsPIE)
 	{
 		UContentBundleManager* ContentBundleManager = GetInjectedWorld()->ContentBundleManager;
 		UContentBundleDuplicateForPIEHelper* DuplicateForPIEHelper = ContentBundleManager->GetPIEDuplicateHelper();
@@ -351,7 +351,10 @@ void FContentBundleEditor::OnBeginCook(IWorldPartitionCookPackageContext& CookCo
 bool FContentBundleEditor::GatherPackagesToCook(class IWorldPartitionCookPackageContext& CookContext)
 {
 	TArray<FString> PackageToGenerate;
-	GenerateStreaming(&PackageToGenerate);
+	const bool bIsPIE = false;
+	GenerateStreaming(&PackageToGenerate, bIsPIE);
+
+	CookPackageIdsToCell.Empty();
 
 	bool bIsSuccess = true;
 
@@ -470,6 +473,11 @@ bool FContentBundleEditor::PopulateGeneratedPackageForCook(class IWorldPartition
 	return bIsSuccess;
 }
 
+UWorldPartitionRuntimeCell* FContentBundleEditor::GetCellForPackage(const FWorldPartitionCookPackage& PackageToCook) const
+{
+	UWorldPartitionRuntimeCell** MatchingCell = const_cast<UWorldPartitionRuntimeCell**>(CookPackageIdsToCell.Find(PackageToCook.PackageId));
+	return MatchingCell ? *MatchingCell : nullptr;
+}
 
 void FContentBundleEditor::BroadcastChanged()
 {
diff --git a/Engine/Source/Runtime/Engine/Private/WorldPartition/Cook/WorldPartitionCookPackageContext.cpp b/Engine/Source/Runtime/Engine/Private/WorldPartition/Cook/WorldPartitionCookPackageContext.cpp
--- a/Engine/Source/Runtime/Engine/Private/WorldPartition/Cook/WorldPartitionCookPackageContext.cpp
+++ b/Engine/Source/Runtime/Engine/Private/WorldPartition/Cook/WorldPartitionCookPackageContext.cpp
@@ -5,6 +5,7 @@
 
 #include "WorldPartition/WorldPartitionLog.h"
 #include "WorldPartition/Cook/WorldPartitionCookPackageGenerator.h"
+#include "WorldPartition/WorldPartitionRuntimeCell.h"
 
 FWorldPartitionCookPackageContext::FWorldPartitionCookPackageContext()
 {
@@ -79,4 +80,32 @@ const FWorldPartitionCookPackage* FWorldPartitionCookPackageContext::AddPackageT
 	return nullptr;
 }
 
+bool FWorldPartitionCookPackageContext::GatherPackagesToCook()
+{
+	bool bIsSuccess = true;
+
+	for (IWorldPartitionCookPackageGenerator* CookPackageGenerator : CookPackageGenerators)
+	{
+		if (CookPackageGenerator->GatherPackagesToCook(*this))
+		{
+			if (const TArray<FWorldPartitionCookPackage*>* CookPackages = GetCookPackages(CookPackageGenerator))
+			{
+				for (FWorldPartitionCookPackage* CookPackage : *CookPackages)
+				{
+					if (UWorldPartitionRuntimeCell* Cell = CookPackage ? CookPackageGenerator->GetCellForPackage(*CookPackage) : nullptr)
+					{
+						Cell->SetLevelPackageName(*CookPackage->GetFullGeneratedPath());
+					}
+				}
+			}
+		}
+		else
+		{
+			bIsSuccess = false;
+		}
+	}
+
+	return bIsSuccess;
+}
+
 #endif
\ No newline at end of file
diff --git a/Engine/Source/Runtime/Engine/Private/WorldPartition/Cook/WorldPartitionCookPackageContext.h b/Engine/Source/Runtime/Engine/Private/WorldPartition/Cook/WorldPartitionCookPackageContext.h
--- a/Engine/Source/Runtime/Engine/Private/WorldPartition/Cook/WorldPartitionCookPackageContext.h
+++ b/Engine/Source/Runtime/Engine/Private/WorldPartition/Cook/WorldPartitionCookPackageContext.h
@@ -21,6 +21,7 @@ public:
 
 	virtual const FWorldPartitionCookPackage* AddLevelStreamingPackageToGenerate(IWorldPartitionCookPackageGenerator* CookPackageGenerator, const FString& Root, const FString& RelativePath) override;
 	virtual const FWorldPartitionCookPackage* AddGenericPackageToGenerate(IWorldPartitionCookPackageGenerator* CookPackageGenerator, const FString& Root, const FString& RelativePath) override;
+	virtual bool GatherPackagesToCook() override;
 	//~ End IWorldPartitionCookPackageContext Interface 
 
 	const TArray<FWorldPartitionCookPackage*>* GetCookPackages(const IWorldPartitionCookPackageGenerator* CookPackageGenerator) const { return PackagesToCookByGenerator.Find(CookPackageGenerator); }
diff --git a/Engine/Source/Runtime/Engine/Private/WorldPartition/WorldPartition.cpp b/Engine/Source/Runtime/Engine/Private/WorldPartition/WorldPartition.cpp
--- a/Engine/Source/Runtime/Engine/Private/WorldPartition/WorldPartition.cpp
+++ b/Engine/Source/Runtime/Engine/Private/WorldPartition/WorldPartition.cpp
@@ -1395,6 +1395,15 @@ void UWorldPartition::UpdateStreamingState()
 	}
 }
 
+bool UWorldPartition::GetIntersectingCells(const TArray<FWorldPartitionStreamingQuerySource>& InSources, TArray<const IWorldPartitionCell*>& OutCells) const
+{
+	if (StreamingPolicy)
+	{
+		return StreamingPolicy->GetIntersectingCells(InSources, OutCells);
+	}
+	return false;
+}
+
 bool UWorldPartition::CanAddLoadedLevelToWorld(class ULevel* InLevel) const
 {
 	if (GetWorld()->IsGameWorld())
@@ -1559,6 +1568,12 @@ bool UWorldPartition::PopulateGeneratedPackageForCook(IWorldPartitionCookPackage
 	return RuntimeHash->PopulateGeneratedPackageForCook(InPackagesToCook, OutModifiedPackages);
 }
 
+UWorldPartitionRuntimeCell* UWorldPartition::GetCellForPackage(const FWorldPartitionCookPackage& PackageToCook) const
+{
+	check(RuntimeHash);
+	return RuntimeHash->GetCellForPackage(PackageToCook);
+}
+
 TArray<FBox> UWorldPartition::GetUserLoadedEditorRegions() const
 {
 	TArray<FBox> Result;
diff --git a/Engine/Source/Runtime/Engine/Private/WorldPartition/WorldPartitionCookPackageSplitter.cpp b/Engine/Source/Runtime/Engine/Private/WorldPartition/WorldPartitionCookPackageSplitter.cpp
--- a/Engine/Source/Runtime/Engine/Private/WorldPartition/WorldPartitionCookPackageSplitter.cpp
+++ b/Engine/Source/Runtime/Engine/Private/WorldPartition/WorldPartitionCookPackageSplitter.cpp
@@ -104,11 +104,8 @@ TArray<ICookPackageSplitter::FGeneratedPackage> FWorldPartitionCookPackageSplitt
 
 	WorldPartition->BeginCook(CookContext);
 
-	bool bIsSuccess = true;
-	for (IWorldPartitionCookPackageGenerator* CookPackageGenerator : CookContext.GetCookPackageGenerators())
-	{
-		bIsSuccess &= CookPackageGenerator->GatherPackagesToCook(CookContext);
-	}
+	bool bIsSuccess = CookContext.GatherPackagesToCook();
+	UE_CLOG(!bIsSuccess, LogWorldPartition, Warning, TEXT("[Cook] Errors while gathering packages to took from generators for owner object %s."), *GetFullNameSafe(OwnerObject));
 
 	UE_LOG(LogWorldPartition, Log, TEXT("[Cook] Gathered %u packages to generate from %u Generators."), CookContext.NumPackageToGenerate(), CookContext.NumGenerators());
 
diff --git a/Engine/Source/Runtime/Engine/Private/WorldPartition/WorldPartitionRuntimeCell.cpp b/Engine/Source/Runtime/Engine/Private/WorldPartition/WorldPartitionRuntimeCell.cpp
--- a/Engine/Source/Runtime/Engine/Private/WorldPartition/WorldPartitionRuntimeCell.cpp
+++ b/Engine/Source/Runtime/Engine/Private/WorldPartition/WorldPartitionRuntimeCell.cpp
@@ -122,10 +122,10 @@ void UWorldPartitionRuntimeCell::AppendStreamingSourceInfo(const FWorldPartition
 void UWorldPartitionRuntimeCell::MergeStreamingSourceInfo() const
 {}
 
-int32 UWorldPartitionRuntimeCell::SortCompare(const UWorldPartitionRuntimeCell* Other) const
+int32 UWorldPartitionRuntimeCell::SortCompare(const UWorldPartitionRuntimeCell* Other, bool bCanUseSortingCache) const
 {
 	// Source priority (lower value is higher prio)
-	const int32 Comparison = (int32)CachedMinSourcePriority - (int32)Other->CachedMinSourcePriority;
+	const int32 Comparison = bCanUseSortingCache ? ((int32)CachedMinSourcePriority - (int32)Other->CachedMinSourcePriority) : 0;
 	// Cell priority (lower value is higher prio)
 	return (Comparison != 0) ? Comparison : (Priority - Other->Priority);
 }
@@ -177,3 +177,12 @@ bool UWorldPartitionRuntimeCell::ContainsDataLayer(const UDataLayerInstance* Dat
 {
 	return GetDataLayers().Contains(DataLayerInstance->GetDataLayerFName());
 }
+
+FName UWorldPartitionRuntimeCell::GetLevelPackageName() const
+{ 
+#if WITH_EDITOR
+	return LevelPackageName;
+#else
+	return NAME_None;
+#endif
+}
\ No newline at end of file
diff --git a/Engine/Source/Runtime/Engine/Private/WorldPartition/WorldPartitionRuntimeHash.cpp b/Engine/Source/Runtime/Engine/Private/WorldPartition/WorldPartitionRuntimeHash.cpp
--- a/Engine/Source/Runtime/Engine/Private/WorldPartition/WorldPartitionRuntimeHash.cpp
+++ b/Engine/Source/Runtime/Engine/Private/WorldPartition/WorldPartitionRuntimeHash.cpp
@@ -204,4 +204,23 @@ void UWorldPartitionRuntimeHash::FStreamingSourceCells::AddCell(const UWorldPart
 	Cells.Add(Cell);
 }
 
+void FWorldPartitionQueryCache::AddCellInfo(const UWorldPartitionRuntimeCell* Cell, const FSphericalSector& SourceShape)
+{
+	const double SquareDistance = FVector::DistSquared2D(SourceShape.GetCenter(), Cell->GetCellBounds().GetCenter());
+	if (double* ExistingSquareDistance = CellToSourceMinSqrDistances.Find(Cell))
+	{
+		*ExistingSquareDistance = FMath::Min(*ExistingSquareDistance, SquareDistance);
+	}
+	else
+	{
+		CellToSourceMinSqrDistances.Add(Cell, SquareDistance);
+	}
+}
+
+double FWorldPartitionQueryCache::GetCellMinSquareDist(const UWorldPartitionRuntimeCell* Cell) const
+{
+	const double* Dist = CellToSourceMinSqrDistances.Find(Cell);
+	return Dist ? *Dist : MAX_dbl;
+}
+
 #undef LOCTEXT_NAMESPACE
diff --git a/Engine/Source/Runtime/Engine/Private/WorldPartition/WorldPartitionRuntimeLevelStreamingCell.cpp b/Engine/Source/Runtime/Engine/Private/WorldPartition/WorldPartitionRuntimeLevelStreamingCell.cpp
--- a/Engine/Source/Runtime/Engine/Private/WorldPartition/WorldPartitionRuntimeLevelStreamingCell.cpp
+++ b/Engine/Source/Runtime/Engine/Private/WorldPartition/WorldPartitionRuntimeLevelStreamingCell.cpp
@@ -60,6 +60,35 @@ bool UWorldPartitionRuntimeLevelStreamingCell::HasActors() const
 #endif
 }
 
+#if WITH_EDITOR
+TSet<FName> UWorldPartitionRuntimeLevelStreamingCell::GetActorPackageNames() const
+{
+	TSet<FName> Actors;
+	Actors.Reserve(Packages.Num());
+	for (const FWorldPartitionRuntimeCellObjectMapping& Package : Packages)
+	{
+		Actors.Add(Package.Package);
+	}
+	return Actors;
+}
+#endif
+
+FName UWorldPartitionRuntimeLevelStreamingCell::GetLevelPackageName() const
+{
+#if WITH_EDITOR
+	UWorld* World = GetCellOwner()->GetOwningWorld();
+	if (World->IsPlayInEditor())
+	{
+		return FName(UWorldPartitionLevelStreamingPolicy::GetCellPackagePath(GetFName(), World));
+	}
+#endif
+	if (LevelStreaming)
+	{
+		return LevelStreaming->GetWorldAssetPackageFName();
+	}
+	return Super::GetLevelPackageName();
+}
+
 TArray<FName> UWorldPartitionRuntimeLevelStreamingCell::GetActors() const
 {
 	TArray<FName> Actors;
diff --git a/Engine/Source/Runtime/Engine/Private/WorldPartition/WorldPartitionRuntimeSpatialHash.cpp b/Engine/Source/Runtime/Engine/Private/WorldPartition/WorldPartitionRuntimeSpatialHash.cpp
--- a/Engine/Source/Runtime/Engine/Private/WorldPartition/WorldPartitionRuntimeSpatialHash.cpp
+++ b/Engine/Source/Runtime/Engine/Private/WorldPartition/WorldPartitionRuntimeSpatialHash.cpp
@@ -159,7 +159,7 @@ int64 FSpatialHashStreamingGrid::GetCellSize(int32 Level) const
 	return GetGridHelper().Levels[Level].CellSize;
 }
 
-void FSpatialHashStreamingGrid::GetCells(const FWorldPartitionStreamingQuerySource& QuerySource, TSet<const UWorldPartitionRuntimeCell*>& OutCells, bool bEnableZCulling) const
+void FSpatialHashStreamingGrid::GetCells(const FWorldPartitionStreamingQuerySource& QuerySource, TSet<const UWorldPartitionRuntimeCell*>& OutCells, bool bEnableZCulling, FWorldPartitionQueryCache* QueryCache) const
 {
 	TRACE_CPUPROFILER_EVENT_SCOPE(FSpatialHashStreamingGrid::GetCells_QuerySource);
 
@@ -196,6 +196,10 @@ void FSpatialHashStreamingGrid::GetCells(const FWorldPartitionStreamingQuerySour
 						if (ShouldAddCell(Cell, QuerySource))
 						{
 							OutCells.Add(Cell);
+							if (QueryCache)
+							{
+								QueryCache->AddCellInfo(Cell, Shape);
+							}
 						}
 					}
 				});
@@ -1158,6 +1162,7 @@ bool UWorldPartitionRuntimeSpatialHash::GenerateStreaming(UWorldPartitionStreami
 	// Fix case where StreamingGrids might have been persisted.
 	bIsNameToGridMappingDirty = true;
 	StreamingGrids.Empty();
+	check(PackagesToGenerateForCook.IsEmpty());
 
 	// Append grids from ASpatialHashRuntimeGridInfo actors to runtime spatial hash grids
 	TArray<FSpatialHashRuntimeGrid> AllGrids;
@@ -1474,12 +1479,10 @@ bool UWorldPartitionRuntimeSpatialHash::CreateStreamingGrid(const FSpatialHashRu
 						check(!PackageRelativePath.IsEmpty());
 						OutPackagesToGenerate->Add(PackageRelativePath);
 
-						if (IsRunningCookCommandlet())
-						{
-							UE_LOG(LogWorldPartition, Log, TEXT("Creating runtime streaming cells %s."), *StreamingCell->GetName());
-							// Map relative package to StreamingCell for PopulateGeneratedPackageForCook/PopulateGeneratorPackageForCook
-							PackagesToGenerateForCook.Add(PackageRelativePath, StreamingCell);
-						}
+						// Map relative package to StreamingCell for PopulateGeneratedPackageForCook/PopulateGeneratorPackageForCook/GetCellForPackage
+						PackagesToGenerateForCook.Add(PackageRelativePath, StreamingCell);
+
+						UE_CLOG(IsRunningCookCommandlet(), LogWorldPartition, Log, TEXT("Creating runtime streaming cells %s."), *StreamingCell->GetName());
 					}
 				}
 			}
@@ -1503,6 +1506,12 @@ bool UWorldPartitionRuntimeSpatialHash::PopulateGeneratedPackageForCook(const FW
 	return false;
 }
 
+UWorldPartitionRuntimeCell* UWorldPartitionRuntimeSpatialHash::GetCellForPackage(const FWorldPartitionCookPackage& PackageToCook) const
+{
+	UWorldPartitionRuntimeCell** MatchingCell = const_cast<UWorldPartitionRuntimeCell**>(PackagesToGenerateForCook.Find(PackageToCook.RelativePath));
+	return MatchingCell ? *MatchingCell : nullptr;
+}
+
 TArray<UWorldPartitionRuntimeCell*> UWorldPartitionRuntimeSpatialHash::GetAlwaysLoadedCells() const
 {
 	TArray<UWorldPartitionRuntimeCell*> Result;
@@ -1530,6 +1539,8 @@ bool UWorldPartitionRuntimeSpatialHash::PrepareGeneratorPackageForCook(TArray<UP
 			return false;
 		}
 	}
+
+	//@todo_ow: here we can safely remove always loaded cells as they are not part of the OutPackagesToGenerate
 	return true;
 }
 
@@ -1554,6 +1565,7 @@ void UWorldPartitionRuntimeSpatialHash::FlushStreaming()
 {
 	bIsNameToGridMappingDirty = true;
 	StreamingGrids.Empty();
+	PackagesToGenerateForCook.Empty();
 }
 
 #endif //WITH_EDITOR
@@ -1601,20 +1613,24 @@ void UWorldPartitionRuntimeSpatialHash::ForEachStreamingCells(TFunctionRef<bool(
 	});
 }
 
-void UWorldPartitionRuntimeSpatialHash::ForEachStreamingCellsQuery(const FWorldPartitionStreamingQuerySource& QuerySource, TFunctionRef<bool(const UWorldPartitionRuntimeCell*)> Func) const
+void UWorldPartitionRuntimeSpatialHash::ForEachStreamingCellsQuery(const FWorldPartitionStreamingQuerySource& QuerySource, TFunctionRef<bool(const UWorldPartitionRuntimeCell*)> Func, FWorldPartitionQueryCache* QueryCache) const
 {
 	ForEachStreamingGrid([&](const FSpatialHashStreamingGrid& StreamingGrid)
 	{
 		if (IsCellRelevantFor(StreamingGrid.bClientOnlyVisible))
 		{
 			TSet<const UWorldPartitionRuntimeCell*> Cells;
-			StreamingGrid.GetCells(QuerySource, Cells, GetEffectiveEnableZCulling(bEnableZCulling));
+			StreamingGrid.GetCells(QuerySource, Cells, GetEffectiveEnableZCulling(bEnableZCulling), QueryCache);
 
 			for (const UWorldPartitionRuntimeCell* Cell : Cells)
 			{
-				if (!Func(Cell))
+				//@todo_ow: Remove this test once always loaded cells (which are empty) are removed at cook time
+				if (!Cell->IsAlwaysLoaded())
 				{
-					break;
+					if (!Func(Cell))
+					{
+						break;
+					}
 				}
 			}
 		}
@@ -1670,8 +1686,6 @@ void UWorldPartitionRuntimeSpatialHash::ForEachStreamingCellsSources(const TArra
 
 bool UWorldPartitionRuntimeSpatialHash::InjectExternalStreamingObject(URuntimeHashExternalStreamingObjectBase* ExternalStreamingObject)
 {
-	check(GetWorld()->IsGameWorld());
-
 	if (!ExternalStreamingObjects.Contains(ExternalStreamingObject))
 	{
 		if (URuntimeSpatialHashExternalStreamingObject* SpatialHashExternalStreamingObject = Cast<URuntimeSpatialHashExternalStreamingObject>(ExternalStreamingObject))
@@ -1707,8 +1721,6 @@ bool UWorldPartitionRuntimeSpatialHash::InjectExternalStreamingObject(URuntimeHa
 
 bool UWorldPartitionRuntimeSpatialHash::RemoveExternalStreamingObject(URuntimeHashExternalStreamingObjectBase* ExternalStreamingObject)
 {
-	check(GetWorld()->IsGameWorld());
-
 	if (URuntimeSpatialHashExternalStreamingObject* SpatialHashExternalStreamingObject = Cast<URuntimeSpatialHashExternalStreamingObject>(ExternalStreamingObject))
 	{
 		if (ExternalStreamingObjects.Remove(SpatialHashExternalStreamingObject) > 0)
diff --git a/Engine/Source/Runtime/Engine/Private/WorldPartition/WorldPartitionRuntimeSpatialHashCell.cpp b/Engine/Source/Runtime/Engine/Private/WorldPartition/WorldPartitionRuntimeSpatialHashCell.cpp
--- a/Engine/Source/Runtime/Engine/Private/WorldPartition/WorldPartitionRuntimeSpatialHashCell.cpp
+++ b/Engine/Source/Runtime/Engine/Private/WorldPartition/WorldPartitionRuntimeSpatialHashCell.cpp
@@ -165,9 +165,9 @@ void UWorldPartitionRuntimeSpatialHashCell::MergeStreamingSourceInfo() const
 	}
 }
 
-int32 UWorldPartitionRuntimeSpatialHashCell::SortCompare(const UWorldPartitionRuntimeCell* InOther) const
+int32 UWorldPartitionRuntimeSpatialHashCell::SortCompare(const UWorldPartitionRuntimeCell* InOther, bool bCanUseSortingCache) const
 {
-	int32 Result = Super::SortCompare(InOther);
+	int32 Result = Super::SortCompare(InOther, bCanUseSortingCache);
 	if (Result == 0)
 	{
 		const UWorldPartitionRuntimeSpatialHashCell* Other = Cast<const UWorldPartitionRuntimeSpatialHashCell>(InOther);
@@ -175,7 +175,7 @@ int32 UWorldPartitionRuntimeSpatialHashCell::SortCompare(const UWorldPartitionRu
 		
 		// Level (higher value is higher prio)
 		Result = Other->Level - Level;
-		if (Result == 0)
+		if (bCanUseSortingCache && (Result == 0))
 		{
 			// Closest distance (lower value is higher prio)
 			const double Diff = CachedSourceSortingDistance - Other->CachedSourceSortingDistance;
@@ -191,4 +191,4 @@ int32 UWorldPartitionRuntimeSpatialHashCell::SortCompare(const UWorldPartitionRu
 		}
 	}
 	return Result;
-}
+}
\ No newline at end of file
diff --git a/Engine/Source/Runtime/Engine/Private/WorldPartition/WorldPartitionStreamingPolicy.cpp b/Engine/Source/Runtime/Engine/Private/WorldPartition/WorldPartitionStreamingPolicy.cpp
--- a/Engine/Source/Runtime/Engine/Private/WorldPartition/WorldPartitionStreamingPolicy.cpp
+++ b/Engine/Source/Runtime/Engine/Private/WorldPartition/WorldPartitionStreamingPolicy.cpp
@@ -311,6 +311,53 @@ uint32 UWorldPartitionStreamingPolicy::ComputeStreamingSourceHash(const FWorldPa
 	return Hash;
 };
 
+bool UWorldPartitionStreamingPolicy::GetIntersectingCells(const TArray<FWorldPartitionStreamingQuerySource>& InSources, TArray<const IWorldPartitionCell*>& OutCells) const
+{
+	if (!WorldPartition || !WorldPartition->RuntimeHash)
+	{
+		return false;
+	}
+
+	FWorldPartitionQueryCache QueryCache;
+	TSet<const UWorldPartitionRuntimeCell*> Cells;
+	for (const FWorldPartitionStreamingQuerySource& Source : InSources)
+	{
+		WorldPartition->RuntimeHash->ForEachStreamingCellsQuery(Source, [&Cells](const UWorldPartitionRuntimeCell* Cell)
+		{
+			Cells.Add(Cell);
+			return true;
+		}, &QueryCache);
+	}
+
+	TArray<const UWorldPartitionRuntimeCell*> SortedCells = Cells.Array();
+	Algo::Sort(SortedCells, [&QueryCache](const UWorldPartitionRuntimeCell* CellA, const UWorldPartitionRuntimeCell* CellB)
+	{
+		const bool bCanUseSortingCache = false;
+		int32 SortCompare = CellA->SortCompare(CellB, bCanUseSortingCache);
+		if (SortCompare == 0)
+		{
+			// Closest distance (lower value is higher prio)
+			const double Diff = QueryCache.GetCellMinSquareDist(CellA) - QueryCache.GetCellMinSquareDist(CellB);
+			if (FMath::IsNearlyZero(Diff))
+			{
+				return CellA->GetLevelPackageName().LexicalLess(CellB->GetLevelPackageName());
+			}
+			else
+			{
+				return Diff < 0;
+			}
+		}
+		return SortCompare < 0;
+	});
+
+	OutCells.Reserve(SortedCells.Num());
+	for (const UWorldPartitionRuntimeCell* Cell : SortedCells)
+	{
+		OutCells.Add(Cell);
+	}
+	return true;
+}
+
 void UWorldPartitionStreamingPolicy::UpdateStreamingState()
 {
 	TRACE_CPUPROFILER_EVENT_SCOPE(UWorldPartitionStreamingPolicy::UpdateStreamingState);
diff --git a/Engine/Source/Runtime/Engine/Private/WorldPartition/WorldPartitionUtils.cpp b/Engine/Source/Runtime/Engine/Private/WorldPartition/WorldPartitionUtils.cpp
new file mode 100644
--- /dev/null
+++ b/Engine/Source/Runtime/Engine/Private/WorldPartition/WorldPartitionUtils.cpp
@@ -0,0 +1,116 @@
+// Copyright Epic Games, Inc. All Rights Reserved.
+
+#if WITH_EDITOR
+
+#include "WorldPartition/WorldPartitionUtils.h"
+#include "WorldPartition/WorldPartitionRuntimeCellInterface.h"
+#include "WorldPartition/WorldPartitionRuntimeHash.h"
+#include "WorldPartition/WorldPartitionStreamingSource.h"
+#include "WorldPartition/Cook/WorldPartitionCookPackageContext.h"
+#include "WorldPartition/ContentBundle/ContentBundleEditor.h"
+#include "WorldPartition/ContentBundle/ContentBundleWorldSubsystem.h"
+
+#define LOCTEXT_NAMESPACE "WorldPartition"
+
+FWorldPartitionUtils::FSimulateCookedSession::FSimulateCookedSession(UWorld* InWorld)
+: bIsValid(false)
+{
+	UE_CLOG(!InWorld, LogWorldPartition, Error, TEXT("FSimulateCookedSession was given an invalid world."));
+	UE_CLOG(InWorld && !InWorld->IsPartitionedWorld(), LogWorldPartition, Warning, TEXT("FSimulateCookedSession was given a non-partitioned world."));
+	
+	if (InWorld)
+	{
+		WorldPartition = InWorld->GetWorldPartition();
+		bIsValid = WorldPartition.IsValid() ? SimulateCook() : false;
+	}
+
+	UE_CLOG(!bIsValid, LogWorldPartition, Warning, TEXT("FSimulateCookedSession failed to generate streaming."));
+}
+
+bool FWorldPartitionUtils::FSimulateCookedSession::SimulateCook()
+{
+	check(WorldPartition.IsValid());
+	if (!WorldPartition->CanGenerateStreaming() || IsRunningGame() || WorldPartition->GetWorld()->IsGameWorld() || IsRunningCookCommandlet())
+	{
+		UE_CLOG(!WorldPartition->CanGenerateStreaming(), LogWorldPartition, Warning, TEXT("FSimulateCookedSession cannot generate streaming, world partition is already generating streaming."));
+		UE_CLOG(IsRunningGame() || WorldPartition->GetWorld()->IsGameWorld(), LogWorldPartition, Warning, TEXT("FSimulateCookedSession cannot generate streaming while running PIE/game."));
+		UE_CLOG(IsRunningCookCommandlet(), LogWorldPartition, Warning, TEXT("FSimulateCookedSession cannot generate streaming while cooking."));
+		return false;
+	}
+
+	// Simulate cook
+	FWorldPartitionCookPackageContext CookContext;
+	WorldPartition->BeginCook(CookContext);
+
+	if (!CookContext.GatherPackagesToCook())
+	{
+		return false;
+	}
+
+	// Simulate Content Bundles injection
+	TArray<TSharedPtr<FContentBundleEditor>> ContentBundles;
+	WorldPartition->GetWorld()->ContentBundleManager->GetEditorContentBundle(ContentBundles);
+	for (TSharedPtr<FContentBundleEditor>& ContentBundle : ContentBundles)
+	{
+		if (ContentBundle->HasCookedContent())
+		{
+			WorldPartition->RuntimeHash->InjectExternalStreamingObject(ContentBundle->GetStreamingObject());
+		}
+	}
+
+	return true;
+}
+
+FWorldPartitionUtils::FSimulateCookedSession::~FSimulateCookedSession()
+{
+	if (WorldPartition.IsValid())
+	{
+		// Simulate Content Bundles injection
+		TArray<TSharedPtr<FContentBundleEditor>> ContentBundles;
+		WorldPartition->GetWorld()->ContentBundleManager->GetEditorContentBundle(ContentBundles);
+		for (TSharedPtr<FContentBundleEditor>& ContentBundle : ContentBundles)
+		{
+			if (ContentBundle->HasCookedContent())
+			{
+				URuntimeHashExternalStreamingObjectBase* ExternalStreamingObject = ContentBundle->GetStreamingObject();
+				WorldPartition->RuntimeHash->RemoveExternalStreamingObject(ExternalStreamingObject);
+				// Trash external streaming object
+				ExternalStreamingObject->Rename(*MakeUniqueObjectName(GetTransientPackage(), URuntimeHashExternalStreamingObjectBase::StaticClass()).ToString(), GetTransientPackage(), REN_NonTransactional | REN_DontCreateRedirectors | REN_ForceNoResetLoaders | REN_DoNotDirty);
+			}
+		}
+
+		WorldPartition->FlushStreaming();
+	}
+}
+
+bool FWorldPartitionUtils::FSimulateCookedSession::ForEachStreamingCells(TFunctionRef<void(const IWorldPartitionCell*)> Func)
+{
+	if (!WorldPartition.IsValid() || !WorldPartition->RuntimeHash)
+	{
+		return false;
+	}
+
+	WorldPartition->RuntimeHash->ForEachStreamingCells([Func](const UWorldPartitionRuntimeCell* Cell)
+	{
+		if (Cell)
+		{
+			Func(Cell);
+		}
+		return true;
+	});
+
+	return true;
+}
+
+bool FWorldPartitionUtils::FSimulateCookedSession::GetIntersectingCells(const TArray<FWorldPartitionStreamingQuerySource>& InSources, TArray<const IWorldPartitionCell*>& OutCells)
+{
+	if (WorldPartition.IsValid())
+	{
+		return WorldPartition->GetIntersectingCells(InSources, OutCells);
+	}
+	return false;
+}
+	
+#undef LOCTEXT_NAMESPACE
+	
+#endif
\ No newline at end of file
diff --git a/Engine/Source/Runtime/Engine/Public/WorldPartition/ContentBundle/ContentBundleEditor.h b/Engine/Source/Runtime/Engine/Public/WorldPartition/ContentBundle/ContentBundleEditor.h
--- a/Engine/Source/Runtime/Engine/Public/WorldPartition/ContentBundle/ContentBundleEditor.h
+++ b/Engine/Source/Runtime/Engine/Public/WorldPartition/ContentBundle/ContentBundleEditor.h
@@ -52,16 +52,19 @@ public:
 	void ReferenceAllActors();
 	void UnreferenceAllActors();
 
-	void GenerateStreaming(TArray<FString>* OutPackageToGenerate);
+	void GenerateStreaming(TArray<FString>* OutPackageToGenerate, bool bIsPIE);
 
 	URuntimeHashExternalStreamingObjectBase* GetStreamingObject() const { return ExternalStreamingObject; }
 
+	// Cooking
 	void OnBeginCook(IWorldPartitionCookPackageContext& CookContext);
-
-	bool GatherPackagesToCook(class IWorldPartitionCookPackageContext& CookContext);
-	bool PopulateGeneratorPackageForCook(class IWorldPartitionCookPackageContext& CookContext, const TArray<FWorldPartitionCookPackage*>& PackagesToCook, TArray<UPackage*>& OutModifiedPackages);
-	bool PopulateGeneratedPackageForCook(class IWorldPartitionCookPackageContext& CookContext, const FWorldPartitionCookPackage& PackagesToCook, TArray<UPackage*>& OutModifiedPackages);
 	bool HasCookedContent() const { return ExternalStreamingObject != nullptr; }
+	//~Begin IWorldPartitionCookPackageGenerator
+	virtual bool GatherPackagesToCook(class IWorldPartitionCookPackageContext& CookContext) override;
+	virtual bool PopulateGeneratorPackageForCook(class IWorldPartitionCookPackageContext& CookContext, const TArray<FWorldPartitionCookPackage*>& PackagesToCook, TArray<UPackage*>& OutModifiedPackages) override;
+	virtual bool PopulateGeneratedPackageForCook(class IWorldPartitionCookPackageContext& CookContext, const FWorldPartitionCookPackage& PackageToCook, TArray<UPackage*>& OutModifiedPackages) override;
+	virtual UWorldPartitionRuntimeCell* GetCellForPackage(const FWorldPartitionCookPackage& PackageToCook) const override;
+	//~End IWorldPartitionCookPackageGenerator
 
 protected:
 	//~ Begin IContentBundle Interface
diff --git a/Engine/Source/Runtime/Engine/Public/WorldPartition/Cook/WorldPartitionCookPackageContextInterface.h b/Engine/Source/Runtime/Engine/Public/WorldPartition/Cook/WorldPartitionCookPackageContextInterface.h
--- a/Engine/Source/Runtime/Engine/Public/WorldPartition/Cook/WorldPartitionCookPackageContextInterface.h
+++ b/Engine/Source/Runtime/Engine/Public/WorldPartition/Cook/WorldPartitionCookPackageContextInterface.h
@@ -19,6 +19,8 @@ public:
 
 	virtual const FWorldPartitionCookPackage* AddLevelStreamingPackageToGenerate(IWorldPartitionCookPackageGenerator* Generator, const FString& Root, const FString& RelativePath) = 0;
 	virtual const FWorldPartitionCookPackage* AddGenericPackageToGenerate(IWorldPartitionCookPackageGenerator* Generator, const FString& Root, const FString& RelativePath) = 0;
+
+	virtual bool GatherPackagesToCook() = 0;
 };
 
 #endif
\ No newline at end of file
diff --git a/Engine/Source/Runtime/Engine/Public/WorldPartition/Cook/WorldPartitionCookPackageGenerator.h b/Engine/Source/Runtime/Engine/Public/WorldPartition/Cook/WorldPartitionCookPackageGenerator.h
--- a/Engine/Source/Runtime/Engine/Public/WorldPartition/Cook/WorldPartitionCookPackageGenerator.h
+++ b/Engine/Source/Runtime/Engine/Public/WorldPartition/Cook/WorldPartitionCookPackageGenerator.h
@@ -6,6 +6,7 @@
 #if WITH_EDITOR
 struct FWorldPartitionCookPackage;
 class IWorldPartitionCookPackageContext;
+class UWorldPartitionRuntimeCell;
 #endif
 
 class IWorldPartitionCookPackageGenerator
@@ -17,6 +18,7 @@ public:
 	virtual bool GatherPackagesToCook(IWorldPartitionCookPackageContext& CookContext) = 0;
 	virtual bool PrepareGeneratorPackageForCook(IWorldPartitionCookPackageContext& CookContext, TArray<UPackage*>& OutModifiedPackages) { return true; }
 	virtual bool PopulateGeneratorPackageForCook(IWorldPartitionCookPackageContext& CookContext, const TArray<FWorldPartitionCookPackage*>& InPackagesToCook, TArray<UPackage*>& OutModifiedPackages) = 0;
-	virtual bool PopulateGeneratedPackageForCook(IWorldPartitionCookPackageContext& CookContext, const FWorldPartitionCookPackage& InPackagesToCook, TArray<UPackage*>& OutModifiedPackages) = 0;
+	virtual bool PopulateGeneratedPackageForCook(IWorldPartitionCookPackageContext& CookContext, const FWorldPartitionCookPackage& InPackageToCook, TArray<UPackage*>& OutModifiedPackages) = 0;
+	virtual UWorldPartitionRuntimeCell* GetCellForPackage(const FWorldPartitionCookPackage& PackageToCook) const = 0;
 #endif
 };
diff --git a/Engine/Source/Runtime/Engine/Public/WorldPartition/WorldPartition.h b/Engine/Source/Runtime/Engine/Public/WorldPartition/WorldPartition.h
--- a/Engine/Source/Runtime/Engine/Public/WorldPartition/WorldPartition.h
+++ b/Engine/Source/Runtime/Engine/Public/WorldPartition/WorldPartition.h
@@ -31,6 +31,7 @@ class UWorldPartitionEditorHash;
 class UWorldPartitionRuntimeCell;
 class UWorldPartitionRuntimeHash;
 class UWorldPartitionStreamingPolicy;
+class IWorldPartitionCell;
 class IStreamingGenerationErrorHandler;
 class FLoaderAdapterAlwaysLoadedActors;
 class FLoaderAdapterActorList;
@@ -151,6 +152,7 @@ public:
 	FName GetWorldPartitionEditorName() const;
 
 	// Streaming generation
+	bool CanGenerateStreaming() const { return !StreamingPolicy; }
 	bool GenerateStreaming(TArray<FString>* OutPackagesToGenerate = nullptr);
 	bool GenerateContainerStreaming(const UActorDescContainer* ActorDescContainer, TArray<FString>* OutPackagesToGenerate = nullptr);
 	void FlushStreaming();
@@ -171,8 +173,8 @@ public:
 	virtual bool PrepareGeneratorPackageForCook(IWorldPartitionCookPackageContext& CookContext, TArray<UPackage*>& OutModifiedPackages) override;
 	virtual bool PopulateGeneratorPackageForCook(IWorldPartitionCookPackageContext& CookContext, const TArray<FWorldPartitionCookPackage*>& InPackagesToCook, TArray<UPackage*>& OutModifiedPackages) override;
 	virtual bool PopulateGeneratedPackageForCook(IWorldPartitionCookPackageContext& CookContext, const FWorldPartitionCookPackage& InPackagesToCool, TArray<UPackage*>& OutModifiedPackages) override;
+	virtual UWorldPartitionRuntimeCell* GetCellForPackage(const FWorldPartitionCookPackage& PackageToCook) const override;
 	//~ End IWorldPartitionCookPackageGenerator Interface 
-
 	// End Cooking
 
 	UE_DEPRECATED(5.1, "GetWorldBounds is deprecated, use GetEditorWorldBounds or GetRuntimeWorldBounds instead.")
@@ -256,6 +258,7 @@ public:
 	bool CanAddLoadedLevelToWorld(ULevel* InLevel) const;
 	bool IsStreamingCompleted(const TArray<FWorldPartitionStreamingSource>* InStreamingSources) const;
 	bool IsStreamingCompleted(EWorldPartitionRuntimeCellState QueryState, const TArray<FWorldPartitionStreamingQuerySource>& QuerySources, bool bExactState) const;
+	bool GetIntersectingCells(const TArray<FWorldPartitionStreamingQuerySource>& InSources, TArray<const IWorldPartitionCell*>& OutCells) const;
 
 	const TArray<FWorldPartitionStreamingSource>& GetStreamingSources() const;
 
diff --git a/Engine/Source/Runtime/Engine/Public/WorldPartition/WorldPartitionRuntimeCell.h b/Engine/Source/Runtime/Engine/Public/WorldPartition/WorldPartitionRuntimeCell.h
--- a/Engine/Source/Runtime/Engine/Public/WorldPartition/WorldPartitionRuntimeCell.h
+++ b/Engine/Source/Runtime/Engine/Public/WorldPartition/WorldPartitionRuntimeCell.h
@@ -170,7 +170,7 @@ class ENGINE_API UWorldPartitionRuntimeCell : public UObject, public IWorldParti
 	virtual bool IsLoading() const { return false; }
 	virtual const FString& GetDebugName() const { return DebugInfo.Name; }
 	virtual bool IsDebugShown() const;
-	virtual int32 SortCompare(const UWorldPartitionRuntimeCell* Other) const;
+	virtual int32 SortCompare(const UWorldPartitionRuntimeCell* Other, bool bCanUseSortingCache = true) const;
 	virtual FName GetGridName() const { return DebugInfo.GridName; }
 	bool GetClientOnlyVisible() const { return bClientOnlyVisible; }
 	virtual FGuid const& GetContentBundleID() const { return ContentBundleID; }
@@ -186,7 +186,7 @@ class ENGINE_API UWorldPartitionRuntimeCell : public UObject, public IWorldParti
 
 	static void DirtyStreamingSourceCacheEpoch() { ++UWorldPartitionRuntimeCell::StreamingSourceCacheEpoch; }
 
-	// IWorldPartitionCell Interface begin
+	//~Begin IWorldPartitionCell Interface
 	virtual TArray<const UDataLayerInstance*> GetDataLayerInstances() const override;
 	virtual bool ContainsDataLayer(const UDataLayerAsset* DataLayerAsset) const override;
 	virtual bool ContainsDataLayer(const UDataLayerInstance* DataLayerInstance) const override;
@@ -198,7 +198,8 @@ class ENGINE_API UWorldPartitionRuntimeCell : public UObject, public IWorldParti
 	}
 	virtual const FBox& GetContentBounds() const override { return ContentBounds; }
 	virtual FBox GetCellBounds() const override { return FBox(ForceInit); }
-	// IWorldPartitionCell Interface end
+	virtual FName GetLevelPackageName() const override;
+	//~End IWorldPartitionCell Interface
 
 	const FVector2D GetMinMaxZ() const { return FVector2D(ContentBounds.Min.Z, ContentBounds.Max.Z); }
 
@@ -206,13 +207,17 @@ class ENGINE_API UWorldPartitionRuntimeCell : public UObject, public IWorldParti
 #if WITH_EDITOR
 	bool NeedsActorToCellRemapping() const;
 	void SetBlockOnSlowLoading(bool bInBlockOnSlowLoading) { bBlockOnSlowLoading = bInBlockOnSlowLoading; }
-
 	void SetClientOnlyVisible(bool bInClientOnlyVisible) { bClientOnlyVisible = bInClientOnlyVisible; }
-
 	void SetDataLayers(const TArray<const UDataLayerInstance*>& InDataLayerInstances);
 	void SetContentBundleUID(const FGuid& InContentBundleID) { ContentBundleID = InContentBundleID; }
 	void SetContentBounds(const FBox& InBounds) { ContentBounds = InBounds; }
 	void SetDebugInfo(int64 InCoordX, int64 InCoordY, int64 InCoordZ, FName InGridName);
+	void SetLevelPackageName(const FName& InLevelPackageName) { LevelPackageName = InLevelPackageName; }
+	
+	//~Begin IWorldPartitionCell Interface
+	virtual TSet<FName> GetActorPackageNames() const override { return TSet<FName>(); }
+	//~End IWorldPartitionCell Interface
+
 	virtual void AddActorToCell(const FWorldPartitionActorDescView& ActorDescView, const FActorContainerID& InContainerID, const FTransform& InContainerTransform, const UActorDescContainer* InContainer) PURE_VIRTUAL(UWorldPartitionRuntimeCell::AddActorToCell,);
 	virtual int32 GetActorCount() const PURE_VIRTUAL(UWorldPartitionRuntimeCell::GetActorCount, return 0;);
 
@@ -286,4 +291,8 @@ protected:
 	// Represents the streaming priority relative to other cells
 	float DebugStreamingPriority;
 #endif
+
+#if WITH_EDITOR
+	FName LevelPackageName;
+#endif
 };
diff --git a/Engine/Source/Runtime/Engine/Public/WorldPartition/WorldPartitionRuntimeCellInterface.h b/Engine/Source/Runtime/Engine/Public/WorldPartition/WorldPartitionRuntimeCellInterface.h
--- a/Engine/Source/Runtime/Engine/Public/WorldPartition/WorldPartitionRuntimeCellInterface.h
+++ b/Engine/Source/Runtime/Engine/Public/WorldPartition/WorldPartitionRuntimeCellInterface.h
@@ -25,5 +25,12 @@ public:
 	virtual const FBox& GetContentBounds() const = 0;
 	/** Returns the cell bounds. */
 	virtual FBox GetCellBounds() const = 0;
+	/** Returns the associated level package name. */
+	virtual FName GetLevelPackageName() const = 0;
+
+#if WITH_EDITOR
+	/** Returns the referenced actor packages. */
+	virtual TSet<FName> GetActorPackageNames() const = 0;
+#endif
 };
 
diff --git a/Engine/Source/Runtime/Engine/Public/WorldPartition/WorldPartitionRuntimeHash.h b/Engine/Source/Runtime/Engine/Public/WorldPartition/WorldPartitionRuntimeHash.h
--- a/Engine/Source/Runtime/Engine/Public/WorldPartition/WorldPartitionRuntimeHash.h
+++ b/Engine/Source/Runtime/Engine/Public/WorldPartition/WorldPartitionRuntimeHash.h
@@ -65,6 +65,16 @@ private:
 	TSoftObjectPtr<UWorld> OuterWorld;
 };
 
+struct FWorldPartitionQueryCache
+{
+public:
+	void AddCellInfo(const UWorldPartitionRuntimeCell* Cell, const FSphericalSector& SourceShape);
+	double GetCellMinSquareDist(const UWorldPartitionRuntimeCell* Cell) const;
+
+private:
+	TMap<const UWorldPartitionRuntimeCell*, double> CellToSourceMinSqrDistances;
+};
+
 UCLASS(Abstract, Config=Engine, AutoExpandCategories=(WorldPartition), Within = WorldPartition)
 class ENGINE_API UWorldPartitionRuntimeHash : public UObject, public IWorldPartitionRuntimeCellOwner
 {
@@ -76,6 +86,7 @@ class ENGINE_API UWorldPartitionRuntimeHash : public UObject, public IWorldParti
 	virtual bool PrepareGeneratorPackageForCook(TArray<UPackage*>& OutModifiedPackages) { return false; }
 	virtual bool PopulateGeneratorPackageForCook(const TArray<FWorldPartitionCookPackage*>& PackagesToCook, TArray<UPackage*>& OutModifiedPackages) { return false; }
 	virtual bool PopulateGeneratedPackageForCook(const FWorldPartitionCookPackage& PackagesToCook, TArray<UPackage*>& OutModifiedPackages) { return false; }
+	virtual UWorldPartitionRuntimeCell* GetCellForPackage(const FWorldPartitionCookPackage& PackageToCook) const { return nullptr; }
 	virtual bool GenerateStreaming(class UWorldPartitionStreamingPolicy* StreamingPolicy, const IStreamingGenerationContext* StreamingGenerationContext, TArray<FString>* OutPackagesToGenerate) { return false; }
 	virtual void FlushStreaming() {}
 	virtual bool GenerateHLOD(ISourceControlHelper* SourceControlHelper, const IStreamingGenerationContext* StreamingGenerationContext, bool bCreateActorsOnly) { return false; }
@@ -123,7 +134,7 @@ public:
 
 	// Streaming interface
 	virtual void ForEachStreamingCells(TFunctionRef<bool(const UWorldPartitionRuntimeCell*)> Func) const {}
-	virtual void ForEachStreamingCellsQuery(const FWorldPartitionStreamingQuerySource& QuerySource, TFunctionRef<bool(const UWorldPartitionRuntimeCell*)> Func) const {}
+	virtual void ForEachStreamingCellsQuery(const FWorldPartitionStreamingQuerySource& QuerySource, TFunctionRef<bool(const UWorldPartitionRuntimeCell*)> Func, FWorldPartitionQueryCache* QueryCache = nullptr) const {}
 	virtual void ForEachStreamingCellsSources(const TArray<FWorldPartitionStreamingSource>& Sources, TFunctionRef<bool(const UWorldPartitionRuntimeCell*, EStreamingSourceTargetState)> Func) const {}
 
 	bool IsCellRelevantFor(bool bClientOnlyVisible) const;
diff --git a/Engine/Source/Runtime/Engine/Public/WorldPartition/WorldPartitionRuntimeLevelStreamingCell.h b/Engine/Source/Runtime/Engine/Public/WorldPartition/WorldPartitionRuntimeLevelStreamingCell.h
--- a/Engine/Source/Runtime/Engine/Public/WorldPartition/WorldPartitionRuntimeLevelStreamingCell.h
+++ b/Engine/Source/Runtime/Engine/Public/WorldPartition/WorldPartitionRuntimeLevelStreamingCell.h
@@ -29,6 +29,10 @@ class UWorldPartitionRuntimeLevelStreamingCell : public UWorldPartitionRuntimeSp
 	virtual bool IsLoading() const override;
 	//~End UWorldPartitionRuntimeCell Interface
 
+	//~Begin IWorldPartitionCell Interface
+	FName GetLevelPackageName() const override;
+	//~End IWorldPartitionCell Interface
+
 	virtual void SetStreamingPriority(int32 InStreamingPriority) const override;
 	class UWorldPartitionLevelStreamingDynamic* GetLevelStreaming() const;
 
@@ -51,6 +55,10 @@ class UWorldPartitionRuntimeLevelStreamingCell : public UWorldPartitionRuntimeSp
 	virtual FString GetPackageNameToCreate() const override;
 	//~End UWorldPartitionRuntimeCell Interface
 
+	//~Begin IWorldPartitionCell Interface
+	virtual TSet<FName> GetActorPackageNames() const override;
+	//~End IWorldPartitionCell Interface
+
 	const TArray<FWorldPartitionRuntimeCellObjectMapping>& GetPackages() const { return Packages; }
 #endif
 
diff --git a/Engine/Source/Runtime/Engine/Public/WorldPartition/WorldPartitionRuntimeSpatialHash.h b/Engine/Source/Runtime/Engine/Public/WorldPartition/WorldPartitionRuntimeSpatialHash.h
--- a/Engine/Source/Runtime/Engine/Public/WorldPartition/WorldPartitionRuntimeSpatialHash.h
+++ b/Engine/Source/Runtime/Engine/Public/WorldPartition/WorldPartitionRuntimeSpatialHash.h
@@ -141,7 +141,7 @@ struct FSpatialHashStreamingGrid
 
 	// Used by PIE/Game
 	int64 GetCellSize(int32 Level) const;
-	void GetCells(const FWorldPartitionStreamingQuerySource& QuerySource, TSet<const UWorldPartitionRuntimeCell*>& OutCells, bool bEnableZCulling) const;
+	void GetCells(const FWorldPartitionStreamingQuerySource& QuerySource, TSet<const UWorldPartitionRuntimeCell*>& OutCells, bool bEnableZCulling, FWorldPartitionQueryCache* QueryCache = nullptr) const;
 	void GetCells(const TArray<FWorldPartitionStreamingSource>& Sources, const class UDataLayerSubsystem* DataLayerSubsystem, UWorldPartitionRuntimeHash::FStreamingSourceCells& OutActivateCells, UWorldPartitionRuntimeHash::FStreamingSourceCells& OutLoadCells, bool bEnableZCulling) const;
 	void GetNonSpatiallyLoadedCells(const UDataLayerSubsystem* DataLayerSubsystem, TSet<const UWorldPartitionRuntimeCell*>& OutActivateCells, TSet<const UWorldPartitionRuntimeCell*>& OutLoadCells) const;
 	void Draw2D(UCanvas* Canvas, UWorld* World, const TArray<FWorldPartitionStreamingSource>& Sources, const FBox& Region, const FBox2D& GridScreenBounds, TFunctionRef<FVector2D(const FVector2D&)> WorldToScreen) const;
@@ -273,6 +273,7 @@ public:
 	virtual bool PrepareGeneratorPackageForCook(TArray<UPackage*>& OutModifiedPackages) override;
 	virtual bool PopulateGeneratorPackageForCook(const TArray<FWorldPartitionCookPackage*>& PackagesToCook, TArray<UPackage*>& OutModifiedPackage) override;
 	virtual bool PopulateGeneratedPackageForCook(const FWorldPartitionCookPackage& PackageToCook, TArray<UPackage*>& OutModifiedPackages) override;
+	virtual UWorldPartitionRuntimeCell* GetCellForPackage(const FWorldPartitionCookPackage& PackageToCook) const override;
 	virtual void FlushStreaming() override;
 	virtual bool GenerateHLOD(ISourceControlHelper* SourceControlHelper, const IStreamingGenerationContext* StreamingGenerationContext, bool bCreateActorsOnly) override;
 	virtual void DrawPreview() const override;
@@ -287,7 +288,7 @@ public:
 
 	// streaming interface
 	virtual void ForEachStreamingCells(TFunctionRef<bool(const UWorldPartitionRuntimeCell*)> Func) const override;
-	virtual void ForEachStreamingCellsQuery(const FWorldPartitionStreamingQuerySource& QuerySource, TFunctionRef<bool(const UWorldPartitionRuntimeCell*)> Func) const override;
+	virtual void ForEachStreamingCellsQuery(const FWorldPartitionStreamingQuerySource& QuerySource, TFunctionRef<bool(const UWorldPartitionRuntimeCell*)> Func, FWorldPartitionQueryCache* QueryCache = nullptr) const override;
 	virtual void ForEachStreamingCellsSources(const TArray<FWorldPartitionStreamingSource>& Sources, TFunctionRef<bool(const UWorldPartitionRuntimeCell*, EStreamingSourceTargetState)> Func) const override;
 
 	virtual bool InjectExternalStreamingObject(URuntimeHashExternalStreamingObjectBase* ExternalStreamingObject) override;
diff --git a/Engine/Source/Runtime/Engine/Public/WorldPartition/WorldPartitionRuntimeSpatialHashCell.h b/Engine/Source/Runtime/Engine/Public/WorldPartition/WorldPartitionRuntimeSpatialHashCell.h
--- a/Engine/Source/Runtime/Engine/Public/WorldPartition/WorldPartitionRuntimeSpatialHashCell.h
+++ b/Engine/Source/Runtime/Engine/Public/WorldPartition/WorldPartitionRuntimeSpatialHashCell.h
@@ -20,7 +20,7 @@ class UWorldPartitionRuntimeSpatialHashCell : public UWorldPartitionRuntimeCell
 	virtual void ResetStreamingSourceInfo() const override;
 	virtual void AppendStreamingSourceInfo(const FWorldPartitionStreamingSource& Source, const FSphericalSector& SourceShape) const override;
 	virtual void MergeStreamingSourceInfo() const override;
-	virtual int32 SortCompare(const UWorldPartitionRuntimeCell* InOther) const override;
+	virtual int32 SortCompare(const UWorldPartitionRuntimeCell* InOther, bool bCanUseSortingCache = true) const override;
 	//~End UWorldPartitionRuntimeCell
 
 	//~ Begin IWorldPartitionCell Interface
diff --git a/Engine/Source/Runtime/Engine/Public/WorldPartition/WorldPartitionStreamingPolicy.h b/Engine/Source/Runtime/Engine/Public/WorldPartition/WorldPartitionStreamingPolicy.h
--- a/Engine/Source/Runtime/Engine/Public/WorldPartition/WorldPartitionStreamingPolicy.h
+++ b/Engine/Source/Runtime/Engine/Public/WorldPartition/WorldPartitionStreamingPolicy.h
@@ -45,6 +45,7 @@ class UWorldPartitionStreamingPolicy : public UObject
 	GENERATED_UCLASS_BODY()
 
 public:
+	virtual bool GetIntersectingCells(const TArray<FWorldPartitionStreamingQuerySource>& InSources, TArray<const IWorldPartitionCell*>& OutCells) const;
 	virtual void UpdateStreamingState();
 	virtual bool CanAddLoadedLevelToWorld(class ULevel* InLevel) const;
 	virtual bool DrawRuntimeHash2D(class UCanvas* Canvas, const FVector2D& PartitionCanvasSize, const FVector2D& Offset, FVector2D& OutUsedCanvasSize);
diff --git a/Engine/Source/Runtime/Engine/Public/WorldPartition/WorldPartitionUtils.h b/Engine/Source/Runtime/Engine/Public/WorldPartition/WorldPartitionUtils.h
new file mode 100644
--- /dev/null
+++ b/Engine/Source/Runtime/Engine/Public/WorldPartition/WorldPartitionUtils.h
@@ -0,0 +1,32 @@
+// Copyright Epic Games, Inc. All Rights Reserved.
+
+#if WITH_EDITOR
+
+#include "UObject/GCObject.h"
+
+class UWorld;
+class UWorldPartition;
+class IWorldPartitionCell;
+struct FWorldPartitionStreamingQuerySource;
+
+struct FWorldPartitionUtils
+{
+	class FSimulateCookedSession
+	{
+	public:
+		FSimulateCookedSession(UWorld* InWorld);
+		~FSimulateCookedSession();
+
+		bool IsValid() const { return bIsValid; }
+		bool ForEachStreamingCells(TFunctionRef<void(const IWorldPartitionCell*)> Func);
+		bool GetIntersectingCells(const TArray<FWorldPartitionStreamingQuerySource>& InSources, TArray<const IWorldPartitionCell*>& OutCells);
+
+	private:
+		bool SimulateCook();
+
+		bool bIsValid;
+		TWeakObjectPtr<UWorldPartition> WorldPartition;
+	};
+};
+
+#endif
\ No newline at end of file
