Add IsInParallelGameThread when accessing the game thread shader map because the PSOPrecache collection jobs can access these as well while collect the PSO initializers

diff --git a/Engine/Source/Runtime/Engine/Private/Materials/MaterialShared.cpp b/Engine/Source/Runtime/Engine/Private/Materials/MaterialShared.cpp
--- a/Engine/Source/Runtime/Engine/Private/Materials/MaterialShared.cpp
+++ b/Engine/Source/Runtime/Engine/Private/Materials/MaterialShared.cpp
@@ -3279,7 +3279,7 @@ bool FMaterial::TryGetShaders(const FMaterialShaderTypes& InTypes, const FVertex
 	// TRACE_CPUPROFILER_EVENT_SCOPE(FMaterial::TryGetShaders); // <-- disabled by default due to verbosity (hundreds of calls per frame)
 
 	static const auto* CVarShaderPipelines = IConsoleManager::Get().FindTConsoleVariableDataInt(TEXT("r.ShaderPipelines"));
-	const bool bIsInGameThread = IsInGameThread();
+	const bool bIsInGameThread = (IsInGameThread() || IsInParallelGameThread());
 	const FMaterialShaderMap* ShaderMap = bIsInGameThread ? GameThreadShaderMap : RenderingThreadShaderMap;
 	const bool bShaderMapComplete = bIsInGameThread ? IsGameThreadShaderMapComplete() : IsRenderingThreadShaderMapComplete();
 
diff --git a/Engine/Source/Runtime/Engine/Public/MaterialShared.h b/Engine/Source/Runtime/Engine/Public/MaterialShared.h
--- a/Engine/Source/Runtime/Engine/Public/MaterialShared.h
+++ b/Engine/Source/Runtime/Engine/Public/MaterialShared.h
@@ -2061,7 +2061,7 @@ public:
 
 	class FMaterialShaderMap* GetGameThreadShaderMap() const 
 	{ 
-		checkSlow(IsInGameThread() || IsInAsyncLoadingThread());
+		checkSlow(IsInGameThread() || IsInParallelGameThread() || IsInAsyncLoadingThread());
 		return GameThreadShaderMap; 
 	}
 
@@ -2073,7 +2073,7 @@ public:
 
 	inline bool IsGameThreadShaderMapComplete() const
 	{
-		checkSlow(IsInGameThread());
+		checkSlow(IsInGameThread() || IsInParallelGameThread());
 		return bGameThreadShaderMapIsComplete;
 	}
 
