Fixed CIS non-unity build.

diff --git a/Engine/Source/Runtime/Engine/Public/WorldPartition/WorldPartitionUtils.h b/Engine/Source/Runtime/Engine/Public/WorldPartition/WorldPartitionUtils.h
--- a/Engine/Source/Runtime/Engine/Public/WorldPartition/WorldPartitionUtils.h
+++ b/Engine/Source/Runtime/Engine/Public/WorldPartition/WorldPartitionUtils.h
@@ -2,7 +2,7 @@
 
 #if WITH_EDITOR
 
-#include "UObject/GCObject.h"
+#include "UObject/WeakObjectPtr.h"
 
 class UWorld;
 class UWorldPartition;
