Fix an issue where looping force feedback effects would not be evaluated past their time

#rnx

diff --git a/Engine/Source/Runtime/Engine/Private/GameFramework/ForceFeedbackEffect.cpp b/Engine/Source/Runtime/Engine/Private/GameFramework/ForceFeedbackEffect.cpp
--- a/Engine/Source/Runtime/Engine/Private/GameFramework/ForceFeedbackEffect.cpp
+++ b/Engine/Source/Runtime/Engine/Private/GameFramework/ForceFeedbackEffect.cpp
@@ -182,8 +182,8 @@ bool FActiveForceFeedbackEffect::Update(const float DeltaTime, FForceFeedbackVal
 	{
 		return false;
 	}
-	// Update the effect values if we can
-	if (PlayTime <= EffectDuration)
+	// Update the effect values if we can. Always get the values for a looping effect.
+	if (PlayTime <= EffectDuration || Parameters.bLooping)
 	{
 		GetValues(Values);
 	}	
