Prevent a crash if FMaterial filters out all the shader types (UE-170113).

diff --git a/Engine/Source/Runtime/Engine/Private/Materials/MaterialShader.cpp b/Engine/Source/Runtime/Engine/Private/Materials/MaterialShader.cpp
--- a/Engine/Source/Runtime/Engine/Private/Materials/MaterialShader.cpp
+++ b/Engine/Source/Runtime/Engine/Private/Materials/MaterialShader.cpp
@@ -1742,7 +1742,7 @@ void FMaterialShaderMap::RemoveCompilingDependency(FMaterial* Material)
 	}
 }
 
-void FMaterialShaderMap::SubmitCompileJobs(uint32 CompilingShaderMapId,
+int32 FMaterialShaderMap::SubmitCompileJobs(uint32 CompilingShaderMapId,
 	const FMaterial* Material,
 	const TRefCountPtr<FSharedShaderCompilerEnvironment>& MaterialEnvironment,
 	EShaderCompileJobPriority InPriority) const
@@ -1986,6 +1986,8 @@ void FMaterialShaderMap::SubmitCompileJobs(uint32 CompilingShaderMapId,
 	UE_LOG(LogShaders, Verbose, TEXT("		%u Shaders among %u VertexFactories"), NumShaders, NumVertexFactories);
 
 	GShaderCompilingManager->SubmitJobs(CompileJobs, Material->GetBaseMaterialPathName(), GetDebugDescription());
+
+	return CompileJobs.Num();
 }
 
 /**
@@ -2131,12 +2133,33 @@ void FMaterialShaderMap::Compile(
 			CompilePriority = EShaderCompileJobPriority::Normal;
 		}
 
-		SubmitCompileJobs(CompilingId, Material, MaterialEnvironment, CompilePriority);
+		// Material can filter out all our shader types, essentially preventing the compilation from happening, which can make the shadermap stuck in
+		// "always being compiled" mode. If we find out that we submitted 0 jobs, consider compilation finished.
+		// Note that the shader map can be shared between FMaterials. Here we assume that they all will filter the shader types the same
+		// (and have no good way to check this).
+		if (SubmitCompileJobs(CompilingId, Material, MaterialEnvironment, CompilePriority) == 0)
+		{
+			RemoveCompilingDependency(Material);
+			if (CompilingId == 0u)
+			{
+				bCompilationFinalized = true;
+				bCompiledSuccessfully = true;
+				// create resource code even if it's empty (needed during the serialization and possibly other places).
+				GetResourceCode();
+
+#if WITH_EDITOR
+				if (bIsPersistent)
+				{
+					SaveToDerivedDataCache();
+				}
+#endif
+			}
+		}
 	}
 
 	// Compile the shaders for this shader map now if the material is not deferring and deferred compiles are not enabled globally
 	// If we're early in the startup we can save some time by compiling all special/default materials asynchronously, even if normally they are synchronous
-	if (PrecompileMode == EMaterialShaderPrecompileMode::Synchronous && !PoolSpecialMaterialsCompileJobs())
+	if (CompilingId != 0u && PrecompileMode == EMaterialShaderPrecompileMode::Synchronous && !PoolSpecialMaterialsCompileJobs())
 	{
 		TArray<int32> CurrentShaderMapId;
 		CurrentShaderMapId.Add(CompilingId);
diff --git a/Engine/Source/Runtime/Engine/Public/MaterialShared.h b/Engine/Source/Runtime/Engine/Public/MaterialShared.h
--- a/Engine/Source/Runtime/Engine/Public/MaterialShared.h
+++ b/Engine/Source/Runtime/Engine/Public/MaterialShared.h
@@ -1170,7 +1170,8 @@ public:
 	uint32 GetMaxNumInstructionsForShader(FShaderType* ShaderType) const { return GetContent()->GetMaxNumInstructionsForShader(*this, ShaderType); }
 
 #if WITH_EDITOR
-	void SubmitCompileJobs(uint32 CompilingShaderMapId,
+	/** Submits compile jobs for this shadermap, returns number of jobs submitted. */
+	int32 SubmitCompileJobs(uint32 CompilingShaderMapId,
 		const FMaterial* Material,
 		const TRefCountPtr<FSharedShaderCompilerEnvironment>& MaterialEnvironment,
 		EShaderCompileJobPriority Priority) const;
