Prevented check from triggering in moderation commandlet.

diff --git a/Engine/Source/Runtime/Engine/Private/LevelActor.cpp b/Engine/Source/Runtime/Engine/Private/LevelActor.cpp
--- a/Engine/Source/Runtime/Engine/Private/LevelActor.cpp
+++ b/Engine/Source/Runtime/Engine/Private/LevelActor.cpp
@@ -296,7 +296,7 @@ public:
 		if (MacAddress.Num() != 6)
 		{
 			// During cooking, some platform can't retrieve the MAC adress, so force a default one.
-			check(IsRunningCookCommandlet());
+			check(IsRunningCommandlet());
 			MacAddress.Empty(6);
 			MacAddress.AddZeroed(6);
 		}
