Add "ShowWorldContextPin" metadata to device properties so that you can get the player controller, game instance, etc from within it's BP graphs

#rnx

diff --git a/Engine/Source/Runtime/Engine/Classes/GameFramework/InputDeviceProperties.h b/Engine/Source/Runtime/Engine/Classes/GameFramework/InputDeviceProperties.h
--- a/Engine/Source/Runtime/Engine/Classes/GameFramework/InputDeviceProperties.h
+++ b/Engine/Source/Runtime/Engine/Classes/GameFramework/InputDeviceProperties.h
@@ -28,7 +28,7 @@ class UCurveFloat;
 * support certain device properties. An older gamepad may not have any advanced trigger haptics for 
 * example. 
 */
-UCLASS(Abstract, Blueprintable, BlueprintType, EditInlineNew, CollapseCategories)
+UCLASS(Abstract, Blueprintable, BlueprintType, EditInlineNew, CollapseCategories, meta = (ShowWorldContextPin))
 class ENGINE_API UInputDeviceProperty : public UObject
 {
 	GENERATED_BODY()
