Add accessors to get a Player Controller from a Platform User ID and Input Device ID onto the input device subsystem and blueprint library

#rnx

diff --git a/Engine/Source/Runtime/Engine/Classes/GameFramework/InputDeviceSubsystem.h b/Engine/Source/Runtime/Engine/Classes/GameFramework/InputDeviceSubsystem.h
--- a/Engine/Source/Runtime/Engine/Classes/GameFramework/InputDeviceSubsystem.h
+++ b/Engine/Source/Runtime/Engine/Classes/GameFramework/InputDeviceSubsystem.h
@@ -11,6 +11,7 @@
 ENGINE_API DECLARE_LOG_CATEGORY_EXTERN(LogInputDeviceProperties, Log, All);
 
 class UInputDeviceProperty;
+class APlayerController;
 
 /** A handle to an active input device property that is being used by the InputDeviceSubsytem. */
 USTRUCT(BlueprintType)
@@ -140,6 +141,14 @@ public:
 	virtual void Tick(float DeltaTime) override;
 	//~ End FTickableGameObject interface
 
+	/** Get the player controller who has the given Platform User ID. */
+	UFUNCTION(BlueprintCallable, Category = "Input Devices")
+	static APlayerController* GetPlayerControllerFromPlatformUser(const FPlatformUserId UserId);
+
+	/** Get the player controller who owns the given input device id */
+	UFUNCTION(BlueprintCallable, Category = "Input Devices")
+	static APlayerController* GetPlayerControllerFromInputDevice(const FInputDeviceId DeviceId);
+
 	/** Set the given device property, which will start the evaluation and application of it to the given platform user. */
 	UFUNCTION(BlueprintCallable, Category = "Input Devices", meta = (AutoCreateRefTerm = "Params"))
 	FInputDevicePropertyHandle SetDeviceProperty(const FSetDevicePropertyParams& Params);
diff --git a/Engine/Source/Runtime/Engine/Private/GameFramework/InputDeviceSubsystem.cpp b/Engine/Source/Runtime/Engine/Private/GameFramework/InputDeviceSubsystem.cpp
--- a/Engine/Source/Runtime/Engine/Private/GameFramework/InputDeviceSubsystem.cpp
+++ b/Engine/Source/Runtime/Engine/Private/GameFramework/InputDeviceSubsystem.cpp
@@ -9,6 +9,7 @@
 #include "GenericPlatform/GenericPlatformApplicationMisc.h"	// For FInputDeviceScope
 #include "GenericPlatform/GenericPlatformInputDeviceMapper.h"
 #include "Framework/Application/SlateApplication.h"			// For RegisterInputPreProcessor
+#include "GameFramework/PlayerController.h"
 
 #include UE_INLINE_GENERATED_CPP_BY_NAME(InputDeviceSubsystem)
 
@@ -294,6 +295,39 @@ void UInputDeviceSubsystem::Tick(float DeltaTime)
 	}
 }
 
+APlayerController* UInputDeviceSubsystem::GetPlayerControllerFromPlatformUser(const FPlatformUserId UserId)
+{
+	if (UserId.IsValid())
+	{
+		if (const UInputDeviceSubsystem* System = UInputDeviceSubsystem::Get())
+		{
+			if (const UWorld* World = System->GetTickableGameObjectWorld())
+			{
+				for (FConstPlayerControllerIterator Iterator = World->GetPlayerControllerIterator(); Iterator; ++Iterator)
+				{
+					APlayerController* PlayerController = Iterator->Get();
+					const FPlatformUserId PlayerControllerUserID = PlayerController ? PlayerController->GetPlatformUserId() : PLATFORMUSERID_NONE;
+					if (PlayerControllerUserID.IsValid() && PlayerControllerUserID == UserId)
+					{
+						return PlayerController;
+					}
+				}
+			}
+		}
+	}
+	else
+	{
+		UE_LOG(LogInputDeviceProperties, Warning, TEXT("Attempting to find the player controller for an invalid Platform"));
+	}
+
+	return nullptr;
+}
+
+APlayerController* UInputDeviceSubsystem::GetPlayerControllerFromInputDevice(const FInputDeviceId DeviceId)
+{
+	return GetPlayerControllerFromPlatformUser(IPlatformInputDeviceMapper::Get().GetUserForInputDevice(DeviceId));
+}
+
 FInputDevicePropertyHandle UInputDeviceSubsystem::SetDeviceProperty(const FSetDevicePropertyParams& Params)
 {
 	if (!Params.DevicePropertyClass)
