Crash fix for deref null pointer in audio debug w/ au.debug.soundreverb

#pf 63867bf5f213a116eb4ed992

diff --git a/Engine/Source/Runtime/Engine/Private/AudioDebug.cpp b/Engine/Source/Runtime/Engine/Private/AudioDebug.cpp
--- a/Engine/Source/Runtime/Engine/Private/AudioDebug.cpp
+++ b/Engine/Source/Runtime/Engine/Private/AudioDebug.cpp
@@ -1757,9 +1757,13 @@ namespace Audio
 			else if (ActivatedReverbs.Num() == 1)
 			{
 				auto It = ActivatedReverbs.CreateConstIterator();
-				TheString = FString::Printf(TEXT("  Activated Reverb Effect: %s (Priority: %g Tag: '%s')"), *It.Value().ReverbSettings.ReverbEffect->GetName(), It.Value().Priority, *It.Key().ToString());
-				Canvas->DrawShadowedString(X, Y, *TheString, GetStatsFont(), LinearBodyColor);
-				Y += Height;
+				const FActivatedReverb& ActiveReverb = It.Value();
+				if (ActiveReverb.ReverbSettings.ReverbEffect)
+				{
+					TheString = FString::Printf(TEXT("  Activated Reverb Effect: %s (Priority: %g Tag: '%s')"), *ActiveReverb.ReverbSettings.ReverbEffect->GetName(), ActiveReverb.Priority, *It.Key().ToString());
+					Canvas->DrawShadowedString(X, Y, *TheString, GetStatsFont(), LinearBodyColor);
+					Y += Height;
+				}
 			}
 			else
 			{
@@ -1768,8 +1772,12 @@ namespace Audio
 				TMap<int32, FString> PrioritySortedActivatedReverbs;
 				for (auto It = ActivatedReverbs.CreateConstIterator(); It; ++It)
 				{
-					TheString = FString::Printf(TEXT("    %s (Priority: %g Tag: '%s')"), *It.Value().ReverbSettings.ReverbEffect->GetName(), It.Value().Priority, *It.Key().ToString());
-					PrioritySortedActivatedReverbs.Add(It.Value().Priority, TheString);
+					const FActivatedReverb& ActiveReverb = It.Value();
+					if (ActiveReverb.ReverbSettings.ReverbEffect)
+					{
+						TheString = FString::Printf(TEXT("    %s (Priority: %g Tag: '%s')"), *ActiveReverb.ReverbSettings.ReverbEffect->GetName(), ActiveReverb.Priority, *It.Key().ToString());
+						PrioritySortedActivatedReverbs.Add(ActiveReverb.Priority, TheString);
+					}
 				}
 				for (auto It = PrioritySortedActivatedReverbs.CreateConstIterator(); It; ++It)
 				{
