Fix setting incorrect CVar - submitted by hulken.

#9857

diff --git a/Engine/Source/Runtime/Engine/Private/PacketHandlers/StatelessConnectHandlerComponent.cpp b/Engine/Source/Runtime/Engine/Private/PacketHandlers/StatelessConnectHandlerComponent.cpp
--- a/Engine/Source/Runtime/Engine/Private/PacketHandlers/StatelessConnectHandlerComponent.cpp
+++ b/Engine/Source/Runtime/Engine/Private/PacketHandlers/StatelessConnectHandlerComponent.cpp
@@ -295,7 +295,7 @@ namespace UE::Net
 
 	static FAutoConsoleVariableRef CVarNetVerifyNetClientID(
 		TEXT("net.VerifyNetClientID"),
-		GVerifyNetSessionID,
+		GVerifyNetClientID,
 		TEXT("Whether or not verification of the packet ClientID value is performed."));
 
 	static int32 GVerifyMagicHeader = 0;
