Fixed CIS static analysis.

diff --git a/Engine/Source/Runtime/Engine/Public/WorldPartition/WorldPartitionUtils.h b/Engine/Source/Runtime/Engine/Public/WorldPartition/WorldPartitionUtils.h
--- a/Engine/Source/Runtime/Engine/Public/WorldPartition/WorldPartitionUtils.h
+++ b/Engine/Source/Runtime/Engine/Public/WorldPartition/WorldPartitionUtils.h
@@ -1,5 +1,7 @@
 // Copyright Epic Games, Inc. All Rights Reserved.
 
+#pragma once
+
 #if WITH_EDITOR
 
 #include "UObject/WeakObjectPtr.h"
