- Expose the FInputDevicePropertyHandle::GetTypeHash function with ENGINE_API so that you can use it in plugins.
- Make IsDevicePropertyHandleValid Blueprint Pure

#rnx

diff --git a/Engine/Source/Runtime/Engine/Classes/GameFramework/InputDeviceSubsystem.h b/Engine/Source/Runtime/Engine/Classes/GameFramework/InputDeviceSubsystem.h
--- a/Engine/Source/Runtime/Engine/Classes/GameFramework/InputDeviceSubsystem.h
+++ b/Engine/Source/Runtime/Engine/Classes/GameFramework/InputDeviceSubsystem.h
@@ -30,7 +30,7 @@ struct ENGINE_API FInputDevicePropertyHandle
 	/** An invalid Input Device Property handle */
 	static FInputDevicePropertyHandle InvalidHandle;
 
-	friend uint32 GetTypeHash(const FInputDevicePropertyHandle& InHandle);
+	ENGINE_API friend uint32 GetTypeHash(const FInputDevicePropertyHandle& InHandle);
 
 	bool operator==(const FInputDevicePropertyHandle& Other) const;
 	bool operator!=(const FInputDevicePropertyHandle& Other) const;
@@ -194,7 +194,7 @@ public:
 	void RemoveAllDeviceProperties();
 
 	/** Returns true if the given handle is valid */
-	UFUNCTION(BlueprintCallable, Category = "Input Devices")
+	UFUNCTION(BlueprintCallable, BlueprintPure, Category = "Input Devices")
 	static bool IsDevicePropertyHandleValid(const FInputDevicePropertyHandle& InHandle);
 
 	/** Gets the most recently used hardware input device for the given platform user */
