Fixed preview of regular node when Strata is enabled.
Strata node can still be previewed.

diff --git a/Engine/Source/Runtime/Engine/Private/Materials/MaterialExpressions.cpp b/Engine/Source/Runtime/Engine/Private/Materials/MaterialExpressions.cpp
--- a/Engine/Source/Runtime/Engine/Private/Materials/MaterialExpressions.cpp
+++ b/Engine/Source/Runtime/Engine/Private/Materials/MaterialExpressions.cpp
@@ -2040,13 +2040,26 @@ void UMaterialExpression::ConnectToPreviewMaterial(UMaterial* InMaterial, int32
 {
 	if (InMaterial && OutputIndex >= 0 && OutputIndex < Outputs.Num())
 	{
-		if (IsResultStrataMaterial(0))
+		if (Engine_IsStrataEnabled())
 		{
-			InMaterial->SetShadingModel(MSM_DefaultLit);
-			InMaterial->bUseMaterialAttributes = false;
-			FExpressionInput* MaterialInput = InMaterial->GetExpressionInputForProperty(MP_FrontMaterial);
-			check(MaterialInput);
-			ConnectExpression(MaterialInput, OutputIndex);
+			if (IsResultStrataMaterial(0))
+			{
+				InMaterial->SetShadingModel(MSM_DefaultLit);
+				InMaterial->bUseMaterialAttributes = false;
+				FExpressionInput* MaterialInput = InMaterial->GetExpressionInputForProperty(MP_FrontMaterial);
+				check(MaterialInput);
+				ConnectExpression(MaterialInput, OutputIndex);
+			}
+			else
+			{
+				InMaterial->SetShadingModel(MSM_Unlit);
+				UMaterialExpressionStrataUnlitBSDF* UnlitBSDF = NewObject<UMaterialExpressionStrataUnlitBSDF>(this);
+				UnlitBSDF->EmissiveColor.Connect(OutputIndex, this);
+
+				FExpressionInput* MaterialInput = InMaterial->GetExpressionInputForProperty(MP_FrontMaterial);
+				check(MaterialInput);
+				MaterialInput->Connect(0, UnlitBSDF);
+			}
 		}
 		else if(IsResultMaterialAttributes(0))
 		{
