PR #9860: Fix for static meshes not recreating render states on property changed if async compilation is disabled (Contributed by [REDACTED_PERSON])

#rnx

diff --git a/Engine/Source/Runtime/Engine/Private/StaticMeshBuild.cpp b/Engine/Source/Runtime/Engine/Private/StaticMeshBuild.cpp
--- a/Engine/Source/Runtime/Engine/Private/StaticMeshBuild.cpp
+++ b/Engine/Source/Runtime/Engine/Private/StaticMeshBuild.cpp
@@ -199,9 +199,12 @@ void UStaticMesh::BatchBuild(const TArray<UStaticMesh*>& InStaticMeshes, const F
 					if (Component->IsRenderStateCreated())
 					{
 						Component->DestroyRenderState_Concurrent();
-						StaticMeshComponents[StaticMesh].Add(Component);
 						Scenes.Add(Component->GetScene());
 					}
+					if (Component->IsRegistered())
+					{
+						StaticMeshComponents[StaticMesh].Add(Component);
+					}
 				}
 			}
 		}
@@ -223,7 +226,7 @@ void UStaticMesh::BatchBuild(const TArray<UStaticMesh*>& InStaticMeshes, const F
 				{
 					for (UPrimitiveComponent* Component : *MeshComponents)
 					{
-						if (Component->IsRegistered() && !Component->IsRenderStateCreated())
+						if (Component->IsRegistered() && !Component->IsRenderStateCreated() && Component->ShouldCreateRenderState())
 						{
 							Component->CreateRenderState_Concurrent(nullptr);
 							Scenes.Add(Component->GetScene());
