PR #9714: AnimNodeReference check assertion fix (Contributed by [REDACTED_PERSON])

diff --git a/Engine/Source/Runtime/Engine/Classes/Animation/AnimNodeReference.h b/Engine/Source/Runtime/Engine/Classes/Animation/AnimNodeReference.h
--- a/Engine/Source/Runtime/Engine/Classes/Animation/AnimNodeReference.h
+++ b/Engine/Source/Runtime/Engine/Classes/Animation/AnimNodeReference.h
@@ -51,7 +51,7 @@ public:
 	{
 		check(AnimNodeStruct);
 		check(AnimNode);
-		check(NodeType::StaticStruct()->IsChildOf(AnimNodeStruct));
+		check(AnimNodeStruct->IsChildOf(NodeType::StaticStruct()));
 		return *static_cast<NodeType*>(AnimNode);
 	}
 
