Add options to ignore time dilation and play effects even if the game is paused.

#rnx

diff --git a/Engine/Source/Runtime/Engine/Classes/GameFramework/InputDeviceSubsystem.h b/Engine/Source/Runtime/Engine/Classes/GameFramework/InputDeviceSubsystem.h
--- a/Engine/Source/Runtime/Engine/Classes/GameFramework/InputDeviceSubsystem.h
+++ b/Engine/Source/Runtime/Engine/Classes/GameFramework/InputDeviceSubsystem.h
@@ -66,6 +66,9 @@ struct ENGINE_API FActiveDeviceProperty
 	/** The platform user that is actively receiving this device property */
 	FPlatformUserId PlatformUser = PLATFORMUSERID_NONE;
 
+	/** The handle of this active property. */
+	FInputDevicePropertyHandle PropertyHandle = FInputDevicePropertyHandle::InvalidHandle;
+
 	/**
 	* If true, then when the InputDeviceProperty is done being evaluated and has fulfilled its duration, then 
 	* it will have it's Reset function called. Set this to false if you want your device property to stay 
@@ -73,8 +76,11 @@ struct ENGINE_API FActiveDeviceProperty
 	*/
 	bool bRemoveAfterEvaluationTime = true;
 
-	/** The handle of this active property. */
-	FInputDevicePropertyHandle PropertyHandle = FInputDevicePropertyHandle::InvalidHandle;
+	/** If true, then this device property will ignore dilated delta time and use the Applications delta time instead */
+	bool bIgnoreTimeDilation = false;
+
+	/** If true, then this device property will be played even if the game world is paused. */
+	bool bPlayWhilePaused = false;
 };
 
 /** Parameters for the UInputDeviceSubsystem::SetDeviceProperty function */
@@ -103,6 +109,14 @@ struct ENGINE_API FSetDevicePropertyParams
 	*/
 	UPROPERTY(EditAnywhere, BlueprintReadWrite, Category = "Input Devices")
 	bool bRemoveAfterEvaluationTime = true;
+
+	/** If true, then this device property will ignore dilated delta time and use the Applications delta time instead */
+	UPROPERTY(EditAnywhere, BlueprintReadWrite, Category = "Input Devices")
+	bool bIgnoreTimeDilation = false;
+
+	/** If true, then this device property will be played even if the game world is paused. */
+	UPROPERTY(EditAnywhere, BlueprintReadWrite, Category = "Input Devices")
+	bool bPlayWhilePaused = false;
 };
 
 /**
@@ -138,7 +152,7 @@ public:
 	virtual bool IsAllowedToTick() const override;
 	virtual bool IsTickableInEditor() const override;
 	virtual TStatId GetStatId() const override;
-	virtual void Tick(float DeltaTime) override;
+	virtual void Tick(float InDeltaTime) override;
 	//~ End FTickableGameObject interface
 
 	/** Get the player controller who has the given Platform User ID. */
diff --git a/Engine/Source/Runtime/Engine/Private/GameFramework/InputDeviceSubsystem.cpp b/Engine/Source/Runtime/Engine/Private/GameFramework/InputDeviceSubsystem.cpp
--- a/Engine/Source/Runtime/Engine/Private/GameFramework/InputDeviceSubsystem.cpp
+++ b/Engine/Source/Runtime/Engine/Private/GameFramework/InputDeviceSubsystem.cpp
@@ -10,6 +10,7 @@
 #include "GenericPlatform/GenericPlatformInputDeviceMapper.h"
 #include "Framework/Application/SlateApplication.h"			// For RegisterInputPreProcessor
 #include "GameFramework/PlayerController.h"
+#include "Misc/App.h"	// For FApp::GetDeltaTime()
 
 #include UE_INLINE_GENERATED_CPP_BY_NAME(InputDeviceSubsystem)
 
@@ -263,8 +264,14 @@ TStatId UInputDeviceSubsystem::GetStatId() const
 	RETURN_QUICK_DECLARE_CYCLE_STAT(UInputDeviceSubsystem, STATGROUP_Tickables);
 }
 
-void UInputDeviceSubsystem::Tick(float DeltaTime)
+void UInputDeviceSubsystem::Tick(float InDeltaTime)
 {
+	// If a property doesn't want to be affected by time dilation then we can use this instead
+	const double NonDialatedDeltaTime = FApp::GetDeltaTime();
+
+	const UWorld* World = GetTickableGameObjectWorld();
+	const bool bIsGamePaused = World ? World->IsPaused() : false;
+
 	for (int32 Index = ActiveProperties.Num() - 1; Index >= 0; --Index)
 	{
 		if (!ActiveProperties[Index].Property)
@@ -275,6 +282,15 @@ void UInputDeviceSubsystem::Tick(float DeltaTime)
 			continue;
 		}
 
+		// If the game is paused, only play effects that have explicitly specified that they should be 
+		// played while paused.
+		if (bIsGamePaused && !ActiveProperties[Index].bPlayWhilePaused)
+		{
+			continue;
+		}
+
+		const double DeltaTime = ActiveProperties[Index].bIgnoreTimeDilation ? NonDialatedDeltaTime : static_cast<double>(InDeltaTime);
+		
 		// Increase the evaluated time of this property
 		ActiveProperties[Index].EvaluatedDuration += DeltaTime;
 
@@ -347,8 +363,10 @@ FInputDevicePropertyHandle UInputDeviceSubsystem::SetDeviceProperty(const FSetDe
 		ensure(ActiveProp.Property);
 
 		ActiveProp.PlatformUser = Params.UserId;
-		ActiveProp.bRemoveAfterEvaluationTime = Params.bRemoveAfterEvaluationTime;
 		ActiveProp.PropertyHandle = OutHandle;
+		ActiveProp.bRemoveAfterEvaluationTime = Params.bRemoveAfterEvaluationTime;	
+		ActiveProp.bIgnoreTimeDilation = Params.bIgnoreTimeDilation;
+		ActiveProp.bPlayWhilePaused = Params.bPlayWhilePaused;
 
 		ActiveProperties.Add(ActiveProp);
 	}
