Make the device property stateful functions protected so that you cant call them from outside of the input device subsystem.

#rnx

diff --git a/Engine/Source/Runtime/Engine/Classes/GameFramework/InputDeviceProperties.h b/Engine/Source/Runtime/Engine/Classes/GameFramework/InputDeviceProperties.h
--- a/Engine/Source/Runtime/Engine/Classes/GameFramework/InputDeviceProperties.h
+++ b/Engine/Source/Runtime/Engine/Classes/GameFramework/InputDeviceProperties.h
@@ -31,11 +31,17 @@ class UCurveFloat;
 UCLASS(Abstract, Blueprintable, BlueprintType, EditInlineNew, CollapseCategories, meta = (ShowWorldContextPin))
 class ENGINE_API UInputDeviceProperty : public UObject
 {
+	friend class UInputDeviceSubsystem;
+	// TODO: Make force feedback effects just use the subsystem instead of ticking stuff themselves
+	friend class UForceFeedbackEffect;
+
 	GENERATED_BODY()
 public:
 
 	UInputDeviceProperty(const FObjectInitializer& ObjectInitializer = FObjectInitializer::Get());
 
+protected:
+
 	/**
 	* Evaluate this device property for a given duration. 
 	* If overriding in Blueprints, make sure to call the parent function!
@@ -85,6 +91,8 @@ public:
 	/** Gets a pointer to the current input device property that the IInputInterface can use. */
 	virtual FInputDeviceProperty* GetInternalDeviceProperty() { return nullptr; };
 
+public:
+
 	/**
 	* The duration that this device property should last. Override this if your property has any dynamic curves 
 	* to be the max time range.
