Reset any active device properties when PIE ends

#rnx

diff --git a/Engine/Source/Runtime/Engine/Private/GameFramework/InputDeviceSubsystem.cpp b/Engine/Source/Runtime/Engine/Private/GameFramework/InputDeviceSubsystem.cpp
--- a/Engine/Source/Runtime/Engine/Private/GameFramework/InputDeviceSubsystem.cpp
+++ b/Engine/Source/Runtime/Engine/Private/GameFramework/InputDeviceSubsystem.cpp
@@ -449,6 +449,14 @@ bool UInputDeviceSubsystem::IsDevicePropertyHandleValid(const FInputDeviceProper
 
 void UInputDeviceSubsystem::RemoveAllDeviceProperties()
 {
+	for (FActiveDeviceProperty& ActiveProperty : ActiveProperties)
+	{
+		if (ActiveProperty.Property)
+		{
+			ActiveProperty.Property->ResetDeviceProperty(ActiveProperty.PlatformUser);
+		}		
+	}
+
 	ActiveProperties.Empty();
 }
 
