Use FMalloc instead of TLockFreeFixedSizeAllocator in PageAllocator when ASan is enabled

#rnx

diff --git a/Engine/Source/Runtime/Core/Private/Misc/MemStack.cpp b/Engine/Source/Runtime/Core/Private/Misc/MemStack.cpp
--- a/Engine/Source/Runtime/Core/Private/Misc/MemStack.cpp
+++ b/Engine/Source/Runtime/Core/Private/Misc/MemStack.cpp
@@ -235,16 +235,24 @@ uint64 FPageAllocator::BytesFree()
 
 void *FPageAllocator::Alloc()
 {
+#if USING_ADDRESS_SANITISER
+	return FMemory::Malloc(FPageAllocator::PageSize);
+#else
 	LLM_SCOPE_BY_BOOTSTRAP_TAG(MemStack);
 	void *Result = TheAllocator.Allocate();
 	STAT(UpdateStats());
 	return Result;
+#endif
 }
 
 void FPageAllocator::Free(void *Mem)
 {
+#if USING_ADDRESS_SANITISER
+	FMemory::Free(Mem);
+#else
 	TheAllocator.Free(Mem);
 	STAT(UpdateStats());
+#endif
 }
 
 void *FPageAllocator::AllocSmall()
diff --git a/Engine/Source/Runtime/Engine/Private/UnrealEngine.cpp b/Engine/Source/Runtime/Engine/Private/UnrealEngine.cpp
--- a/Engine/Source/Runtime/Engine/Private/UnrealEngine.cpp
+++ b/Engine/Source/Runtime/Engine/Private/UnrealEngine.cpp
@@ -10521,6 +10521,24 @@ bool UEngine::PerformError(const TCHAR* Cmd, FOutputDevice& Ar)
 		// Reuse deleted array
 		FMemory::Memzero(UseAfterFreeArray, UseAfterFreeArraySize);
 	}
+	else if (FParse::Command(&Cmd, TEXT("USEAFTERFREEMEMSTACK")))
+	{
+		static constexpr size_t MemStackUseAfterFreeArraySize = 4096; 
+		uint8* DanglingMemStackPtr = nullptr;
+		{
+			FMemMark Mark(FMemStack::Get());
+			DanglingMemStackPtr = static_cast<uint8*>(FMemStack::Get().Alloc(MemStackUseAfterFreeArraySize, 16));
+		}
+
+		// Do another Memstack allocation to validate the freed chunk from above has been put in quarantine
+		{
+			FMemMark Mark(FMemStack::Get());
+			FMemStack::Get().Alloc(MemStackUseAfterFreeArraySize, 16);
+		}
+
+		// Reuse deleted array
+		FMemory::Memzero(DanglingMemStackPtr, MemStackUseAfterFreeArraySize);
+	}
 #endif
 #endif // !UE_BUILD_SHIPPING
 	return false;
