Move FMemberReference::ResolveMember to the cpp, groundwork for optimizing it
#rnx

diff --git a/Engine/Source/Runtime/Engine/Classes/Engine/MemberReference.h b/Engine/Source/Runtime/Engine/Classes/Engine/MemberReference.h
--- a/Engine/Source/Runtime/Engine/Classes/Engine/MemberReference.h
+++ b/Engine/Source/Runtime/Engine/Classes/Engine/MemberReference.h
@@ -276,30 +276,6 @@ protected:
 	/** Only intended for backwards compat! */
 	ENGINE_API void SetGivenSelfScope(const FName InMemberName, const FGuid InMemberGuid, TSubclassOf<class UObject> InMemberParentClass, TSubclassOf<class UObject> SelfScope) const;
 #endif
-
-	template<class TFieldType>
-	TFieldType* ResolveUFunction() const
-	{
-		return nullptr;
-	}
-
-	template<class TFieldType>
-	TFieldType* ResolveUField(FFieldClass* InClass, UPackage* TargetPackage) const
-	{
-		return nullptr;
-	}
-	template<class TFieldType>
-	TFieldType* ResolveUField(UClass* InClass, UPackage* TargetPackage) const
-	{
-		return FindObject<TFieldType>(TargetPackage, *MemberName.ToString());;
-	}
-
-	template<class TFieldType>
-	UObject* GetFieldOuter(TFieldType Field) const
-	{
-		return Field->GetOuter();
-	}
-
 public:
 
 	/** Get the class that owns this member */
@@ -343,12 +319,6 @@ public:
 	{
 		return bSelfContext ? SelfScope : GetMemberParentClass();
 	}
-	
-	template <class TFieldClassA, class TFieldClassB>
-	bool CompareClassesHelper(const TFieldClassA* ClassA, const TFieldClassB* ClassB) const
-	{
-		return ClassA == ClassB;
-	}
 
 	/** 
 	 *	Returns the member FProperty/UFunction this reference is pointing to, or NULL if it no longer exists 
@@ -358,152 +328,15 @@ public:
 	template<class TFieldType>
 	TFieldType* ResolveMember(UClass* SelfScope = nullptr, const bool bAlwaysFollowRedirects = false) const
 	{
-		TFieldType* ReturnField = nullptr;
-
-		if(bSelfContext && SelfScope == nullptr)
-		{
-			UE_LOG(LogBlueprint, Warning, TEXT("FMemberReference::ResolveMember (%s) bSelfContext == true, but no scope supplied!"), *MemberName.ToString() );
-		}
-
-		// Check if the member reference is function scoped
-		if(IsLocalScope())
-		{
-			UStruct* MemberScopeStruct = FindUField<UStruct>(SelfScope, *MemberScope);
-
-			// Find in target scope
-			ReturnField = FindUFieldOrFProperty<TFieldType>(MemberScopeStruct, MemberName, EFieldIterationFlags::IncludeAll);
-
-#if WITH_EDITOR
-			if(ReturnField == nullptr)
-			{
-				// If the property was not found, refresh the local variable name and try again
-				const FName RenamedMemberName = RefreshLocalVariableName(SelfScope);
-				if (RenamedMemberName != NAME_None)
-				{
-					ReturnField = FindUFieldOrFProperty<TFieldType>(MemberScopeStruct, MemberName, EFieldIterationFlags::IncludeAll);
-				}
-			}
-#endif
-		}
-		else
-		{
-#if WITH_EDITOR
-			const bool bCanFollowRedirects = !GIsSavingPackage;
-#else
-			const bool bCanFollowRedirects = bAlwaysFollowRedirects;
-#endif
-
-			// Look for remapped member
-			UClass* TargetScope = GetScope(SelfScope);
-			if (bCanFollowRedirects && TargetScope)
-			{
-				ReturnField = FindRemappedField<TFieldType>(TargetScope, MemberName, true);
-			}
-
-			if(ReturnField != nullptr)
-			{
-				// Fix up this struct, we found a redirect
-				MemberName = ReturnField->GetFName();
-				MemberParent = Cast<UClass>(GetFieldOuter(static_cast<typename TFieldType::BaseFieldClass*>(ReturnField)));
-
-				MemberGuid.Invalidate();
-				UBlueprint::GetGuidFromClassByFieldName<TFieldType>(TargetScope, MemberName, MemberGuid);
-
-				if (UClass* ParentAsClass = GetMemberParentClass())
-				{
-#if WITH_EDITOR
-					ParentAsClass = ParentAsClass->GetAuthoritativeClass();
-#endif
-					MemberParent  = ParentAsClass;
-
-					// Re-evaluate self-ness against the redirect if we were given a valid SelfScope
-					// For functions and multicast delegates we don't want to go from not-self to self as the target pin type should remain consistent
-					if (SelfScope != nullptr && (bSelfContext || (!CompareClassesHelper(TFieldType::StaticClass(), UFunction::StaticClass()) && !CompareClassesHelper(TFieldType::StaticClass(), FMulticastDelegateProperty::StaticClass()))))
-					{
-#if WITH_EDITOR
-						bSelfContext = SelfScope->IsChildOf(ParentAsClass) || SelfScope->ClassGeneratedBy == ParentAsClass->ClassGeneratedBy;
-#else
-						bSelfContext = SelfScope->IsChildOf(ParentAsClass);
-#endif
-
-						if (bSelfContext)
-						{
-							MemberParent = nullptr;
-						}
-					}
-				}	
-			}
-			else if (TargetScope != nullptr)
-			{
-#if WITH_EDITOR
-				bool bUseUpToDateClass = SelfScope && SelfScope->GetAuthoritativeClass() != SelfScope;
-				TargetScope = GetClassToUse(TargetScope, bUseUpToDateClass);
-				if (TargetScope)
-#endif
-				{
-					// Find in target scope or in the sparse class data
-					ReturnField = FindUFieldOrFProperty<TFieldType>(TargetScope, MemberName, EFieldIterationFlags::IncludeAll);
-					if (!ReturnField)
-					{
-						if (UScriptStruct* SparseClassDataStruct = TargetScope->GetSparseClassDataStruct())
-						{
-							ReturnField = FindUFieldOrFProperty<TFieldType>(SparseClassDataStruct, MemberName, EFieldIterationFlags::IncludeAll);
-						}
-					}
-				}
-
-
-#if !WITH_EDITOR
-				if (bAlwaysFollowRedirects)
-#endif
-				{
-					// If the reference variable is valid we need to make sure that our GUID matches
-					if (ReturnField != nullptr)
-					{
-						UBlueprint::GetGuidFromClassByFieldName<TFieldType>(TargetScope, MemberName, MemberGuid);
-					}
-					// If we have a GUID find the reference variable and make sure the name is up to date and find the field again
-					// For now only variable references will have valid GUIDs.  Will have to deal with finding other names subsequently
-					else if (MemberGuid.IsValid())
-					{
-						const FName RenamedMemberName = UBlueprint::GetFieldNameFromClassByGuid<TFieldType>(TargetScope, MemberGuid);
-						if (RenamedMemberName != NAME_None)
-						{
-							MemberName = RenamedMemberName;
-							ReturnField = FindUFieldOrFProperty<TFieldType>(TargetScope, MemberName, EFieldIterationFlags::IncludeAll);
-						}
-					}
-				}
-			}
-			else if (UPackage* TargetPackage = GetMemberParentPackage())
-			{
-				ReturnField = ResolveUField<TFieldType>(TFieldType::StaticClass(), TargetPackage);
-			}
-			// For backwards compatibility: as of CL 2412156, delegate signatures 
-			// could have had a null MemberParentClass (for those natively 
-			// declared outside of a class), we used to rely on the following 
-			// FindObject<>; however this was not reliable (hence the addition 
-			// of GetMemberParentPackage(), etc.)
-			else if (MemberName.ToString().EndsWith(HEADER_GENERATED_DELEGATE_SIGNATURE_SUFFIX))
-			{
-				ReturnField = ResolveUFunction<TFieldType>();
-				if (ReturnField != nullptr)
-				{
-					UE_LOG(LogBlueprint, Display, TEXT("Generic delegate signature ref (%s). Explicitly setting it to: '%s'. Make sure this is correct (there could be multiple native delegate types with this name)."), *MemberName.ToString(), *ReturnField->GetPathName());
-					MemberParent = ReturnField->GetOutermost();
-				}
-			}
-		}
-
-		// Check to see if the member has been deprecated
-		if (FProperty* Property = FFieldVariant(ReturnField).Get<FProperty>())
-		{
-			bWasDeprecated = Property->HasAnyPropertyFlags(CPF_Deprecated);
-		}
-
-		return ReturnField;
+		return static_cast<TFieldType*>(ResolveMemberProperty(SelfScope, bAlwaysFollowRedirects, TFieldType::StaticClass()));
+	}
+	template<>
+	UFunction* ResolveMember(UClass* SelfScope, const bool bAlwaysFollowRedirects) const
+	{
+		return ResolveMemberFunction(SelfScope, bAlwaysFollowRedirects);
 	}
 
+	/** ResolveMember overload for UBlueprint, uses the skeleton class so it can be more up to date */
 	template<class TFieldType>
 	TFieldType* ResolveMember(UBlueprint* SelfScope)
 	{
@@ -541,6 +374,11 @@ protected:
 	/** @return the 'real' generated class for blueprint classes, but only if we're already passed through CompileClassLayout */
 	ENGINE_API static UClass* GetClassToUse(UClass* InClass, bool bUseUpToDateClass);
 #endif
+	template<class TFieldType, class TFieldTypeClass>
+	TFieldType* ResolveMemberImpl(UClass* SelfScope, TFieldTypeClass* FieldClas, const bool bAlwaysFollowRedirects, const bool bIsUFunctionOrMulticastDelegate) const;
+
+	ENGINE_API FProperty* ResolveMemberProperty(UClass* SelfScope, const bool bAlwaysFollowRedirects, FFieldClass* FieldClass) const;
+	ENGINE_API UFunction* ResolveMemberFunction(UClass* SelfScope, const bool bAlwaysFollowRedirects) const;
 
 public:
 	template<class TFieldType>
@@ -580,39 +418,3 @@ public:
 	}
 };
 
-
-template<>
-inline UFunction* FMemberReference::ResolveUFunction() const
-{
-	UFunction* ReturnField = nullptr;
-	FString const StringName = MemberName.ToString();
-	for (TObjectIterator<UPackage> PackageIt; PackageIt && (ReturnField == nullptr); ++PackageIt)
-	{
-		if (PackageIt->HasAnyPackageFlags(PKG_CompiledIn) == false)
-		{
-			continue;
-		}
-
-		// NOTE: this could return the wrong field (if there are 
-		//       two like-named delegates defined in separate packages)
-		ReturnField = FindObject<UFunction>(*PackageIt, *StringName);
-	}
-	return ReturnField;
-}
-
-template<>
-inline UObject* FMemberReference::GetFieldOuter(FField* Field) const
-{
-	return Field->GetOwner<UObject>();
-}
-
-template <>
-inline bool FMemberReference::CompareClassesHelper(const UClass* ClassA, const FFieldClass* ClassB) const
-{
-	return false;
-}
-template <>
-inline bool FMemberReference::CompareClassesHelper(const FFieldClass* ClassA, const UClass* ClassB) const
-{
-	return false;
-}
\ No newline at end of file
diff --git a/Engine/Source/Runtime/Engine/Private/MemberReference.cpp b/Engine/Source/Runtime/Engine/Private/MemberReference.cpp
--- a/Engine/Source/Runtime/Engine/Private/MemberReference.cpp
+++ b/Engine/Source/Runtime/Engine/Private/MemberReference.cpp
@@ -279,6 +279,214 @@ UClass* FMemberReference::GetClassToUse(UClass* InClass, bool bUseUpToDateClass)
 }
 
 #endif
+template<class TFieldType>
+TFieldType* ResolveUFunctionImpl(FName MemberName)
+{
+	return nullptr;
+}
+
+template<>
+inline UFunction* ResolveUFunctionImpl(FName MemberName)
+{
+	UFunction* ReturnField = nullptr;
+	FString const StringName = MemberName.ToString();
+	for (TObjectIterator<UPackage> PackageIt; PackageIt && (ReturnField == nullptr); ++PackageIt)
+	{
+		if (PackageIt->HasAnyPackageFlags(PKG_CompiledIn) == false)
+		{
+			continue;
+		}
+
+		// NOTE: this could return the wrong field (if there are 
+		//       two like-named delegates defined in separate packages)
+		ReturnField = FindObject<UFunction>(*PackageIt, *StringName);
+	}
+	return ReturnField;
+}
+
+template<class TFieldType>
+TFieldType* ResolveUField(FFieldClass* InClass, UPackage* TargetPackage, FName MemberName)
+{
+	return nullptr;
+}
+template<class TFieldType>
+TFieldType* ResolveUField(UClass* InClass, UPackage* TargetPackage, FName MemberName) 
+{
+	return FindObject<TFieldType>(TargetPackage, *MemberName.ToString());;
+}
+
+template<class TFieldType>
+UObject* GetFieldOuter(TFieldType Field)
+{
+	return Field->GetOuter();
+}
+
+template<>
+inline UObject* GetFieldOuter(FField* Field)
+{
+	return Field->GetOwner<UObject>();
+}
+
+template<class TFieldType, class TFieldTypeClass>
+TFieldType* FMemberReference::ResolveMemberImpl(UClass* SelfScope, TFieldTypeClass* FieldClass, const bool bAlwaysFollowRedirects, const bool bIsUFunctionOrMulticastDelegate) const
+{
+	TFieldType* ReturnField = nullptr;
+
+	if (bSelfContext && SelfScope == nullptr)
+	{
+		UE_LOG(LogBlueprint, Warning, TEXT("FMemberReference::ResolveMember (%s) bSelfContext == true, but no scope supplied!"), *MemberName.ToString());
+	}
+
+	// Check if the member reference is function scoped
+	if (IsLocalScope())
+	{
+		UStruct* MemberScopeStruct = FindUField<UStruct>(SelfScope, *MemberScope);
+
+		// Find in target scope
+		ReturnField = FindUFieldOrFProperty<TFieldType>(MemberScopeStruct, MemberName, EFieldIterationFlags::IncludeAll);
+
+#if WITH_EDITOR
+		if (ReturnField == nullptr)
+		{
+			// If the property was not found, refresh the local variable name and try again
+			const FName RenamedMemberName = RefreshLocalVariableName(SelfScope);
+			if (RenamedMemberName != NAME_None)
+			{
+				ReturnField = FindUFieldOrFProperty<TFieldType>(MemberScopeStruct, MemberName, EFieldIterationFlags::IncludeAll);
+			}
+		}
+#endif
+	}
+	else
+	{
+#if WITH_EDITOR
+		const bool bCanFollowRedirects = !GIsSavingPackage;
+#else
+		const bool bCanFollowRedirects = bAlwaysFollowRedirects;
+#endif
+
+		// Look for remapped member
+		UClass* TargetScope = GetScope(SelfScope);
+		if (bCanFollowRedirects && TargetScope)
+		{
+			ReturnField = static_cast<TFieldType*>(FindRemappedField(FieldClass, TargetScope, MemberName, true));
+		}
+
+		if (ReturnField != nullptr)
+		{
+			// Fix up this struct, we found a redirect
+			MemberName = ReturnField->GetFName();
+			MemberParent = Cast<UClass>(GetFieldOuter(static_cast<typename TFieldType::BaseFieldClass*>(ReturnField)));
+
+			MemberGuid.Invalidate();
+			UBlueprint::GetGuidFromClassByFieldName<TFieldType>(TargetScope, MemberName, MemberGuid);
+
+			if (UClass* ParentAsClass = GetMemberParentClass())
+			{
+#if WITH_EDITOR
+				ParentAsClass = ParentAsClass->GetAuthoritativeClass();
+#endif
+				MemberParent = ParentAsClass;
+
+				// Re-evaluate self-ness against the redirect if we were given a valid SelfScope
+				// For functions and multicast delegates we don't want to go from not-self to self as the target pin type should remain consistent
+				if (SelfScope != nullptr && 
+					(bSelfContext || !bIsUFunctionOrMulticastDelegate))
+				{
+#if WITH_EDITOR
+					bSelfContext = SelfScope->IsChildOf(ParentAsClass) || SelfScope->ClassGeneratedBy == ParentAsClass->ClassGeneratedBy;
+#else
+					bSelfContext = SelfScope->IsChildOf(ParentAsClass);
+#endif
+
+					if (bSelfContext)
+					{
+						MemberParent = nullptr;
+					}
+				}
+			}
+		}
+		else if (TargetScope != nullptr)
+		{
+#if WITH_EDITOR
+			bool bUseUpToDateClass = SelfScope && SelfScope->GetAuthoritativeClass() != SelfScope;
+			TargetScope = GetClassToUse(TargetScope, bUseUpToDateClass);
+			if (TargetScope)
+#endif
+			{
+				// Find in target scope or in the sparse class data
+				ReturnField = FindUFieldOrFProperty<TFieldType>(TargetScope, MemberName, EFieldIterationFlags::IncludeAll);
+				if (!ReturnField)
+				{
+					if (UScriptStruct* SparseClassDataStruct = TargetScope->GetSparseClassDataStruct())
+					{
+						ReturnField = FindUFieldOrFProperty<TFieldType>(SparseClassDataStruct, MemberName, EFieldIterationFlags::IncludeAll);
+					}
+				}
+			}
+
+
+#if !WITH_EDITOR
+			if (bAlwaysFollowRedirects)
+#endif
+			{
+				// If the reference variable is valid we need to make sure that our GUID matches
+				if (ReturnField != nullptr)
+				{
+					UBlueprint::GetGuidFromClassByFieldName<TFieldType>(TargetScope, MemberName, MemberGuid);
+				}
+				// If we have a GUID find the reference variable and make sure the name is up to date and find the field again
+				// For now only variable references will have valid GUIDs.  Will have to deal with finding other names subsequently
+				else if (MemberGuid.IsValid())
+				{
+					const FName RenamedMemberName = UBlueprint::GetFieldNameFromClassByGuid<TFieldType>(TargetScope, MemberGuid);
+					if (RenamedMemberName != NAME_None)
+					{
+						MemberName = RenamedMemberName;
+						ReturnField = FindUFieldOrFProperty<TFieldType>(TargetScope, MemberName, EFieldIterationFlags::IncludeAll);
+					}
+				}
+			}
+		}
+		else if (UPackage* TargetPackage = GetMemberParentPackage())
+		{
+			ReturnField = ResolveUField<TFieldType>(FieldClass, TargetPackage, MemberName);
+		}
+		// For backwards compatibility: as of CL 2412156, delegate signatures 
+		// could have had a null MemberParentClass (for those natively 
+		// declared outside of a class), we used to rely on the following 
+		// FindObject<>; however this was not reliable (hence the addition 
+		// of GetMemberParentPackage(), etc.)
+		else if (MemberName.ToString().EndsWith(HEADER_GENERATED_DELEGATE_SIGNATURE_SUFFIX))
+		{
+			ReturnField = ResolveUFunctionImpl<TFieldType>(MemberName);
+			if (ReturnField != nullptr)
+			{
+				UE_LOG(LogBlueprint, Display, TEXT("Generic delegate signature ref (%s). Explicitly setting it to: '%s'. Make sure this is correct (there could be multiple native delegate types with this name)."), *MemberName.ToString(), *ReturnField->GetPathName());
+				MemberParent = ReturnField->GetOutermost();
+			}
+		}
+	}
+
+	// Check to see if the member has been deprecated
+	if (FProperty* Property = FFieldVariant(ReturnField).Get<FProperty>())
+	{
+		bWasDeprecated = Property->HasAnyPropertyFlags(CPF_Deprecated);
+	}
+
+	return ReturnField;
+}
+
+FProperty* FMemberReference::ResolveMemberProperty(UClass* SelfScope, const bool bAlwaysFollowRedirects, FFieldClass* FieldClass) const
+{
+	return ResolveMemberImpl<FProperty, FFieldClass>(SelfScope, FieldClass, bAlwaysFollowRedirects, FieldClass == FMulticastDelegateProperty::StaticClass());
+}
+
+UFunction* FMemberReference::ResolveMemberFunction(UClass* SelfScope, const bool bAlwaysFollowRedirects) const
+{
+	const bool bIsAUFunction = true;
+	return ResolveMemberImpl<UFunction, UClass>(SelfScope, UFunction::StaticClass(), bAlwaysFollowRedirects, bIsAUFunction);
+}
 
 template <typename TFieldType>
 TFieldType* FindRemappedFieldImpl(FName FieldClassOutermostName, FName FieldClassName, UClass* InitialScope, FName InitialName, bool bInitialScopeMustBeOwnerOfField)
