NetConnection:
* Deprecated the DormantReplicatorMap and replaced it with the private DormantReplicatorSet
* The new structs holds dormant object replicators grouped by their owning actor.
* Current public usage of the dormant map have been replaced with equivalent functions in UNetConnection

diff --git a/Engine/Source/Runtime/Engine/Classes/Engine/NetConnection.h b/Engine/Source/Runtime/Engine/Classes/Engine/NetConnection.h
--- a/Engine/Source/Runtime/Engine/Classes/Engine/NetConnection.h
+++ b/Engine/Source/Runtime/Engine/Classes/Engine/NetConnection.h
@@ -35,6 +35,7 @@
 #include "Net/NetAnalyticsTypes.h"
 #include "Net/Core/Connection/NetCloseResult.h"
 #include "Net/TrafficControl.h"
+#include "Net/NetDormantHolder.h"
 #include "HAL/LowLevelMemTracker.h"
 #include "GameFramework/Actor.h"
 
@@ -58,7 +59,8 @@ namespace UE::Net
 {
 	class FNetPing;
 	class FNetConnectionFaultRecovery;
-}
+
+} // end namespace UE::Net
 
 typedef TMap<TWeakObjectPtr<AActor>, UActorChannel*, FDefaultSetAllocator, TWeakObjectPtrMapKeyFuncs<TWeakObjectPtr<AActor>, UActorChannel*>> FActorChannelMap;
 
@@ -781,7 +783,14 @@ public:
 	TMap<FNetworkGUID, TArray<class UActorChannel*>> KeepProcessingActorChannelBunchesMap;
 
 	/** A list of replicators that belong to recently dormant actors/objects */
-	TMap<FObjectKey, TSharedRef<FObjectReplicator>> DormantReplicatorMap;	
+	UE_DEPRECATED(5.2, "The DormantReplicatorMap is deprecated in favor of the private DormantReplicatorSet.")
+	TMap<FObjectKey, TSharedRef<FObjectReplicator>> DormantReplicatorMap;
+
+private:
+
+	UE::Net::Private::FDormantReplicatorHolder DormantReplicatorSet;
+
+public:
 
 	ENGINE_API FName GetClientWorldPackageName() const { return ClientWorldPackageName; }
 
@@ -1198,7 +1207,16 @@ public:
 	ENGINE_API void ForcePropertyCompare( AActor* Actor );
 
 	/** Wrapper for validating an objects dormancy state, and to prepare the object for replication again */
-	void FlushDormancyForObject( UObject* Object );
+	UE_DEPRECATED(5.2, "FlushDormancyForObject has been replaced with a version that needs to receive the dormant actor.")
+	void FlushDormancyForObject( UObject* Object ) {}
+
+	/**
+	* Validate an objects dormancy state and prepare the object for replication again
+	* 
+	* @param DormantActor The dormant actor that owns the replicated object that needs to be flushed
+	* @param ReplicatedObject The replicated object that is flushed.
+	*/
+	void FlushDormancyForObject(AActor* DormantActor, UObject* ReplicatedObject);
 
 	/** 
 	 * Wrapper for setting the current client login state, so we can trap for debugging, and verbosity purposes. 
@@ -1290,14 +1308,32 @@ public:
 	TArray<FOutBunch *>& GetOutgoingBunches() { return OutgoingBunches; }
 
 	/** Add a replicator to the dormancy map and release its strong pointer to its object */
-	void AddDormantReplicator(UObject* Object, const TSharedRef<FObjectReplicator>& Replicator);
+	UE_DEPRECATED(5.2, "AddDormantReplicator has been replaced by StoreDormantReplicator and will be removed soon.")
+	void AddDormantReplicator(UObject* Object, const TSharedRef<FObjectReplicator>& Replicator) {}
+
+	/** Store a replicator to the dormancy map and release its strong pointer to its object */
+	void StoreDormantReplicator(AActor* OwnerActor, UObject* Object, const TSharedRef<FObjectReplicator>& ObjectReplicator);
 	
 	/** Find a dormant replicator for the channel actor or one of its subobjects. Removes it from the map if found. */
-	TSharedPtr<FObjectReplicator> FindAndRemoveDormantReplicator(UObject* Object);
+	UE_DEPRECATED(5.2, "FindAndRemoveDormantReplicator is deprecated. Use the new version that needs to receive the owning actor.")
+	TSharedPtr<FObjectReplicator> FindAndRemoveDormantReplicator(UObject* Object) { return {}; }
+
+	/** Find a dormant replicator for the channel actor or one of its subobjects. Removes it from the map if found. */
+	TSharedPtr<FObjectReplicator> FindAndRemoveDormantReplicator(AActor* OwnerActor, UObject* Object);
+
+	/** Remove any reference to the dormant object replicator */
+	void RemoveDormantReplicator(AActor* Actor, UObject* Object);
 
+	/** Remove the reference of all dormant object replicators owned by an actor */
 	void CleanupDormantReplicatorsForActor(AActor* Actor);
 
-	/** Removes stale entries from DormantReplicatorMap. */
+	/** Trigger a callback on all dormant replicators of every dormant actor we stored */
+	void ExecuteOnAllDormantReplicators(UE::Net::FExecuteForEachDormantReplicator ExecuteFunction);
+
+	/** Trigger a callback on all dormant replicators owned by a dormant actor */
+	void ExecuteOnAllDormantReplicatorsOfActor(AActor* OwnerActor, UE::Net::FExecuteForEachDormantReplicator ExecuteFunction);
+
+	/** Removes dormant object replicators from objects now invalid. */
 	void CleanupStaleDormantReplicators();
 
 	/** Called before Driver.TickDispatch processes received packets */
diff --git a/Engine/Source/Runtime/Engine/Private/DataChannel.cpp b/Engine/Source/Runtime/Engine/Private/DataChannel.cpp
--- a/Engine/Source/Runtime/Engine/Private/DataChannel.cpp
+++ b/Engine/Source/Runtime/Engine/Private/DataChannel.cpp
@@ -2275,7 +2275,7 @@ void UActorChannel::CleanupReplicators(const bool bKeepReplicators)
 			//		KeepProcessingActorChannelBunchesMap will get in here, then when the channel closes a second time, we'll hit this assert
 			//		It should be okay to just set the most recent replicator
 			//check( Connection->DormantReplicatorMap.Find( CompIt.Value()->GetObject() ) == NULL );
-			Connection->AddDormantReplicator(ObjectReplicatorRef->GetObject(), ObjectReplicatorRef);
+			Connection->StoreDormantReplicator(Actor, ObjectReplicatorRef->GetObject(), ObjectReplicatorRef);
 			ObjectReplicatorRef->StopReplicating(this);		// Stop replicating on this channel
 		}
 		else
@@ -4951,23 +4951,25 @@ TSharedRef<FObjectReplicator>& UActorChannel::CreateReplicator(UObject* Obj, boo
 	TSharedPtr<FObjectReplicator> NewReplicator;
 	if (bCheckDormantReplicators)
 	{
-		NewReplicator = Connection->FindAndRemoveDormantReplicator(Obj);
+		NewReplicator = Connection->FindAndRemoveDormantReplicator(Actor, Obj);
 	}
 	else
 	{
-		Connection->FindAndRemoveDormantReplicator(Obj); // Discard dormant replicator if there was one that we're not using;
+		// Discard dormant replicator if there was one that we're not using
+		Connection->RemoveDormantReplicator(Actor, Obj);
 	}
 
 	// Check if we found it and that it is has a valid object
 	if(NewReplicator.IsValid())
 	{
-		UE_LOG( LogNetTraffic, Log, TEXT( "Found existing replicator for %s" ), *Obj->GetName() );
+		UE_LOG(LogNetTraffic, Log, TEXT("CreateReplicator - Found existing replicator for (0x%p) %s - Replicator (0x%p)"), Obj , *Obj->GetName(), NewReplicator.Get());
 	}
 	else
 	{
 		// Didn't find one, need to create
-		UE_LOG( LogNetTraffic, Log, TEXT( "Creating Replicator for %s" ), *Obj->GetName() );
 		NewReplicator = Connection->CreateReplicatorForNewActorChannel(Obj);
+
+		UE_LOG(LogNetTraffic, Log, TEXT("CreateReplicator - creating new replicator for (0x%p) %s - Replicator (0x%p)"), Obj, *Obj->GetName(), NewReplicator.Get());
 	}
 
 	// Add to the replication map
diff --git a/Engine/Source/Runtime/Engine/Private/DemoNetDriver.cpp b/Engine/Source/Runtime/Engine/Private/DemoNetDriver.cpp
--- a/Engine/Source/Runtime/Engine/Private/DemoNetDriver.cpp
+++ b/Engine/Source/Runtime/Engine/Private/DemoNetDriver.cpp
@@ -3017,10 +3017,12 @@ void UDemoNetDriver::FinalizeFastForward(const double StartTime)
 			}
 		}
 
-		for (auto& DormantPair : ServerConnection->DormantReplicatorMap)
+		UE::Net::FExecuteForEachDormantReplicator CallRepNotifies = [](AActor* OwnerActor, FObjectKey ObjectKey, const TSharedRef<FObjectReplicator>& ReplicatorRef)
 		{
-			DormantPair.Value->CallRepNotifies(true);
-		}
+			ReplicatorRef->CallRepNotifies(true);
+		};
+
+		ServerConnection->ExecuteOnAllDormantReplicators(CallRepNotifies);
 	}
 
 	// We may have been fast-forwarding immediately after loading a checkpoint
@@ -3918,10 +3920,12 @@ bool UDemoNetDriver::FastForwardLevels(const FGotoResult& GotoResult)
 			}
 		}
 
-		for (auto& DormantPair : ServerConnection->DormantReplicatorMap)
+		UE::Net::FExecuteForEachDormantReplicator CallRepNotifies = [](AActor* OwnerActor, FObjectKey ObjectKey, const TSharedRef<FObjectReplicator>& ReplicatorRef)
 		{
-			DormantPair.Value->CallRepNotifies( true );
-		}
+			ReplicatorRef->CallRepNotifies(true);
+		};
+
+		ServerConnection->ExecuteOnAllDormantReplicators(CallRepNotifies);
 	}
 
 	return true;
@@ -5096,7 +5100,7 @@ void UDemoNetDriver::NotifyActorDestroyed(AActor* Actor, bool IsSeamlessTravel)
 					Channel->ReleaseReferences(false);
 				}
 
-				Connection->DormantReplicatorMap.Remove(Actor);
+				Connection->CleanupDormantReplicatorsForActor(Actor);
 			}
 
 			GetNetworkObjectList().Remove(Actor);
diff --git a/Engine/Source/Runtime/Engine/Private/Net/NetDormantHolder.cpp b/Engine/Source/Runtime/Engine/Private/Net/NetDormantHolder.cpp
new file mode 100644
--- /dev/null
+++ b/Engine/Source/Runtime/Engine/Private/Net/NetDormantHolder.cpp
@@ -0,0 +1,178 @@
+// Copyright Epic Games, Inc. All Rights Reserved.
+
+#include "Net/NetDormantHolder.h"
+
+#include "Net/DataReplication.h"
+
+namespace UE::Net::Private
+{
+
+/*-----------------------------------------------------------------------------
+	FDormantObjectReplicator
+-----------------------------------------------------------------------------*/
+
+FDormantObjectReplicator::FDormantObjectReplicator(FObjectKey InObjectKey)
+	: ObjectKey(InObjectKey)
+	, Replicator(MakeShared<FObjectReplicator>())
+{
+	//nothing
+}
+
+FDormantObjectReplicator::FDormantObjectReplicator(FObjectKey InObjectKey, const TSharedRef<FObjectReplicator>& ExistingReplicator)
+	: ObjectKey(InObjectKey)
+	, Replicator(ExistingReplicator)
+{
+	//nothing
+}
+
+/*-----------------------------------------------------------------------------
+	FDormantReplicatorHolder
+-----------------------------------------------------------------------------*/
+
+bool FDormantReplicatorHolder::DoesReplicatorExist(AActor* DormantActor, UObject* ReplicatedObject) const
+{
+	if (const FActorDormantReplicators* ActorReplicators = ActorReplicatorSet.Find(DormantActor))
+	{
+		const FObjectKey SubObjectKey = ReplicatedObject;
+
+		return (ActorReplicators->DormantReplicators.Find(SubObjectKey) != nullptr);
+	}
+
+	return false;
+}
+
+TSharedPtr<FObjectReplicator> FDormantReplicatorHolder::FindReplicator(AActor* DormantActor, UObject* ReplicatedObject)
+{
+	TSharedPtr<FObjectReplicator> ReplicatorPtr;
+
+	if (FActorDormantReplicators* ActorReplicators = ActorReplicatorSet.Find(DormantActor))
+	{
+		const FObjectKey SubObjectKey = ReplicatedObject;
+		
+		if (FDormantObjectReplicator* SubObjectReplicator = ActorReplicators->DormantReplicators.Find(SubObjectKey))
+		{
+			ReplicatorPtr = SubObjectReplicator->Replicator;
+		}
+	}
+
+	return ReplicatorPtr;
+}
+
+TSharedPtr<FObjectReplicator> FDormantReplicatorHolder::FindAndRemoveReplicator(AActor* DormantActor, UObject* ReplicatedObject)
+{
+	TSharedPtr<FObjectReplicator> ReplicatorPtr;
+
+	if (FActorDormantReplicators* ActorReplicators = ActorReplicatorSet.Find(DormantActor))
+	{
+		const FObjectKey SubObjectKey = ReplicatedObject;
+		FSetElementId Index = ActorReplicators->DormantReplicators.FindId(ReplicatedObject);
+
+		if (Index.IsValidId())
+		{
+			ReplicatorPtr = ActorReplicators->DormantReplicators[Index].Replicator;
+			ActorReplicators->DormantReplicators.Remove(Index);
+		}
+	}
+
+	return ReplicatorPtr;
+}
+
+const TSharedRef<FObjectReplicator>& FDormantReplicatorHolder::CreateAndStoreReplicator(AActor* DormantActor, UObject* ReplicatedObject)
+{
+	FActorDormantReplicators& ActorReplicators = ActorReplicatorSet.FindOrAdd(FActorDormantReplicators(DormantActor));
+
+	const FObjectKey SubObjectKey = ReplicatedObject;	
+
+	// Add a new replicator tied to this object. 
+	// If there was already a replicator for the same object in the set, it will be destroyed and overwritten by this new one.
+	FSetElementId Index = ActorReplicators.DormantReplicators.Add(FDormantObjectReplicator(SubObjectKey));
+
+	return ActorReplicators.DormantReplicators[Index].Replicator;
+}
+
+void FDormantReplicatorHolder::StoreReplicator(AActor* DormantActor, UObject* ReplicatedObject, const TSharedRef<FObjectReplicator>& ObjectReplicator)
+{
+	FActorDormantReplicators& ActorReplicators = ActorReplicatorSet.FindOrAdd(FActorDormantReplicators(DormantActor));
+
+	ActorReplicators.DormantReplicators.Add(FDormantObjectReplicator(ReplicatedObject, ObjectReplicator));
+}
+
+void FDormantReplicatorHolder::RemoveStoredReplicator(AActor* DormantActor, FObjectKey ReplicatedObjectKey)
+{
+	FSetElementId Index = ActorReplicatorSet.FindId(DormantActor);
+	if (Index.IsValidId())
+	{
+		ActorReplicatorSet[Index].DormantReplicators.Remove(ReplicatedObjectKey);
+
+		// Cleanup the actor entry if its not holding any other replicators
+		if (ActorReplicatorSet[Index].DormantReplicators.IsEmpty())		
+		{
+			ActorReplicatorSet.Remove(Index);
+		}
+	}
+}
+
+void FDormantReplicatorHolder::CleanupAllReplicatorsOfActor(AActor* DormantActor)
+{
+	ActorReplicatorSet.Remove(DormantActor);
+}
+
+void FDormantReplicatorHolder::CleanupStaleObjects()
+{
+	for (FActorReplicatorSet::TIterator ActorSetIt = ActorReplicatorSet.CreateIterator(); ActorSetIt; ++ActorSetIt )
+	{
+		for (FActorDormantReplicators::FObjectReplicatorSet::TIterator ReplicatorSetIt = ActorSetIt->DormantReplicators.CreateIterator(); ReplicatorSetIt; ++ReplicatorSetIt)
+		{
+			FDormantObjectReplicator& DormantReplicator = *ReplicatorSetIt;
+
+			if (!DormantReplicator.Replicator->GetWeakObjectPtr().IsValid())
+			{
+				ReplicatorSetIt.RemoveCurrent();
+			}
+		}
+		
+		if (ActorSetIt->DormantReplicators.IsEmpty())
+		{
+			ActorSetIt.RemoveCurrent();
+		}
+	}
+}
+
+void FDormantReplicatorHolder::ForEachDormantReplicator(UE::Net::FExecuteForEachDormantReplicator Function)
+{
+	for (const FActorDormantReplicators& ActorReplicators : ActorReplicatorSet)
+	{
+		for (const FDormantObjectReplicator& DormantReplicator : ActorReplicators.DormantReplicators)
+		{
+			Function(ActorReplicators.OwnerActor, DormantReplicator.ObjectKey, DormantReplicator.Replicator);
+		}
+	}
+}
+
+void FDormantReplicatorHolder::ForEachDormantReplicatorOfActor(AActor* DormantActor, UE::Net::FExecuteForEachDormantReplicator Function)
+{
+	if (FActorDormantReplicators* ActorReplicators = ActorReplicatorSet.Find(DormantActor))
+	{
+		for (const FDormantObjectReplicator& DormantReplicator : ActorReplicators->DormantReplicators)
+		{
+			Function(ActorReplicators->OwnerActor, DormantReplicator.ObjectKey, DormantReplicator.Replicator);
+		}
+	}
+}
+
+void FDormantReplicatorHolder::EmptySet()
+{
+	ActorReplicatorSet.Empty();
+}
+
+void FDormantReplicatorHolder::CountBytes(FArchive& Ar) const
+{
+	ActorReplicatorSet.CountBytes(Ar);
+	for (const FActorDormantReplicators& ActorReplicators : ActorReplicatorSet)
+	{
+		ActorReplicators.CountBytes(Ar);
+	}
+}
+
+
+} //end namespace UE::Net::Private
\ No newline at end of file
diff --git a/Engine/Source/Runtime/Engine/Private/NetConnection.cpp b/Engine/Source/Runtime/Engine/Private/NetConnection.cpp
--- a/Engine/Source/Runtime/Engine/Private/NetConnection.cpp
+++ b/Engine/Source/Runtime/Engine/Private/NetConnection.cpp
@@ -340,7 +340,9 @@ UNetConnection::UNetConnection(FVTableHelper& Helper)
 {
 }
 
+PRAGMA_DISABLE_DEPRECATION_WARNINGS
 UNetConnection::~UNetConnection() = default;
+PRAGMA_ENABLE_DEPRECATION_WARNINGS
 
 void UNetConnection::InitChannelData()
 {
@@ -854,7 +856,7 @@ void UNetConnection::Serialize( FArchive& Ar )
 		);
 
 		// ObjectReplicators are going to be counted by UNetDriver::Serialize AllOwnedReplicators.
-		GRANULAR_NETWORK_MEMORY_TRACKING_TRACK("DormantReplicatorMap", DormantReplicatorMap.CountBytes(Ar));
+		GRANULAR_NETWORK_MEMORY_TRACKING_TRACK("DormantReplicatorSet", DormantReplicatorSet.CountBytes(Ar));
 
 		GRANULAR_NETWORK_MEMORY_TRACKING_TRACK("ClientVisibleLevelNames", ClientVisibleLevelNames.CountBytes(Ar));
 		GRANULAR_NETWORK_MEMORY_TRACKING_TRACK("ClientVisibleActorOuters", ClientVisibleActorOuters.CountBytes(Ar));
@@ -4736,7 +4738,6 @@ void UNetConnection::ResetGameWorldState()
 	ClientVisibleLevelNames.Empty();
 	ClientMakingVisibleLevelNames.Empty();
 	KeepProcessingActorChannelBunchesMap.Empty();
-	DormantReplicatorMap.Empty();
 	CleanupDormantActorState();
 	ClientVisibleActorOuters.Empty();
 
@@ -4754,23 +4755,23 @@ void UNetConnection::ResetGameWorldState()
 
 void UNetConnection::CleanupDormantActorState()
 {
-	DormantReplicatorMap.Empty();
+	DormantReplicatorSet.EmptySet();
 }
 
-void UNetConnection::FlushDormancy(class AActor* Actor)
+void UNetConnection::FlushDormancy(AActor* Actor)
 {
 	UE_LOG( LogNetDormancy, Verbose, TEXT( "FlushDormancy: %s. Connection: %s" ), *Actor->GetName(), *GetName() );
 	
 	if ( Driver->GetNetworkObjectList().MarkActive( Actor, this, Driver ) )
 	{
-		FlushDormancyForObject( Actor );
+		FlushDormancyForObject( Actor, Actor );
 
 		// TODO: Is this set of objects sufficient? Should we query the dormancy map from the connection instead?
 		for ( UActorComponent* ActorComp : Actor->GetReplicatedComponents() )
 		{
 			if ( ActorComp && ActorComp->GetIsReplicated() )
 			{
-				FlushDormancyForObject( ActorComp );
+				FlushDormancyForObject(Actor, ActorComp );
 			}
 		}
 	}
@@ -4804,29 +4805,40 @@ void UNetConnection::ForcePropertyCompare( AActor* Actor )
 	}
 }
 
-/** Wrapper for validating an objects dormancy state, and to prepare the object for replication again */
-void UNetConnection::FlushDormancyForObject(UObject* Object)
+void UNetConnection::FlushDormancyForObject(AActor* DormantActor, UObject* ReplicatedObject)
 {
 	QUICK_SCOPE_CYCLE_COUNTER(STAT_NetConnection_FlushDormancyForObject)
 
-	TSharedRef<FObjectReplicator>* Replicator = DormantReplicatorMap.Find(Object);
-	if (Replicator)
+	bool bReuseReplicators = false;
+
+	if (GbNetReuseReplicatorsForDormantObjects && Driver && Driver->IsServer())
 	{
-		if (GNetDormancyValidate == 1)
-		{
-			Replicator->Get().ValidateAgainstState( Object );
-		}
+		bReuseReplicators = true;
+	}
 
-		if (!GbNetReuseReplicatorsForDormantObjects || !Driver || !Driver->IsServer())
+	if (GNetDormancyValidate == 1)
+	{
+		TSharedPtr<FObjectReplicator> Replicator = DormantReplicatorSet.FindReplicator(DormantActor, ReplicatedObject);
+		if (Replicator.IsValid())
 		{
-			Replicator = nullptr;
+			Replicator->ValidateAgainstState(ReplicatedObject);
 		}
 	}
 
-	if (Replicator == nullptr)
+	// If we want to reuse replicators, make sure they exist for this object
+	if (bReuseReplicators)
 	{
-		Replicator = &DormantReplicatorMap.Add(Object, MakeShared<FObjectReplicator>());
-		Replicator->Get().InitWithObject(Object, this, false);		// Init using the objects current state
+		bReuseReplicators = DormantReplicatorSet.DoesReplicatorExist(DormantActor, ReplicatedObject);
+	}
+
+	// If we need to create a new replicator when flushing
+	if (!bReuseReplicators)
+	{
+		const TSharedRef<FObjectReplicator>& ObjectReplicatorRef = DormantReplicatorSet.CreateAndStoreReplicator(DormantActor, ReplicatedObject);
+		
+		// Init using the objects current state
+		constexpr bool bUseDefaultState = false; 
+		ObjectReplicatorRef->InitWithObject(ReplicatedObject, this, bUseDefaultState);
 
 		// Flush the must be mapped GUIDs, the initialization may add them, but they're phantom and will be remapped when actually sending
 		if (UPackageMapClient* PackageMapClient = CastChecked<UPackageMapClient>(PackageMap))
@@ -5007,55 +5019,67 @@ void UNetConnection::DestroyIgnoredActor(AActor* Actor)
 	}
 }
 
-void UNetConnection::AddDormantReplicator(UObject* Object, const TSharedRef<FObjectReplicator>& Replicator)
+void UNetConnection::StoreDormantReplicator(AActor* OwnerActor, UObject* Object, const TSharedRef<FObjectReplicator>& ObjectReplicator)
 {
-	Replicator->ReleaseStrongReference();
-	DormantReplicatorMap.Add(Object, Replicator);
+	ObjectReplicator->ReleaseStrongReference();
+
+	if (ensureMsgf(OwnerActor, TEXT("StoreDormantReplicator cannot receive a null owner while storing %s"), *GetNameSafe(Object)))
+	{
+		DormantReplicatorSet.StoreReplicator(OwnerActor, Object, ObjectReplicator);
+	}
 }
 
-TSharedPtr<FObjectReplicator> UNetConnection::FindAndRemoveDormantReplicator(UObject* Object)
+TSharedPtr<FObjectReplicator> UNetConnection::FindAndRemoveDormantReplicator(AActor* OwnerActor, UObject* Object)
 {
-	FObjectKey Key(Object);
-	
-	if (DormantReplicatorMap.Contains(Key))
-	{
-		TSharedRef<FObjectReplicator> Ref = DormantReplicatorMap.FindAndRemoveChecked(Key);
+	const FObjectKey ObjectKey(Object);
+
+	TSharedPtr<FObjectReplicator> ReplicatorPtr = DormantReplicatorSet.FindAndRemoveReplicator(OwnerActor, Object);
 
+	if (ReplicatorPtr.IsValid())
+	{
 		// Only return the replicator if the object is still valid, otherwise just remove it from the cache and allow the caller to create a new one
-		if (UObject* StrongPtr = Ref->GetWeakObjectPtr().Get())
+		if (UObject* StrongPtr = ReplicatorPtr->GetWeakObjectPtr().Get())
 		{
 			// Reassign the strong pointer for GC/faster resolve
-			Ref->SetObject(StrongPtr);
-			return Ref;
+			ReplicatorPtr->SetObject(StrongPtr);
+		}
+		else
+		{
+			ReplicatorPtr.Reset();
 		}
 	}
-	return {};
+
+	return ReplicatorPtr;
+}
+
+void UNetConnection::RemoveDormantReplicator(AActor* Actor, UObject* Object)
+{
+	const FObjectKey ObjectKey = Object;
+
+	DormantReplicatorSet.RemoveStoredReplicator(Actor, ObjectKey);
+}
+
+void UNetConnection::ExecuteOnAllDormantReplicators(UE::Net::FExecuteForEachDormantReplicator ExecuteFunction)
+{
+	DormantReplicatorSet.ForEachDormantReplicator(ExecuteFunction);
+}
+
+void UNetConnection::ExecuteOnAllDormantReplicatorsOfActor(AActor* OwnerActor, UE::Net::FExecuteForEachDormantReplicator ExecuteFunction)
+{
+	DormantReplicatorSet.ForEachDormantReplicatorOfActor(OwnerActor, ExecuteFunction);	
 }
 
 void UNetConnection::CleanupDormantReplicatorsForActor(AActor* Actor)
 {
 	if (Actor)
 	{
-		// TODO: The DormantReplicator map contains not only replicated components, but any subobjects
-		// that are replicated (such as Gameplay Attribute Sets).
-		// That means we are likely leaking entries here (at least until CleanupStaleDormantReplicators is called).
-		DormantReplicatorMap.Remove(Actor);
-		for (UActorComponent* const Component : Actor->GetReplicatedComponents())
-		{
-			DormantReplicatorMap.Remove(Component);
-		}
+		DormantReplicatorSet.CleanupAllReplicatorsOfActor(Actor);
 	}
 }
 
 void UNetConnection::CleanupStaleDormantReplicators()
 {
-	for (auto It = DormantReplicatorMap.CreateIterator(); It; ++It)
-	{
-		if (!It.Value()->GetWeakObjectPtr().IsValid())
-		{
-			It.RemoveCurrent();
-		}
-	}
+	DormantReplicatorSet.CleanupStaleObjects();
 }
 
 void UNetConnection::SetPendingCloseDueToReplicationFailure()
diff --git a/Engine/Source/Runtime/Engine/Private/NetDriver.cpp b/Engine/Source/Runtime/Engine/Private/NetDriver.cpp
--- a/Engine/Source/Runtime/Engine/Private/NetDriver.cpp
+++ b/Engine/Source/Runtime/Engine/Private/NetDriver.cpp
@@ -5171,18 +5171,19 @@ int32 UNetDriver::ServerReplicateActors(float DeltaSeconds)
 		// net.DormancyValidate can be set to 2 to validate all dormant actors against last known state before going dormant
 		if ( GNetDormancyValidate == 2 )
 		{
-			// TODO: DormantReplicatorMap will actually contain all Actors and Subobjects.
-			// This means that we will call FObjectReplicator::ValidateAgainstState multiple times for
-			// the same object (once for itself and again for each subobject).
-			for ( auto It = Connection->DormantReplicatorMap.CreateIterator(); It; ++It )
+			UE::Net::FExecuteForEachDormantReplicator ValidateFunction = [](AActor* OwnerActor, FObjectKey ObjectKey, const TSharedRef<FObjectReplicator>& ReplicatorRef)
 			{
-				FObjectReplicator& Replicator = It.Value().Get();
+				FObjectReplicator& Replicator = ReplicatorRef.Get();
 
-				if ( Replicator.OwningChannel != nullptr )
+				// We will call FObjectReplicator::ValidateAgainstState multiple times for
+				// the same object (once for itself and again for each subobject).
+				if (Replicator.OwningChannel != nullptr)
 				{
-					Replicator.ValidateAgainstState( Replicator.OwningChannel->GetActor() );
+					Replicator.ValidateAgainstState(Replicator.OwningChannel->GetActor());
 				}
-			}
+			};
+				
+			Connection->ExecuteOnAllDormantReplicators(ValidateFunction);
 		}
 
 		// if this client shouldn't be ticked this frame
diff --git a/Engine/Source/Runtime/Engine/Public/Net/NetDormantHolder.h b/Engine/Source/Runtime/Engine/Public/Net/NetDormantHolder.h
new file mode 100644
--- /dev/null
+++ b/Engine/Source/Runtime/Engine/Public/Net/NetDormantHolder.h
@@ -0,0 +1,175 @@
+// Copyright Epic Games, Inc. All Rights Reserved.
+
+#pragma once
+
+#include "UObject/ObjectKey.h"
+#include "Templates/SharedPointer.h"
+#include "Serialization/Archive.h"
+
+class AActor;
+class FObjectReplicator;
+
+namespace UE::Net
+{
+	/** Function definition allowed in the ForEach functions */
+	typedef TFunctionRef<void(AActor* OwnerActor, FObjectKey ObjectKey, const TSharedRef<FObjectReplicator>& ReplicatorRef)> FExecuteForEachDormantReplicator;
+}
+
+namespace UE::Net::Private
+{
+
+/** 
+* Container holding the FObjectReplicator of a dormant replicated object
+*/
+struct FDormantObjectReplicator
+{
+	/** Default constructor that will always create a replicator */
+	explicit FDormantObjectReplicator(FObjectKey InObjectKey);
+
+	/** Constructor that copies an existing replicator */
+	explicit FDormantObjectReplicator(FObjectKey InObjectKey, const TSharedRef<FObjectReplicator>& ExistingReplicator);
+
+	bool operator==(const FDormantObjectReplicator& rhs) const { return ObjectKey == rhs.ObjectKey; }
+	bool operator==(FObjectKey rhs) const { return ObjectKey == rhs; }
+
+	/** The replicated object's key mapped to this replicator */
+	FObjectKey ObjectKey;
+
+	/** The replicator of the replicated object we are storing */
+	TSharedRef<FObjectReplicator> Replicator;
+};
+
+
+/** 
+* Container that stores all the replicators owned by a given dormant Actor. 
+* Note that the actor's own replicator will be part of this list.
+*/
+struct FActorDormantReplicators
+{
+	explicit FActorDormantReplicators(AActor* InOwnerActor) : OwnerActor(InOwnerActor) {}
+
+	bool operator==(const FActorDormantReplicators& rhs) const { return OwnerActor == rhs.OwnerActor; }
+	bool operator==(AActor* rhs) const { return OwnerActor == rhs; }
+
+	void CountBytes(FArchive& Ar) const { DormantReplicators.CountBytes(Ar); }
+
+	/** KeyFuncs that make it so the TSet can match a struct with a simple FObjectKey. */
+	struct FDormantObjectReplicatorKeyFuncs : BaseKeyFuncs<FDormantObjectReplicator, FObjectKey, false>
+	{
+		static KeyInitType GetSetKey(ElementInitType Element) { return Element.ObjectKey; }
+		static bool Matches(KeyInitType lhs, KeyInitType rhs) { return lhs == rhs; }
+		static uint32 GetKeyHash(KeyInitType Key) { return GetTypeHash(Key); }
+	};
+
+	/** The dormant actor who owns all the different object replicators */
+	AActor* OwnerActor = nullptr;
+
+	typedef TSet<FDormantObjectReplicator, FDormantObjectReplicatorKeyFuncs> FObjectReplicatorSet;
+	/** List of all object replicates stored for this dormant actor */
+	FObjectReplicatorSet DormantReplicators;
+};
+
+/**
+* Container that stores the object replicators of dormant actors and their subobjects.
+* Note that the actor stores it's own replicator and the replicators of it's subobjects
+*/
+struct FDormantReplicatorHolder
+{
+	/** 
+	* Are we currently storing an object replicator owned by the dormant actor. 
+	* 
+	* @param DormantActor The dormant actor that owns the replicated object
+	* @param ReplicatedObject The object that was replicated and for whom we could be storing it's FObjectReplicator.
+	* @return True if we are storing a replicator.
+	*/
+	bool DoesReplicatorExist(AActor* DormantActor, UObject* ReplicatedObject) const;
+
+	/**
+	* Return the objector replicator for the given replicated object .
+	* 
+	* @param DormantActor The dormant actor that owns the replicated object.
+	* @param ReplicatedObject The object that was replicated
+	* @return The object replicator of the replicated object if we were holding one. Otherwise an invalid pointer
+	*/
+	TSharedPtr<FObjectReplicator> FindReplicator(AActor* DormantActor, UObject* ReplicatedObject);
+
+	/**
+	* Return the object replicator for the given replicated object but also remove the reference to that object replicator.
+	* 
+	* @param DormantActor The dormant actor that owns the replicated object
+	* @param ReplicatedObject The object that was replicated and for whom we want it's FObjectReplicator.
+	* @return The object replicator of the replicated object if we were holding one. Otherwise an invalid pointer
+	*/
+	TSharedPtr<FObjectReplicator> FindAndRemoveReplicator(AActor* DormantActor, UObject* ReplicatedObject);
+
+	/**
+	* Create an object replicator for the given replicated object and store a reference to it.
+	* Does not initialize the Replicator but simply allocates it.
+	* 
+	* @param DormantActor The dormant actor that owns the replicated object
+	* @param ReplicatedObject The object that will be tied to the new ObjectReplicator.
+	* @return Return a reference to the object replicator we created and are now storing.
+	*/
+	const TSharedRef<FObjectReplicator>& CreateAndStoreReplicator(AActor* DormantActor, UObject* ReplicatedObject);
+
+	/**
+	* Store an existing replicator tied to the given replicated object
+	* 
+	* @param DormantActor The dormant actor that owns the replicated object
+	* @param ReplicatedObject The object that is tied to the existing replicator
+	* @param ObjectReplicator The existing replicator we that will be stored
+	*/
+	void StoreReplicator(AActor* DormantActor, UObject* ReplicatedObject, const TSharedRef<FObjectReplicator>& ObjectReplicator);
+
+	/**
+	* Remove our reference to the object replicator tied to the given replicated object
+	* 
+	* @param DormantActor The dormant actor that owns the replicated object
+	* @param ReplicatedObjectKey The object key to a replicated object that may have stored it's object replicator here
+	*/
+	void RemoveStoredReplicator(AActor* DormantActor, FObjectKey ReplicatedObjectKey);
+
+	/**
+	* Remove the references to all the object replicators tied to the given actor
+	* 
+	* @param DormantActor The dormant actor we want to clean any reference to
+	*/
+	void CleanupAllReplicatorsOfActor(AActor* DormantActor);
+
+	/**
+	* Iterate over all the stored object replicators and destroy any that are tied to a replicated object that is now considered invalid.
+	*/
+	void CleanupStaleObjects();
+
+	/** 
+	* Execute the passed function on each dormant replicator we are holding
+	*/
+	void ForEachDormantReplicator(UE::Net::FExecuteForEachDormantReplicator Function);
+
+	/**
+	* Execute the passed function on each dormant replicatof owned by the given dormant actor
+	*/
+	void ForEachDormantReplicatorOfActor(AActor* DormantActor, UE::Net::FExecuteForEachDormantReplicator Function);
+
+	/**
+	* Remove every reference that is still stored.
+	*/
+	void EmptySet();
+
+	void CountBytes(FArchive& Ar) const;
+
+	/** KeyFuncs that make it so the TSet can only needs a simple AActor pointer for the Key */
+	struct FActorDormantReplicatorsKeyFuncs : BaseKeyFuncs<FActorDormantReplicators, AActor*, false>
+	{
+		static KeyInitType GetSetKey(ElementInitType Element) { return Element.OwnerActor; }
+		static bool Matches(KeyInitType lhs, KeyInitType rhs) { return lhs == rhs; }
+		static uint32 GetKeyHash(KeyInitType Key) { return GetTypeHash(Key); }
+	};
+
+	typedef TSet<FActorDormantReplicators, FActorDormantReplicatorsKeyFuncs> FActorReplicatorSet;
+	/** The TSet indexed by Actor that stores all their dormant object replicators */
+	FActorReplicatorSet ActorReplicatorSet;
+};
+
+}//end namespace UE::Net::Private
+
