Small optimizations to the device property subsytem. Dont process input if we are at a breakpoint, and only broadcast the input device changed delegate if it was different then the known hardware

#rnx

diff --git a/Engine/Source/Runtime/Engine/Private/GameFramework/InputDeviceSubsystem.cpp b/Engine/Source/Runtime/Engine/Private/GameFramework/InputDeviceSubsystem.cpp
--- a/Engine/Source/Runtime/Engine/Private/GameFramework/InputDeviceSubsystem.cpp
+++ b/Engine/Source/Runtime/Engine/Private/GameFramework/InputDeviceSubsystem.cpp
@@ -30,6 +30,17 @@ class FInputDeviceSubsystemProcessor : public IInputProcessor
 	
 	void UpdateLatestDevice(const FInputDeviceId InDeviceId)
 	{
+#if WITH_EDITOR
+		// If we're stopped at a breakpoint we need for this input preprocessor to just ignore all incoming input
+		// because we're now doing stuff outside the game loop in the editor and it needs to not block all that.
+		// This can happen if you suspend input while spawning a dialog and then hit another breakpoint and then
+		// try and use the editor, you can suddenly be unable to do anything.
+		if (GIntraFrameDebuggingGameThread)
+		{
+			return;
+		}
+#endif
+		
 		if (UInputDeviceSubsystem* SubSystem = UInputDeviceSubsystem::Get())
 		{
 			if (const FInputDeviceScope* Scope = FInputDeviceScope::GetCurrent())
@@ -481,7 +492,16 @@ FHardwareDeviceIdentifier UInputDeviceSubsystem::GetInputDeviceHardwareIdentifie
 }
 
 void UInputDeviceSubsystem::SetMostRecentlyUsedHardwareDevice(const FInputDeviceId InDeviceId, const FHardwareDeviceIdentifier& InHardwareId)
-{
+{	
+	// If the hardware hasn't changed, then just ignore it because we don't need to update anything.
+	if (const FHardwareDeviceIdentifier* ExistingDevice = LatestInputDeviceIdentifiers.Find(InDeviceId))
+	{
+		if (InHardwareId == *ExistingDevice)
+		{
+			return;
+		}
+	}
+
 	FPlatformUserId OwningUserId = IPlatformInputDeviceMapper::Get().GetUserForInputDevice(InDeviceId);
 
 	// Keep track of each input device's latest hardware id
