Move some UPrimitiveComponent templated methods from the .cpp to the .h as they're used by other compile units.

diff --git a/Engine/Source/Runtime/Engine/Classes/Components/PrimitiveComponent.h b/Engine/Source/Runtime/Engine/Classes/Components/PrimitiveComponent.h
--- a/Engine/Source/Runtime/Engine/Classes/Components/PrimitiveComponent.h
+++ b/Engine/Source/Runtime/Engine/Classes/Components/PrimitiveComponent.h
@@ -3,6 +3,7 @@
 #pragma once
 
 #include "CoreMinimal.h"
+#include "Algo/Copy.h"
 #include "HAL/ThreadSafeCounter.h"
 #include "UObject/ObjectMacros.h"
 #include "UObject/UObjectGlobals.h"
@@ -10,6 +11,7 @@
 #include "InputCoreTypes.h"
 #include "Templates/SubclassOf.h"
 #include "Engine/EngineTypes.h"
+#include "Engine/ScopedMovementUpdate.h"
 #include "Components/SceneComponent.h"
 #include "RenderCommandFence.h"
 #include "GameFramework/Actor.h"
@@ -44,6 +46,8 @@ namespace PrimitiveComponentCVars
 {
 	extern float HitDistanceToleranceCVar;
 	extern float InitialOverlapToleranceCVar;
+	extern int32 bAllowCachedOverlapsCVar;
+	extern int32 bEnableFastOverlapCheck;
 }
 
 /** Determines whether a Character can attempt to step up onto a component when they walk in to it. */
@@ -167,6 +171,24 @@ struct FRendererStencilMaskEvaluation
 	}
 };
 
+// Predicate to determine if an overlap is *NOT* with a certain AActor.
+struct FPredicateOverlapHasDifferentActor
+{
+	FPredicateOverlapHasDifferentActor(const AActor& Owner)
+		: MyOwnerPtr(&Owner)
+	{
+	}
+
+	bool operator() (const FOverlapInfo& Info)
+	{
+		// MyOwnerPtr is always valid, so we don't need the IsValid() checks in the WeakObjectPtr comparison operator.
+		return !MyOwnerPtr.HasSameIndexAndSerialNumber(Info.OverlapInfo.HitObjectHandle.FetchActor());
+	}
+
+private:
+	const TWeakObjectPtr<const AActor> MyOwnerPtr;
+};
+
 // TODO: Add sleep and wake state change types to this enum, so that the
 // OnComponentWake and OnComponentSleep delegates may be deprecated.
 // Doing so would save a couple bytes per primitive component.
@@ -2675,6 +2697,9 @@ public:
 	/** If true then DoCustomNavigableGeometryExport will be called to collect navigable geometry of this component. */
 	FORCEINLINE EHasCustomNavigableGeometry::Type HasCustomNavigableGeometry() const { return bHasCustomNavigableGeometry; }
 
+	// Returns true if we should check the GetGenerateOverlapEvents() flag when gathering overlaps, otherwise we'll always just do it.
+	FORCEINLINE_DEBUGGABLE bool ShouldCheckOverlapFlagToQueueOverlaps(const UPrimitiveComponent& ThisComponent) const;
+
 	/** Set value of HasCustomNavigableGeometry */
 	void SetCustomNavigableGeometry(const EHasCustomNavigableGeometry::Type InType);
 
@@ -2769,4 +2794,128 @@ FORCEINLINE_DEBUGGABLE bool UPrimitiveComponent::K2_IsPhysicsCollisionEnabled()
 FORCEINLINE_DEBUGGABLE bool UPrimitiveComponent::GetGenerateOverlapEvents() const
 {
 	return bGenerateOverlapEvents;
-}
\ No newline at end of file
+}
+
+FORCEINLINE_DEBUGGABLE bool UPrimitiveComponent::ShouldCheckOverlapFlagToQueueOverlaps(const UPrimitiveComponent& ThisComponent) const
+{
+	const FScopedMovementUpdate* CurrentUpdate = ThisComponent.GetCurrentScopedMovement();
+	if (CurrentUpdate)
+	{
+		return CurrentUpdate->RequiresOverlapsEventFlag();
+	}
+	// By default we require the GetGenerateOverlapEvents() to queue up overlaps, since we require it to trigger events.
+	return true;
+}
+
+//////////////////////////////////////////////////////////////////////////
+// PrimitiveComponent templates
+
+template<typename AllocatorType>
+bool UPrimitiveComponent::ConvertSweptOverlapsToCurrentOverlaps(
+	TArray<FOverlapInfo, AllocatorType>& OverlapsAtEndLocation, const TOverlapArrayView& SweptOverlaps, int32 SweptOverlapsIndex,
+	const FVector& EndLocation, const FQuat& EndRotationQuat)
+{
+	checkSlow(SweptOverlapsIndex >= 0);
+
+	bool bResult = false;
+	const bool bForceGatherOverlaps = !ShouldCheckOverlapFlagToQueueOverlaps(*this);
+	if ((GetGenerateOverlapEvents() || bForceGatherOverlaps) && PrimitiveComponentCVars::bAllowCachedOverlapsCVar)
+	{
+		const AActor* Actor = GetOwner();
+		if (Actor && Actor->GetRootComponent() == this)
+		{
+			// We know we are not overlapping any new components at the end location. Children are ignored here (see note below).
+			if (PrimitiveComponentCVars::bEnableFastOverlapCheck)
+			{
+				SCOPE_CYCLE_COUNTER(STAT_MoveComponent_FastOverlap);
+
+				// Check components we hit during the sweep, keep only those still overlapping
+				const FCollisionQueryParams UnusedQueryParams(NAME_None, FCollisionQueryParams::GetUnknownStatId());
+				const int32 NumSweptOverlaps = SweptOverlaps.Num();
+				OverlapsAtEndLocation.Reserve(OverlapsAtEndLocation.Num() + NumSweptOverlaps);
+				for (int32 Index = SweptOverlapsIndex; Index < NumSweptOverlaps; ++Index)
+				{
+					const FOverlapInfo& OtherOverlap = SweptOverlaps[Index];
+					UPrimitiveComponent* OtherPrimitive = OtherOverlap.OverlapInfo.GetComponent();
+					if (OtherPrimitive && (OtherPrimitive->GetGenerateOverlapEvents() || bForceGatherOverlaps))
+					{
+						if (OtherPrimitive->bMultiBodyOverlap)
+						{
+							// Not handled yet. We could do it by checking every body explicitly and track each body index in the overlap test, but this seems like a rare need.
+							return false;
+						}
+						else if (Cast<USkeletalMeshComponent>(OtherPrimitive) || Cast<USkeletalMeshComponent>(this))
+						{
+							// SkeletalMeshComponent does not support this operation, and would return false in the test when an actual query could return true.
+							return false;
+						}
+						else if (OtherPrimitive->ComponentOverlapComponent(this, EndLocation, EndRotationQuat, UnusedQueryParams))
+						{
+							OverlapsAtEndLocation.Add(OtherOverlap);
+						}
+					}
+				}
+
+				// Note: we don't worry about adding any child components here, because they are not included in the sweep results.
+				// Children test for their own overlaps after we update our own, and we ignore children in our own update.
+				checkfSlow(OverlapsAtEndLocation.FindByPredicate(FPredicateOverlapHasSameActor(*Actor)) == nullptr,
+					TEXT("Child overlaps should not be included in the SweptOverlaps() array in UPrimitiveComponent::ConvertSweptOverlapsToCurrentOverlaps()."));
+
+				bResult = true;
+			}
+			else
+			{
+				if (SweptOverlaps.Num() == 0 && AreAllCollideableDescendantsRelative())
+				{
+					// Add overlaps with components in this actor.
+					GetOverlapsWithActor_Template(Actor, OverlapsAtEndLocation);
+					bResult = true;
+				}
+			}
+		}
+	}
+
+	return bResult;
+}
+
+template<typename AllocatorType>
+bool UPrimitiveComponent::ConvertRotationOverlapsToCurrentOverlaps(TArray<FOverlapInfo, AllocatorType>& OutOverlapsAtEndLocation, const TOverlapArrayView& CurrentOverlaps)
+{
+	bool bResult = false;
+	const bool bForceGatherOverlaps = !ShouldCheckOverlapFlagToQueueOverlaps(*this);
+	if ((GetGenerateOverlapEvents() || bForceGatherOverlaps) && PrimitiveComponentCVars::bAllowCachedOverlapsCVar)
+	{
+		const AActor* Actor = GetOwner();
+		if (Actor && Actor->GetRootComponent() == this)
+		{
+			if (PrimitiveComponentCVars::bEnableFastOverlapCheck)
+			{
+				// Add all current overlaps that are not children. Children test for their own overlaps after we update our own, and we ignore children in our own update.
+				OutOverlapsAtEndLocation.Reserve(OutOverlapsAtEndLocation.Num() + CurrentOverlaps.Num());
+				Algo::CopyIf(CurrentOverlaps, OutOverlapsAtEndLocation, FPredicateOverlapHasDifferentActor(*Actor));
+				bResult = true;
+			}
+		}
+	}
+
+	return bResult;
+}
+
+template<typename AllocatorType>
+bool UPrimitiveComponent::GetOverlapsWithActor_Template(const AActor* Actor, TArray<FOverlapInfo, AllocatorType>& OutOverlaps) const
+{
+	const int32 InitialCount = OutOverlaps.Num();
+	if (Actor)
+	{
+		for (int32 OverlapIdx = 0; OverlapIdx < OverlappingComponents.Num(); ++OverlapIdx)
+		{
+			UPrimitiveComponent const* const PrimComp = OverlappingComponents[OverlapIdx].OverlapInfo.Component.Get();
+			if (PrimComp && (PrimComp->GetOwner() == Actor))
+			{
+				OutOverlaps.Add(OverlappingComponents[OverlapIdx]);
+			}
+		}
+	}
+
+	return InitialCount != OutOverlaps.Num();
+}
diff --git a/Engine/Source/Runtime/Engine/Private/Components/PrimitiveComponent.cpp b/Engine/Source/Runtime/Engine/Private/Components/PrimitiveComponent.cpp
--- a/Engine/Source/Runtime/Engine/Private/Components/PrimitiveComponent.cpp
+++ b/Engine/Source/Runtime/Engine/Private/Components/PrimitiveComponent.cpp
@@ -34,7 +34,6 @@
 #include "GameFramework/CheatManager.h"
 #include "Streaming/TextureStreamingHelpers.h"
 #include "PrimitiveSceneProxy.h"
-#include "Algo/Copy.h"
 #include "UObject/RenderingObjectVersion.h"
 #include "UObject/FortniteMainBranchObjectVersion.h"
 #include "UObject/UE5PrivateFrostyStreamObjectVersion.h"
@@ -70,14 +69,6 @@ typedef TArray<const FOverlapInfo*, TInlineAllocator<8>> TInlineOverlapPointerAr
 
 DEFINE_LOG_CATEGORY_STATIC(LogPrimitiveComponent, Log, All);
 
-static int32 bAllowCachedOverlapsCVar = 1;
-static FAutoConsoleVariableRef CVarAllowCachedOverlaps(
-	TEXT("p.AllowCachedOverlaps"), 
-	bAllowCachedOverlapsCVar,
-	TEXT("Primitive Component physics\n")
-	TEXT("0: disable cached overlaps, 1: enable (default)"),
-	ECVF_Default);
-
 static int32 AlwaysCreatePhysicsStateConversionHackCVar = 0;
 static FAutoConsoleVariableRef CVarAlwaysCreatePhysicsStateConversionHack(
 	TEXT("p.AlwaysCreatePhysicsStateConversionHack"),
@@ -87,6 +78,14 @@ static FAutoConsoleVariableRef CVarAlwaysCreatePhysicsStateConversionHack(
 
 namespace PrimitiveComponentCVars
 {
+	int32 bAllowCachedOverlapsCVar = 1;
+	static FAutoConsoleVariableRef CVarAllowCachedOverlaps(
+		TEXT("p.AllowCachedOverlaps"),
+		bAllowCachedOverlapsCVar,
+		TEXT("Primitive Component physics\n")
+		TEXT("0: disable cached overlaps, 1: enable (default)"),
+		ECVF_Default);
+
 	float InitialOverlapToleranceCVar = 0.0f;
 	static FAutoConsoleVariableRef CVarInitialOverlapTolerance(
 		TEXT("p.InitialOverlapTolerance"),
@@ -103,6 +102,9 @@ namespace PrimitiveComponentCVars
 		TEXT("Tolerance for hit distance for overlap test in PrimitiveComponent movement.\n")
 		TEXT("Hits that are less than this distance are ignored."),
 		ECVF_Default);
+
+	int32 bEnableFastOverlapCheck = 1;
+	static FAutoConsoleVariableRef CVarEnableFastOverlapCheck(TEXT("p.EnableFastOverlapCheck"), bEnableFastOverlapCheck, TEXT("Enable fast overlap check against sweep hits, avoiding UpdateOverlaps (for the swept component)."));
 }
 
 #if !(UE_BUILD_SHIPPING || UE_BUILD_TEST)
@@ -115,8 +117,6 @@ FAutoConsoleVariableRef CVarRefShowInitialOverlaps(
 	ECVF_Cheat);
 #endif // !(UE_BUILD_SHIPPING || UE_BUILD_TEST)
 
-static int32 bEnableFastOverlapCheck = 1;
-static FAutoConsoleVariableRef CVarEnableFastOverlapCheck(TEXT("p.EnableFastOverlapCheck"), bEnableFastOverlapCheck, TEXT("Enable fast overlap check against sweep hits, avoiding UpdateOverlaps (for the swept component)."));
 DECLARE_CYCLE_STAT(TEXT("MoveComponent FastOverlap"), STAT_MoveComponent_FastOverlap, STATGROUP_Game);
 DECLARE_CYCLE_STAT(TEXT("BeginComponentOverlap"), STAT_BeginComponentOverlap, STATGROUP_Game);
 DECLARE_CYCLE_STAT(TEXT("EndComponentOverlap"), STAT_EndComponentOverlap, STATGROUP_Game);
@@ -148,24 +148,6 @@ private:
 	const TWeakObjectPtr<const AActor> MyOwnerPtr;
 };
 
-// Predicate to determine if an overlap is *NOT* with a certain AActor.
-struct FPredicateOverlapHasDifferentActor
-{
-	FPredicateOverlapHasDifferentActor(const AActor& Owner)
-	: MyOwnerPtr(&Owner)
-	{
-	}
-
-	bool operator() (const FOverlapInfo& Info)
-	{
-		// MyOwnerPtr is always valid, so we don't need the IsValid() checks in the WeakObjectPtr comparison operator.
-		return !MyOwnerPtr.HasSameIndexAndSerialNumber(Info.OverlapInfo.HitObjectHandle.FetchActor());
-	}
-
-private:
-	const TWeakObjectPtr<const AActor> MyOwnerPtr;
-};
-
 // Helper for finding the index of an FOverlapInfo in an Array using the FFastOverlapInfoCompare predicate, knowing that at least one overlap is valid (non-null).
 template<class AllocatorType>
 FORCEINLINE_DEBUGGABLE int32 IndexOfOverlapFast(const TArray<FOverlapInfo, AllocatorType>& OverlapArray, const FOverlapInfo& SearchItem)
@@ -2354,20 +2336,6 @@ static bool ShouldIgnoreHitResult(const UWorld* InWorld, FHitResult const& TestH
 	return false;
 }
 
-
-// Returns true if we should check the GetGenerateOverlapEvents() flag when gathering overlaps, otherwise we'll always just do it.
-static FORCEINLINE_DEBUGGABLE bool ShouldCheckOverlapFlagToQueueOverlaps(const UPrimitiveComponent& ThisComponent)
-{
-	const FScopedMovementUpdate* CurrentUpdate = ThisComponent.GetCurrentScopedMovement();
-	if (CurrentUpdate)
-	{
-		return CurrentUpdate->RequiresOverlapsEventFlag();
-	}
-	// By default we require the GetGenerateOverlapEvents() to queue up overlaps, since we require it to trigger events.
-	return true;
-}
-
-
 static FORCEINLINE_DEBUGGABLE bool ShouldIgnoreOverlapResult(const UWorld* World, const AActor* ThisActor, const UPrimitiveComponent& ThisComponent, const AActor* OtherActor, const UPrimitiveComponent& OtherComponent, bool bCheckOverlapFlags)
 {
 	// Don't overlap with self
@@ -2971,25 +2939,6 @@ bool UPrimitiveComponent::IsOverlappingActor(const AActor* Other) const
 	return false;
 }
 
-template<typename AllocatorType>
-bool UPrimitiveComponent::GetOverlapsWithActor_Template(const AActor* Actor, TArray<FOverlapInfo, AllocatorType>& OutOverlaps) const
-{
-	const int32 InitialCount = OutOverlaps.Num();
-	if (Actor)
-	{
-		for (int32 OverlapIdx=0; OverlapIdx<OverlappingComponents.Num(); ++OverlapIdx)
-		{
-			UPrimitiveComponent const* const PrimComp = OverlappingComponents[OverlapIdx].OverlapInfo.Component.Get();
-			if ( PrimComp && (PrimComp->GetOwner() == Actor) )
-			{
-				OutOverlaps.Add(OverlappingComponents[OverlapIdx]);
-			}
-		}
-	}
-
-	return InitialCount != OutOverlaps.Num();
-}
-
 bool UPrimitiveComponent::GetOverlapsWithActor(const AActor* Actor, TArray<FOverlapInfo>& OutOverlaps) const
 {
 	return GetOverlapsWithActor_Template(Actor, OutOverlaps);
@@ -3303,98 +3252,6 @@ void UPrimitiveComponent::GetOverlappingComponents(TSet<UPrimitiveComponent*>& O
 	}
 }
 
-template<typename AllocatorType>
-bool UPrimitiveComponent::ConvertSweptOverlapsToCurrentOverlaps(
-	TArray<FOverlapInfo, AllocatorType>& OverlapsAtEndLocation, const TOverlapArrayView& SweptOverlaps, int32 SweptOverlapsIndex,
-	const FVector& EndLocation, const FQuat& EndRotationQuat)
-{
-	checkSlow(SweptOverlapsIndex >= 0);
-
-	bool bResult = false;
-	const bool bForceGatherOverlaps = !ShouldCheckOverlapFlagToQueueOverlaps(*this);
-	if ((GetGenerateOverlapEvents() || bForceGatherOverlaps) && bAllowCachedOverlapsCVar)
-	{
-		const AActor* Actor = GetOwner();
-		if (Actor && Actor->GetRootComponent() == this)
-		{
-			// We know we are not overlapping any new components at the end location. Children are ignored here (see note below).
-			if (bEnableFastOverlapCheck)
-			{
-				SCOPE_CYCLE_COUNTER(STAT_MoveComponent_FastOverlap);
-
-				// Check components we hit during the sweep, keep only those still overlapping
-				const FCollisionQueryParams UnusedQueryParams(NAME_None, FCollisionQueryParams::GetUnknownStatId());
-				const int32 NumSweptOverlaps = SweptOverlaps.Num();
-				OverlapsAtEndLocation.Reserve(OverlapsAtEndLocation.Num() + NumSweptOverlaps);
-				for (int32 Index = SweptOverlapsIndex; Index < NumSweptOverlaps; ++Index)
-				{
-					const FOverlapInfo& OtherOverlap = SweptOverlaps[Index];
-					UPrimitiveComponent* OtherPrimitive = OtherOverlap.OverlapInfo.GetComponent();
-					if (OtherPrimitive && (OtherPrimitive->GetGenerateOverlapEvents() || bForceGatherOverlaps))
-					{
-						if (OtherPrimitive->bMultiBodyOverlap)
-						{
-							// Not handled yet. We could do it by checking every body explicitly and track each body index in the overlap test, but this seems like a rare need.
-							return false;
-						}
-						else if (Cast<USkeletalMeshComponent>(OtherPrimitive) || Cast<USkeletalMeshComponent>(this))
-						{
-							// SkeletalMeshComponent does not support this operation, and would return false in the test when an actual query could return true.
-							return false;
-						}
-						else if (OtherPrimitive->ComponentOverlapComponent(this, EndLocation, EndRotationQuat, UnusedQueryParams))
-						{
-							OverlapsAtEndLocation.Add(OtherOverlap);
-						}
-					}
-				}
-
-				// Note: we don't worry about adding any child components here, because they are not included in the sweep results.
-				// Children test for their own overlaps after we update our own, and we ignore children in our own update.
-				checkfSlow(OverlapsAtEndLocation.FindByPredicate(FPredicateOverlapHasSameActor(*Actor)) == nullptr,
-					TEXT("Child overlaps should not be included in the SweptOverlaps() array in UPrimitiveComponent::ConvertSweptOverlapsToCurrentOverlaps()."));
-
-				bResult = true;
-			}
-			else
-			{
-				if (SweptOverlaps.Num() == 0 && AreAllCollideableDescendantsRelative())
-				{
-					// Add overlaps with components in this actor.
-					GetOverlapsWithActor_Template(Actor, OverlapsAtEndLocation);
-					bResult = true;
-				}
-			}
-		}
-	}
-
-	return bResult;
-}
-
-template<typename AllocatorType>
-bool UPrimitiveComponent::ConvertRotationOverlapsToCurrentOverlaps(TArray<FOverlapInfo, AllocatorType>& OutOverlapsAtEndLocation, const TOverlapArrayView& CurrentOverlaps)
-{
-	bool bResult = false;
-	const bool bForceGatherOverlaps = !ShouldCheckOverlapFlagToQueueOverlaps(*this);
-	if ((GetGenerateOverlapEvents() || bForceGatherOverlaps) && bAllowCachedOverlapsCVar)
-	{
-		const AActor* Actor = GetOwner();
-		if (Actor && Actor->GetRootComponent() == this)
-		{
-			if (bEnableFastOverlapCheck)
-			{
-				// Add all current overlaps that are not children. Children test for their own overlaps after we update our own, and we ignore children in our own update.
-				OutOverlapsAtEndLocation.Reserve(OutOverlapsAtEndLocation.Num() + CurrentOverlaps.Num());
-				Algo::CopyIf(CurrentOverlaps, OutOverlapsAtEndLocation, FPredicateOverlapHasDifferentActor(*Actor));
-				bResult = true;
-			}
-		}
-	}
-
-	return bResult;
-}
-
-
 bool UPrimitiveComponent::AreAllCollideableDescendantsRelative(bool bAllowCachedValue) const
 {
 	UPrimitiveComponent* MutableThis = const_cast<UPrimitiveComponent*>(this);
@@ -3571,7 +3428,7 @@ bool UPrimitiveComponent::UpdateOverlapsImpl(const TOverlapArrayView* NewPending
 			if (IsValid(this) && GetGenerateOverlapEvents())
 			{
 				// Might be able to avoid testing for new overlaps at the end location.
-				if (OverlapsAtEndLocation != nullptr && bAllowCachedOverlapsCVar && PrevTransform.Equals(GetComponentTransform()))
+				if (OverlapsAtEndLocation != nullptr && PrimitiveComponentCVars::bAllowCachedOverlapsCVar && PrevTransform.Equals(GetComponentTransform()))
 				{
 					UE_LOG(LogPrimitiveComponent, VeryVerbose, TEXT("%s->%s Skipping overlap test!"), *GetNameSafe(GetOwner()), *GetName());
 					const bool bCheckForInvalid = (NewPendingOverlaps && NewPendingOverlaps->Num() > 0);
