Fix non-unity compile issues with CL#23342352

diff --git a/Engine/Source/Runtime/Engine/Classes/Components/PrimitiveComponent.h b/Engine/Source/Runtime/Engine/Classes/Components/PrimitiveComponent.h
--- a/Engine/Source/Runtime/Engine/Classes/Components/PrimitiveComponent.h
+++ b/Engine/Source/Runtime/Engine/Classes/Components/PrimitiveComponent.h
@@ -4,6 +4,7 @@
 
 #include "CoreMinimal.h"
 #include "Algo/Copy.h"
+#include "EngineStats.h"
 #include "HAL/ThreadSafeCounter.h"
 #include "UObject/ObjectMacros.h"
 #include "UObject/UObjectGlobals.h"
diff --git a/Engine/Source/Runtime/Engine/Private/Components/PrimitiveComponent.cpp b/Engine/Source/Runtime/Engine/Private/Components/PrimitiveComponent.cpp
--- a/Engine/Source/Runtime/Engine/Private/Components/PrimitiveComponent.cpp
+++ b/Engine/Source/Runtime/Engine/Private/Components/PrimitiveComponent.cpp
@@ -6,7 +6,6 @@
 
 #include "Components/PrimitiveComponent.h"
 #include "Engine/Level.h"
-#include "EngineStats.h"
 #include "GameFramework/DamageType.h"
 #include "GameFramework/Pawn.h"
 #include "PSOPrecache.h"
