#UE CIS quick fix. Disable stat

diff --git a/Engine/Source/Runtime/Engine/Classes/Components/PrimitiveComponent.h b/Engine/Source/Runtime/Engine/Classes/Components/PrimitiveComponent.h
--- a/Engine/Source/Runtime/Engine/Classes/Components/PrimitiveComponent.h
+++ b/Engine/Source/Runtime/Engine/Classes/Components/PrimitiveComponent.h
@@ -2828,7 +2828,7 @@ bool UPrimitiveComponent::ConvertSweptOverlapsToCurrentOverlaps(
 			// We know we are not overlapping any new components at the end location. Children are ignored here (see note below).
 			if (PrimitiveComponentCVars::bEnableFastOverlapCheck)
 			{
-				SCOPE_CYCLE_COUNTER(STAT_MoveComponent_FastOverlap);
+				//SCOPE_CYCLE_COUNTER(STAT_MoveComponent_FastOverlap);
 
 				// Check components we hit during the sweep, keep only those still overlapping
 				const FCollisionQueryParams UnusedQueryParams(NAME_None, FCollisionQueryParams::GetUnknownStatId());
