Fix non-unity Stat issues and Debug build compile issues caused by CL#23342352. Reenables STAT_MoveComponent_FastOverlap.

diff --git a/Engine/Source/Runtime/Engine/Classes/Components/PrimitiveComponent.h b/Engine/Source/Runtime/Engine/Classes/Components/PrimitiveComponent.h
--- a/Engine/Source/Runtime/Engine/Classes/Components/PrimitiveComponent.h
+++ b/Engine/Source/Runtime/Engine/Classes/Components/PrimitiveComponent.h
@@ -26,11 +26,15 @@
 #include "HitProxies.h"
 #include "Interfaces/Interface_AsyncCompilation.h"
 #include "HLOD/HLODBatchingPolicy.h"
+#include "Stats/Stats2.h"
 #if UE_ENABLE_INCLUDE_ORDER_DEPRECATED_IN_5_1
 #include "Engine/OverlapInfo.h"
 #endif
 #include "PrimitiveComponent.generated.h"
 
+DECLARE_CYCLE_STAT_EXTERN(TEXT("BeginComponentOverlap"), STAT_BeginComponentOverlap, STATGROUP_Game, ENGINE_API);
+DECLARE_CYCLE_STAT_EXTERN(TEXT("MoveComponent FastOverlap"), STAT_MoveComponent_FastOverlap, STATGROUP_Game, ENGINE_API);
+
 class AController;
 class FPrimitiveSceneProxy;
 class UMaterialInterface;
@@ -172,6 +176,24 @@ struct FRendererStencilMaskEvaluation
 	}
 };
 
+// Predicate to determine if an overlap is with a certain AActor.
+struct FPredicateOverlapHasSameActor
+{
+	FPredicateOverlapHasSameActor(const AActor& Owner)
+		: MyOwnerPtr(&Owner)
+	{
+	}
+
+	bool operator() (const FOverlapInfo& Info)
+	{
+		// MyOwnerPtr is always valid, so we don't need the IsValid() checks in the WeakObjectPtr comparison operator.
+		return MyOwnerPtr.HasSameIndexAndSerialNumber(Info.OverlapInfo.HitObjectHandle.FetchActor());
+	}
+
+private:
+	const TWeakObjectPtr<const AActor> MyOwnerPtr;
+};
+
 // Predicate to determine if an overlap is *NOT* with a certain AActor.
 struct FPredicateOverlapHasDifferentActor
 {
@@ -2828,7 +2850,7 @@ bool UPrimitiveComponent::ConvertSweptOverlapsToCurrentOverlaps(
 			// We know we are not overlapping any new components at the end location. Children are ignored here (see note below).
 			if (PrimitiveComponentCVars::bEnableFastOverlapCheck)
 			{
-				//SCOPE_CYCLE_COUNTER(STAT_MoveComponent_FastOverlap);
+				SCOPE_CYCLE_COUNTER(STAT_MoveComponent_FastOverlap);
 
 				// Check components we hit during the sweep, keep only those still overlapping
 				const FCollisionQueryParams UnusedQueryParams(NAME_None, FCollisionQueryParams::GetUnknownStatId());
diff --git a/Engine/Source/Runtime/Engine/Private/Components/PrimitiveComponent.cpp b/Engine/Source/Runtime/Engine/Private/Components/PrimitiveComponent.cpp
--- a/Engine/Source/Runtime/Engine/Private/Components/PrimitiveComponent.cpp
+++ b/Engine/Source/Runtime/Engine/Private/Components/PrimitiveComponent.cpp
@@ -116,8 +116,9 @@ FAutoConsoleVariableRef CVarRefShowInitialOverlaps(
 	ECVF_Cheat);
 #endif // !(UE_BUILD_SHIPPING || UE_BUILD_TEST)
 
-DECLARE_CYCLE_STAT(TEXT("MoveComponent FastOverlap"), STAT_MoveComponent_FastOverlap, STATGROUP_Game);
-DECLARE_CYCLE_STAT(TEXT("BeginComponentOverlap"), STAT_BeginComponentOverlap, STATGROUP_Game);
+DEFINE_STAT(STAT_BeginComponentOverlap);
+DEFINE_STAT(STAT_MoveComponent_FastOverlap);
+
 DECLARE_CYCLE_STAT(TEXT("EndComponentOverlap"), STAT_EndComponentOverlap, STATGROUP_Game);
 DECLARE_CYCLE_STAT(TEXT("PrimComp DispatchBlockingHit"), STAT_DispatchBlockingHit, STATGROUP_Game);
 
@@ -129,24 +130,6 @@ FOverlapInfo::FOverlapInfo(UPrimitiveComponent* InComponent, int32 InBodyIndex)
 	OverlapInfo.Item = InBodyIndex;
 }
 
-// Predicate to determine if an overlap is with a certain AActor.
-struct FPredicateOverlapHasSameActor
-{
-	FPredicateOverlapHasSameActor(const AActor& Owner)
-	: MyOwnerPtr(&Owner)
-	{
-	}
-
-	bool operator() (const FOverlapInfo& Info)
-	{
-		// MyOwnerPtr is always valid, so we don't need the IsValid() checks in the WeakObjectPtr comparison operator.
-		return MyOwnerPtr.HasSameIndexAndSerialNumber(Info.OverlapInfo.HitObjectHandle.FetchActor());
-	}
-
-private:
-	const TWeakObjectPtr<const AActor> MyOwnerPtr;
-};
-
 // Helper for finding the index of an FOverlapInfo in an Array using the FFastOverlapInfoCompare predicate, knowing that at least one overlap is valid (non-null).
 template<class AllocatorType>
 FORCEINLINE_DEBUGGABLE int32 IndexOfOverlapFast(const TArray<FOverlapInfo, AllocatorType>& OverlapArray, const FOverlapInfo& SearchItem)
