Improved HLOD warmup to support warmup from HLOD1 to HLOD0
* Previously, warmup was only performed on cell unload
* With these changes, HLOD warmup also occurs on cell loading, so going from HLOD1 to HLOD0 will trigger it
* Added fast path to skip warmup under certain conditions (if VT/nanite are either disabled or unsupported)

diff --git a/Engine/Source/Runtime/Engine/Classes/Engine/LevelStreaming.h b/Engine/Source/Runtime/Engine/Classes/Engine/LevelStreaming.h
--- a/Engine/Source/Runtime/Engine/Classes/Engine/LevelStreaming.h
+++ b/Engine/Source/Runtime/Engine/Classes/Engine/LevelStreaming.h
@@ -350,6 +350,11 @@ protected:
 	/** Updates the current state of the streaming level and notifies any observers. */
 	void SetCurrentState(ELevelStreamingState NewState);
 
+	/** Returns whether the streaming level can make visible (can call AddToWorld). */
+	virtual bool CanMakeVisible();
+	/** Returns whether the streaming level can make invisible (can call RemoveFromWorld). */
+	virtual bool CanMakeInvisible();
+
 private:
 
 	/** If true client will wait for acknowledgment from server before making streaming levels invisible */
@@ -360,10 +365,6 @@ private:
 	bool ShouldWaitForServerAckBeforeChangingVisibilityState(ENetLevelVisibilityRequest InRequestType, bool bInShouldBeVisible);
 	/** Ack a client instigated visibility/streaming transaction */
 	void AckNetVisibilityTransaction(FNetLevelVisibilityTransactionId AckedClientTransactionId, bool bClientAckCanMakeVisible);
-	/** Returns whether the streaming level can make visible (can call AddToWorld). */
-	bool CanMakeVisible();
-	/** Returns whether the streaming level can make invisible (can call RemoveFromWorld). */
-	bool CanMakeInvisible();
 
 	/** Determine what the streaming level's target state should be. */
 	ELevelStreamingTargetState DetermineTargetState() const;
diff --git a/Engine/Source/Runtime/Engine/Private/WorldPartition/HLOD/HLODSubsystem.cpp b/Engine/Source/Runtime/Engine/Private/WorldPartition/HLOD/HLODSubsystem.cpp
--- a/Engine/Source/Runtime/Engine/Private/WorldPartition/HLOD/HLODSubsystem.cpp
+++ b/Engine/Source/Runtime/Engine/Private/WorldPartition/HLOD/HLODSubsystem.cpp
@@ -33,40 +33,56 @@
 
 DEFINE_LOG_CATEGORY_STATIC(LogHLODSubsystem, Log, All);
 
-static TAutoConsoleVariable<int32> CVarHLODWarmupNumFrames(
-	TEXT("wp.Runtime.HLOD.WarmupNumFrames"),
-	5,
-	TEXT("Delay unloading of a cell for this amount of frames to ensure HLOD assets are ready to be shown at the proper resolution. Set to 0 to force disable warmup."),
-	ECVF_Default
-);
-
 static TAutoConsoleVariable<int32> CVarHLODWarmupEnabled(
 	TEXT("wp.Runtime.HLOD.WarmupEnabled"),
 	1,
-	TEXT("Enable HLOD assets warmup. Will delay unloading of cells & transition to HLODs for wp.Runtime.HLOD.WarmupNumFrames frames."),
-	ECVF_Default
-);
+	TEXT("Enable HLOD assets warmup. Will delay unloading of cells & transition to HLODs for wp.Runtime.HLOD.WarmupNumFrames frames."));
+
+static TAutoConsoleVariable<int32> CVarHLODWarmupVT(
+	TEXT("wp.Runtime.HLOD.WarmupVT"),
+	1,
+	TEXT("Enable virtual texture warmup for HLOD assets. Requires wp.Runtime.HLOD.WarmupEnabled to be 1."));
+
+static TAutoConsoleVariable<int32> CVarHLODWarmupNanite(
+	TEXT("wp.Runtime.HLOD.WarmupNanite"),
+	1,
+	TEXT("Enable Nanite warmup for HLOD assets. Requires wp.Runtime.HLOD.WarmupEnabled to be 1."));
+
+static TAutoConsoleVariable<int32> CVarHLODWarmupNumFrames(
+	TEXT("wp.Runtime.HLOD.WarmupNumFrames"),
+	5,
+	TEXT("Delay unloading of a cell for this amount of frames to ensure HLOD assets are ready to be shown at the proper resolution. Set to 0 to force disable warmup."));
 
 static TAutoConsoleVariable<int32> CVarHLODWarmupDebugDraw(
 	TEXT("wp.Runtime.HLOD.WarmupDebugDraw"),
 	0,
-	TEXT("Draw debug display for the warmup requests"),
-	ECVF_Default
-);
+	TEXT("Draw debug display for the warmup requests"));
 
 static TAutoConsoleVariable<float> CVarHLODWarmupVTScaleFactor(
 	TEXT("wp.Runtime.HLOD.WarmupVTScaleFactor"),
 	2.0f,
-	TEXT("Scale the VT size we ask to prefetch by this factor."),
-	ECVF_Default
-);
+	TEXT("Scale the VT size we ask to prefetch by this factor."));
 
 static TAutoConsoleVariable<int32> CVarHLODWarmupVTSizeClamp(
 	TEXT("wp.Runtime.HLOD.WarmupVTSizeClamp"),
 	2048,
-	TEXT("Clamp VT warmup requests for safety."),
-	ECVF_Default
-);
+	TEXT("Clamp VT warmup requests for safety."));
+
+static void HLODSubsystemCVarSinkFunction()
+{
+	for (UWorld* World : TObjectRange<UWorld>(RF_ClassDefaultObject | RF_ArchetypeObject, true, EInternalObjectFlags::Garbage))
+	{
+		if (World->WorldType == EWorldType::Game || World->WorldType == EWorldType::PIE)
+		{
+			if (UHLODSubsystem* HLODSubsystem = World->GetSubsystem<UHLODSubsystem>())
+			{
+				HLODSubsystem->OnCVarsChanged();
+			}
+		}
+	}
+}
+
+static FAutoConsoleVariableSink CVarHLODSink(FConsoleCommandDelegate::CreateStatic(&HLODSubsystemCVarSinkFunction));
 
 namespace FHLODSubsystem
 {
@@ -114,6 +130,7 @@ public:
 
 UHLODSubsystem::UHLODSubsystem()
 	: UWorldSubsystem()
+	, bCachedShouldPerformWarmup(true)
 {
 }
 
@@ -216,6 +233,8 @@ void UHLODSubsystem::Initialize(FSubsystemCollectionBase& Collection)
 	GetWorld()->OnWorldPartitionInitialized().AddUObject(this, &UHLODSubsystem::OnWorldPartitionInitialized);
 	GetWorld()->OnWorldPartitionUninitialized().AddUObject(this, &UHLODSubsystem::OnWorldPartitionUninitialized);
 
+	bCachedShouldPerformWarmup = ShouldPerformWarmup();
+
 	SceneViewExtension = FSceneViewExtensions::NewExtension<FHLODResourcesResidencySceneViewExtension>(World);
 }
 
@@ -256,42 +275,33 @@ void UHLODSubsystem::OnWorldPartitionUninitialized(UWorldPartition* InWorldParti
 	}
 }
 
-void UHLODSubsystem::RegisterHLODActor(AWorldPartitionHLOD* InWorldPartitionHLOD)
+const UHLODSubsystem::FCellData* UHLODSubsystem::GetCellData(const UWorldPartitionRuntimeCell* InCell) const
 {
-	TRACE_CPUPROFILER_EVENT_SCOPE(UHLODSubsystem::RegisterHLODActor);
+	return const_cast<UHLODSubsystem*>(this)->GetCellData(InCell);
+}
 
-	const UWorldPartition* WorldPartition = FHLODSubsystem::GetWorldPartition(InWorldPartitionHLOD);
+UHLODSubsystem::FCellData* UHLODSubsystem::GetCellData(const UWorldPartitionRuntimeCell* InCell)
+{
+	// @todo_ow ContentBundles do not support Hlods
+	if (InCell->GetContentBundleID().IsValid())
+	{
+		return nullptr;
+	}
+
+	const UWorldPartition* WorldPartition = InCell->GetCellOwner()->GetOuterWorld()->GetWorldPartition();
 	check(WorldPartition && WorldPartition->IsStreamingEnabled());
 
 	FWorldPartitionHLODRuntimeData* WorldPartitionHLODRuntimeData = WorldPartitionsHLODRuntimeData.Find(WorldPartition);
 	if (WorldPartitionHLODRuntimeData)
 	{
-		const FName CellName = InWorldPartitionHLOD->GetSourceCellName();
-		FCellData* CellData = WorldPartitionHLODRuntimeData->CellsData.Find(CellName);
-
-#if WITH_EDITOR
-		UE_LOG(LogHLODSubsystem, Verbose, TEXT("Registering HLOD %s (%s) for cell %s"), *InWorldPartitionHLOD->GetActorLabel(), *InWorldPartitionHLOD->GetActorGuid().ToString(), *CellName.ToString());
-#endif
-
-		if (CellData)
-		{
-			CellData->LoadedHLODs.Add(InWorldPartitionHLOD);
-			InWorldPartitionHLOD->SetVisibility(UHLODSubsystem::WorldPartitionHLODEnabled && !CellData->bIsCellVisible);
-		}
-		else
-		{
-			UE_LOG(LogHLODSubsystem, Verbose, TEXT("Found HLOD referencing nonexistent cell '%s'"), *CellName.ToString());
-			InWorldPartitionHLOD->SetVisibility(false);
-		}
+		return WorldPartitionHLODRuntimeData->CellsData.Find(InCell->GetFName());
 	}
 
-	HLODActorRegisteredEvent.Broadcast(InWorldPartitionHLOD);
+	return nullptr;
 }
 
-void UHLODSubsystem::UnregisterHLODActor(AWorldPartitionHLOD* InWorldPartitionHLOD)
+UHLODSubsystem::FCellData* UHLODSubsystem::GetCellData(AWorldPartitionHLOD* InWorldPartitionHLOD)
 {
-	TRACE_CPUPROFILER_EVENT_SCOPE(UHLODSubsystem::UnregisterHLODActor);
-
 	const UWorldPartition* WorldPartition = FHLODSubsystem::GetWorldPartition(InWorldPartitionHLOD);
 	check(WorldPartition && WorldPartition->IsStreamingEnabled());
 
@@ -299,33 +309,17 @@ void UHLODSubsystem::UnregisterHLODActor(AWorldPartitionHLOD* InWorldPartitionHL
 	if (WorldPartitionHLODRuntimeData)
 	{
 		const FName CellName = InWorldPartitionHLOD->GetSourceCellName();
-		FCellData* CellData = WorldPartitionHLODRuntimeData->CellsData.Find(CellName);
-		if (CellData)
-		{
-#if WITH_EDITOR
-			UE_LOG(LogHLODSubsystem, Verbose, TEXT("Unregistering HLOD %s (%s) for cell %s"), *InWorldPartitionHLOD->GetActorLabel(), *InWorldPartitionHLOD->GetActorGuid().ToString(), *CellName.ToString());
-#endif
-
-			int32 NumRemoved = CellData->LoadedHLODs.Remove(InWorldPartitionHLOD);
-			check(NumRemoved == 1);
-		}
+		return WorldPartitionHLODRuntimeData->CellsData.Find(CellName);
 	}
 
-	HLODActorUnregisteredEvent.Broadcast(InWorldPartitionHLOD);
+	return nullptr;		
 }
 
 const TArray<AWorldPartitionHLOD*>& UHLODSubsystem::GetHLODActorsForCell(const UWorldPartitionRuntimeCell* InCell) const
 {
-	const UWorldPartition* WorldPartition = InCell->GetWorld()->GetWorldPartition();
-
-	const FWorldPartitionHLODRuntimeData* WorldPartitionHLODRuntimeData = WorldPartitionsHLODRuntimeData.Find(WorldPartition);
-	if (WorldPartitionHLODRuntimeData)
+	if (const FCellData* CellData = GetCellData(InCell))
 	{
-		const FCellData* CellData = WorldPartitionHLODRuntimeData->CellsData.Find(InCell->GetFName());
-		if (CellData)
-		{
-			return CellData->LoadedHLODs;
-		}
+		return CellData->LoadedHLODs;
 	}
 
 	// No HLOD found for the given cell, return a dummy array
@@ -333,70 +327,82 @@ const TArray<AWorldPartitionHLOD*>& UHLODSubsystem::GetHLODActorsForCell(const U
 	return DummyArray;
 }
 
-void UHLODSubsystem::OnCellShown(const UWorldPartitionRuntimeCell* InCell)
+void UHLODSubsystem::RegisterHLODActor(AWorldPartitionHLOD* InWorldPartitionHLOD)
 {
-	// @todo_ow ContentBundles do not support Hlods
-	if (InCell->GetContentBundleID().IsValid())
+	TRACE_CPUPROFILER_EVENT_SCOPE(UHLODSubsystem::RegisterHLODActor);
+	
+	if (FCellData* CellData = GetCellData(InWorldPartitionHLOD))
 	{
-		return;
+#if WITH_EDITOR
+		UE_LOG(LogHLODSubsystem, Verbose, TEXT("Registering HLOD %s (%s) for cell %s"), *InWorldPartitionHLOD->GetActorLabel(), *InWorldPartitionHLOD->GetActorGuid().ToString(), *InWorldPartitionHLOD->GetSourceCellName().ToString());
+#endif
+		CellData->LoadedHLODs.Add(InWorldPartitionHLOD);
+		InWorldPartitionHLOD->SetVisibility(UHLODSubsystem::WorldPartitionHLODEnabled && !CellData->bIsCellVisible);
+	}
+	else
+	{
+		UE_LOG(LogHLODSubsystem, Verbose, TEXT("Found HLOD referencing nonexistent cell '%s'"), *InWorldPartitionHLOD->GetSourceCellName().ToString());
+		InWorldPartitionHLOD->SetVisibility(false);
 	}
 
-	const UWorldPartition* WorldPartition = InCell->GetCellOwner()->GetOuterWorld()->GetWorldPartition();
-	check(WorldPartition && WorldPartition->IsStreamingEnabled());
+	HLODActorsToWarmup.Remove(InWorldPartitionHLOD);
+	HLODActorRegisteredEvent.Broadcast(InWorldPartitionHLOD);
+}
 
-	FWorldPartitionHLODRuntimeData* WorldPartitionHLODRuntimeData = WorldPartitionsHLODRuntimeData.Find(WorldPartition);
-	if (WorldPartitionHLODRuntimeData)
+void UHLODSubsystem::UnregisterHLODActor(AWorldPartitionHLOD* InWorldPartitionHLOD)
+{
+	TRACE_CPUPROFILER_EVENT_SCOPE(UHLODSubsystem::UnregisterHLODActor);
+
+	if (FCellData* CellData = GetCellData(InWorldPartitionHLOD))
 	{
-		FCellData* CellData = WorldPartitionHLODRuntimeData->CellsData.Find(InCell->GetFName());
-		if (CellData)
-		{
-			CellData->bIsCellVisible = true;
+#if WITH_EDITOR
+		UE_LOG(LogHLODSubsystem, Verbose, TEXT("Unregistering HLOD %s (%s) for cell %s"), *InWorldPartitionHLOD->GetActorLabel(), *InWorldPartitionHLOD->GetActorGuid().ToString(), *InWorldPartitionHLOD->GetSourceCellName().ToString());
+#endif
+
+		int32 NumRemoved = CellData->LoadedHLODs.Remove(InWorldPartitionHLOD);
+		check(NumRemoved == 1);
+	}
+
+	HLODActorUnregisteredEvent.Broadcast(InWorldPartitionHLOD);
+}
+
+void UHLODSubsystem::OnCellShown(const UWorldPartitionRuntimeCell* InCell)
+{
+	if (FCellData* CellData = GetCellData(InCell))
+	{
+		CellData->bIsCellVisible = true;
 
 #if WITH_EDITOR
-			UE_LOG(LogHLODSubsystem, Verbose, TEXT("Cell shown - %s - hiding %d HLOD actors"), *InCell->GetName(), CellData->LoadedHLODs.Num());
+		UE_LOG(LogHLODSubsystem, Verbose, TEXT("Cell shown - %s - hiding %d HLOD actors"), *InCell->GetName(), CellData->LoadedHLODs.Num());
 #endif
 
-			for (AWorldPartitionHLOD* HLODActor : CellData->LoadedHLODs)
-			{
+		for (AWorldPartitionHLOD* HLODActor : CellData->LoadedHLODs)
+		{
 #if WITH_EDITOR
-				UE_LOG(LogHLODSubsystem, Verbose, TEXT("\t\t%s - %s"), *HLODActor->GetActorLabel(), *HLODActor->GetActorGuid().ToString());
+			UE_LOG(LogHLODSubsystem, Verbose, TEXT("\t\t%s - %s"), *HLODActor->GetActorLabel(), *HLODActor->GetActorGuid().ToString());
 #endif
-				HLODActor->SetVisibility(false);
-			}
+			HLODActor->SetVisibility(false);
 		}
 	}
 }
 
 void UHLODSubsystem::OnCellHidden(const UWorldPartitionRuntimeCell* InCell)
 {
-	// @todo_ow ContentBundles do not support Hlods
-	if (InCell->GetContentBundleID().IsValid())
-	{
-		return;
-	}
-
-	const UWorldPartition* WorldPartition = InCell->GetCellOwner()->GetOuterWorld()->GetWorldPartition();
-	check(WorldPartition && WorldPartition->IsStreamingEnabled());
-
-	FWorldPartitionHLODRuntimeData* WorldPartitionHLODRuntimeData = WorldPartitionsHLODRuntimeData.Find(WorldPartition);
-	if (WorldPartitionHLODRuntimeData)
+	if (FCellData* CellData = GetCellData(InCell))
 	{
-		FCellData* CellData = WorldPartitionHLODRuntimeData->CellsData.Find(InCell->GetFName());
-		if (CellData)
-		{
-			CellData->bIsCellVisible = false;
+		CellData->bIsCellVisible = false;
 
 #if WITH_EDITOR
-			UE_LOG(LogHLODSubsystem, Verbose, TEXT("Cell hidden - %s - showing %d HLOD actors"), *InCell->GetName(), CellData->LoadedHLODs.Num());
+		UE_LOG(LogHLODSubsystem, Verbose, TEXT("Cell hidden - %s - showing %d HLOD actors"), *InCell->GetName(), CellData->LoadedHLODs.Num());
 #endif
 
-			for (AWorldPartitionHLOD* HLODActor : CellData->LoadedHLODs)
-			{
+		for (AWorldPartitionHLOD* HLODActor : CellData->LoadedHLODs)
+		{
 #if WITH_EDITOR
-				UE_LOG(LogHLODSubsystem, Verbose, TEXT("\t\t%s - %s"), *HLODActor->GetActorLabel(), *HLODActor->GetActorGuid().ToString());
+			UE_LOG(LogHLODSubsystem, Verbose, TEXT("\t\t%s - %s"), *HLODActor->GetActorLabel(), *HLODActor->GetActorGuid().ToString());
 #endif
-				HLODActor->SetVisibility(UHLODSubsystem::WorldPartitionHLODEnabled);
-			}
+			HLODActor->SetVisibility(UHLODSubsystem::WorldPartitionHLODEnabled);
+			HLODActorsToWarmup.Remove(HLODActor);
 		}
 	}
 }
@@ -425,174 +431,279 @@ static void PrepareNaniteRequests(TSet<Nanite::FResources*>& InOutNaniteRequests
 	UStaticMesh* StaticMesh = InStaticMeshComponent->GetStaticMesh();
 	if (StaticMesh && StaticMesh->HasValidNaniteData())
 	{
-		// UE_LOG(LogHLODSubsystem, Warning, TEXT("NanitePrefetch: %s, %d pages"), *StaticMesh->GetFullName(), StaticMesh->GetRenderData()->NaniteResources.PageStreamingStates.Num());
 		InOutNaniteRequests.Add(&StaticMesh->GetRenderData()->NaniteResources);
 	}
 }
 
-static float EstimateScreenSize(UStaticMeshComponent* InStaticMeshComponent, const FSceneViewFamily& InViewFamily)
+bool UHLODSubsystem::PrepareToWarmup(const UWorldPartitionRuntimeCell* InCell, AWorldPartitionHLOD* InHLODActor)
 {
-	float MaxPixels = 0;
+	TRACE_CPUPROFILER_EVENT_SCOPE(UHLODSubsystem::PrepareToWarmup)
 
-	// Estimate the highest screen pixel size of this component in the provided views
-	for (const FSceneView* View : InViewFamily.Views)
+	bool bHLODActorNeedsWarmUp = false;
+
+	if (InHLODActor->DoesRequireWarmup())
 	{
-		// Make sure the HLOD actor we're about to show is actually in the frustum
-		if (View->ViewFrustum.IntersectSphere(InStaticMeshComponent->Bounds.Origin, InStaticMeshComponent->Bounds.SphereRadius))
-		{
-			float ScreenDiameter = ComputeBoundsScreenSize(InStaticMeshComponent->Bounds.Origin, InStaticMeshComponent->Bounds.SphereRadius, *View);
-			float PixelSize = ScreenDiameter * View->ViewMatrices.GetScreenScale() * 2.0f;
+		FWorldPartitionHLODWarmupState& WarmupState = HLODActorsToWarmup.FindOrAdd(InHLODActor);
+
+		// In case a previous request to unload was aborted and the cell never unloaded... assume warmup requests are expired after a given amount of frames.
+		const uint32 WarmupExpiredFrames = 30;
+		const uint32 CurrentFrameNumber = GetWorld()->Scene->GetFrameNumber();
 
-			MaxPixels = FMath::Max(MaxPixels, PixelSize);
+		// Trigger warmup on the first request to unload, or if a warmup request expired
+		const bool bInitiateWarmup = WarmupState.WarmupEndFrame == INDEX_NONE || CurrentFrameNumber > (WarmupState.WarmupEndFrame + WarmupExpiredFrames);
+		const bool bWarmupCompleted = !bInitiateWarmup && CurrentFrameNumber >= WarmupState.WarmupEndFrame;
+		
+		if (bInitiateWarmup)
+		{
+			// Warmup will be triggered in the next BeginRenderView() call, at which point the frame number will have been incremented.
+			WarmupState.WarmupStartFrame = CurrentFrameNumber + 1;
+			WarmupState.WarmupEndFrame = WarmupState.WarmupStartFrame + CVarHLODWarmupNumFrames.GetValueOnGameThread();
+			WarmupState.Location = InCell->GetCellBounds().GetCenter();
 		}
+
+		bHLODActorNeedsWarmUp = !bWarmupCompleted;
 	}
 
-	return MaxPixels;
+	return bHLODActorNeedsWarmUp;
 }
 
-void UHLODSubsystem::MakeRenderResourcesResident(const FCellData& CellData, const FSceneViewFamily& InViewFamily)
+void UHLODSubsystem::OnCVarsChanged()
 {
-	TRACE_CPUPROFILER_EVENT_SCOPE(UHLODSubsystem::MakeRenderResourcesResident)
+	bCachedShouldPerformWarmup = ShouldPerformWarmup();
+}
 
-	TMap<UMaterialInterface*, float> VTRequests;
-	TSet<Nanite::FResources*> NaniteRequests;
+bool UHLODSubsystem::ShouldPerformWarmup() const
+{
+	// Test if warmup is disabled globally.
+	const bool bWarmupEnabled = CVarHLODWarmupEnabled.GetValueOnGameThread() != 0;
+	if (!bWarmupEnabled)
+	{
+		return false;
+	}
 
-	// For each HLOD actor representing this cell
-	for(AWorldPartitionHLOD* HLODActor : CellData.LoadedHLODs)
+	// If warmup num of frames is invalid, no warmup needed
+	if (CVarHLODWarmupNumFrames.GetValueOnGameThread() <= 0)
 	{
-		// Skip HLOD actors that doesn't require warmup.
-		// For example, instanced HLODs, as they reuse the same meshes/textures as their source actors.
-		// These resources should already be resident & at the proper resolution.
-		if (!HLODActor->DoesRequireWarmup())
-		{
-			continue;
-		}
+		return false;
+	}
 
-		HLODActor->ForEachComponent<UStaticMeshComponent>(false, [&](UStaticMeshComponent* SMC)
-		{
-			float PixelSize = EstimateScreenSize(SMC, InViewFamily);
+	// If warmup num of frames is invalid, no warmup needed
+	const EShaderPlatform ShaderPlatform = GShaderPlatformForFeatureLevel[GMaxRHIFeatureLevel];
+	const bool bNaniteEnabled = UseNanite(ShaderPlatform);
+	const bool bVirtualTextureEnabled = UseVirtualTexturing(ShaderPlatform);
+	const bool bWarmupNanite = CVarHLODWarmupNanite.GetValueOnGameThread() != 0;
+	const bool bWarmupVT = CVarHLODWarmupVT.GetValueOnGameThread() != 0;
+	const bool bWarmupNeeded = (bNaniteEnabled && bWarmupNanite) || (bVirtualTextureEnabled && bWarmupVT);
+	if (!bWarmupNeeded)
+	{
+		return false;
+	}
 
-			if (PixelSize > 0)
-			{
-				PrepareVTRequests(VTRequests, SMC, PixelSize);
+	// If we're running a dedicated server, no warmup needed
+	const bool bIsDedicatedServer = GetWorld()->GetNetMode() == NM_DedicatedServer;
+	if (bIsDedicatedServer)
+	{
+		return false;
+	}
 
-				// Only issue Nanite requests on the first warmup frame
-				if (CellData.WarmupStartFrame == InViewFamily.FrameNumber)
-				{
-					PrepareNaniteRequests(NaniteRequests, SMC);
-				}
+	return true;
+}
 
-#if ENABLE_DRAW_DEBUG
-				if (CVarHLODWarmupDebugDraw.GetValueOnAnyThread())
-				{
-					const FBox& Box = SMC->CalcLocalBounds().GetBox();
-					DrawDebugBox(HLODActor->GetWorld(), Box.GetCenter(), Box.GetExtent(), FColor::Yellow, /*bPersistentLine*/ false, /*Lifetime*/ 1.0f);
-				}
-#endif
-			}
-		});
+bool UHLODSubsystem::ShouldPerformWarmupForCell(const UWorldPartitionRuntimeCell* InCell) const
+{
+	if (!bCachedShouldPerformWarmup)
+	{
+		return false;
 	}
 
-	if (!VTRequests.IsEmpty() || !NaniteRequests.IsEmpty())
-	{
-		ENQUEUE_RENDER_COMMAND(MakeHLODRenderResourcesResident)(
-			[VTRequests = MoveTemp(VTRequests), NaniteRequests = MoveTemp(NaniteRequests), FeatureLevel = InViewFamily.GetFeatureLevel()](FRHICommandListImmediate& RHICmdList)
-			{
-				for (const TPair<UMaterialInterface*, float>& VTRequest : VTRequests)
-				{
-					UMaterialInterface* Material = VTRequest.Key;
-					FMaterialRenderProxy* MaterialRenderProxy = Material->GetRenderProxy();
+	const UWorld* World = GetWorld();
 
-					GetRendererModule().RequestVirtualTextureTiles(MaterialRenderProxy, FVector2D(VTRequest.Value, VTRequest.Value), FeatureLevel);
-				}
+	// Blocking loading shouldn't trigger warmup
+	const bool bIsInBlockingLoading = World->GetIsInBlockTillLevelStreamingCompleted();
+	if (bIsInBlockingLoading)
+	{
+		return false;
+	}
 
-				const uint32 NumFramesBeforeRender = CVarHLODWarmupNumFrames.GetValueOnRenderThread();
-				for (const Nanite::FResources* Resource : NaniteRequests)
-				{
-					GetRendererModule().PrefetchNaniteResource(Resource, NumFramesBeforeRender);
-				}
-			});
+	// If processing for this cell in Add/RemoveFromWorld() has already started, it's too late for warmup
+	const ULevel* CellLevel = InCell->GetLevel();
+	const bool bCurrentlyProcessingLevel = CellLevel == World->GetCurrentLevelPendingVisibility() ||
+										   CellLevel == World->GetCurrentLevelPendingInvisibility();
+	if (bCurrentlyProcessingLevel)
+	{
+		return false;
 	}
+
+	return true;
 }
 
-bool UHLODSubsystem::RequestUnloading(const UWorldPartitionRuntimeCell* InCell)
+bool UHLODSubsystem::CanMakeVisible(const UWorldPartitionRuntimeCell* InCell)
 {
-	// Test if warmup is disabled globally.
-	bool bPerformWarmpup = (CVarHLODWarmupEnabled.GetValueOnGameThread() != 0) && (CVarHLODWarmupNumFrames.GetValueOnGameThread() > 0) && (GetWorld()->GetNetMode() != NM_DedicatedServer);
-	if (!bPerformWarmpup)
-	{
-		return true;
-	}
+	TRACE_CPUPROFILER_EVENT_SCOPE(UHLODSubsystem::CanMakeVisible)
 
-	const UWorldPartition* WorldPartition = InCell->GetCellOwner()->GetOuterWorld()->GetWorldPartition();;
-	FWorldPartitionHLODRuntimeData* WorldPartitionHLODRuntimeData = WorldPartitionsHLODRuntimeData.Find(WorldPartition);
-	if (!WorldPartitionHLODRuntimeData)
+	if (!ShouldPerformWarmupForCell(InCell))
 	{
 		return true;
 	}
 
-	FCellData* CellData = WorldPartitionHLODRuntimeData->CellsData.Find(InCell->GetFName());
-	if (!CellData)
+	bool bCanMakeVisible = true;
+
+	// Prevent cells containing HLODs actors from being made visible until warmup has been performed
+	if (InCell->GetIsHLOD() && InCell->GetLevel())
 	{
-		return true;
+		for (AActor* Actor : InCell->GetLevel()->Actors)
+		{
+			if (AWorldPartitionHLOD* HLODActor = Cast<AWorldPartitionHLOD>(Actor))
+			{
+				const bool bHLODActorNeedsWarmup = PrepareToWarmup(InCell, HLODActor);
+				bCanMakeVisible &= !bHLODActorNeedsWarmup;
+			}
+		}
 	}
 
-	// If cell wasn't even visible yet or doesn't have HLOD actors, skip warmup.
-	bPerformWarmpup = !CellData->LoadedHLODs.IsEmpty() && CellData->bIsCellVisible;
-	if (!bPerformWarmpup)
+	return bCanMakeVisible;
+}
+
+bool UHLODSubsystem::CanMakeInvisible(const UWorldPartitionRuntimeCell* InCell)
+{
+	TRACE_CPUPROFILER_EVENT_SCOPE(UHLODSubsystem::CanMakeInvisible)
+
+	if (!ShouldPerformWarmupForCell(InCell))
 	{
 		return true;
 	}
 
-	// Verify that at least one HLOD actor associated with this cell needs warmup.
-	bPerformWarmpup = Algo::AnyOf(CellData->LoadedHLODs, [](const AWorldPartitionHLOD* HLODActor) { return HLODActor->DoesRequireWarmup(); });
-	if (!bPerformWarmpup)
+	bool bCanMakeInvisible = true;
+
+	// Prevent cells which have an HLOD representation from being hidden until their matching HLOD actors have been warmed up
+	if (FCellData* CellData = GetCellData(InCell))
 	{
-		return true;
+		for (AWorldPartitionHLOD* HLODActor : CellData->LoadedHLODs)
+		{
+			const bool bHLODActorNeedsWarmup = PrepareToWarmup(InCell, HLODActor);
+			bCanMakeInvisible &= !bHLODActorNeedsWarmup;
+		}
 	}
 
-	// In case a previous request to unload was aborted and the cell never unloaded... assume warmup requests are expired after a given amount of frames.
-	const uint32 WarmupExpiredFrames = 30;
+	return bCanMakeInvisible;
+}
 
-	uint32 CurrentFrameNumber = GetWorld()->Scene->GetFrameNumber();
+static float IsInView(const FVector& BoundsOrigin, const FVector& BoundsExtent, const FSceneViewFamily& InViewFamily, bool bComputeScreenSize, float& MaxScreenSizePixels)
+{
+	MaxScreenSizePixels = 0;
 
-	// Trigger warmup on the first request to unload, or if a warmup request expired
-	if (CellData->WarmupEndFrame == INDEX_NONE || CurrentFrameNumber > (CellData->WarmupEndFrame + WarmupExpiredFrames))
+	// Estimate the highest screen pixel size of this component in the provided views
+	for (const FSceneView* View : InViewFamily.Views)
 	{
-		// Warmup will be triggered in the next BeginRenderView() call, at which point the frame number will have been incremented.
-		CellData->WarmupStartFrame = CurrentFrameNumber + 1; 
-		CellData->WarmupEndFrame = CellData->WarmupStartFrame + CVarHLODWarmupNumFrames.GetValueOnGameThread();
-		WorldPartitionHLODRuntimeData->CellsToWarmup.Add(CellData);
-	}
+		// Make sure the HLOD actor we're about to show is actually in the frustum
+		if (View->ViewFrustum.IntersectBox(BoundsOrigin, BoundsExtent))
+		{
+			if (bComputeScreenSize)
+			{
+				float ScreenDiameter = ComputeBoundsScreenSize(BoundsOrigin, BoundsExtent.Size(), *View);
+				float ScreenSizePixels = ScreenDiameter * View->ViewMatrices.GetScreenScale() * 2.0f;
 
-	// Test if warmup is completed
-	bool bCanUnload = CurrentFrameNumber >= CellData->WarmupEndFrame;
-	if (bCanUnload)
-	{
-		CellData->WarmupStartFrame = INDEX_NONE;
-		CellData->WarmupEndFrame = INDEX_NONE;
+				MaxScreenSizePixels = FMath::Max(MaxScreenSizePixels, ScreenSizePixels);
+			}
+			else
+			{
+				return true;
+			}
+		}
 	}
 
-	return bCanUnload;
+	return MaxScreenSizePixels > 0;
 }
 
-void UHLODSubsystem::OnBeginRenderViews(const FSceneViewFamily& InViewFamily)
+static void MakeHLODRenderResourcesResident(TMap<UMaterialInterface*, float>& VTRequests, TSet<Nanite::FResources*>& NaniteRequests, const FSceneViewFamily& InViewFamily)
 {
-	for (auto& KeyValuePair : WorldPartitionsHLODRuntimeData)
+	if (!VTRequests.IsEmpty() || !NaniteRequests.IsEmpty())
 	{
-		for (TSet<FCellData*>::TIterator CellIt(KeyValuePair.Value.CellsToWarmup); CellIt; ++CellIt)
+		ENQUEUE_RENDER_COMMAND(MakeHLODRenderResourcesResident)(
+			[VTRequests = MoveTemp(VTRequests), NaniteRequests = MoveTemp(NaniteRequests), FeatureLevel = InViewFamily.GetFeatureLevel()](FRHICommandListImmediate& RHICmdList)
 		{
-			FCellData* CellPendingUnload = *CellIt;
+			for (const TPair<UMaterialInterface*, float>& VTRequest : VTRequests)
+			{
+				UMaterialInterface* Material = VTRequest.Key;
+				FMaterialRenderProxy* MaterialRenderProxy = Material->GetRenderProxy();
 
-			MakeRenderResourcesResident(*CellPendingUnload, InViewFamily);
+				GetRendererModule().RequestVirtualTextureTiles(MaterialRenderProxy, FVector2D(VTRequest.Value, VTRequest.Value), FeatureLevel);
+			}
 
-			// Stop processing this cell if warmup is done.
-			if (InViewFamily.FrameNumber >= CellPendingUnload->WarmupEndFrame)
+			const uint32 NumFramesBeforeRender = CVarHLODWarmupNumFrames.GetValueOnRenderThread();
+			for (const Nanite::FResources* Resource : NaniteRequests)
 			{
-				CellIt.RemoveCurrent();
+				GetRendererModule().PrefetchNaniteResource(Resource, NumFramesBeforeRender);
 			}
+		});
+	}
+}
+
+void UHLODSubsystem::OnBeginRenderViews(const FSceneViewFamily& InViewFamily)
+{
+	TRACE_CPUPROFILER_EVENT_SCOPE(UHLODSubsystem::OnBeginRenderViews)
+
+	TMap<UMaterialInterface*, float> VTRequests;
+	TSet<Nanite::FResources*> NaniteRequests;
+
+	const bool bWarmupNanite = CVarHLODWarmupNanite.GetValueOnGameThread() != 0;
+	const bool bWarmupVT = CVarHLODWarmupVT.GetValueOnGameThread() != 0;
+
+	for (FHLODWarmupStateMap::TIterator HLODActorWarmupStateIt(HLODActorsToWarmup); HLODActorWarmupStateIt; ++HLODActorWarmupStateIt)
+	{
+		const FObjectKey& HLODActorObjectKey = HLODActorWarmupStateIt.Key();
+		const FWorldPartitionHLODWarmupState& HLODWarmupState = HLODActorWarmupStateIt.Value();
+		const AWorldPartitionHLOD* HLODActor = Cast<AWorldPartitionHLOD>(HLODActorObjectKey.ResolveObjectPtr());
+		
+		if (HLODActor && InViewFamily.FrameNumber < HLODWarmupState.WarmupEndFrame)
+		{
+			HLODActor->ForEachComponent<UStaticMeshComponent>(false, [&](UStaticMeshComponent* SMC)
+			{
+				// Assume ISM HLOD don't need warmup, as they are actually found in the source level
+				if (SMC->IsA<UInstancedStaticMeshComponent>() || !SMC->GetStaticMesh())
+				{
+					return;
+				}
+
+				// Retrieve this component's bound - we must support getting the bounds before the component is even registered.
+				FVector BoundsOrigin = SMC->GetStaticMesh()->GetBounds().Origin + SMC->GetRelativeLocation();
+				FVector BoundsExtent = SMC->GetStaticMesh()->GetBounds().BoxExtent;
+				
+				float ScreenSizePixels = 0;
+				if (IsInView(BoundsOrigin, BoundsExtent, InViewFamily, bWarmupVT, ScreenSizePixels))
+				{
+					if (bWarmupVT)
+					{
+						PrepareVTRequests(VTRequests, SMC, ScreenSizePixels);
+					}
+
+					if (bWarmupNanite)
+					{
+						// Only issue Nanite requests on the first warmup frame
+						if (HLODWarmupState.WarmupStartFrame == InViewFamily.FrameNumber)
+						{
+							PrepareNaniteRequests(NaniteRequests, SMC);
+						}
+					}
+
+#if ENABLE_DRAW_DEBUG
+					if (CVarHLODWarmupDebugDraw.GetValueOnAnyThread())
+					{
+						DrawDebugBox(HLODActor->GetWorld(), BoundsOrigin, BoundsExtent, FColor::Yellow, /*bPersistentLine*/ false, /*Lifetime*/ 1.0f);
+					}
+#endif
+				}
+			});
+		}
+
+		if (!HLODActor)
+		{
+			// HLOD actor has been unloaded, stop tracking it's warmup state
+			HLODActorWarmupStateIt.RemoveCurrent();
 		}
 	}
+
+	MakeHLODRenderResourcesResident(VTRequests, NaniteRequests, InViewFamily);
 }
 
 #if WITH_EDITOR
diff --git a/Engine/Source/Runtime/Engine/Private/WorldPartition/WorldPartitionLevelStreamingDynamic.cpp b/Engine/Source/Runtime/Engine/Private/WorldPartition/WorldPartitionLevelStreamingDynamic.cpp
--- a/Engine/Source/Runtime/Engine/Private/WorldPartition/WorldPartitionLevelStreamingDynamic.cpp
+++ b/Engine/Source/Runtime/Engine/Private/WorldPartition/WorldPartitionLevelStreamingDynamic.cpp
@@ -2,6 +2,7 @@
 
 #include "WorldPartition/WorldPartitionLevelStreamingDynamic.h"
 #include "WorldPartition/WorldPartition.h"
+#include "WorldPartition/HLOD/HLODSubsystem.h"
 #include "Engine/World.h"
 
 #include UE_INLINE_GENERATED_CPP_BY_NAME(WorldPartitionLevelStreamingDynamic)
@@ -626,6 +627,62 @@ void UWorldPartitionLevelStreamingDynamic::Deactivate()
 	SetShouldBeVisible(false);
 }
 
+bool UWorldPartitionLevelStreamingDynamic::CanMakeVisible()
+{
+	const ENetMode NetMode = GetWorld()->GetNetMode();
+	if (NetMode != NM_DedicatedServer)
+	{
+		// @todo_ow ContentBundles do not support Hlods and events below are forwarding
+		// Events to the HLodSubsystem which assumes knowledge of all cells (not true with plugins)
+		if (const UWorldPartitionRuntimeLevelStreamingCell* RuntimeLevelStreamingCell = StreamingCell.Get())
+		{
+			if (RuntimeLevelStreamingCell->GetIsHLOD() && !RuntimeLevelStreamingCell->GetContentBundleID().IsValid())
+			{
+				if (const UWorldPartition* WorldPartition = GetWorld()->GetWorldPartition())
+				{
+					if (UHLODSubsystem* HLODSubsystem = GetWorld()->GetSubsystem<UHLODSubsystem>())
+					{
+						if (!HLODSubsystem->CanMakeVisible(StreamingCell.Get()))
+						{
+							return false;
+						}
+					}
+				}
+			}
+		}
+	}
+
+	return Super::CanMakeVisible();
+}
+
+bool UWorldPartitionLevelStreamingDynamic::CanMakeInvisible()
+{
+	const ENetMode NetMode = GetWorld()->GetNetMode();
+	if (NetMode != NM_DedicatedServer)
+	{
+		// @todo_ow ContentBundles do not support Hlods and events below are forwarding
+		// Events to the HLodSubsystem which assumes knowledge of all cells (not true with plugins)
+		if (const UWorldPartitionRuntimeLevelStreamingCell* RuntimeLevelStreamingCell = StreamingCell.Get())
+		{
+			if (!RuntimeLevelStreamingCell->GetContentBundleID().IsValid())
+			{
+				if (const UWorldPartition* WorldPartition = GetWorld()->GetWorldPartition())
+				{
+					if (UHLODSubsystem* HLODSubsystem = GetWorld()->GetSubsystem<UHLODSubsystem>())
+					{
+						if (!HLODSubsystem->CanMakeInvisible(StreamingCell.Get()))
+						{
+							return false;
+						}
+					}
+				}
+			}
+		}
+	}
+
+	return Super::CanMakeInvisible();
+}
+
 UWorld* UWorldPartitionLevelStreamingDynamic::GetOuterWorld() const
 {
 	check(OuterWorldPartition.IsValid());
diff --git a/Engine/Source/Runtime/Engine/Private/WorldPartition/WorldPartitionRuntimeLevelStreamingCell.cpp b/Engine/Source/Runtime/Engine/Private/WorldPartition/WorldPartitionRuntimeLevelStreamingCell.cpp
--- a/Engine/Source/Runtime/Engine/Private/WorldPartition/WorldPartitionRuntimeLevelStreamingCell.cpp
+++ b/Engine/Source/Runtime/Engine/Private/WorldPartition/WorldPartitionRuntimeLevelStreamingCell.cpp
@@ -417,19 +417,6 @@ ULevel* UWorldPartitionRuntimeLevelStreamingCell::GetLevel() const
 
 bool UWorldPartitionRuntimeLevelStreamingCell::CanUnload() const
 {
-	// @todo_ow ContentBundles do not support Hlods and events below are forwarding
-	// Events to the HLodSubsystem which assumes knowledge of all cells (not true with plugins)
-	if (LevelStreaming && !GetContentBundleID().IsValid())
-	{
-		if (const UWorldPartition* WorldPartition = LevelStreaming->GetWorld()->GetWorldPartition(); WorldPartition && WorldPartition->IsStreamingEnabled())
-		{
-			if (UHLODSubsystem* HLODSubsystem = LevelStreaming->GetWorld()->GetSubsystem<UHLODSubsystem>())
-			{
-				return HLODSubsystem->RequestUnloading(this);
-			}
-		}
-	}
-
 	return true;
 }
 
diff --git a/Engine/Source/Runtime/Engine/Public/WorldPartition/HLOD/HLODSubsystem.h b/Engine/Source/Runtime/Engine/Public/WorldPartition/HLOD/HLODSubsystem.h
--- a/Engine/Source/Runtime/Engine/Public/WorldPartition/HLOD/HLODSubsystem.h
+++ b/Engine/Source/Runtime/Engine/Public/WorldPartition/HLOD/HLODSubsystem.h
@@ -4,6 +4,7 @@
 #include "Subsystems/WorldSubsystem.h"
 #include "Containers/Array.h"
 #include "Containers/Map.h"
+#include "UObject/ObjectKey.h"
 #include "HLODSubsystem.generated.h"
 
 
@@ -46,7 +47,8 @@ public:
 	void OnCellShown(const UWorldPartitionRuntimeCell* InCell);
 	void OnCellHidden(const UWorldPartitionRuntimeCell* InCell);
 
-	bool RequestUnloading(const UWorldPartitionRuntimeCell* InCell);
+	bool CanMakeVisible(const UWorldPartitionRuntimeCell* InCell);
+	bool CanMakeInvisible(const UWorldPartitionRuntimeCell* InCell);
 
 	const TArray<AWorldPartitionHLOD*>& GetHLODActorsForCell(const UWorldPartitionRuntimeCell* InCell) const;
 
@@ -56,38 +58,43 @@ public:
 	FWorldPartitionHLODActorUnregisteredEvent& OnHLODActorUnregisteredEvent() { return HLODActorUnregisteredEvent; }
 
 	void SetHLODAlwaysLoadedCullDistance(int32 InCullDistance);
+
+	void OnCVarsChanged();
 	
 private:
 	struct FCellData
 	{
-		bool bIsCellVisible;
+		bool bIsCellVisible = false;
 		TArray<AWorldPartitionHLOD*> LoadedHLODs;
-		uint32 WarmupStartFrame;
-		uint32 WarmupEndFrame;
-
-		FCellData() 
-			: bIsCellVisible(false)
-			, WarmupStartFrame(INDEX_NONE)
-			, WarmupEndFrame(INDEX_NONE)
-		{
-		}
 	};
 
 	struct FWorldPartitionHLODRuntimeData
 	{
 		TMap<FName, FCellData> CellsData;
-		TSet<FCellData*> CellsToWarmup;
 	};
 	
 	TMap<TObjectPtr<UWorldPartition>, FWorldPartitionHLODRuntimeData> WorldPartitionsHLODRuntimeData;
+	const FCellData* GetCellData(const UWorldPartitionRuntimeCell* InCell) const;
+	FCellData* GetCellData(const UWorldPartitionRuntimeCell* InCell);
+	FCellData* GetCellData(AWorldPartitionHLOD* InWorldPartitionHLOD);
 
 	TArray<AWorldPartitionHLOD*> AlwaysLoadedHLODActors;
 
-	void OnWorldPartitionInitialized(UWorldPartition* InWorldPartition);
-	void OnWorldPartitionUninitialized(UWorldPartition* InWorldPartition);
-	void OnBeginRenderViews(const FSceneViewFamily& InViewFamily);
+	struct FWorldPartitionHLODWarmupState
+	{
+		uint32 WarmupStartFrame = INDEX_NONE;
+		uint32 WarmupEndFrame = INDEX_NONE;
+		FVector Location;
+	};
 
-	void MakeRenderResourcesResident(const FCellData& CellHLODs, const FSceneViewFamily& InViewFamily);
+	typedef TMap<FObjectKey, FWorldPartitionHLODWarmupState> FHLODWarmupStateMap;
+	FHLODWarmupStateMap HLODActorsToWarmup;
+
+	void OnBeginRenderViews(const FSceneViewFamily& InViewFamily);
+	bool PrepareToWarmup(const UWorldPartitionRuntimeCell* InCell, AWorldPartitionHLOD* InHLODActor);
+	bool ShouldPerformWarmup() const;
+	bool ShouldPerformWarmupForCell(const UWorldPartitionRuntimeCell* InCell) const;
+	bool bCachedShouldPerformWarmup;
 
 	/** Console command used to turn on/off loading & rendering of world partition HLODs */
 	static class FAutoConsoleCommand EnableHLODCommand;
@@ -97,6 +104,9 @@ private:
 	friend class FHLODResourcesResidencySceneViewExtension;
 	TSharedPtr<FHLODResourcesResidencySceneViewExtension, ESPMode::ThreadSafe> SceneViewExtension;
 
+	void OnWorldPartitionInitialized(UWorldPartition* InWorldPartition);
+	void OnWorldPartitionUninitialized(UWorldPartition* InWorldPartition);
+
 	FWorldPartitionHLODActorRegisteredEvent		HLODActorRegisteredEvent;
 	FWorldPartitionHLODActorUnregisteredEvent	HLODActorUnregisteredEvent;
 };
diff --git a/Engine/Source/Runtime/Engine/Public/WorldPartition/WorldPartitionLevelStreamingDynamic.h b/Engine/Source/Runtime/Engine/Public/WorldPartition/WorldPartitionLevelStreamingDynamic.h
--- a/Engine/Source/Runtime/Engine/Public/WorldPartition/WorldPartitionLevelStreamingDynamic.h
+++ b/Engine/Source/Runtime/Engine/Public/WorldPartition/WorldPartitionLevelStreamingDynamic.h
@@ -35,6 +35,10 @@ class ENGINE_API UWorldPartitionLevelStreamingDynamic : public ULevelStreamingDy
 
 	virtual bool ShouldBeAlwaysLoaded() const override { return bShouldBeAlwaysLoaded; }
 	virtual bool ShouldRequireFullVisibilityToRender() const override { return true; }
+
+	virtual bool CanMakeVisible() override;
+	virtual bool CanMakeInvisible() override;
+
 #if !WITH_EDITOR
 	virtual void PostLoad();
 #endif
