Add option to delay component proxy creation when the requested PSOs are still compiling - disabled by default

diff --git a/Engine/Source/Runtime/Engine/Classes/Components/MeshComponent.h b/Engine/Source/Runtime/Engine/Classes/Components/MeshComponent.h
--- a/Engine/Source/Runtime/Engine/Classes/Components/MeshComponent.h
+++ b/Engine/Source/Runtime/Engine/Classes/Components/MeshComponent.h
@@ -134,6 +134,12 @@ public:
 	/** Precache all PSOs which can be used by the mesh component */
 	virtual void PrecachePSOs() {}
 
+	/** Schedule task to mark render state dirty when the PSO precaching tasks are done */
+	void RequestRecreateRenderStateWhenPSOPrecacheFinished(const FGraphEventArray& PSOPrecacheCompileEvents);
+
+	/** Check if PSOs are still precaching */
+	bool IsPSOPrecaching();
+
 	/** Generate streaming data for all materials. */
 	void GetStreamingTextureInfoInner(FStreamingTextureLevelContext& LevelContext, const TArray<FStreamingTextureBuildInfo>* PreBuiltData, float ComponentScaling, TArray<FStreamingRenderAssetPrimitiveInfo>& OutStreamingTextures) const;
 
@@ -195,4 +201,10 @@ protected:
 	/** Flag whether or not the cached material parameter indices map is dirty (defaults to true, and is set from SetMaterial/Set(Skeletal)Mesh */
 	uint8 bCachedMaterialParameterIndicesAreDirty : 1;
 
+	/** Helper flag to check if PSOs have been precached already */
+	uint8 bPSOPrecacheCalled : 1;
+
+	/** Graph event used to track all the PSO precache events */
+	FGraphEventRef PSOPrecacheCompileEvent;
+
 };
diff --git a/Engine/Source/Runtime/Engine/Private/Components/MeshComponent.cpp b/Engine/Source/Runtime/Engine/Private/Components/MeshComponent.cpp
--- a/Engine/Source/Runtime/Engine/Private/Components/MeshComponent.cpp
+++ b/Engine/Source/Runtime/Engine/Private/Components/MeshComponent.cpp
@@ -7,6 +7,7 @@
 #include "ContentStreaming.h"
 #include "Streaming/TextureStreamingHelpers.h"
 #include "Engine/World.h"
+#include "PSOPrecache.h"
 
 #include UE_INLINE_GENERATED_CPP_BY_NAME(MeshComponent)
 
@@ -25,6 +26,7 @@ UMeshComponent::UMeshComponent(const FObjectInitializer& ObjectInitializer)
 	bCanEverAffectNavigation = true;
 	bCachedMaterialParameterIndicesAreDirty = true;
 	bEnableMaterialParameterCaching = false;
+	bPSOPrecacheCalled = false;
 }
 
 UMaterialInterface* UMeshComponent::GetMaterial(int32 ElementIndex) const
@@ -74,6 +76,10 @@ void UMeshComponent::SetMaterial(int32 ElementIndex, UMaterialInterface* Materia
 
 			// Set the material and invalidate things
 			OverrideMaterials[ElementIndex] = Material;
+
+			// Precache PSOs again
+			PrecachePSOs();
+
 			MarkRenderStateDirty();
 			// If MarkRenderStateDirty didn't notify the streamer, do it now
 			if (!bIgnoreStreamingManagerUpdate && OwnerLevelHasRegisteredStaticComponentsInStreamingManager(GetOwner()))
@@ -91,8 +97,6 @@ void UMeshComponent::SetMaterial(int32 ElementIndex, UMaterialInterface* Materia
 				BodyInst->UpdatePhysicalMaterials();
 			}
 
-			PrecachePSOs();
-
 #if WITH_EDITOR
 			// Static Lighting is updated when compilation finishes
 			if (!IsCompiling())
@@ -144,6 +148,59 @@ int32 UMeshComponent::GetNumOverrideMaterials() const
 	return OverrideMaterials.Num();
 }
 
+class FMarkRenderStateDirtyTask
+{
+public:
+	explicit FMarkRenderStateDirtyTask(UActorComponent* InActorComponent)
+		: ActorComponent(InActorComponent)
+	{
+	}
+
+	void DoTask(ENamedThreads::Type CurrentThread, const FGraphEventRef& MyCompletionGraphEvent)
+	{
+		if (ActorComponent.IsValid())
+		{
+			ActorComponent->MarkRenderStateDirty();
+			ActorComponent = nullptr;
+		}
+	}
+
+public:
+
+	TWeakObjectPtr<UActorComponent> ActorComponent;
+
+	static ESubsequentsMode::Type	GetSubsequentsMode() { return ESubsequentsMode::TrackSubsequents; }
+	ENamedThreads::Type				GetDesiredThread() { return ENamedThreads::GameThread; }
+	FORCEINLINE TStatId				GetStatId() const { return TStatId(); }
+};
+
+void UMeshComponent::RequestRecreateRenderStateWhenPSOPrecacheFinished(const FGraphEventArray& PSOPrecacheCompileEvents)
+{
+	// Mark the render state dirty when all PSOs are compiled so the proxy gets recreated
+	if (ProxyCreationWhenPSOReady() && !PSOPrecacheCompileEvents.IsEmpty())
+	{
+		PSOPrecacheCompileEvent = TGraphTask<FMarkRenderStateDirtyTask>::CreateTask(&PSOPrecacheCompileEvents).ConstructAndDispatchWhenReady(this);
+	}
+
+	bPSOPrecacheCalled = true;
+}
+
+bool UMeshComponent::IsPSOPrecaching()
+{
+	if (!ProxyCreationWhenPSOReady())
+	{
+		return false;
+	}
+
+	ensure(!IsComponentPSOPrecachingEnabled() || bPSOPrecacheCalled);
+
+	if (PSOPrecacheCompileEvent && PSOPrecacheCompileEvent->IsComplete())
+	{
+		PSOPrecacheCompileEvent = nullptr;
+	}
+	return PSOPrecacheCompileEvent != nullptr;
+}
+
 #if WITH_EDITOR
 void UMeshComponent::PostEditChangeChainProperty(FPropertyChangedChainEvent& PropertyChangedEvent)
 {
diff --git a/Engine/Source/Runtime/Engine/Private/Components/SkinnedMeshComponent.cpp b/Engine/Source/Runtime/Engine/Private/Components/SkinnedMeshComponent.cpp
--- a/Engine/Source/Runtime/Engine/Private/Components/SkinnedMeshComponent.cpp
+++ b/Engine/Source/Runtime/Engine/Private/Components/SkinnedMeshComponent.cpp
@@ -503,7 +503,8 @@ FPrimitiveSceneProxy* USkinnedMeshComponent::CreateSceneProxy()
 	if (SkelMeshRenderData &&
 		SkelMeshRenderData->LODRenderData.IsValidIndex(GetPredictedLODLevel()) &&
 		!bHideSkin &&
-		MeshObject)
+		MeshObject &&
+		!IsPSOPrecaching())
 	{
 		// Only create a scene proxy if the bone count being used is supported, or if we don't have a skeleton (this is the case with destructibles)
 		int32 MinLODIndex = ComputeMinLOD();
@@ -577,14 +578,17 @@ void USkinnedMeshComponent::PrecachePSOs()
 	SetupPrecachePSOParams(PrecachePSOParams);
 	PrecachePSOParams.bCastShadow = PrecachePSOParams.bCastShadow && bAnySectionCastsShadows;
 
+	FGraphEventArray PSOPrecacheCompileEvents;
 	for (FSkinnedAssetVertexFactoryTypesPerMaterialData& VFsPerMaterial : VFsPerMaterials)
 	{
 		UMaterialInterface* MaterialInterface = GetMaterial(VFsPerMaterial.MaterialIndex);
 		if (MaterialInterface)
 		{
-			MaterialInterface->PrecachePSOs(VFsPerMaterial.VertexFactoryTypes, PrecachePSOParams);
+			PSOPrecacheCompileEvents.Append(MaterialInterface->PrecachePSOs(VFsPerMaterial.VertexFactoryTypes, PrecachePSOParams));
 		}
 	}
+
+	RequestRecreateRenderStateWhenPSOPrecacheFinished(PSOPrecacheCompileEvents);
 }
 
 void USkinnedMeshComponent::OnRegister()
@@ -645,6 +649,8 @@ void USkinnedMeshComponent::OnRegister()
 	MeshDeformerInstance = ActiveMeshDeformer != nullptr ? ActiveMeshDeformer->CreateInstance(this, MeshDeformerInstanceSettings) : nullptr;
 
 	RefreshExternalMorphTargetWeights();
+
+	PrecachePSOs();
 }
 
 void USkinnedMeshComponent::OnUnregister()
diff --git a/Engine/Source/Runtime/Engine/Private/Components/StaticMeshComponent.cpp b/Engine/Source/Runtime/Engine/Private/Components/StaticMeshComponent.cpp
--- a/Engine/Source/Runtime/Engine/Private/Components/StaticMeshComponent.cpp
+++ b/Engine/Source/Runtime/Engine/Private/Components/StaticMeshComponent.cpp
@@ -1559,6 +1559,7 @@ void UStaticMeshComponent::PrecachePSOs()
 
 	const FVertexFactoryType* VFType = ShouldCreateNaniteProxy() ? &Nanite::FVertexFactory::StaticType :  &FLocalVertexFactory::StaticType;
 
+	FGraphEventArray PSOPrecacheCompileEvents;
 	for (uint16 MaterialIndex : UsedMaterialIndices)
 	{
 		UMaterialInterface* MaterialInterface = GetMaterial(MaterialIndex);
@@ -1567,9 +1568,12 @@ void UStaticMeshComponent::PrecachePSOs()
 			PrecachePSOParams.bHasWorldPositionOffsetVelocity = SupportsWorldPositionOffsetVelocity()
 				&& MaterialInterface->GetRelevance_Concurrent(FeatureLevel).bUsesWorldPositionOffset;
 
-			MaterialInterface->PrecachePSOs(VFType, PrecachePSOParams);
+			PSOPrecacheCompileEvents.Append(MaterialInterface->PrecachePSOs(VFType, PrecachePSOParams));
 		}
 	}
+
+	// Mark the render state dirty when all PSOs are compiled so the proxy gets recreated
+	RequestRecreateRenderStateWhenPSOPrecacheFinished(PSOPrecacheCompileEvents);
 }
 
 #if WITH_EDITOR
@@ -2203,6 +2207,9 @@ bool UStaticMeshComponent::SetStaticMesh(UStaticMesh* NewMesh)
 		);
 	}
 
+	// Precache the PSOs
+	PrecachePSOs();
+
 	// Need to send this to render thread at some point
 	if (IsRenderStateCreated())
 	{
@@ -2231,9 +2238,6 @@ bool UStaticMeshComponent::SetStaticMesh(UStaticMesh* NewMesh)
 	// Mark cached material parameter names dirty
 	MarkCachedMaterialParameterNameIndicesDirty();
 
-	// Precache the PSOs
-	PrecachePSOs();
-
 #if WITH_EDITOR
 	// Broadcast that the static mesh has changed
 	OnStaticMeshChangedEvent.Broadcast(this);
diff --git a/Engine/Source/Runtime/Engine/Private/HierarchicalInstancedStaticMesh.cpp b/Engine/Source/Runtime/Engine/Private/HierarchicalInstancedStaticMesh.cpp
--- a/Engine/Source/Runtime/Engine/Private/HierarchicalInstancedStaticMesh.cpp
+++ b/Engine/Source/Runtime/Engine/Private/HierarchicalInstancedStaticMesh.cpp
@@ -3163,7 +3163,8 @@ FPrimitiveSceneProxy* UHierarchicalInstancedStaticMeshComponent::CreateSceneProx
 		// Make sure we have an actual static mesh.
 		GetStaticMesh() &&
 		!GetStaticMesh()->IsCompiling() &&
-		GetStaticMesh()->HasValidRenderData(false);
+		GetStaticMesh()->HasValidRenderData(false) &&
+		!IsPSOPrecaching();
 
 	if (bMeshIsValid)
 	{
diff --git a/Engine/Source/Runtime/Engine/Private/InstancedStaticMesh.cpp b/Engine/Source/Runtime/Engine/Private/InstancedStaticMesh.cpp
--- a/Engine/Source/Runtime/Engine/Private/InstancedStaticMesh.cpp
+++ b/Engine/Source/Runtime/Engine/Private/InstancedStaticMesh.cpp
@@ -2353,7 +2353,8 @@ FPrimitiveSceneProxy* UInstancedStaticMeshComponent::CreateSceneProxy()
 		// make sure we have an actual static mesh
 		GetStaticMesh() &&
 		GetStaticMesh()->IsCompiling() == false &&
-		GetStaticMesh()->HasValidRenderData();
+		GetStaticMesh()->HasValidRenderData() &&
+		!IsPSOPrecaching();
 
 	if (bMeshIsValid)
 	{
@@ -4548,14 +4549,18 @@ void UInstancedStaticMeshComponent::PrecachePSOs()
 	
 	const FVertexFactoryType* VFType = ShouldCreateNaniteProxy() ? &Nanite::FVertexFactory::StaticType : &FInstancedStaticMeshVertexFactory::StaticType;
 
+	FGraphEventArray PSOPrecacheCompileEvents;
 	for (uint16 MaterialIndex : UsedMaterialIndices)
 	{
 		UMaterialInterface* MaterialInterface = GetMaterial(MaterialIndex);
 		if (MaterialInterface)
 		{
-			MaterialInterface->PrecachePSOs(VFType, PrecachePSOParams);
+			PSOPrecacheCompileEvents.Append(MaterialInterface->PrecachePSOs(VFType, PrecachePSOParams));
 		}
 	}
+
+	// Mark the render state dirty when all PSOs are compiled so the proxy gets recreated
+	RequestRecreateRenderStateWhenPSOPrecacheFinished(PSOPrecacheCompileEvents);
 }
 
 void UInstancedStaticMeshComponent::OnPostLoadPerInstanceData()
diff --git a/Engine/Source/Runtime/Engine/Private/Materials/Material.cpp b/Engine/Source/Runtime/Engine/Private/Materials/Material.cpp
--- a/Engine/Source/Runtime/Engine/Private/Materials/Material.cpp
+++ b/Engine/Source/Runtime/Engine/Private/Materials/Material.cpp
@@ -566,7 +566,15 @@ void UMaterialInterface::InitDefaultMaterials()
 			}
 		}
 
-		// Now precache PSOs for all the default materials
+		RecursionLevel--;
+#if USE_EVENT_DRIVEN_ASYNC_LOAD_AT_BOOT_TIME
+		bInitialized = !GEventDrivenLoaderEnabled || RecursionLevel == 0;
+#else
+		bInitialized = true;
+#endif
+
+		// Now precache PSOs for all the default materials after the default materials are marked initialize
+		// PSO precaching can request default materials so they have to marked as initialized to avoid endless recursion
 		if (PipelineStateCache::IsPSOPrecachingEnabled())
 		{
 			FPSOPrecacheParams PrecachePSOParams;
@@ -587,13 +595,6 @@ void UMaterialInterface::InitDefaultMaterials()
 				}
 			}
 		}
-
-		RecursionLevel--;
-#if USE_EVENT_DRIVEN_ASYNC_LOAD_AT_BOOT_TIME
-		bInitialized = !GEventDrivenLoaderEnabled || RecursionLevel == 0;
-#else
-		bInitialized = true;
-#endif
 	}
 }
 
diff --git a/Engine/Source/Runtime/Engine/Private/Materials/MaterialShared.cpp b/Engine/Source/Runtime/Engine/Private/Materials/MaterialShared.cpp
--- a/Engine/Source/Runtime/Engine/Private/Materials/MaterialShared.cpp
+++ b/Engine/Source/Runtime/Engine/Private/Materials/MaterialShared.cpp
@@ -2738,7 +2738,7 @@ public:
 		TStrongObjectPtr<UMaterialInterface>* InMaterialInterface,
 		FMaterialShaderMap* InMaterialShaderMap,
 		ERHIFeatureLevel::Type InFeatureLevel,
-		const FMaterial* InMaterial,
+		FMaterial* InMaterial,
 		const TConstArrayView<const FVertexFactoryType*>& InVertexFactoryTypes,
 		const FPSOPrecacheParams& InPreCacheParams)
 		: MaterialInterface(InMaterialInterface)
@@ -2759,6 +2759,12 @@ public:
 		// Won't touch the material interface anymore - PSO compile jobs take refs to all RHI resources while creating the task
 		TGraphTask<FMaterialInterfaceReleaseTask>::CreateTask().ConstructAndDispatchWhenReady(MaterialInterface);
 
+		// Mark this vertex factory with precache params as processed with current compile events
+		for (const FVertexFactoryType* VFType : VertexFactoryTypes)
+		{
+			Material->AddPrecachedPSOVertexFactoryType(VFType, PreCacheParams, PSOCompileGraphEvents);
+		}
+
 		// Extend MyCompletionGraphEvent to wait for all the async compile events
 		for (FGraphEventRef& GraphEvent : PSOCompileGraphEvents)
 		{
@@ -2771,7 +2777,7 @@ public:
 	TStrongObjectPtr<UMaterialInterface>* MaterialInterface;
 	FMaterialShaderMap* MaterialShaderMap;
 	ERHIFeatureLevel::Type FeatureLevel;
-	const FMaterial* Material;
+	FMaterial* Material;
 	TArray<const FVertexFactoryType*, TInlineAllocator<4>> VertexFactoryTypes;
 	const FPSOPrecacheParams PreCacheParams;
 
@@ -2792,6 +2798,18 @@ FGraphEventArray FMaterial::CollectPSOs(ERHIFeatureLevel::Type InFeatureLevel, c
 		return GraphEvents;
 	}
 
+	auto FindPrecacheEntry = [this](const FPrecacheVertexTypeWithParams& Other) -> FPSOPrecacheEntry*
+	{
+		for (FPSOPrecacheEntry& Entry : PrecachedPSOVertexFactories)
+		{
+			if (Entry.PrecacheVertexTypeWithParams == Other)
+			{
+				return &Entry;
+			}
+		}
+		return nullptr;
+	};
+
 	TArray<const FVertexFactoryType*, TInlineAllocator<4>> MissingVFs;
 
 	// Try and find missing entries which still need to be precached
@@ -2802,7 +2820,9 @@ FGraphEventArray FMaterial::CollectPSOs(ERHIFeatureLevel::Type InFeatureLevel, c
 			FPrecacheVertexTypeWithParams PrecacheVertexTypeWithParams;
 			PrecacheVertexTypeWithParams.VertexFactoryType = VFType;
 			PrecacheVertexTypeWithParams.PrecachePSOParams = PreCacheParams;
-			if (!PrecachedPSOVertexFactories.Contains(PrecacheVertexTypeWithParams))
+
+			FPSOPrecacheEntry* CurrentEntry = FindPrecacheEntry(PrecacheVertexTypeWithParams);			
+			if (CurrentEntry == nullptr || !CurrentEntry->CompilingGraphEvents.IsEmpty())
 			{
 				MissingVFs.Add(VFType);
 			}
@@ -2823,11 +2843,27 @@ FGraphEventArray FMaterial::CollectPSOs(ERHIFeatureLevel::Type InFeatureLevel, c
 				FPrecacheVertexTypeWithParams PrecacheVertexTypeWithParams;
 				PrecacheVertexTypeWithParams.VertexFactoryType = VFType;
 				PrecacheVertexTypeWithParams.PrecachePSOParams = PreCacheParams;
-				if (PrecachedPSOVertexFactories.Contains(PrecacheVertexTypeWithParams))
-					continue;
 
-				ActualMissingVFs.Add(VFType);
-				PrecachedPSOVertexFactories.Add(PrecacheVertexTypeWithParams);
+				FPSOPrecacheEntry* CurrentEntry = FindPrecacheEntry(PrecacheVertexTypeWithParams);
+				if (CurrentEntry == nullptr)
+				{
+					ActualMissingVFs.Add(VFType);
+				}
+				else if (!CurrentEntry->CompilingGraphEvents.IsEmpty())
+				{
+					// trim the current list of open compile events
+					for (int32 Index = 0; Index < CurrentEntry->CompilingGraphEvents.Num(); ++Index)
+					{
+						if (CurrentEntry->CompilingGraphEvents[Index]->IsComplete())
+						{
+							CurrentEntry->CompilingGraphEvents.RemoveAtSwap(Index);
+							Index--;
+						}
+					}
+
+					// add the current active compile events
+					GraphEvents.Append(CurrentEntry->CompilingGraphEvents);
+				}
 			}
 		}
 
@@ -2854,6 +2890,10 @@ FGraphEventArray FMaterial::CollectPSOs(ERHIFeatureLevel::Type InFeatureLevel, c
 			else
 			{				
 				GraphEvents = GameThreadShaderMap->CollectPSOs(InFeatureLevel, this, ActualMissingVFs, PreCacheParams);
+				for (const FVertexFactoryType* VFType : ActualMissingVFs)
+				{
+					AddPrecachedPSOVertexFactoryType(VFType, PreCacheParams, GraphEvents);
+				}
 			}
 		}
 	}
@@ -2861,6 +2901,33 @@ FGraphEventArray FMaterial::CollectPSOs(ERHIFeatureLevel::Type InFeatureLevel, c
 	return GraphEvents;
 }
 
+void FMaterial::AddPrecachedPSOVertexFactoryType(const FVertexFactoryType* VertexFactoryType, const FPSOPrecacheParams& PreCacheParams, const FGraphEventArray& PSOCompileEvents)
+{
+	FPrecacheVertexTypeWithParams PrecacheVertexTypeWithParams;
+	PrecacheVertexTypeWithParams.VertexFactoryType = VertexFactoryType;
+	PrecacheVertexTypeWithParams.PrecachePSOParams = PreCacheParams;
+
+	FRWScopeLock WriteLock(PrecachePSOVFLock, SLT_Write);
+
+	for (FPSOPrecacheEntry& Entry : PrecachedPSOVertexFactories)
+	{
+		if (Entry.PrecacheVertexTypeWithParams == PrecacheVertexTypeWithParams)
+		{
+			// Already found, then merge the outstanding compile events
+			for (FGraphEventRef CompileEvent : PSOCompileEvents)
+			{
+				Entry.CompilingGraphEvents.AddUnique(CompileEvent);
+			}
+
+			return;
+		}
+	}
+
+	FPSOPrecacheEntry& NewEntry = PrecachedPSOVertexFactories[PrecachedPSOVertexFactories.AddDefaulted(1)];
+	NewEntry.PrecacheVertexTypeWithParams = PrecacheVertexTypeWithParams;
+	NewEntry.CompilingGraphEvents = PSOCompileEvents;
+}
+
 #if WITH_EDITOR
 
 void FMaterial::BeginCacheShaders(EShaderPlatform Platform, EMaterialShaderPrecompileMode PrecompileMode, const ITargetPlatform* TargetPlatform, TUniqueFunction<void(bool bSuccess)>&& CompletionCallback)
diff --git a/Engine/Source/Runtime/Engine/Private/PSOPrecache.cpp b/Engine/Source/Runtime/Engine/Private/PSOPrecache.cpp
--- a/Engine/Source/Runtime/Engine/Private/PSOPrecache.cpp
+++ b/Engine/Source/Runtime/Engine/Private/PSOPrecache.cpp
@@ -31,6 +31,14 @@ static FAutoConsoleVariableRef CVarPSOPrecacheResources(
 	ECVF_ReadOnly
 );
 
+int32 GPSOProxyCreationWhenPSOReady = 0;
+static FAutoConsoleVariableRef CVarPSOProxyCreationWhenPSOReady(
+	TEXT("r.PSOPrecache.ProxyCreationWhenPSOReady"),
+	GPSOProxyCreationWhenPSOReady,
+	TEXT("Delay the component proxy creation when the requested PSOs for precaching are still compiling."),
+	ECVF_ReadOnly
+);
+
 PSOCollectorCreateFunction FPSOCollectorCreateManager::JumpTable[(int32)EShadingPath::Num][FPSOCollectorCreateManager::MaxPSOCollectorCount] = {};
 
 bool IsComponentPSOPrecachingEnabled()
@@ -43,6 +51,11 @@ bool IsResourcePSOPrecachingEnabled()
 	return FApp::CanEverRender() && PipelineStateCache::IsPSOPrecachingEnabled() && GPSOPrecacheResources && !GIsEditor;
 }
 
+bool ProxyCreationWhenPSOReady()
+{
+	return FApp::CanEverRender() && PipelineStateCache::IsPSOPrecachingEnabled() && GPSOProxyCreationWhenPSOReady && !GIsEditor;
+}
+
 FGraphEventArray PrecachePSOs(const TArray<FPSOPrecacheData>& PSOInitializers)
 {
 	FGraphEventArray GraphEvents;
diff --git a/Engine/Source/Runtime/Engine/Private/StaticMeshRender.cpp b/Engine/Source/Runtime/Engine/Private/StaticMeshRender.cpp
--- a/Engine/Source/Runtime/Engine/Private/StaticMeshRender.cpp
+++ b/Engine/Source/Runtime/Engine/Private/StaticMeshRender.cpp
@@ -2484,6 +2484,12 @@ FPrimitiveSceneProxy* UStaticMeshComponent::CreateSceneProxy()
 		return nullptr;
 	}
 
+	if (IsPSOPrecaching())
+	{
+		UE_LOG(LogStaticMesh, Verbose, TEXT("Skipping CreateSceneProxy for StaticMeshComponent %s (Static mesh component PSOs are still compiling)"), *GetFullName());
+		return nullptr;
+	}
+
 	// Is Nanite supported, and is there built Nanite data for this static mesh?
 	if (ShouldCreateNaniteProxy() 
 	#if WITH_EDITORONLY_DATA
diff --git a/Engine/Source/Runtime/Engine/Public/MaterialShared.h b/Engine/Source/Runtime/Engine/Public/MaterialShared.h
--- a/Engine/Source/Runtime/Engine/Public/MaterialShared.h
+++ b/Engine/Source/Runtime/Engine/Public/MaterialShared.h
@@ -1745,6 +1745,11 @@ public:
 	 */
 	ENGINE_API FGraphEventArray CollectPSOs(ERHIFeatureLevel::Type InFeatureLevel, const TConstArrayView<const FVertexFactoryType*>& VertexFactoryTypes, const FPSOPrecacheParams& PreCacheParams);
 
+	/**
+	 * Helper function used by background thread when PSO collection for given vertex factory type and params are done - it also passes the currently compiling PSO graph events to keep track of
+	 */
+	void AddPrecachedPSOVertexFactoryType(const FVertexFactoryType* VertexFactoryType, const FPSOPrecacheParams& PreCacheParams, const FGraphEventArray& PSOCompileEvents);
+
 	/**
 	 * Should the shader for this material with the given platform, shader type and vertex 
 	 * factory type combination be compiled
@@ -2385,8 +2390,19 @@ private:
 			return HashCombine(GetTypeHash(Params.PrecachePSOParams), GetTypeHash(Params.VertexFactoryType));
 		}
 	};
+
+	/**
+	 * Struct used to keep track of all precached vertex factories with pso params and the currently still compiling graph events 
+	 * No map is used reduce the memory footprint for storage (mostly for platforms where PSO precaching is not enabled)
+	 */
+	struct FPSOPrecacheEntry
+	{
+		FPrecacheVertexTypeWithParams PrecacheVertexTypeWithParams;
+		FGraphEventArray CompilingGraphEvents;
+	};
+
 	FRWLock PrecachePSOVFLock;
-	TArray<FPrecacheVertexTypeWithParams> PrecachedPSOVertexFactories;
+	TArray<FPSOPrecacheEntry> PrecachedPSOVertexFactories;
 
 #if WITH_EDITOR
 	/**
diff --git a/Engine/Source/Runtime/Engine/Public/PSOPrecache.h b/Engine/Source/Runtime/Engine/Public/PSOPrecache.h
--- a/Engine/Source/Runtime/Engine/Public/PSOPrecache.h
+++ b/Engine/Source/Runtime/Engine/Public/PSOPrecache.h
@@ -165,6 +165,11 @@ extern ENGINE_API bool IsComponentPSOPrecachingEnabled();
  */
 extern ENGINE_API bool IsResourcePSOPrecachingEnabled();
 
+/**
+ * Delay component proxy creation when it's requested PSOs are still precaching
+ */
+extern ENGINE_API bool ProxyCreationWhenPSOReady();
+
 /**
  * Try and create PSOs for all the given initializers and return an optional array of graph events of async compiling PSOs
  */
