PR #9219: Fix FRootMotionSource::AccumulateMode not being checked inside the "Matches" checks, which could cause erroneous matching leading to movement de-syncs

diff --git a/Engine/Source/Runtime/Engine/Private/GameFramework/RootMotionSource.cpp b/Engine/Source/Runtime/Engine/Private/GameFramework/RootMotionSource.cpp
--- a/Engine/Source/Runtime/Engine/Private/GameFramework/RootMotionSource.cpp
+++ b/Engine/Source/Runtime/Engine/Private/GameFramework/RootMotionSource.cpp
@@ -230,7 +230,8 @@ bool FRootMotionSource::Matches(const FRootMotionSource* Other) const
 {
 	return Other != nullptr && 
 		GetScriptStruct() == Other->GetScriptStruct() && 
-		Priority == Other->Priority && 
+		Priority == Other->Priority &&
+		AccumulateMode == Other->AccumulateMode &&
 		bInLocalSpace == Other->bInLocalSpace &&
 		InstanceName == Other->InstanceName &&
 		FMath::IsNearlyEqual(Duration, Other->Duration, UE_SMALL_NUMBER);
