Fix for an editor crash caused by reinstancing a transacted AnimBP outered to a reinstanced actor-based Blueprint w/ a skeletal mesh component.

Change summary:
- Added a UActorComponent::OnComponentDestroyed() override to USkeletalMeshComponent to clear the current AnimBP instance.
- Modified USkeletalMeshComponent::ClearAnimScriptInstance() to mark the current AnimBP instance as garbage before clearing the reference.
- Modified UAnimInstance::BeginDestroy() to unregister the instance from the global core delegate list for future reinstanced UObject notifications.

diff --git a/Engine/Source/Runtime/Engine/Classes/Components/SkeletalMeshComponent.h b/Engine/Source/Runtime/Engine/Classes/Components/SkeletalMeshComponent.h
--- a/Engine/Source/Runtime/Engine/Classes/Components/SkeletalMeshComponent.h
+++ b/Engine/Source/Runtime/Engine/Classes/Components/SkeletalMeshComponent.h
@@ -1663,6 +1663,7 @@ public:
 	virtual void TickComponent(float DeltaTime, enum ELevelTick TickType, FActorComponentTickFunction *ThisTickFunction) override;
 	virtual void BeginPlay() override;
 	virtual void SetComponentTickEnabled(bool bEnabled) override;
+	virtual void OnComponentDestroyed(bool bDestroyingHierarchy) override;
 
 	//Handle registering our end physics tick function
 	virtual void RegisterEndPhysicsTick(bool bRegister);
diff --git a/Engine/Source/Runtime/Engine/Private/Animation/AnimInstance.cpp b/Engine/Source/Runtime/Engine/Private/Animation/AnimInstance.cpp
--- a/Engine/Source/Runtime/Engine/Private/Animation/AnimInstance.cpp
+++ b/Engine/Source/Runtime/Engine/Private/Animation/AnimInstance.cpp
@@ -1328,6 +1328,13 @@ void UAnimInstance::BeginDestroy()
 		AnimInstanceProxy = nullptr;
 	}
 
+#if WITH_EDITOR
+	if (!HasAnyFlags(RF_ClassDefaultObject))
+	{
+		FCoreUObjectDelegates::OnObjectsReinstanced.RemoveAll(this);
+	}
+#endif // WITH_EDITOR
+
 	Super::BeginDestroy();
 }
 
diff --git a/Engine/Source/Runtime/Engine/Private/Components/SkeletalMeshComponent.cpp b/Engine/Source/Runtime/Engine/Private/Components/SkeletalMeshComponent.cpp
--- a/Engine/Source/Runtime/Engine/Private/Components/SkeletalMeshComponent.cpp
+++ b/Engine/Source/Runtime/Engine/Private/Components/SkeletalMeshComponent.cpp
@@ -410,6 +410,13 @@ void USkeletalMeshComponent::SetComponentTickEnabled(bool bEnabled)
 	}
 }
 
+void USkeletalMeshComponent::OnComponentDestroyed(bool bDestroyingHierarchy)
+{
+	ClearAnimScriptInstance();
+
+	Super::OnComponentDestroyed(bDestroyingHierarchy);
+}
+
 void USkeletalMeshComponent::RegisterEndPhysicsTick(bool bRegister)
 {
 	if (bRegister != EndPhysicsTickFunction.IsTickFunctionRegistered())
@@ -1000,6 +1007,7 @@ void USkeletalMeshComponent::ClearAnimScriptInstance()
 		HandleExistingParallelEvaluationTask(bBlockOnTask, bPerformPostAnimEvaluation);
 
 		AnimScriptInstance->EndNotifyStates();
+		AnimScriptInstance->MarkAsGarbage();
 	}
 	AnimScriptInstance = nullptr;
 	ResetLinkedAnimInstances();
