Fix crash when re-instancing linked anim instances that have custom proxies

diff --git a/Engine/Source/Runtime/Engine/Private/Animation/AnimNode_LinkedAnimLayer.cpp b/Engine/Source/Runtime/Engine/Private/Animation/AnimNode_LinkedAnimLayer.cpp
--- a/Engine/Source/Runtime/Engine/Private/Animation/AnimNode_LinkedAnimLayer.cpp
+++ b/Engine/Source/Runtime/Engine/Private/Animation/AnimNode_LinkedAnimLayer.cpp
@@ -229,11 +229,16 @@ void FAnimNode_LinkedAnimLayer::HandleObjectsReinstanced_Impl(UObject* InSourceO
 	if(UseLegacyAnimInstanceReinstancingBehavior == nullptr || !UseLegacyAnimInstanceReinstancingBehavior->GetBool())
 	{
 		UAnimInstance* SourceAnimInstance = CastChecked<UAnimInstance>(InSourceObject);
+		FAnimInstanceProxy& SourceProxy = SourceAnimInstance->GetProxyOnAnyThread<FAnimInstanceProxy>();
+
+		// Call Initialize here to ensure any custom proxies are initialized (as they may have been re-created during
+		// re-instancing, and they dont call the constructor that takes a UAnimInstance*)
+		SourceProxy.Initialize(SourceAnimInstance);
+
 		InitializeProperties(SourceAnimInstance, Interface.Get() != nullptr ? Interface.Get() : InSourceObject->GetClass());
 		DynamicUnlink(SourceAnimInstance);
 		DynamicLink(SourceAnimInstance);
 
-		FAnimInstanceProxy& SourceProxy = SourceAnimInstance->GetProxyOnAnyThread<FAnimInstanceProxy>();
 		SourceProxy.InitializeCachedClassData();
 		
 		FAnimationInitializeContext Context(&SourceProxy);
