Fix static analysis warning

diff --git a/Engine/Source/Runtime/Engine/Private/StaticMesh.cpp b/Engine/Source/Runtime/Engine/Private/StaticMesh.cpp
--- a/Engine/Source/Runtime/Engine/Private/StaticMesh.cpp
+++ b/Engine/Source/Runtime/Engine/Private/StaticMesh.cpp
@@ -1362,7 +1362,7 @@ void FStaticMeshLODResources::InitResources(UStaticMesh* Parent, int32 LODIndex)
 		BeginInitResource(&AdditionalIndexBuffers->ReversedDepthOnlyIndexBuffer);
 	}
 
-	if (Parent->bSupportGpuUniformlyDistributedSampling && Parent->bSupportUniformlyDistributedSampling && (AreaWeightedSampler.GetNumEntries() > 0))
+	if (Parent && Parent->bSupportGpuUniformlyDistributedSampling && Parent->bSupportUniformlyDistributedSampling && (AreaWeightedSampler.GetNumEntries() > 0))
 	{
 		AreaWeightedSectionSamplersBuffer.Init(&AreaWeightedSectionSamplers);
 		AreaWeightedSectionSamplersBuffer.SetOwnerName(OwnerName);
@@ -1370,7 +1370,7 @@ void FStaticMeshLODResources::InitResources(UStaticMesh* Parent, int32 LODIndex)
 	}
 
 #if RHI_RAYTRACING
-	if (IsRayTracingAllowed() && Parent->bSupportRayTracing)
+	if (IsRayTracingAllowed() && Parent && Parent->bSupportRayTracing)
 	{
 		const bool bProceduralPrimitive = Parent->HasValidNaniteData() && Nanite::GetSupportsRayTracingProceduralPrimitive(GMaxRHIShaderPlatform);
 		ENQUEUE_RENDER_COMMAND(InitStaticMeshRayTracingGeometry)(
@@ -1404,7 +1404,7 @@ void FStaticMeshLODResources::InitResources(UStaticMesh* Parent, int32 LODIndex)
 
 #if !(UE_BUILD_SHIPPING || UE_BUILD_TEST)
 	ENQUEUE_RENDER_COMMAND(NameRHIResources)(
-		[this, DebugName = Parent->GetFName()](FRHICommandListImmediate&)
+		[this, DebugName = (Parent ? Parent->GetFName() : NAME_None)](FRHICommandListImmediate&)
 	{
 		TStringBuilder<512> StringBuilder;
 		auto SetDebugName = [&StringBuilder, DebugName](FRHIBuffer* RHIBuffer, const TCHAR* Extension)
