Fix licensee-reported null access crash when two debug camera modes are active at once.
The debug camera code has always been a bit flaky, but this visilbility function is fine to skip if things are confused.

diff --git a/Engine/Source/Runtime/Engine/Private/DebugCameraController.cpp b/Engine/Source/Runtime/Engine/Private/DebugCameraController.cpp
--- a/Engine/Source/Runtime/Engine/Private/DebugCameraController.cpp
+++ b/Engine/Source/Runtime/Engine/Private/DebugCameraController.cpp
@@ -492,7 +492,11 @@ void ADebugCameraController::OnDeactivate( APlayerController* RestoredPC )
 	ConsoleCommand(TEXT("show camfrustums"));
 	DrawFrustum->UnregisterComponent();
 	RestoredPC->SetActorHiddenInGame(true);
-	RestoredPC->PlayerCameraManager->SetActorHiddenInGame(true);
+	
+	if (RestoredPC->PlayerCameraManager)
+	{
+		RestoredPC->PlayerCameraManager->SetActorHiddenInGame(true);
+	}
 
 	OriginalControllerRef = nullptr;
 	OriginalPlayer = nullptr;
