Fix actor GetNetMode early in PIE/network startup by deferring to the world if there is no net driver yet. This allows proper detection of dedicated server at the actor level

diff --git a/Engine/Source/Runtime/Engine/Private/Actor.cpp b/Engine/Source/Runtime/Engine/Private/Actor.cpp
--- a/Engine/Source/Runtime/Engine/Private/Actor.cpp
+++ b/Engine/Source/Runtime/Engine/Private/Actor.cpp
@@ -4722,9 +4722,10 @@ ENetMode AActor::InternalGetNetMode() const
 		return NetDriver->GetNetMode();
 	}
 
-	if (UDemoNetDriver* DemoNetDriver = World ? World->GetDemoNetDriver() : nullptr)
+	if (World)
 	{
-		return DemoNetDriver->GetNetMode();
+		// World handles the demo net mode and has some special case checks for PIE
+		return World->GetNetMode();
 	}
 
 	return NM_Standalone;
