Change how thickness is setup on Strata material.
* Add a new thickness input on the root node for thin surface
* Add a new thickness input on vertical layer operator, to provide the thickness of the 'layered' material
* Mark Slab thickness input as deprecated, for converting existing asset once this CL is checked in.

diff --git a/Engine/Plugins/Enterprise/GLTFExporter/Source/GLTFMaterialBaking/Private/GLTFExportMaterialProxy.h b/Engine/Plugins/Enterprise/GLTFExporter/Source/GLTFMaterialBaking/Private/GLTFExportMaterialProxy.h
--- a/Engine/Plugins/Enterprise/GLTFExporter/Source/GLTFMaterialBaking/Private/GLTFExportMaterialProxy.h
+++ b/Engine/Plugins/Enterprise/GLTFExporter/Source/GLTFMaterialBaking/Private/GLTFExportMaterialProxy.h
@@ -420,10 +420,13 @@ public:
 			return MaterialInterface->CompileProperty(Compiler, MP_ShadingModel);
 
 		}
+		else if (Property == MP_SurfaceThickness)
+		{
+			return MaterialInterface->CompileProperty(Compiler, MP_SurfaceThickness);
+		}
 		else if (Property == MP_FrontMaterial)
 		{
 			return MaterialInterface->CompileProperty(Compiler, MP_FrontMaterial);
-
 		}
 		else
 		{
diff --git a/Engine/Shaders/Private/Strata/Strata.ush b/Engine/Shaders/Private/Strata/Strata.ush
--- a/Engine/Shaders/Private/Strata/Strata.ush
+++ b/Engine/Shaders/Private/Strata/Strata.ush
@@ -1351,15 +1351,10 @@ FStrataData GetStrataSlabBSDF(
 {
 	FStrataData StrataData = GetInitialisedStrataData();
 
-	if (bIsThinSurface)
-	{
-		ThicknessCm = STRATA_LAYER_ISTHIN_THICKNESS_THRESHOLD_CM - 1e-5f;
-	}
-
 	SSSMFPScale = saturate(SSSMFPScale);
 	const bool bSSSProfilePlugged = SSSProfileID > 0.0f;
 	const bool bMFPPlugged = any(SSSMFP > 0);
-	bool bIsThin = bMFPPlugged && ThicknessCm < STRATA_LAYER_ISTHIN_THICKNESS_THRESHOLD_CM && !UseSSSDiffusion; // Thin is used for flagging 'legacy foliage' case
+	bool bIsThin = bMFPPlugged && bIsThinSurface; // Enable thin only if MFP is plugged, as otherwise there will be no back-tramission
 	bool bHasSSSWrap = bMFPPlugged && !UseSSSDiffusion;
 	bool bHasSSSProfile = !bIsThin && bSSSProfilePlugged && SSSMFPScale > 0.f;	// When no profile is provide, we use Burley. STRATA_TODO we should clarify as burley is also an option of the profile.
 	bool bHasSSS = bMFPPlugged || bHasSSSProfile || bHasSSSWrap;
@@ -1413,7 +1408,7 @@ FStrataData GetStrataSlabBSDF(
 	}
 
 	// About the thin, a.k.a. two sided lighting, model:
-	// 	- Thin lighting model is used when thickness becomes less than STRATA_LAYER_ISTHIN_THICKNESS_THRESHOLD_CM
+	// 	- Thin lighting model is used when MFP is plugged and the surface is flagged as 'is-thin-surface'
 	//  - It cannot work with a SSS Profile because the profile id is stored in the register used for the MFP (e.g. see SLAB_SSSPROFILEID).
 	// 	- Thin lighting model force disabled the SSSProfile.
 	//  - It can work with the Strata SSS because the MFP is shared in this case between both.
diff --git a/Engine/Shaders/Private/Strata/StrataExport.ush b/Engine/Shaders/Private/Strata/StrataExport.ush
--- a/Engine/Shaders/Private/Strata/StrataExport.ush
+++ b/Engine/Shaders/Private/Strata/StrataExport.ush
@@ -312,6 +312,7 @@ void PackStrataOut(
 				const uint SingleIsSimpleVolume		= 0x10;
 				const uint SingleIsWrap				= 0x20;
 				const uint SingleIsSSSProfile		= 0x40;
+				const uint SingleIsWrapThin			= SingleIsThin | SingleIsWrap;
 				// No need to add anisotropy since it is not supported by the simple mode.
 
 				uint OptimisedModeMask = 0;
@@ -353,9 +354,8 @@ void PackStrataOut(
 				else if ((OptimisedModeMask & SingleHasCloth) == SingleHasCloth)
 				{
 					OptimisedLegacyMode = SINGLE_OPTLEGACYMODE_CLOTH;
-				}
-				
-				else if ((OptimisedModeMask & SingleIsThin) == SingleIsThin)
+				}				
+				else if ((OptimisedModeMask & SingleIsWrapThin) == SingleIsWrapThin)
 				{
 					OptimisedLegacyMode = SINGLE_OPTLEGACYMODE_SSSWRAP_THIN;
 					SSSData.Header.Bytes = 0; // Force to skip SSSData writing
diff --git a/Engine/Shaders/Private/Strata/StrataLegacyConversion.ush b/Engine/Shaders/Private/Strata/StrataLegacyConversion.ush
--- a/Engine/Shaders/Private/Strata/StrataLegacyConversion.ush
+++ b/Engine/Shaders/Private/Strata/StrataLegacyConversion.ush
@@ -344,7 +344,7 @@ FStrataData StrataConvertLegacyMaterialDynamic(
 	else if (ShadingModel == SHADINGMODELID_TWOSIDED_FOLIAGE)
 	{
 		// Set a thickness that will enabled the thin lighting model (corresponding to the legacy two-sided lighting model)
-		const float Thickness = STRATA_LAYER_ISTHIN_THICKNESS_THRESHOLD_CM - 1e-5f;
+		const float Thickness = STRATA_LAYER_DEFAULT_THICKNESS_CM;
 		const float3 MFPInCm = TransmittanceToMeanFreePath(SubSurfaceColor /*TransmittanceColor*/, Thickness * CENTIMETER_TO_METER) * METER_TO_CENTIMETER;
 
 		LegacySlab.BaseColor = BaseColor;
diff --git a/Engine/Shaders/Private/Strata/StrataTree.ush b/Engine/Shaders/Private/Strata/StrataTree.ush
--- a/Engine/Shaders/Private/Strata/StrataTree.ush
+++ b/Engine/Shaders/Private/Strata/StrataTree.ush
@@ -18,6 +18,7 @@
 void PostProcessSlabBeforeLighting(inout FStrataBSDF BSDF)
 {
 	// Two sided lighting can only be evaluated on the bottom layer, together with option SSS. Otherwise, the simple volume will be used.
+	// At the moment, a material can't be thin & simple volume at the same time. This should be fixed by unifying SSSDiffusion|SSSWrap|SimpleVolume into a single evaluation function
 	const bool bIsThin = BSDF.bIsBottom && BSDF_GETISTHIN(BSDF);
 
 	// We want simple volumetric only if the MFP input is pluged in, otherwise we avoid dealing with simple volumetric.
diff --git a/Engine/Shaders/Private/Strata/StrataVisualize.usf b/Engine/Shaders/Private/Strata/StrataVisualize.usf
--- a/Engine/Shaders/Private/Strata/StrataVisualize.usf
+++ b/Engine/Shaders/Private/Strata/StrataVisualize.usf
@@ -104,7 +104,8 @@ void PrintPixelType(inout FShaderPrintContext Ctx, in FStrataPixelHeader Header,
 
 void PrintSSSType(inout FShaderPrintContext Ctx, uint SSSType, bool bSimpleVolume, bool bThin, FFontColor InColor)
 {
-		 if (SSSType == SSS_TYPE_WRAP)				Print(Ctx, TEXT("Wrap"), InColor);
+		 if (SSSType == SSS_TYPE_WRAP && bThin)		Print(Ctx, TEXT("Wrap & Thin"), InColor);
+	else if (SSSType == SSS_TYPE_WRAP)				Print(Ctx, TEXT("Wrap"), InColor);
 	else if (SSSType == SSS_TYPE_DIFFUSION)			Print(Ctx, TEXT("Diffusion"), InColor);
 	else if (SSSType == SSS_TYPE_DIFFUSION_PROFILE)	Print(Ctx, TEXT("Diffusion Profile"), InColor);
 	else if (bSimpleVolume)							Print(Ctx, TEXT("Simple Volume"), InColor);
diff --git a/Engine/Shaders/Shared/StrataDefinitions.h b/Engine/Shaders/Shared/StrataDefinitions.h
--- a/Engine/Shaders/Shared/StrataDefinitions.h
+++ b/Engine/Shaders/Shared/StrataDefinitions.h
@@ -36,9 +36,6 @@
 // The default thickness of a layer is considered to be 0.01 centimeter = 0.1 millimeter
 #define STRATA_LAYER_DEFAULT_THICKNESS_CM	0.01f
 
-// The thin lighting model is used when the thickness becomes strictly lower than 0.1 millimeter
-#define STRATA_LAYER_ISTHIN_THICKNESS_THRESHOLD_CM 0.01f
-
 #define STRATA_BASE_PASS_MRT_OUTPUT_COUNT	3
 
 #define STRATA_SSS_DATA_UINT_COUNT		2
diff --git a/Engine/Source/Developer/MaterialBaking/Private/ExportMaterialProxy.h b/Engine/Source/Developer/MaterialBaking/Private/ExportMaterialProxy.h
--- a/Engine/Source/Developer/MaterialBaking/Private/ExportMaterialProxy.h
+++ b/Engine/Source/Developer/MaterialBaking/Private/ExportMaterialProxy.h
@@ -424,10 +424,13 @@ public:
 			return MaterialInterface->CompileProperty(Compiler, MP_ShadingModel);
 
 		}
+		else if (Property == MP_SurfaceThickness)
+		{
+			return MaterialInterface->CompileProperty(Compiler, MP_SurfaceThickness);
+		}
 		else if (Property == MP_FrontMaterial)
 		{
 			return MaterialInterface->CompileProperty(Compiler, MP_FrontMaterial);
-
 		}
 		else
 		{
diff --git a/Engine/Source/Developer/MaterialUtilities/Private/MaterialUtilities.cpp b/Engine/Source/Developer/MaterialUtilities/Private/MaterialUtilities.cpp
--- a/Engine/Source/Developer/MaterialUtilities/Private/MaterialUtilities.cpp
+++ b/Engine/Source/Developer/MaterialUtilities/Private/MaterialUtilities.cpp
@@ -597,6 +597,8 @@ public:
 				break;
 			case MP_ShadingModel:
 				return MaterialInterface->CompileProperty(&ProxyCompiler, MP_ShadingModel);
+			case MP_SurfaceThickness:
+				return MaterialInterface->CompileProperty(&ProxyCompiler, MP_SurfaceThickness);
 			case MP_FrontMaterial:
 				return MaterialInterface->CompileProperty(&ProxyCompiler, MP_FrontMaterial);
 			default:
@@ -619,6 +621,10 @@ public:
 		{
 			return MaterialInterface->CompileProperty(Compiler, MP_ShadingModel);
 		}
+		else if (Property == MP_SurfaceThickness)
+		{
+			return MaterialInterface->CompileProperty(Compiler, MP_SurfaceThickness);
+		}
 		else if (Property == MP_FrontMaterial)
 		{
 			return MaterialInterface->CompileProperty(Compiler, MP_FrontMaterial);
diff --git a/Engine/Source/Editor/MaterialEditor/Private/MaterialEditor.cpp b/Engine/Source/Editor/MaterialEditor/Private/MaterialEditor.cpp
--- a/Engine/Source/Editor/MaterialEditor/Private/MaterialEditor.cpp
+++ b/Engine/Source/Editor/MaterialEditor/Private/MaterialEditor.cpp
@@ -4627,6 +4627,7 @@ UClass* FMaterialEditor::GetOnPromoteToParameterClass(const UEdGraphPin* TargetP
 			case MP_PixelDepthOffset:
 			case MP_ShadingModel:
 			case MP_OpacityMask:
+			case MP_SurfaceThickness:
 				return UMaterialExpressionScalarParameter::StaticClass();
 
 			case MP_WorldPositionOffset:
diff --git a/Engine/Source/Editor/UnrealEd/Private/Lightmass/LightmassRender.cpp b/Engine/Source/Editor/UnrealEd/Private/Lightmass/LightmassRender.cpp
--- a/Engine/Source/Editor/UnrealEd/Private/Lightmass/LightmassRender.cpp
+++ b/Engine/Source/Editor/UnrealEd/Private/Lightmass/LightmassRender.cpp
@@ -503,6 +503,10 @@ public:
 		{
 			return MaterialInterface->CompileProperty(Compiler, MP_OpacityMask);
 		}
+		else if (Property == MP_SurfaceThickness)
+		{
+			return MaterialInterface->CompileProperty(Compiler, MP_SurfaceThickness);
+		}
 		else if (Property == MP_FrontMaterial)
 		{
 			if (bStrataEnabled)
diff --git a/Engine/Source/Editor/UnrealEd/Private/MaterialGraph.cpp b/Engine/Source/Editor/UnrealEd/Private/MaterialGraph.cpp
--- a/Engine/Source/Editor/UnrealEd/Private/MaterialGraph.cpp
+++ b/Engine/Source/Editor/UnrealEd/Private/MaterialGraph.cpp
@@ -258,6 +258,7 @@ void UMaterialGraph::RebuildGraphInternal(const TMap<UMaterialExpression*, TArra
 
 		MaterialInputs.Add(FMaterialInputInfo(FMaterialAttributeDefinitionMap::GetDisplayNameForMaterial(MP_PixelDepthOffset, Material), MP_PixelDepthOffset, LOCTEXT("PixelDepthOffsetToolTip", "Pixel Depth Offset")));
 		MaterialInputs.Add(FMaterialInputInfo(FMaterialAttributeDefinitionMap::GetDisplayNameForMaterial(MP_ShadingModel, Material), MP_ShadingModel, LOCTEXT("ShadingModelToolTip", "Selects which shading model should be used per pixel")));
+		MaterialInputs.Add(FMaterialInputInfo(FMaterialAttributeDefinitionMap::GetDisplayNameForMaterial(MP_SurfaceThickness, Material), MP_SurfaceThickness, LOCTEXT("SurfaceThicknessToolTip", "Defines the surface's thickness when IsThinSurface is enabled")));
 		MaterialInputs.Add(FMaterialInputInfo(FMaterialAttributeDefinitionMap::GetDisplayNameForMaterial(MP_FrontMaterial, Material), MP_FrontMaterial, LOCTEXT("FrontMaterialToolTip", "Specify the front facing material")));
 
 		//^^^ New material properties go above here. ^^^^
diff --git a/Engine/Source/Runtime/Engine/Classes/Materials/Material.h b/Engine/Source/Runtime/Engine/Classes/Materials/Material.h
--- a/Engine/Source/Runtime/Engine/Classes/Materials/Material.h
+++ b/Engine/Source/Runtime/Engine/Classes/Materials/Material.h
@@ -375,6 +375,9 @@ public:
 	UPROPERTY()
 	FShadingModelMaterialInput ShadingModelFromMaterialExpression;
 
+	UPROPERTY()
+	FScalarMaterialInput SurfaceThickness;
+
 	UPROPERTY()
 	FStrataMaterialInput FrontMaterial;
 
@@ -1762,6 +1765,7 @@ public:
 	ENGINE_API bool HasMetallicConnected() const;
 	ENGINE_API bool HasEmissiveColorConnected() const;
 	ENGINE_API bool HasAnisotropyConnected() const;
+	ENGINE_API bool HasSurfaceThicknessConnected() const;
 	ENGINE_API bool HasStrataFrontMaterialConnected() const;
 	ENGINE_API bool HasVertexPositionOffsetConnected() const;
 	ENGINE_API bool HasPixelDepthOffsetConnected() const;
diff --git a/Engine/Source/Runtime/Engine/Classes/Materials/MaterialExpressionStrata.h b/Engine/Source/Runtime/Engine/Classes/Materials/MaterialExpressionStrata.h
--- a/Engine/Source/Runtime/Engine/Classes/Materials/MaterialExpressionStrata.h
+++ b/Engine/Source/Runtime/Engine/Classes/Materials/MaterialExpressionStrata.h
@@ -281,7 +281,7 @@ class UMaterialExpressionStrataSlabBSDF : public UMaterialExpressionStrataBSDF
 	 * The slab thickness. (type = float, centimeters, default = 0.01 centimeter = 0.1 millimeter)
 	 */
 	UPROPERTY()
-	FExpressionInput Thickness;
+	FExpressionInput Thickness; // TODO: Remove
 
 	/**
 	 * The amount of fuzz on top of the surface used to simulate cloth-like appearance.
@@ -823,6 +823,12 @@ class UMaterialExpressionStrataVerticalLayering : public UMaterialExpressionStra
 	UPROPERTY()
 	FExpressionInput Base;
 
+	/**
+	 * Thickness of the Top material layer
+	 */
+	UPROPERTY()
+	FExpressionInput Thickness;
+
 	/**
 	 * Merge Top and Base into a single material by mixing their inputs rather than their evaluation. This makes lighting evaluation cheaper (Default: off)
 	 */
@@ -838,6 +844,7 @@ class UMaterialExpressionStrataVerticalLayering : public UMaterialExpressionStra
 	virtual bool IsResultStrataMaterial(int32 OutputIndex) override;
 	virtual void GatherStrataMaterialInfo(FStrataMaterialInfo& StrataMaterialInfo, int32 OutputIndex) override;
 	virtual FStrataOperator* StrataGenerateMaterialTopologyTree(class FMaterialCompiler* Compiler, UMaterialExpression* Parent, int32 OutputIndex) override;
+	virtual FName GetInputName(int32 InputIndex) const override;
 #endif
 	//~ End UMaterialExpression Interface
 };
diff --git a/Engine/Source/Runtime/Engine/Private/Materials/HLSLMaterialTranslator.cpp b/Engine/Source/Runtime/Engine/Private/Materials/HLSLMaterialTranslator.cpp
--- a/Engine/Source/Runtime/Engine/Private/Materials/HLSLMaterialTranslator.cpp
+++ b/Engine/Source/Runtime/Engine/Private/Materials/HLSLMaterialTranslator.cpp
@@ -315,6 +315,7 @@ FHLSLMaterialTranslator::FHLSLMaterialTranslator(FMaterial* InMaterial,
 	SharedPixelProperties[MP_PixelDepthOffset] = true;
 	SharedPixelProperties[MP_SubsurfaceColor] = true;
 	SharedPixelProperties[MP_ShadingModel] = true;
+	SharedPixelProperties[MP_SurfaceThickness] = true;
 	SharedPixelProperties[MP_FrontMaterial] = true;
 
 	for (int32 Frequency = 0; Frequency < SF_NumFrequencies; ++Frequency)
@@ -742,7 +743,18 @@ bool FHLSLMaterialTranslator::Translate()
 		#if DEBUG_STRATA_TREE_STACK
 			UE_LOG(LogMaterial, Display, TEXT(" StrataTreeStack: StrataGenerateMaterialTopologyTree"));
 		#endif
-			FrontMaterialExpr->StrataGenerateMaterialTopologyTree(this, nullptr, 0);
+			{
+				StrataThicknessIndexToExpressionInput.SetNum(0);
+				StrataThicknessStack.SetNum(0);
+			
+				FExpressionInput* SurfaceThickness = Material->IsThinSurface() && Material->GetMaterialInterface() ? &Material->GetMaterialInterface()->GetMaterial()->GetEditorOnlyData()->SurfaceThickness : nullptr;
+				StrataThicknessStackPush(nullptr, SurfaceThickness);
+				FrontMaterialExpr->StrataGenerateMaterialTopologyTree(this, nullptr, 0);
+				StrataThicknessStackPop();
+
+				check(StrataThicknessStack.Num() == 0);
+			}
+
 			if (!StrataGenerateDerivedMaterialOperatorData())
 			{
 				Errorf(TEXT("Strata material errors encountered."));
@@ -872,6 +884,7 @@ bool FHLSLMaterialTranslator::Translate()
 	#if DEBUG_STRATA_TREE_STACK
 		UE_LOG(LogMaterial, Display, TEXT(" StrataTreeStack: Material->CompilePropertyAndSetMaterialProperty(MP_FrontMaterial)"));
 	#endif
+		Chunk[MP_SurfaceThickness]				= Material->CompilePropertyAndSetMaterialProperty(MP_SurfaceThickness		,this);
 		Chunk[MP_FrontMaterial]					= Material->CompilePropertyAndSetMaterialProperty(MP_FrontMaterial			,this);
 		}
 
@@ -10700,6 +10713,7 @@ int32 FHLSLMaterialTranslator::StrataSlabBSDF(
 {
 	const FString NormalCode = GetParameterCode(Normal);
 	const FString TangentCode = Tangent != INDEX_NONE ? *GetParameterCode(Tangent) : TEXT("NONE");
+	const FString ThicknessCode = GetParameterCode(Thickness);
 	const bool bIsThinSurface = Material->IsThinSurface();
 
 	if (PromoteToOperator)
@@ -10710,7 +10724,7 @@ int32 FHLSLMaterialTranslator::StrataSlabBSDF(
 			return INDEX_NONE;
 		}
 		return AddCodeChunk(
-			MCT_Strata, TEXT("PromoteParameterBlendedBSDFToOperator(GetStrataSlabBSDF(Parameters.StrataPixelFootprint, %s, %s, %s, %s, %s, %s, %s, %s, %s, %s, %s, %s, %s, %s, %s, %s, %s, %s, %s, Parameters.SharedLocalBases.Types) /* Normal = %s ; Tangent = %s */, Parameters.StrataTree, %u, %u, %u, %u)"),
+			MCT_Strata, TEXT("PromoteParameterBlendedBSDFToOperator(GetStrataSlabBSDF(Parameters.StrataPixelFootprint, %s, %s, %s, %s, %s, %s, %s, %s, %s, %s, %s, %s, %s, %s, %s, %s, %s, %s, %s, Parameters.SharedLocalBases.Types) /* Normal = %s ; Tangent = %s ; Thickness = %s */, Parameters.StrataTree, %u, %u, %u, %u)"),
 			*StrataGetCastParameterCode(DiffuseAlbedo,			MCT_Float3),
 			*StrataGetCastParameterCode(F0,						MCT_Float3),
 			*StrataGetCastParameterCode(F90,					MCT_Float3),
@@ -10732,6 +10746,7 @@ int32 FHLSLMaterialTranslator::StrataSlabBSDF(
 			*SharedLocalBasisIndexMacro,
 			*NormalCode,
 			*TangentCode,
+			*ThicknessCode,
 			PromoteToOperator->Index,
 			PromoteToOperator->BSDFIndex,
 			PromoteToOperator->LayerDepth,
@@ -10740,7 +10755,7 @@ int32 FHLSLMaterialTranslator::StrataSlabBSDF(
 	}
 	
 	return AddCodeChunk(
-		MCT_Strata, TEXT("GetStrataSlabBSDF(Parameters.StrataPixelFootprint, %s, %s, %s, %s, %s, %s, %s, %s, %s, %s, %s, %s, %s, %s, %s, %s, %s, %s, %s, Parameters.SharedLocalBases.Types) /* Normal = %s ; Tangent = %s */"),
+		MCT_Strata, TEXT("GetStrataSlabBSDF(Parameters.StrataPixelFootprint, %s, %s, %s, %s, %s, %s, %s, %s, %s, %s, %s, %s, %s, %s, %s, %s, %s, %s, %s, Parameters.SharedLocalBases.Types) /* Normal = %s ; Tangent = %s ; Thickness = %s */"),
 		*StrataGetCastParameterCode(DiffuseAlbedo,			MCT_Float3),
 		*StrataGetCastParameterCode(F0,						MCT_Float3),
 		*StrataGetCastParameterCode(F90,					MCT_Float3),
@@ -10761,7 +10776,8 @@ int32 FHLSLMaterialTranslator::StrataSlabBSDF(
 		bIsThinSurface ? TEXT("true") : TEXT("false"),
 		*SharedLocalBasisIndexMacro,
 		*NormalCode,
-		*TangentCode
+		*TangentCode,
+		*ThicknessCode
 	);
 }
 
@@ -11011,34 +11027,38 @@ int32 FHLSLMaterialTranslator::StrataHorizontalMixingParameterBlending(
 	);
 }
 
-int32 FHLSLMaterialTranslator::StrataVerticalLayering(int32 Top, int32 Base, int OperatorIndex, uint32 MaxDistanceFromLeaves)
+int32 FHLSLMaterialTranslator::StrataVerticalLayering(int32 Top, int32 Base, int32 Thickness, int OperatorIndex, uint32 MaxDistanceFromLeaves)
 {
 	if (Top == INDEX_NONE || Base == INDEX_NONE)
 	{
 		return INDEX_NONE;
 	}
+	const FString ThicknessCode = Thickness != INDEX_NONE ? GetParameterCode(Thickness) : TEXT("NONE");;
 	return AddCodeChunk(
-		MCT_Strata, TEXT("StrataVerticalLayering(%s, %s, Parameters.StrataTree, %u, %u)"),
+		MCT_Strata, TEXT("StrataVerticalLayering(%s, %s, Parameters.StrataTree, %u, %u) /* Thickness = %s */"),
 		*GetParameterCode(Top),
 		*GetParameterCode(Base),
 		OperatorIndex,
-		MaxDistanceFromLeaves
+		MaxDistanceFromLeaves,
+		*ThicknessCode
 	);
 }
 
-int32 FHLSLMaterialTranslator::StrataVerticalLayeringParameterBlending(int32 Top, int32 Base, const FString& SharedLocalBasisIndexMacro, int32 TopBSDFNormalCodeChunk, FStrataOperator* PromoteToOperator)
+int32 FHLSLMaterialTranslator::StrataVerticalLayeringParameterBlending(int32 Top, int32 Base, int32 Thickness, const FString& SharedLocalBasisIndexMacro, int32 TopBSDFNormalCodeChunk, FStrataOperator* PromoteToOperator)
 {
 	if (Top == INDEX_NONE || Base == INDEX_NONE)
 	{
 		return INDEX_NONE;
 	}
 
+	const FString ThicknessCode = Thickness != INDEX_NONE ? GetParameterCode(Thickness) : TEXT("NONE");
+
 	if (PromoteToOperator)
 	{
 		check(PromoteToOperator->Index != INDEX_NONE);
 		check(PromoteToOperator->BSDFIndex != INDEX_NONE);
 		return AddCodeChunk(
-			MCT_Strata, TEXT("PromoteParameterBlendedBSDFToOperator(StrataVerticalLayeringParameterBlending(%s, %s, %s, dot(%s, %s)), Parameters.StrataTree, %u, %u, %u, %u)"),
+			MCT_Strata, TEXT("PromoteParameterBlendedBSDFToOperator(StrataVerticalLayeringParameterBlending(%s, %s, %s, dot(%s, %s)), Parameters.StrataTree, %u, %u, %u, %u) /* Thickness = %s */"),
 			*GetParameterCode(Top),
 			*GetParameterCode(Base),
 			*SharedLocalBasisIndexMacro,
@@ -11047,17 +11067,19 @@ int32 FHLSLMaterialTranslator::StrataVerticalLayeringParameterBlending(int32 Top
 			PromoteToOperator->Index,
 			PromoteToOperator->BSDFIndex,
 			PromoteToOperator->LayerDepth,
-			PromoteToOperator->bIsBottom ? 1 : 0
+			PromoteToOperator->bIsBottom ? 1 : 0,
+			*ThicknessCode
 		);
 	}
 
 	return AddCodeChunk(
-		MCT_Strata, TEXT("StrataVerticalLayeringParameterBlending(%s, %s, %s, dot(%s, %s))"),
+		MCT_Strata, TEXT("StrataVerticalLayeringParameterBlending(%s, %s, %s, dot(%s, %s)) /* Thickness = %s */"),
 		*GetParameterCode(Top),
 		*GetParameterCode(Base),
 		*SharedLocalBasisIndexMacro,
 		*GetParameterCode(TopBSDFNormalCodeChunk),
-		*GetParameterCode(CameraVector())
+		*GetParameterCode(CameraVector()),
+		*ThicknessCode
 	);
 }
 
@@ -11347,6 +11369,49 @@ void FHLSLMaterialTranslator::StrataTreeStackPop()
 #endif
 }
 
+int32 FHLSLMaterialTranslator::StrataThicknessStackGetThicknessIndex()
+{
+	return StrataThicknessStack.Top();
+}
+
+int32 FHLSLMaterialTranslator::StrataThicknessStackGetThicknessCode(int32 Index)
+{
+	int32 OutCode = INDEX_NONE;
+	if (Index == INDEX_NONE || Index >= StrataThicknessIndexToExpressionInput.Num())
+	{
+		UE_LOG(LogMaterial, Error, TEXT(" StrataThichkness: %i could not be found)"), Index);
+	}
+	else if (FExpressionInput* Input = StrataThicknessIndexToExpressionInput[Index])
+	{
+		OutCode = Input->GetTracedInput().Expression ? Input->Compile(this) : Constant(STRATA_LAYER_DEFAULT_THICKNESS_CM);
+		EMaterialValueType Type = GetType(OutCode);
+		if (IsLWCType(Type))
+		{
+			Type = MakeNonLWCType(Type);
+			OutCode = ValidCast(OutCode, Type);
+		}	
+	}
+	if (OutCode == INDEX_NONE)
+	{
+		OutCode = Constant(STRATA_LAYER_DEFAULT_THICKNESS_CM);
+	}
+	return OutCode;
+}
+
+int32 FHLSLMaterialTranslator::StrataThicknessStackPush(UMaterialExpression* Expression, FExpressionInput* Input)
+{
+	int32 Index = StrataThicknessIndexToExpressionInput.Num();
+	StrataThicknessIndexToExpressionInput.Add(Input);
+	StrataThicknessStack.Push(Index);
+	return Index;	
+}
+
+void FHLSLMaterialTranslator::StrataThicknessStackPop()
+{
+	check(StrataThicknessStack.Num() >= 1);
+	StrataThicknessStack.Pop();
+}
+
 int32 FHLSLMaterialTranslator::MapARPassthroughCameraUV(int32 UV)
 {
 	if (UV == INDEX_NONE)
diff --git a/Engine/Source/Runtime/Engine/Private/Materials/HLSLMaterialTranslator.h b/Engine/Source/Runtime/Engine/Private/Materials/HLSLMaterialTranslator.h
--- a/Engine/Source/Runtime/Engine/Private/Materials/HLSLMaterialTranslator.h
+++ b/Engine/Source/Runtime/Engine/Private/Materials/HLSLMaterialTranslator.h
@@ -450,6 +450,12 @@ protected:
 	*/
 	TArray<FGuid> StrataNodeIdentifierStack;
 
+	/** Stack of thickness input used for propagating thickness information from root node and vertical operation
+	* This is transient and updated when calling StrataGenerateMaterialTopologyTree. The information is then stored into the FStratOperator
+	*/
+	TArray<int32> StrataThicknessStack;
+	TArray<FExpressionInput*> StrataThicknessIndexToExpressionInput;
+
 	struct FStrataSimplificationStatus
 	{
 		bool bMaterialFitsInMemoryBudget = false;	// Track whether or not the material fits.
@@ -1097,8 +1103,8 @@ protected:
 		int32 ColorScaleBehindWater, int32 Normal, const FString& SharedLocalBasisIndexMacro, FStrataOperator* PromoteToOperator) override;
 	virtual int32 StrataHorizontalMixing(int32 Background, int32 Foreground, int32 Mix, int OperatorIndex, uint32 MaxDistanceFromLeaves) override;
 	virtual int32 StrataHorizontalMixingParameterBlending(int32 Background, int32 Foreground, int32 HorizontalMixCodeChunk, int32 NormalMixCodeChunk, const FString& SharedLocalBasisIndexMacro, FStrataOperator* PromoteToOperator) override;
-	virtual int32 StrataVerticalLayering(int32 Top, int32 Base, int OperatorIndex, uint32 MaxDistanceFromLeaves) override;
-	virtual int32 StrataVerticalLayeringParameterBlending(int32 Top, int32 Base, const FString& SharedLocalBasisIndexMacro, int32 TopBSDFNormalCodeChunk, FStrataOperator* PromoteToOperator) override;
+	virtual int32 StrataVerticalLayering(int32 Top, int32 Base, int32 Thickness, int OperatorIndex, uint32 MaxDistanceFromLeaves) override;
+	virtual int32 StrataVerticalLayeringParameterBlending(int32 Top, int32 Base, int32 Thickness, const FString& SharedLocalBasisIndexMacro, int32 TopBSDFNormalCodeChunk, FStrataOperator* PromoteToOperator) override;
 	virtual int32 StrataAdd(int32 A, int32 B, int OperatorIndex, uint32 MaxDistanceFromLeaves) override;
 	virtual int32 StrataAddParameterBlending(int32 A, int32 B, int32 AMixWeight, const FString& SharedLocalBasisIndexMacro, FStrataOperator* PromoteToOperator) override;
 	virtual int32 StrataWeight(int32 A, int32 Weight, int OperatorIndex, uint32 MaxDistanceFromLeaves) override;
@@ -1115,6 +1121,11 @@ protected:
 	virtual FGuid StrataTreeStackGetParentPathUniqueId() override;
 	virtual void StrataTreeStackPop() override;
 	
+	virtual int32 StrataThicknessStackGetThicknessIndex() override;
+	virtual int32 StrataThicknessStackGetThicknessCode(int32 Index) override;
+	virtual int32 StrataThicknessStackPush(UMaterialExpression* Expression, FExpressionInput* Input) override;
+	virtual void StrataThicknessStackPop() override;
+
 	virtual FStrataOperator& StrataCompilationRegisterOperator(int32 OperatorType, FGuid StrataExpressionGuid, UMaterialExpression* Parent, FGuid StrataParentExpressionGuid, bool bUseParameterBlending = false) override;
 	virtual FStrataOperator& StrataCompilationGetOperator(FGuid StrataExpressionGuid) override;
 	virtual FStrataOperator* StrataCompilationGetOperatorFromIndex(int32 OperatorIndex) override;
diff --git a/Engine/Source/Runtime/Engine/Private/Materials/Material.cpp b/Engine/Source/Runtime/Engine/Private/Materials/Material.cpp
--- a/Engine/Source/Runtime/Engine/Private/Materials/Material.cpp
+++ b/Engine/Source/Runtime/Engine/Private/Materials/Material.cpp
@@ -2719,7 +2719,7 @@ void UMaterial::Serialize(FArchive& Ar)
 	}
 #endif // #if WITH_EDITOR
 
-	static_assert(MP_MAX == 33, "New material properties must have DoMaterialAttributeReorder called on them to ensure that any future reordering of property pins is correctly applied.");
+	static_assert(MP_MAX == 34, "New material properties must have DoMaterialAttributeReorder called on them to ensure that any future reordering of property pins is correctly applied.");
 
 	if (Ar.UEVer() < VER_UE4_MATERIAL_MASKED_BLENDMODE_TIDY)
 	{
@@ -3694,6 +3694,7 @@ void UMaterial::PostLoad()
 	DoMaterialAttributeReorder(&EditorOnly->PixelDepthOffset, UEVer, RenderObjVer, UE5MainVer);
 	DoMaterialAttributeReorder(&EditorOnly->ShadingModelFromMaterialExpression, UEVer, RenderObjVer, UE5MainVer);
 	DoMaterialAttributeReorder(&EditorOnly->FrontMaterial, UEVer, RenderObjVer, UE5MainVer);
+	DoMaterialAttributeReorder(&EditorOnly->SurfaceThickness, UEVer, RenderObjVer, UE5MainVer);
 
 	if (ParameterGroupData_DEPRECATED.Num() > 0)
 	{
@@ -4109,6 +4110,11 @@ bool UMaterial::HasAnisotropyConnected() const
 	return IsPropertyConnected(MP_Anisotropy);
 }
 
+bool UMaterial::HasSurfaceThicknessConnected() const
+{
+	return IsPropertyConnected(MP_SurfaceThickness);
+}
+
 bool UMaterial::HasStrataFrontMaterialConnected() const
 {
 	return IsPropertyConnected(MP_FrontMaterial);
@@ -5560,6 +5566,7 @@ bool UMaterial::GetExpressionInputDescription(EMaterialProperty InProperty, FMat
 	case MP_MaterialAttributes: SetMaterialInputDescription(EditorOnly->MaterialAttributes, false, OutDescription); return true;
 	case MP_PixelDepthOffset: SetMaterialInputDescription(EditorOnly->PixelDepthOffset, false, OutDescription); return true;
 	case MP_ShadingModel: SetMaterialInputDescription(EditorOnly->ShadingModelFromMaterialExpression, false, OutDescription); return true;
+	case MP_SurfaceThickness: SetMaterialInputDescription(EditorOnly->SurfaceThickness, false, OutDescription); return true;
 	case MP_FrontMaterial: SetMaterialInputDescription(EditorOnly->FrontMaterial, false, OutDescription); return true;
 	default:
 		if (InProperty >= MP_CustomizedUVs0 && InProperty <= MP_CustomizedUVs7)
@@ -6165,6 +6172,7 @@ int32 UMaterial::CompilePropertyEx( FMaterialCompiler* Compiler, const FGuid& At
 		case MP_WorldPositionOffset:	return EditorOnly->WorldPositionOffset.CompileWithDefault(Compiler, Property);
 		case MP_PixelDepthOffset:		return EditorOnly->PixelDepthOffset.CompileWithDefault(Compiler, Property);
 		case MP_ShadingModel:			return EditorOnly->ShadingModelFromMaterialExpression.CompileWithDefault(Compiler, Property);
+		case MP_SurfaceThickness:		return EditorOnly->SurfaceThickness.CompileWithDefault(Compiler, Property);
 		case MP_FrontMaterial:			return EditorOnly->FrontMaterial.CompileWithDefault(Compiler, Property);
 
 		default:
@@ -6437,6 +6445,7 @@ bool UMaterial::IsPropertySupported(EMaterialProperty InProperty) const
 		case MP_AmbientOcclusion:
 		case MP_WorldPositionOffset:
 		case MP_PixelDepthOffset:
+		case MP_SurfaceThickness:
 		case MP_FrontMaterial:
 			bSupported = true;
 			break;
@@ -6478,6 +6487,7 @@ static bool IsPropertyActive_Internal(EMaterialProperty InProperty,
 	bool bHasRefraction,
 	bool bUsesShadingModelFromMaterialExpression,
 	bool bIsTranslucencyWritingVelocity,
+	bool bIsThinSurface,
 	bool bIsSupported)
 {
 	const bool bStrataEnabled = Engine_IsStrataEnabled();
@@ -6619,6 +6629,9 @@ static bool IsPropertyActive_Internal(EMaterialProperty InProperty,
 			case MP_PixelDepthOffset:
 				Active = (!bIsTranslucentBlendMode) || (bIsTranslucencyWritingVelocity);
 				break;
+			case MP_SurfaceThickness:
+				Active = bIsThinSurface;
+				break;
 			case MP_FrontMaterial:
 				Active = true;
 				break;
@@ -6690,6 +6703,7 @@ static bool IsPropertyActive_Internal(EMaterialProperty InProperty,
 		case MP_ShadingModel:
 			Active = bUsesShadingModelFromMaterialExpression;
 			break;
+		case MP_SurfaceThickness:
 		case MP_FrontMaterial:
 			{
 				Active = bStrataEnabled;
@@ -6724,6 +6738,7 @@ bool UMaterial::IsPropertyActiveInEditor(EMaterialProperty InProperty) const
 		GetEditorOnlyData()->Refraction.IsConnected(),
 		IsShadingModelFromMaterialExpression(),
 		IsTranslucencyWritingVelocity(),
+		IsThinSurface(),
 		IsPropertySupported(InProperty));
 }
 #endif // WITH_EDITOR
@@ -6749,6 +6764,7 @@ bool UMaterial::IsPropertyActiveInDerived(EMaterialProperty InProperty, const UM
 		!bUseMaterialAttributes ? GetCachedExpressionData().IsPropertyConnected(MP_Refraction) : bUsesDistortion,
 		DerivedMaterial->IsShadingModelFromMaterialExpression(),
 		IsTranslucencyWritingVelocity(),
+		IsThinSurface(),
 		IsPropertySupported(InProperty));
 }
 
diff --git a/Engine/Source/Runtime/Engine/Private/Materials/MaterialCachedData.cpp b/Engine/Source/Runtime/Engine/Private/Materials/MaterialCachedData.cpp
--- a/Engine/Source/Runtime/Engine/Private/Materials/MaterialCachedData.cpp
+++ b/Engine/Source/Runtime/Engine/Private/Materials/MaterialCachedData.cpp
@@ -50,7 +50,7 @@ const FMaterialCachedParameterEntry FMaterialCachedParameterEntry::EmptyData{};
 const FMaterialCachedExpressionData FMaterialCachedExpressionData::EmptyData{};
 const FMaterialCachedExpressionEditorOnlyData FMaterialCachedExpressionEditorOnlyData::EmptyData{};
 
-static_assert((uint32)(EMaterialProperty::MP_MAX)-1 <= (8 * sizeof(FMaterialCachedExpressionData::PropertyConnectedBitmask)), "PropertyConnectedBitmask cannot contain entire EMaterialProperty enumeration.");
+static_assert((uint32)(EMaterialProperty::MP_MaterialAttributes)-1 <= (8 * sizeof(FMaterialCachedExpressionData::PropertyConnectedBitmask)), "PropertyConnectedBitmask cannot contain entire EMaterialProperty enumeration.");
 
 FMaterialCachedExpressionData::FMaterialCachedExpressionData()
 	: bHasMaterialLayers(false)
diff --git a/Engine/Source/Runtime/Engine/Private/Materials/MaterialExpressions.cpp b/Engine/Source/Runtime/Engine/Private/Materials/MaterialExpressions.cpp
--- a/Engine/Source/Runtime/Engine/Private/Materials/MaterialExpressions.cpp
+++ b/Engine/Source/Runtime/Engine/Private/Materials/MaterialExpressions.cpp
@@ -6691,7 +6691,7 @@ int32 UMaterialExpressionMakeMaterialAttributes::Compile(class FMaterialCompiler
 	int32 Ret = INDEX_NONE;
 	UMaterialExpression* Expression = nullptr;
 
- 	static_assert(MP_MAX == 33, 
+ 	static_assert(MP_MAX == 34, 
 		"New material properties should be added to the end of the inputs for this expression. \
 		The order of properties here should match the material results pins, the make material attriubtes node inputs and the mapping of IO indices to properties in GetMaterialPropertyFromInputOutputIndex().\
 		Insertions into the middle of the properties or a change in the order of properties will also require that existing data is fixed up in DoMaterialAttributeReorder().\
@@ -6777,7 +6777,7 @@ UMaterialExpressionBreakMaterialAttributes::UMaterialExpressionBreakMaterialAttr
 
 	MenuCategories.Add(ConstructorStatics.NAME_MaterialAttributes);
 	
- 	static_assert(MP_MAX == 33, 
+ 	static_assert(MP_MAX == 34, 
 		"New material properties should be added to the end of the outputs for this expression. \
 		The order of properties here should match the material results pins, the make material attriubtes node inputs and the mapping of IO indices to properties in GetMaterialPropertyFromInputOutputIndex().\
 		Insertions into the middle of the properties or a change in the order of properties will also require that existing data is fixed up in DoMaterialAttriubtesReorder().\
@@ -22386,6 +22386,9 @@ void UMaterialExpressionStrataLegacyConversion::GatherStrataMaterialInfo(FStrata
 
 FStrataOperator* UMaterialExpressionStrataLegacyConversion::StrataGenerateMaterialTopologyTree(class FMaterialCompiler* Compiler, class UMaterialExpression* Parent, int32 OutputIndex)
 {
+	// Note Thickness has no meaning/usage in the context of StrataLegacyConversionNode
+	int32 ThicknessIndex = Compiler->StrataThicknessStackGetThicknessIndex();
+
 	auto AddDefaultWorstCase = [&](bool bSSS)
 	{
 		FStrataOperator& SlabOperator = Compiler->StrataCompilationRegisterOperator(STRATA_OPERATOR_BSDF_LEGACY, Compiler->StrataTreeStackGetPathUniqueId(), Parent, Compiler->StrataTreeStackGetParentPathUniqueId());
@@ -22394,6 +22397,7 @@ FStrataOperator* UMaterialExpressionStrataLegacyConversion::StrataGenerateMateri
 		SlabOperator.bBSDFHasMFPPluggedIn = true;
 		SlabOperator.bBSDFHasFuzz = true;
 		SlabOperator.bBSDFHasAnisotropy = Anisotropy.IsConnected();
+		SlabOperator.ThicknessIndex = ThicknessIndex;
 
 		return &SlabOperator;
 	};
@@ -22411,6 +22415,7 @@ FStrataOperator* UMaterialExpressionStrataLegacyConversion::StrataGenerateMateri
 		{
 			FStrataOperator& Operator = Compiler->StrataCompilationRegisterOperator(STRATA_OPERATOR_BSDF_LEGACY, Compiler->StrataTreeStackGetPathUniqueId(), Parent, Compiler->StrataTreeStackGetParentPathUniqueId());
 			Operator.BSDFType = STRATA_BSDF_TYPE_UNLIT;
+			Operator.ThicknessIndex = ThicknessIndex;
 			return &Operator;
 		}
 		else if (ConvertedStrataMaterialInfo.HasShadingModel(SSM_DefaultLit))
@@ -22425,18 +22430,21 @@ FStrataOperator* UMaterialExpressionStrataLegacyConversion::StrataGenerateMateri
 		{
 			FStrataOperator& Operator = Compiler->StrataCompilationRegisterOperator(STRATA_OPERATOR_BSDF_LEGACY, Compiler->StrataTreeStackGetPathUniqueId(), Parent, Compiler->StrataTreeStackGetParentPathUniqueId());
 			Operator.BSDFType = STRATA_BSDF_TYPE_HAIR;
+			Operator.ThicknessIndex = ThicknessIndex;
 			return &Operator;
 		}
 		else if (ConvertedStrataMaterialInfo.HasShadingModel(SSM_Eye))
 		{
 			FStrataOperator& Operator = Compiler->StrataCompilationRegisterOperator(STRATA_OPERATOR_BSDF_LEGACY, Compiler->StrataTreeStackGetPathUniqueId(), Parent, Compiler->StrataTreeStackGetParentPathUniqueId());
 			Operator.BSDFType = STRATA_BSDF_TYPE_EYE;
+			Operator.ThicknessIndex = ThicknessIndex;
 			return &Operator;
 		}
 		else if (ConvertedStrataMaterialInfo.HasShadingModel(SSM_SingleLayerWater))
 		{
 			FStrataOperator& Operator = Compiler->StrataCompilationRegisterOperator(STRATA_OPERATOR_BSDF_LEGACY, Compiler->StrataTreeStackGetPathUniqueId(), Parent, Compiler->StrataTreeStackGetParentPathUniqueId());
 			Operator.BSDFType = STRATA_BSDF_TYPE_SINGLELAYERWATER;
+			Operator.ThicknessIndex = ThicknessIndex;
 			return &Operator;
 		}
 
@@ -22521,6 +22529,10 @@ int32 UMaterialExpressionStrataSlabBSDF::Compile(class FMaterialCompiler* Compil
 
 	FStrataOperator& StrataOperator = Compiler->StrataCompilationGetOperator(Compiler->StrataTreeStackGetPathUniqueId());
 	StrataOperator.BSDFRegisteredSharedLocalBasis = NewRegisteredSharedLocalBasis;
+	
+	check(StrataOperator.ThicknessIndex != INDEX_NONE);
+	int32 ThicknesCodeChunk = Compiler->StrataThicknessStackGetThicknessCode(StrataOperator.ThicknessIndex);
+	check(ThicknesCodeChunk != INDEX_NONE);
 
 	int32 OutputCodeChunk = Compiler->StrataSlabBSDF(
 		CompileWithDefaultFloat3(Compiler, DiffuseAlbedo, 0.18f, 0.18f, 0.18f),
@@ -22539,7 +22551,7 @@ int32 UMaterialExpressionStrataSlabBSDF::Compile(class FMaterialCompiler* Compil
 		Compiler->Constant(0.0f),										// SecondRoughnessAsSimpleClearCoat
 		CompileWithDefaultFloat1(Compiler, FuzzAmount, 0.0f),
 		CompileWithDefaultFloat3(Compiler, FuzzColor, 0.0f, 0.0f, 0.0f),
-		CompileWithDefaultFloat1(Compiler, Thickness, STRATA_LAYER_DEFAULT_THICKNESS_CM),
+		ThicknesCodeChunk,
 		NormalCodeChunk,
 		TangentCodeChunk,
 		Compiler->GetStrataSharedLocalBasisIndexMacro(NewRegisteredSharedLocalBasis),
@@ -22708,7 +22720,7 @@ FName UMaterialExpressionStrataSlabBSDF::GetInputName(int32 InputIndex) const
 	}
 	else if (InputIndex == 13)
 	{
-		return TEXT("Thickness");
+		return TEXT("Thickness (DEPRECATED)");
 	}
 	else if (InputIndex == 14)
 	{
@@ -22775,6 +22787,7 @@ FStrataOperator* UMaterialExpressionStrataSlabBSDF::StrataGenerateMaterialTopolo
 	StrataOperator.bBSDFHasSSS = HasSSS();
 	StrataOperator.bBSDFHasMFPPluggedIn = HasMFPPluggedIn();
 	StrataOperator.bBSDFHasAnisotropy = HasAnisotropy();
+	StrataOperator.ThicknessIndex = Compiler->StrataThicknessStackGetThicknessIndex();
 	return &StrataOperator;
 }
 
@@ -22843,6 +22856,9 @@ int32 UMaterialExpressionStrataSimpleClearCoatBSDF::Compile(class FMaterialCompi
 	FStrataOperator& StrataOperator = Compiler->StrataCompilationGetOperator(Compiler->StrataTreeStackGetPathUniqueId());
 	StrataOperator.BSDFRegisteredSharedLocalBasis = NewRegisteredSharedLocalBasis;
 
+	int32 ThicknessCodeChunk = Compiler->StrataThicknessStackGetThicknessCode(StrataOperator.ThicknessIndex);
+	check(ThicknessCodeChunk != INDEX_NONE);
+
 	int32 OutputCodeChunk = Compiler->StrataSlabBSDF(
 		CompileWithDefaultFloat3(Compiler, DiffuseAlbedo, 0.18f, 0.18f, 0.18f),		// DiffuseAlbedo
 		CompileWithDefaultFloat3(Compiler, F0, DefaultF0, DefaultF0, DefaultF0),	// F0
@@ -22861,7 +22877,7 @@ int32 UMaterialExpressionStrataSimpleClearCoatBSDF::Compile(class FMaterialCompi
 		Compiler->Constant(1.0f),								// SecondRoughnessAsSimpleClearCoat == true for UMaterialExpressionStrataSimpleClearCoatBSDF
 		Compiler->Constant(0.0f),								// FuzzAmount
 		Compiler->Constant3(0.0f, 0.0f, 0.0f),					// FuzzColor
-		Compiler->Constant(STRATA_LAYER_DEFAULT_THICKNESS_CM),	// Thickness
+		ThicknessCodeChunk,										// Thickness
 		NormalCodeChunk,
 		NullTangentCodeChunk,
 		Compiler->GetStrataSharedLocalBasisIndexMacro(NewRegisteredSharedLocalBasis),
@@ -22991,6 +23007,7 @@ FStrataOperator* UMaterialExpressionStrataSimpleClearCoatBSDF::StrataGenerateMat
 	StrataOperator.bBSDFHasSSS = false;
 	StrataOperator.bBSDFHasMFPPluggedIn = false;
 	StrataOperator.bBSDFHasAnisotropy = false;
+	StrataOperator.ThicknessIndex = Compiler->StrataThicknessStackGetThicknessIndex();
 	return &StrataOperator;
 }
 #endif // WITH_EDITOR
@@ -23069,6 +23086,7 @@ FStrataOperator* UMaterialExpressionStrataVolumetricFogCloudBSDF::StrataGenerate
 {
 	FStrataOperator& StrataOperator = Compiler->StrataCompilationRegisterOperator(STRATA_OPERATOR_BSDF, Compiler->StrataTreeStackGetPathUniqueId(), Parent, Compiler->StrataTreeStackGetParentPathUniqueId());
 	StrataOperator.BSDFType = STRATA_BSDF_TYPE_VOLUMETRICFOGCLOUD;
+	StrataOperator.ThicknessIndex = Compiler->StrataThicknessStackGetThicknessIndex();
 	return &StrataOperator;
 }
 #endif // WITH_EDITOR
@@ -23136,6 +23154,7 @@ FStrataOperator* UMaterialExpressionStrataLightFunction::StrataGenerateMaterialT
 {
 	FStrataOperator& StrataOperator = Compiler->StrataCompilationRegisterOperator(STRATA_OPERATOR_BSDF, Compiler->StrataTreeStackGetPathUniqueId(), Parent, Compiler->StrataTreeStackGetParentPathUniqueId());
 	StrataOperator.BSDFType = STRATA_BSDF_TYPE_UNLIT;
+	StrataOperator.ThicknessIndex = Compiler->StrataThicknessStackGetThicknessIndex();
 	return &StrataOperator;
 }
 #endif // WITH_EDITOR
@@ -23209,6 +23228,7 @@ FStrataOperator* UMaterialExpressionStrataPostProcess::StrataGenerateMaterialTop
 {
 	FStrataOperator& StrataOperator = Compiler->StrataCompilationRegisterOperator(STRATA_OPERATOR_BSDF, Compiler->StrataTreeStackGetPathUniqueId(), Parent, Compiler->StrataTreeStackGetParentPathUniqueId());
 	StrataOperator.BSDFType = STRATA_BSDF_TYPE_UNLIT;
+	StrataOperator.ThicknessIndex = Compiler->StrataThicknessStackGetThicknessIndex();
 	return &StrataOperator;
 }
 #endif // WITH_EDITOR
@@ -23391,6 +23411,7 @@ FStrataOperator* UMaterialExpressionStrataUnlitBSDF::StrataGenerateMaterialTopol
 {
 	FStrataOperator& StrataOperator = Compiler->StrataCompilationRegisterOperator(STRATA_OPERATOR_BSDF, Compiler->StrataTreeStackGetPathUniqueId(), Parent, Compiler->StrataTreeStackGetParentPathUniqueId());
 	StrataOperator.BSDFType = STRATA_BSDF_TYPE_UNLIT;
+	StrataOperator.ThicknessIndex = Compiler->StrataThicknessStackGetThicknessIndex();
 	return &StrataOperator;
 }
 #endif // WITH_EDITOR
@@ -23506,6 +23527,7 @@ FStrataOperator* UMaterialExpressionStrataHairBSDF::StrataGenerateMaterialTopolo
 {
 	FStrataOperator& StrataOperator = Compiler->StrataCompilationRegisterOperator(STRATA_OPERATOR_BSDF, Compiler->StrataTreeStackGetPathUniqueId(), Parent, Compiler->StrataTreeStackGetParentPathUniqueId());
 	StrataOperator.BSDFType = STRATA_BSDF_TYPE_HAIR;
+	StrataOperator.ThicknessIndex = Compiler->StrataThicknessStackGetThicknessIndex();
 	return &StrataOperator;
 }
 #endif // WITH_EDITOR
@@ -23619,6 +23641,7 @@ FStrataOperator* UMaterialExpressionStrataEyeBSDF::StrataGenerateMaterialTopolog
 {
 	FStrataOperator& StrataOperator = Compiler->StrataCompilationRegisterOperator(STRATA_OPERATOR_BSDF, Compiler->StrataTreeStackGetPathUniqueId(), Parent, Compiler->StrataTreeStackGetParentPathUniqueId());
 	StrataOperator.BSDFType = STRATA_BSDF_TYPE_EYE;
+	StrataOperator.ThicknessIndex = Compiler->StrataThicknessStackGetThicknessIndex();
 	return &StrataOperator;
 }
 #endif // WITH_EDITOR
@@ -23749,6 +23772,7 @@ FStrataOperator* UMaterialExpressionStrataSingleLayerWaterBSDF::StrataGenerateMa
 {
 	FStrataOperator& StrataOperator = Compiler->StrataCompilationRegisterOperator(STRATA_OPERATOR_BSDF, Compiler->StrataTreeStackGetPathUniqueId(), Parent, Compiler->StrataTreeStackGetParentPathUniqueId());
 	StrataOperator.BSDFType = STRATA_BSDF_TYPE_SINGLELAYERWATER;
+	StrataOperator.ThicknessIndex = Compiler->StrataThicknessStackGetThicknessIndex();
 	return &StrataOperator;
 }
 #endif // WITH_EDITOR
@@ -23952,6 +23976,9 @@ int32 UMaterialExpressionStrataVerticalLayering::Compile(class FMaterialCompiler
 	Compiler->StrataTreeStackPush(this, 1);
 	int32 BaseCodeChunk = Base.Compile(Compiler);
 	Compiler->StrataTreeStackPop();
+	Compiler->StrataTreeStackPush(this, 2);
+	int32 ThicknessCodeChunk = Thickness.Compile(Compiler);
+	Compiler->StrataTreeStackPop();
 
 	int32 OutputCodeChunk = INDEX_NONE;
 	FStrataOperator& StrataOperator = Compiler->StrataCompilationGetOperator(Compiler->StrataTreeStackGetPathUniqueId());
@@ -23993,7 +24020,7 @@ int32 UMaterialExpressionStrataVerticalLayering::Compile(class FMaterialCompiler
 		const FStrataRegisteredSharedLocalBasis NewRegisteredSharedLocalBasis = StrataCompilationInfoCreateSharedLocalBasis(Compiler, NewNormalCodeChunk, NewTangentCodeChunk);
 
 		OutputCodeChunk = Compiler->StrataVerticalLayeringParameterBlending(
-			TopCodeChunk, BaseCodeChunk, Compiler->GetStrataSharedLocalBasisIndexMacro(NewRegisteredSharedLocalBasis), TopBSDFOperator->BSDFRegisteredSharedLocalBasis.NormalCodeChunk,
+			TopCodeChunk, BaseCodeChunk, ThicknessCodeChunk, Compiler->GetStrataSharedLocalBasisIndexMacro(NewRegisteredSharedLocalBasis), TopBSDFOperator->BSDFRegisteredSharedLocalBasis.NormalCodeChunk,
 			StrataOperator.bRootOfParameterBlendingSubTree ? &StrataOperator : nullptr);
 
 		// Propagate the parameter blended normal
@@ -24001,7 +24028,7 @@ int32 UMaterialExpressionStrataVerticalLayering::Compile(class FMaterialCompiler
 	}
 	else
 	{
-		OutputCodeChunk = Compiler->StrataVerticalLayering(TopCodeChunk, BaseCodeChunk, StrataOperator.Index, StrataOperator.MaxDistanceFromLeaves);
+		OutputCodeChunk = Compiler->StrataVerticalLayering(TopCodeChunk, BaseCodeChunk, ThicknessCodeChunk, StrataOperator.Index, StrataOperator.MaxDistanceFromLeaves);
 	}
 
 
@@ -24020,6 +24047,24 @@ void UMaterialExpressionStrataVerticalLayering::GetCaption(TArray<FString>& OutC
 	}
 }
 
+FName UMaterialExpressionStrataVerticalLayering::GetInputName(int32 InputIndex) const
+{
+	if (InputIndex == 0)
+	{
+		return TEXT("Top");
+	}
+	else if (InputIndex == 1)
+	{
+		return TEXT("Bottom");
+	}
+	else if (InputIndex == 2)
+	{
+		return  TEXT("Top Thickness");
+	}
+
+	return TEXT("Unknown");
+}
+
 uint32 UMaterialExpressionStrataVerticalLayering::GetOutputType(int32 OutputIndex)
 {
 	return MCT_Strata;
@@ -24027,6 +24072,10 @@ uint32 UMaterialExpressionStrataVerticalLayering::GetOutputType(int32 OutputInde
 
 uint32 UMaterialExpressionStrataVerticalLayering::GetInputType(int32 InputIndex)
 {
+	if (InputIndex == 2)
+	{
+		return MCT_Float;
+	}
 	return MCT_Strata;
 }
 
@@ -24055,13 +24104,18 @@ FStrataOperator* UMaterialExpressionStrataVerticalLayering::StrataGenerateMateri
 	UMaterialExpression* ChildBExpression = Base.GetTracedInput().Expression;
 	FStrataOperator* OpA = nullptr;
 	FStrataOperator* OpB = nullptr;
+
+	// Top - Use the vertical operator thickness
 	if (ChildAExpression)
 	{
+		Compiler->StrataThicknessStackPush(this, &Thickness);
 		Compiler->StrataTreeStackPush(this, 0);
 		OpA = ChildAExpression->StrataGenerateMaterialTopologyTree(Compiler, this, Top.OutputIndex);
 		Compiler->StrataTreeStackPop();
+		Compiler->StrataThicknessStackPop();
 		AssignOperatorIndexIfNotNull(StrataOperator.LeftIndex, OpA);
 	}
+	// Bottom - Use the propagated thickness from parent
 	if (ChildBExpression)
 	{
 		Compiler->StrataTreeStackPush(this, 1);
@@ -24196,7 +24250,7 @@ void UMaterialExpressionStrataAdd::GatherStrataMaterialInfo(FStrataMaterialInfo&
 {
 	if (A.GetTracedInput().Expression)
 	{
-		A.GetTracedInput().Expression->GatherStrataMaterialInfo(StrataMaterialInfo, A.OutputIndex);
+		A.GetTracedInput().Expression->GatherStrataMaterialInfo(StrataMaterialInfo, A.OutputIndex);	
 	}
 	if (B.GetTracedInput().Expression)
 	{
diff --git a/Engine/Source/Runtime/Engine/Private/Materials/MaterialHLSLEmitter.cpp b/Engine/Source/Runtime/Engine/Private/Materials/MaterialHLSLEmitter.cpp
--- a/Engine/Source/Runtime/Engine/Private/Materials/MaterialHLSLEmitter.cpp
+++ b/Engine/Source/Runtime/Engine/Private/Materials/MaterialHLSLEmitter.cpp
@@ -790,6 +790,7 @@ bool MaterialEmitHLSL(const FMaterialCompileTargetParameters& InCompilerTarget,
 	SharedPixelProperties[MP_PixelDepthOffset] = true;
 	SharedPixelProperties[MP_SubsurfaceColor] = true;
 	SharedPixelProperties[MP_ShadingModel] = true;
+	SharedPixelProperties[MP_SurfaceThickness] = true;
 	SharedPixelProperties[MP_FrontMaterial] = true;
 
 	bool bUsesWorldPositionOffset = false;
diff --git a/Engine/Source/Runtime/Engine/Private/Materials/MaterialShared.cpp b/Engine/Source/Runtime/Engine/Private/Materials/MaterialShared.cpp
--- a/Engine/Source/Runtime/Engine/Private/Materials/MaterialShared.cpp
+++ b/Engine/Source/Runtime/Engine/Private/Materials/MaterialShared.cpp
@@ -63,6 +63,7 @@
 #endif
 #include "ProfilingDebugging/CountersTrace.h"
 #include "RenderCore.h"
+#include "StrataDefinitions.h"
 
 #define LOCTEXT_NAMESPACE "MaterialShared"
 
@@ -5503,6 +5504,7 @@ void FMaterialAttributeDefinitionMap::InitializeAttributeMap()
 	Add(FGuid(0xD0B0FA03, 0x14D74455, 0xA851BAC5, 0x81A0788B), TEXT("Refraction"),				MP_Refraction,				MCT_Float3,	FVector4(1,0,0,0),	SF_Pixel);
 	Add(FGuid(0x0AC97EC3, 0xE3D047BA, 0xB610167D, 0xC4D919FF), TEXT("PixelDepthOffset"),		MP_PixelDepthOffset,		MCT_Float,	FVector4(0,0,0,0),	SF_Pixel);
 	Add(FGuid(0xD9423FFF, 0xD77E4D82, 0x8FF9CF5E, 0x055D1255), TEXT("ShadingModel"),			MP_ShadingModel,			MCT_ShadingModel, FVector4(0, 0, 0, 0), SF_Pixel, INDEX_NONE, false, &CompileShadingModelBlendFunction);
+	Add(FGuid(0x42BDD2E0, 0xBE714189, 0xA0984BC3, 0xDD0BE872), TEXT("SurfaceThickness"),		MP_SurfaceThickness,		MCT_Float,  FVector4(STRATA_LAYER_DEFAULT_THICKNESS_CM, 0, 0, 0), SF_Pixel);
 	Add(FGuid(0x5973A03E, 0x13A74E08, 0x92D0CEDD, 0xF2936CF8), TEXT("FrontMaterial"),			MP_FrontMaterial,			MCT_Strata, FVector4(0,0,0,0),	SF_Pixel, INDEX_NONE, false, &CompileStrataBlendFunction);
 
 	// Used when compiling material with execution pins, which are compiling all attributes together
@@ -5663,6 +5665,8 @@ FText FMaterialAttributeDefinitionMap::GetAttributeOverrideForMaterial(const FGu
 		return LOCTEXT("PixelDepthOffset", "Pixel Depth Offset");
 	case MP_ShadingModel:
 		return LOCTEXT("ShadingModel", "Shading Model");
+	case MP_SurfaceThickness:
+		return LOCTEXT("SurfaceThickness", "Surface Thickness");
 	case MP_FrontMaterial:
 		return LOCTEXT("FrontMaterial", "Front Material");
 	case MP_CustomOutput:
diff --git a/Engine/Source/Runtime/Engine/Public/MaterialCompiler.h b/Engine/Source/Runtime/Engine/Public/MaterialCompiler.h
--- a/Engine/Source/Runtime/Engine/Public/MaterialCompiler.h
+++ b/Engine/Source/Runtime/Engine/Public/MaterialCompiler.h
@@ -100,13 +100,14 @@ struct FStrataOperator
 	int32 OperatorType;
 	bool bNodeRequestParameterBlending;
 
-	int32 Index;		// Index into the array of operators
-	int32 ParentIndex;	// Parent operator index
-	int32 LeftIndex;	// Left child operator index
-	int32 RightIndex;	// Right child operator index
+	int32 Index;			// Index into the array of operators
+	int32 ParentIndex;		// Parent operator index
+	int32 LeftIndex;		// Left child operator index
+	int32 RightIndex;		// Right child operator index
+	int32 ThicknessIndex;	// Thickness expression index
 
 	// Data used for BSDF type nodes only
-	int32 BSDFIndex;	// Index in the array of BSDF if a BSDF operator
+	int32 BSDFIndex;		// Index in the array of BSDF if a BSDF operator
 	uint8 BSDFType;
 	FStrataRegisteredSharedLocalBasis BSDFRegisteredSharedLocalBasis;
 	bool  bBSDFHasSSS;
@@ -132,6 +133,7 @@ struct FStrataOperator
 		ParentIndex = INDEX_NONE;
 		LeftIndex = INDEX_NONE;
 		RightIndex = INDEX_NONE;
+		ThicknessIndex = INDEX_NONE;
 
 		BSDFIndex = INDEX_NONE;
 		BSDFType = 0;
@@ -666,8 +668,8 @@ public:
 		int32 WaterAlbedo, int32 WaterExtinction, int32 WaterPhaseG, int32 ColorScaleBehindWater, int32 Normal, const FString& SharedLocalBasisIndexMacro, FStrataOperator* PromoteToOperator) = 0;
 	virtual int32 StrataHorizontalMixing(int32 Background, int32 Foreground, int32 Mix, int OperatorIndex, uint32 MaxDistanceFromLeaves) = 0;
 	virtual int32 StrataHorizontalMixingParameterBlending(int32 Background, int32 Foreground, int32 HorizontalMixCodeChunk, int32 NormalMixCodeChunk, const FString& SharedLocalBasisIndexMacro, FStrataOperator* PromoteToOperator) = 0;
-	virtual int32 StrataVerticalLayering(int32 Top, int32 Base, int OperatorIndex, uint32 MaxDistanceFromLeaves) = 0;
-	virtual int32 StrataVerticalLayeringParameterBlending(int32 Top, int32 Base, const FString& SharedLocalBasisIndexMacro, int32 TopBSDFNormalCodeChunk, FStrataOperator* PromoteToOperator) = 0;
+	virtual int32 StrataVerticalLayering(int32 Top, int32 Base, int32 Thickness, int OperatorIndex, uint32 MaxDistanceFromLeaves) = 0;
+	virtual int32 StrataVerticalLayeringParameterBlending(int32 Top, int32 Base, int32 Thickness, const FString& SharedLocalBasisIndexMacro, int32 TopBSDFNormalCodeChunk, FStrataOperator* PromoteToOperator) = 0;
 	virtual int32 StrataAdd(int32 A, int32 B, int OperatorIndex, uint32 MaxDistanceFromLeaves) = 0;
 	virtual int32 StrataAddParameterBlending(int32 A, int32 B, int32 AMixWeight, const FString& SharedLocalBasisIndexMacro, FStrataOperator* PromoteToOperator) = 0;
 	virtual int32 StrataWeight(int32 A, int32 Weight, int OperatorIndex, uint32 MaxDistanceFromLeaves) = 0;
@@ -701,6 +703,11 @@ public:
 	 */
 	virtual void StrataTreeStackPop() = 0;
 
+	virtual int32 StrataThicknessStackGetThicknessIndex() = 0;
+	virtual int32 StrataThicknessStackGetThicknessCode(int32 Index) = 0;
+	virtual int32 StrataThicknessStackPush(UMaterialExpression* Expression, FExpressionInput* Input) = 0;
+	virtual void StrataThicknessStackPop() = 0;
+
 	/**
 	 * Register an operator of the tree representation the Strata material and its topology.
 	 */
@@ -1318,14 +1325,14 @@ public:
 		return Compiler->StrataHorizontalMixingParameterBlending(Background, Foreground, HorizontalMixCodeChunk, NormalMixCodeChunk, SharedLocalBasisIndexMacro, PromoteToOperator);
 	}
 
-	virtual int32 StrataVerticalLayering(int32 Top, int32 Base, int OperatorIndex, uint32 MaxDistanceFromLeaves) override
+	virtual int32 StrataVerticalLayering(int32 Top, int32 Base, int32 Thickness, int OperatorIndex, uint32 MaxDistanceFromLeaves) override
 	{
-		return Compiler->StrataVerticalLayering(Top, Base, OperatorIndex, MaxDistanceFromLeaves);
+		return Compiler->StrataVerticalLayering(Top, Base, Thickness, OperatorIndex, MaxDistanceFromLeaves);
 	}
 
-	virtual int32 StrataVerticalLayeringParameterBlending(int32 Top, int32 Base, const FString& SharedLocalBasisIndexMacro, int32 TopBSDFNormalCodeChunk, FStrataOperator* PromoteToOperator) override
+	virtual int32 StrataVerticalLayeringParameterBlending(int32 Top, int32 Base, int32 Thickness, const FString& SharedLocalBasisIndexMacro, int32 TopBSDFNormalCodeChunk, FStrataOperator* PromoteToOperator) override
 	{
-		return Compiler->StrataVerticalLayeringParameterBlending(Top, Base, SharedLocalBasisIndexMacro, TopBSDFNormalCodeChunk, PromoteToOperator);
+		return Compiler->StrataVerticalLayeringParameterBlending(Top, Base, Thickness, SharedLocalBasisIndexMacro, TopBSDFNormalCodeChunk, PromoteToOperator);
 	}
 
 	virtual int32 StrataAdd(int32 A, int32 B, int OperatorIndex, uint32 MaxDistanceFromLeaves) override
@@ -1394,6 +1401,23 @@ public:
 		Compiler->StrataTreeStackPop();
 	}
 
+	virtual int32 StrataThicknessStackGetThicknessIndex() override
+	{
+		return Compiler->StrataThicknessStackGetThicknessIndex();
+	}
+	virtual int32 StrataThicknessStackGetThicknessCode(int32 Index) override
+	{
+		return Compiler->StrataThicknessStackGetThicknessCode(Index);
+	}
+	virtual int32 StrataThicknessStackPush(UMaterialExpression* Expression, FExpressionInput* Input) override
+	{
+		return Compiler->StrataThicknessStackPush(Expression, Input);
+	}
+	virtual void StrataThicknessStackPop() override
+	{
+		Compiler->StrataThicknessStackPop();
+	}
+
 	virtual FStrataOperator& StrataCompilationRegisterOperator(int32 OperatorType, FGuid StrataExpressionGuid, UMaterialExpression* Parent, FGuid StrataParentExpressionGuid, bool bUseParameterBlending = false) override
 	{
 		return Compiler->StrataCompilationRegisterOperator(OperatorType, StrataExpressionGuid, Parent, StrataParentExpressionGuid, bUseParameterBlending);
diff --git a/Engine/Source/Runtime/Engine/Public/MaterialExpressionIO.h b/Engine/Source/Runtime/Engine/Public/MaterialExpressionIO.h
--- a/Engine/Source/Runtime/Engine/Public/MaterialExpressionIO.h
+++ b/Engine/Source/Runtime/Engine/Public/MaterialExpressionIO.h
@@ -302,7 +302,7 @@ struct FMaterialAttributesInput : FExpressionInput
 	: PropertyConnectedBitmask(0)
 	{ 
 		// ensure PropertyConnectedBitmask can contain all properties.
-		static_assert((uint32)(MP_MAX)-1 <= (8 * sizeof(PropertyConnectedBitmask)), "PropertyConnectedBitmask cannot contain entire EMaterialProperty enumeration.");
+		static_assert((uint32)(MP_MaterialAttributes)-1 <= (8 * sizeof(PropertyConnectedBitmask)), "PropertyConnectedBitmask cannot contain entire EMaterialProperty enumeration.");
 	}
 
 #if WITH_EDITOR
diff --git a/Engine/Source/Runtime/Engine/Public/SceneTypes.h b/Engine/Source/Runtime/Engine/Public/SceneTypes.h
--- a/Engine/Source/Runtime/Engine/Public/SceneTypes.h
+++ b/Engine/Source/Runtime/Engine/Public/SceneTypes.h
@@ -233,6 +233,7 @@ enum EMaterialProperty
 	MP_PixelDepthOffset UMETA(Hidden),
 	MP_ShadingModel UMETA(Hidden),
 	MP_FrontMaterial UMETA(Hidden),
+	MP_SurfaceThickness UMETA(Hidden),
 
 	//^^^ New material properties go above here ^^^^
 	MP_MaterialAttributes UMETA(Hidden),
