Break the InputDevicePropertyHandle into its own .h/.cpp file to help improve compile times. Now users of it dont need to include the whole subsystem, just the handle.

#rnx
#pf skip

diff --git a/Engine/Source/Runtime/Engine/Classes/GameFramework/InputDevicePropertyHandle.h b/Engine/Source/Runtime/Engine/Classes/GameFramework/InputDevicePropertyHandle.h
new file mode 100644
--- /dev/null
+++ b/Engine/Source/Runtime/Engine/Classes/GameFramework/InputDevicePropertyHandle.h
@@ -0,0 +1,42 @@
+// Copyright Epic Games, Inc. All Rights Reserved.
+
+#pragma once
+
+#include "InputDevicePropertyHandle.generated.h"
+
+/** A handle to an active input device property that is being used by the InputDeviceSubsytem. */
+USTRUCT(BlueprintType)
+struct ENGINE_API FInputDevicePropertyHandle
+{
+	friend class UInputDeviceSubsystem;
+
+	GENERATED_BODY()
+
+	FInputDevicePropertyHandle();
+	~FInputDevicePropertyHandle() = default;
+
+	/** Returns true if this handle is valid */
+	bool IsValid() const;
+
+	/** An invalid Input Device Property handle */
+	static FInputDevicePropertyHandle InvalidHandle;
+
+	ENGINE_API friend uint32 GetTypeHash(const FInputDevicePropertyHandle& InHandle);
+
+	bool operator==(const FInputDevicePropertyHandle& Other) const;
+	bool operator!=(const FInputDevicePropertyHandle& Other) const;
+
+	FString ToString() const;
+
+private:
+
+	// Private constructor because we don't want any other users to make a valid device property handle.
+	FInputDevicePropertyHandle(uint32 InternalID);
+
+	/** Static function to get a unique device handle. */
+	static FInputDevicePropertyHandle AcquireValidHandle();
+
+	/** The internal ID of this handle. Populated by the private constructor in AcquireValidHandle */
+	UPROPERTY()
+	uint32 InternalId;
+};
diff --git a/Engine/Source/Runtime/Engine/Classes/GameFramework/InputDeviceSubsystem.h b/Engine/Source/Runtime/Engine/Classes/GameFramework/InputDeviceSubsystem.h
--- a/Engine/Source/Runtime/Engine/Classes/GameFramework/InputDeviceSubsystem.h
+++ b/Engine/Source/Runtime/Engine/Classes/GameFramework/InputDeviceSubsystem.h
@@ -3,6 +3,7 @@
 #pragma once
 
 #include "InputSettings.h"
+#include "GameFramework/InputDevicePropertyHandle.h"
 #include "Subsystems/EngineSubsystem.h"
 #include "Tickable.h"
 #include "Templates/SubclassOf.h"
@@ -13,43 +14,6 @@ ENGINE_API DECLARE_LOG_CATEGORY_EXTERN(LogInputDeviceProperties, Log, All);
 class UInputDeviceProperty;
 class APlayerController;
 
-/** A handle to an active input device property that is being used by the InputDeviceSubsytem. */
-USTRUCT(BlueprintType)
-struct ENGINE_API FInputDevicePropertyHandle
-{
-	friend class UInputDeviceSubsystem;
-
-	GENERATED_BODY()
-
-	FInputDevicePropertyHandle();
-	~FInputDevicePropertyHandle() = default;
-
-	/** Returns true if this handle is valid */
-	bool IsValid() const;
-
-	/** An invalid Input Device Property handle */
-	static FInputDevicePropertyHandle InvalidHandle;
-
-	ENGINE_API friend uint32 GetTypeHash(const FInputDevicePropertyHandle& InHandle);
-
-	bool operator==(const FInputDevicePropertyHandle& Other) const;
-	bool operator!=(const FInputDevicePropertyHandle& Other) const;
-
-	FString ToString() const;
-
-private:
-
-	// Private constructor because we don't want any other users to make a valid device property handle.
-	FInputDevicePropertyHandle(uint32 InternalID);
-
-	/** Static function to get a unique device handle. */
-	static FInputDevicePropertyHandle AcquireValidHandle();
-
-	/** The internal ID of this handle. Populated by the private constructor in AcquireValidHandle */
-	UPROPERTY()
-	uint32 InternalId;
-};
-
 /** Contains a pointer to an active device property and keeps track of how long it has been evaluated for */
 USTRUCT()
 struct ENGINE_API FActiveDeviceProperty
diff --git a/Engine/Source/Runtime/Engine/Private/GameFramework/InputDevicePropertyHandle.cpp b/Engine/Source/Runtime/Engine/Private/GameFramework/InputDevicePropertyHandle.cpp
new file mode 100644
--- /dev/null
+++ b/Engine/Source/Runtime/Engine/Private/GameFramework/InputDevicePropertyHandle.cpp
@@ -0,0 +1,53 @@
+// Copyright Epic Games, Inc. All Rights Reserved.
+
+#include "GameFramework/InputDevicePropertyHandle.h"
+
+////////////////////////////////////////////////////////
+// FInputDevicePropertyHandle
+
+FInputDevicePropertyHandle FInputDevicePropertyHandle::InvalidHandle(0);
+
+FInputDevicePropertyHandle::FInputDevicePropertyHandle()
+	: InternalId(0)
+{
+
+}
+
+FInputDevicePropertyHandle::FInputDevicePropertyHandle(uint32 InInternalID)
+	: InternalId(InInternalID)
+{
+
+}
+
+bool FInputDevicePropertyHandle::operator==(const FInputDevicePropertyHandle& Other) const
+{
+	return InternalId == Other.InternalId;
+}
+
+bool FInputDevicePropertyHandle::operator!=(const FInputDevicePropertyHandle& Other) const
+{
+	return InternalId != Other.InternalId;
+}
+
+uint32 GetTypeHash(const FInputDevicePropertyHandle& InHandle)
+{
+	return GetTypeHash(InHandle.InternalId);
+}
+
+FString FInputDevicePropertyHandle::ToString() const
+{
+	return IsValid() ? FString::FromInt(InternalId) : TEXT("Invalid");
+}
+
+bool FInputDevicePropertyHandle::IsValid() const
+{
+	return InternalId != FInputDevicePropertyHandle::InvalidHandle.InternalId;
+}
+
+FInputDevicePropertyHandle FInputDevicePropertyHandle::AcquireValidHandle()
+{
+	// 0 is the "Invalid" index for these handles. Start them at 1
+	static uint32 GHandleIndex = 1;
+
+	return FInputDevicePropertyHandle(++GHandleIndex);
+}
\ No newline at end of file
diff --git a/Engine/Source/Runtime/Engine/Private/GameFramework/InputDeviceSubsystem.cpp b/Engine/Source/Runtime/Engine/Private/GameFramework/InputDeviceSubsystem.cpp
--- a/Engine/Source/Runtime/Engine/Private/GameFramework/InputDeviceSubsystem.cpp
+++ b/Engine/Source/Runtime/Engine/Private/GameFramework/InputDeviceSubsystem.cpp
@@ -93,56 +93,6 @@ public:
 	}
 };
 
-////////////////////////////////////////////////////////
-// FInputDevicePropertyHandle
-
-FInputDevicePropertyHandle FInputDevicePropertyHandle::InvalidHandle(0);
-
-FInputDevicePropertyHandle::FInputDevicePropertyHandle()
-	: InternalId(0)
-{
-
-}
-
-FInputDevicePropertyHandle::FInputDevicePropertyHandle(uint32 InInternalID)
-	: InternalId(InInternalID)
-{
-
-}
-
-bool FInputDevicePropertyHandle::operator==(const FInputDevicePropertyHandle& Other) const
-{
-	return InternalId == Other.InternalId;
-}
-
-bool FInputDevicePropertyHandle::operator!=(const FInputDevicePropertyHandle& Other) const
-{
-	return InternalId != Other.InternalId;
-}
-
-uint32 GetTypeHash(const FInputDevicePropertyHandle& InHandle)
-{
-	return GetTypeHash(InHandle.InternalId);
-}
-
-FString FInputDevicePropertyHandle::ToString() const
-{
-	return IsValid() ? FString::FromInt(InternalId) : TEXT("Invalid");
-}
-
-bool FInputDevicePropertyHandle::IsValid() const
-{
-	return InternalId != FInputDevicePropertyHandle::InvalidHandle.InternalId;
-}
-
-FInputDevicePropertyHandle FInputDevicePropertyHandle::AcquireValidHandle()
-{
-	// 0 is the "Invalid" index for these handles. Start them at 1
-	static uint32 GHandleIndex = 1;
-
-	return FInputDevicePropertyHandle(++GHandleIndex);
-}
-
 ////////////////////////////////////////////////////////
 // FSetDevicePropertyParams
 
