Fix for FFT analysis only using one channel

diff --git a/Engine/Source/Runtime/Engine/Private/SoundWave.cpp b/Engine/Source/Runtime/Engine/Private/SoundWave.cpp
--- a/Engine/Source/Runtime/Engine/Private/SoundWave.cpp
+++ b/Engine/Source/Runtime/Engine/Private/SoundWave.cpp
@@ -2282,7 +2282,7 @@ void USoundWave::BakeFFTAnalysis()
 			float SampleValue = 0.0f;
 			for (uint16 ChannelIndex = 0; ChannelIndex < RawDataNumChannels; ++ChannelIndex)
 			{
-				SampleValue += (float)InputData[FrameIndex * RawDataNumChannels] / 32767.0f;
+				SampleValue += (float)InputData[(FrameIndex * RawDataNumChannels) + ChannelIndex] / 32767.0f;
 			}
 			SampleValue /= RawDataNumChannels;
 
@@ -2438,7 +2438,7 @@ void USoundWave::BakeEnvelopeAnalysis()
 			float SampleValue = 0.0f;
 			for (uint16 ChannelIndex = 0; ChannelIndex < RawDataNumChannels; ++ChannelIndex)
 			{
-				SampleValue += (float)InputData[FrameIndex * RawDataNumChannels] / 32767.0f;
+				SampleValue += (float)InputData[(FrameIndex * RawDataNumChannels) + ChannelIndex] / 32767.0f;
 			}
 			SampleValue /= RawDataNumChannels;
 
