Add a const GetTypeHash function to the input device property handle, that way other types can just use the handle as a hash key if desired.

#rnx

diff --git a/Engine/Source/Runtime/Engine/Classes/GameFramework/InputDevicePropertyHandle.h b/Engine/Source/Runtime/Engine/Classes/GameFramework/InputDevicePropertyHandle.h
--- a/Engine/Source/Runtime/Engine/Classes/GameFramework/InputDevicePropertyHandle.h
+++ b/Engine/Source/Runtime/Engine/Classes/GameFramework/InputDevicePropertyHandle.h
@@ -21,6 +21,8 @@ struct ENGINE_API FInputDevicePropertyHandle
 	/** An invalid Input Device Property handle */
 	static FInputDevicePropertyHandle InvalidHandle;
 
+	uint32 GetTypeHash() const;
+
 	ENGINE_API friend uint32 GetTypeHash(const FInputDevicePropertyHandle& InHandle);
 
 	bool operator==(const FInputDevicePropertyHandle& Other) const;
diff --git a/Engine/Source/Runtime/Engine/Private/GameFramework/InputDevicePropertyHandle.cpp b/Engine/Source/Runtime/Engine/Private/GameFramework/InputDevicePropertyHandle.cpp
--- a/Engine/Source/Runtime/Engine/Private/GameFramework/InputDevicePropertyHandle.cpp
+++ b/Engine/Source/Runtime/Engine/Private/GameFramework/InputDevicePropertyHandle.cpp
@@ -29,9 +29,14 @@ bool FInputDevicePropertyHandle::operator!=(const FInputDevicePropertyHandle& Ot
 	return InternalId != Other.InternalId;
 }
 
+uint32 FInputDevicePropertyHandle::GetTypeHash() const
+{
+	return ::GetTypeHash(InternalId);
+}
+
 uint32 GetTypeHash(const FInputDevicePropertyHandle& InHandle)
 {
-	return GetTypeHash(InHandle.InternalId);
+	return InHandle.GetTypeHash();
 }
 
 FString FInputDevicePropertyHandle::ToString() const
