Convert the ActiveDeviceProperty array to a TSet that uses the Device Property handle as a hash key. This will be much faster for removal/adding things and lets us have a PendingRemoval queue that we can check while we are iterating on Tick, instead of needing to iterate all the active properties on the removal calls themselves.

#rnx
#pf skip

diff --git a/Engine/Source/Runtime/Engine/Classes/GameFramework/InputDeviceSubsystem.h b/Engine/Source/Runtime/Engine/Classes/GameFramework/InputDeviceSubsystem.h
--- a/Engine/Source/Runtime/Engine/Classes/GameFramework/InputDeviceSubsystem.h
+++ b/Engine/Source/Runtime/Engine/Classes/GameFramework/InputDeviceSubsystem.h
@@ -20,6 +20,14 @@ struct ENGINE_API FActiveDeviceProperty
 {
 	GENERATED_BODY()
 
+	/** Active properties can just use the hash of their FInputDevicePropertyHandle for a fast and unique lookup */
+	ENGINE_API friend uint32 GetTypeHash(const FActiveDeviceProperty& InProp);
+
+	bool operator==(const FActiveDeviceProperty& Other) const;
+	bool operator!=(const FActiveDeviceProperty& Other) const;
+	ENGINE_API friend bool operator==(const FActiveDeviceProperty& ActiveProp, const FInputDevicePropertyHandle& Handle);
+	ENGINE_API friend bool operator!=(const FActiveDeviceProperty& ActiveProp, const FInputDevicePropertyHandle& Handle);
+	
 	/** The active device property */
 	UPROPERTY(Transient)
 	TObjectPtr<UInputDeviceProperty> Property = nullptr;
@@ -132,40 +140,32 @@ public:
 	FInputDevicePropertyHandle SetDeviceProperty(const FSetDevicePropertyParams& Params);
 
 	/** Returns a pointer to the active input device property with the given handle. Returns null if the property doesn't exist */
-	TObjectPtr<UInputDeviceProperty> GetActiveDeviceProperty(const FInputDevicePropertyHandle Handle);
-
-	/** 
-	* Remove any active device properties that have the same class as the one given.
-	* 
-	* @param UserId					The PlatformUser that would have the device property applied to them
-	* @param DevicePropertyClass	The Device Property type to remove
-	* 
-	* @return						The number of removed device properties. 
-	*/
-	UFUNCTION(BlueprintCallable, Category = "Input Devices", meta = (ReturnDisplayName = "Num Removed"))
-	int32 RemoveDevicePropertiesOfClass(const FPlatformUserId UserId, TSubclassOf<UInputDeviceProperty> DevicePropertyClass);
+	UFUNCTION(BlueprintCallable, Category = "Input Devices", meta=(ReturnDisplayName="Device Property"))
+	UInputDeviceProperty* GetActiveDeviceProperty(const FInputDevicePropertyHandle Handle) const;
+	
+	/** Returns true if the property associated with the given handle is currently active, and it is not pending removal */
+	UFUNCTION(BlueprintCallable, Category = "Input Devices", meta = (ReturnDisplayName = "Is Property Active"))
+	bool IsPropertyActive(const FInputDevicePropertyHandle Handle) const;
 
 	/**
 	* Remove a single device property based on it's handle
 	*
-	* @param FInputDevicePropertyHandle		Device property handle to be removed
-	* @param bResetOnRemoval				If true, call ResetDeviceProperty before it is removed. Default is true.
+	* @param FInputDevicePropertyHandle		Device property handle to be removed	
 	*
 	* @return								The number of removed device properties.
 	*/
-	UFUNCTION(BlueprintCallable, Category = "Input Devices", meta = (ReturnDisplayName = "Num Removed"))
-	int32 RemoveDevicePropertyByHandle(const FInputDevicePropertyHandle HandleToRemove, const bool bResetOnRemoval = true);
+	UFUNCTION(BlueprintCallable, Category = "Input Devices")
+	void RemoveDevicePropertyByHandle(const FInputDevicePropertyHandle HandleToRemove);
 
 	/**
 	* Remove a set of device properties based on their handles. 
 	* 
 	* @param HandlesToRemove	The set of device property handles to remove
-	* @param bResetOnRemoval	If true, call ResetDeviceProperty before it is removed. Default is true.
 	* 
 	* @return					The number of removed device properties
 	*/
-	UFUNCTION(BlueprintCallable, Category = "Input Devices", meta = (ReturnDisplayName = "Num Removed"))
-	int32 RemoveDevicePropertyHandles(const TSet<FInputDevicePropertyHandle>& HandlesToRemove, const bool bResetOnRemoval = true);
+	UFUNCTION(BlueprintCallable, Category = "Input Devices")
+	void RemoveDevicePropertyHandles(const TSet<FInputDevicePropertyHandle>& HandlesToRemove);
 
 	/** Removes all the current Input Device Properties that are active, regardless of the Platform User */
 	UFUNCTION(BlueprintCallable, Category = "Input Devices")
@@ -206,7 +206,14 @@ protected:
 	* Array of currently active input device properties that will be evaluated on tick
 	*/
 	UPROPERTY(Transient)
-	TArray<FActiveDeviceProperty> ActiveProperties;
+	TSet<FActiveDeviceProperty> ActiveProperties;
+
+	/**
+	 * Set of property handles the properties that are currently pending manual removal.
+	 * This is populated by the "Remove device property" functions. 
+	 */
+	UPROPERTY(Transient)
+	TSet<FInputDevicePropertyHandle> PropertiesPendingRemoval;
 	
 	/** A map of an input device to it's most recent hardware device identifier */
 	TMap<FInputDeviceId, FHardwareDeviceIdentifier> LatestInputDeviceIdentifiers;
diff --git a/Engine/Source/Runtime/Engine/Private/GameFramework/InputDeviceSubsystem.cpp b/Engine/Source/Runtime/Engine/Private/GameFramework/InputDeviceSubsystem.cpp
--- a/Engine/Source/Runtime/Engine/Private/GameFramework/InputDeviceSubsystem.cpp
+++ b/Engine/Source/Runtime/Engine/Private/GameFramework/InputDeviceSubsystem.cpp
@@ -93,6 +93,36 @@ public:
 	}
 };
 
+////////////////////////////////////////////////////////
+// FActiveDeviceProperty
+
+/** Active properties can just use the hash of their Property Device handle for a fast and unique lookup */
+uint32 GetTypeHash(const FActiveDeviceProperty& InProp)
+{
+	return InProp.PropertyHandle.GetTypeHash();
+}
+
+// The property handles are the only things that matter when comparing, they will always be unique.
+bool FActiveDeviceProperty::operator==(const FActiveDeviceProperty& Other) const
+{
+	return PropertyHandle == Other.PropertyHandle;
+}
+
+bool FActiveDeviceProperty::operator!=(const FActiveDeviceProperty& Other) const
+{
+	return PropertyHandle != Other.PropertyHandle;
+}
+
+bool operator==(const FActiveDeviceProperty& ActiveProp, const FInputDevicePropertyHandle& Handle)
+{
+	return ActiveProp.PropertyHandle == Handle;
+}
+
+bool operator!=(const FActiveDeviceProperty& ActiveProp, const FInputDevicePropertyHandle& Handle)
+{
+	return ActiveProp.PropertyHandle != Handle;
+}
+
 ////////////////////////////////////////////////////////
 // FSetDevicePropertyParams
 
@@ -202,16 +232,17 @@ ETickableTickType UInputDeviceSubsystem::GetTickableTickType() const
 
 bool UInputDeviceSubsystem::IsAllowedToTick() const
 {
+	// Only tick when there are active device properties or ones we want to remove
+	const bool bWantsTick = !ActiveProperties.IsEmpty() || !PropertiesPendingRemoval.IsEmpty();
 #if WITH_EDITOR
 	// If we are PIE'ing, then check if PIE is paused
 	if (GEditor && (GEditor->bIsSimulatingInEditor || GEditor->PlayWorld))
 	{
-		return bIsPIEPlaying && !ActiveProperties.IsEmpty();
+		return bIsPIEPlaying && bWantsTick;
 	}	
 #endif
-
-	// Only tick when there are active device properties
-	return !ActiveProperties.IsEmpty();
+	
+	return bWantsTick;
 }
 
 bool UInputDeviceSubsystem::IsTickableInEditor() const
@@ -233,43 +264,56 @@ void UInputDeviceSubsystem::Tick(float InDeltaTime)
 	const UWorld* World = GetTickableGameObjectWorld();
 	const bool bIsGamePaused = World ? World->IsPaused() : false;
 
-	for (int32 Index = ActiveProperties.Num() - 1; Index >= 0; --Index)
+	for (auto It = ActiveProperties.CreateIterator(); It; ++It)
 	{
-		if (!ActiveProperties[Index].Property)
+		FActiveDeviceProperty& ActiveProp = *It;
+		if (!ActiveProp.Property)
 		{
 			// Something has gone wrong if we get here... maybe the Property has been GC'd ?
 			// This is really just an emergency handling case
 			ensure(false);
 			continue;
 		}
+		
+		if (PropertiesPendingRemoval.Contains(ActiveProp.PropertyHandle))
+		{
+			ActiveProp.Property->ResetDeviceProperty(ActiveProp.PlatformUser);
+			It.RemoveCurrent();
+			PropertiesPendingRemoval.Remove(ActiveProp.PropertyHandle);
+			continue;
+		}
 
-		// If the game is paused, only play effects that have explicitly specified that they should be 
+		// If the game is paused, only play effects that have explicitly specified that they should be
 		// played while paused.
-		if (bIsGamePaused && !ActiveProperties[Index].bPlayWhilePaused)
+		if (bIsGamePaused && !ActiveProp.bPlayWhilePaused)
 		{
 			continue;
 		}
 
-		const double DeltaTime = ActiveProperties[Index].bIgnoreTimeDilation ? NonDialatedDeltaTime : static_cast<double>(InDeltaTime);
-		
+		const double DeltaTime = ActiveProp.bIgnoreTimeDilation ? NonDialatedDeltaTime : static_cast<double>(InDeltaTime);
+
 		// Increase the evaluated time of this property
-		ActiveProperties[Index].EvaluatedDuration += DeltaTime;
+		ActiveProp.EvaluatedDuration += DeltaTime;
 
 		// If the property has run past it's duration, reset it and remove it from our active properties
 		// Only do this if it is marked as 'bRemoveAfterEvaluationTime' so that you can keep device properties set
 		// without having to worry about duration. 
-		if (ActiveProperties[Index].bRemoveAfterEvaluationTime && ActiveProperties[Index].EvaluatedDuration > ActiveProperties[Index].Property->GetDuration())
+		if (ActiveProp.bRemoveAfterEvaluationTime && ActiveProp.EvaluatedDuration > ActiveProp.Property->GetDuration())
 		{
-			ActiveProperties[Index].Property->ResetDeviceProperty(ActiveProperties[Index].PlatformUser);
-			ActiveProperties.RemoveAtSwap(Index);			
+			ActiveProp.Property->ResetDeviceProperty(ActiveProp.PlatformUser);
+			It.RemoveCurrent();
+			continue;
 		}
 		// Otherwise, we can evaluate and apply it as normal
 		else
 		{
-			ActiveProperties[Index].Property->EvaluateDeviceProperty(ActiveProperties[Index].PlatformUser, DeltaTime, ActiveProperties[Index].EvaluatedDuration);
-			ActiveProperties[Index].Property->ApplyDeviceProperty(ActiveProperties[Index].PlatformUser);
-		}
+			ActiveProp.Property->EvaluateDeviceProperty(ActiveProp.PlatformUser, DeltaTime, ActiveProp.EvaluatedDuration);
+			ActiveProp.Property->ApplyDeviceProperty(ActiveProp.PlatformUser);
+		}		 
 	}
+
+	// After we are done ticking, there should be no properties pending removal
+	ensure(PropertiesPendingRemoval.IsEmpty());
 }
 
 APlayerController* UInputDeviceSubsystem::GetPlayerControllerFromPlatformUser(const FPlatformUserId UserId)
@@ -335,72 +379,42 @@ FInputDevicePropertyHandle UInputDeviceSubsystem::SetDeviceProperty(const FSetDe
 	return OutHandle;
 }
 
-TObjectPtr<UInputDeviceProperty> UInputDeviceSubsystem::GetActiveDeviceProperty(const FInputDevicePropertyHandle Handle)
+UInputDeviceProperty* UInputDeviceSubsystem::GetActiveDeviceProperty(const FInputDevicePropertyHandle Handle) const
 {
-	for (const FActiveDeviceProperty& ActiveProp : ActiveProperties)
+	// Don't include any properties that are pending removal
+	if (!PropertiesPendingRemoval.Contains(Handle))
 	{
-		if (ActiveProp.PropertyHandle == Handle)
+		// We can find the active property based on the handle's hash because FActiveDeviceProperty::GetTypeHash
+		// just returns its FInputDevicePropertyHandle's GetTypeHash
+		if (const FActiveDeviceProperty* ExistingProp = ActiveProperties.FindByHash(Handle.GetTypeHash(), Handle))
 		{
-			return ActiveProp.Property;
+			return ExistingProp->Property;
 		}
 	}
+	
 	return nullptr;
 }
 
-int32 UInputDeviceSubsystem::RemoveDevicePropertiesOfClass(const FPlatformUserId UserId, TSubclassOf<UInputDeviceProperty> DevicePropertyClass)
+bool UInputDeviceSubsystem::IsPropertyActive(const FInputDevicePropertyHandle Handle) const
 {
-	int32 NumRemoved = 0;
-
-	if (DevicePropertyClass)
-	{
-		// Remove all active properties that are of the same class type
-		for (int32 Index = ActiveProperties.Num() - 1; Index >= 0; --Index)
-		{
-			if (ActiveProperties[Index].PlatformUser == UserId && ActiveProperties[Index].Property && ActiveProperties[Index].Property->GetClass() == DevicePropertyClass)
-			{
-				ActiveProperties[Index].Property->ResetDeviceProperty(ActiveProperties[Index].PlatformUser);
-				ActiveProperties.RemoveAtSwap(Index);
-				++NumRemoved;
-			}
-		}
-	}	
-	else
-	{
-		UE_LOG(LogInputDeviceProperties, Error, TEXT("Invalid DevicePropertyClass passed into RemoveDeviceProperty! Nothing will happen."));
-	}
-
-	return NumRemoved;
+	return ActiveProperties.ContainsByHash(Handle.GetTypeHash(), Handle) && !PropertiesPendingRemoval.Contains(Handle);
 }
 
-int32 UInputDeviceSubsystem::RemoveDevicePropertyByHandle(const FInputDevicePropertyHandle HandleToRemove, const bool bResetOnRemoval /*= true*/)
+void UInputDeviceSubsystem::RemoveDevicePropertyByHandle(const FInputDevicePropertyHandle HandleToRemove)
 {
-	return RemoveDevicePropertyHandles({ HandleToRemove }, bResetOnRemoval);
+	PropertiesPendingRemoval.Add(HandleToRemove);
 }
 
-int32 UInputDeviceSubsystem::RemoveDevicePropertyHandles(const TSet<FInputDevicePropertyHandle>& HandlesToRemove, const bool bResetOnRemoval /*= true*/)
+void UInputDeviceSubsystem::RemoveDevicePropertyHandles(const TSet<FInputDevicePropertyHandle>& HandlesToRemove)
 {
-	int32 NumRemoved = 0;
-
-	for (int32 Index = ActiveProperties.Num() - 1; Index >= 0; --Index)
+	if (!HandlesToRemove.IsEmpty())
 	{
-		if (HandlesToRemove.Contains(ActiveProperties[Index].PropertyHandle))
-		{
-			if (bResetOnRemoval)
-			{
-				ActiveProperties[Index].Property->ResetDeviceProperty(ActiveProperties[Index].PlatformUser);
-			}			
-			ActiveProperties.RemoveAtSwap(Index);
-			++NumRemoved;
-			break;
-		}
+		PropertiesPendingRemoval.Append(HandlesToRemove);	
 	}
-
-	if (NumRemoved <= 0)
+	else
 	{
-		UE_LOG(LogInputDeviceProperties, Warning, TEXT("Unable to remove any device property handles!"));
+		UE_LOG(LogInputDeviceProperties, Warning, TEXT("Provided an empty set of handles to remove. Nothing will happen."));
 	}
-
-	return NumRemoved;
 }
 
 bool UInputDeviceSubsystem::IsDevicePropertyHandleValid(const FInputDevicePropertyHandle& InHandle)
@@ -417,8 +431,9 @@ void UInputDeviceSubsystem::RemoveAllDeviceProperties()
 			ActiveProperty.Property->ResetDeviceProperty(ActiveProperty.PlatformUser);
 		}		
 	}
-
+	
 	ActiveProperties.Empty();
+	PropertiesPendingRemoval.Empty();
 }
 
 FHardwareDeviceIdentifier UInputDeviceSubsystem::GetMostRecentlyUsedHardwareDevice(const FPlatformUserId InUserId) const
