Fix crash when opening a material function with Strata, due to thickness propagation not being completed.

diff --git a/Engine/Source/Runtime/Engine/Private/Materials/MaterialExpressions.cpp b/Engine/Source/Runtime/Engine/Private/Materials/MaterialExpressions.cpp
--- a/Engine/Source/Runtime/Engine/Private/Materials/MaterialExpressions.cpp
+++ b/Engine/Source/Runtime/Engine/Private/Materials/MaterialExpressions.cpp
@@ -22530,8 +22530,16 @@ int32 UMaterialExpressionStrataSlabBSDF::Compile(class FMaterialCompiler* Compil
 	FStrataOperator& StrataOperator = Compiler->StrataCompilationGetOperator(Compiler->StrataTreeStackGetPathUniqueId());
 	StrataOperator.BSDFRegisteredSharedLocalBasis = NewRegisteredSharedLocalBasis;
 	
-	check(StrataOperator.ThicknessIndex != INDEX_NONE);
-	int32 ThicknesCodeChunk = Compiler->StrataThicknessStackGetThicknessCode(StrataOperator.ThicknessIndex);
+	int32 ThicknesCodeChunk = INDEX_NONE;
+	if (StrataOperator.ThicknessIndex != INDEX_NONE)
+	{
+		ThicknesCodeChunk = Compiler->StrataThicknessStackGetThicknessCode(StrataOperator.ThicknessIndex);
+	}
+	else
+	{
+		// Thickness is not tracked properly, this can happen when opening a material function in editor
+		ThicknesCodeChunk = Compiler->Constant(STRATA_LAYER_DEFAULT_THICKNESS_CM);
+	}
 	check(ThicknesCodeChunk != INDEX_NONE);
 
 	int32 OutputCodeChunk = Compiler->StrataSlabBSDF(
