Deprecating an unused function and moving it to its cpp file to avoid include dependency.

diff --git a/Engine/Source/Runtime/Engine/Private/Rendering/MorphTargetVertexInfoBuffers.cpp b/Engine/Source/Runtime/Engine/Private/Rendering/MorphTargetVertexInfoBuffers.cpp
--- a/Engine/Source/Runtime/Engine/Private/Rendering/MorphTargetVertexInfoBuffers.cpp
+++ b/Engine/Source/Runtime/Engine/Private/Rendering/MorphTargetVertexInfoBuffers.cpp
@@ -4,6 +4,7 @@
 #include "ProfilingDebugging/LoadTimeTracker.h"
 #include "Rendering/SkeletalMeshLODRenderData.h"
 #include "Animation/MorphTarget.h"
+#include "RHI.h"
 
 extern int32 GSkinCacheRecomputeTangents;
 
@@ -40,6 +41,13 @@ void FMorphTargetVertexInfoBuffers::ReleaseRHI()
 	bRHIIntialized = false;
 }
 
+uint32 FMorphTargetVertexInfoBuffers::GetMaximumThreadGroupSize()
+{
+	//D3D11 there can be at most 65535 Thread Groups in each dimension of a Dispatch call.
+	uint64 MaximumThreadGroupSize = uint64(GMaxComputeDispatchDimension) * 32ull;
+	return uint32(FMath::Min<uint64>(MaximumThreadGroupSize, UINT32_MAX));
+}
+
 const float FMorphTargetVertexInfoBuffers::CalculatePositionPrecision(float TargetPositionErrorTolerance)
 {
 	const float UnrealUnitPerMeter = 100.0f;
diff --git a/Engine/Source/Runtime/Engine/Public/Rendering/MorphTargetVertexInfoBuffers.h b/Engine/Source/Runtime/Engine/Public/Rendering/MorphTargetVertexInfoBuffers.h
--- a/Engine/Source/Runtime/Engine/Public/Rendering/MorphTargetVertexInfoBuffers.h
+++ b/Engine/Source/Runtime/Engine/Public/Rendering/MorphTargetVertexInfoBuffers.h
@@ -26,12 +26,8 @@ public:
 	ENGINE_API virtual void InitRHI() override;
 	ENGINE_API virtual void ReleaseRHI() override;
 
-	static uint32 GetMaximumThreadGroupSize()
-	{
-		//D3D11 there can be at most 65535 Thread Groups in each dimension of a Dispatch call.
-		uint64 MaximumThreadGroupSize = uint64(GMaxComputeDispatchDimension) * 32ull;
-		return uint32(FMath::Min<uint64>(MaximumThreadGroupSize, UINT32_MAX));
-	}
+	UE_DEPRECATED(5.2, "GetMaximumThreadGroupSize will be removed as it is no longer used.")
+	static ENGINE_API uint32 GetMaximumThreadGroupSize();
 
 	ENGINE_API uint32 GetNumBatches(uint32 index = UINT_MAX) const
 	{
