fix Texture LODBias not working in Editor with NeverStream textures
in Editor, the drop mip bias is simulated by treating those mips as streaming

diff --git a/Engine/Source/Runtime/Engine/Private/Texture.cpp b/Engine/Source/Runtime/Engine/Private/Texture.cpp
--- a/Engine/Source/Runtime/Engine/Private/Texture.cpp
+++ b/Engine/Source/Runtime/Engine/Private/Texture.cpp
@@ -1523,6 +1523,14 @@ FStreamableRenderResourceState UTexture::GetResourcePostInitState(const FTexture
 	// "ExpectedAssetLODBias" is the number of mips that would be dropped in cook
 	//		in a cooked run, it is zero
 
+	// in Editor the mips that will be dropped in cook are still present
+	//	dropping them is simulated by treating them as streamable
+	if ( ExpectedAssetLODBias > 0 && ! bTextureIsStreamable )
+	{
+		check( ! FPlatformProperties::RequiresCookedData() ); // ExpectedAssetLODBias should have been zero
+		NumOfNonStreamingMips = FullMipCount - ExpectedAssetLODBias;
+	}
+
 	// GMaxTextureMipCount is for the current running RHI
 	//  it may be lower than the number of mips we cooked (eg. on mobile)
 	//	we must limit the number of mips to this count.
