Adding includes and forward declarations for upcoming include changes.

diff --git a/Engine/Source/Runtime/Engine/Private/Particles/ParticleResources.cpp b/Engine/Source/Runtime/Engine/Private/Particles/ParticleResources.cpp
--- a/Engine/Source/Runtime/Engine/Private/Particles/ParticleResources.cpp
+++ b/Engine/Source/Runtime/Engine/Private/Particles/ParticleResources.cpp
@@ -6,6 +6,7 @@
 
 #include "ParticleResources.h"
 #include "ParticleCurveTexture.h"
+#include "RHI.h"
 
 /** The size of the scratch vertex buffer. */
 const int32 GParticleScratchVertexBufferSize = 64 * (1 << 10); // 64KB
diff --git a/Engine/Source/Runtime/Engine/Private/Rendering/ColorVertexBuffer.cpp b/Engine/Source/Runtime/Engine/Private/Rendering/ColorVertexBuffer.cpp
--- a/Engine/Source/Runtime/Engine/Private/Rendering/ColorVertexBuffer.cpp
+++ b/Engine/Source/Runtime/Engine/Private/Rendering/ColorVertexBuffer.cpp
@@ -2,11 +2,11 @@
 
 #include "Rendering/ColorVertexBuffer.h"
 #include "CoreMinimal.h"
-#include "RHI.h"
 #include "Components.h"
 #include "EngineUtils.h"
 #include "ProfilingDebugging/LoadTimeTracker.h"
 #include "VertexFactory.h"
+#include "DataDrivenShaderPlatformInfo.h"
 
 /*-----------------------------------------------------------------------------
 FColorVertexBuffer
diff --git a/Engine/Source/Runtime/Engine/Private/VectorField.h b/Engine/Source/Runtime/Engine/Private/VectorField.h
--- a/Engine/Source/Runtime/Engine/Private/VectorField.h
+++ b/Engine/Source/Runtime/Engine/Private/VectorField.h
@@ -22,6 +22,8 @@ class FSceneView;
 /** A vertex factory that can be used to visualize vector fields. */
 class FVectorFieldVisualizationVertexFactory;
 
+class FRHICommandListImmediate;
+
 /*------------------------------------------------------------------------------
 	Vector field delcarations.
 ------------------------------------------------------------------------------*/
diff --git a/Engine/Source/Runtime/Engine/Public/GPUSkinCache.h b/Engine/Source/Runtime/Engine/Public/GPUSkinCache.h
--- a/Engine/Source/Runtime/Engine/Public/GPUSkinCache.h
+++ b/Engine/Source/Runtime/Engine/Public/GPUSkinCache.h
@@ -43,6 +43,7 @@
 #include "VertexFactory.h"
 #include "CanvasTypes.h"
 #include "CachedGeometry.h"
+#include "DataDrivenShaderPlatformInfo.h"
 
 class FGPUSkinPassthroughVertexFactory;
 class FGPUBaseSkinVertexFactory;
diff --git a/Engine/Source/Runtime/Engine/Public/PrimitiveSceneProxy.h b/Engine/Source/Runtime/Engine/Public/PrimitiveSceneProxy.h
--- a/Engine/Source/Runtime/Engine/Public/PrimitiveSceneProxy.h
+++ b/Engine/Source/Runtime/Engine/Public/PrimitiveSceneProxy.h
@@ -31,6 +31,7 @@ enum class ERuntimeVirtualTextureMaterialType : uint8;
 struct FMeshBatch;
 class FColorVertexBuffer;
 struct FInstanceUpdateCmdBuffer;
+class FRayTracingGeometry;
 
 namespace Nanite
 {
diff --git a/Engine/Source/Runtime/Engine/Public/RayTracingInstance.h b/Engine/Source/Runtime/Engine/Public/RayTracingInstance.h
--- a/Engine/Source/Runtime/Engine/Public/RayTracingInstance.h
+++ b/Engine/Source/Runtime/Engine/Public/RayTracingInstance.h
@@ -7,6 +7,8 @@
 
 #if RHI_RAYTRACING
 
+class FRayTracingGeometry;
+
 struct FRayTracingMaskAndFlags
 {
 	/** Instance mask that can be used to exclude the instance from specific effects (eg. ray traced shadows). */
diff --git a/Engine/Source/Runtime/Engine/Public/Rendering/ColorVertexBuffer.h b/Engine/Source/Runtime/Engine/Public/Rendering/ColorVertexBuffer.h
--- a/Engine/Source/Runtime/Engine/Public/Rendering/ColorVertexBuffer.h
+++ b/Engine/Source/Runtime/Engine/Public/Rendering/ColorVertexBuffer.h
@@ -4,6 +4,7 @@
 
 #include "RenderResource.h"
 #include "StaticMeshVertexData.h"
+#include "RHI.h"
 
 struct FStaticMeshBuildVertex;
 
diff --git a/Engine/Source/Runtime/Engine/Public/Rendering/PositionVertexBuffer.h b/Engine/Source/Runtime/Engine/Public/Rendering/PositionVertexBuffer.h
--- a/Engine/Source/Runtime/Engine/Public/Rendering/PositionVertexBuffer.h
+++ b/Engine/Source/Runtime/Engine/Public/Rendering/PositionVertexBuffer.h
@@ -3,6 +3,7 @@
 #pragma once
 
 #include "RenderResource.h"
+#include "RHI.h"
 
 struct FStaticMeshBuildVertex;
 
diff --git a/Engine/Source/Runtime/Engine/Public/SceneManagement.h b/Engine/Source/Runtime/Engine/Public/SceneManagement.h
--- a/Engine/Source/Runtime/Engine/Public/SceneManagement.h
+++ b/Engine/Source/Runtime/Engine/Public/SceneManagement.h
@@ -57,6 +57,7 @@ struct FDynamicMeshVertex;
 class ULightMapVirtualTexture2D;
 class FGPUScenePrimitiveCollector;
 class FVirtualShadowMapArrayCacheManager;
+class FRayTracingGeometry;
 
 namespace UE { namespace Color { class FColorSpace; } }
 
diff --git a/Engine/Source/Runtime/Engine/Public/SimpleElementShaders.h b/Engine/Source/Runtime/Engine/Public/SimpleElementShaders.h
--- a/Engine/Source/Runtime/Engine/Public/SimpleElementShaders.h
+++ b/Engine/Source/Runtime/Engine/Public/SimpleElementShaders.h
@@ -12,6 +12,7 @@
 #include "GlobalShader.h"
 #include "SceneTypes.h"
 #include "Engine/EngineTypes.h"
+#include "DataDrivenShaderPlatformInfo.h"
 
 class FSceneView;
 struct FRelativeViewMatrices;
diff --git a/Engine/Source/Runtime/Engine/Public/SkeletalRenderPublic.h b/Engine/Source/Runtime/Engine/Public/SkeletalRenderPublic.h
--- a/Engine/Source/Runtime/Engine/Public/SkeletalRenderPublic.h
+++ b/Engine/Source/Runtime/Engine/Public/SkeletalRenderPublic.h
@@ -22,6 +22,7 @@ class UMorphTarget;
 struct FSkelMeshRenderSection;
 struct FCachedGeometry;
 class FGPUSkinCacheEntry;
+class FRayTracingGeometry;
 
 /** data for a single skinned skeletal mesh vertex */
 struct FFinalSkinVertex
diff --git a/Engine/Source/Runtime/Engine/Public/UpdateTextureShaders.h b/Engine/Source/Runtime/Engine/Public/UpdateTextureShaders.h
--- a/Engine/Source/Runtime/Engine/Public/UpdateTextureShaders.h
+++ b/Engine/Source/Runtime/Engine/Public/UpdateTextureShaders.h
@@ -7,6 +7,7 @@
 #include "ShaderParameters.h"
 #include "Shader.h"
 #include "GlobalShader.h"
+#include "DataDrivenShaderPlatformInfo.h"
 
 enum class EUpdateTextureValueType
 {
diff --git a/Engine/Source/Runtime/RenderCore/Private/ShaderCodeLibrary.cpp b/Engine/Source/Runtime/RenderCore/Private/ShaderCodeLibrary.cpp
--- a/Engine/Source/Runtime/RenderCore/Private/ShaderCodeLibrary.cpp
+++ b/Engine/Source/Runtime/RenderCore/Private/ShaderCodeLibrary.cpp
@@ -35,6 +35,7 @@ ShaderCodeLibrary.cpp: Bound shader state cache implementation.
 #include "Shader.h"
 #include "ShaderCodeArchive.h"
 #include "ShaderPipelineCache.h"
+#include "DataDrivenShaderPlatformInfo.h"
 #include "String/ParseTokens.h"
 
 #if WITH_EDITORONLY_DATA
diff --git a/Engine/Source/Runtime/RenderCore/Private/ShaderCore.cpp b/Engine/Source/Runtime/RenderCore/Private/ShaderCore.cpp
--- a/Engine/Source/Runtime/RenderCore/Private/ShaderCore.cpp
+++ b/Engine/Source/Runtime/RenderCore/Private/ShaderCore.cpp
@@ -23,6 +23,7 @@
 #include "Interfaces/IShaderFormat.h"
 #include "Interfaces/IShaderFormatModule.h"
 #include "RHIShaderFormatDefinitions.inl"
+#include "DataDrivenShaderPlatformInfo.h"
 #include "Compression/OodleDataCompression.h"
 #include "Serialization/MemoryWriter.h"
 #if WITH_EDITOR
diff --git a/Engine/Source/Runtime/RenderCore/Private/ShaderPipelineCache.cpp b/Engine/Source/Runtime/RenderCore/Private/ShaderPipelineCache.cpp
--- a/Engine/Source/Runtime/RenderCore/Private/ShaderPipelineCache.cpp
+++ b/Engine/Source/Runtime/RenderCore/Private/ShaderPipelineCache.cpp
@@ -13,6 +13,7 @@
 #include "Serialization/CustomVersion.h"
 #include "Serialization/MemoryReader.h"
 #include "Shader.h"
+#include "RenderUtils.h"
 #include "Misc/EngineVersion.h"
 #include "PipelineStateCache.h"
 #include "PipelineFileCache.h"
diff --git a/Engine/Source/Runtime/RenderCore/Public/CopyTextureShaders.h b/Engine/Source/Runtime/RenderCore/Public/CopyTextureShaders.h
--- a/Engine/Source/Runtime/RenderCore/Public/CopyTextureShaders.h
+++ b/Engine/Source/Runtime/RenderCore/Public/CopyTextureShaders.h
@@ -13,6 +13,7 @@
 #include "ShaderCore.h"
 #include "ShaderParameterUtils.h"
 #include "ShaderParameters.h"
+#include "DataDrivenShaderPlatformInfo.h"
 
 class FPointerTableBase;
 class FRHIComputeShader;
diff --git a/Engine/Source/Runtime/RenderCore/Public/ResolveShader.h b/Engine/Source/Runtime/RenderCore/Public/ResolveShader.h
--- a/Engine/Source/Runtime/RenderCore/Public/ResolveShader.h
+++ b/Engine/Source/Runtime/RenderCore/Public/ResolveShader.h
@@ -11,6 +11,7 @@
 #include "Shader.h"
 #include "ShaderCore.h"
 #include "ShaderParameters.h"
+#include "DataDrivenShaderPlatformInfo.h"
 
 class FPointerTableBase;
 class FRHICommandList;
diff --git a/Engine/Source/Runtime/Renderer/Private/WideCustomResolveShaders.h b/Engine/Source/Runtime/Renderer/Private/WideCustomResolveShaders.h
--- a/Engine/Source/Runtime/Renderer/Private/WideCustomResolveShaders.h
+++ b/Engine/Source/Runtime/Renderer/Private/WideCustomResolveShaders.h
@@ -7,6 +7,7 @@
 #include "Shader.h"
 #include "GlobalShader.h"
 #include "ShaderParameterUtils.h"
+#include "DataDrivenShaderPlatformInfo.h"
 
 class FWideCustomResolveVS : public FGlobalShader
 {
diff --git a/Engine/Source/Runtime/Renderer/Public/HdrCustomResolveShaders.h b/Engine/Source/Runtime/Renderer/Public/HdrCustomResolveShaders.h
--- a/Engine/Source/Runtime/Renderer/Public/HdrCustomResolveShaders.h
+++ b/Engine/Source/Runtime/Renderer/Public/HdrCustomResolveShaders.h
@@ -14,6 +14,7 @@
 #include "ShaderCore.h"
 #include "ShaderParameterUtils.h"
 #include "ShaderParameters.h"
+#include "DataDrivenShaderPlatformInfo.h"
 
 class FPointerTableBase;
 
diff --git a/Engine/Source/Runtime/Renderer/Public/TessellationTable.h b/Engine/Source/Runtime/Renderer/Public/TessellationTable.h
--- a/Engine/Source/Runtime/Renderer/Public/TessellationTable.h
+++ b/Engine/Source/Runtime/Renderer/Public/TessellationTable.h
@@ -8,6 +8,8 @@
 #include "Math/IntVector.h"
 #include "Math/UnrealMathSSE.h"
 #include "RenderResource.h"
+#include "RHI.h"
+#include "RHIUtilities.h"
 
 namespace Nanite
 {
