SceneCapture: Update display name for SCS_FinalColorHDR to indicate that the color remains in the working color space.

#rnx

diff --git a/Engine/Source/Runtime/Engine/Classes/Engine/EngineTypes.h b/Engine/Source/Runtime/Engine/Classes/Engine/EngineTypes.h
--- a/Engine/Source/Runtime/Engine/Classes/Engine/EngineTypes.h
+++ b/Engine/Source/Runtime/Engine/Classes/Engine/EngineTypes.h
@@ -494,7 +494,7 @@ enum ESceneCaptureSource
 	SCS_DeviceDepth UMETA(DisplayName = "DeviceDepth in RGB"),
 	SCS_Normal UMETA(DisplayName="Normal in RGB (Deferred Renderer only)"),
 	SCS_BaseColor UMETA(DisplayName = "BaseColor in RGB (Deferred Renderer only)"),
-	SCS_FinalColorHDR UMETA(DisplayName = "Final Color (HDR) in Linear sRGB gamut"),
+	SCS_FinalColorHDR UMETA(DisplayName = "Final Color (HDR) in Linear Working Color Space"),
 	SCS_FinalToneCurveHDR UMETA(DisplayName = "Final Color (with tone curve) in Linear sRGB gamut"),
 
 	SCS_MAX
