Reworked distortion to:
- Use material F0=>IOR by default.
- Then use the root node refraction pin input to override IOR for art direction.
- Refraction is now disable using the mode None.
- This data converted according to the graph content. Usingthe Rectraction input is detected by looking at all the node from the mateiral graph. When material attributes are used however, It is however over conservative since it does not check if the node a really plugged in an actual attribute path.

It was needed to deprecated the previouls RefractionMode variable as default value needs to now be RM_None. The new property is RefractionMethod.

The new GetRootNodeOverridesDefaultRefraction function on FMaterial has a default but to have correct refraction it must be implemented when it matters.

Refraction engine test were all green.

diff --git a/Engine/Plugins/Enterprise/GLTFExporter/Source/GLTFMaterialBaking/Private/GLTFExportMaterialProxy.h b/Engine/Plugins/Enterprise/GLTFExporter/Source/GLTFMaterialBaking/Private/GLTFExportMaterialProxy.h
--- a/Engine/Plugins/Enterprise/GLTFExporter/Source/GLTFMaterialBaking/Private/GLTFExportMaterialProxy.h
+++ b/Engine/Plugins/Enterprise/GLTFExporter/Source/GLTFMaterialBaking/Private/GLTFExportMaterialProxy.h
@@ -323,6 +323,9 @@ public:
 
 	virtual enum EStrataBlendMode GetStrataBlendMode() const override { return EStrataBlendMode::SBM_Opaque; }
 
+	virtual enum ERefractionMode GetRefractionMode() const override { return Material ? (ERefractionMode)Material->RefractionMethod : RM_None; }
+	virtual bool GetRootNodeOverridesDefaultRefraction()const override { return Material ? Material->bRootNodeOverridesDefaultDistortion : false; }
+
 	virtual bool GetParameterValue(EMaterialParameterType Type, const FHashedMaterialParameterInfo& ParameterInfo, FMaterialParameterValue& OutValue, const FMaterialRenderContext& Context) const override
 	{
 		return MaterialInterface->GetRenderProxy()->GetParameterValue(Type, ParameterInfo, OutValue, Context);
diff --git a/Engine/Shaders/Private/DistortAccumulatePS.usf b/Engine/Shaders/Private/DistortAccumulatePS.usf
--- a/Engine/Shaders/Private/DistortAccumulatePS.usf
+++ b/Engine/Shaders/Private/DistortAccumulatePS.usf
@@ -114,6 +114,21 @@ void Main(
 		Sigma = RefractionLobeStat.Sigma;
 	}
 
+	// We accumulate F0 only from the top layer to evaluate the material F0 used to evaluate distortion
+	float3 TopF0 = 0.0f;
+	STRATA_UNROLL_N(STRATA_CLAMPED_BSDF_COUNT)
+	for (int BSDFIdx = 0; BSDFIdx < StrataTree.BSDFCount; ++BSDFIdx)
+	{
+	#define BSDF StrataTree.BSDFs[BSDFIdx]
+		const bool bIsVisible = StrataIsBSDFVisible(BSDF);
+		if (bIsVisible)
+		{
+			TopF0 += BSDF.TopLayerDataWeight * StrataGetBSDFSpecularF0(BSDF);
+		}
+	#undef BSDF
+	}
+	TopF0 = SanitizeF0(TopF0);	// Sanitize, in case a Add node has been used
+
 	// Write out variance (as this is more linear to accumulate than roughness)
 	OutVariance = float4(Sigma, 0.0f, 0.0f, 0.0f);
 
@@ -123,18 +138,25 @@ void Main(
 	// Compute UV distortion
 	float2 BufferUVDistortion = ComputeBufferUVDistortion(
 		MaterialParameters, PixelMaterialInputs, ResolvedView,
-		RefractionWorldNormal, DistortionParams, ScreenUV, RefractionData, TryToClip, EyeIndex);
+		RefractionWorldNormal, DielectricF0RGBToIor(TopF0),
+		DistortionParams, ScreenUV, RefractionData, TryToClip, EyeIndex);
 
 	// Sample depth at distortion offset
 	float2 DistortBufferUV = ScreenUV + BufferUVDistortion;
 
 #else // STRATA_ENABLED
 
+	float3 BaseColor = GetMaterialBaseColor(PixelMaterialInputs);
+	float  Metallic = GetMaterialMetallic(PixelMaterialInputs);
+	float  Specular = GetMaterialSpecular(PixelMaterialInputs);
+	float3 F0 = ComputeF0(Specular, BaseColor, Metallic);
+	float MaterialIOR = DielectricF0RGBToIor(F0);
 
 	// Compute UV distortion
 	float2 BufferUVDistortion = ComputeBufferUVDistortion(
 		MaterialParameters, PixelMaterialInputs, ResolvedView,
-		MaterialParameters.WorldNormal, DistortionParams, ScreenUV, RefractionData, true, EyeIndex);
+		MaterialParameters.WorldNormal, MaterialIOR,
+		DistortionParams, ScreenUV, RefractionData, true, EyeIndex);
 
 	// Sample depth at distortion offset
 	float2 DistortBufferUV = ScreenUV + BufferUVDistortion;
diff --git a/Engine/Shaders/Private/DistortionCommon.ush b/Engine/Shaders/Private/DistortionCommon.ush
--- a/Engine/Shaders/Private/DistortionCommon.ush
+++ b/Engine/Shaders/Private/DistortionCommon.ush
@@ -6,7 +6,8 @@
 
 float2 ComputeBufferUVDistortion(
 	in FMaterialPixelParameters MaterialParameters, in FPixelMaterialInputs PixelMaterialInputs, in ViewState ResolvedView,
-	in half3 WorldNormal, in float4 DistortionParameters, in float2 ScreenUV, in FMaterialRefractionData RefractionData, in bool TryToClip, in uint EyeIndex)
+	in half3 WorldNormal, in float InMaterialIOR,
+	in float4 DistortionParameters, in float2 ScreenUV, in FMaterialRefractionData RefractionData, in bool TryToClip, in uint EyeIndex)
 {
 	half3 ViewNormal = normalize(TransformWorldVectorToView(WorldNormal));
 
@@ -14,12 +15,23 @@ float2 ComputeBufferUVDistortion(
 	half3 ViewVertexNormal = TransformTangentVectorToView(MaterialParameters, float3(0, 0, 1));
 	// Treat refraction of 1 as no refraction, to be consistent with IOR mode
 	float2 ViewportUVDistortion = (ViewVertexNormal.xy - ViewNormal.xy) * (GetMaterialRefractionPixelNormalStrength(RefractionData) - 1);
+
 #elif REFRACTION_USE_INDEX_OF_REFRACTION
 	// we assume the camera is in air
 	float AirIOR = 1.0f;
-	float2 ViewportUVDistortion = ViewNormal.xy * (GetMaterialRefractionIOR(RefractionData) - AirIOR);
+#if REFRACTION_ROOT_NODE_OVERRIDES_DEFAULT || STRATA_OPTIMIZED_UNLIT
+	float MaterialIOR = GetMaterialRefractionIOR(RefractionData);
+#else
+	float MaterialIOR = InMaterialIOR;
+#endif
+	float2 ViewportUVDistortion = ViewNormal.xy * (MaterialIOR - AirIOR);
+
 #elif REFRACTION_USE_2D_OFFSET
 	float2 ViewportUVDistortion = GetMaterialRefraction2DOffset(RefractionData);
+
+#elif REFRACTION_USE_NONE
+	float2 ViewportUVDistortion = 0.0f;
+
 #else
 #error Unkown refraction mode.
 #endif
diff --git a/Engine/Shaders/Private/PathTracing/PathTracingMaterialHitShader.usf b/Engine/Shaders/Private/PathTracing/PathTracingMaterialHitShader.usf
--- a/Engine/Shaders/Private/PathTracing/PathTracingMaterialHitShader.usf
+++ b/Engine/Shaders/Private/PathTracing/PathTracingMaterialHitShader.usf
@@ -30,6 +30,20 @@
 #include "/Engine/Private/PathTracing/Material/PathTracingFresnel.ush"
 #include "/Engine/Private/PathTracing/Material/PathTracingThinGlass.ush"
 
+float GetRefractionIor(FPixelMaterialInputs PixelMaterialInputs)
+{
+#if REFRACTION_ROOT_NODE_OVERRIDES_DEFAULT || STRATA_OPTIMIZED_UNLIT
+	float Ior = GetMaterialRefractionIOR(GetMaterialRefraction(PixelMaterialInputs));
+#else
+	float3 BaseColor = GetMaterialBaseColor(PixelMaterialInputs);
+	float  Metallic = GetMaterialMetallic(PixelMaterialInputs);
+	float  Specular = GetMaterialSpecular(PixelMaterialInputs);
+	float3 F0 = ComputeF0(Specular, BaseColor, Metallic);
+	float Ior = DielectricF0RGBToIor(F0);
+#endif
+	Ior = Ior == 0.0 ? 1.0 : Ior;
+	return Ior;
+}
 
 RAY_TRACING_ENTRY_CLOSEST_HIT(PathTracingMaterialCHS,
 	FPackedPathTracingPayload, PackedPayload,
@@ -117,15 +131,14 @@ RAY_TRACING_ENTRY_CLOSEST_HIT(PathTracingMaterialCHS,
 
 	// Store the results in local variables and reuse instead of calling the functions multiple times.
 	half3 BaseColor = GetMaterialBaseColor(PixelMaterialInputs);
+	half  Metallic = GetMaterialMetallic(PixelMaterialInputs);
+	half  Specular = GetMaterialSpecular(PixelMaterialInputs);
+	half Roughness = GetMaterialRoughness(PixelMaterialInputs);
 #if MATERIALBLENDING_TRANSLUCENT && REFRACTION_USE_INDEX_OF_REFRACTION
-	Payload.Ior = GetMaterialRefractionIOR(GetMaterialRefraction(PixelMaterialInputs));
-	Payload.Ior = (Payload.Ior == 0.0) ? 1.0 : Payload.Ior;
+	Payload.Ior = GetRefractionIor(PixelMaterialInputs);
 #else
 	Payload.Ior = 1;
 #endif
-	half  Metallic = GetMaterialMetallic(PixelMaterialInputs);
-	half  Specular = GetMaterialSpecular(PixelMaterialInputs);
-	half Roughness = GetMaterialRoughness(PixelMaterialInputs);
 
 #if !SIMPLIFIED_MATERIAL_SHADER
 
@@ -486,12 +499,13 @@ RAY_TRACING_ENTRY_ANY_HIT(PathTracingMaterialAHS,
 	// Assume Payload.ShadingModelID == SHADINGMODELID_THIN_TRANSLUCENT
 	float3 Transmission = GetThinTranslucentMaterialOutput0(MaterialParameters);
 	float Opacity = GetMaterialOpacity(PixelMaterialInputs);
+
 #if REFRACTION_USE_INDEX_OF_REFRACTION && !SIMPLIFIED_MATERIAL_SHADER
-	float Ior = GetMaterialRefractionIOR(GetMaterialRefraction(PixelMaterialInputs));
-	Ior = Ior == 0 ? 1.0 : Ior;
+	float Ior = GetRefractionIor(PixelMaterialInputs);
 #else
 	float Ior = 1.0;
 #endif
+
 	float3 Transparency = 1 - Opacity;
 	float3 V = WorldRayDirection();
 	float3 N = normalize(MaterialParameters.WorldNormal);
@@ -534,8 +548,7 @@ RAY_TRACING_ENTRY_ANY_HIT(PathTracingMaterialAHS,
 	
 	// Is refraction enabled?
 #if !SIMPLIFIED_MATERIAL_SHADER && MATERIAL_SHADINGMODEL_DEFAULT_LIT && REFRACTION_USE_INDEX_OF_REFRACTION
-	float Ior = GetMaterialRefractionIOR(GetMaterialRefraction(PixelMaterialInputs));
-	Ior = Ior == 0.0 ? 1.0 : Ior;
+	float Ior = GetRefractionIor(PixelMaterialInputs);
 	int ShadingModelID = GetMaterialShadingModel(PixelMaterialInputs);
 	if (Transparency > 0 && ShadingModelID == SHADINGMODELID_DEFAULT_LIT && Ior != 1.0)
 	{
diff --git a/Engine/Shaders/Private/RayTracing/RayTracingMaterialHitShaders.usf b/Engine/Shaders/Private/RayTracing/RayTracingMaterialHitShaders.usf
--- a/Engine/Shaders/Private/RayTracing/RayTracingMaterialHitShaders.usf
+++ b/Engine/Shaders/Private/RayTracing/RayTracingMaterialHitShaders.usf
@@ -297,6 +297,8 @@ void GetMaterialPayload(
 	Payload.Radiance += EmissiveLuminance;
 	Payload.TranslatedWorldPos = LWCToFloat(LWCAdd(MaterialParameters.AbsoluteWorldPosition, ResolvedView.PreViewTranslation));
 
+	// STRATA_TODO Payload.Ior when available (accounting for REFRACTION_USE_INDEX_OF_REFRACTION and REFRACTION_ROOT_NODE_OVERRIDES_DEFAULT)
+
 	// Opacity
 	Payload.Opacity = TotalCoverage;
 	#if MATERIALBLENDING_MASKED
@@ -312,14 +314,19 @@ void GetMaterialPayload(
 {
 	// Store the results in local variables and reuse instead of calling the functions multiple times.
 	half3 BaseColor = GetMaterialBaseColor(PixelMaterialInputs);
+	half  Metallic = GetMaterialMetallic(PixelMaterialInputs);
+	half  Specular = GetMaterialSpecular(PixelMaterialInputs);
+	half Roughness = GetMaterialRoughness(PixelMaterialInputs);
 #if MATERIALBLENDING_TRANSLUCENT && REFRACTION_USE_INDEX_OF_REFRACTION
+	#if REFRACTION_ROOT_NODE_OVERRIDES_DEFAULT
 	half IOR = GetMaterialRefractionIOR(GetMaterialRefraction(PixelMaterialInputs));
+	#else
+	half3 F0 = ComputeF0(Specular, BaseColor, Metallic);
+	half IOR = DielectricF0RGBToIor(F0);
+	#endif
 #else
 	half IOR = 1;
 #endif
-	half  Metallic = GetMaterialMetallic(PixelMaterialInputs);
-	half  Specular = GetMaterialSpecular(PixelMaterialInputs);
-	half Roughness = GetMaterialRoughness(PixelMaterialInputs);
 
 	/**
 	 * Set material attributes for full materials
diff --git a/Engine/Shaders/Private/ShadingCommon.ush b/Engine/Shaders/Private/ShadingCommon.ush
--- a/Engine/Shaders/Private/ShadingCommon.ush
+++ b/Engine/Shaders/Private/ShadingCommon.ush
@@ -102,7 +102,7 @@ float F0ToDielectricSpecular(float F0)
 
 float F0RGBToDielectricSpecular(float3 F0)
 {
-	return F0ToDielectricSpecular(max(F0.r, max(F0.g, F0.b)));
+	return F0ToDielectricSpecular(max3(F0.r, F0.g, F0.b));
 }
 
 half DielectricSpecularToF0(half Specular)
@@ -116,6 +116,11 @@ float DielectricF0ToIor(float F0)
 	return 2.0f / (1.0f - sqrt(F0)) - 1.0f;
 }
 
+float DielectricF0RGBToIor(float3 F0)
+{
+	return DielectricF0ToIor(max3(F0.r, F0.g, F0.b));
+}
+
 float DielectricIorToF0(float Ior)
 {
 	const float F0Sqrt = (Ior-1)/(Ior+1);
diff --git a/Engine/Shaders/Private/SingleLayerWaterShading.ush b/Engine/Shaders/Private/SingleLayerWaterShading.ush
--- a/Engine/Shaders/Private/SingleLayerWaterShading.ush
+++ b/Engine/Shaders/Private/SingleLayerWaterShading.ush
@@ -78,6 +78,8 @@ WaterVolumeLightingOutput EvaluateWaterVolumeLighting(
 
 #else  // !SINGLE_LAYER_WATER_MOBILE_SHADING
 
+		float IorWater = DielectricF0ToIor(DielectricSpecularToF0(Specular)); // Wrong if metal is set to >1. But we still keep refraction on the water surface nonetheless.
+
 		ViewportUV = clamp(ViewportUV, SceneWithoutSingleLayerWaterMinUV, SceneWithoutSingleLayerWaterMaxUV);
 
 		float PixelSceneDeviceZ = WaterSampleSceneDepthWithoutWater(
@@ -92,7 +94,7 @@ WaterVolumeLightingOutput EvaluateWaterVolumeLighting(
 		FMaterialRefractionData RefractionData = GetMaterialRefraction(PixelMaterialInputs);
 		float2 BufferUVDistortion = ComputeBufferUVDistortion(
 			MaterialParameters, PixelMaterialInputs, ResolvedView,
-			MaterialParameters.WorldNormal,
+			MaterialParameters.WorldNormal, IorWater,
 			DistortionParams, ViewportUV, RefractionData, false, EyeIndex);
 		PostProcessUVDistortion(MaterialParameters, PixelMaterialInputs, PixelSceneDepth, BufferUVDistortion, RefractionData);
 		// Also reduce refraction according to water thickness
@@ -158,8 +160,7 @@ WaterVolumeLightingOutput EvaluateWaterVolumeLighting(
 			DirLightPhaseValue = IsotropicPhase();
 #else
 			float IorFrom = 1.0f; // assumes we come from air
-			float IorTo   = DielectricF0ToIor(DielectricSpecularToF0(Specular)); // Wrong if metal is set to >1. But we still keep refraction on the water surface nonetheless.
-			const float relativeIOR = IorFrom / IorTo;
+			const float relativeIOR = IorFrom / IorWater;
 			float3 UnderWaterRayDir = 0.0f;
 			if (WaterRefract(MaterialParameters.CameraVector, MaterialParameters.WorldNormal, relativeIOR, UnderWaterRayDir))
 			{
diff --git a/Engine/Shaders/Public/ShaderVersion.ush b/Engine/Shaders/Public/ShaderVersion.ush
--- a/Engine/Shaders/Public/ShaderVersion.ush
+++ b/Engine/Shaders/Public/ShaderVersion.ush
@@ -2,4 +2,4 @@
 // Each time the console command is executed it generates a new GUID. As this file is included
 // in Platform.ush (which should be included in any shader) it allows to invalidate the shader DDC.
 // 
-// GUID = BDF733E1-F4F8-4332-A403-7C2DAAC8A2E9
\ No newline at end of file
+// GUID = 28A1A473-9B5D-40E9-B951-B09A08EB4DA3
diff --git a/Engine/Source/Developer/MaterialBaking/Private/ExportMaterialProxy.h b/Engine/Source/Developer/MaterialBaking/Private/ExportMaterialProxy.h
--- a/Engine/Source/Developer/MaterialBaking/Private/ExportMaterialProxy.h
+++ b/Engine/Source/Developer/MaterialBaking/Private/ExportMaterialProxy.h
@@ -516,6 +516,8 @@ public:
 	virtual bool IsMasked() const override { return false; }
 	virtual enum EBlendMode GetBlendMode() const override { return BLEND_Opaque; }
 	virtual enum EStrataBlendMode GetStrataBlendMode() const override { return EStrataBlendMode::SBM_Opaque; }
+	virtual enum ERefractionMode GetRefractionMode() const override { return Material ? (ERefractionMode)Material->RefractionMethod : RM_None; }
+	virtual bool GetRootNodeOverridesDefaultRefraction()const override { return Material ? Material->bRootNodeOverridesDefaultDistortion : false; }
 	virtual FMaterialShadingModelField GetShadingModels() const override { return MSM_DefaultLit; }
 	virtual bool IsShadingModelFromMaterialExpression() const override { return false; }
 	virtual float GetOpacityMaskClipValue() const override { return 0.5f; }
diff --git a/Engine/Source/Developer/MaterialUtilities/Private/MaterialUtilities.cpp b/Engine/Source/Developer/MaterialUtilities/Private/MaterialUtilities.cpp
--- a/Engine/Source/Developer/MaterialUtilities/Private/MaterialUtilities.cpp
+++ b/Engine/Source/Developer/MaterialUtilities/Private/MaterialUtilities.cpp
@@ -702,6 +702,8 @@ public:
 	virtual bool IsMasked() const override									{ return false; }
 	virtual enum EBlendMode GetBlendMode() const override					{ return BLEND_Opaque; }
 	virtual enum EStrataBlendMode GetStrataBlendMode() const override		{ return EStrataBlendMode::SBM_Opaque; }
+	virtual enum ERefractionMode GetRefractionMode() const override			{ return Material ? (ERefractionMode)Material->RefractionMethod : RM_None; }
+	virtual bool GetRootNodeOverridesDefaultRefraction()const override		{ return Material ? Material->bRootNodeOverridesDefaultDistortion : false; }
 	virtual FMaterialShadingModelField GetShadingModels() const override	{ return MSM_Unlit; }
 	virtual bool IsShadingModelFromMaterialExpression() const override		{ return false; }
 	virtual float GetOpacityMaskClipValue() const override					{ return 0.5f; }
diff --git a/Engine/Source/Editor/MaterialEditor/Private/MaterialEditor.h b/Engine/Source/Editor/MaterialEditor/Private/MaterialEditor.h
--- a/Engine/Source/Editor/MaterialEditor/Private/MaterialEditor.h
+++ b/Engine/Source/Editor/MaterialEditor/Private/MaterialEditor.h
@@ -126,6 +126,7 @@ public:
 	virtual bool IsMasked() const override { return false; }
 	virtual enum EBlendMode GetBlendMode() const override { return BLEND_Translucent; }
 	virtual enum EStrataBlendMode GetStrataBlendMode() const override { return EStrataBlendMode::SBM_TranslucentGreyTransmittance; }
+	virtual bool GetRootNodeOverridesDefaultRefraction()const override { return false; } // refraction unused for material preview
 	virtual FMaterialShadingModelField GetShadingModels() const override { return MSM_Unlit; }
 	virtual bool IsShadingModelFromMaterialExpression() const override { return false; }
 	virtual float GetOpacityMaskClipValue() const override { return 0.5f; }
diff --git a/Engine/Source/Editor/UnrealEd/Private/Lightmass/LightmassRender.cpp b/Engine/Source/Editor/UnrealEd/Private/Lightmass/LightmassRender.cpp
--- a/Engine/Source/Editor/UnrealEd/Private/Lightmass/LightmassRender.cpp
+++ b/Engine/Source/Editor/UnrealEd/Private/Lightmass/LightmassRender.cpp
@@ -600,6 +600,8 @@ public:
 	virtual bool IsMasked() const override									{ return false; }
 	virtual enum EBlendMode GetBlendMode() const override					{ return BLEND_Opaque; }
 	virtual enum EStrataBlendMode GetStrataBlendMode() const override		{ return EStrataBlendMode::SBM_Opaque; }
+	virtual enum ERefractionMode GetRefractionMode() const override			{ return Material ? (ERefractionMode)Material->RefractionMethod : RM_None; }
+	virtual bool GetRootNodeOverridesDefaultRefraction()const override		{ return Material ? Material->bRootNodeOverridesDefaultDistortion : false; }
 	virtual FMaterialShadingModelField GetShadingModels() const override	{ return MSM_Unlit; }
 	virtual bool IsShadingModelFromMaterialExpression() const override		{ return false; }
 	virtual float GetOpacityMaskClipValue() const override					{ return 0.5f; }
diff --git a/Engine/Source/Runtime/Core/Public/UObject/UE5MainStreamObjectVersion.h b/Engine/Source/Runtime/Core/Public/UObject/UE5MainStreamObjectVersion.h
--- a/Engine/Source/Runtime/Core/Public/UObject/UE5MainStreamObjectVersion.h
+++ b/Engine/Source/Runtime/Core/Public/UObject/UE5MainStreamObjectVersion.h
@@ -313,6 +313,9 @@ struct CORE_API FUE5MainStreamObjectVersion
 		// Add ThinSurface instance override support
 		MaterialInstanceBasePropertyOverridesThinSurface,
 
+		// Add refraction mode None, converted from legacy when the refraction pin is not plugged.
+		MaterialRefractionModeNone,
+
 		// -----<new versions can be added above this line>-------------------------------------------------
 		VersionPlusOne,
 		LatestVersion = VersionPlusOne - 1
diff --git a/Engine/Source/Runtime/Engine/Classes/Engine/EngineTypes.h b/Engine/Source/Runtime/Engine/Classes/Engine/EngineTypes.h
--- a/Engine/Source/Runtime/Engine/Classes/Engine/EngineTypes.h
+++ b/Engine/Source/Runtime/Engine/Classes/Engine/EngineTypes.h
@@ -371,13 +371,15 @@ UENUM()
 enum ERefractionMode
 {
 	/** 
+	 * By default, when the root node refraction pin is unplugged, relies on the material IOR evaluated from F0.
 	 * Refraction is computed based on the camera vector entering a medium whose index of refraction is defined by the Refraction material input.  
 	 * The new medium's surface is defined by the material's normal.  With this mode, a flat plane seen from the side will have a constant refraction offset.
 	 * This is a physical model of refraction but causes reading outside the scene color texture so is a poor fit for large refractive surfaces like water.
 	 */
 	RM_IndexOfRefraction UMETA(DisplayName="Index Of Refraction"),
 
-	/** 
+	/**
+	 * By default, when the root node refraction pin is unplugged, no refraction will appear.
 	 * The refraction offset into Scene Color is computed based on the difference between the per-pixel normal and the per-vertex normal.  
 	 * With this mode, a material whose normal is the default (0, 0, 1) will never cause any refraction.  This mode is only valid with tangent space normals.
 	 * The refraction material input scales the offset, although a value of 1.0 maps to no refraction, and a value of 2 maps to a scale of 1.0 on the offset.
@@ -386,9 +388,15 @@ enum ERefractionMode
 	RM_PixelNormalOffset UMETA(DisplayName="Pixel Normal Offset"),
 
 	/**
+	 * By default, when the root node refraction pin is unplugged, no refraction will appear.
 	 * Explicit 2D screen offset. This offset is independent of screen resolution and aspect ratio. The user is in charge of any strength and fading.
 	 */
 	RM_2DOffset UMETA(DisplayName = "2D Offset"),
+
+	/**
+	 * Refraction is disabled.
+	 */
+	RM_None UMETA(DisplayName = "None"),
 };
 
 /**
diff --git a/Engine/Source/Runtime/Engine/Classes/Materials/Material.h b/Engine/Source/Runtime/Engine/Classes/Materials/Material.h
--- a/Engine/Source/Runtime/Engine/Classes/Materials/Material.h
+++ b/Engine/Source/Runtime/Engine/Classes/Materials/Material.h
@@ -723,6 +723,12 @@ public:
 	UPROPERTY()
 	uint8 bUsesDistortion : 1;
 
+	/** 
+	 * Indicates that the material do not use the material physically based refraction (e.g. IOR from reflectivity F0), but overrides it for artistic purposes.
+	 */
+	 UPROPERTY()
+	uint8 bRootNodeOverridesDefaultDistortion : 1;
+
 	/** 
 	 * Indicates that the material and its instances can be used with clothing
 	 * This will result in the shaders required to support clothing being compiled which will increase shader compile time and memory usage.
@@ -953,9 +959,12 @@ public:
 	UPROPERTY(EditAnywhere, Category = PostProcessMaterial, AdvancedDisplay, meta = (EditCondition = "bEnableStencilTest"))
 	uint8 StencilRefValue = 0;
 
+	UPROPERTY()
+	TEnumAsByte<enum ERefractionMode> RefractionMode_DEPRECATED;
+
 	/** Controls how the Refraction input is interpreted and how the refraction offset into scene color is computed for this material. */
 	UPROPERTY(EditAnywhere, Category=Refraction)
-	TEnumAsByte<enum ERefractionMode> RefractionMode;
+	TEnumAsByte<enum ERefractionMode> RefractionMethod;
 
 	/** If multiple nodes with the same  type are inserted at the same point, this defined order and if they get combined, only used if domain is PostProcess */
 	UPROPERTY(EditAnywhere, BlueprintReadWrite, Category = PostProcessMaterial, meta = (DisplayName = "Blendable Priority"))
@@ -1827,6 +1836,8 @@ public:
 private:
 #if WITH_EDITOR
 	static FMaterialCompilationFinished MaterialCompilationFinishedEvent;
+
+	bool IsRefractionPinPluggedIn(const UMaterialEditorOnlyData* EditorOnly);
 #endif // WITH_EDITOR
 
 	friend class FLightmassMaterialProxy;
diff --git a/Engine/Source/Runtime/Engine/Private/Materials/Material.cpp b/Engine/Source/Runtime/Engine/Private/Materials/Material.cpp
--- a/Engine/Source/Runtime/Engine/Private/Materials/Material.cpp
+++ b/Engine/Source/Runtime/Engine/Private/Materials/Material.cpp
@@ -1002,6 +1002,8 @@ UMaterial::UMaterial(const FObjectInitializer& ObjectInitializer)
 	bAllowDevelopmentShaderCompile = true;
 	bIsMaterialEditorStatsMaterial = false;
 
+	RefractionMethod = RM_None;
+
 #if WITH_EDITORONLY_DATA
 	MaterialGraph = NULL;
 #endif //WITH_EDITORONLY_DATA
@@ -3940,6 +3942,37 @@ void UMaterial::PostLoad()
 	}
 #endif // #if WITH_EDITOR
 
+#if WITH_EDITOR
+	// Before, refraction was only enabled when the refraction pin was plugged in.
+	// Now it is enabled only when not OFF. Otherwise:
+	//    - if plugged the pin override the physically based material refraction
+	//    - if unplugged the refraction is converted from the material F0.
+	if (UE5MainVer < FUE5MainStreamObjectVersion::MaterialRefractionModeNone)
+	{
+		const bool bRefractionPinPluggedIn = IsRefractionPinPluggedIn(EditorOnly);
+
+		// Update to the new variable, accounting for the old default value.
+		RefractionMethod = RefractionMode_DEPRECATED;
+
+		if (!bRefractionPinPluggedIn)
+		{
+			// The root node refraction pin was not plugged in so set refraction mode as none.
+			// We do this for all domains and blending modes (translucent, opaque, etc).
+			RefractionMethod = RM_None;
+			bRootNodeOverridesDefaultDistortion = false;
+
+			// Need to mutate since UPROPERTY like bUsesDistortion is changed and shaders need to recompile.
+			static FGuid MaterialRefractionModeOFFConversionGuid(TEXT("094C0316-EA95-4B39-A3FA-CA126A92989B"));
+			ReleaseResourcesAndMutateDDCKey(MaterialRefractionModeOFFConversionGuid);
+		}
+		else
+		{
+			// Keep the current refraction mode and notify that it is overriden on the root node.
+			bRootNodeOverridesDefaultDistortion = true;
+		}
+	}
+#endif // WITH_EDITOR
+
 	STAT(double MaterialLoadTime = 0);
 	{
 		SCOPE_SECONDS_COUNTER(MaterialLoadTime);
@@ -4445,52 +4478,9 @@ void UMaterial::PostEditChangePropertyInternal(FPropertyChangedEvent& PropertyCh
 
 	const UMaterialEditorOnlyData* EditorOnly = GetEditorOnlyData();
 
-	// check for distortion in material 
-	{
-		bUsesDistortion = false;
-
-		// check for a distortion value
-		if (EditorOnly->Refraction.Expression
-			|| (EditorOnly->Refraction.UseConstant && FMath::Abs(EditorOnly->Refraction.Constant - 1.0f) >= UE_KINDA_SMALL_NUMBER))
-		{
-			bUsesDistortion = true;
-		}
-
-		// check the material attributes for refraction expressions as well
-		if (EditorOnly->MaterialAttributes.Expression)
-		{
-			// The node that possibly sets the refraction attribute can be anywhere in the sub-graph that eventually connects
-			// to the material output expression. We need to look for any node of type [Make/Set]MaterialAttributes in this subgraph
-			// and check whether it writes to the Refraction attribute.
-			TArray<UMaterialExpression*> Expressions;
-			EditorOnly->MaterialAttributes.Expression->GetAllInputExpressions(Expressions);
-
-			for (int32 ExprIndex = 0; ExprIndex < Expressions.Num() && !bUsesDistortion; ++ExprIndex)
-			{
-				// Case when node is a MakeMaterialAttributes
-				UMaterialExpressionMakeMaterialAttributes* MakeAttributeExpression = Cast<UMaterialExpressionMakeMaterialAttributes>(Expressions[ExprIndex]);
-				if (MakeAttributeExpression && MakeAttributeExpression->Refraction.Expression)
-				{
-					bUsesDistortion = true;
-				}
-
-				// Case when node is a SetMaterialAttributes
-				UMaterialExpressionSetMaterialAttributes* SetAttributeExpression = Cast<UMaterialExpressionSetMaterialAttributes>(Expressions[ExprIndex]);
-				if (!bUsesDistortion && SetAttributeExpression)
-				{
-					for (int32 InputIndex = 0; InputIndex < SetAttributeExpression->Inputs.Num(); InputIndex++)
-					{
-						FName InputName = SetAttributeExpression->GetInputName(InputIndex);
-						if (InputName == TEXT("Refraction"))
-						{
-							bUsesDistortion = true;
-							break;
-						}
-					}
-				}
-			}
-		}
-	}
+	// Check for distortion in material 
+	bUsesDistortion = RefractionMethod != RM_None;
+	bRootNodeOverridesDefaultDistortion = bUsesDistortion ? IsRefractionPinPluggedIn(EditorOnly) : false;
 
 	//If we can be sure this material would be the same opaque as it is masked then allow it to be assumed opaque.
 	bCanMaskedBeAssumedOpaque = !EditorOnly->OpacityMask.Expression && !(EditorOnly->OpacityMask.UseConstant && EditorOnly->OpacityMask.Constant < 0.999f) && !bUseMaterialAttributes;
@@ -6875,6 +6865,55 @@ void UMaterial::GetLightingGuidChain(bool bIncludeTextures, TArray<FGuid>& OutGu
 #endif
 }
 
+#if WITH_EDITOR
+bool UMaterial::IsRefractionPinPluggedIn(const UMaterialEditorOnlyData* EditorOnly)
+{
+	// check for a distortion value
+	if (EditorOnly->Refraction.Expression
+		|| (EditorOnly->Refraction.UseConstant && FMath::Abs(EditorOnly->Refraction.Constant - 1.0f) >= UE_KINDA_SMALL_NUMBER))
+	{
+		return true;
+	}
+
+	// check the material attributes for refraction expressions as well
+	if (EditorOnly->MaterialAttributes.Expression)
+	{
+		// The node that possibly sets the refraction attribute can be anywhere in the sub-graph that eventually connects
+		// to the material output expression. We need to look for any node of type [Make/Set]MaterialAttributes in this subgraph
+		// and check whether it writes to the Refraction attribute.
+
+		// For that we use GetAllExpressionsInMaterialAndFunctionsOfType which truly parse all functions.  
+		// It however do not check if the node is really connected to the graph so it can be over conservative.
+
+		TArray<UMaterialExpressionMakeMaterialAttributes*> MakeAttributeExpressions;
+		GetAllExpressionsInMaterialAndFunctionsOfType<UMaterialExpressionMakeMaterialAttributes>(MakeAttributeExpressions);
+		for (UMaterialExpressionMakeMaterialAttributes* MakeAttributeExpression : MakeAttributeExpressions)
+		{
+			if (MakeAttributeExpression && MakeAttributeExpression->Refraction.Expression)
+			{
+				return true;
+			}
+		}
+
+		TArray<UMaterialExpressionSetMaterialAttributes*> SetMaterialAttributesExpressions;
+		GetAllExpressionsInMaterialAndFunctionsOfType<UMaterialExpressionSetMaterialAttributes>(SetMaterialAttributesExpressions);
+		for (UMaterialExpressionSetMaterialAttributes* SetMaterialAttributesExpression : SetMaterialAttributesExpressions)
+		{
+			for (int32 InputIndex = 0; InputIndex < SetMaterialAttributesExpression->Inputs.Num(); InputIndex++)
+			{
+				FName InputName = SetMaterialAttributesExpression->GetInputName(InputIndex);
+				if (InputName == TEXT("Refraction"))
+				{
+					return true;
+				}
+			}
+		}
+	}
+
+	return false;
+}
+#endif // WITH_EDITOR
+
 UMaterialEditorOnlyData::UMaterialEditorOnlyData()
 {
 	BaseColor.Constant = FColor(128, 128, 128);
diff --git a/Engine/Source/Runtime/Engine/Private/Materials/MaterialShared.cpp b/Engine/Source/Runtime/Engine/Private/Materials/MaterialShared.cpp
--- a/Engine/Source/Runtime/Engine/Private/Materials/MaterialShared.cpp
+++ b/Engine/Source/Runtime/Engine/Private/Materials/MaterialShared.cpp
@@ -798,7 +798,7 @@ void FMaterial::GetStaticParameterSet(EShaderPlatform Platform, FStaticParameter
 
 ERefractionMode FMaterial::GetRefractionMode() const 
 { 
-	return RM_IndexOfRefraction; 
+	return RM_None;
 }
 
 const FMaterialCachedExpressionData& FMaterial::GetCachedExpressionData() const
@@ -1760,7 +1760,12 @@ EStrataBlendMode FMaterialResource::GetStrataBlendMode() const
 
 ERefractionMode FMaterialResource::GetRefractionMode() const
 {
-	return Material->RefractionMode;
+	return Material->RefractionMethod;
+}
+
+bool FMaterialResource::GetRootNodeOverridesDefaultRefraction() const
+{
+	return Material->bRootNodeOverridesDefaultDistortion;
 }
 
 FMaterialShadingModelField FMaterialResource::GetShadingModels() const 
@@ -2335,10 +2340,12 @@ void FMaterial::SetupMaterialEnvironment(
 	case RM_IndexOfRefraction:	OutEnvironment.SetDefine(TEXT("REFRACTION_USE_INDEX_OF_REFRACTION"),	TEXT("1")); break;
 	case RM_PixelNormalOffset:	OutEnvironment.SetDefine(TEXT("REFRACTION_USE_PIXEL_NORMAL_OFFSET"),	TEXT("1")); break;
 	case RM_2DOffset:			OutEnvironment.SetDefine(TEXT("REFRACTION_USE_2D_OFFSET"),				TEXT("1")); break;
+	case RM_None:				OutEnvironment.SetDefine(TEXT("REFRACTION_USE_NONE"),					TEXT("1")); break;
 	default: 
 		UE_LOG(LogMaterial, Warning, TEXT("Unknown material refraction mode: %u  Setting to RM_IndexOfRefraction"),(int32)GetRefractionMode());
 		OutEnvironment.SetDefine(TEXT("REFRACTION_USE_INDEX_OF_REFRACTION"),TEXT("1"));
 	}
+	OutEnvironment.SetDefine(TEXT("REFRACTION_ROOT_NODE_OVERRIDES_DEFAULT"), GetRootNodeOverridesDefaultRefraction() ? 1 : 0);
 
 	OutEnvironment.SetDefine(TEXT("USE_DITHERED_LOD_TRANSITION_FROM_MATERIAL"), IsDitheredLODTransition());
 	OutEnvironment.SetDefine(TEXT("MATERIAL_TWOSIDED"), IsTwoSided());
diff --git a/Engine/Source/Runtime/Engine/Public/MaterialShared.h b/Engine/Source/Runtime/Engine/Public/MaterialShared.h
--- a/Engine/Source/Runtime/Engine/Public/MaterialShared.h
+++ b/Engine/Source/Runtime/Engine/Public/MaterialShared.h
@@ -1844,6 +1844,7 @@ public:
 	virtual enum EBlendMode GetBlendMode() const = 0;
 	virtual enum EStrataBlendMode GetStrataBlendMode() const = 0;
 	ENGINE_API virtual enum ERefractionMode GetRefractionMode() const;
+	virtual bool GetRootNodeOverridesDefaultRefraction() const = 0;
 	virtual FMaterialShadingModelField GetShadingModels() const = 0;
 	virtual bool IsShadingModelFromMaterialExpression() const = 0;
 	virtual enum ETranslucencyLightingMode GetTranslucencyLightingMode() const { return TLM_VolumetricNonDirectional; };
@@ -2904,6 +2905,7 @@ public:
 	ENGINE_API virtual enum EBlendMode GetBlendMode() const override;
 	ENGINE_API virtual enum EStrataBlendMode GetStrataBlendMode() const override;
 	ENGINE_API virtual enum ERefractionMode GetRefractionMode() const override;
+	ENGINE_API virtual bool GetRootNodeOverridesDefaultRefraction() const override;
 	ENGINE_API virtual uint32 GetMaterialDecalResponse() const override;
 	ENGINE_API virtual bool HasBaseColorConnected() const override;
 	ENGINE_API virtual bool HasNormalConnected() const override;
