Fixed refraction pin to not be active when refraction method is none.

diff --git a/Engine/Source/Runtime/Engine/Private/Materials/Material.cpp b/Engine/Source/Runtime/Engine/Private/Materials/Material.cpp
--- a/Engine/Source/Runtime/Engine/Private/Materials/Material.cpp
+++ b/Engine/Source/Runtime/Engine/Private/Materials/Material.cpp
@@ -6474,7 +6474,7 @@ static bool IsPropertyActive_Internal(EMaterialProperty InProperty,
 	FMaterialShadingModelField ShadingModels,
 	ETranslucencyLightingMode TranslucencyLightingMode,
 	bool bBlendableOutputAlpha,
-	bool bHasRefraction,
+	bool bUsesDistortion,
 	bool bUsesShadingModelFromMaterialExpression,
 	bool bIsTranslucencyWritingVelocity,
 	bool bIsThinSurface,
@@ -6602,7 +6602,7 @@ static bool IsPropertyActive_Internal(EMaterialProperty InProperty,
 			switch (InProperty)
 			{
 			case MP_Refraction:
-				Active = (bIsTranslucentBlendMode && BlendMode != BLEND_AlphaHoldout && BlendMode != BLEND_Modulate) || ShadingModels.HasShadingModel(MSM_SingleLayerWater);
+				Active = ((bIsTranslucentBlendMode && BlendMode != BLEND_AlphaHoldout && BlendMode != BLEND_Modulate) || ShadingModels.HasShadingModel(MSM_SingleLayerWater)) && bUsesDistortion;
 				break;
 			case MP_Opacity:
 				Active = (bIsTranslucentBlendMode && BlendMode != BLEND_Modulate) || ShadingModels.HasShadingModel(MSM_SingleLayerWater);
@@ -6637,7 +6637,7 @@ static bool IsPropertyActive_Internal(EMaterialProperty InProperty,
 			Active = false;
 			break;
 		case MP_Refraction:
-			Active = (bIsTranslucentBlendMode && BlendMode != BLEND_AlphaHoldout && BlendMode != BLEND_Modulate) || ShadingModels.HasShadingModel(MSM_SingleLayerWater);
+			Active = ((bIsTranslucentBlendMode && BlendMode != BLEND_AlphaHoldout && BlendMode != BLEND_Modulate) || ShadingModels.HasShadingModel(MSM_SingleLayerWater)) && bUsesDistortion;
 			break;
 		case MP_Opacity:
 			Active = (bIsTranslucentBlendMode && BlendMode != BLEND_Modulate) || ShadingModels.HasShadingModel(MSM_SingleLayerWater);
@@ -6665,7 +6665,7 @@ static bool IsPropertyActive_Internal(EMaterialProperty InProperty,
 			Active = ShadingModels.IsLit() && (!bIsTranslucentBlendMode || !bIsVolumetricTranslucencyLightingMode);
 			break;
 		case MP_Normal:
-			Active = (ShadingModels.IsLit() && (!bIsTranslucentBlendMode || !bIsNonDirectionalTranslucencyLightingMode)) || bHasRefraction;
+			Active = (ShadingModels.IsLit() && (!bIsTranslucentBlendMode || !bIsNonDirectionalTranslucencyLightingMode)) || bUsesDistortion;
 			break;
 		case MP_Tangent:
 			Active = ShadingModels.HasAnyShadingModel({ MSM_DefaultLit, MSM_ClearCoat }) && (!bIsTranslucentBlendMode || !bIsVolumetricTranslucencyLightingMode);
@@ -6725,7 +6725,7 @@ bool UMaterial::IsPropertyActiveInEditor(EMaterialProperty InProperty) const
 		ShadingModels,
 		TranslucencyLightingMode,
 		BlendableOutputAlpha,
-		GetEditorOnlyData()->Refraction.IsConnected(),
+		bUsesDistortion,
 		IsShadingModelFromMaterialExpression(),
 		IsTranslucencyWritingVelocity(),
 		IsThinSurface(),
@@ -6751,7 +6751,7 @@ bool UMaterial::IsPropertyActiveInDerived(EMaterialProperty InProperty, const UM
 		DerivedMaterial->GetShadingModels(),
 		TranslucencyLightingMode,
 		BlendableOutputAlpha,
-		!bUseMaterialAttributes ? GetCachedExpressionData().IsPropertyConnected(MP_Refraction) : bUsesDistortion,
+		bUsesDistortion,
 		DerivedMaterial->IsShadingModelFromMaterialExpression(),
 		IsTranslucencyWritingVelocity(),
 		IsThinSurface(),
