[shaders] use _just_ static parameters and path when generating the unique asset name for debug info paths

diff --git a/Engine/Source/Runtime/Engine/Private/Materials/MaterialShader.cpp b/Engine/Source/Runtime/Engine/Private/Materials/MaterialShader.cpp
--- a/Engine/Source/Runtime/Engine/Private/Materials/MaterialShader.cpp
+++ b/Engine/Source/Runtime/Engine/Private/Materials/MaterialShader.cpp
@@ -1010,6 +1010,22 @@ void FMaterialShaderMapId::UpdateFromParameterSet(const FStaticParameterSet& Sta
 
 uint32 GetTypeHash(FPlatformTypeLayoutParameters Params) { return HashCombine(Params.Flags, Params.MaxFieldAlignment); }
 
+void FMaterialShaderMapId::AppendStaticParametersString(FString& ParamsString) const
+{
+	for (const FStaticSwitchParameter& StaticSwitchParameter : StaticSwitchParameters)
+	{
+		StaticSwitchParameter.AppendKeyString(ParamsString);
+	}
+	for (const FStaticComponentMaskParameter& StaticComponentMaskParameter : StaticComponentMaskParameters)
+	{
+		StaticComponentMaskParameter.AppendKeyString(ParamsString);
+	}
+	for (const FStaticTerrainLayerWeightParameter& StaticTerrainLayerWeightParameter : TerrainLayerWeightParameters)
+	{
+		StaticTerrainLayerWeightParameter.AppendKeyString(ParamsString);
+	}
+}
+
 void FMaterialShaderMapId::AppendKeyString(FString& KeyString) const
 {
 	check(IsContentValid());
@@ -1026,18 +1042,8 @@ void FMaterialShaderMapId::AppendKeyString(FString& KeyString) const
 
 	LayoutParams.AppendKeyString(KeyString);
 
-	for (const FStaticSwitchParameter& StaticSwitchParameter : StaticSwitchParameters)
-	{
-		StaticSwitchParameter.AppendKeyString(KeyString);
-	}
-	for (const FStaticComponentMaskParameter& StaticComponentMaskParameter : StaticComponentMaskParameters)
-	{
-		StaticComponentMaskParameter.AppendKeyString(KeyString);
-	}
-	for (const FStaticTerrainLayerWeightParameter& StaticTerrainLayerWeightParameter : TerrainLayerWeightParameters)
-	{
-		StaticTerrainLayerWeightParameter.AppendKeyString(KeyString);
-	}
+	AppendStaticParametersString(KeyString);
+
 	if (MaterialLayersId)
 	{
 		MaterialLayersId->AppendKeyString(KeyString);
diff --git a/Engine/Source/Runtime/Engine/Private/Materials/MaterialShared.cpp b/Engine/Source/Runtime/Engine/Private/Materials/MaterialShared.cpp
--- a/Engine/Source/Runtime/Engine/Private/Materials/MaterialShared.cpp
+++ b/Engine/Source/Runtime/Engine/Private/Materials/MaterialShared.cpp
@@ -2699,7 +2699,7 @@ bool FMaterial::CacheShaders(const FMaterialShaderMapId& ShaderMapId, EShaderPla
 FString FMaterial::GetUniqueAssetName(EShaderPlatform Platform, const FMaterialShaderMapId& ShaderMapId) const
 {
 	FString ToHash = GetAssetPath().GetPlainNameString();
-	ShaderMapId.AppendKeyString(ToHash);
+	ShaderMapId.AppendStaticParametersString(ToHash);
 	uint64 Hash = CityHash64((const char*)*ToHash, ToHash.Len());
 	return FString::Printf(TEXT("%s_%llx"), *GetAssetName(), Hash);
 }
diff --git a/Engine/Source/Runtime/Engine/Public/MaterialShared.h b/Engine/Source/Runtime/Engine/Public/MaterialShared.h
--- a/Engine/Source/Runtime/Engine/Public/MaterialShared.h
+++ b/Engine/Source/Runtime/Engine/Public/MaterialShared.h
@@ -984,6 +984,7 @@ public:
 
 	/** Appends string representations of this Id to a key string. */
 	void AppendKeyString(FString& KeyString) const;
+	void AppendStaticParametersString(FString& ParamsString) const;
 
 	const TArray<FStaticSwitchParameter> &GetStaticSwitchParameters() const 					{ return StaticSwitchParameters; }
 	const TArray<FStaticComponentMaskParameter> &GetStaticComponentMaskParameters() const 		{ return StaticComponentMaskParameters; }
