Fixed subobject refcounting for Dormant replicators.
Renable subobject refcount feature that enables the tearoff/delete subobject feature.

diff --git a/Engine/Source/Runtime/Engine/Classes/Engine/NetworkObjectList.h b/Engine/Source/Runtime/Engine/Classes/Engine/NetworkObjectList.h
--- a/Engine/Source/Runtime/Engine/Classes/Engine/NetworkObjectList.h
+++ b/Engine/Source/Runtime/Engine/Classes/Engine/NetworkObjectList.h
@@ -11,12 +11,13 @@ class FArchive;
 #ifndef UE_REPLICATED_OBJECT_REFCOUNTING
 	// Allows every network actor to keep track of the number of channels an individual subobject was replicated. 
 	// Not needed on clients so its compiled out to reduce memory usage.
-	#define UE_REPLICATED_OBJECT_REFCOUNTING 0
+	#define UE_REPLICATED_OBJECT_REFCOUNTING WITH_SERVER_CODE
 #endif
 
 #ifndef DO_REPLICATED_OBJECT_CHANNELREF_CHECKS
 	// Run additional code to validate any errors in the SubObject channel refcounting at the cost of extra memory usage and cpu overhead.
-	#define DO_REPLICATED_OBJECT_CHANNELREF_CHECKS (UE_REPLICATED_OBJECT_REFCOUNTING && WITH_SERVER_CODE && !UE_BUILD_SHIPPING) // will be disabled in Test builds after a few days.
+	// Temp: will be disabled from Test builds after a few days.
+	#define DO_REPLICATED_OBJECT_CHANNELREF_CHECKS (UE_REPLICATED_OBJECT_REFCOUNTING && WITH_SERVER_CODE && !UE_BUILD_SHIPPING) 
 #endif
 
 /** Indicates the status of a replicated subobject replicated by actor channels. */
@@ -52,14 +53,25 @@ struct FNetworkObjectInfo
 	* @warning: internal net driver time, not related to WorldSettings.TimeSeconds */
 	double LastNetUpdateTimestamp;
 
+
+	/**
+	* Key definitions for TSet/TMap that works with invalidated weak pointers
+	*/
+	struct FNetConnectionKeyFuncs : BaseKeyFuncs<TWeakObjectPtr<UNetConnection>, const TWeakObjectPtr<UNetConnection>&, false>
+	{
+		static KeyInitType GetSetKey(const ElementInitType& Element)	{ return Element; }
+		static bool Matches(KeyInitType A, KeyInitType B)				{ return A.HasSameIndexAndSerialNumber(B); }
+		static uint32 GetKeyHash(KeyInitType Key)						{ return GetTypeHash(Key); }
+	};
+
 	/** List of connections that this actor is dormant on */
-	TSet<TWeakObjectPtr<UNetConnection>> DormantConnections;
+	TSet<TWeakObjectPtr<UNetConnection>, FNetConnectionKeyFuncs> DormantConnections;
 
 	/** A list of connections that this actor has recently been dormant on, but the actor doesn't have a channel open yet.
 	*  These need to be differentiated from actors that the client doesn't know about, but there's no explicit list for just those actors.
 	*  (this list will be very transient, with connections being moved off the DormantConnections list, onto this list, and then off once the actor has a channel again)
 	*/
-	TSet<TWeakObjectPtr<UNetConnection>> RecentlyDormantConnections;
+	TSet<TWeakObjectPtr<UNetConnection>, FNetConnectionKeyFuncs> RecentlyDormantConnections;
 
 	/** 
 	* Keeps track of the number of channels that have replicated a subobject.
@@ -77,8 +89,11 @@ struct FNetworkObjectInfo
 		ENetSubObjectStatus Status = ENetSubObjectStatus::Active;
 
 #if DO_REPLICATED_OBJECT_CHANNELREF_CHECKS
-		/** Debug array to keep track of every individual channel */
-		TArray<UObject*> RegisteredChannels;
+		/** 
+		* Debug array to keep track of every individual references. Helps to trigger an ensure early if something wrong is detected.
+		* Consists of UActorChannel for active replicators or UNetConnection's for dormant replicators
+		*/
+		TArray<UObject*> RegisteredOwners;
 #endif
 		inline bool operator==(const FSubObjectChannelReference& rhs) const { return SubObjectPtr.HasSameIndexAndSerialNumber(rhs.SubObjectPtr); }
 		inline bool operator==(const TWeakObjectPtr<UObject>& rhs) const { return SubObjectPtr.HasSameIndexAndSerialNumber(rhs); }
@@ -101,26 +116,16 @@ struct FNetworkObjectInfo
 		}
 	};
 
+	/** 
+	* Key definitions for TSet/TMap that works with invalided weak pointers 
+	*/
 	struct FSubObjectChannelRefKeyFuncs : BaseKeyFuncs<FSubObjectChannelReference, const TWeakObjectPtr<UObject>&, false>
 	{
-		static KeyInitType GetSetKey(ElementInitType Element)
-		{
-			return Element.SubObjectPtr;
-		}
-
-		static bool Matches(KeyInitType A, KeyInitType B)
-		{
-			return A.HasSameIndexAndSerialNumber(B);
-		}
-
-		static uint32 GetKeyHash(KeyInitType Key)
-		{
-			// Probably should use a better has like FObjectKey but can't access the internals of the WeakPtr.
-			return GetTypeHash(Key);
-		}
+		static KeyInitType GetSetKey(const ElementInitType& Element)	{ return Element.SubObjectPtr; }
+		static bool Matches(KeyInitType A, KeyInitType B)				{ return A.HasSameIndexAndSerialNumber(B); }
+		static uint32 GetKeyHash(KeyInitType Key)						{ return GetTypeHash(Key); }
 	};
 
-
 #if UE_REPLICATED_OBJECT_REFCOUNTING
 	/** The set of active replicated subobjects for an actor */
 	TSet<FSubObjectChannelReference, FSubObjectChannelRefKeyFuncs> ActiveSubObjectChannelReferences;
@@ -296,16 +301,30 @@ public:
 	* Called when a channel starts replicating a subobject for the first time.
 	* Used to keep track of the number of channels having an active reference to a specific subobject
 	*/
-	void AddSubObjectChannelReference(AActor* OwnerActor, UObject* ReplicatedSubObject, UActorChannel* ReferencingChannel);
+	void AddSubObjectChannelReference(AActor* OwnerActor, UObject* ReplicatedSubObject, UObject* ReferenceOwner);
 
 	/** Called when a channel stops replicating a subobject. */
-	void RemoveSubObjectChannelReference(AActor* OwnerActor, const TWeakObjectPtr<UObject>& ReplicatedSubObject, UActorChannel* ReferencingChannel);
+	void RemoveSubObjectChannelReference(AActor* OwnerActor, const TWeakObjectPtr<UObject>& ReplicatedSubObject, UObject* ReferenceOwner);
+
+	/** Called when multiple subobjects need to remove their reference from either the active or inactive list */
+	void RemoveMultipleSubObjectChannelReference(AActor* OwnerActor, const TArray< TWeakObjectPtr<UObject>, TInlineAllocator<16> >& SubObjectsToRemove, UObject* ReferenceOwner);
 
 	/** Called when multiple subobjects that were flagged torn off or delete have removed their channel reference */
-	void RemoveMultipleInvalidSubObjectChannelReference(FNetworkObjectInfo* ActorNetInfo, const TArray< TWeakObjectPtr<UObject>, TInlineAllocator<16> >& SubObjectsToRemove, UActorChannel* ReferencingChannel);
+	void RemoveMultipleInvalidSubObjectChannelReference(FNetworkObjectInfo* ActorNetInfo, const TArray< TWeakObjectPtr<UObject>, TInlineAllocator<16> >& SubObjectsToRemove, UObject* ReferenceOwner);
 
 	/** Called when multiple subobjects that were still considered active have removed their channel reference */
-	void RemoveMultipleActiveSubObjectChannelReference(FNetworkObjectInfo* ActorNetInfo, const TArray< TWeakObjectPtr<UObject>, TInlineAllocator<16> >& SubObjectsToRemove, UActorChannel* ReferencingChannel);
+	void RemoveMultipleActiveSubObjectChannelReference(FNetworkObjectInfo* ActorNetInfo, const TArray< TWeakObjectPtr<UObject>, TInlineAllocator<16> >& SubObjectsToRemove, UObject* ReferenceOwner);
+
+	/** 
+	* Keep track of the transfer of ownership from the channel to the connection when the actor becomes dormant.
+	* This is only needed to ensure the debug reference tracking is up to date.
+	* No actual logic is modified here.
+	*/
+#if DO_REPLICATED_OBJECT_CHANNELREF_CHECKS
+	void SwapMultipleReferencesForDormancy(AActor* OwnerActor, const TArray< TWeakObjectPtr<UObject>, TInlineAllocator<16> >& SubObjectsToSwap, UActorChannel* PreviousChannelRefOwner, UNetConnection* NewConnectionRefOwner);
+	void SwapReferenceForDormancy(AActor* OwnerActor, UObject* ReplicatedSubObject, UNetConnection* PreviousConnectionRefOwner, UActorChannel* NewChannelRefOwner);
+#endif 
+
 #endif //#if UE_REPLICATED_OBJECT_REFCOUNTING
 
 private:
@@ -314,8 +333,14 @@ private:
 #if UE_REPLICATED_OBJECT_REFCOUNTING
 	void InvalidateSubObject(AActor* Actor, UObject* SubObject, ENetSubObjectStatus InvalidStatus);
 
-	bool HandleRemoveActiveSubObjectRef(FNetworkObjectInfo* ActorNetInfo, AActor* OwnerActor, const TWeakObjectPtr<UObject>& ReplicatedSubObject, UActorChannel* ReferencingChannel);
-	bool HandleRemoveInvalidSubObjectRef(FNetworkObjectInfo* ActorNetInfo, AActor* OwnerActor, const TWeakObjectPtr<UObject>& ReplicatedSubObject, UActorChannel* ReferencingChannel);
+	void HandleRemoveAnySubObjectChannelRef(FNetworkObjectInfo* ActorNetInfo, const TWeakObjectPtr<UObject>& ReplicatedSubObject, UObject* ReferenceOwner);
+	bool HandleRemoveActiveSubObjectRef(FNetworkObjectInfo* ActorNetInfo, const TWeakObjectPtr<UObject>& ReplicatedSubObject, UObject* ReferenceOwner);
+	bool HandleRemoveInvalidSubObjectRef(FNetworkObjectInfo* ActorNetInfo, const TWeakObjectPtr<UObject>& ReplicatedSubObject, UObject* ReferenceOwner);
+#endif
+
+
+#if DO_REPLICATED_OBJECT_CHANNELREF_CHECKS
+	void HandleSwapReferenceForDormancy(FNetworkObjectInfo* ActorNetInfo, const TWeakObjectPtr<UObject>& SubObjectPtr, UObject* PreviousRefOwner, UObject* NewRefOwner);
 #endif
 
 	FNetworkObjectSet AllNetworkObjects;
diff --git a/Engine/Source/Runtime/Engine/Private/DataChannel.cpp b/Engine/Source/Runtime/Engine/Private/DataChannel.cpp
--- a/Engine/Source/Runtime/Engine/Private/DataChannel.cpp
+++ b/Engine/Source/Runtime/Engine/Private/DataChannel.cpp
@@ -2254,6 +2254,16 @@ int64 UActorChannel::Close(EChannelCloseReason Reason)
 
 void UActorChannel::CleanupReplicators(const bool bKeepReplicators)
 {
+#if UE_REPLICATED_OBJECT_REFCOUNTING
+	TArray<TWeakObjectPtr<UObject>, TInlineAllocator<16>> ReferencesToRemove;
+#endif
+
+#if DO_REPLICATED_OBJECT_CHANNELREF_CHECKS
+    // The reference swap is only needed when we keep track of individual references since the refcount stays identical.
+    // So we only keep track of swapped objects when we run with the checks enabled.
+	TArray<TWeakObjectPtr<UObject>, TInlineAllocator<16>> ReferencesToSwap;
+#endif
+
 	// Cleanup or save replicators
 	for (auto CompIt = ReplicationMap.CreateIterator(); CompIt; ++CompIt)
 	{
@@ -2265,6 +2275,14 @@ void UActorChannel::CleanupReplicators(const bool bKeepReplicators)
 		{
 			LLM_SCOPE_BYTAG(NetConnection);
 
+#if DO_REPLICATED_OBJECT_CHANNELREF_CHECKS
+			// Ignore the main actor's replicator
+			if (Actor != ObjectReplicatorRef->ObjectPtr)
+			{
+				ReferencesToSwap.Add(ObjectReplicatorRef->GetWeakObjectPtr());
+			}
+#endif
+
 			// If we want to keep the replication state of the actor/sub-objects around, transfer ownership to the connection
 			// This way, if this actor opens another channel on this connection, we can reclaim or use this replicator to compare state, etc.
 			// For example, we may want to see if any state changed since the actor went dormant, and is now active again. 
@@ -2284,13 +2302,27 @@ void UActorChannel::CleanupReplicators(const bool bKeepReplicators)
 			// Ignore the main actor's replicator
 			if (Actor != ObjectReplicatorRef->ObjectPtr)
 			{
-				Connection->Driver->GetNetworkObjectList().RemoveSubObjectChannelReference(Actor, ObjectReplicatorRef->GetWeakObjectPtr(), this);
+				ReferencesToRemove.Add(ObjectReplicatorRef->GetWeakObjectPtr());
 			}
 #endif
 			ObjectReplicatorRef->CleanUp();
 		}
 	}
 
+#if UE_REPLICATED_OBJECT_REFCOUNTING
+	if (ReferencesToRemove.Num() > 0)
+	{
+		Connection->Driver->GetNetworkObjectList().RemoveMultipleSubObjectChannelReference(Actor, ReferencesToRemove, this);
+	}
+#endif
+
+#if DO_REPLICATED_OBJECT_CHANNELREF_CHECKS
+	if (ReferencesToSwap.Num() > 0)
+	{
+		Connection->Driver->GetNetworkObjectList().SwapMultipleReferencesForDormancy(Actor, ReferencesToSwap, this, Connection);
+	}
+#endif
+
 	ReplicationMap.Empty();
 
 	ActorReplicator.Reset();
@@ -4952,6 +4984,20 @@ TSharedRef<FObjectReplicator>& UActorChannel::CreateReplicator(UObject* Obj, boo
 	if (bCheckDormantReplicators)
 	{
 		NewReplicator = Connection->FindAndRemoveDormantReplicator(Actor, Obj);
+
+#if DO_REPLICATED_OBJECT_CHANNELREF_CHECKS
+		if (NewReplicator.IsValid())
+		{
+			// If this is a subobject (and not the main actor) swap the reference from the NetConnection to this channel
+			if (Obj != Actor)
+			{
+				if (UNetDriver* NetDriver = Connection->GetDriver())
+				{
+					NetDriver->GetNetworkObjectList().SwapReferenceForDormancy(Actor, Obj, Connection, this);
+				}
+			}
+		}
+#endif
 	}
 	else
 	{
@@ -4970,21 +5016,23 @@ TSharedRef<FObjectReplicator>& UActorChannel::CreateReplicator(UObject* Obj, boo
 		NewReplicator = Connection->CreateReplicatorForNewActorChannel(Obj);
 
 		UE_LOG(LogNetTraffic, Log, TEXT("CreateReplicator - creating new replicator for (0x%p) %s - Replicator (0x%p)"), Obj, *Obj->GetName(), NewReplicator.Get());
-	}
-
-	// Add to the replication map
-	TSharedRef<FObjectReplicator>& NewRef = ReplicationMap.Add(Obj, NewReplicator.ToSharedRef());
 
 #if UE_REPLICATED_OBJECT_REFCOUNTING
-	// If this is a subobject (and not the main actor) add a reference
-	if (Obj != Actor)
-	{
-		if (UNetDriver* NetDriver = Connection->GetDriver())
+		// If this is a subobject (and not the main actor) add a reference
+		if (Obj != Actor)
 		{
-			NetDriver->GetNetworkObjectList().AddSubObjectChannelReference(Actor, Obj, this);
+			if (UNetDriver* NetDriver = Connection->GetDriver())
+			{
+				NetDriver->GetNetworkObjectList().AddSubObjectChannelReference(Actor, Obj, this);
+			}
 		}
-	}
 #endif
+	}
+
+	// Add to the replication map
+	TSharedRef<FObjectReplicator>& NewRef = ReplicationMap.Add(Obj, NewReplicator.ToSharedRef());
+
+
 
 	// Start replicating with this replicator
 	NewRef->StartReplicating(this);
diff --git a/Engine/Source/Runtime/Engine/Private/Net/NetDormantHolder.cpp b/Engine/Source/Runtime/Engine/Private/Net/NetDormantHolder.cpp
--- a/Engine/Source/Runtime/Engine/Private/Net/NetDormantHolder.cpp
+++ b/Engine/Source/Runtime/Engine/Private/Net/NetDormantHolder.cpp
@@ -3,6 +3,7 @@
 #include "Net/NetDormantHolder.h"
 
 #include "Net/DataReplication.h"
+#include "Engine/NetworkObjectList.h"
 
 namespace UE::Net::Private
 {
@@ -138,6 +139,46 @@ void FDormantReplicatorHolder::CleanupStaleObjects()
 	}
 }
 
+void FDormantReplicatorHolder::CleanupStaleObjects(FNetworkObjectList& NetworkObjectList, UObject* ReferenceOwner)
+{
+#if UE_REPLICATED_OBJECT_REFCOUNTING
+	TArray<TWeakObjectPtr<UObject>, TInlineAllocator<16>> CleanedUpObjects;
+
+	for (FActorReplicatorSet::TIterator ActorSetIt = ActorReplicatorSet.CreateIterator(); ActorSetIt; ++ActorSetIt)
+	{
+		for (FActorDormantReplicators::FObjectReplicatorSet::TIterator ReplicatorSetIt = ActorSetIt->DormantReplicators.CreateIterator(); ReplicatorSetIt; ++ReplicatorSetIt)
+		{
+			FDormantObjectReplicator& DormantReplicator = *ReplicatorSetIt;
+			TWeakObjectPtr<UObject> DormantObjectPtr = DormantReplicator.Replicator->GetWeakObjectPtr();
+
+			if (!DormantObjectPtr.IsValid())
+			{
+				// If it's a subobject
+				if (ActorSetIt->OwnerActor != DormantReplicator.Replicator->GetObject())
+				{
+					CleanedUpObjects.Add(DormantObjectPtr);
+				}
+
+				ReplicatorSetIt.RemoveCurrent();
+			}
+		}
+
+		if (ActorSetIt->DormantReplicators.IsEmpty())
+		{
+			ActorSetIt.RemoveCurrent();
+		}
+
+		if (CleanedUpObjects.Num() > 0)
+		{
+			NetworkObjectList.RemoveMultipleSubObjectChannelReference(ActorSetIt->OwnerActor, CleanedUpObjects, ReferenceOwner);
+			CleanedUpObjects.Reset();
+		}
+	}
+#else
+	CleanupStaleObjects();
+#endif
+}
+
 void FDormantReplicatorHolder::ForEachDormantReplicator(UE::Net::FExecuteForEachDormantReplicator Function)
 {
 	for (const FActorDormantReplicators& ActorReplicators : ActorReplicatorSet)
diff --git a/Engine/Source/Runtime/Engine/Private/NetConnection.cpp b/Engine/Source/Runtime/Engine/Private/NetConnection.cpp
--- a/Engine/Source/Runtime/Engine/Private/NetConnection.cpp
+++ b/Engine/Source/Runtime/Engine/Private/NetConnection.cpp
@@ -4840,6 +4840,13 @@ void UNetConnection::FlushDormancyForObject(AActor* DormantActor, UObject* Repli
 		constexpr bool bUseDefaultState = false; 
 		ObjectReplicatorRef->InitWithObject(ReplicatedObject, this, bUseDefaultState);
 
+#if UE_REPLICATED_OBJECT_REFCOUNTING
+		if (Driver && DormantActor != ReplicatedObject)
+		{
+			Driver->GetNetworkObjectList().AddSubObjectChannelReference(DormantActor, ReplicatedObject, this);
+		}
+#endif
+
 		// Flush the must be mapped GUIDs, the initialization may add them, but they're phantom and will be remapped when actually sending
 		if (UPackageMapClient* PackageMapClient = CastChecked<UPackageMapClient>(PackageMap))
 		{
@@ -5057,6 +5064,17 @@ void UNetConnection::RemoveDormantReplicator(AActor* Actor, UObject* Object)
 	const FObjectKey ObjectKey = Object;
 
 	DormantReplicatorSet.RemoveStoredReplicator(Actor, ObjectKey);
+
+#if UE_REPLICATED_OBJECT_REFCOUNTING
+	// If this is a subobject (and not the main actor) remove the reference
+	if (Object != Actor)
+	{
+		if (Driver)
+		{
+			Driver->GetNetworkObjectList().RemoveSubObjectChannelReference(Actor, Object, this);
+		}
+	}
+#endif
 }
 
 void UNetConnection::ExecuteOnAllDormantReplicators(UE::Net::FExecuteForEachDormantReplicator ExecuteFunction)
@@ -5073,13 +5091,42 @@ void UNetConnection::CleanupDormantReplicatorsForActor(AActor* Actor)
 {
 	if (Actor)
 	{
+#if UE_REPLICATED_OBJECT_REFCOUNTING
+		if (Driver)
+		{
+			TArray<TWeakObjectPtr<UObject>, TInlineAllocator<16>> RemovedObjects;
+			UE::Net::FExecuteForEachDormantReplicator ExecuteFunction = [&RemovedObjects](AActor* OwnerActor, FObjectKey ObjectKey, const TSharedRef<FObjectReplicator>& ReplicatorRef)
+			{
+				if (OwnerActor != ReplicatorRef->GetObject())
+				{
+					RemovedObjects.Add(ReplicatorRef->GetWeakObjectPtr());
+				}
+			};
+
+			DormantReplicatorSet.ForEachDormantReplicatorOfActor(Actor, ExecuteFunction);
+
+			Driver->GetNetworkObjectList().RemoveMultipleSubObjectChannelReference(Actor, RemovedObjects, this);
+		}
+#endif
+
 		DormantReplicatorSet.CleanupAllReplicatorsOfActor(Actor);
 	}
 }
 
 void UNetConnection::CleanupStaleDormantReplicators()
 {
+#if UE_REPLICATED_OBJECT_REFCOUNTING
+	if (Driver)	
+	{
+		DormantReplicatorSet.CleanupStaleObjects(Driver->GetNetworkObjectList(), this);
+	}
+	else
+	{
+		DormantReplicatorSet.CleanupStaleObjects();
+	}
+#else
 	DormantReplicatorSet.CleanupStaleObjects();
+#endif
 }
 
 void UNetConnection::SetPendingCloseDueToReplicationFailure()
diff --git a/Engine/Source/Runtime/Engine/Private/NetworkObjectList.cpp b/Engine/Source/Runtime/Engine/Private/NetworkObjectList.cpp
--- a/Engine/Source/Runtime/Engine/Private/NetworkObjectList.cpp
+++ b/Engine/Source/Runtime/Engine/Private/NetworkObjectList.cpp
@@ -591,110 +591,136 @@ void FNetworkObjectList::InvalidateSubObject(AActor* Actor, UObject* SubObject,
 	}
 }
 
-void FNetworkObjectList::AddSubObjectChannelReference(AActor* OwnerActor, UObject* ReplicatedSubObject, UActorChannel* ReferencingChannel)
+void FNetworkObjectList::AddSubObjectChannelReference(AActor* OwnerActor, UObject* ReplicatedSubObject, UObject* ReferenceOwner)
 {
 	TSharedPtr<FNetworkObjectInfo>* InfoPtr = AllNetworkObjects.Find(OwnerActor);
 
-	if (ensureMsgf(InfoPtr, TEXT("The channel %s replicated %s (0x%p) without having any network objectinfo for %s"), *ReferencingChannel->Describe(), *GetNameSafe(ReplicatedSubObject), ReplicatedSubObject, *GetNameSafe(OwnerActor)))
+	if (ensureMsgf(InfoPtr, TEXT("The owner (0x%p) %s replicated %s (0x%p) without having any network objectinfo for %s"), ReferenceOwner, *GetNameSafe(ReferenceOwner), *GetNameSafe(ReplicatedSubObject), ReplicatedSubObject, *OwnerActor->GetName()))
 	{
 		FNetworkObjectInfo* ActorNetInfo = InfoPtr->Get();
-
 		const TWeakObjectPtr<UObject> SubObjectPtr = ReplicatedSubObject;
 
 #if DO_REPLICATED_OBJECT_CHANNELREF_CHECKS
 		// Make sure the object wasn't previously put in the invalid list  
 		if (FNetworkObjectInfo::FSubObjectChannelReference* DestroyedObjectRef = ActorNetInfo->InvalidSubObjectChannelReferences.FindByKey(SubObjectPtr))
 		{
-			ensureMsgf(DestroyedObjectRef == nullptr, TEXT("SubObject %s (0x%p) owned by %s was replicated by channel (0x%p) %s after it was set to %s"), *GetNameSafe(ReplicatedSubObject), ReplicatedSubObject, *GetNameSafe(OwnerActor), ReferencingChannel , *ReferencingChannel->Describe(), UE::Net::Private::LexToString(DestroyedObjectRef->Status));
+			ensureMsgf(DestroyedObjectRef == nullptr, TEXT("SubObject %s (0x%p) owned by %s was replicated by connection (0x%p) %s after it was set to (0x%p)"), *GetNameSafe(ReplicatedSubObject), ReplicatedSubObject, *GetNameSafe(OwnerActor), ReferenceOwner, *GetNameSafe(ReferenceOwner), UE::Net::Private::LexToString(DestroyedObjectRef->Status));
 		}
 #endif
 
 		if (FNetworkObjectInfo::FSubObjectChannelReference* ChannelReferences = ActorNetInfo->ActiveSubObjectChannelReferences.Find(SubObjectPtr))
 		{
 #if DO_REPLICATED_OBJECT_CHANNELREF_CHECKS
-			checkf(!ChannelReferences->RegisteredChannels.Contains(ReferencingChannel), TEXT("SubObject %s (0x%p) owned by %s was already referenced by channel %s"), *GetNameSafe(ReplicatedSubObject), ReplicatedSubObject, *GetNameSafe(OwnerActor), *ReferencingChannel->Describe());
-			ChannelReferences->RegisteredChannels.Add(ReferencingChannel);
+			ensureMsgf(!ChannelReferences->RegisteredOwners.Contains(ReferenceOwner), TEXT("SubObject %s (0x%p) owned by %s was already referenced by (0x%p) %s"), *GetNameSafe(ReplicatedSubObject), ReplicatedSubObject, *GetNameSafe(OwnerActor), ReferenceOwner, *GetNameSafe(ReferenceOwner));
+			ChannelReferences->RegisteredOwners.Add(ReferenceOwner);
 #endif
 
 			ChannelReferences->ChannelRefCount++;
 			check(ChannelReferences->ChannelRefCount < MAX_uint16);
 
-			UE_LOG(LogNetSubObject, Verbose, TEXT("Adding ChannelRef (%d) to %s (0x%p) owned by %s for channel %s"), ChannelReferences->ChannelRefCount, *GetNameSafe(ReplicatedSubObject), ReplicatedSubObject,  *GetNameSafe(OwnerActor), *ReferencingChannel->Describe());
+			UE_LOG(LogNetSubObject, Verbose, TEXT("Adding ChannelRef (%d) to %s (0x%p) owned by %s for (0x%p) %s"), ChannelReferences->ChannelRefCount, *GetNameSafe(ReplicatedSubObject), ReplicatedSubObject,  *GetNameSafe(OwnerActor), ReferenceOwner, *GetNameSafe(ReferenceOwner));
 		}
 		else
 		{
 			bool bWasSet = false;
 			FSetElementId ElementId = ActorNetInfo->ActiveSubObjectChannelReferences.Emplace(FNetworkObjectInfo::FSubObjectChannelReference(SubObjectPtr), &bWasSet);
 
-			UE_LOG(LogNetSubObject, Verbose, TEXT("Adding ChannelRef (new) to %s (0x%p) owned by %s for channel %s"), *GetNameSafe(ReplicatedSubObject), ReplicatedSubObject, *GetNameSafe(OwnerActor), *ReferencingChannel->Describe());
+			UE_LOG(LogNetSubObject, Verbose, TEXT("Adding ChannelRef (new) to %s (0x%p) owned by %s for (0x%p) %s"), *GetNameSafe(ReplicatedSubObject), ReplicatedSubObject, *GetNameSafe(OwnerActor), ReferenceOwner, *GetNameSafe(ReferenceOwner));
 
 #if DO_REPLICATED_OBJECT_CHANNELREF_CHECKS
-			ActorNetInfo->ActiveSubObjectChannelReferences[ElementId].RegisteredChannels.Add(ReferencingChannel);
+			ActorNetInfo->ActiveSubObjectChannelReferences[ElementId].RegisteredOwners.Add(ReferenceOwner);
 #endif
 		}
 	}
 }
 
-void FNetworkObjectList::RemoveMultipleInvalidSubObjectChannelReference(FNetworkObjectInfo* ActorNetInfo, const TArray< TWeakObjectPtr<UObject>, TInlineAllocator<16> >& SubObjectsToRemove, UActorChannel* ReferencingChannel)
+void FNetworkObjectList::RemoveMultipleSubObjectChannelReference(AActor* OwnerActor, const TArray< TWeakObjectPtr<UObject>, TInlineAllocator<16> >& SubObjectsToRemove, UObject* ReferenceOwner)
+{
+	if (TSharedPtr<FNetworkObjectInfo>* InfoPtr = AllNetworkObjects.Find(OwnerActor))
+	{
+		check(InfoPtr->IsValid());
+
+		FNetworkObjectInfo* ActorNetInfo = InfoPtr->Get();
+
+		for (const TWeakObjectPtr<UObject>& SubObjectPtr : SubObjectsToRemove)
+		{
+			HandleRemoveAnySubObjectChannelRef(ActorNetInfo, SubObjectPtr, ReferenceOwner);
+		}
+	}
+}
+
+void FNetworkObjectList::RemoveMultipleInvalidSubObjectChannelReference(FNetworkObjectInfo* ActorNetInfo, const TArray< TWeakObjectPtr<UObject>, TInlineAllocator<16> >& SubObjectsToRemove, UObject* ReferenceOwner)
 {
-	for (const TWeakObjectPtr<UObject>& ReplicatedSubObject : SubObjectsToRemove)
+	check(ActorNetInfo);
+
+	for (const TWeakObjectPtr<UObject>& SubObjectPtr : SubObjectsToRemove)
 	{
-		HandleRemoveInvalidSubObjectRef(ActorNetInfo, ActorNetInfo->Actor, ReplicatedSubObject, ReferencingChannel);
+		HandleRemoveInvalidSubObjectRef(ActorNetInfo, SubObjectPtr, ReferenceOwner);
 	}
 }
 
-void FNetworkObjectList::RemoveMultipleActiveSubObjectChannelReference(FNetworkObjectInfo* ActorNetInfo, const TArray< TWeakObjectPtr<UObject>, TInlineAllocator<16> >& SubObjectsToRemove, UActorChannel* ReferencingChannel)
+void FNetworkObjectList::RemoveMultipleActiveSubObjectChannelReference(FNetworkObjectInfo* ActorNetInfo, const TArray< TWeakObjectPtr<UObject>, TInlineAllocator<16> >& SubObjectsToRemove, UObject* ReferenceOwner)
 {
-	for (const TWeakObjectPtr<UObject>& ReplicatedSubObject : SubObjectsToRemove)
+	check(ActorNetInfo);
+
+	for (const TWeakObjectPtr<UObject>& SubObjectPtr : SubObjectsToRemove)
 	{
-		bool bWasRemoved = HandleRemoveActiveSubObjectRef(ActorNetInfo, ActorNetInfo->Actor, ReplicatedSubObject, ReferencingChannel);
+		bool bWasRemoved = HandleRemoveActiveSubObjectRef(ActorNetInfo, SubObjectPtr, ReferenceOwner);
 
 		// If somehow the object wasn't in the active list, check if it was in the invalid list just to be sure.
 		// It's possible for an object set to be deleted to have it's uobject ptr become innacessible before we replicate it's owner.
 		// Calling this here would act as a safeguard to ensure its reference is removed even if the tear off/force delete command was never sent.
 		if (!bWasRemoved)
 		{
-			HandleRemoveInvalidSubObjectRef(ActorNetInfo, ActorNetInfo->Actor, ReplicatedSubObject, ReferencingChannel);
+			HandleRemoveInvalidSubObjectRef(ActorNetInfo, SubObjectPtr, ReferenceOwner);
 		}
 	}
 }
 
-void FNetworkObjectList::RemoveSubObjectChannelReference(AActor* OwnerActor, const TWeakObjectPtr<UObject>& ReplicatedSubObject, UActorChannel* ReferencingChannel)
+void FNetworkObjectList::RemoveSubObjectChannelReference(AActor* OwnerActor, const TWeakObjectPtr<UObject>& SubObjectPtr, UObject* ReferenceOwner)
 {
-	check(ReplicatedSubObject.IsExplicitlyNull() == false);
-
-	TSharedPtr<FNetworkObjectInfo>* InfoPtr = AllNetworkObjects.Find(OwnerActor);
-	if (InfoPtr)
+	if (TSharedPtr<FNetworkObjectInfo>* InfoPtr = AllNetworkObjects.Find(OwnerActor))
 	{
-		FNetworkObjectInfo* ActorNetInfo = InfoPtr->Get();
+		check(InfoPtr->IsValid());
+		HandleRemoveAnySubObjectChannelRef(InfoPtr->Get(), SubObjectPtr, ReferenceOwner);
+	}
+}
 
-		// Look in the active set
-		bool bWasRemoved = HandleRemoveActiveSubObjectRef(ActorNetInfo, OwnerActor, ReplicatedSubObject, ReferencingChannel);
+void FNetworkObjectList::HandleRemoveAnySubObjectChannelRef(FNetworkObjectInfo* ActorNetInfo, const TWeakObjectPtr<UObject>& SubObjectPtr, UObject* ReferenceOwner)
+{
+	check(ActorNetInfo);
+	check(SubObjectPtr.IsExplicitlyNull() == false);
 
-		// Look in the destroyed array
-		if (!bWasRemoved)
-		{
-			bWasRemoved = HandleRemoveInvalidSubObjectRef(ActorNetInfo, OwnerActor, ReplicatedSubObject, ReferencingChannel);
-		}
+	// Look in the active set
+	bool bWasRemoved = HandleRemoveActiveSubObjectRef(ActorNetInfo, SubObjectPtr, ReferenceOwner);
+
+	// Look in the destroyed array
+	if (!bWasRemoved)
+	{
+		bWasRemoved = HandleRemoveInvalidSubObjectRef(ActorNetInfo, SubObjectPtr, ReferenceOwner);
 	}
+
+	ensureMsgf(bWasRemoved, TEXT("HandleRemoveAnySubObjectChannelRef could not find any references for %s (0x%p) owned by %s"), *GetNameSafe(SubObjectPtr.GetEvenIfUnreachable()), SubObjectPtr.GetEvenIfUnreachable(), *GetNameSafe(ActorNetInfo->Actor));
 }
 
-bool FNetworkObjectList::HandleRemoveActiveSubObjectRef(FNetworkObjectInfo* ActorNetInfo, AActor* OwnerActor, const TWeakObjectPtr<UObject>& SubObjectPtr, UActorChannel* ReferencingChannel)
+bool FNetworkObjectList::HandleRemoveActiveSubObjectRef(FNetworkObjectInfo* ActorNetInfo, const TWeakObjectPtr<UObject>& SubObjectPtr, UObject* ReferenceOwner)
 {
+	check(ActorNetInfo);
+
 	FSetElementId FoundId = ActorNetInfo->ActiveSubObjectChannelReferences.FindId(SubObjectPtr);
 	if (FoundId.IsValidId())
 	{
 		FNetworkObjectInfo::FSubObjectChannelReference& ActiveReference = ActorNetInfo->ActiveSubObjectChannelReferences[FoundId];
 
-		check(ActiveReference.SubObjectPtr == SubObjectPtr);
+		check(ActiveReference.SubObjectPtr.HasSameIndexAndSerialNumber(SubObjectPtr));
 
 		ActiveReference.ChannelRefCount--;
 
-		UE_LOG(LogNetSubObject, Verbose, TEXT("Removed ACTIVE ChannelRef (%d) for %s (0x%p) owned by %s for channel %s"), ActiveReference.ChannelRefCount, *GetNameSafe(SubObjectPtr.GetEvenIfUnreachable()), SubObjectPtr.GetEvenIfUnreachable(), *GetNameSafe(OwnerActor), *ReferencingChannel->Describe());
+		UE_LOG(LogNetSubObject, Verbose, TEXT("Removed ACTIVE ChannelRef (%d) for %s (0x%p) owned by %s for connection (0x%p) %s"), ActiveReference.ChannelRefCount, *GetNameSafe(SubObjectPtr.GetEvenIfUnreachable()), SubObjectPtr.GetEvenIfUnreachable(), *GetNameSafe(ActorNetInfo->Actor), ReferenceOwner, *GetNameSafe(ReferenceOwner));
 
 #if DO_REPLICATED_OBJECT_CHANNELREF_CHECKS
-		int32 Removed = ActiveReference.RegisteredChannels.RemoveSingleSwap(ReferencingChannel);
-		checkf(Removed > 0, TEXT("Removed ACTIVE ref for Subobject %s (0x%p) owned by %s but it was never referenced by the channel %s"), *GetNameSafe(SubObjectPtr.GetEvenIfUnreachable()), SubObjectPtr.GetEvenIfUnreachable(), *GetNameSafe(OwnerActor), *ReferencingChannel->Describe());
+		int32 Removed = ActiveReference.RegisteredOwners.RemoveSingleSwap(ReferenceOwner);
+		ensureMsgf(Removed > 0, TEXT("Removed ACTIVE ref for Subobject %s (0x%p) owned by %s but it was never referenced by the connection (0x%p) %s"), *GetNameSafe(SubObjectPtr.GetEvenIfUnreachable()), SubObjectPtr.GetEvenIfUnreachable(), *GetNameSafe(ActorNetInfo->Actor), ReferenceOwner, *GetNameSafe(ReferenceOwner));
 #endif
 
 		if (ActiveReference.ChannelRefCount == 0)
@@ -704,27 +730,25 @@ bool FNetworkObjectList::HandleRemoveActiveSubObjectRef(FNetworkObjectInfo* Acto
 
 		return true;
 	}
-	else
-	{
-		check(false);
-	}
 
 	return false;
 }
 
-bool FNetworkObjectList::HandleRemoveInvalidSubObjectRef(FNetworkObjectInfo* ActorNetInfo, AActor* OwnerActor, const TWeakObjectPtr<UObject>& SubObjectPtr, UActorChannel* ReferencingChannel)
+bool FNetworkObjectList::HandleRemoveInvalidSubObjectRef(FNetworkObjectInfo* ActorNetInfo, const TWeakObjectPtr<UObject>& SubObjectPtr, UObject* ReferenceOwner)
 {
+	check(ActorNetInfo);
+
 	const int32 Index = ActorNetInfo->InvalidSubObjectChannelReferences.IndexOfByKey(SubObjectPtr);
 	if (Index != INDEX_NONE)
 	{
 		FNetworkObjectInfo::FSubObjectChannelReference& InvalidReference = ActorNetInfo->InvalidSubObjectChannelReferences[Index];
 		InvalidReference.ChannelRefCount--;
 
-		UE_LOG(LogNetSubObject, Verbose, TEXT("Removed INVALID ChannelRef (%d) for %s (0x%p) owned by %s for channel %s"), InvalidReference.ChannelRefCount,*GetNameSafe(SubObjectPtr.GetEvenIfUnreachable()), SubObjectPtr.GetEvenIfUnreachable(), *GetNameSafe(OwnerActor), *ReferencingChannel->Describe());
+		UE_LOG(LogNetSubObject, Verbose, TEXT("Removed INVALID ChannelRef (%d) for %s (0x%p) owned by %s for connection (0x%p) %s"), InvalidReference.ChannelRefCount,*GetNameSafe(SubObjectPtr.GetEvenIfUnreachable()), SubObjectPtr.GetEvenIfUnreachable(), *GetNameSafe(ActorNetInfo->Actor), ReferenceOwner, *GetNameSafe(ReferenceOwner));
 
 #if DO_REPLICATED_OBJECT_CHANNELREF_CHECKS
-		int32 Removed = InvalidReference.RegisteredChannels.RemoveSingleSwap(ReferencingChannel);
-		checkf(Removed > 0, TEXT("Removed INVALID ref for Subobject %s (0x%p) owned by %s but it was never referenced by the channel %s"), *GetNameSafe(SubObjectPtr.GetEvenIfUnreachable()), SubObjectPtr.GetEvenIfUnreachable(), *GetNameSafe(OwnerActor), *ReferencingChannel->Describe());
+		int32 Removed = InvalidReference.RegisteredOwners.RemoveSingleSwap(ReferenceOwner);
+		ensureMsgf(Removed > 0, TEXT("Removed INVALID ref for Subobject %s (0x%p) owned by %s but it was never referenced by the connection (0x%p) %s"), *GetNameSafe(SubObjectPtr.GetEvenIfUnreachable()), SubObjectPtr.GetEvenIfUnreachable(), *GetNameSafe(ActorNetInfo->Actor), ReferenceOwner, *GetNameSafe(ReferenceOwner));
 #endif
 		if (InvalidReference.ChannelRefCount == 0)
 		{
@@ -737,4 +761,61 @@ bool FNetworkObjectList::HandleRemoveInvalidSubObjectRef(FNetworkObjectInfo* Act
 	return false;
 }
 
+#if DO_REPLICATED_OBJECT_CHANNELREF_CHECKS
+void FNetworkObjectList::SwapMultipleReferencesForDormancy(AActor* OwnerActor, const TArray< TWeakObjectPtr<UObject>, TInlineAllocator<16> >& SubObjectsToSwap, UActorChannel* PreviousChannelRefOwner, UNetConnection* NewConnectionRefOwner)
+{
+	if (TSharedPtr<FNetworkObjectInfo>* InfoPtr = AllNetworkObjects.Find(OwnerActor))
+	{
+		check(InfoPtr->IsValid());
+
+		FNetworkObjectInfo* ActorNetInfo = InfoPtr->Get();
+
+		for (const TWeakObjectPtr<UObject>& SubObjectPtr : SubObjectsToSwap)
+		{
+			HandleSwapReferenceForDormancy(ActorNetInfo, SubObjectPtr, PreviousChannelRefOwner, NewConnectionRefOwner);
+		}
+	}
+}
+
+void FNetworkObjectList::SwapReferenceForDormancy(AActor* OwnerActor, UObject* ReplicatedSubObject, UNetConnection* PreviousConnectionRefOwner, UActorChannel* NewChannelRefOwner)
+{
+	if (TSharedPtr<FNetworkObjectInfo>* InfoPtr = AllNetworkObjects.Find(OwnerActor))
+	{
+		check(InfoPtr->IsValid());
+
+		const TWeakObjectPtr<UObject> SubObjectPtr = ReplicatedSubObject;
+		HandleSwapReferenceForDormancy(InfoPtr->Get(), SubObjectPtr, PreviousConnectionRefOwner, NewChannelRefOwner);
+	}
+}
+
+void FNetworkObjectList::HandleSwapReferenceForDormancy(FNetworkObjectInfo* ActorNetInfo, const TWeakObjectPtr<UObject>& SubObjectPtr, UObject* PreviousRefOwner, UObject* NewRefOwner)
+{
+	check(ActorNetInfo);
+
+	// First check the active list
+	FNetworkObjectInfo::FSubObjectChannelReference* SubObjectChannelRef = ActorNetInfo->ActiveSubObjectChannelReferences.Find(SubObjectPtr);
+	
+	if (!SubObjectChannelRef)
+	{
+		// Then check the inactive list
+		SubObjectChannelRef = ActorNetInfo->InvalidSubObjectChannelReferences.FindByKey(SubObjectPtr);
+	}
+
+	if (SubObjectChannelRef)
+	{
+		UE_LOG(LogNetSubObject, Verbose, TEXT("Swapped ChannelRef (%d) for SubObject %s (0x%p) owned by %s. From (0x%p) %s to (0x%p) %s"), SubObjectChannelRef->ChannelRefCount, *GetNameSafe(SubObjectPtr.GetEvenIfUnreachable()), SubObjectPtr.GetEvenIfUnreachable(), *GetNameSafe(ActorNetInfo->Actor), PreviousRefOwner, *GetNameSafe(PreviousRefOwner), NewRefOwner, *GetNameSafe(NewRefOwner));
+
+		int32 Removed = SubObjectChannelRef->RegisteredOwners.RemoveSingleSwap(PreviousRefOwner);
+		ensureMsgf(Removed > 0, TEXT("SwapReferencesForDormancy could not find reference to previous reference (0x%p) %s for subobject %s (0x%p) owned by %s. Swapping to (0x%p) %s"), PreviousRefOwner, *GetNameSafe(PreviousRefOwner), *GetNameSafe(SubObjectPtr.GetEvenIfUnreachable()), SubObjectPtr.GetEvenIfUnreachable(), *GetNameSafe(ActorNetInfo->Actor), NewRefOwner, *GetNameSafe(NewRefOwner));
+
+		ensureMsgf(!SubObjectChannelRef->RegisteredOwners.Contains(NewRefOwner), TEXT("SwapReferencesForDormancy found new reference (0x%p) %s was already registered to %s (0x%p) owned by %s"), NewRefOwner, *GetNameSafe(NewRefOwner), *GetNameSafe(SubObjectPtr.GetEvenIfUnreachable()), SubObjectPtr.GetEvenIfUnreachable(), *GetNameSafe(ActorNetInfo->Actor));
+		SubObjectChannelRef->RegisteredOwners.Add(NewRefOwner);
+	}
+	else
+	{
+		ensureMsgf(false, TEXT("SwapReferencesForDormancy could not find any references to %s (0x%p) owned by %s. Swapping from (0x%p) %s to (0x%p) %s"), *GetNameSafe(SubObjectPtr.GetEvenIfUnreachable()), SubObjectPtr.GetEvenIfUnreachable(), *GetNameSafe(ActorNetInfo->Actor), PreviousRefOwner, *GetNameSafe(PreviousRefOwner), NewRefOwner, *GetNameSafe(NewRefOwner));
+	}
+}
+#endif //#if DO_REPLICATED_OBJECT_CHANNELREF_CHECKS
+
 #endif //#if UE_REPLICATED_OBJECT_REFCOUNTING
\ No newline at end of file
diff --git a/Engine/Source/Runtime/Engine/Public/Net/NetDormantHolder.h b/Engine/Source/Runtime/Engine/Public/Net/NetDormantHolder.h
--- a/Engine/Source/Runtime/Engine/Public/Net/NetDormantHolder.h
+++ b/Engine/Source/Runtime/Engine/Public/Net/NetDormantHolder.h
@@ -8,6 +8,7 @@
 
 class AActor;
 class FObjectReplicator;
+class FNetworkObjectList;
 
 namespace UE::Net
 {
@@ -141,6 +142,15 @@ struct FDormantReplicatorHolder
 	*/
 	void CleanupStaleObjects();
 
+	/**
+	* Iterate over all the stored object replicators and destroy any that are tied to a replicated object that is now considered invalid.
+	* This version will decrement the subobject references stored in the networkobject list too.
+	* 
+	* @param FNetworkObjectList The netdriver's network object list
+	* @param ReferenceOwner The connection owner that is storing those dormant objects
+	*/
+	void CleanupStaleObjects(FNetworkObjectList& NetworkObjectList, UObject* ReferenceOwner);
+
 	/** 
 	* Execute the passed function on each dormant replicator we are holding
 	*/
