Make FActiveDeviceProperty a protected type of the UInputDeviceSubsystem because no external users should be attempting to use it.

#rnx
#pf 63890934303395f6c943ea1b

diff --git a/Engine/Plugins/Runtime/InputDebugging/Source/InputDebugging/Private/InputDeviceDebugTools.cpp b/Engine/Plugins/Runtime/InputDebugging/Source/InputDebugging/Private/InputDeviceDebugTools.cpp
--- a/Engine/Plugins/Runtime/InputDebugging/Source/InputDebugging/Private/InputDeviceDebugTools.cpp
+++ b/Engine/Plugins/Runtime/InputDebugging/Source/InputDebugging/Private/InputDeviceDebugTools.cpp
@@ -93,18 +93,18 @@ void FInputDeviceDebugTools::OnShowDebugDeviceProperties(UCanvas* Canvas)
 	}
 
 	// Sort the active properties by their owning Platform User
-	static TMap<FPlatformUserId, TArray<const FActiveDeviceProperty*>> PlatformUserDebugStrings;
+	static TMap<FPlatformUserId, TArray<const UInputDeviceSubsystem::FActiveDeviceProperty*>> PlatformUserDebugStrings;
 	{
 		PlatformUserDebugStrings.Empty();
 
-		for (const FActiveDeviceProperty& ActiveProperty : System->ActiveProperties)
+		for (const UInputDeviceSubsystem::FActiveDeviceProperty& ActiveProperty : System->ActiveProperties)
 		{
 			PlatformUserDebugStrings.FindOrAdd(ActiveProperty.PlatformUser).Emplace(&ActiveProperty);
 		}
 	}	
 
 	// Actually display the debug information	
-	for (TPair<FPlatformUserId, TArray<const FActiveDeviceProperty*>> Pair : PlatformUserDebugStrings)
+	for (TPair<FPlatformUserId, TArray<const UInputDeviceSubsystem::FActiveDeviceProperty*>> Pair : PlatformUserDebugStrings)
 	{
 		const FString PlatUserLabel = FString::Printf(TEXT("Platform User %d has %d active Device Properties"),
 			Pair.Key.GetInternalId(),
@@ -113,7 +113,7 @@ void FInputDeviceDebugTools::OnShowDebugDeviceProperties(UCanvas* Canvas)
 		DisplayDebugManager.SetDrawColor(FColor::Yellow);
 		DisplayDebugManager.DrawString(PlatUserLabel);
 
-		for (const FActiveDeviceProperty* ActiveProperty: Pair.Value)
+		for (const UInputDeviceSubsystem::FActiveDeviceProperty* ActiveProperty: Pair.Value)
 		{
 			// Draw the name of the property
 			DisplayDebugManager.SetDrawColor(FColor::White);
diff --git a/Engine/Source/Runtime/Engine/Classes/GameFramework/InputDeviceSubsystem.h b/Engine/Source/Runtime/Engine/Classes/GameFramework/InputDeviceSubsystem.h
--- a/Engine/Source/Runtime/Engine/Classes/GameFramework/InputDeviceSubsystem.h
+++ b/Engine/Source/Runtime/Engine/Classes/GameFramework/InputDeviceSubsystem.h
@@ -14,47 +14,6 @@ ENGINE_API DECLARE_LOG_CATEGORY_EXTERN(LogInputDeviceProperties, Log, All);
 class UInputDeviceProperty;
 class APlayerController;
 
-/** Contains a pointer to an active device property and keeps track of how long it has been evaluated for */
-USTRUCT()
-struct ENGINE_API FActiveDeviceProperty
-{
-	GENERATED_BODY()
-
-	/** Active properties can just use the hash of their FInputDevicePropertyHandle for a fast and unique lookup */
-	ENGINE_API friend uint32 GetTypeHash(const FActiveDeviceProperty& InProp);
-
-	bool operator==(const FActiveDeviceProperty& Other) const;
-	bool operator!=(const FActiveDeviceProperty& Other) const;
-	ENGINE_API friend bool operator==(const FActiveDeviceProperty& ActiveProp, const FInputDevicePropertyHandle& Handle);
-	ENGINE_API friend bool operator!=(const FActiveDeviceProperty& ActiveProp, const FInputDevicePropertyHandle& Handle);
-	
-	/** The active device property */
-	UPROPERTY(Transient)
-	TObjectPtr<UInputDeviceProperty> Property = nullptr;
-
-	/** How long this property has been evaluated for. DeltaTime is added to this on tick */
-	double EvaluatedDuration = 0.0;
-
-	/** The platform user that is actively receiving this device property */
-	FPlatformUserId PlatformUser = PLATFORMUSERID_NONE;
-
-	/** The handle of this active property. */
-	FInputDevicePropertyHandle PropertyHandle = FInputDevicePropertyHandle::InvalidHandle;
-
-	/**
-	* If true, then when the InputDeviceProperty is done being evaluated and has fulfilled its duration, then 
-	* it will have it's Reset function called. Set this to false if you want your device property to stay 
-	* in it's ending state. 
-	*/
-	bool bRemoveAfterEvaluationTime = true;
-
-	/** If true, then this device property will ignore dilated delta time and use the Applications delta time instead */
-	bool bIgnoreTimeDilation = false;
-
-	/** If true, then this device property will be played even if the game world is paused. */
-	bool bPlayWhilePaused = false;
-};
-
 /** Parameters for the UInputDeviceSubsystem::SetDeviceProperty function */
 USTRUCT(BlueprintType)
 struct ENGINE_API FSetDevicePropertyParams
@@ -110,6 +69,56 @@ class ENGINE_API UInputDeviceSubsystem : public UEngineSubsystem, public FTickab
 	
 	GENERATED_BODY()
 
+protected:
+
+	/** Contains a pointer to an active device property and keeps track of how long it has been evaluated for */
+	struct FActiveDeviceProperty
+	{		
+		/** Active properties can just use the hash of their FInputDevicePropertyHandle for a fast and unique lookup */
+		friend uint32 GetTypeHash(const UInputDeviceSubsystem::FActiveDeviceProperty& InProp)
+		{
+			return InProp.PropertyHandle.GetTypeHash();
+		}
+
+		friend bool operator==(const UInputDeviceSubsystem::FActiveDeviceProperty& ActiveProp, const FInputDevicePropertyHandle& Handle)
+		{
+			return ActiveProp.PropertyHandle == Handle;
+		}
+
+		friend bool operator!=(const UInputDeviceSubsystem::FActiveDeviceProperty& ActiveProp, const FInputDevicePropertyHandle& Handle)
+		{
+			return ActiveProp.PropertyHandle != Handle;
+		}
+
+		bool operator==(const UInputDeviceSubsystem::FActiveDeviceProperty& Other) const;
+		bool operator!=(const UInputDeviceSubsystem::FActiveDeviceProperty& Other) const;
+
+		/** The active device property */	
+		TObjectPtr<UInputDeviceProperty> Property = nullptr;
+
+		/** How long this property has been evaluated for. DeltaTime is added to this on tick */
+		double EvaluatedDuration = 0.0;
+
+		/** The platform user that is actively receiving this device property */
+		FPlatformUserId PlatformUser = PLATFORMUSERID_NONE;
+
+		/** The handle of this active property. */
+		FInputDevicePropertyHandle PropertyHandle = FInputDevicePropertyHandle::InvalidHandle;
+
+		/**
+		* If true, then when the InputDeviceProperty is done being evaluated and has fulfilled its duration, then
+		* it will have it's Reset function called. Set this to false if you want your device property to stay
+		* in it's ending state.
+		*/
+		bool bRemoveAfterEvaluationTime = true;
+
+		/** If true, then this device property will ignore dilated delta time and use the Applications delta time instead */
+		bool bIgnoreTimeDilation = false;
+
+		/** If true, then this device property will be played even if the game world is paused. */
+		bool bPlayWhilePaused = false;
+	};
+
 public:
 
 	static UInputDeviceSubsystem* Get();
@@ -205,8 +214,7 @@ protected:
 	/**
 	* Array of currently active input device properties that will be evaluated on tick
 	*/
-	UPROPERTY(Transient)
-	TSet<FActiveDeviceProperty> ActiveProperties;
+	TSet<UInputDeviceSubsystem::FActiveDeviceProperty> ActiveProperties;
 
 	/**
 	 * Set of property handles the properties that are currently pending manual removal.
diff --git a/Engine/Source/Runtime/Engine/Private/GameFramework/InputDeviceSubsystem.cpp b/Engine/Source/Runtime/Engine/Private/GameFramework/InputDeviceSubsystem.cpp
--- a/Engine/Source/Runtime/Engine/Private/GameFramework/InputDeviceSubsystem.cpp
+++ b/Engine/Source/Runtime/Engine/Private/GameFramework/InputDeviceSubsystem.cpp
@@ -96,33 +96,17 @@ public:
 ////////////////////////////////////////////////////////
 // FActiveDeviceProperty
 
-/** Active properties can just use the hash of their Property Device handle for a fast and unique lookup */
-uint32 GetTypeHash(const FActiveDeviceProperty& InProp)
-{
-	return InProp.PropertyHandle.GetTypeHash();
-}
-
 // The property handles are the only things that matter when comparing, they will always be unique.
-bool FActiveDeviceProperty::operator==(const FActiveDeviceProperty& Other) const
+bool UInputDeviceSubsystem::FActiveDeviceProperty::operator==(const UInputDeviceSubsystem::FActiveDeviceProperty& Other) const
 {
 	return PropertyHandle == Other.PropertyHandle;
 }
 
-bool FActiveDeviceProperty::operator!=(const FActiveDeviceProperty& Other) const
+bool UInputDeviceSubsystem::FActiveDeviceProperty::operator!=(const UInputDeviceSubsystem::FActiveDeviceProperty& Other) const
 {
 	return PropertyHandle != Other.PropertyHandle;
 }
 
-bool operator==(const FActiveDeviceProperty& ActiveProp, const FInputDevicePropertyHandle& Handle)
-{
-	return ActiveProp.PropertyHandle == Handle;
-}
-
-bool operator!=(const FActiveDeviceProperty& ActiveProp, const FInputDevicePropertyHandle& Handle)
-{
-	return ActiveProp.PropertyHandle != Handle;
-}
-
 ////////////////////////////////////////////////////////
 // FSetDevicePropertyParams
 
