Fix clang compilation error due to const warning

diff --git a/Engine/Source/Runtime/Engine/Classes/Engine/NetworkObjectList.h b/Engine/Source/Runtime/Engine/Classes/Engine/NetworkObjectList.h
--- a/Engine/Source/Runtime/Engine/Classes/Engine/NetworkObjectList.h
+++ b/Engine/Source/Runtime/Engine/Classes/Engine/NetworkObjectList.h
@@ -121,7 +121,7 @@ struct FNetworkObjectInfo
 	*/
 	struct FSubObjectChannelRefKeyFuncs : BaseKeyFuncs<FSubObjectChannelReference, const TWeakObjectPtr<UObject>&, false>
 	{
-		static KeyInitType GetSetKey(const ElementInitType& Element)	{ return Element.SubObjectPtr; }
+		static KeyInitType GetSetKey(ElementInitType& Element)			{ return Element.SubObjectPtr; }
 		static bool Matches(KeyInitType A, KeyInitType B)				{ return A.HasSameIndexAndSerialNumber(B); }
 		static uint32 GetKeyHash(KeyInitType Key)						{ return GetTypeHash(Key); }
 	};
