Do not reset light color for the static light color property. Add an option for if you want to reset the color to curve color properties. This lets you just do very simple "one shot" settings for device color, and then a more complex setup using curves if you want to do a smooth transition of colors or something like that.

#rnx

diff --git a/Engine/Source/Runtime/Engine/Classes/GameFramework/InputDeviceProperties.h b/Engine/Source/Runtime/Engine/Classes/GameFramework/InputDeviceProperties.h
--- a/Engine/Source/Runtime/Engine/Classes/GameFramework/InputDeviceProperties.h
+++ b/Engine/Source/Runtime/Engine/Classes/GameFramework/InputDeviceProperties.h
@@ -173,7 +173,8 @@ struct FDeviceColorData
 };
 
 /**
-* Set the color of an Input Device to a static color. 
+* Set the color of an Input Device to a static color. This will NOT reset the device color when the property
+* is done evaluating. You can think of this as a "One Shot" effect, where you set the device property color.
 * 
 * NOTE: This property has platform specific implementations and may behave differently per platform.
 * See the docs for more details on each platform.
@@ -186,7 +187,6 @@ class UColorInputDeviceProperty : public UInputDeviceProperty
 public:
 
 	virtual void EvaluateDeviceProperty_Implementation(const FPlatformUserId PlatformUser, const float DeltaTime, const float Duration) override;
-	virtual void ResetDeviceProperty_Implementation(const FPlatformUserId PlatformUser) override;
 	virtual FInputDeviceProperty* GetInternalDeviceProperty() override;
 
 	/** Default color data that will be used by default. Device Specific overrides will be used when the current input device matches */
@@ -217,6 +217,10 @@ struct FDeviceColorCurveData
 	UPROPERTY(EditAnywhere, BlueprintReadWrite, Category = "Color")
 	bool bEnable = true;
 
+	/** If true, the light color will be reset to "off" after the curve values are finished evaluating. */
+	UPROPERTY(EditAnywhere, BlueprintReadWrite, Category = "Color")
+	bool bResetAfterCompletion = true;
+
 	/** The color the device light should be */
 	UPROPERTY(EditAnywhere, BlueprintReadWrite, Category = "Color")
 	TObjectPtr<UCurveLinearColor> DeviceColorCurve;
diff --git a/Engine/Source/Runtime/Engine/Private/GameFramework/InputDeviceProperties.cpp b/Engine/Source/Runtime/Engine/Private/GameFramework/InputDeviceProperties.cpp
--- a/Engine/Source/Runtime/Engine/Private/GameFramework/InputDeviceProperties.cpp
+++ b/Engine/Source/Runtime/Engine/Private/GameFramework/InputDeviceProperties.cpp
@@ -85,13 +85,6 @@ void UColorInputDeviceProperty::EvaluateDeviceProperty_Implementation(const FPla
 	}
 }
 
-void UColorInputDeviceProperty::ResetDeviceProperty_Implementation(const FPlatformUserId PlatformUser)
-{
-	// Disabling the light will reset the color
-	InternalProperty.bEnable = false;
-	ApplyDeviceProperty(PlatformUser);
-}
-
 FInputDeviceProperty* UColorInputDeviceProperty::GetInternalDeviceProperty()
 {
 	return &InternalProperty;
@@ -128,9 +121,18 @@ void UColorInputDeviceCurveProperty::EvaluateDeviceProperty_Implementation(const
 
 void UColorInputDeviceCurveProperty::ResetDeviceProperty_Implementation(const FPlatformUserId PlatformUser)
 {
-	// Disabling the light will reset the color
-	InternalProperty.bEnable = false;
-	ApplyDeviceProperty(PlatformUser);
+	bool bReset = ColorData.bResetAfterCompletion;
+	if (const FDeviceColorCurveData* Data = GetDeviceSpecificData<FDeviceColorCurveData>(PlatformUser, DeviceOverrideData))
+	{
+		bReset = Data->bResetAfterCompletion;
+	}
+
+	if (bReset)
+	{
+		// Disabling the light will reset the color
+    	InternalProperty.bEnable = false;
+    	ApplyDeviceProperty(PlatformUser);
+	}	
 }
 
 FInputDeviceProperty* UColorInputDeviceCurveProperty::GetInternalDeviceProperty()
