Add const version of attenuation function

diff --git a/Engine/Source/Runtime/Engine/Classes/Sound/SoundCue.h b/Engine/Source/Runtime/Engine/Classes/Sound/SoundCue.h
--- a/Engine/Source/Runtime/Engine/Classes/Sound/SoundCue.h
+++ b/Engine/Source/Runtime/Engine/Classes/Sound/SoundCue.h
@@ -24,6 +24,7 @@ class UEdGraph;
 
 class USoundCue;
 class USoundNode;
+class USoundNodeAttenuation;
 struct FActiveSound;
 struct FSoundParseParameters;
 
@@ -293,10 +294,8 @@ public:
 	// This is used to cache the quality level if it has not been cached yet.
 	static void CacheQualityLevel();
 
-	/**
-	 * Instantiate certain functions to work around a linker issue
-	 */
-	void RecursiveFindAttenuation( USoundNode* Node, TArray<class USoundNodeAttenuation*> &OutNodes );
+	void RecursiveFindAttenuation(USoundNode* Node, TArray<USoundNodeAttenuation*> &OutNodes);
+	void RecursiveFindAttenuation(const USoundNode* Node, TArray<const USoundNodeAttenuation*>& OutNodes) const;
 
 #if WITH_EDITOR
 	/** Create the basic sound graph */
diff --git a/Engine/Source/Runtime/Engine/Private/SoundCue.cpp b/Engine/Source/Runtime/Engine/Private/SoundCue.cpp
--- a/Engine/Source/Runtime/Engine/Private/SoundCue.cpp
+++ b/Engine/Source/Runtime/Engine/Private/SoundCue.cpp
@@ -412,11 +412,16 @@ void USoundCue::PostEditChangeProperty(struct FPropertyChangedEvent& PropertyCha
 }
 #endif // WITH_EDITOR
 
-void USoundCue::RecursiveFindAttenuation( USoundNode* Node, TArray<class USoundNodeAttenuation*> &OutNodes )
+void USoundCue::RecursiveFindAttenuation(USoundNode* Node, TArray<USoundNodeAttenuation*> &OutNodes)
 {
 	RecursiveFindNode<USoundNodeAttenuation>( Node, OutNodes );
 }
 
+void USoundCue::RecursiveFindAttenuation(const USoundNode* Node, TArray<const USoundNodeAttenuation*>& OutNodes) const
+{
+	RecursiveFindNode<USoundNodeAttenuation>(Node, OutNodes);
+}
+
 void USoundCue::RecursiveFindAllNodes( USoundNode* Node, TArray<class USoundNode*> &OutNodes )
 {
 	if( Node )
