- Change "SetDeviceProperty" to a more aptly named "Activate Device Property".
- Create a BP version of Activate Device Property that takes in a subclass and will make a NewObject device prop
- Create a C++ version of ActivateDeviceProperty that will let you pass in your own DeviceProperty pointer. This is not BP exposed because I think it would be too easy to pass in the same instance of a property multiple times to designers.
- Rename "bRemoveAfterEvaluationTime" to just "bLooping". Now, when activate a device property as "looping" it will not be removed until you call "remove device property" on it. This makes more sense to use and reads better to me because thats what we do for force feedback.

#rnx

diff --git a/Engine/Plugins/Runtime/GameplayAbilities/Source/GameplayAbilities/Private/GameplayCueNotifyTypes.cpp b/Engine/Plugins/Runtime/GameplayAbilities/Source/GameplayAbilities/Private/GameplayCueNotifyTypes.cpp
--- a/Engine/Plugins/Runtime/GameplayAbilities/Source/GameplayAbilities/Private/GameplayCueNotifyTypes.cpp
+++ b/Engine/Plugins/Runtime/GameplayAbilities/Source/GameplayAbilities/Private/GameplayCueNotifyTypes.cpp
@@ -862,13 +862,12 @@ bool FGameplayCueNotify_InputDevicePropertyInfo::SetDeviceProperties(const FGame
 		{
 			FSetDevicePropertyParams Params = {};
 			Params.UserId = TargetPC->GetPlatformUserId();
-			Params.bRemoveAfterEvaluationTime = true;
+			
 			if (ensure(Params.UserId.IsValid()))
 			{
 				for (TSubclassOf<UInputDeviceProperty> PropClass : DeviceProperties)
 				{
-					Params.DevicePropertyClass = PropClass;
-					FInputDevicePropertyHandle Handle = System->SetDeviceProperty(Params);
+					FInputDevicePropertyHandle Handle = System->ActivateDevicePropertyOfClass(PropClass, Params);
 					ensure(Handle.IsValid());
 				}
 				return true;
diff --git a/Engine/Plugins/Runtime/InputDebugging/Source/InputDebugging/Private/InputDeviceDebugTools.cpp b/Engine/Plugins/Runtime/InputDebugging/Source/InputDebugging/Private/InputDeviceDebugTools.cpp
--- a/Engine/Plugins/Runtime/InputDebugging/Source/InputDebugging/Private/InputDeviceDebugTools.cpp
+++ b/Engine/Plugins/Runtime/InputDebugging/Source/InputDebugging/Private/InputDeviceDebugTools.cpp
@@ -122,8 +122,8 @@ void FInputDeviceDebugTools::OnShowDebugDeviceProperties(UCanvas* Canvas)
 			// Draw some information about the flags that have been set on the given property
 			DisplayDebugManager.SetDrawColor(FColor::Silver);
 			DisplayDebugManager.DrawString(FString::Printf(TEXT("\t Property Handle: %s"), *ActiveProperty->PropertyHandle.ToString()));
-			DisplayDebugManager.DrawString(FString::Printf(TEXT("\t Flags: \n \t \t \t bRemoveAfterEvaluationTime: %d \n \t \t \t bIgnoreTimeDilation: %d \n \t \t \t bPlayWhilePaused: %d"),
-				ActiveProperty->bRemoveAfterEvaluationTime, ActiveProperty->bIgnoreTimeDilation, ActiveProperty->bPlayWhilePaused));
+			DisplayDebugManager.DrawString(FString::Printf(TEXT("\t Flags: \n \t \t \t bLooping: %d \n \t \t \t bIgnoreTimeDilation: %d \n \t \t \t bPlayWhilePaused: %d"),
+				ActiveProperty->bLooping, ActiveProperty->bIgnoreTimeDilation, ActiveProperty->bPlayWhilePaused));
 
 			// Highlight if a property's evaluation has completed. If can go higher then the duration if
 			// is set to not be removed after its evaluation is complete
diff --git a/Engine/Source/Runtime/Engine/Classes/GameFramework/InputDeviceSubsystem.h b/Engine/Source/Runtime/Engine/Classes/GameFramework/InputDeviceSubsystem.h
--- a/Engine/Source/Runtime/Engine/Classes/GameFramework/InputDeviceSubsystem.h
+++ b/Engine/Source/Runtime/Engine/Classes/GameFramework/InputDeviceSubsystem.h
@@ -21,25 +21,17 @@ struct ENGINE_API FSetDevicePropertyParams
 	GENERATED_BODY()
 
 	FSetDevicePropertyParams();
-
-	FSetDevicePropertyParams(TSubclassOf<UInputDeviceProperty> InPropertyClass, const FPlatformUserId InUserId, const bool bInRemoveAfterEvaluationTime = true);
-	
-	/** The device property to set */
-	UPROPERTY(EditInstanceOnly, BlueprintReadWrite, Category = "Input Devices")
-	TSubclassOf<UInputDeviceProperty> DevicePropertyClass;
 	
 	/** The Platform User whose device's should receive the device property */
 	UPROPERTY(EditAnywhere, BlueprintReadWrite, Category = "Input Devices")
 	FPlatformUserId UserId;
 
 	/**
-	* If true, then when the InputDeviceProperty is done being evaluated and has fulfilled its duration, then
-	* it will have it's Reset function called and be removed as an active property. 
-	* 
-	* Set this to false if you want your device property to keep being applied without regards to it's duration.
+	* If true, then the input device property will not be removed after it's evaluation time has completed.
+	* Instead, it will remain active until manually removed with a RemoveDeviceProperty call.
 	*/
 	UPROPERTY(EditAnywhere, BlueprintReadWrite, Category = "Input Devices")
-	bool bRemoveAfterEvaluationTime = true;
+	bool bLooping = false;
 
 	/** If true, then this device property will ignore dilated delta time and use the Applications delta time instead */
 	UPROPERTY(EditAnywhere, BlueprintReadWrite, Category = "Input Devices")
@@ -106,11 +98,10 @@ protected:
 		FInputDevicePropertyHandle PropertyHandle = FInputDevicePropertyHandle::InvalidHandle;
 
 		/**
-		* If true, then when the InputDeviceProperty is done being evaluated and has fulfilled its duration, then
-		* it will have it's Reset function called. Set this to false if you want your device property to stay
-		* in it's ending state.
+		* If true, then the input device property will not be removed after it's evaluation time has completed.
+		* Instead, it will remain active until manually removed with a RemoveDeviceProperty call.
 		*/
-		bool bRemoveAfterEvaluationTime = true;
+		bool bLooping = false;
 
 		/** If true, then this device property will ignore dilated delta time and use the Applications delta time instead */
 		bool bIgnoreTimeDilation = false;
@@ -144,10 +135,18 @@ public:
 	UFUNCTION(BlueprintCallable, Category = "Input Devices")
 	static APlayerController* GetPlayerControllerFromInputDevice(const FInputDeviceId DeviceId);
 
-	/** Set the given device property, which will start the evaluation and application of it to the given platform user. */
-	UFUNCTION(BlueprintCallable, Category = "Input Devices", meta = (AutoCreateRefTerm = "Params"))
-	FInputDevicePropertyHandle SetDeviceProperty(const FSetDevicePropertyParams& Params);
+	/**
+	 * Starts tracking the given device property as an "Active" property. This means that the property will be evaluted and applied to its platform user
+	 *
+	 * NOTE: This does NOT make a new instance of the given property. If you pass in the same object before it is completely
+	 * evaluated, then you see undesired effects.
+	 */
+	FInputDevicePropertyHandle ActivateDeviceProperty(UInputDeviceProperty* Property, const FSetDevicePropertyParams& Params);
 
+	/** Spawn a new instance of the given device property class and activate it. */
+	UFUNCTION(BlueprintCallable, Category = "Input Devices", meta = (AutoCreateRefTerm = "Params", ReturnDisplayName = "Device Property Handle"))
+	FInputDevicePropertyHandle ActivateDevicePropertyOfClass(TSubclassOf<UInputDeviceProperty> PropertyClass, const FSetDevicePropertyParams& Params);
+	
 	/** Returns a pointer to the active input device property with the given handle. Returns null if the property doesn't exist */
 	UFUNCTION(BlueprintCallable, Category = "Input Devices", meta=(ReturnDisplayName="Device Property"))
 	UInputDeviceProperty* GetActiveDeviceProperty(const FInputDevicePropertyHandle Handle) const;
@@ -155,6 +154,10 @@ public:
 	/** Returns true if the property associated with the given handle is currently active, and it is not pending removal */
 	UFUNCTION(BlueprintCallable, Category = "Input Devices", meta = (ReturnDisplayName = "Is Property Active"))
 	bool IsPropertyActive(const FInputDevicePropertyHandle Handle) const;
+	
+	/** Returns true if the given UInputDeviceProperty object is already being evaluated. */
+	//UFUNCTION(BlueprintCallable, Category = "Input Devices", meta = (ReturnDisplayName = "Is Property Active"))
+	//bool IsDevicePropertyActive(UInputDeviceProperty* Property) const;
 
 	/**
 	* Remove a single device property based on it's handle
@@ -212,7 +215,7 @@ protected:
 #endif
 	
 	/**
-	* Array of currently active input device properties that will be evaluated on tick
+	* Set of currently active input device properties that will be evaluated on tick
 	*/
 	TSet<UInputDeviceSubsystem::FActiveDeviceProperty> ActiveProperties;
 
diff --git a/Engine/Source/Runtime/Engine/Private/GameFramework/InputDeviceSubsystem.cpp b/Engine/Source/Runtime/Engine/Private/GameFramework/InputDeviceSubsystem.cpp
--- a/Engine/Source/Runtime/Engine/Private/GameFramework/InputDeviceSubsystem.cpp
+++ b/Engine/Source/Runtime/Engine/Private/GameFramework/InputDeviceSubsystem.cpp
@@ -116,14 +116,6 @@ FSetDevicePropertyParams::FSetDevicePropertyParams()
 
 }
 
-FSetDevicePropertyParams::FSetDevicePropertyParams(TSubclassOf<UInputDeviceProperty> InPropertyClass, const FPlatformUserId InUserId, const bool bInRemoveAfterEvaluationTime /* = true */)
-	: DevicePropertyClass(InPropertyClass)
-	, UserId(InUserId)
-	, bRemoveAfterEvaluationTime(bInRemoveAfterEvaluationTime)
-{
-	
-}
-
 ////////////////////////////////////////////////////////
 // UInputDeviceSubsystem
 
@@ -280,11 +272,11 @@ void UInputDeviceSubsystem::Tick(float InDeltaTime)
 		ActiveProp.EvaluatedDuration += DeltaTime;
 
 		// If the property has run past it's duration, reset it and remove it from our active properties
-		// Only do this if it is marked as 'bRemoveAfterEvaluationTime' so that you can keep device properties set
-		// without having to worry about duration. 
-		if (ActiveProp.bRemoveAfterEvaluationTime && ActiveProp.EvaluatedDuration > ActiveProp.Property->GetDuration())
+		// Don't remove properties that are looping, because the user desires them to keep playing
+		if (!ActiveProp.bLooping && ActiveProp.EvaluatedDuration > ActiveProp.Property->GetDuration())
 		{
 			ActiveProp.Property->ResetDeviceProperty(ActiveProp.PlatformUser);
+			PropertiesPendingRemoval.Remove(ActiveProp.PropertyHandle);
 			It.RemoveCurrent();
 			continue;
 		}
@@ -293,11 +285,12 @@ void UInputDeviceSubsystem::Tick(float InDeltaTime)
 		{
 			ActiveProp.Property->EvaluateDeviceProperty(ActiveProp.PlatformUser, DeltaTime, ActiveProp.EvaluatedDuration);
 			ActiveProp.Property->ApplyDeviceProperty(ActiveProp.PlatformUser);
-		}		 
+		}
 	}
 
-	// After we are done ticking, there should be no properties pending removal
-	ensure(PropertiesPendingRemoval.IsEmpty());
+	// If there are still properties pending removal, then they have already been removed by the loop
+	// above. Empty out the pending removal set to start fresh next frame.
+	PropertiesPendingRemoval.Empty();
 }
 
 APlayerController* UInputDeviceSubsystem::GetPlayerControllerFromPlatformUser(const FPlatformUserId UserId)
@@ -333,27 +326,23 @@ APlayerController* UInputDeviceSubsystem::GetPlayerControllerFromInputDevice(con
 	return GetPlayerControllerFromPlatformUser(IPlatformInputDeviceMapper::Get().GetUserForInputDevice(DeviceId));
 }
 
-FInputDevicePropertyHandle UInputDeviceSubsystem::SetDeviceProperty(const FSetDevicePropertyParams& Params)
+FInputDevicePropertyHandle UInputDeviceSubsystem::ActivateDeviceProperty(UInputDeviceProperty* Property, const FSetDevicePropertyParams& Params)
 {
-	if (!Params.DevicePropertyClass)
+	if (!Property)
 	{
-		UE_LOG(LogInputDeviceProperties, Error, TEXT("Invalid DevicePropertyClass passed into SetDeviceProperty! Nothing will happen."));
+		UE_LOG(LogInputDeviceProperties, Error, TEXT("Invalid Property passed into ActivateDeviceProperty! Nothing will happen."));
 		return FInputDevicePropertyHandle::InvalidHandle;
 	}
-		
+
 	FInputDevicePropertyHandle OutHandle = FInputDevicePropertyHandle::AcquireValidHandle();
 
-	if (ensureAlwaysMsgf(OutHandle.IsValid(), TEXT("Unable to acquire a valid input device property handle!")))
+	if (ensureMsgf(OutHandle.IsValid(), TEXT("Unable to acquire a valid input device property handle! The input device property cannot be activated")))
 	{
 		FActiveDeviceProperty ActiveProp = {};
-
-		// Spawn an instance of this device property		
-		ActiveProp.Property = NewObject<UInputDeviceProperty>(/* Outer = */ this, /* Class */ Params.DevicePropertyClass);
-		ensure(ActiveProp.Property);
-
+		ActiveProp.Property = Property;
 		ActiveProp.PlatformUser = Params.UserId;
 		ActiveProp.PropertyHandle = OutHandle;
-		ActiveProp.bRemoveAfterEvaluationTime = Params.bRemoveAfterEvaluationTime;	
+		ActiveProp.bLooping = Params.bLooping;
 		ActiveProp.bIgnoreTimeDilation = Params.bIgnoreTimeDilation;
 		ActiveProp.bPlayWhilePaused = Params.bPlayWhilePaused;
 
@@ -363,6 +352,18 @@ FInputDevicePropertyHandle UInputDeviceSubsystem::SetDeviceProperty(const FSetDe
 	return OutHandle;
 }
 
+FInputDevicePropertyHandle UInputDeviceSubsystem::ActivateDevicePropertyOfClass(TSubclassOf<UInputDeviceProperty> PropertyClass, const FSetDevicePropertyParams& Params)
+{
+	if (!PropertyClass)
+	{
+		UE_LOG(LogInputDeviceProperties, Error, TEXT("Invalid PropertyClass passed into ActivateDeviceProperty! Nothing will happen."));
+		return FInputDevicePropertyHandle::InvalidHandle;
+	}
+
+	// Spawn a new instance of the given device property and activate it as normal
+	return ActivateDeviceProperty(NewObject<UInputDeviceProperty>(/* Outer = */ this, /* Class */ PropertyClass), Params);
+}
+
 UInputDeviceProperty* UInputDeviceSubsystem::GetActiveDeviceProperty(const FInputDevicePropertyHandle Handle) const
 {
 	// Don't include any properties that are pending removal
