PackedLevelActor Blueprint: Recompile skeleton class after clearing SimpleConstructionScript to avoid component name clashes when adding the new components

PR #9719: UE-168469 proposed bug fix (Contributed by [REDACTED_PERSON])

diff --git a/Engine/Source/Runtime/Engine/Private/PackedLevelActor/PackedLevelActorBuilder.cpp b/Engine/Source/Runtime/Engine/Private/PackedLevelActor/PackedLevelActorBuilder.cpp
--- a/Engine/Source/Runtime/Engine/Private/PackedLevelActor/PackedLevelActorBuilder.cpp
+++ b/Engine/Source/Runtime/Engine/Private/PackedLevelActor/PackedLevelActorBuilder.cpp
@@ -417,6 +417,7 @@ bool FPackedLevelActorBuilder::CreateOrUpdateBlueprintFromPacked(APackedLevelAct
 		{
 			BP->SimpleConstructionScript->RemoveNodeAndPromoteChildren(Node);
 		}
+		FKismetEditorUtilities::GenerateBlueprintSkeleton(BP,true);
 	}
 	else
 	{
