Update force feedback effects to use the subsystem instead of the prototype version of device properties

#rnx
#pf 638a1e983b617fd4be47bb3a

diff --git a/Engine/Source/Developer/AssetTools/Private/AssetTypeActions/AssetTypeActions_ForceFeedbackEffect.cpp b/Engine/Source/Developer/AssetTools/Private/AssetTypeActions/AssetTypeActions_ForceFeedbackEffect.cpp
--- a/Engine/Source/Developer/AssetTools/Private/AssetTypeActions/AssetTypeActions_ForceFeedbackEffect.cpp
+++ b/Engine/Source/Developer/AssetTools/Private/AssetTypeActions/AssetTypeActions_ForceFeedbackEffect.cpp
@@ -142,6 +142,7 @@ void FAssetTypeActions_ForceFeedbackEffect::PlayEffect(UForceFeedbackEffect* Eff
 		PreviewForceFeedbackEffect.ForceFeedbackEffect = Effect;
 		PreviewForceFeedbackEffect.PlayTime = 0.f;
 		PreviewForceFeedbackEffect.PlatformUser = IPlatformInputDeviceMapper::Get().GetPrimaryPlatformUser();
+		PreviewForceFeedbackEffect.ActivateDeviceProperties();
 	}
 	else
 	{
diff --git a/Engine/Source/Runtime/Engine/Classes/GameFramework/ForceFeedbackEffect.h b/Engine/Source/Runtime/Engine/Classes/GameFramework/ForceFeedbackEffect.h
--- a/Engine/Source/Runtime/Engine/Classes/GameFramework/ForceFeedbackEffect.h
+++ b/Engine/Source/Runtime/Engine/Classes/GameFramework/ForceFeedbackEffect.h
@@ -6,6 +6,7 @@
 #include "UObject/Object.h"
 #include "Curves/CurveFloat.h"
 #include "ForceFeedbackParameters.h"
+#include "GameFramework/InputDevicePropertyHandle.h"
 #include "ForceFeedbackEffect.generated.h"
 
 class UForceFeedbackEffect;
@@ -55,10 +56,18 @@ struct ENGINE_API FActiveForceFeedbackEffect
 	/** The platform user that should receive this effect */
 	FPlatformUserId PlatformUser = PLATFORMUSERID_NONE;
 
+	/** Set to true after this force feedback effect has activated it's device properties */
+	bool bActivatedDeviceProperties;
+
+	/** Array of device properties that have been activated by this force feedback effect */
+	UPROPERTY()
+	TSet<FInputDevicePropertyHandle> ActiveDeviceProperties;
+
 	FActiveForceFeedbackEffect()
 		: ForceFeedbackEffect(nullptr)
 		, PlayTime(0.f)
 		, PlatformUser(PLATFORMUSERID_NONE)
+		, bActivatedDeviceProperties(false)
 	{
 	}
 
@@ -67,12 +76,18 @@ struct ENGINE_API FActiveForceFeedbackEffect
 		, Parameters(InParameters)
 		, PlayTime(0.f)
 		, PlatformUser(InPlatformUser)
+		, bActivatedDeviceProperties(false)
 	{
 	}
 
+	~FActiveForceFeedbackEffect();
+
 	// Updates the final force feedback values based on this effect.  Returns true if the effect should continue playing, false if it is finished.
 	bool Update(float DeltaTime, FForceFeedbackValues& Values);
 
+	// Activates all the device properties with the input device subsystem
+	void ActivateDeviceProperties();
+
 	/** Reset any device properties that may need to be after the duration of this effect has ended. */
 	void ResetDeviceProperties();
 
@@ -126,11 +141,6 @@ class UForceFeedbackEffect : public UObject
 
 	void GetValues(const float EvalTime, FForceFeedbackValues& Values, const FPlatformUserId PlatformUser, float ValueMultiplier = 1.f) const;
 
-	void SetDeviceProperties(const FPlatformUserId PlatformUser, const float DeltaTime, const float EvalTime);
-
-	/** Reset any device properties that may need to be after the duration of this effect has ended. */
-	void ResetDeviceProperties(const FPlatformUserId PlatformUser);
-
 	/**
 	 * Returns the channel details that should currently be used for the given platform.
 	 * This will return one of the OverridenDetails if the current platform has one specified.
diff --git a/Engine/Source/Runtime/Engine/Classes/GameFramework/InputDeviceProperties.h b/Engine/Source/Runtime/Engine/Classes/GameFramework/InputDeviceProperties.h
--- a/Engine/Source/Runtime/Engine/Classes/GameFramework/InputDeviceProperties.h
+++ b/Engine/Source/Runtime/Engine/Classes/GameFramework/InputDeviceProperties.h
@@ -32,8 +32,6 @@ UCLASS(Abstract, Blueprintable, BlueprintType, EditInlineNew, CollapseCategories
 class ENGINE_API UInputDeviceProperty : public UObject
 {
 	friend class UInputDeviceSubsystem;
-	// TODO: Make force feedback effects just use the subsystem instead of ticking stuff themselves
-	friend class UForceFeedbackEffect;
 
 	GENERATED_BODY()
 public:
diff --git a/Engine/Source/Runtime/Engine/Private/GameFramework/ForceFeedbackEffect.cpp b/Engine/Source/Runtime/Engine/Private/GameFramework/ForceFeedbackEffect.cpp
--- a/Engine/Source/Runtime/Engine/Private/GameFramework/ForceFeedbackEffect.cpp
+++ b/Engine/Source/Runtime/Engine/Private/GameFramework/ForceFeedbackEffect.cpp
@@ -105,36 +105,6 @@ void UForceFeedbackEffect::GetValues(const float EvalTime, FForceFeedbackValues&
 	}
 }
 
-void UForceFeedbackEffect::SetDeviceProperties(const FPlatformUserId PlatformUser, const float DeltaTime, const float EvalTime)
-{
-	for (TObjectPtr<UInputDeviceProperty> DeviceProp : DeviceProperties)
-	{
-		if (DeviceProp)
-		{
-			if (EvalTime > DeviceProp->GetDuration())
-			{
-				DeviceProp->ResetDeviceProperty(PlatformUser);
-			}
-			else
-			{
-				DeviceProp->EvaluateDeviceProperty(PlatformUser, DeltaTime, EvalTime);
-				DeviceProp->ApplyDeviceProperty(PlatformUser);
-			}			
-		}
-	}
-}
-
-void UForceFeedbackEffect::ResetDeviceProperties(const FPlatformUserId PlatformUser)
-{
-	for (TObjectPtr<UInputDeviceProperty> DeviceProp : DeviceProperties)
-	{
-		if (DeviceProp)
-		{
-			DeviceProp->ResetDeviceProperty(PlatformUser);
-		}
-	}
-}
-
 const TArray<FForceFeedbackChannelDetails>& UForceFeedbackEffect::GetCurrentChannelDetails(const FPlatformUserId PlatformUser) const
 {	
 	if (const UInputDeviceSubsystem* SubSystem = UInputDeviceSubsystem::Get())
@@ -150,6 +120,11 @@ const TArray<FForceFeedbackChannelDetails>& UForceFeedbackEffect::GetCurrentChan
 	return ChannelDetails;
 }
 
+FActiveForceFeedbackEffect::~FActiveForceFeedbackEffect()
+{
+	ResetDeviceProperties();
+}
+
 void FActiveForceFeedbackEffect::GetValues(FForceFeedbackValues& Values) const
 {
 	if (ForceFeedbackEffect)
@@ -172,13 +147,18 @@ bool FActiveForceFeedbackEffect::Update(const float DeltaTime, FForceFeedbackVal
 	}
 
 	const float EffectDuration = ForceFeedbackEffect->GetDuration();
-	const float DevicePropDuration = ForceFeedbackEffect->GetTotalDevicePropertyDuration();
+
+	// If this is the first time that we are playing the effect
+	if (!bActivatedDeviceProperties)
+	{
+		ActivateDeviceProperties();
+	}
 
 	PlayTime += (Parameters.bIgnoreTimeDilation ? FApp::GetDeltaTime() : DeltaTime);
 
 	// If the play time is longer then the force feedback effect curve's last key value, 
 	// or if there are still device properties that need to be evaluated
-	if (PlayTime > EffectDuration && PlayTime > DevicePropDuration && (!Parameters.bLooping || (EffectDuration == 0.0f && DevicePropDuration == 0.0f)))
+	if (PlayTime > EffectDuration && (!Parameters.bLooping || (EffectDuration == 0.0f)))
 	{
 		return false;
 	}
@@ -186,23 +166,43 @@ bool FActiveForceFeedbackEffect::Update(const float DeltaTime, FForceFeedbackVal
 	if (PlayTime <= EffectDuration || Parameters.bLooping)
 	{
 		GetValues(Values);
-	}	
-	
-	// Update device properties if we can
-	if (PlayTime <= DevicePropDuration)
-	{
-		// Set any input device properties associated with this effect
-		const float EvalTime = PlayTime - DevicePropDuration * FMath::FloorToFloat(PlayTime / DevicePropDuration);
-		ForceFeedbackEffect->SetDeviceProperties(PlatformUser, DeltaTime, EvalTime);
-	}	
+	}
 
 	return true;
 }
 
+void FActiveForceFeedbackEffect::ActivateDeviceProperties()
+{
+	if (ensure(ForceFeedbackEffect))
+	{
+		if (UInputDeviceSubsystem* System = UInputDeviceSubsystem::Get())
+		{
+			for (TObjectPtr<UInputDeviceProperty> DeviceProp : ForceFeedbackEffect->DeviceProperties)
+			{
+				if (ensure(DeviceProp))
+				{
+					FSetDevicePropertyParams Params = {};
+					Params.bIgnoreTimeDilation = Parameters.bIgnoreTimeDilation;
+					Params.UserId = PlatformUser;
+					Params.bPlayWhilePaused = Parameters.bPlayWhilePaused;
+					Params.bLooping = Parameters.bLooping;
+					ActiveDeviceProperties.Emplace(System->ActivateDeviceProperty(DeviceProp, Params));
+				}
+			}
+		}
+
+		bActivatedDeviceProperties = true;
+	}	
+}
+
 void FActiveForceFeedbackEffect::ResetDeviceProperties()
 {
-	if (ForceFeedbackEffect)
+	if (!ActiveDeviceProperties.IsEmpty())
 	{
-		ForceFeedbackEffect->ResetDeviceProperties(PlatformUser);
+		if (UInputDeviceSubsystem* System = UInputDeviceSubsystem::Get())
+		{
+			System->RemoveDevicePropertyHandles(ActiveDeviceProperties);
+			ActiveDeviceProperties.Empty();
+		}
 	}
 }
\ No newline at end of file
