StreamingTexture.cpp fix streaming cinematic lod bias broken in recent refactor
broken by 22409186 Texture: remove cached LOD Bias

diff --git a/Engine/Source/Runtime/Engine/Private/Streaming/StreamingTexture.cpp b/Engine/Source/Runtime/Engine/Private/Streaming/StreamingTexture.cpp
--- a/Engine/Source/Runtime/Engine/Private/Streaming/StreamingTexture.cpp
+++ b/Engine/Source/Runtime/Engine/Private/Streaming/StreamingTexture.cpp
@@ -190,21 +190,9 @@ void FStreamingRenderAsset::UpdateDynamicData(const int32* NumStreamedMips, int3
 
 		int32 LODBias = 0;
 		if (!Settings.bUseAllMips)
-		{			
-			UTexture * Texture = Cast<UTexture>(RenderAsset);
-			if ( Texture )
-			{
-				LODBias = FMath::Max<int32>(0, Texture->CalculateLODBias(false) - ResourceState.AssetLODBias);
-			}
-			else
-			{
-				// GetCachedLODBias() - NumCinematicMipLevels
-				//  is not the right way to get non-cinematic LODBias
-				//	you should call CalculateLODBias(false)
-				//	the difference occurs if any of the clamps come into play
-				const int32 ResourceLODBias = FMath::Max<int32>(0, RenderAsset->GetCachedLODBias() - ResourceState.AssetLODBias);
-				LODBias = FMath::Max<int32>(ResourceLODBias - NumCinematicMipLevels, 0);
-			}
+		{
+			const int32 ResourceLODBias = FMath::Max<int32>(0, RenderAsset->GetCachedLODBias() - ResourceState.AssetLODBias);
+			LODBias = FMath::Max<int32>(ResourceLODBias - NumCinematicMipLevels, 0);
 
 			// Reduce the max allowed resolution according to LODBias if the texture group allows it.
 			if (IsMaxResolutionAffectedByGlobalBias() && !Settings.bUsePerTextureBias)
