Add a "bEnableInputDeviceSubsystem" setting to disable the input device subsystem entirely via a hotfix if need be.
You can add :

[/Script/Engine.InputSettings]
bEnableInputDeviceSubsystem=false

to your game's DefaultInput.ini to disable it, or turn it off in the editor preferences.

#rnx

diff --git a/Engine/Config/BaseInput.ini b/Engine/Config/BaseInput.ini
--- a/Engine/Config/BaseInput.ini
+++ b/Engine/Config/BaseInput.ini
@@ -15,6 +15,7 @@ TriggerKeyIsConsideredPressed=0.75
 TriggerKeyIsConsideredReleased=0.25
 InitialButtonRepeatDelay=0.2
 ButtonRepeatDelay=0.1
+bEnableInputDeviceSubsystem=true
 
 ;-----------------------------------------------------------------------------------------
 ; Axis properties
diff --git a/Engine/Source/Runtime/Engine/Classes/GameFramework/InputDeviceSubsystem.h b/Engine/Source/Runtime/Engine/Classes/GameFramework/InputDeviceSubsystem.h
--- a/Engine/Source/Runtime/Engine/Classes/GameFramework/InputDeviceSubsystem.h
+++ b/Engine/Source/Runtime/Engine/Classes/GameFramework/InputDeviceSubsystem.h
@@ -112,6 +112,12 @@ protected:
 
 public:
 
+	/**
+	 * Returns a pointer to the Input Device Engine Subsystem if it is available.
+	 * 
+	 * NOTE: This may be null if the bEnableInputDeviceSubsystem flag in UInputSettings
+	 * is set to false!
+	 */
 	static UInputDeviceSubsystem* Get();
 
 	virtual void Initialize(FSubsystemCollectionBase& Collection) override;
@@ -153,11 +159,7 @@ public:
 	
 	/** Returns true if the property associated with the given handle is currently active, and it is not pending removal */
 	UFUNCTION(BlueprintCallable, Category = "Input Devices", meta = (ReturnDisplayName = "Is Property Active"))
-	bool IsPropertyActive(const FInputDevicePropertyHandle Handle) const;
-	
-	/** Returns true if the given UInputDeviceProperty object is already being evaluated. */
-	//UFUNCTION(BlueprintCallable, Category = "Input Devices", meta = (ReturnDisplayName = "Is Property Active"))
-	//bool IsDevicePropertyActive(UInputDeviceProperty* Property) const;
+	bool IsPropertyActive(const FInputDevicePropertyHandle Handle) const;	
 
 	/**
 	* Remove a single device property based on it's handle
diff --git a/Engine/Source/Runtime/Engine/Classes/GameFramework/InputSettings.h b/Engine/Source/Runtime/Engine/Classes/GameFramework/InputSettings.h
--- a/Engine/Source/Runtime/Engine/Classes/GameFramework/InputSettings.h
+++ b/Engine/Source/Runtime/Engine/Classes/GameFramework/InputSettings.h
@@ -78,6 +78,15 @@ class ENGINE_API UInputSettings
 	UPROPERTY(config, EditAnywhere, Category = "Input")
 	uint8 bFilterInputByPlatformUser:1;
 
+	/**
+	 * If true, then the input device subsystem will be allowed to Initalize when the engine boots.
+	 * NOTE: For this setting to take effect, and editor restart is required.
+	 * 
+	 * @see UInputDeviceSubsystem
+	 */
+	UPROPERTY(config, EditAnywhere, Category = "Input")
+	uint8 bEnableInputDeviceSubsystem:1;
+
 	/**
 	 * If true, then the Player Controller will have it's Pressed Keys flushed when the input mode is changed
 	 * to Game and UI mode or the game viewport loses focus. The default behavior is true.
diff --git a/Engine/Source/Runtime/Engine/Private/GameFramework/InputDeviceSubsystem.cpp b/Engine/Source/Runtime/Engine/Private/GameFramework/InputDeviceSubsystem.cpp
--- a/Engine/Source/Runtime/Engine/Private/GameFramework/InputDeviceSubsystem.cpp
+++ b/Engine/Source/Runtime/Engine/Private/GameFramework/InputDeviceSubsystem.cpp
@@ -163,8 +163,15 @@ bool UInputDeviceSubsystem::ShouldCreateSubsystem(UObject* Outer) const
 	{
 		return false;
 	}
+
+	// There is a setting to turn off this subsystem entirely.
+	const bool bShouldCreate = GetDefault<UInputSettings>()->bEnableInputDeviceSubsystem;
+	if (!bShouldCreate)
+	{
+		UE_LOG(LogInputDeviceProperties, Log, TEXT("UInputSettings::bEnableInputDeviceSubsystem is false, the Input Device Subsystem will NOT be created!"));
+	}
 	
-	return Super::ShouldCreateSubsystem(Outer);
+	return bShouldCreate && Super::ShouldCreateSubsystem(Outer);
 }
 
 UWorld* UInputDeviceSubsystem::GetTickableGameObjectWorld() const
diff --git a/Engine/Source/Runtime/Engine/Private/UserInterface/InputSettings.cpp b/Engine/Source/Runtime/Engine/Private/UserInterface/InputSettings.cpp
--- a/Engine/Source/Runtime/Engine/Private/UserInterface/InputSettings.cpp
+++ b/Engine/Source/Runtime/Engine/Private/UserInterface/InputSettings.cpp
@@ -26,6 +26,7 @@ UInputSettings::UInputSettings(const FObjectInitializer& ObjectInitializer)
 	, bEnableLegacyInputScales(true)
 	, bEnableMotionControls(true)
 	, bFilterInputByPlatformUser(false)
+	, bEnableInputDeviceSubsystem(true)
 	, bShouldFlushPressedKeysOnViewportFocusLost(true)
 	, bEnableDynamicComponentInputBinding(true)
 	, DefaultViewportMouseCaptureMode(EMouseCaptureMode::CapturePermanently_IncludingInitialMouseDown)
