Make FActiveDeviceProperty a USTRUCT again because it has to have a the UPROPERTY for the TobjectPtr to stay valid over time

#rnx
#pf skip

diff --git a/Engine/Plugins/Runtime/InputDebugging/Source/InputDebugging/Private/InputDeviceDebugTools.cpp b/Engine/Plugins/Runtime/InputDebugging/Source/InputDebugging/Private/InputDeviceDebugTools.cpp
--- a/Engine/Plugins/Runtime/InputDebugging/Source/InputDebugging/Private/InputDeviceDebugTools.cpp
+++ b/Engine/Plugins/Runtime/InputDebugging/Source/InputDebugging/Private/InputDeviceDebugTools.cpp
@@ -93,18 +93,18 @@ void FInputDeviceDebugTools::OnShowDebugDeviceProperties(UCanvas* Canvas)
 	}
 
 	// Sort the active properties by their owning Platform User
-	static TMap<FPlatformUserId, TArray<const UInputDeviceSubsystem::FActiveDeviceProperty*>> PlatformUserDebugStrings;
+	static TMap<FPlatformUserId, TArray<const FActiveDeviceProperty*>> PlatformUserDebugStrings;
 	{
 		PlatformUserDebugStrings.Empty();
 
-		for (const UInputDeviceSubsystem::FActiveDeviceProperty& ActiveProperty : System->ActiveProperties)
+		for (const FActiveDeviceProperty& ActiveProperty : System->ActiveProperties)
 		{
 			PlatformUserDebugStrings.FindOrAdd(ActiveProperty.PlatformUser).Emplace(&ActiveProperty);
 		}
 	}	
 
 	// Actually display the debug information	
-	for (TPair<FPlatformUserId, TArray<const UInputDeviceSubsystem::FActiveDeviceProperty*>> Pair : PlatformUserDebugStrings)
+	for (TPair<FPlatformUserId, TArray<const FActiveDeviceProperty*>> Pair : PlatformUserDebugStrings)
 	{
 		const FString PlatUserLabel = FString::Printf(TEXT("Platform User %d has %d active Device Properties"),
 			Pair.Key.GetInternalId(),
@@ -113,7 +113,7 @@ void FInputDeviceDebugTools::OnShowDebugDeviceProperties(UCanvas* Canvas)
 		DisplayDebugManager.SetDrawColor(FColor::Yellow);
 		DisplayDebugManager.DrawString(PlatUserLabel);
 
-		for (const UInputDeviceSubsystem::FActiveDeviceProperty* ActiveProperty: Pair.Value)
+		for (const FActiveDeviceProperty* ActiveProperty: Pair.Value)
 		{
 			// Draw the name of the property
 			DisplayDebugManager.SetDrawColor(FColor::White);
diff --git a/Engine/Source/Runtime/Engine/Classes/GameFramework/InputDeviceSubsystem.h b/Engine/Source/Runtime/Engine/Classes/GameFramework/InputDeviceSubsystem.h
--- a/Engine/Source/Runtime/Engine/Classes/GameFramework/InputDeviceSubsystem.h
+++ b/Engine/Source/Runtime/Engine/Classes/GameFramework/InputDeviceSubsystem.h
@@ -49,6 +49,55 @@ struct ENGINE_API FSetDevicePropertyParams
  */
 DECLARE_DYNAMIC_MULTICAST_DELEGATE_TwoParams(FHardwareInputDeviceChanged, const FPlatformUserId, UserId, const FInputDeviceId, DeviceId);
 
+/** Contains a pointer to an active device property and keeps track of how long it has been evaluated for */
+USTRUCT()
+struct ENGINE_API FActiveDeviceProperty
+{
+	GENERATED_BODY()
+	
+	/** Active properties can just use the hash of their FInputDevicePropertyHandle for a fast and unique lookup */
+	ENGINE_API friend uint32 GetTypeHash(const FActiveDeviceProperty& InProp);
+	ENGINE_API friend bool operator==(const FActiveDeviceProperty& ActiveProp, const FInputDevicePropertyHandle& Handle);
+	ENGINE_API friend bool operator!=(const FActiveDeviceProperty& ActiveProp, const FInputDevicePropertyHandle& Handle);
+
+	bool operator==(const FActiveDeviceProperty& Other) const;
+	bool operator!=(const FActiveDeviceProperty& Other) const;
+
+	/** The active device property */
+	UPROPERTY()
+	TObjectPtr<UInputDeviceProperty> Property = nullptr;
+
+	/** How long this property has been evaluated for. DeltaTime is added to this on tick */
+	double EvaluatedDuration = 0.0;
+
+	/** The platform user that is actively receiving this device property */
+	FPlatformUserId PlatformUser = PLATFORMUSERID_NONE;
+
+	/** The handle of this active property. */
+	FInputDevicePropertyHandle PropertyHandle = FInputDevicePropertyHandle::InvalidHandle;
+
+	/**
+	* If true, then the input device property will not be removed after it's evaluation time has completed.
+	* Instead, it will remain active until manually removed with a RemoveDeviceProperty call.
+	*/
+	bool bLooping = false;
+
+	/** If true, then this device property will ignore dilated delta time and use the Applications delta time instead */
+	bool bIgnoreTimeDilation = false;
+
+	/** If true, then this device property will be played even if the game world is paused. */
+	bool bPlayWhilePaused = false;
+
+	/**
+	 * This is set to true when this device property has been applied.
+	 * All device properties should be applied at least one time, no matter what.
+	 * This handles cases where the evaluation time is longer then the duration of a property.
+	 * I.e., your property is set to a duration of 0.1, but you get a delta time of .12 seconds for some reason
+	 * (choppy frames, low perf client, debugging, etc). 
+	 */
+	bool bHasBeenAppliedAtLeastOnce = false;
+};
+
 /**
 * The input device subsystem provides an interface to allow users to set Input Device Properties
 * on any Platform User. 
@@ -60,65 +109,6 @@ class ENGINE_API UInputDeviceSubsystem : public UEngineSubsystem, public FTickab
 	friend class FInputDeviceDebugTools;
 	
 	GENERATED_BODY()
-
-protected:
-
-	/** Contains a pointer to an active device property and keeps track of how long it has been evaluated for */
-	struct FActiveDeviceProperty
-	{		
-		/** Active properties can just use the hash of their FInputDevicePropertyHandle for a fast and unique lookup */
-		friend uint32 GetTypeHash(const UInputDeviceSubsystem::FActiveDeviceProperty& InProp)
-		{
-			return InProp.PropertyHandle.GetTypeHash();
-		}
-
-		friend bool operator==(const UInputDeviceSubsystem::FActiveDeviceProperty& ActiveProp, const FInputDevicePropertyHandle& Handle)
-		{
-			return ActiveProp.PropertyHandle == Handle;
-		}
-
-		friend bool operator!=(const UInputDeviceSubsystem::FActiveDeviceProperty& ActiveProp, const FInputDevicePropertyHandle& Handle)
-		{
-			return ActiveProp.PropertyHandle != Handle;
-		}
-
-		bool operator==(const UInputDeviceSubsystem::FActiveDeviceProperty& Other) const;
-		bool operator!=(const UInputDeviceSubsystem::FActiveDeviceProperty& Other) const;
-
-		/** The active device property */	
-		TObjectPtr<UInputDeviceProperty> Property = nullptr;
-
-		/** How long this property has been evaluated for. DeltaTime is added to this on tick */
-		double EvaluatedDuration = 0.0;
-
-		/** The platform user that is actively receiving this device property */
-		FPlatformUserId PlatformUser = PLATFORMUSERID_NONE;
-
-		/** The handle of this active property. */
-		FInputDevicePropertyHandle PropertyHandle = FInputDevicePropertyHandle::InvalidHandle;
-
-		/**
-		* If true, then the input device property will not be removed after it's evaluation time has completed.
-		* Instead, it will remain active until manually removed with a RemoveDeviceProperty call.
-		*/
-		bool bLooping = false;
-
-		/** If true, then this device property will ignore dilated delta time and use the Applications delta time instead */
-		bool bIgnoreTimeDilation = false;
-
-		/** If true, then this device property will be played even if the game world is paused. */
-		bool bPlayWhilePaused = false;
-
-		/**
-		 * This is set to true when this device property has been applied.
-		 * All device properties should be applied at least one time, no matter what.
-		 * This handles cases where the evaluation time is longer then the duration of a property.
-		 * I.e., your property is set to a duration of 0.1, but you get a delta time of .12 seconds for some reason
-		 * (choppy frames, low perf client, debugging, etc). 
-		 */
-		bool bHasBeenAppliedAtLeastOnce = false;
-	};
-
 public:
 
 	/**
@@ -228,7 +218,8 @@ protected:
 	/**
 	* Set of currently active input device properties that will be evaluated on tick
 	*/
-	TSet<UInputDeviceSubsystem::FActiveDeviceProperty> ActiveProperties;
+	UPROPERTY(Transient)
+	TSet<FActiveDeviceProperty> ActiveProperties;
 
 	/**
 	 * Set of property handles the properties that are currently pending manual removal.
diff --git a/Engine/Source/Runtime/Engine/Private/GameFramework/InputDeviceSubsystem.cpp b/Engine/Source/Runtime/Engine/Private/GameFramework/InputDeviceSubsystem.cpp
--- a/Engine/Source/Runtime/Engine/Private/GameFramework/InputDeviceSubsystem.cpp
+++ b/Engine/Source/Runtime/Engine/Private/GameFramework/InputDeviceSubsystem.cpp
@@ -96,13 +96,28 @@ public:
 ////////////////////////////////////////////////////////
 // FActiveDeviceProperty
 
+uint32 GetTypeHash(const FActiveDeviceProperty& InProp)
+{
+	return InProp.PropertyHandle.GetTypeHash();
+}
+
+bool operator==(const FActiveDeviceProperty& ActiveProp, const FInputDevicePropertyHandle& Handle)
+{
+	return ActiveProp.PropertyHandle == Handle;
+}
+
+bool operator!=(const FActiveDeviceProperty& ActiveProp, const FInputDevicePropertyHandle& Handle)
+{
+	return ActiveProp.PropertyHandle != Handle;
+}
+
 // The property handles are the only things that matter when comparing, they will always be unique.
-bool UInputDeviceSubsystem::FActiveDeviceProperty::operator==(const UInputDeviceSubsystem::FActiveDeviceProperty& Other) const
+bool FActiveDeviceProperty::operator==(const FActiveDeviceProperty& Other) const
 {
 	return PropertyHandle == Other.PropertyHandle;
 }
 
-bool UInputDeviceSubsystem::FActiveDeviceProperty::operator!=(const UInputDeviceSubsystem::FActiveDeviceProperty& Other) const
+bool FActiveDeviceProperty::operator!=(const FActiveDeviceProperty& Other) const
 {
 	return PropertyHandle != Other.PropertyHandle;
 }
