UE-171631 UBodySetup::AbortPhysicsMeshAsyncCreation() needs to be implemented for Chaos
UE-169219 HoloLens MRMesh crashes when collisions are enabled
-Implemented  it in a simple way that is compatible for 5.1.1 hotfix.

diff --git a/Engine/Source/Runtime/Engine/Private/PhysicsEngine/BodySetup.cpp b/Engine/Source/Runtime/Engine/Private/PhysicsEngine/BodySetup.cpp
--- a/Engine/Source/Runtime/Engine/Private/PhysicsEngine/BodySetup.cpp
+++ b/Engine/Source/Runtime/Engine/Private/PhysicsEngine/BodySetup.cpp
@@ -498,12 +498,32 @@ void UBodySetup::CreatePhysicsMeshesAsync(FOnAsyncPhysicsCookFinished OnAsyncPhy
 
 void UBodySetup::AbortPhysicsMeshAsyncCreation()
 {
-	// TODO Chaos - support async cook abort
+	check(IsInGameThread());
+
+	// If we have a current cook helper we will ask it to cancel, null our ref to it, and ignore any results it produces.
+	// Note that CreatePhysicsMeshesAsync, AbortPhysicsMeshAsyncCreation and FinishCreatePhysicsMeshesAsync all run on the game thread.
+	// After this function returns we could safely call CreatePhysicsMeshesAsync again on this body from the game thread.
+	if (CurrentCookHelper)
+	{
+		CurrentCookHelper->CancelCookAsync();
+		CurrentCookHelper = nullptr;
+	}
 }
 
 void UBodySetup::FinishCreatePhysicsMeshesAsync(FAsyncCookHelper* AsyncPhysicsCookHelper, FOnAsyncPhysicsCookFinished OnAsyncPhysicsCookFinished)
 {
-	// Ensure we haven't gotten multiple cooks going
+	check(IsInGameThread());
+
+	// If a canceled CookHelper finishes delete it.
+	if (AsyncPhysicsCookHelper && AsyncPhysicsCookHelper->WasCanceled())
+	{
+		// CurrentCookHelper could be null or it could be a new FAsyncCookHelper, but it should never be a canceled cook helper.
+		check(CurrentCookHelper != AsyncPhysicsCookHelper);
+		delete AsyncPhysicsCookHelper;
+		return;
+	}
+
+	// Un-canceled cooks should match the CurrentCookHelper.
 	// Then clear it
 	check(CurrentCookHelper == AsyncPhysicsCookHelper);
 	CurrentCookHelper = nullptr;
diff --git a/Engine/Source/Runtime/Engine/Private/PhysicsEngine/Experimental/ChaosCooking.cpp b/Engine/Source/Runtime/Engine/Private/PhysicsEngine/Experimental/ChaosCooking.cpp
--- a/Engine/Source/Runtime/Engine/Private/PhysicsEngine/Experimental/ChaosCooking.cpp
+++ b/Engine/Source/Runtime/Engine/Private/PhysicsEngine/Experimental/ChaosCooking.cpp
@@ -345,6 +345,7 @@ namespace Chaos
 
 	FCookHelper::FCookHelper(UBodySetup* InSetup)
 		: SourceSetup(InSetup)
+		, bCanceled(false)
 	{
 		check(SourceSetup);
 		EPhysXMeshCookFlags TempFlags = static_cast<EPhysXMeshCookFlags>(0);
@@ -353,6 +354,8 @@ namespace Chaos
 
 	void FCookHelper::Cook()
 	{
+		if (bCanceled) return;
+
 		Cooking::BuildConvexMeshes(SimpleImplicits, CookInfo);
 		Cooking::BuildTriangleMeshes(ComplexImplicits, FaceRemap, VertexRemap, CookInfo);
 
diff --git a/Engine/Source/Runtime/Engine/Private/PhysicsEngine/Experimental/ChaosCooking.h b/Engine/Source/Runtime/Engine/Private/PhysicsEngine/Experimental/ChaosCooking.h
--- a/Engine/Source/Runtime/Engine/Private/PhysicsEngine/Experimental/ChaosCooking.h
+++ b/Engine/Source/Runtime/Engine/Private/PhysicsEngine/Experimental/ChaosCooking.h
@@ -40,8 +40,15 @@ namespace Chaos
 		void CookAsync(FSimpleDelegateGraphTask::FDelegate CompletionDelegate);
 		bool HasWork() const;
 
+		// CancelCookAsync is not guaranteed to have any effect on the work done.
+		// If it is called the cook work may be abandoned and the CookAsync may return early.
+		// If bCancel is true in the CompletionDelegate the results must be ignored.
+		void CancelCookAsync() { bCanceled = true; }
+		bool WasCanceled() const { return bCanceled; }
+
 	private:
 		UBodySetup* SourceSetup;
 		FCookBodySetupInfo CookInfo;
+		std::atomic<bool> bCanceled;
 	};
 }
\ No newline at end of file
