Improve editor responsiveness under heavy load by unlocking the use of foreground threads even when max concurrency is reached for background work

#rnx

diff --git a/Engine/Source/Runtime/Engine/Private/AssetCompilingManager.cpp b/Engine/Source/Runtime/Engine/Private/AssetCompilingManager.cpp
--- a/Engine/Source/Runtime/Engine/Private/AssetCompilingManager.cpp
+++ b/Engine/Source/Runtime/Engine/Private/AssetCompilingManager.cpp
@@ -120,8 +120,10 @@ public:
 	 * InMaxConcurrency           Maximum number of concurrent tasks allowed, -1 will limit concurrency to number of threads available in the underlying thread pool.
 	 * InPriorityMapper           Thread-safe function used to map any priority from this Queue to the priority that should be used when scheduling the task on the underlying thread pool.
 	 */
-	FMemoryBoundQueuedThreadPoolWrapper(FQueuedThreadPool* InWrappedQueuedThreadPool, int32 InMaxConcurrency = -1, TFunction<EQueuedWorkPriority(EQueuedWorkPriority)> InPriorityMapper = [](EQueuedWorkPriority InPriority) { return InPriority; })
-		: FQueuedThreadPoolWrapper(InWrappedQueuedThreadPool, InMaxConcurrency, InPriorityMapper)
+	FMemoryBoundQueuedThreadPoolWrapper(FQueuedThreadPool* InWrappedQueuedThreadPool, TFunction<int32 ()> InMaxForegroundConcurrency, TFunction<int32()> InMaxBackgroundConcurrency, TFunction<EQueuedWorkPriority(EQueuedWorkPriority)> InPriorityMapper = [](EQueuedWorkPriority InPriority) { return InPriority; })
+		: FQueuedThreadPoolWrapper(InWrappedQueuedThreadPool, -1, InPriorityMapper)
+		, MaxForegroundConcurrency(InMaxForegroundConcurrency)
+		, MaxBackgroundConcurrency(InMaxBackgroundConcurrency)
 	{
 	}
 
@@ -224,12 +226,28 @@ public:
 		UpdateCounters();
 	}
 
+	int32 GetAdjustedMaxConcurrency(EQueuedWorkPriority Priority) const
+	{
+		// Only unlocks foreground concurrency when priority is highest or blocking. This improves
+		// editor responsiveness under heavy load when GT is blocking on something and
+		// max concurrency has already been reached by background work.
+		if (Priority <= EQueuedWorkPriority::Highest)
+		{
+			return MaxBackgroundConcurrency() + MaxForegroundConcurrency();
+		}
+		else
+		{
+			return MaxBackgroundConcurrency();
+		}
+	}
+
 	int32 GetMaxConcurrency() const override
 	{
 		int64 NewRequiredMemory = 0;
 		
 		// Add next work memory requirement to see if it still fits in available memory
-		if (IQueuedWork* NextWork = QueuedWork.Peek())
+		EQueuedWorkPriority Priority;
+		if (IQueuedWork* NextWork = QueuedWork.Peek(&Priority))
 		{
 			NewRequiredMemory += GetRequiredMemory(NextWork);
 		}
@@ -237,7 +255,7 @@ public:
 		// TotalEstimatedMemory is the sum of all "RequiredMemory" for currently scheduled tasks
 		int64 TotalRequiredMemory = TotalEstimatedMemory.load() + NewRequiredMemory;
 
-		int32 DynamicMaxConcurrency = FQueuedThreadPoolWrapper::GetMaxConcurrency();
+		int32 DynamicMaxConcurrency = GetAdjustedMaxConcurrency(Priority);
 		int32 Concurrency = FQueuedThreadPoolWrapper::GetCurrentConcurrency();
 
 		// Never limit below a concurrency of 1 or else we'll starve the asset processing and never be scheduled again
@@ -268,11 +286,13 @@ public:
 
 		TRACE_COUNTER_SET(AsyncCompilationMaxConcurrency, DynamicMaxConcurrency);
 
-		check(DynamicMaxConcurrency > 0);
 		return DynamicMaxConcurrency;
 	}
 
 private:
+	TFunction<int32()> MaxForegroundConcurrency;
+	TFunction<int32()> MaxBackgroundConcurrency;
+
 	std::atomic<int64> TotalEstimatedMemory {0};
 };
 
@@ -310,6 +330,8 @@ FAssetCompilingManager::~FAssetCompilingManager()
 #endif
 }
 
+extern CORE_API int32 GNumForegroundWorkers; // TaskGraph.cpp
+
 FQueuedThreadPool* FAssetCompilingManager::GetThreadPool() const
 {
 #if WITH_EDITOR
@@ -321,12 +343,15 @@ FQueuedThreadPool* FAssetCompilingManager::GetThreadPool() const
 		// Recently found out that GThreadPool and GLargeThreadPool have the same amount of workers, so can't rely on GThreadPool to be our limiter here.
 		// FPlatformMisc::NumberOfCores() and FPlatformMisc::NumberOfCoresIncludingHyperthreads() also return the same value when -corelimit is used so we can't use FPlatformMisc::NumberOfCores()
 		// if we want to keep the same 1:2 relationship with worker count.
-		const int32 MaxConcurrency = FMath::Max(GLargeThreadPool->GetNumThreads() / 2, 1);
+		auto MaxBackgroundConcurrency = []() { return FMath::Max(1, GLargeThreadPool->GetNumThreads() / 2); };
+
+		// Additional concurrency to unlock at highest priority
+		auto MaxForegroundConcurrency = []() { return FMath::Max(0, GNumForegroundWorkers); };
 
 		// All asset priorities will resolve to a Low priority once being scheduled.
 		// Any asset supporting being built async should be scheduled lower than Normal to let non-async stuff go first
 		// However, we let Highest and Blocking priority pass-through as it to benefit from going to foreground threads when required (i.e. Game-thread is waiting on some assets)
-		GAssetThreadPool = new FMemoryBoundQueuedThreadPoolWrapper(GLargeThreadPool, MaxConcurrency, [](EQueuedWorkPriority Priority) { return Priority <= EQueuedWorkPriority::Highest ? Priority : EQueuedWorkPriority::Low; });
+		GAssetThreadPool = new FMemoryBoundQueuedThreadPoolWrapper(GLargeThreadPool, MaxForegroundConcurrency, MaxBackgroundConcurrency, [](EQueuedWorkPriority Priority) { return Priority <= EQueuedWorkPriority::Highest ? Priority : EQueuedWorkPriority::Low; });
 
 		AsyncCompilationHelpers::BindThreadPoolToCVar(
 			GAssetThreadPool,
