Popup window on Metal platforms when no cooked content at startup

diff --git a/Engine/Source/Runtime/Engine/Private/ShaderCompiler/ShaderCompiler.cpp b/Engine/Source/Runtime/Engine/Private/ShaderCompiler/ShaderCompiler.cpp
--- a/Engine/Source/Runtime/Engine/Private/ShaderCompiler/ShaderCompiler.cpp
+++ b/Engine/Source/Runtime/Engine/Private/ShaderCompiler/ShaderCompiler.cpp
@@ -7110,9 +7110,14 @@ void VerifyGlobalShaders(EShaderPlatform Platform, const ITargetPlatform* Target
 				{
 					if (bErrorOnMissing)
 					{
-						UE_LOG(LogShaders, Fatal, TEXT("Missing global shader %s's permutation %i, Please make sure cooking was successful."),
-							GlobalShaderType->GetName(), PermutationId);
-					}
+                        if (IsMetalPlatform(GMaxRHIShaderPlatform))
+                        {
+                            check(IsInGameThread());
+                            FMessageDialog::Open(EAppMsgType::Ok, LOCTEXT("NoGlobalShader_Error", "Missing shader permutation. Please make sure cooking was successful and refer to Engine log for details."));
+                        }
+                            UE_LOG(LogShaders, Fatal, TEXT("Missing global shader %s's permutation %i, Please make sure cooking was successful."),
+                                GlobalShaderType->GetName(), PermutationId);
+                    }
 					else
 					{
 #if WITH_EDITOR
diff --git a/Engine/Source/Runtime/RenderCore/Private/ShaderCodeLibrary.cpp b/Engine/Source/Runtime/RenderCore/Private/ShaderCodeLibrary.cpp
--- a/Engine/Source/Runtime/RenderCore/Private/ShaderCodeLibrary.cpp
+++ b/Engine/Source/Runtime/RenderCore/Private/ShaderCodeLibrary.cpp
@@ -23,6 +23,7 @@ ShaderCodeLibrary.cpp: Bound shader state cache implementation.
 #include "Misc/ConfigCacheIni.h"
 #include "Misc/CoreDelegates.h"
 #include "Misc/FileHelper.h"
+#include "Misc/MessageDialog.h"
 #include "Misc/Paths.h"
 #include "Misc/ScopeLock.h"
 #include "Misc/ScopeRWLock.h"
@@ -3293,6 +3294,7 @@ void FShaderCodeLibrary::InitForRuntime(EShaderPlatform ShaderPlatform)
 			}
 			else
 			{
+                FMessageDialog::Open(EAppMsgType::Ok, NSLOCTEXT("MessageDialog", "MissingGlobalShaderLibraryFilesClient_Body", "Game files required to initialize the global shader and cooked content are most likely missing. Refer to Engine log for details."));
 				UE_LOG(LogShaderLibrary, Fatal, TEXT("Failed to initialize ShaderCodeLibrary required by the project because part of the Global shader library is missing from %s."), *FPaths::ProjectContentDir());
 			}
 			FPlatformMisc::RequestExit(true);
