Write demo time to a csv stat. This is useful for working out CSV capture start times for deterministic profiling

diff --git a/Engine/Source/Runtime/Engine/Private/DemoNetDriver.cpp b/Engine/Source/Runtime/Engine/Private/DemoNetDriver.cpp
--- a/Engine/Source/Runtime/Engine/Private/DemoNetDriver.cpp
+++ b/Engine/Source/Runtime/Engine/Private/DemoNetDriver.cpp
@@ -2741,7 +2741,8 @@ void UDemoNetDriver::TickDemoPlayback(float DeltaSeconds)
 	{
 		return;
 	}
-	
+	CSV_CUSTOM_STAT(Demo, Time, GetDemoCurrentTime()+DeltaSeconds, ECsvCustomStatOp::Set);
+
 	// This will be true when watching a live replay and we're grabbing an up to date header.
 	// In that case, we want to pause playback until we can actually travel.
 	if (bIsWaitingForHeaderDownload)
