AssetManager GetPackageCookRule: Change NeverCook behavior and add new values to the EPrimaryAssetCookRule enum.
NeverCook, DevelopmentCook, and DevelopmentAlwaysCook now specify exclusion from the cook only for the assets that are directly referenced from the PrimaryAsset. Indirectly assets remain as Unknown.
Inclusion remains unchanged however: DevelopmentAlwaysCook and AlwaysCook still specify that indirectly referenced assts are included in the cook.
EPrimaryAssetCookRule::DevelopmentAlwaysProductionUnknownCook represents the new possibility that an asset is always cooked in development because it is referenced from a DevelopmentAlwaysCook PrimaryAsset, but it MAY still be cooked in Production because it is an indirectly referenced asset from that DevelopmentAlwaysCook PrimaryAsset.
Two other new values in EPrimaryAssetCookRule are just aliases for existing values that are more descriptive of their inclusive/exclusive marking:
DevelopmentCook -> ProductionNeverCook
DevelopmentAlwaysCook -> DevelopmentAlwaysProductionNeverCook.
#rn Minor, AssetManager

diff --git a/Engine/Plugins/Editor/AssetManagerEditor/Source/AssetManagerEditor/Private/AssetManagerEditorModule.cpp b/Engine/Plugins/Editor/AssetManagerEditor/Source/AssetManagerEditor/Private/AssetManagerEditorModule.cpp
--- a/Engine/Plugins/Editor/AssetManagerEditor/Source/AssetManagerEditor/Private/AssetManagerEditorModule.cpp
+++ b/Engine/Plugins/Editor/AssetManagerEditor/Source/AssetManagerEditor/Private/AssetManagerEditorModule.cpp
@@ -1202,13 +1202,16 @@ bool FAssetManagerEditorModule::GetStringValueForCustomColumn(const FAssetData&
 		case EPrimaryAssetCookRule::AlwaysCook: 
 			OutValue = TEXT("Always");
 			return true;
-		case EPrimaryAssetCookRule::DevelopmentAlwaysCook:
-			OutValue = TEXT("DevelopmentAlways");
+		case EPrimaryAssetCookRule::DevelopmentAlwaysProductionUnknownCook:
+			OutValue = TEXT("DevelopmentAlwaysProductionUnknown");
 			return true;
-		case EPrimaryAssetCookRule::DevelopmentCook: 
-			OutValue = TEXT("Development");
+		case EPrimaryAssetCookRule::DevelopmentAlwaysProductionNeverCook:
+			OutValue = TEXT("DevelopmentAlwaysProductionNever");
 			return true;
-		case EPrimaryAssetCookRule::NeverCook: 
+		case EPrimaryAssetCookRule::ProductionNeverCook:
+			OutValue = TEXT("ProductionNeverCook");
+			return true;
+		case EPrimaryAssetCookRule::NeverCook:
 			OutValue = TEXT("Never");
 			return true;
 		}
@@ -1291,11 +1294,14 @@ bool FAssetManagerEditorModule::GetDisplayTextForCustomColumn(const FAssetData&
 		case EPrimaryAssetCookRule::AlwaysCook:
 			OutValue = LOCTEXT("AlwaysCook", "Always");
 			return true;
-		case EPrimaryAssetCookRule::DevelopmentAlwaysCook:
-			OutValue = LOCTEXT("DevelopmentAlwaysCook", "DevelopmentAlways");
+		case EPrimaryAssetCookRule::DevelopmentAlwaysProductionUnknownCook:
+			OutValue = LOCTEXT("DevelopmentAlwaysProductionUnknownCook", "DevelopmentAlwaysProductionUnknown");
+			return true;
+		case EPrimaryAssetCookRule::DevelopmentAlwaysProductionNeverCook:
+			OutValue = LOCTEXT("DevelopmentAlwaysProductionNeverCook", "DevelopmentAlwaysProductionNever");
 			return true;
-		case EPrimaryAssetCookRule::DevelopmentCook:
-			OutValue = LOCTEXT("DevelopmentCook", "Development");
+		case EPrimaryAssetCookRule::ProductionNeverCook:
+			OutValue = LOCTEXT("ProductionNeverCook", "ProductionNever");
 			return true;
 		case EPrimaryAssetCookRule::NeverCook:
 			OutValue = LOCTEXT("NeverCook", "Never");
diff --git a/Engine/Source/Runtime/Engine/Classes/Engine/AssetManager.h b/Engine/Source/Runtime/Engine/Classes/Engine/AssetManager.h
--- a/Engine/Source/Runtime/Engine/Classes/Engine/AssetManager.h
+++ b/Engine/Source/Runtime/Engine/Classes/Engine/AssetManager.h
@@ -387,6 +387,9 @@ public:
 
 	/** Returns list of PrimaryAssetIds that manage a package. Will optionally recurse up the management chain */
 	virtual bool GetPackageManagers(FName PackageName, bool bRecurseToParents, TSet<FPrimaryAssetId>& ManagerSet) const;
+	/** Returns PrimaryAssetIds that manage a package, with property describing the reference (direct or indirect). */
+	virtual bool GetPackageManagers(FName PackageName, bool bRecurseToParents,
+		TMap<FPrimaryAssetId, UE::AssetRegistry::EDependencyProperty>& Managers) const;
 
 
 	// GENERAL ASSET UTILITY FUNCTIONS
@@ -525,6 +528,15 @@ public:
 	/** Returns cook rule for a package name using Management rules, games should override this to take into account their individual workflows */
 	virtual EPrimaryAssetCookRule GetPackageCookRule(FName PackageName) const;
 
+	/**
+	 * Helper function for GetPackageCookRule. Given the list of Managers that manage a package, calculate the unioned cookrule for the
+	 * package. @see FPrimaryAssetCookRuleUnion::UnionWith. If two managers are in conflict (e.g. one is CookAlways, the other is
+	 * ProductionNeverCook), the higher-priority will win. If they have the same priority, the NeverCook rule will win, and the managers
+	 * will be reported in the OutConflictIds field.
+	 */
+	EPrimaryAssetCookRule CalculateCookRuleUnion(const TMap<FPrimaryAssetId, UE::AssetRegistry::EDependencyProperty>& Managers,
+		TOptional<TPair<FPrimaryAssetId, FPrimaryAssetId>>* OutConflictIds) const;
+
 	UE_DEPRECATED(5.0, "Use version that takes ICookInfo instead")
 	virtual bool VerifyCanCookPackage(FName PackageName, bool bLogError = true) const;
 
diff --git a/Engine/Source/Runtime/Engine/Classes/Engine/AssetManagerTypes.h b/Engine/Source/Runtime/Engine/Classes/Engine/AssetManagerTypes.h
--- a/Engine/Source/Runtime/Engine/Classes/Engine/AssetManagerTypes.h
+++ b/Engine/Source/Runtime/Engine/Classes/Engine/AssetManagerTypes.h
@@ -2,37 +2,69 @@
 
 #pragma once
 
-#include "CoreMinimal.h"
-#include "UObject/PrimaryAssetId.h"
 #include "AssetRegistry/AssetIdentifier.h"
-#include "EngineTypes.h"
-#include "AssetManagerTypes.generated.h"
+#include "Containers/Array.h"
+#include "Containers/Set.h"
+#include "Containers/UnrealString.h"
+#include "Delegates/Delegate.h"
+#include "HAL/Platform.h"
+#include "Misc/EnumClassFlags.h"
+#include "Misc/Optional.h"
+#include "Templates/Tuple.h"
+#include "UObject/Class.h"
+#include "UObject/Object.h"
+#include "UObject/ObjectMacros.h"
+#include "UObject/NameTypes.h"
+#include "UObject/ObjectPtr.h"
+#include "UObject/PrimaryAssetId.h"
+#include "UObject/SoftObjectPtr.h"
 
 #if UE_ENABLE_INCLUDE_ORDER_DEPRECATED_IN_5_2
 #include "AssetRegistry/AssetData.h"
 #include "AssetRegistry/AssetBundleData.h"
+#include "CoreMinimal.h"
 #endif
 
+#include "AssetManagerTypes.generated.h"
+
 /** Rule about when to cook/ship a primary asset */
 UENUM()
 enum class EPrimaryAssetCookRule : uint8
 {
-	/** Nothing is known about this asset specifically. It will cook if something else depends on it. */
+	/** Nothing is known about this asset specifically. It will cook in both Development and Production if something else depends on it. */
 	Unknown,
 
 	/** Asset should never be cooked/shipped in any situation. An error will be generated if something depends on it. */
 	NeverCook,
 
 	/** Asset will be cooked in development if something else depends on it, but will never be cooked in a production build. */
-	DevelopmentCook,
+	ProductionNeverCook,
+	/** Legacy name equivalent to Production Never Cook */
+	DevelopmentCook=ProductionNeverCook,
 
 	/** Asset will always be cooked in development, but should never be cooked in a production build. */
-	DevelopmentAlwaysCook,
+	DevelopmentAlwaysProductionNeverCook,
+	/** Legacy name equivalent to DevelopmentAlwaysProductionNeverCook */
+	DevelopmentAlwaysCook = DevelopmentAlwaysProductionNeverCook,
+
+	/**
+	 * Asset will always be cooked in development; nothing is known about whether it should cook in Production. It will cook
+	 * in production if something else depends on it.
+	 */
+	DevelopmentAlwaysProductionUnknownCook,
 
 	/** Asset will always be cooked, in both production and development. */
 	AlwaysCook,
 };
 
+/** The production levels referenced by values of EPrimaryAssetCookRule. */
+enum class EPrimaryAssetProductionLevel
+{
+	Development = 0,
+	Production,
+	Count
+};
+
 /** Structure defining rules for what to do with assets, this is defined per type and can be overridden per asset */
 USTRUCT()
 struct FPrimaryAssetRules
@@ -261,4 +293,55 @@ struct FAssetManagerSearchRules
 
 	/** Returns true if there are any rules set that need to be verified */
 	ENGINE_API bool AreRulesSet() const;
-};
\ No newline at end of file
+};
+
+/** Merges CookRules from multiple managers to decide the final CookRule for an asset they manage. */
+struct FPrimaryAssetCookRuleUnion
+{
+	/**
+	 * Add the information from another manager to the Union.
+	 *
+	 * @param CookRule The CookRule in the manager's FPrimaryAssetRules. CookRules that add references (AlwaysCook,
+	 * DevelopmentAlwaysCook) apply an inclusion for their production level (Production or Development) and all lower
+	 * levels. (e.g. AlwaysCooked applies an inclusion for both Production and Development). CookRules that add
+	 * exclusions (NeverCook, ProductionNeverCook, DevelopmentAlwaysProductionNeverCook) apply an exclusion for their
+	 * production level and all higher levels (e.g. NeverCooked applies an exclusion for both Production and
+	 * Development.) Inclusions for each level are unioned, and exclusions for each level are unioned.
+	 *
+	 * @param bDirectReference Whether the manager's reference to the asset being managed is direct. Inclusions
+	 * (AlwaysCook, DevelopmentAlwaysCook) apply transitively - both direct and indirect references apply the
+	 * inclusion. Exclusions (NeverCook, ProductionNeverCook, DevelopmentAlwaysProductionNeverCook) apply only to
+	 * direct references. e.g. 
+	 * PrimaryAssetId A has DevelopmentAlwaysProductionNeverCook
+	 * A -> B
+	 * B -> C
+	 * B == DevelopmentAlwaysProductionNeverCook
+	 * C == DevelopmentAlwaysProductionUnknownCook
+	 *
+	 * @param Id Record of the Manager that has the reference to the asset, used for error feedback if there is
+	 * a conflict.
+	 *
+	 * @param Priority The Manager's Priority. In the case of a conflict between two managers (one excluding and one
+	 * including at the same level), if the priorities are different this is not an error, and the exclusion or
+	 * inclusion directive from the higher-priority mangager will be kept. If the priorities are equal, the exclusion
+	 * will be kept and a conflict result will be returned from GetRule.
+	 */
+	ENGINE_API void UnionWith(EPrimaryAssetCookRule CookRule, bool bDirectReference, const FPrimaryAssetId& Id, int32 Priority);
+
+	/**
+	 * Get the result of the unions up to this point. If OutConflictIds is non-null, it will be cleared if there
+	 * are no conflicts, and will be set to the Ids that have a conflict between inclusion and exclusion.
+	 */
+	ENGINE_API EPrimaryAssetCookRule GetRule(TOptional<TTuple<FPrimaryAssetId, FPrimaryAssetId>>* OutConflictIds);
+
+private:
+	struct FAssignmentInfo
+	{
+		FPrimaryAssetId Id;
+		int32 Priority = -1;
+		bool bSet = false;
+	};
+
+	FAssignmentInfo InclusionByLevel[(int32)EPrimaryAssetProductionLevel::Count];
+	FAssignmentInfo ExclusionByLevel[(int32)EPrimaryAssetProductionLevel::Count];
+};
diff --git a/Engine/Source/Runtime/Engine/Private/AssetManager.cpp b/Engine/Source/Runtime/Engine/Private/AssetManager.cpp
--- a/Engine/Source/Runtime/Engine/Private/AssetManager.cpp
+++ b/Engine/Source/Runtime/Engine/Private/AssetManager.cpp
@@ -18,6 +18,7 @@
 #include "UObject/LinkerLoad.h"
 #include "UObject/ObjectSaveContext.h"
 #include "UObject/UObjectHash.h"
+#include "Misc/AssertionMacros.h"
 #include "Misc/FileHelper.h"
 #include "Misc/ScopedSlowTask.h"
 #include "Misc/Paths.h"
@@ -3416,37 +3417,66 @@ bool UAssetManager::GetManagedPackageList(FPrimaryAssetId PrimaryAssetId, TArray
 }
 
 bool UAssetManager::GetPackageManagers(FName PackageName, bool bRecurseToParents, TSet<FPrimaryAssetId>& ManagerSet) const
+{
+	TMap<FPrimaryAssetId, UE::AssetRegistry::EDependencyProperty> Managers;
+	bool bFoundAny = GetPackageManagers(PackageName, bRecurseToParents, Managers);
+	for (TPair<FPrimaryAssetId, UE::AssetRegistry::EDependencyProperty>& Pair : Managers)
+	{
+		ManagerSet.Add(Pair.Key);
+	}
+	return bFoundAny;
+}
+
+bool UAssetManager::GetPackageManagers(FName PackageName, bool bRecurseToParents,
+	TMap<FPrimaryAssetId, UE::AssetRegistry::EDependencyProperty>& Managers) const
 {
 	IAssetRegistry& AssetRegistry = GetAssetRegistry();
+	auto UnionManageProperties = [](UE::AssetRegistry::EDependencyProperty A, UE::AssetRegistry::EDependencyProperty B)
+	{
+		return A | B;
+	};
 
 	bool bFoundAny = false;
-	TArray<FAssetIdentifier> ReferencingPrimaryAssets;
+	TArray<FAssetDependency> ReferencingPrimaryAssets;
 	ReferencingPrimaryAssets.Reserve(128);
 
 	AssetRegistry.GetReferencers(PackageName, ReferencingPrimaryAssets, UE::AssetRegistry::EDependencyCategory::Manage);
 
 	for (int32 IdentifierIndex = 0; IdentifierIndex < ReferencingPrimaryAssets.Num(); IdentifierIndex++)
 	{
-		FPrimaryAssetId PrimaryAssetId = ReferencingPrimaryAssets[IdentifierIndex].GetPrimaryAssetId();
-		if (PrimaryAssetId.IsValid())
+		FAssetDependency& AssetDependency = ReferencingPrimaryAssets[IdentifierIndex];
+		FPrimaryAssetId PrimaryAssetId = AssetDependency.AssetId.GetPrimaryAssetId();
+		if (!PrimaryAssetId.IsValid())
 		{
-			bFoundAny = true;
-			ManagerSet.Add(PrimaryAssetId);
+			continue;
+		}
+		bFoundAny = true;
+		UE::AssetRegistry::EDependencyProperty& ExistingProperties = Managers.FindOrAdd(PrimaryAssetId, UE::AssetRegistry::EDependencyProperty::None);
+		ExistingProperties = UnionManageProperties(ExistingProperties, AssetDependency.Properties);
 
-			if (bRecurseToParents)
+		if (bRecurseToParents)
+		{
+			const TArray<FPrimaryAssetId>* ManagementParents = ManagementParentMap.Find(PrimaryAssetId);
+			if (ManagementParents)
 			{
-				const TArray<FPrimaryAssetId> *ManagementParents = ManagementParentMap.Find(PrimaryAssetId);
-
-				if (ManagementParents)
+				// AssetDependency can be invalidated because we are modifying the array it is in. Copy its data now before we modify the array.
+				UE::AssetRegistry::EDependencyCategory IndirectCategory = AssetDependency.Category;
+				UE::AssetRegistry::EDependencyProperty IndirectProperties = AssetDependency.Properties;
+				EnumRemoveFlags(IndirectProperties, UE::AssetRegistry::EDependencyProperty::Direct);
+				for (const FPrimaryAssetId& Manager : *ManagementParents)
 				{
-					for (const FPrimaryAssetId& Manager : *ManagementParents)
+					// Call FindOrAdd with -1 so we can use value != -1 to decide whether it already existed
+					UE::AssetRegistry::EDependencyProperty& ExistingParent = Managers.FindOrAdd(Manager,
+						(UE::AssetRegistry::EDependencyProperty)-1);
+					if (ExistingParent == (UE::AssetRegistry::EDependencyProperty)-1)
 					{
-						if (!ManagerSet.Contains(Manager))
-						{
-							ManagerSet.Add(Manager);
-							// Add to end of list to recurse into the parent
-							ReferencingPrimaryAssets.Add(Manager);
-						}
+						ExistingParent = UE::AssetRegistry::EDependencyProperty::None;
+						// Add to end of list to recurse into the parent.
+						FAssetDependency& Added = ReferencingPrimaryAssets.Emplace_GetRef();
+						Added.AssetId = Manager;
+						Added.Category = IndirectCategory;
+						// Set the parent's property equal to the child's properties, but change it to Indirect
+						Added.Properties = IndirectProperties;
 					}
 				}
 			}
@@ -3973,7 +4003,7 @@ void UAssetManager::ModifyCook(TArray<FName>& PackagesToCook, TArray<FName>& Pac
 				bTargetPlatformsAllowDevelopmentObjects ? TargetPlatform : TargetPlatforms[0];
 			UE_LOG(LogAssetManager, Error,
 				TEXT("Cooking platform %s and %s in a single cook is not supported, because they have different values for AllowsDevelopmentObjects. ")
-				TEXT("This cook session will use AllowsDevelopmentObjects = true, which will add packages to platform % s that should not be present."),
+				TEXT("This cook session will use AllowsDevelopmentObjects = true, which will add packages to platform %s that should not be present."),
 				*TargetPlatforms[0]->PlatformName(), *TargetPlatform->PlatformName(),
 				*PlatformThatDoesNotAllow->PlatformName());
 			bTargetPlatformsAllowDevelopmentObjects = true;
@@ -3988,6 +4018,8 @@ void UAssetManager::ModifyCook(TArray<FName>& PackagesToCook, TArray<FName>& Pac
 
 	GetPrimaryAssetTypeInfoList(TypeList);
 
+	bool bIncludeDevelopmentAssets = !bOnlyCookProductionAssets || bTargetPlatformsAllowDevelopmentObjects;
+
 	// Get package names in the libraries that we care about for cooking. Only get ones that are needed in production
 	for (const FPrimaryAssetTypeInfo& TypeInfo : TypeList)
 	{
@@ -4023,9 +4055,10 @@ void UAssetManager::ModifyCook(TArray<FName>& PackagesToCook, TArray<FName>& Pac
 		for (FName PackageName : AssetPackages)
 		{
 			EPrimaryAssetCookRule CookRule = GetPackageCookRule(PackageName);
-
-			// Treat DevAlwaysCook as AlwaysCook, may get excluded in VerifyCanCookPackage
-			bool bAlwaysCook = (CookRule == EPrimaryAssetCookRule::AlwaysCook || CookRule == EPrimaryAssetCookRule::DevelopmentAlwaysCook);
+			bool bAlwaysCook = CookRule == EPrimaryAssetCookRule::AlwaysCook ||
+				(bIncludeDevelopmentAssets && (
+					CookRule == EPrimaryAssetCookRule::DevelopmentAlwaysProductionUnknownCook ||
+					CookRule == EPrimaryAssetCookRule::DevelopmentAlwaysProductionNeverCook));
 			bool bCanCook = VerifyCanCookPackage(nullptr, PackageName, false);
 
 			if (bAlwaysCook && bCanCook && !TypeInfo.bIsEditorOnly)
@@ -4125,34 +4158,155 @@ bool UAssetManager::ShouldCookForPlatform(const UPackage* Package, const ITarget
 EPrimaryAssetCookRule UAssetManager::GetPackageCookRule(FName PackageName) const
 {
 	TRACE_CPUPROFILER_EVENT_SCOPE(UAssetManager::GetPackageCookRule);
-	FPrimaryAssetRules BestRules;
-	FPrimaryAssetId BestId;
-	TSet<FPrimaryAssetId> Managers;
+
+	TMap<FPrimaryAssetId, UE::AssetRegistry::EDependencyProperty> Managers;
 	GetPackageManagers(PackageName, true, Managers);
 
-	for (const FPrimaryAssetId& PrimaryAssetId : Managers)
+	TOptional<TTuple<FPrimaryAssetId, FPrimaryAssetId>> ConflictIds;
+	EPrimaryAssetCookRule CookRule = CalculateCookRuleUnion(Managers, &ConflictIds);
+	if (ConflictIds)
 	{
-		FPrimaryAssetRules Rules = GetPrimaryAssetRules(PrimaryAssetId);
+		UE_LOG(LogAssetManager, Error, TEXT("GetPackageCookRule: Conflicting Cook Rule for package %s! %s and %s have the same priority and disagree."),
+			*PackageName.ToString(), *ConflictIds->Get<0>().ToString(), *ConflictIds->Get<1>().ToString());
+	}
+
+	return CookRule;
+}
+
+EPrimaryAssetCookRule UAssetManager::CalculateCookRuleUnion(const TMap<FPrimaryAssetId, UE::AssetRegistry::EDependencyProperty>& Managers,
+	TOptional<TPair<FPrimaryAssetId, FPrimaryAssetId>>* OutConflictIds) const
+{
+	FPrimaryAssetCookRuleUnion Union;
+	for (const TPair<FPrimaryAssetId, UE::AssetRegistry::EDependencyProperty>& Pair : Managers)
+	{
+		FPrimaryAssetRules Rules = GetPrimaryAssetRules(Pair.Key);
+		bool bDirect = EnumHasAllFlags(Pair.Value, UE::AssetRegistry::EDependencyProperty::Direct);
+		Union.UnionWith(Rules.CookRule, bDirect, Pair.Key, Rules.Priority);
+	}
+
+	return Union.GetRule(OutConflictIds);
+}
 
-		if (Rules.CookRule != EPrimaryAssetCookRule::Unknown && Rules.CookRule != BestRules.CookRule)
+void FPrimaryAssetCookRuleUnion::UnionWith(EPrimaryAssetCookRule CookRule, bool bDirectReference, const FPrimaryAssetId& Id, int32 Priority)
+{
+	auto MarkExcluded = [this, bDirectReference, &Id, Priority](EPrimaryAssetProductionLevel LowestLevelToExclude)
+	{
+		// Exclusion only applies to direct references
+		if (bDirectReference)
 		{
-			if (BestRules.CookRule == EPrimaryAssetCookRule::Unknown || Rules.Priority > BestRules.Priority)
+			for (int32 LevelInt = (int32)LowestLevelToExclude; LevelInt < (int32)EPrimaryAssetProductionLevel::Count; ++LevelInt)
 			{
-				BestRules = Rules;
-				BestId = PrimaryAssetId;
-			}
-			else
-			{
-				// Lower priority, ignore
-				if (BestRules.Priority == Rules.Priority)
+				FAssignmentInfo& Info = ExclusionByLevel[LevelInt];
+				if (!Info.bSet || Info.Priority < Priority)
 				{
-					UE_LOG(LogAssetManager, Error, TEXT("GetPackageCookRule: Conflicting Cook Rule for package %s! %s and %s have the same priority and disagree."), *PackageName.ToString(), *PrimaryAssetId.ToString(), *BestId.ToString());
+					Info.bSet = true;
+					Info.Priority = Priority;
+					Info.Id = Id;
 				}
 			}
 		}
+	};
+	auto MarkIncluded = [this, &Id, Priority](EPrimaryAssetProductionLevel HighestLevelToInclude)
+	{
+		// Referenced applies to direct and indirect references
+		for (int32 LevelInt = 0; LevelInt <= (int32)HighestLevelToInclude; ++LevelInt)
+		{
+			FAssignmentInfo& Info = InclusionByLevel[LevelInt];
+			if (!Info.bSet || Info.Priority < Priority)
+			{
+				Info.bSet = true;
+				Info.Priority = Priority;
+				Info.Id = Id;
+			}
+		}
+	};
+
+	switch (CookRule)
+	{
+	case EPrimaryAssetCookRule::Unknown:
+		// Managers with CookRule Unknown are only used to define Chunks for assets that other managers include
+		// and do not affect whether the Asset should be cooked
+		break;
+	case EPrimaryAssetCookRule::NeverCook:
+		// Managers with NeverCook require that the asset is NOT cooked in production or development builds,
+		// but only for direct references
+		MarkExcluded(EPrimaryAssetProductionLevel::Development);
+		break;
+	case EPrimaryAssetCookRule::ProductionNeverCook:
+		// Managers with ProductionNeverCook (1) do not imply the asset should be cooked for development but
+		// (2) DO require that the asset is NOT cooked in production builds, but only for direct references
+		MarkExcluded(EPrimaryAssetProductionLevel::Production);
+		break;
+	case EPrimaryAssetCookRule::DevelopmentAlwaysProductionNeverCook:
+		// Managers with DevelopmentAlwaysProductionNeverCook (1) require the asset should be cooked in development builds
+		// (2) require that the asset is NOT cooked in production builds, but only for direct references
+		MarkIncluded(EPrimaryAssetProductionLevel::Development);
+		MarkExcluded(EPrimaryAssetProductionLevel::Production);
+		break;
+	case EPrimaryAssetCookRule::DevelopmentAlwaysProductionUnknownCook:
+		// Managers with DevelopmentAlwaysProductionUnknownCook (1) require the asset should be cooked in development builds
+		// (2) imply neither inclusion nor exclusion for production builds
+		MarkIncluded(EPrimaryAssetProductionLevel::Development);
+		break;
+	case EPrimaryAssetCookRule::AlwaysCook:
+		// Managers with AlwaysCook (1) require the asset should be cooked in production and development
+		MarkIncluded(EPrimaryAssetProductionLevel::Production);
+		break;
+	default:
+		checkNoEntry();
+		break;
+	}
+}
+
+EPrimaryAssetCookRule FPrimaryAssetCookRuleUnion::GetRule(TOptional<TTuple<FPrimaryAssetId, FPrimaryAssetId>>* OutConflictId)
+{
+	if (OutConflictId)
+	{
+		OutConflictId->Reset();
+	}
+
+	FAssignmentInfo& ProdInclusion= InclusionByLevel[(int32)EPrimaryAssetProductionLevel::Production];
+	FAssignmentInfo& ProdExclusion= ExclusionByLevel[(int32)EPrimaryAssetProductionLevel::Production];
+	FAssignmentInfo& DevInclusion = InclusionByLevel[(int32)EPrimaryAssetProductionLevel::Development];
+	FAssignmentInfo& DevExclusion = ExclusionByLevel[(int32)EPrimaryAssetProductionLevel::Development];
+	if (DevExclusion.bSet && (!DevInclusion.bSet || DevInclusion.Priority <= DevExclusion.Priority))
+	{
+		if (DevInclusion.bSet && DevInclusion.Priority == DevExclusion.Priority)
+		{
+			OutConflictId->Emplace(DevExclusion.Id, DevInclusion.Id);
+		}
+		return EPrimaryAssetCookRule::NeverCook;
 	}
 
-	return BestRules.CookRule;
+	if (ProdExclusion.bSet && (!ProdInclusion.bSet || ProdInclusion.Priority <= ProdExclusion.Priority))
+	{
+		if (ProdInclusion.bSet && ProdInclusion.Priority == ProdExclusion.Priority)
+		{
+			OutConflictId->Emplace(ProdExclusion.Id, ProdInclusion.Id);
+		}
+
+		if (!DevInclusion.bSet)
+		{
+			return EPrimaryAssetCookRule::ProductionNeverCook;
+		}
+		else
+		{
+			return EPrimaryAssetCookRule::DevelopmentAlwaysProductionNeverCook;
+		}
+	}
+
+	if (!DevInclusion.bSet)
+	{
+		return EPrimaryAssetCookRule::Unknown;
+	}
+	else if (!ProdInclusion.bSet)
+	{
+		return EPrimaryAssetCookRule::DevelopmentAlwaysProductionUnknownCook;
+	}
+	else
+	{
+		return EPrimaryAssetCookRule::AlwaysCook;
+	}
 }
 
 static FString GetInstigatorChainString(UE::Cook::ICookInfo* CookInfo, FName PackageName)
@@ -4200,12 +4354,12 @@ bool UAssetManager::VerifyCanCookPackage(FName PackageName, bool bLogError) cons
 		
 		bRetVal = false;
 	}
-	else if ((CookRule == EPrimaryAssetCookRule::DevelopmentCook || CookRule == EPrimaryAssetCookRule::DevelopmentAlwaysCook)
+	else if ((CookRule == EPrimaryAssetCookRule::ProductionNeverCook || CookRule == EPrimaryAssetCookRule::DevelopmentAlwaysProductionNeverCook)
 		&& bOnlyCookProductionAssets && !bTargetPlatformsAllowDevelopmentObjects)
 	{
 		if (bLogError)
 		{
-			UE_LOG(LogAssetManager, Warning, TEXT("Package %s is set to Development, but bOnlyCookProductionAssets is true! Instigators: %s"),
+			UE_LOG(LogAssetManager, Warning, TEXT("Package %s is set to ProductionNeverCook, and bOnlyCookProductionAssets is true, but something is trying to cook it! Instigators: %s"),
 				*PackageName.ToString(), *GetInstigatorChainString(CookInfo, PackageName));
 		}
 
