Track owner path names on textures to show in Render Resource Viewer

diff --git a/Engine/Source/Runtime/Engine/Private/Rendering/StreamableTextureResource.cpp b/Engine/Source/Runtime/Engine/Private/Rendering/StreamableTextureResource.cpp
--- a/Engine/Source/Runtime/Engine/Private/Rendering/StreamableTextureResource.cpp
+++ b/Engine/Source/Runtime/Engine/Private/Rendering/StreamableTextureResource.cpp
@@ -190,6 +190,7 @@ void FStreamableTextureResource::InitRHI()
 	// Update mip-level fading.
 	MipBiasFade.SetNewMipCount( State.NumRequestedLODs, State.NumRequestedLODs, LastRenderTime, MipFadeSetting );
 
+	TextureRHI->SetOwnerName(GetOwnerName());
 	TextureRHI->SetName(TextureName);
 	RHIBindDebugLabelName(TextureRHI, *TextureName.ToString());
 
@@ -230,6 +231,7 @@ void FStreamableTextureResource::FinalizeStreaming(FRHITexture* InTextureRHI)
 	}
 
 	TextureRHI = InTextureRHI;
+	TextureRHI->SetOwnerName(GetOwnerName());
 	if (ensure(TextureReferenceRHI.IsValid()))
 	{
 		RHIUpdateTextureReference(TextureReferenceRHI, TextureRHI);
diff --git a/Engine/Source/Runtime/Engine/Private/Texture.cpp b/Engine/Source/Runtime/Engine/Private/Texture.cpp
--- a/Engine/Source/Runtime/Engine/Private/Texture.cpp
+++ b/Engine/Source/Runtime/Engine/Private/Texture.cpp
@@ -345,6 +345,8 @@ void UTexture::UpdateResource()
 			{
 				NewResource->SetTextureReference(TextureReference.TextureReferenceRHI);
 			});
+
+			NewResource->SetOwnerName(FName(GetPathName()));
 			BeginInitResource(NewResource);
 			// Now that the resource is ready for streaming, bind it to the streamer.
 			LinkStreaming();
diff --git a/Engine/Source/Runtime/Engine/Private/Texture2D.cpp b/Engine/Source/Runtime/Engine/Private/Texture2D.cpp
--- a/Engine/Source/Runtime/Engine/Private/Texture2D.cpp
+++ b/Engine/Source/Runtime/Engine/Private/Texture2D.cpp
@@ -1358,6 +1358,8 @@ void FVirtualTexture2DResource::InitRHI()
 #if WITH_EDITOR
 	InitializeEditorResources(VirtualTexture);
 #endif
+
+	TextureRHI->SetOwnerName(GetOwnerName());
 }
 
 #if WITH_EDITOR
diff --git a/Engine/Source/Runtime/Engine/Private/TextureCube.cpp b/Engine/Source/Runtime/Engine/Private/TextureCube.cpp
--- a/Engine/Source/Runtime/Engine/Private/TextureCube.cpp
+++ b/Engine/Source/Runtime/Engine/Private/TextureCube.cpp
@@ -457,6 +457,7 @@ public:
 		TextureCubeRHI = RHICreateTexture(Desc);
 
 		TextureRHI = TextureCubeRHI;
+		TextureRHI->SetOwnerName(GetOwnerName());
 		TextureRHI->SetName(Owner->GetFName());
 		RHIBindDebugLabelName(TextureRHI, *Name);
 		RHIUpdateTextureReference(Owner->TextureReference.TextureReferenceRHI,TextureRHI);
diff --git a/Engine/Source/Runtime/Engine/Private/TextureCubeArray.cpp b/Engine/Source/Runtime/Engine/Private/TextureCubeArray.cpp
--- a/Engine/Source/Runtime/Engine/Private/TextureCubeArray.cpp
+++ b/Engine/Source/Runtime/Engine/Private/TextureCubeArray.cpp
@@ -137,6 +137,7 @@ public:
 
 		TextureCubeRHI = RHICreateTexture(Desc);
 		TextureRHI = TextureCubeRHI;
+		TextureRHI->SetOwnerName(GetOwnerName());
 		TextureRHI->SetName(Owner->GetFName());
 
 		RHIBindDebugLabelName(TextureRHI, *Name);
