Set owner name to morph vertex buffer for resource tracking

diff --git a/Engine/Source/Runtime/Engine/Private/SkeletalRenderGPUSkin.cpp b/Engine/Source/Runtime/Engine/Private/SkeletalRenderGPUSkin.cpp
--- a/Engine/Source/Runtime/Engine/Private/SkeletalRenderGPUSkin.cpp
+++ b/Engine/Source/Runtime/Engine/Private/SkeletalRenderGPUSkin.cpp
@@ -107,6 +107,7 @@ void FMorphVertexBuffer::InitDynamicRHI()
 	Flags = (EBufferUsageFlags)(Flags | BUF_ShaderResource);
 
 	VertexBufferRHI = RHICreateVertexBuffer(Size, Flags, CreateInfo);
+	VertexBufferRHI->SetOwnerName(GetOwnerName());
 	SRVValue = RHICreateShaderResourceView(VertexBufferRHI, 4, PF_R32_FLOAT);
 
 	if (!bUseGPUMorphTargets)
@@ -139,8 +140,11 @@ void FMorphVertexBuffer::ReleaseDynamicRHI()
 /*-----------------------------------------------------------------------------
 FMorphVertexBufferPool
 -----------------------------------------------------------------------------*/
-void FMorphVertexBufferPool::InitResources()
+void FMorphVertexBufferPool::InitResources(const FName& OwnerName)
 {
+	MorphVertexBuffers[0].SetOwnerName(OwnerName);
+	MorphVertexBuffers[1].SetOwnerName(OwnerName);
+
 	check(!MorphVertexBuffers[0].VertexBufferRHI.IsValid());
 	check(!MorphVertexBuffers[1].VertexBufferRHI.IsValid());
 	BeginInitResource(&MorphVertexBuffers[0]);
@@ -2177,7 +2181,8 @@ void FSkeletalMeshObjectGPUSkin::FSkeletalMeshObjectLOD::InitMorphResources(cons
 	FSkeletalMeshLODRenderData& LODData = SkelMeshRenderData->LODRenderData[LODIndex];
 
 	// init the delta vertex buffer for this LOD
-	MorphVertexBufferPool.InitResources();
+	const FName OwnerName = LODData.MorphTargetVertexInfoBuffers.GetOwnerName();
+	MorphVertexBufferPool.InitResources(OwnerName);
 
 	// Vertex buffers available for the LOD
 	FVertexFactoryBuffers VertexBuffers;
diff --git a/Engine/Source/Runtime/Engine/Private/SkeletalRenderGPUSkin.h b/Engine/Source/Runtime/Engine/Private/SkeletalRenderGPUSkin.h
--- a/Engine/Source/Runtime/Engine/Private/SkeletalRenderGPUSkin.h
+++ b/Engine/Source/Runtime/Engine/Private/SkeletalRenderGPUSkin.h
@@ -311,7 +311,7 @@ public:
 		MorphVertexBuffers[1] = FMorphVertexBuffer(InSkelMeshRenderData, InLOD, InFeatureLevel);
 	}
 
-	void InitResources();
+	void InitResources(const FName& OwnerName);
 	void ReleaseResources();
 	SIZE_T GetResourceSize() const;
 	void EnableDoubleBuffer();
