Set skeletal mesh duplicated vertex buffer owner name for resource tracking

diff --git a/Engine/Source/Runtime/Engine/Public/Rendering/SkeletalMeshDuplicatedVerticesBuffer.h b/Engine/Source/Runtime/Engine/Public/Rendering/SkeletalMeshDuplicatedVerticesBuffer.h
--- a/Engine/Source/Runtime/Engine/Public/Rendering/SkeletalMeshDuplicatedVerticesBuffer.h
+++ b/Engine/Source/Runtime/Engine/Public/Rendering/SkeletalMeshDuplicatedVerticesBuffer.h
@@ -113,6 +113,7 @@ public:
 
             FRHIResourceCreateInfo CreateInfo(TEXT("DuplicatedVerticesIndexBuffer"), ResourceArray);
             DuplicatedVerticesIndexBuffer.VertexBufferRHI = RHICreateVertexBuffer(ResourceArray->GetResourceDataSize(), BUF_Static | BUF_ShaderResource, CreateInfo);
+			DuplicatedVerticesIndexBuffer.VertexBufferRHI->SetOwnerName(GetOwnerName());
             DuplicatedVerticesIndexBuffer.VertexBufferSRV = RHICreateShaderResourceView(DuplicatedVerticesIndexBuffer.VertexBufferRHI, sizeof(uint32), PF_R32_UINT);
         }
 
@@ -122,6 +123,7 @@ public:
 
             FRHIResourceCreateInfo CreateInfo(TEXT("LengthAndIndexDuplicatedVerticesIndexBuffer"), ResourceArray);
             LengthAndIndexDuplicatedVerticesIndexBuffer.VertexBufferRHI = RHICreateVertexBuffer(ResourceArray->GetResourceDataSize(), BUF_Static | BUF_ShaderResource, CreateInfo);
+			LengthAndIndexDuplicatedVerticesIndexBuffer.VertexBufferRHI->SetOwnerName(GetOwnerName());
             LengthAndIndexDuplicatedVerticesIndexBuffer.VertexBufferSRV = RHICreateShaderResourceView(LengthAndIndexDuplicatedVerticesIndexBuffer.VertexBufferRHI, sizeof(uint32), PF_R32_UINT);
         }
     }
