Special engine material DDC key logging downgraded to `Log` channel messaging.

diff --git a/Engine/Source/Runtime/Engine/Private/Materials/MaterialShader.cpp b/Engine/Source/Runtime/Engine/Private/Materials/MaterialShader.cpp
--- a/Engine/Source/Runtime/Engine/Private/Materials/MaterialShader.cpp
+++ b/Engine/Source/Runtime/Engine/Private/Materials/MaterialShader.cpp
@@ -1420,7 +1420,7 @@ TSharedRef<FMaterialShaderMap::FAsyncLoadContext> FMaterialShaderMap::BeginLoadF
 			{
 				FString SpecialEngineDDCKey = Result->DataKey;
 				SpecialEngineDDCKey.RemoveFromEnd(TEXT("\n"));
-				UE_LOG(LogMaterial, Display, TEXT("%s-%s-%s: %s"), *Material->GetAssetName(), *LexToString(ShaderMapId.FeatureLevel), *LexToString(ShaderMapId.QualityLevel), *SpecialEngineDDCKey);
+				UE_LOG(LogMaterial, Log, TEXT("%s-%s-%s: %s"), *Material->GetAssetName(), *LexToString(ShaderMapId.FeatureLevel), *LexToString(ShaderMapId.QualityLevel), *SpecialEngineDDCKey);
 			}
 
 			bool bCheckCache = true;
