Add a FInputDeviceId option for setting input device properties. The default behavior will remain the same. If one is not specifed, then we will default to the Platform User's primary input device.

#rnx
#pf skip

diff --git a/Engine/Plugins/Runtime/InputDebugging/Source/InputDebugging/Private/InputDeviceDebugTools.cpp b/Engine/Plugins/Runtime/InputDebugging/Source/InputDebugging/Private/InputDeviceDebugTools.cpp
--- a/Engine/Plugins/Runtime/InputDebugging/Source/InputDebugging/Private/InputDeviceDebugTools.cpp
+++ b/Engine/Plugins/Runtime/InputDebugging/Source/InputDebugging/Private/InputDeviceDebugTools.cpp
@@ -122,6 +122,7 @@ void FInputDeviceDebugTools::OnShowDebugDeviceProperties(UCanvas* Canvas)
 			// Draw some information about the flags that have been set on the given property
 			DisplayDebugManager.SetDrawColor(FColor::Silver);
 			DisplayDebugManager.DrawString(FString::Printf(TEXT("\t Property Handle: %s"), *ActiveProperty->PropertyHandle.ToString()));
+			DisplayDebugManager.DrawString(FString::Printf(TEXT("\t Playing on DeviceId: %d"), ActiveProperty->DeviceId.GetId()));
 			DisplayDebugManager.DrawString(FString::Printf(TEXT("\t Flags: \n \t \t \t bLooping: %d \n \t \t \t bIgnoreTimeDilation: %d \n \t \t \t bPlayWhilePaused: %d \n \t \t \t bHasBeenAppliedAtLeastOnce: %d"),
 				ActiveProperty->bLooping, ActiveProperty->bIgnoreTimeDilation, ActiveProperty->bPlayWhilePaused, ActiveProperty->bHasBeenAppliedAtLeastOnce));
 
diff --git a/Engine/Source/Runtime/Engine/Classes/GameFramework/InputDeviceProperties.h b/Engine/Source/Runtime/Engine/Classes/GameFramework/InputDeviceProperties.h
--- a/Engine/Source/Runtime/Engine/Classes/GameFramework/InputDeviceProperties.h
+++ b/Engine/Source/Runtime/Engine/Classes/GameFramework/InputDeviceProperties.h
@@ -50,7 +50,7 @@ protected:
 	* @return					A pointer to the evaluated input device property.
 	*/
 	UFUNCTION(BlueprintNativeEvent, Category = "InputDevice")
-	void EvaluateDeviceProperty(const FPlatformUserId PlatformUser, const float DeltaTime, const float Duration);
+	void EvaluateDeviceProperty(const FPlatformUserId PlatformUser, const FInputDeviceId DeviceId, const float DeltaTime, const float Duration);
 
 	/** 
 	* Native C++ implementation of EvaluateDeviceProperty.
@@ -58,7 +58,7 @@ protected:
 	* Override this to alter your device property in native code.
 	* @see UInputDeviceProperty::EvaluateDeviceProperty
 	*/
-	virtual void EvaluateDeviceProperty_Implementation(const FPlatformUserId PlatformUser, const float DeltaTime, const float Duration);
+	virtual void EvaluateDeviceProperty_Implementation(const FPlatformUserId PlatformUser, const FInputDeviceId DeviceId, const float DeltaTime, const float Duration);
 
 	/**
 	* Reset the current device property. Provides an opportunity to reset device state after evaluation is complete. 
@@ -67,7 +67,7 @@ protected:
 	* @param PlatformUser		The platform user that should receive this device property change
 	*/
 	UFUNCTION(BlueprintNativeEvent, Category = "InputDevice")
-	void ResetDeviceProperty(const FPlatformUserId PlatformUser);
+	void ResetDeviceProperty(const FPlatformUserId PlatformUser, const FInputDeviceId DeviceId);
 
 	/**
 	* Native C++ implementation of ResetDeviceProperty
@@ -75,7 +75,7 @@ protected:
 	* 
 	* @see ResetDeviceProperty
 	*/
-	virtual void ResetDeviceProperty_Implementation(const FPlatformUserId PlatformUser);
+	virtual void ResetDeviceProperty_Implementation(const FPlatformUserId PlatformUser, const FInputDeviceId DeviceId);
 
 	/**
 	* Apply the device property from GetInternalDeviceProperty to the given platform user. 
@@ -84,7 +84,7 @@ protected:
 	* @param UserId		The owning Platform User whose input device this property should be applied to.
 	*/
 	UFUNCTION(Category = "InputDevice")
-	virtual void ApplyDeviceProperty(const FPlatformUserId UserId);
+	virtual void ApplyDeviceProperty(const FPlatformUserId UserId, const FInputDeviceId DeviceId);
 	
 	/** Gets a pointer to the current input device property that the IInputInterface can use. */
 	virtual FInputDeviceProperty* GetInternalDeviceProperty() { return nullptr; };
@@ -116,14 +116,14 @@ protected:
 	* @param UserId			The owning Platform User whose input device this property should be applied to.
 	* @param RawProperty	The internal input device property to apply.
 	*/
-	static void ApplyDeviceProperty_Internal(const FPlatformUserId UserId, FInputDeviceProperty* RawProperty);
+	static void ApplyDeviceProperty_Internal(const FPlatformUserId UserId, const FInputDeviceId DeviceId, FInputDeviceProperty* RawProperty);
 
 	/** Returns the device specific data for the given platform user. Returns the default data if none are given */
 	template<class TDataLayout>
-	const TDataLayout* GetDeviceSpecificData(const FPlatformUserId UserId, const TMap<FName, TDataLayout>& InMap) const;
+	const TDataLayout* GetDeviceSpecificData(const FPlatformUserId UserId, const FInputDeviceId DeviceId, const TMap<FName, TDataLayout>& InMap) const;
 
 	template<class TDataLayout>
-	TDataLayout* GetDeviceSpecificDataMutable(const FPlatformUserId UserId, TMap<FName, TDataLayout>& InMap)
+	TDataLayout* GetDeviceSpecificDataMutable(const FPlatformUserId UserId, const FInputDeviceId DeviceId, TMap<FName, TDataLayout>& InMap)
 	{
 		return GetDeviceSpecificData<TDataLayout>(UserId, InMap);
 	}
@@ -137,7 +137,7 @@ protected:
 };
 
 template<class TDataLayout>
-const TDataLayout* UInputDeviceProperty::GetDeviceSpecificData(const FPlatformUserId UserId, const TMap<FName, TDataLayout>& InDeviceData) const
+const TDataLayout* UInputDeviceProperty::GetDeviceSpecificData(const FPlatformUserId UserId, const FInputDeviceId DeviceId, const TMap<FName, TDataLayout>& InDeviceData) const
 {
 	if (const UInputDeviceSubsystem* SubSystem = UInputDeviceSubsystem::Get())
 	{
@@ -184,7 +184,7 @@ class UColorInputDeviceProperty : public UInputDeviceProperty
 
 public:
 
-	virtual void EvaluateDeviceProperty_Implementation(const FPlatformUserId PlatformUser, const float DeltaTime, const float Duration) override;
+	virtual void EvaluateDeviceProperty_Implementation(const FPlatformUserId PlatformUser, const FInputDeviceId DeviceId, const float DeltaTime, const float Duration) override;
 	virtual FInputDeviceProperty* GetInternalDeviceProperty() override;
 
 	/** Default color data that will be used by default. Device Specific overrides will be used when the current input device matches */
@@ -237,8 +237,8 @@ class UColorInputDeviceCurveProperty : public UInputDeviceProperty
 
 public:
 
-	virtual void EvaluateDeviceProperty_Implementation(const FPlatformUserId PlatformUser, const float DeltaTime, const float Duration) override;
-	virtual void ResetDeviceProperty_Implementation(const FPlatformUserId PlatformUser) override;
+	virtual void EvaluateDeviceProperty_Implementation(const FPlatformUserId PlatformUser, const FInputDeviceId DeviceId, const float DeltaTime, const float Duration) override;
+	virtual void ResetDeviceProperty_Implementation(const FPlatformUserId PlatformUser, const FInputDeviceId DeviceId) override;
 	virtual FInputDeviceProperty* GetInternalDeviceProperty() override;
 	virtual float RecalculateDuration() override;
 
@@ -284,7 +284,7 @@ class ENGINE_API UInputDeviceTriggerEffect : public UInputDeviceProperty
 public:	
 
 	virtual FInputDeviceProperty* GetInternalDeviceProperty() override;
-	virtual void ResetDeviceProperty_Implementation(const FPlatformUserId PlatformUser) override;
+	virtual void ResetDeviceProperty_Implementation(const FPlatformUserId PlatformUser, const FInputDeviceId DeviceId) override;
 
 	UPROPERTY(EditAnywhere, BlueprintReadWrite, Category = "DeviceProperty")
 	FDeviceTriggerBaseData BaseTriggerData;
@@ -327,7 +327,7 @@ public:
 	
 	UInputDeviceTriggerFeedbackProperty();
 	
-	virtual void EvaluateDeviceProperty_Implementation(const FPlatformUserId PlatformUser, const float DeltaTime, const float Duration) override;	
+	virtual void EvaluateDeviceProperty_Implementation(const FPlatformUserId PlatformUser, const FInputDeviceId DeviceId, const float DeltaTime, const float Duration) override;	
 	virtual FInputDeviceProperty* GetInternalDeviceProperty() override;
 	virtual float RecalculateDuration() override;
 
@@ -386,7 +386,7 @@ public:
 
 	UInputDeviceTriggerResistanceProperty();
 
-	virtual void EvaluateDeviceProperty_Implementation(const FPlatformUserId PlatformUser, const float DeltaTime, const float Duration) override;
+	virtual void EvaluateDeviceProperty_Implementation(const FPlatformUserId PlatformUser, const FInputDeviceId DeviceId, const float DeltaTime, const float Duration) override;
 	virtual FInputDeviceProperty* GetInternalDeviceProperty() override;
 
 protected:
@@ -443,7 +443,7 @@ public:
 
 	UInputDeviceTriggerVibrationProperty();
 
-	virtual void EvaluateDeviceProperty_Implementation(const FPlatformUserId PlatformUser, const float DeltaTime, const float Duration) override;
+	virtual void EvaluateDeviceProperty_Implementation(const FPlatformUserId PlatformUser, const FInputDeviceId DeviceId, const float DeltaTime, const float Duration) override;
 	virtual FInputDeviceProperty* GetInternalDeviceProperty() override;
 	virtual float RecalculateDuration() override;
 
diff --git a/Engine/Source/Runtime/Engine/Classes/GameFramework/InputDeviceSubsystem.h b/Engine/Source/Runtime/Engine/Classes/GameFramework/InputDeviceSubsystem.h
--- a/Engine/Source/Runtime/Engine/Classes/GameFramework/InputDeviceSubsystem.h
+++ b/Engine/Source/Runtime/Engine/Classes/GameFramework/InputDeviceSubsystem.h
@@ -26,6 +26,15 @@ struct ENGINE_API FSetDevicePropertyParams
 	UPROPERTY(EditAnywhere, BlueprintReadWrite, Category = "Input Devices")
 	FPlatformUserId UserId;
 
+	/** 
+	* The Input Device that should receive the device property. If nothing is specified here,
+	* then the Platform User's default input device will be used. 
+	* 
+	* The default input device is obtained from IPlatformInputDeviceMapper::GetPrimaryInputDeviceForUser
+	*/
+	UPROPERTY(EditAnywhere, BlueprintReadWrite, Category = "Input Devices")
+	FInputDeviceId DeviceId;
+
 	/**
 	* If true, then the input device property will not be removed after it's evaluation time has completed.
 	* Instead, it will remain active until manually removed with a RemoveDeviceProperty call.
@@ -73,6 +82,9 @@ struct ENGINE_API FActiveDeviceProperty
 	/** The platform user that is actively receiving this device property */
 	FPlatformUserId PlatformUser = PLATFORMUSERID_NONE;
 
+	/** The input device id that should receive this property. */
+	FInputDeviceId DeviceId = INPUTDEVICEID_NONE;
+
 	/** The handle of this active property. */
 	FInputDevicePropertyHandle PropertyHandle = FInputDevicePropertyHandle::InvalidHandle;
 
diff --git a/Engine/Source/Runtime/Engine/Private/GameFramework/InputDeviceProperties.cpp b/Engine/Source/Runtime/Engine/Private/GameFramework/InputDeviceProperties.cpp
--- a/Engine/Source/Runtime/Engine/Private/GameFramework/InputDeviceProperties.cpp
+++ b/Engine/Source/Runtime/Engine/Private/GameFramework/InputDeviceProperties.cpp
@@ -17,12 +17,12 @@ UInputDeviceProperty::UInputDeviceProperty(const FObjectInitializer& ObjectIniti
 	RecalculateDuration();
 }
 
-void UInputDeviceProperty::ApplyDeviceProperty(const FPlatformUserId UserId)
+void UInputDeviceProperty::ApplyDeviceProperty(const FPlatformUserId UserId, const FInputDeviceId DeviceId)
 {
-	UInputDeviceProperty::ApplyDeviceProperty_Internal(UserId, GetInternalDeviceProperty());
+	UInputDeviceProperty::ApplyDeviceProperty_Internal(UserId, DeviceId, GetInternalDeviceProperty());
 }
 
-void UInputDeviceProperty::ApplyDeviceProperty_Internal(const FPlatformUserId UserId, FInputDeviceProperty* RawProperty)
+void UInputDeviceProperty::ApplyDeviceProperty_Internal(const FPlatformUserId UserId, const FInputDeviceId DeviceId, FInputDeviceProperty* RawProperty)
 {
 	if (ensure(RawProperty))
 	{
@@ -30,7 +30,7 @@ void UInputDeviceProperty::ApplyDeviceProperty_Internal(const FPlatformUserId Us
 		if (InputInterface)
 		{
 			int32 ControllerId = INDEX_NONE;
-			IPlatformInputDeviceMapper::Get().RemapUserAndDeviceToControllerId(UserId, ControllerId);
+			IPlatformInputDeviceMapper::Get().RemapUserAndDeviceToControllerId(UserId, ControllerId, DeviceId);
 
 			// TODO_BH: Refactor input interface to take an FPlatformUserId directly (UE-158881)
 			InputInterface->SetDeviceProperty(ControllerId, RawProperty);
@@ -56,12 +56,12 @@ void UInputDeviceProperty::PostEditChangeChainProperty(FPropertyChangedChainEven
 }
 #endif	// WITH_EDITOR
 
-void UInputDeviceProperty::EvaluateDeviceProperty_Implementation(const FPlatformUserId PlatformUser, const float DeltaTime, const float Duration)
+void UInputDeviceProperty::EvaluateDeviceProperty_Implementation(const FPlatformUserId PlatformUser, const FInputDeviceId DeviceId, const float DeltaTime, const float Duration)
 {
 
 }
 
-void UInputDeviceProperty::ResetDeviceProperty_Implementation(const FPlatformUserId PlatformUser)
+void UInputDeviceProperty::ResetDeviceProperty_Implementation(const FPlatformUserId PlatformUser, const FInputDeviceId DeviceId)
 {
 
 }
@@ -69,10 +69,10 @@ void UInputDeviceProperty::ResetDeviceProperty_Implementation(const FPlatformUse
 ///////////////////////////////////////////////////////////////////////
 // UColorInputDeviceProperty
 
-void UColorInputDeviceProperty::EvaluateDeviceProperty_Implementation(const FPlatformUserId PlatformUser, const float DeltaTime, const float Duration)
+void UColorInputDeviceProperty::EvaluateDeviceProperty_Implementation(const FPlatformUserId PlatformUser, const FInputDeviceId DeviceId, const float DeltaTime, const float Duration)
 {
 	// Check for an override on the current input device
-	if (const FDeviceColorData* Data = GetDeviceSpecificData<FDeviceColorData>(PlatformUser, DeviceOverrideData))
+	if (const FDeviceColorData* Data = GetDeviceSpecificData<FDeviceColorData>(PlatformUser, DeviceId, DeviceOverrideData))
 	{
 		InternalProperty.bEnable = Data->bEnable;
 		InternalProperty.Color = Data->LightColor;
@@ -93,10 +93,10 @@ FInputDeviceProperty* UColorInputDeviceProperty::GetInternalDeviceProperty()
 ///////////////////////////////////////////////////////////////////////
 // UColorInputDeviceCurveProperty
 
-void UColorInputDeviceCurveProperty::EvaluateDeviceProperty_Implementation(const FPlatformUserId PlatformUser, const float DeltaTime, const float Duration)
+void UColorInputDeviceCurveProperty::EvaluateDeviceProperty_Implementation(const FPlatformUserId PlatformUser, const FInputDeviceId DeviceId, const float DeltaTime, const float Duration)
 {
 	// Check for an override on the current input device
-	if (const FDeviceColorCurveData* Data = GetDeviceSpecificData<FDeviceColorCurveData>(PlatformUser, DeviceOverrideData))
+	if (const FDeviceColorCurveData* Data = GetDeviceSpecificData<FDeviceColorCurveData>(PlatformUser, DeviceId, DeviceOverrideData))
 	{
 		InternalProperty.bEnable = Data->bEnable;
 
@@ -119,10 +119,10 @@ void UColorInputDeviceCurveProperty::EvaluateDeviceProperty_Implementation(const
 	}
 }
 
-void UColorInputDeviceCurveProperty::ResetDeviceProperty_Implementation(const FPlatformUserId PlatformUser)
+void UColorInputDeviceCurveProperty::ResetDeviceProperty_Implementation(const FPlatformUserId PlatformUser, const FInputDeviceId DeviceId)
 {
 	bool bReset = ColorData.bResetAfterCompletion;
-	if (const FDeviceColorCurveData* Data = GetDeviceSpecificData<FDeviceColorCurveData>(PlatformUser, DeviceOverrideData))
+	if (const FDeviceColorCurveData* Data = GetDeviceSpecificData<FDeviceColorCurveData>(PlatformUser, DeviceId, DeviceOverrideData))
 	{
 		bReset = Data->bResetAfterCompletion;
 	}
@@ -131,7 +131,7 @@ void UColorInputDeviceCurveProperty::ResetDeviceProperty_Implementation(const FP
 	{
 		// Disabling the light will reset the color
     	InternalProperty.bEnable = false;
-    	ApplyDeviceProperty(PlatformUser);
+    	ApplyDeviceProperty(PlatformUser, DeviceId);
 	}	
 }
 
@@ -178,13 +178,13 @@ FInputDeviceProperty* UInputDeviceTriggerEffect::GetInternalDeviceProperty()
 	return &ResetProperty;
 }
 
-void UInputDeviceTriggerEffect::ResetDeviceProperty_Implementation(const FPlatformUserId PlatformUser)
+void UInputDeviceTriggerEffect::ResetDeviceProperty_Implementation(const FPlatformUserId PlatformUser, const FInputDeviceId DeviceId)
 {
 	if (BaseTriggerData.bResetUponCompletion)
 	{
 		// Pass in our reset property
 		ResetProperty.AffectedTriggers = BaseTriggerData.AffectedTriggers;
-		ApplyDeviceProperty_Internal(PlatformUser, &ResetProperty);
+		ApplyDeviceProperty_Internal(PlatformUser, DeviceId, &ResetProperty);
 	}	
 }
 
@@ -221,13 +221,13 @@ int32 UInputDeviceTriggerFeedbackProperty::GetStrengthValue(const FDeviceTrigger
 	return 0;
 }
 
-void UInputDeviceTriggerFeedbackProperty::EvaluateDeviceProperty_Implementation(const FPlatformUserId PlatformUser, const float DeltaTime, const float Duration)
+void UInputDeviceTriggerFeedbackProperty::EvaluateDeviceProperty_Implementation(const FPlatformUserId PlatformUser, const FInputDeviceId DeviceId, const float DeltaTime, const float Duration)
 {		
 	InternalProperty.AffectedTriggers = BaseTriggerData.AffectedTriggers;
 
 	const FDeviceTriggerFeedbackData* DataToUse = &TriggerData;
 
-	if (const FDeviceTriggerFeedbackData* OverrideData = GetDeviceSpecificData<FDeviceTriggerFeedbackData>(PlatformUser, DeviceOverrideData))
+	if (const FDeviceTriggerFeedbackData* OverrideData = GetDeviceSpecificData<FDeviceTriggerFeedbackData>(PlatformUser, DeviceId, DeviceOverrideData))
 	{
 		DataToUse = OverrideData;
 	}
@@ -282,11 +282,11 @@ UInputDeviceTriggerResistanceProperty::UInputDeviceTriggerResistanceProperty()
 	PropertyDuration = 1.0f;
 }
 
-void UInputDeviceTriggerResistanceProperty::EvaluateDeviceProperty_Implementation(const FPlatformUserId PlatformUser, const float DeltaTime, const float Duration)
+void UInputDeviceTriggerResistanceProperty::EvaluateDeviceProperty_Implementation(const FPlatformUserId PlatformUser, const FInputDeviceId DeviceId, const float DeltaTime, const float Duration)
 {
 	InternalProperty.AffectedTriggers = BaseTriggerData.AffectedTriggers;
 
-	if (const FDeviceTriggerTriggerResistanceData* Data = GetDeviceSpecificData<FDeviceTriggerTriggerResistanceData>(PlatformUser, DeviceOverrideData))
+	if (const FDeviceTriggerTriggerResistanceData* Data = GetDeviceSpecificData<FDeviceTriggerTriggerResistanceData>(PlatformUser, DeviceId, DeviceOverrideData))
 	{
 		InternalProperty.StartPosition = Data->StartPosition;
 		InternalProperty.StartStrengh = Data->StartStrengh;
@@ -317,11 +317,11 @@ UInputDeviceTriggerVibrationProperty::UInputDeviceTriggerVibrationProperty()
 	PropertyDuration = 1.0f;
 }
 
-void UInputDeviceTriggerVibrationProperty::EvaluateDeviceProperty_Implementation(const FPlatformUserId PlatformUser, const float DeltaTime, const float Duration)
+void UInputDeviceTriggerVibrationProperty::EvaluateDeviceProperty_Implementation(const FPlatformUserId PlatformUser, const FInputDeviceId DeviceId, const float DeltaTime, const float Duration)
 {
 	const FDeviceTriggerTriggerVibrationData* DataToUse = &TriggerData;
 
-	if (const FDeviceTriggerTriggerVibrationData* OverrideData = GetDeviceSpecificData<FDeviceTriggerTriggerVibrationData>(PlatformUser, DeviceOverrideData))
+	if (const FDeviceTriggerTriggerVibrationData* OverrideData = GetDeviceSpecificData<FDeviceTriggerTriggerVibrationData>(PlatformUser, DeviceId, DeviceOverrideData))
 	{
 		DataToUse = OverrideData;
 	}
diff --git a/Engine/Source/Runtime/Engine/Private/GameFramework/InputDeviceSubsystem.cpp b/Engine/Source/Runtime/Engine/Private/GameFramework/InputDeviceSubsystem.cpp
--- a/Engine/Source/Runtime/Engine/Private/GameFramework/InputDeviceSubsystem.cpp
+++ b/Engine/Source/Runtime/Engine/Private/GameFramework/InputDeviceSubsystem.cpp
@@ -127,8 +127,11 @@ bool FActiveDeviceProperty::operator!=(const FActiveDeviceProperty& Other) const
 
 FSetDevicePropertyParams::FSetDevicePropertyParams()
 	: UserId(FSlateApplicationBase::SlateAppPrimaryPlatformUser)
+	// Set this Device Id to NONE by default. The subsystem will detect this being invalid and fall back 
+	// to the default device for the given platform user. This will let the default behavior of the system be 
+	// consistent, while still allowing for setting specific input devices if desired.
+	, DeviceId(INPUTDEVICEID_NONE)
 {
-
 }
 
 ////////////////////////////////////////////////////////
@@ -275,7 +278,7 @@ void UInputDeviceSubsystem::Tick(float InDeltaTime)
 		
 		if (PropertiesPendingRemoval.Contains(ActiveProp.PropertyHandle))
 		{
-			ActiveProp.Property->ResetDeviceProperty(ActiveProp.PlatformUser);
+			ActiveProp.Property->ResetDeviceProperty(ActiveProp.PlatformUser, ActiveProp.DeviceId);
 			It.RemoveCurrent();
 			PropertiesPendingRemoval.Remove(ActiveProp.PropertyHandle);
 			continue;
@@ -297,7 +300,7 @@ void UInputDeviceSubsystem::Tick(float InDeltaTime)
 		// Don't remove properties that are looping, because the user desires them to keep playing
 		if (!ActiveProp.bLooping && ActiveProp.EvaluatedDuration > ActiveProp.Property->GetDuration() && ActiveProp.bHasBeenAppliedAtLeastOnce)
 		{
-			ActiveProp.Property->ResetDeviceProperty(ActiveProp.PlatformUser);
+			ActiveProp.Property->ResetDeviceProperty(ActiveProp.PlatformUser, ActiveProp.DeviceId);
 			PropertiesPendingRemoval.Remove(ActiveProp.PropertyHandle);
 			It.RemoveCurrent();
 			continue;
@@ -305,8 +308,8 @@ void UInputDeviceSubsystem::Tick(float InDeltaTime)
 		// Otherwise, we can evaluate and apply it as normal
 		else
 		{
-			ActiveProp.Property->EvaluateDeviceProperty(ActiveProp.PlatformUser, DeltaTime, ActiveProp.EvaluatedDuration);
-			ActiveProp.Property->ApplyDeviceProperty(ActiveProp.PlatformUser);
+			ActiveProp.Property->EvaluateDeviceProperty(ActiveProp.PlatformUser, ActiveProp.DeviceId, DeltaTime, ActiveProp.EvaluatedDuration);
+			ActiveProp.Property->ApplyDeviceProperty(ActiveProp.PlatformUser, ActiveProp.DeviceId);
 
 			// Track that this property has now successfully been applied at least once
 			ActiveProp.bHasBeenAppliedAtLeastOnce = true;
@@ -359,6 +362,12 @@ FInputDevicePropertyHandle UInputDeviceSubsystem::ActivateDeviceProperty(UInputD
 		return FInputDevicePropertyHandle::InvalidHandle;
 	}
 
+	if (!Params.UserId.IsValid())
+	{
+		UE_LOG(LogInputDeviceProperties, Error, TEXT("Invalid Platform User Id '%d' given to ActivateDeviceProperty! Nothing will happen."), Params.UserId.GetInternalId());
+		return FInputDevicePropertyHandle::InvalidHandle;
+	}
+
 	FInputDevicePropertyHandle OutHandle = FInputDevicePropertyHandle::AcquireValidHandle();
 
 	if (ensureMsgf(OutHandle.IsValid(), TEXT("Unable to acquire a valid input device property handle! The input device property cannot be activated")))
@@ -366,11 +375,17 @@ FInputDevicePropertyHandle UInputDeviceSubsystem::ActivateDeviceProperty(UInputD
 		FActiveDeviceProperty ActiveProp = {};
 		ActiveProp.Property = Property;
 		ActiveProp.PlatformUser = Params.UserId;
+
+		// If a valid input device was not given, then set this property on the default device for the user.
+		ActiveProp.DeviceId = 
+			Params.DeviceId.IsValid() ? Params.DeviceId : IPlatformInputDeviceMapper::Get().GetPrimaryInputDeviceForUser(Params.UserId);		
+
 		ActiveProp.PropertyHandle = OutHandle;
 		ActiveProp.bLooping = Params.bLooping;
 		ActiveProp.bIgnoreTimeDilation = Params.bIgnoreTimeDilation;
 		ActiveProp.bPlayWhilePaused = Params.bPlayWhilePaused;
 
+		// Keep track of this new active property. Placing it in this set will have it updated on this subsystem's Tick
 		ActiveProperties.Add(ActiveProp);
 	}
 
@@ -438,7 +453,7 @@ void UInputDeviceSubsystem::RemoveAllDeviceProperties()
 	{
 		if (ActiveProperty.Property)
 		{
-			ActiveProperty.Property->ResetDeviceProperty(ActiveProperty.PlatformUser);
+			ActiveProperty.Property->ResetDeviceProperty(ActiveProperty.PlatformUser, ActiveProperty.DeviceId);
 		}		
 	}
 	
