Adding a separate ddc key bump version for shared linear encoding.

diff --git a/Engine/Source/Developer/TextureCompressor/Public/TextureCompressorModule.h b/Engine/Source/Developer/TextureCompressor/Public/TextureCompressorModule.h
--- a/Engine/Source/Developer/TextureCompressor/Public/TextureCompressorModule.h
+++ b/Engine/Source/Developer/TextureCompressor/Public/TextureCompressorModule.h
@@ -232,6 +232,10 @@ struct FTextureBuildSettings
 	// If the target format is a tiled format and can leverage reusing the linear encoding, this is not nullptr.
 	const ITextureTiler* Tiler = nullptr;
 
+	// If shared linear is enabled _at all_ and this texture in involved with that _at all_ then we set
+	// this so we can segregate the derived data keys.
+	bool bAffectedBySharedLinearEncoding = false;
+
 	static constexpr uint32 MaxTextureResolutionDefault = TNumericLimits<uint32>::Max();
 
 	/** Default settings. */
diff --git a/Engine/Source/Runtime/Engine/Private/TextureDerivedData.cpp b/Engine/Source/Runtime/Engine/Private/TextureDerivedData.cpp
--- a/Engine/Source/Runtime/Engine/Private/TextureDerivedData.cpp
+++ b/Engine/Source/Runtime/Engine/Private/TextureDerivedData.cpp
@@ -80,6 +80,11 @@ static TAutoConsoleVariable<int32> CVarTexturesCookToDerivedDataReferences(
 // This is put in the DDC1 key but NOT in the DDC2 key
 #define TEXTURE_VT_DERIVEDDATA_VER	TEXT("7C16439390E24F1F9468894FB4D4BC54")
 
+// This GUID is mixed in for textures that are involved in shared linear encoded textures - both base and child. It's used
+// to rebuild textures affects by shared linear in the case of bugs that only affect such textures so we don't force a global
+// rebuild. This is in both texture build paths.
+static const FGuid GTextureSLEDerivedDataVer(0xBD855730U, 0xA5B44BBBU, 0x89D051D0U, 0x695AC618U);
+const FGuid& GetTextureSLEDerivedDataVersion() { return GTextureSLEDerivedDataVer; }
 
 static bool IsUsingNewDerivedData()
 {
@@ -421,6 +426,11 @@ void GetTextureDerivedDataKeySuffix(const UTexture& Texture, const FTextureBuild
 	OutKeySuffix.Append(TEXT("_arm64"));
 #endif
 
+	if (BuildSettings.bAffectedBySharedLinearEncoding)
+	{
+		OutKeySuffix.Append(LexToString(GTextureSLEDerivedDataVer));
+	}
+
 	// Serialize the compressor settings into a temporary array. The archive
 	// is flagged as persistent so that machines of different endianness produce
 	// identical binary results.
@@ -675,6 +685,16 @@ static void FinalizeBuildSettingsForLayer(
 					static auto CVarSharedLinearTextureEncoding = IConsoleManager::Get().FindTConsoleVariableDataInt(TEXT("r.SharedLinearTextureEncoding"));
 					if (CVarSharedLinearTextureEncoding->GetValueOnAnyThread())
 					{
+						//
+						// We want to separate out textures involved in shared linear encoding in order to facilitate
+						// fixing bugs without invalidating the world (even though we expect the exact same data to
+						// get generated). However, virtual textures never tile, and so are exempt from this separation.
+						//
+						if (OutSettings.bVirtualStreamable == false)
+						{
+							OutSettings.bAffectedBySharedLinearEncoding = true;
+						}
+
 						// Shared linear encoding can only work if the base texture format does not expect to
 						// do the tiling itself (SupportsTiling == false).
 						const FChildTextureFormat* ChildTextureFormat = TextureFormat->GetChildFormat();
diff --git a/Engine/Source/Runtime/Engine/Private/TextureDerivedDataBuildUtils.cpp b/Engine/Source/Runtime/Engine/Private/TextureDerivedDataBuildUtils.cpp
--- a/Engine/Source/Runtime/Engine/Private/TextureDerivedDataBuildUtils.cpp
+++ b/Engine/Source/Runtime/Engine/Private/TextureDerivedDataBuildUtils.cpp
@@ -21,6 +21,7 @@
 #include "TextureResource.h"
 
 const FGuid& GetTextureDerivedDataVersion();
+const FGuid& GetTextureSLEDerivedDataVersion();
 void GetTextureDerivedMipKey(int32 MipIndex, const FTexture2DMipMap& Mip, const FString& KeySuffix, FString& OutKey);
 
 template <typename ValueType>
@@ -308,6 +309,11 @@ FCbObject SaveTextureBuildSettings(const UTexture& Texture, const FTextureBuildS
 	Writer.BeginObject();
 
 	Writer.AddUuid("BuildVersion", GetTextureDerivedDataVersion());
+
+	if (BuildSettings.bAffectedBySharedLinearEncoding)
+	{
+		Writer.AddUuid("SharedLinearEncodingVersion", GetTextureSLEDerivedDataVersion());
+	}
 	
 	if (Texture.CompressionCacheId.IsValid())
 	{
