Engine dependency cleanup prep
- Moving FRawStaticIndexBuffer16or32::CreateRHIBuffer_Internal into a cpp file
- Moving more of FMorphTargetVertexInfoBuffers into its cpp file
- Moving SubUVBoundingGeometryBuffer into its cpp file.

diff --git a/Engine/Source/Runtime/Engine/Classes/Particles/SubUVAnimation.h b/Engine/Source/Runtime/Engine/Classes/Particles/SubUVAnimation.h
--- a/Engine/Source/Runtime/Engine/Classes/Particles/SubUVAnimation.h
+++ b/Engine/Source/Runtime/Engine/Classes/Particles/SubUVAnimation.h
@@ -70,18 +70,12 @@ public:
 	TArray<FVector2f>* Vertices;
 	FShaderResourceViewRHIRef ShaderResourceView;
 
-	FSubUVBoundingGeometryBuffer(TArray<FVector2f>* InVertices)
-	{
-		Vertices = InVertices;
-	}
+	FSubUVBoundingGeometryBuffer();
+	FSubUVBoundingGeometryBuffer(TArray<FVector2f>* InVertices);
+	~FSubUVBoundingGeometryBuffer();
 
 	virtual void InitRHI() override;
-
-	virtual void ReleaseRHI() override
-	{
-		FVertexBuffer::ReleaseRHI();
-		ShaderResourceView.SafeRelease();
-	}
+	virtual void ReleaseRHI() override;
 };
 
 /** Resource array to pass  */
diff --git a/Engine/Source/Runtime/Engine/Private/Particles/SubUVAnimation.cpp b/Engine/Source/Runtime/Engine/Private/Particles/SubUVAnimation.cpp
--- a/Engine/Source/Runtime/Engine/Private/Particles/SubUVAnimation.cpp
+++ b/Engine/Source/Runtime/Engine/Private/Particles/SubUVAnimation.cpp
@@ -41,6 +41,15 @@ void FSubUVDerivedData::Serialize(FStructuredArchive::FSlot Slot)
 	Slot << BoundingGeometry;
 }
 
+FSubUVBoundingGeometryBuffer::FSubUVBoundingGeometryBuffer() = default;
+
+FSubUVBoundingGeometryBuffer::FSubUVBoundingGeometryBuffer(TArray<FVector2f>* InVertices)
+{
+	Vertices = InVertices;
+}
+
+FSubUVBoundingGeometryBuffer::~FSubUVBoundingGeometryBuffer() = default;
+
 void FSubUVBoundingGeometryBuffer::InitRHI()
 {
 	const uint32 SizeInBytes = Vertices->Num() * Vertices->GetTypeSize();
@@ -54,6 +63,12 @@ void FSubUVBoundingGeometryBuffer::InitRHI()
 	}
 }
 
+void FSubUVBoundingGeometryBuffer::ReleaseRHI()
+{
+	FVertexBuffer::ReleaseRHI();
+	ShaderResourceView.SafeRelease();
+}
+
 USubUVAnimation::USubUVAnimation(const FObjectInitializer& ObjectInitializer)
 	: Super(ObjectInitializer)
 {
diff --git a/Engine/Source/Runtime/Engine/Private/RawIndexBuffer.cpp b/Engine/Source/Runtime/Engine/Private/RawIndexBuffer.cpp
--- a/Engine/Source/Runtime/Engine/Private/RawIndexBuffer.cpp
+++ b/Engine/Source/Runtime/Engine/Private/RawIndexBuffer.cpp
@@ -7,6 +7,8 @@
 #include "RawIndexBuffer.h"
 #include "Interfaces/ITargetPlatform.h"
 #include "Modules/ModuleManager.h"
+#include "RenderingThread.h"
+#include "DataDrivenShaderPlatformInfo.h"
 
 #if WITH_EDITOR
 #include "MeshUtilities.h"
@@ -442,6 +444,43 @@ void FRawStaticIndexBuffer::Discard()
     IndexStorage.Discard();
 }
 
+FBufferRHIRef FRawStaticIndexBuffer16or32Interface::CreateRHIIndexBufferInternal(
+	const TCHAR* InDebugName,
+	const FName& InOwnerName,
+	int32 IndexCount,
+	size_t IndexSize,
+	FResourceArrayInterface* ResourceArray,
+	bool bNeedSRV,
+	bool bRenderThread
+)
+{
+	// Create the index buffer.
+	FRHIResourceCreateInfo CreateInfo(InDebugName, ResourceArray);
+	EBufferUsageFlags Flags = EBufferUsageFlags::Static;
+
+	if (bNeedSRV)
+	{
+		// BUF_ShaderResource is needed for SkinCache RecomputeSkinTangents
+		Flags |= EBufferUsageFlags::ShaderResource;
+	}
+
+	FBufferRHIRef Ret;
+	const uint32 Size = IndexCount * IndexSize;
+	CreateInfo.bWithoutNativeResource = !Size;
+	if (bRenderThread)
+	{
+		Ret = RHICreateIndexBuffer(IndexSize, Size, Flags, CreateInfo);
+	}
+	else
+	{
+		FRHIAsyncCommandList CommandList;
+		Ret = CommandList->CreateBuffer(Size, Flags | EBufferUsageFlags::IndexBuffer, IndexSize, ERHIAccess::SRVMask, CreateInfo);
+	}
+
+	Ret->SetOwnerName(InOwnerName);
+	return Ret;
+}
+
 /*-----------------------------------------------------------------------------
 FRawStaticIndexBuffer16or32
 -----------------------------------------------------------------------------*/
diff --git a/Engine/Source/Runtime/Engine/Private/Rendering/MorphTargetVertexInfoBuffers.cpp b/Engine/Source/Runtime/Engine/Private/Rendering/MorphTargetVertexInfoBuffers.cpp
--- a/Engine/Source/Runtime/Engine/Private/Rendering/MorphTargetVertexInfoBuffers.cpp
+++ b/Engine/Source/Runtime/Engine/Private/Rendering/MorphTargetVertexInfoBuffers.cpp
@@ -5,9 +5,13 @@
 #include "Rendering/SkeletalMeshLODRenderData.h"
 #include "Animation/MorphTarget.h"
 #include "RHI.h"
+#include "DataDrivenShaderPlatformInfo.h"
 
 extern int32 GSkinCacheRecomputeTangents;
 
+FMorphTargetVertexInfoBuffers::FMorphTargetVertexInfoBuffers() = default;
+FMorphTargetVertexInfoBuffers::~FMorphTargetVertexInfoBuffers() = default;
+
 void FMorphTargetVertexInfoBuffers::InitRHI()
 {
 	SCOPED_LOADTIMER(FFMorphTargetVertexInfoBuffers_InitRHI);
@@ -54,6 +58,20 @@ const float FMorphTargetVertexInfoBuffers::CalculatePositionPrecision(float Targ
 	return TargetPositionErrorTolerance * 2.0f * 1e-6f * UnrealUnitPerMeter;	// * 2.0 because correct rounding guarantees error is at most half of the cell size.
 }
 
+void FMorphTargetVertexInfoBuffers::ResetCPUData()
+{
+	MorphData.Empty();
+	MaximumValuePerMorph.Empty();
+	MinimumValuePerMorph.Empty();
+	BatchStartOffsetPerMorph.Empty();
+	BatchesPerMorph.Empty();
+	NumTotalBatches = 0;
+	PositionPrecision = 0.0f;
+	TangentZPrecision = 0.0f;
+	bResourcesInitialized = false;
+	bIsMorphCPUDataValid = false;
+}
+
 void FMorphTargetVertexInfoBuffers::ValidateVertexBuffers(bool bMorphTargetsShouldBeValid)
 {
 #if DO_CHECK
diff --git a/Engine/Source/Runtime/Engine/Public/BatchedElements.h b/Engine/Source/Runtime/Engine/Public/BatchedElements.h
--- a/Engine/Source/Runtime/Engine/Public/BatchedElements.h
+++ b/Engine/Source/Runtime/Engine/Public/BatchedElements.h
@@ -16,6 +16,7 @@
 #include "SceneView.h"
 #include "StaticBoundShaderState.h"
 #include "PipelineStateCache.h"
+#include "Containers/DynamicRHIResourceArray.h"
 
 struct FBatchedPoint;
 struct FMeshPassProcessorRenderState;
diff --git a/Engine/Source/Runtime/Engine/Public/RawIndexBuffer.h b/Engine/Source/Runtime/Engine/Public/RawIndexBuffer.h
--- a/Engine/Source/Runtime/Engine/Public/RawIndexBuffer.h
+++ b/Engine/Source/Runtime/Engine/Public/RawIndexBuffer.h
@@ -374,6 +374,16 @@ public:
 	}
 
 protected:
+	static ENGINE_API FBufferRHIRef CreateRHIIndexBufferInternal(
+		const TCHAR* InDebugName,
+		const FName& InOwnerName,
+		int32 IndexCount,
+		size_t IndexSize,
+		FResourceArrayInterface* ResourceArray,
+		bool bNeedSRV,
+		bool bRenderThread
+	);
+
 	// guaranteed only to be valid if the vertex buffer is valid and the buffer was created with the SRV flags
 	FShaderResourceViewRHIRef SRVValue;
 };
@@ -548,35 +558,19 @@ private:
 	{
 		if (CachedNumIndices)
 		{
-			// Create the index buffer.
-			FRHIResourceCreateInfo CreateInfo(sizeof(INDEX_TYPE) == 4 ? TEXT("FRawStaticIndexBuffer32") : TEXT("FRawStaticIndexBuffer16"), &Indices);
-			EBufferUsageFlags Flags = EBufferUsageFlags::Static;
-
-			if (IsSRVNeeded())
-			{
-				// BUF_ShaderResource is needed for SkinCache RecomputeSkinTangents
-				Flags |= EBufferUsageFlags::ShaderResource;
-			}
-
 			// Need to cache number of indices from the source array *before* RHICreateIndexBuffer is called
 			// because it will empty the source array.
 			CachedNumIndices = Indices.Num();
 
-			FBufferRHIRef Ret;
-			const uint32 Size = CachedNumIndices * sizeof(INDEX_TYPE);
-			CreateInfo.bWithoutNativeResource = !Size;
-			if (bRenderThread)
-			{
-				Ret = RHICreateIndexBuffer(sizeof(INDEX_TYPE), Size, Flags, CreateInfo);
-			}
-			else
-			{
-				FRHIAsyncCommandList CommandList;
-				Ret = CommandList->CreateBuffer(Size, Flags | EBufferUsageFlags::IndexBuffer, sizeof(INDEX_TYPE), ERHIAccess::SRVMask, CreateInfo);
-			}
-
-			Ret->SetOwnerName(GetOwnerName());
-			return Ret;
+			return CreateRHIIndexBufferInternal(
+				sizeof(INDEX_TYPE) == 4 ? TEXT("FRawStaticIndexBuffer32") : TEXT("FRawStaticIndexBuffer16"),
+				GetOwnerName(),
+				Indices.Num(),
+				sizeof(INDEX_TYPE),
+				&Indices,
+				IsSRVNeeded(),
+				bRenderThread
+			);
 		}
 		return nullptr;
 	}
diff --git a/Engine/Source/Runtime/Engine/Public/Rendering/MorphTargetVertexInfoBuffers.h b/Engine/Source/Runtime/Engine/Public/Rendering/MorphTargetVertexInfoBuffers.h
--- a/Engine/Source/Runtime/Engine/Public/Rendering/MorphTargetVertexInfoBuffers.h
+++ b/Engine/Source/Runtime/Engine/Public/Rendering/MorphTargetVertexInfoBuffers.h
@@ -3,39 +3,39 @@
 #pragma once
 
 #include "RenderResource.h"
+#include "Containers/DynamicRHIResourceArray.h"
 
 struct FSkelMeshRenderSection;
 class UMorphTarget;
 
-class FMorphTargetVertexInfoBuffers : public FRenderResource
+class ENGINE_API FMorphTargetVertexInfoBuffers : public FRenderResource
 {
 public:
-	ENGINE_API FMorphTargetVertexInfoBuffers() : NumTotalBatches(0)
-	{
-	}
+	FMorphTargetVertexInfoBuffers();
+	~FMorphTargetVertexInfoBuffers();
 
-	ENGINE_API void InitMorphResources(EShaderPlatform ShaderPlatform, const TArray<FSkelMeshRenderSection>& RenderSections, const TArray<UMorphTarget*>& MorphTargets, int NumVertices, int32 LODIndex, float TargetPositionErrorTolerance);
+	void InitMorphResources(EShaderPlatform ShaderPlatform, const TArray<FSkelMeshRenderSection>& RenderSections, const TArray<UMorphTarget*>& MorphTargets, int NumVertices, int32 LODIndex, float TargetPositionErrorTolerance);
 
 	inline bool IsMorphResourcesInitialized() const { return bResourcesInitialized; }
 	inline bool IsRHIIntialized() const { return bRHIIntialized; }
 	inline bool IsMorphCPUDataValid() const{ return bIsMorphCPUDataValid; }
 	
-	ENGINE_API bool GetEmptyMorphCPUDataOnInitRHI() const { return bEmptyMorphCPUDataOnInitRHI; }
-	ENGINE_API void SetEmptyMorphCPUDataOnInitRHI(bool bEmpty) { bEmptyMorphCPUDataOnInitRHI = bEmpty; }
+	bool GetEmptyMorphCPUDataOnInitRHI() const { return bEmptyMorphCPUDataOnInitRHI; }
+	void SetEmptyMorphCPUDataOnInitRHI(bool bEmpty) { bEmptyMorphCPUDataOnInitRHI = bEmpty; }
 
-	ENGINE_API virtual void InitRHI() override;
-	ENGINE_API virtual void ReleaseRHI() override;
+	virtual void InitRHI() override;
+	virtual void ReleaseRHI() override;
 
 	UE_DEPRECATED(5.2, "GetMaximumThreadGroupSize will be removed as it is no longer used.")
-	static ENGINE_API uint32 GetMaximumThreadGroupSize();
+	static uint32 GetMaximumThreadGroupSize();
 
-	ENGINE_API uint32 GetNumBatches(uint32 index = UINT_MAX) const
+	uint32 GetNumBatches(uint32 index = UINT_MAX) const
 	{
 		check(index == UINT_MAX || index < (uint32)BatchesPerMorph.Num());
 		return index != UINT_MAX ? BatchesPerMorph[index] : NumTotalBatches;
 	}
 
-	ENGINE_API uint32 GetNumMorphs() const
+	uint32 GetNumMorphs() const
 	{
 		return BatchesPerMorph.Num();
 	}
@@ -76,19 +76,7 @@ public:
 	FShaderResourceViewRHIRef MorphDataSRV;
 
 protected:
-	void ResetCPUData()
-	{
-		MorphData.Empty();
-		MaximumValuePerMorph.Empty();
-		MinimumValuePerMorph.Empty();
-		BatchStartOffsetPerMorph.Empty();
-		BatchesPerMorph.Empty();
-		NumTotalBatches = 0;
-		PositionPrecision = 0.0f;
-		TangentZPrecision = 0.0f;
-		bResourcesInitialized = false;
-		bIsMorphCPUDataValid = false;
-	}
+	void ResetCPUData();
 
 	void ValidateVertexBuffers(bool bMorphTargetsShouldBeValid);
 	void Serialize(FArchive& Ar);
@@ -102,9 +90,9 @@ protected:
 	TArray<uint32> BatchStartOffsetPerMorph;
 	TArray<uint32> BatchesPerMorph;
 	
-	uint32 NumTotalBatches;
-	float PositionPrecision;
-	float TangentZPrecision;
+	uint32 NumTotalBatches = 0;
+	float PositionPrecision = 0.0f;
+	float TangentZPrecision = 0.0f;
 
 	bool bIsMorphCPUDataValid = false;
 	bool bResourcesInitialized = false;
