Fix for nullptr localplayer when disconnecting an account

#rnx

diff --git a/Engine/Source/Runtime/Engine/Private/LocalPlayer.cpp b/Engine/Source/Runtime/Engine/Private/LocalPlayer.cpp
--- a/Engine/Source/Runtime/Engine/Private/LocalPlayer.cpp
+++ b/Engine/Source/Runtime/Engine/Private/LocalPlayer.cpp
@@ -167,7 +167,7 @@ class AGameStateBase* FLocalPlayerContext::GetGameState() const
 	else
 	{
 		ULocalPlayer* LocalPlayerPtr = GetLocalPlayer();
-		if (UWorld* LocalPlayerWorld = LocalPlayerPtr->GetWorld())
+		if (UWorld* LocalPlayerWorld = LocalPlayerPtr ? LocalPlayerPtr->GetWorld() : nullptr)
 		{
 			GameState = LocalPlayerWorld->GetGameState();
 		}
