Fixes r.DynamicRes.DynamicFrameTime.RoundUpToVSyncError with monitors that supports different refresh rate

diff --git a/Engine/Source/Runtime/Engine/Private/DynamicResolution.cpp b/Engine/Source/Runtime/Engine/Private/DynamicResolution.cpp
--- a/Engine/Source/Runtime/Engine/Private/DynamicResolution.cpp
+++ b/Engine/Source/Runtime/Engine/Private/DynamicResolution.cpp
@@ -274,7 +274,6 @@ void FDynamicResolutionHeuristicProxy::RefreshCurrentFrameResolutionFraction_Ren
 	float MinGlobalFrameTime = 0.0f;
 	if (CVarDynamicFrameTimeEnable.GetValueOnRenderThread())
 	{
-		const float VSyncFrameTime = 16.6;
 		static const IConsoleVariable* VSyncCVar = IConsoleManager::Get().FindConsoleVariable(TEXT("r.VSync"));
 		check(VSyncCVar);
 
@@ -318,6 +317,7 @@ void FDynamicResolutionHeuristicProxy::RefreshCurrentFrameResolutionFraction_Ren
 			// When CPU bound and vsync is enabled, the GPU may end up underused due to vsync wait, so we can crank up the available GPU time a bit more.
 			if (float RoundUpToVSyncError = CVarDynamicFrameTimeRoundUpToVsync.GetValueOnGameThread() > 0.0f && VSyncCVar->GetInt())
 			{
+				const float VSyncFrameTime = 1000.0f / float(FPlatformMisc::GetMaxRefreshRate());
 				float VSyncRateChangeMultiplier = 1.0f + RoundUpToVSyncError / 100.f;
 
 				int32 CurrentVSyncFactor = FMath::CeilToInt(MedianFrameTime / VSyncFrameTime);
