Add a "reset after completion" optoin to regular colors to let you clear the light color if ytou want.

#rnx
#pf skip

diff --git a/Engine/Source/Runtime/Engine/Classes/GameFramework/InputDeviceProperties.h b/Engine/Source/Runtime/Engine/Classes/GameFramework/InputDeviceProperties.h
--- a/Engine/Source/Runtime/Engine/Classes/GameFramework/InputDeviceProperties.h
+++ b/Engine/Source/Runtime/Engine/Classes/GameFramework/InputDeviceProperties.h
@@ -164,6 +164,10 @@ struct FDeviceColorData
 	/** True if the light should be enabled at all */
 	UPROPERTY(EditAnywhere, BlueprintReadWrite, Category = "Color")
 	bool bEnable = true;
+	
+	/** If true, the light color will be reset to "off" after this property has been evaluated. */
+	UPROPERTY(EditAnywhere, BlueprintReadWrite, Category = "Color")
+	bool bResetAfterCompletion = false;
 
 	/** The color to set the light on  */
 	UPROPERTY(EditAnywhere, BlueprintReadWrite, Category = "Color")
@@ -185,6 +189,7 @@ class UColorInputDeviceProperty : public UInputDeviceProperty
 public:
 
 	virtual void EvaluateDeviceProperty_Implementation(const FPlatformUserId PlatformUser, const FInputDeviceId DeviceId, const float DeltaTime, const float Duration) override;
+	virtual void ResetDeviceProperty_Implementation(const FPlatformUserId PlatformUser, const FInputDeviceId DeviceId) override;
 	virtual FInputDeviceProperty* GetInternalDeviceProperty() override;
 
 	/** Default color data that will be used by default. Device Specific overrides will be used when the current input device matches */
diff --git a/Engine/Source/Runtime/Engine/Private/GameFramework/InputDeviceProperties.cpp b/Engine/Source/Runtime/Engine/Private/GameFramework/InputDeviceProperties.cpp
--- a/Engine/Source/Runtime/Engine/Private/GameFramework/InputDeviceProperties.cpp
+++ b/Engine/Source/Runtime/Engine/Private/GameFramework/InputDeviceProperties.cpp
@@ -85,6 +85,22 @@ void UColorInputDeviceProperty::EvaluateDeviceProperty_Implementation(const FPla
 	}
 }
 
+void UColorInputDeviceProperty::ResetDeviceProperty_Implementation(const FPlatformUserId PlatformUser, const FInputDeviceId DeviceId)
+{
+	bool bReset = ColorData.bResetAfterCompletion;
+	if (const FDeviceColorData* Data = GetDeviceSpecificData<FDeviceColorData>(PlatformUser, DeviceId, DeviceOverrideData))
+	{
+		bReset = Data->bResetAfterCompletion;
+	}
+
+	if (bReset)
+	{
+		// Disabling the light will reset the color
+		InternalProperty.bEnable = false;
+		ApplyDeviceProperty(PlatformUser, DeviceId);
+	}
+}
+
 FInputDeviceProperty* UColorInputDeviceProperty::GetInternalDeviceProperty()
 {
 	return &InternalProperty;
