Fix incorrect material instance material blend mode.

This is a short term fix to fix crashes, but we later need to expose StrataBlendMode override at the instance material level.

diff --git a/Engine/Source/Runtime/Engine/Private/Materials/Material.cpp b/Engine/Source/Runtime/Engine/Private/Materials/Material.cpp
--- a/Engine/Source/Runtime/Engine/Private/Materials/Material.cpp
+++ b/Engine/Source/Runtime/Engine/Private/Materials/Material.cpp
@@ -3028,6 +3028,43 @@ static void AddStrataShadingModelFromMaterialShadingModel(FStrataMaterialInfo& O
 	if (InShadingModels.HasShadingModel(MSM_ThinTranslucent))	{ OutInfo.AddShadingModel(EStrataShadingModel::SSM_DefaultLit); }
 }
 
+
+EStrataBlendMode ConvertLegacyToStrataBlendMode(EBlendMode InBlendMode, FMaterialShadingModelField InShadingModels, const UMaterialInterface* InMaterial)
+{
+	EStrataBlendMode Out = SBM_Opaque;
+	if (InShadingModels.CountShadingModels() == 1 && InShadingModels.GetFirstShadingModel() == EMaterialShadingModel::MSM_ThinTranslucent)
+	{
+		Out = SBM_TranslucentColoredTransmittance;
+	}
+	else
+	{
+		switch (InBlendMode)
+		{
+		case BLEND_Opaque:
+			Out = SBM_Opaque;
+			break;
+		case BLEND_Masked:
+			Out = SBM_Masked;
+			break;
+		case BLEND_Translucent:
+		case BLEND_AlphaComposite:
+		case BLEND_Additive:
+			Out = SBM_TranslucentGreyTransmittance;
+			break;
+		case BLEND_Modulate:
+			Out = SBM_ColoredTransmittanceOnly;
+			break;
+		case BLEND_AlphaHoldout:
+			Out = SBM_AlphaHoldout;
+			break;
+		default:
+			UE_LOG(LogMaterial, Error, TEXT("%s: Material blend mode could not be converted to Strata."), *InMaterial->GetName());
+			break;
+		}
+	}
+	return Out;
+}
+
 void UMaterial::ConvertMaterialToStrataMaterial()
 {
 #if WITH_EDITOR
@@ -3068,39 +3105,6 @@ void UMaterial::ConvertMaterialToStrataMaterial()
 			NewNode->GetInput(NewInputIndex)->Connect(OldNodeInput.OutputIndex, OldNodeInput.Expression);
 		}
 	};
-
-	auto ConvertLegacyToStrataBlendMode = [&]()
-	{
-		if (ShadingModels.CountShadingModels() == 1 && ShadingModels.GetFirstShadingModel() == EMaterialShadingModel::MSM_ThinTranslucent)
-		{
-			StrataBlendMode = SBM_TranslucentColoredTransmittance;
-		}
-		else
-		{
-			switch (BlendMode)
-			{
-			case BLEND_Opaque:
-				StrataBlendMode = SBM_Opaque;
-				break;
-			case BLEND_Masked:
-				StrataBlendMode = SBM_Masked;
-				break;
-			case BLEND_Translucent:
-			case BLEND_AlphaComposite:
-			case BLEND_Additive:
-				StrataBlendMode = SBM_TranslucentGreyTransmittance;
-				break;
-			case BLEND_Modulate:
-				StrataBlendMode = SBM_ColoredTransmittanceOnly;
-				break;
-			case BLEND_AlphaHoldout:
-				StrataBlendMode = SBM_AlphaHoldout;
-				break;
-			default:
-				UE_LOG(LogMaterial, Error, TEXT("%s: Material blend mode could not be converted to Starta."), *GetName());
-			}
-		}
-	};
 	
 	UMaterialEditorOnlyData* EditorOnly = GetEditorOnlyData();
 
@@ -3276,7 +3280,7 @@ void UMaterial::ConvertMaterialToStrataMaterial()
 			EditorOnly->FrontMaterial.Connect(0, ConvertToDecalNode);
 		}
 
-		ConvertLegacyToStrataBlendMode();
+		StrataBlendMode = ConvertLegacyToStrataBlendMode(BlendMode, ShadingModels, this);
 		bInvalidateShader = true;
 	}
 	else if (!bUseMaterialAttributes && !EditorOnly->FrontMaterial.IsConnected() && GetExpressions().IsEmpty())
@@ -3545,7 +3549,7 @@ void UMaterial::ConvertMaterialToStrataMaterial()
 			bInvalidateShader = true;
 		}
 
-		ConvertLegacyToStrataBlendMode();
+		StrataBlendMode = ConvertLegacyToStrataBlendMode(BlendMode, ShadingModels, this);
 	}
 
 	if (bRelinkCustomOutputNodes)
diff --git a/Engine/Source/Runtime/Engine/Private/Materials/MaterialInstance.cpp b/Engine/Source/Runtime/Engine/Private/Materials/MaterialInstance.cpp
--- a/Engine/Source/Runtime/Engine/Private/Materials/MaterialInstance.cpp
+++ b/Engine/Source/Runtime/Engine/Private/Materials/MaterialInstance.cpp
@@ -2024,6 +2024,8 @@ void UMaterialInstance::InitStaticPermutation(EMaterialShaderPrecompileMode Prec
 	FMaterial::DeferredDeleteArray(ResourcesToFree);
 }
 
+EStrataBlendMode ConvertLegacyToStrataBlendMode(EBlendMode InBlendMode, FMaterialShadingModelField InShadingModels, const UMaterialInterface* InMaterial);
+
 void UMaterialInstance::UpdateOverridableBaseProperties()
 {
 	//Parents base property overrides have to be cached by now.
@@ -4265,7 +4267,7 @@ EBlendMode UMaterialInstance::GetBlendMode() const
 
 EStrataBlendMode UMaterialInstance::GetStrataBlendMode() const
 {
-	return Parent ? Parent->GetStrataBlendMode() : SBM_Opaque;
+	return ConvertLegacyToStrataBlendMode(BlendMode, ShadingModels, this);
 }
 
 FMaterialShadingModelField UMaterialInstance::GetShadingModels() const
