Add owner tracking for skin cache resources

diff --git a/Engine/Source/Runtime/Engine/Private/GPUSkinCache.cpp b/Engine/Source/Runtime/Engine/Private/GPUSkinCache.cpp
--- a/Engine/Source/Runtime/Engine/Private/GPUSkinCache.cpp
+++ b/Engine/Source/Runtime/Engine/Private/GPUSkinCache.cpp
@@ -1244,7 +1244,7 @@ void FGPUSkinCache::DispatchUpdateSkinTangents(FRHICommandListImmediate& RHICmdL
 	}
 }
 
-FGPUSkinCache::FRWBuffersAllocation* FGPUSkinCache::TryAllocBuffer(uint32 NumVertices, bool WithTangnents, bool UseIntermediateTangents, uint32 NumTriangles, FRHICommandListImmediate& RHICmdList)
+FGPUSkinCache::FRWBuffersAllocation* FGPUSkinCache::TryAllocBuffer(uint32 NumVertices, bool WithTangnents, bool UseIntermediateTangents, uint32 NumTriangles, FRHICommandListImmediate& RHICmdList, const FName& OwnerName)
 {
 	uint64 MaxSizeInBytes = (uint64)(GSkinCacheSceneMemoryLimitInMB * 1024.0f * 1024.0f);
 	uint64 RequiredMemInBytes = FRWBuffersAllocation::CalculateRequiredMemory(NumVertices, WithTangnents, UseIntermediateTangents, NumTriangles);
@@ -1256,7 +1256,7 @@ FGPUSkinCache::FRWBuffersAllocation* FGPUSkinCache::TryAllocBuffer(uint32 NumVer
 		return nullptr;
 	}
 
-	FRWBuffersAllocation* NewAllocation = new FRWBuffersAllocation(NumVertices, WithTangnents, UseIntermediateTangents, NumTriangles, RHICmdList);
+	FRWBuffersAllocation* NewAllocation = new FRWBuffersAllocation(NumVertices, WithTangnents, UseIntermediateTangents, NumTriangles, RHICmdList, OwnerName);
 	Allocations.Add(NewAllocation);
 
 	UsedMemoryInBytes += RequiredMemInBytes;
@@ -1553,7 +1553,8 @@ bool FGPUSkinCache::ProcessEntry(
 			}
 		}
 
-		FRWBuffersAllocation* NewPositionAllocation = TryAllocBuffer(TotalNumVertices, WithTangents, bEntryUseIntermediateTangents, InterAccumTangentBufferSize, RHICmdList);
+		const FName OwnerName(FString::Printf(TEXT("%s [LOD%d]"), *Skin->GetAssetPathName().ToString(), LODIndex));
+		FRWBuffersAllocation* NewPositionAllocation = TryAllocBuffer(TotalNumVertices, WithTangents, bEntryUseIntermediateTangents, InterAccumTangentBufferSize, RHICmdList, OwnerName);
 		if (!NewPositionAllocation)
 		{
 			if (GSkinCachePrintMemorySummary > 0)
diff --git a/Engine/Source/Runtime/Engine/Private/SkeletalRender.cpp b/Engine/Source/Runtime/Engine/Private/SkeletalRender.cpp
--- a/Engine/Source/Runtime/Engine/Private/SkeletalRender.cpp
+++ b/Engine/Source/Runtime/Engine/Private/SkeletalRender.cpp
@@ -64,6 +64,9 @@ FSkeletalMeshObject::FSkeletalMeshObject(USkinnedMeshComponent* InMeshComponent,
 ,	StatId(InMeshComponent->GetSkinnedAsset()->GetStatID(true))
 ,	FeatureLevel(InFeatureLevel)
 ,	ComponentId(InMeshComponent->ComponentId.PrimIDValue)
+#if RHI_ENABLE_RESOURCE_INFO
+,	AssetPathName(InMeshComponent->GetSkinnedAsset()->GetPathName())
+#endif
 {
 	check(SkeletalMeshRenderData);
 
@@ -233,6 +236,15 @@ float FSkeletalMeshObject::GetScreenSize(int32 LODIndex) const
 	return 0.f;
 }
 
+FName FSkeletalMeshObject::GetAssetPathName() const
+{
+#if RHI_ENABLE_RESOURCE_INFO
+	return AssetPathName;
+#else
+	return NAME_None;
+#endif
+}
+
 /*-----------------------------------------------------------------------------
 Global functions
 -----------------------------------------------------------------------------*/
diff --git a/Engine/Source/Runtime/Engine/Public/GPUSkinCache.h b/Engine/Source/Runtime/Engine/Public/GPUSkinCache.h
--- a/Engine/Source/Runtime/Engine/Public/GPUSkinCache.h
+++ b/Engine/Source/Runtime/Engine/Public/GPUSkinCache.h
@@ -206,12 +206,13 @@ public:
 	{
 		friend struct FRWBufferTracker;
 
-		FRWBuffersAllocation(uint32 InNumVertices, bool InWithTangents, bool InUseIntermediateTangents, uint32 InIntermediateAccumulatedTangentsSize, FRHICommandListImmediate& RHICmdList)
+		FRWBuffersAllocation(uint32 InNumVertices, bool InWithTangents, bool InUseIntermediateTangents, uint32 InIntermediateAccumulatedTangentsSize, FRHICommandListImmediate& RHICmdList, const FName& OwnerName)
 			: NumVertices(InNumVertices), WithTangents(InWithTangents), UseIntermediateTangents(InUseIntermediateTangents), IntermediateAccumulatedTangentsSize(InIntermediateAccumulatedTangentsSize)
 		{
 			for (int32 Index = 0; Index < NUM_BUFFERS; ++Index)
 			{
 				PositionBuffers[Index].Buffer.Initialize(TEXT("SkinCachePositions"), PosBufferBytesPerElement, NumVertices * 3, PF_R32_FLOAT, BUF_Static);
+				PositionBuffers[Index].Buffer.Buffer->SetOwnerName(OwnerName);
 				PositionBuffers[Index].AccessState = ERHIAccess::Unknown;
 			}
 			if (WithTangents)
@@ -220,16 +221,19 @@ public:
 				const EPixelFormat TangentsFormat = IsOpenGLPlatform(GMaxRHIShaderPlatform) ? PF_R16G16B16A16_SINT : PF_R16G16B16A16_SNORM;
 				
 				Tangents.Buffer.Initialize(TEXT("SkinCacheTangents"), TangentBufferBytesPerElement, NumVertices * 2, TangentsFormat, BUF_Static);
+				Tangents.Buffer.Buffer->SetOwnerName(OwnerName);
 				Tangents.AccessState = ERHIAccess::Unknown;
 				if (UseIntermediateTangents)
 				{
 					IntermediateTangents.Buffer.Initialize(TEXT("SkinCacheIntermediateTangents"), TangentBufferBytesPerElement, NumVertices * 2, TangentsFormat, BUF_Static);
+					IntermediateTangents.Buffer.Buffer->SetOwnerName(OwnerName);
 					IntermediateTangents.AccessState = ERHIAccess::Unknown;
 				}
 			}
 			if (IntermediateAccumulatedTangentsSize > 0)
 			{
 				IntermediateAccumulatedTangents.Buffer.Initialize(TEXT("SkinCacheIntermediateAccumulatedTangents"), sizeof(int32), IntermediateAccumulatedTangentsSize * FGPUSkinCache::IntermediateAccumBufferNumInts, PF_R32_SINT, BUF_UnorderedAccess);
+				IntermediateAccumulatedTangents.Buffer.Buffer->SetOwnerName(OwnerName);
 				IntermediateAccumulatedTangents.AccessState = ERHIAccess::Unknown;
 				// The UAV must be zero-filled. We leave it zeroed after each round (see RecomputeTangentsPerVertexPass.usf), so this is only needed on when the buffer is first created.
 				RHICmdList.ClearUAVUint(IntermediateAccumulatedTangents.Buffer.UAV, FUintVector4(0, 0, 0, 0));
@@ -413,7 +417,7 @@ protected:
 	TSet<FGPUSkinCacheEntry*> PendingProcessRTGeometryEntries;
 	TArray<FDispatchEntry> BatchDispatches;
 
-	FRWBuffersAllocation* TryAllocBuffer(uint32 NumVertices, bool WithTangnents, bool UseIntermediateTangents, uint32 NumTriangles, FRHICommandListImmediate& RHICmdList);
+	FRWBuffersAllocation* TryAllocBuffer(uint32 NumVertices, bool WithTangnents, bool UseIntermediateTangents, uint32 NumTriangles, FRHICommandListImmediate& RHICmdList, const FName& OwnerName);
 	void DoDispatch(FRHICommandListImmediate& RHICmdList);
 	void DoDispatch(FRHICommandListImmediate& RHICmdList, FGPUSkinCacheEntry* SkinCacheEntry, int32 Section, int32 RevisionNumber);
 	void DispatchUpdateSkinTangents(FRHICommandListImmediate& RHICmdList, FGPUSkinCacheEntry* Entry, int32 SectionIndex, FSkinCacheRWBuffer*& StagingBuffer, bool bTrianglePass);
diff --git a/Engine/Source/Runtime/Engine/Public/SkeletalRenderPublic.h b/Engine/Source/Runtime/Engine/Public/SkeletalRenderPublic.h
--- a/Engine/Source/Runtime/Engine/Public/SkeletalRenderPublic.h
+++ b/Engine/Source/Runtime/Engine/Public/SkeletalRenderPublic.h
@@ -213,6 +213,8 @@ public:
 
 	FColor GetSkinCacheVisualizationDebugColor(const FName& GPUSkinCacheVisualizationMode, uint32 SectionIndex) const;
 
+	FName GetAssetPathName() const;
+
 #if RHI_RAYTRACING
 	/** Retrieve ray tracing geometry from the underlying mesh object */
 	virtual FRayTracingGeometry* GetRayTracingGeometry() { return nullptr; }
@@ -325,4 +327,8 @@ protected:
 
 	/** Component ID to which belong this  mesh object  */
 	uint32 ComponentId;
+
+#if RHI_ENABLE_RESOURCE_INFO
+	FName AssetPathName;
+#endif
 };
