Fix texture resource not created in editor commandlet mode, causing crash in FRHIResource::SetOwnername() accessing invalid pointer

diff --git a/Engine/Source/Runtime/Engine/Private/Texture2D.cpp b/Engine/Source/Runtime/Engine/Private/Texture2D.cpp
--- a/Engine/Source/Runtime/Engine/Private/Texture2D.cpp
+++ b/Engine/Source/Runtime/Engine/Private/Texture2D.cpp
@@ -1359,7 +1359,10 @@ void FVirtualTexture2DResource::InitRHI()
 	InitializeEditorResources(VirtualTexture);
 #endif
 
-	TextureRHI->SetOwnerName(GetOwnerName());
+	if (TextureRHI.IsValid())
+	{
+		TextureRHI->SetOwnerName(GetOwnerName());
+	}
 }
 
 #if WITH_EDITOR
