Added input types check to LinearInterpolation that fixes crash in Material Editor when input vector dimensions differ.

diff --git a/Engine/Source/Runtime/Engine/Private/Materials/HLSLMaterialDerivativeAutogen.cpp b/Engine/Source/Runtime/Engine/Private/Materials/HLSLMaterialDerivativeAutogen.cpp
--- a/Engine/Source/Runtime/Engine/Private/Materials/HLSLMaterialDerivativeAutogen.cpp
+++ b/Engine/Source/Runtime/Engine/Private/Materials/HLSLMaterialDerivativeAutogen.cpp
@@ -117,45 +117,29 @@ static inline const TCHAR * GetDerivVectorName(EDerivativeType Type)
 
 EDerivativeType GetDerivType(EMaterialValueType ValueType, bool bAllowNonFloat)
 {
-	EDerivativeType Ret = EDerivativeType::None;
-	if (ValueType == MCT_Float1 || ValueType == MCT_Float)
-	{
-		Ret = EDerivativeType::Float1;
-	}
-	else if (ValueType == MCT_Float2)
-	{
-		Ret = EDerivativeType::Float2;
-	}
-	else if (ValueType == MCT_Float3)
-	{
-		Ret = EDerivativeType::Float3;
-	}
-	else if (ValueType == MCT_Float4)
-	{
-		Ret = EDerivativeType::Float4;
-	}
-	else if (ValueType == MCT_LWCScalar)
-	{
-		Ret = EDerivativeType::LWCScalar;
-	}
-	else if (ValueType == MCT_LWCVector2)
-	{
-		Ret = EDerivativeType::LWCVector2;
-	}
-	else if (ValueType == MCT_LWCVector3)
-	{
-		Ret = EDerivativeType::LWCVector3;
-	}
-	else if (ValueType == MCT_LWCVector4)
-	{
-		Ret = EDerivativeType::LWCVector4;
-	}
-	else
-	{
-		Ret = EDerivativeType::None;
+	switch (ValueType)
+	{
+	case MCT_Float:
+	case MCT_Float1:
+		return EDerivativeType::Float1;
+	case MCT_Float2:
+		return EDerivativeType::Float2;
+	case MCT_Float3:
+		return EDerivativeType::Float3;
+	case MCT_Float4:
+		return EDerivativeType::Float4;
+	case MCT_LWCScalar:
+		return EDerivativeType::LWCScalar;
+	case MCT_LWCVector2:
+		return EDerivativeType::LWCVector2;
+	case MCT_LWCVector3:
+		return EDerivativeType::LWCVector3;
+	case MCT_LWCVector4:
+		return EDerivativeType::LWCVector4;
+	default:
 		check(bAllowNonFloat);
+		return EDerivativeType::None;
 	}
-	return Ret;
 }
 
 static EMaterialValueType GetMaterialTypeFromDerivType(EDerivativeType Type)
@@ -180,10 +164,8 @@ static EMaterialValueType GetMaterialTypeFromDerivType(EDerivativeType Type)
 		return MCT_LWCVector4;
 	default:
 		check(0);
-		break;
+		return (EMaterialValueType)0; // invalid, should be a Float 1/2/3/4, break at the check(0);
 	}
-
-	return (EMaterialValueType)0; // invalid, should be a Float 1/2/3/4, break at the check(0);
 }
 
 static FString CoerceFloat(FHLSLMaterialTranslator& Translator, const TCHAR* Value, EDerivativeType DstType, EDerivativeType SrcType)
@@ -1074,6 +1056,13 @@ int32 FMaterialDerivativeAutogen::GenerateLerpFunc(FHLSLMaterialTranslator& Tran
 	
 	const EMaterialValueType ResultType = Translator.GetArithmeticResultType(A, B);
 	const EMaterialValueType AlphaType = (MakeNonLWCType(ResultType) == MakeNonLWCType(SDerivInfo.Type)) ? MakeNonLWCType(ResultType) : MCT_Float1;
+
+	// Early out if the result type determined by input types is invalid.
+	if (ResultType == EMaterialValueType::MCT_Unknown || AlphaType == EMaterialValueType::MCT_Unknown)
+	{
+		return INDEX_NONE;
+	}
+
 	const uint32 NumResultComponents = GetNumComponents(ResultType);
 
 	const bool bAllZeroDeriv = (ADerivInfo.DerivativeStatus == EDerivativeStatus::Zero && BDerivInfo.DerivativeStatus == EDerivativeStatus::Zero && SDerivInfo.DerivativeStatus == EDerivativeStatus::Zero);
