PR #9754: Material Graph Comments, blueprint feature parity.
Hide Comment Bubbles, Selection/Group Movements, Colored Comment Bubbles. (Contributed by [REDACTED_PERSON])

diff --git a/Engine/Source/Editor/MaterialEditor/Private/MaterialEditor.cpp b/Engine/Source/Editor/MaterialEditor/Private/MaterialEditor.cpp
--- a/Engine/Source/Editor/MaterialEditor/Private/MaterialEditor.cpp
+++ b/Engine/Source/Editor/MaterialEditor/Private/MaterialEditor.cpp
@@ -6515,6 +6515,9 @@ void FMaterialEditor::NotifyPostChange( const FPropertyChangedEvent& PropertyCha
 		const FName PropertyName = PropertyChangedEvent.MemberProperty ? PropertyChangedEvent.MemberProperty->GetFName() : NAME_None;
 		if (PropertyName != GET_MEMBER_NAME_CHECKED(UMaterialExpressionComment, Text)
 		&& PropertyName != GET_MEMBER_NAME_CHECKED(UMaterialExpressionComment, CommentColor)
+		&& PropertyName != GET_MEMBER_NAME_CHECKED(UMaterialExpressionComment, bColorCommentBubble)
+		&& PropertyName != GET_MEMBER_NAME_CHECKED(UMaterialExpressionComment, bGroupMode)
+		&& PropertyName != GET_MEMBER_NAME_CHECKED(UMaterialExpressionComment, bCommentBubbleVisible_InDetailsPanel)
 		&& PropertyName != GET_MEMBER_NAME_CHECKED(UMaterialExpression, Desc)
 		&& PropertyName != GET_MEMBER_NAME_CHECKED(UMaterialExpressionVectorParameter, ChannelNames)
 		&& PropertyName != GET_MEMBER_NAME_CHECKED(UMaterialExpressionTextureSampleParameter, ChannelNames))
diff --git a/Engine/Source/Editor/UnrealEd/Private/MaterialGraph.cpp b/Engine/Source/Editor/UnrealEd/Private/MaterialGraph.cpp
--- a/Engine/Source/Editor/UnrealEd/Private/MaterialGraph.cpp
+++ b/Engine/Source/Editor/UnrealEd/Private/MaterialGraph.cpp
@@ -726,7 +726,10 @@ void UMaterialGraph::LinkMaterialExpressionsFromGraph()
 						|| Comment->Text != CommentNode->NodeComment
 						|| Comment->SizeX != CommentNode->NodeWidth
 						|| Comment->SizeY != CommentNode->NodeHeight
-						|| Comment->CommentColor != CommentNode->CommentColor)
+						|| Comment->CommentColor != CommentNode->CommentColor
+						|| Comment->bCommentBubbleVisible_InDetailsPanel != CommentNode->bCommentBubbleVisible_InDetailsPanel
+						|| Comment->bGroupMode != (CommentNode->MoveMode == ECommentBoxMode::GroupMovement)
+						|| Comment->bColorCommentBubble != CommentNode->bColorCommentBubble)
 					{
 						Comment->Modify();
 
@@ -737,6 +740,9 @@ void UMaterialGraph::LinkMaterialExpressionsFromGraph()
 						Comment->SizeX = CommentNode->NodeWidth;
 						Comment->SizeY = CommentNode->NodeHeight;
 						Comment->CommentColor = CommentNode->CommentColor;
+						Comment->bCommentBubbleVisible_InDetailsPanel = CommentNode->bCommentBubbleVisible_InDetailsPanel;
+						Comment->bGroupMode = (CommentNode->MoveMode == ECommentBoxMode::GroupMovement);
+						Comment->bColorCommentBubble = CommentNode->bColorCommentBubble;
 					}
 				}
 			}
diff --git a/Engine/Source/Editor/UnrealEd/Private/MaterialGraphNode_Comment.cpp b/Engine/Source/Editor/UnrealEd/Private/MaterialGraphNode_Comment.cpp
--- a/Engine/Source/Editor/UnrealEd/Private/MaterialGraphNode_Comment.cpp
+++ b/Engine/Source/Editor/UnrealEd/Private/MaterialGraphNode_Comment.cpp
@@ -133,6 +133,11 @@ void UMaterialGraphNode_Comment::PostPlacedNewNode()
 		NodeWidth = MaterialExpressionComment->SizeX;
 		NodeHeight = MaterialExpressionComment->SizeY;
 		CommentColor = MaterialExpressionComment->CommentColor;
+		bCommentBubbleVisible_InDetailsPanel = MaterialExpressionComment->bCommentBubbleVisible_InDetailsPanel;
+		bCommentBubbleVisible = MaterialExpressionComment->bCommentBubbleVisible_InDetailsPanel;
+		bCommentBubblePinned = MaterialExpressionComment->bCommentBubbleVisible_InDetailsPanel;
+		MoveMode = MaterialExpressionComment->bGroupMode ? ECommentBoxMode::GroupMovement : ECommentBoxMode::NoGroupMovement;
+		bColorCommentBubble = MaterialExpressionComment->bColorCommentBubble;
 	}
 }
 
diff --git a/Engine/Source/Runtime/Engine/Classes/Materials/MaterialExpressionComment.h b/Engine/Source/Runtime/Engine/Classes/Materials/MaterialExpressionComment.h
--- a/Engine/Source/Runtime/Engine/Classes/Materials/MaterialExpressionComment.h
+++ b/Engine/Source/Runtime/Engine/Classes/Materials/MaterialExpressionComment.h
@@ -32,6 +32,18 @@ class UMaterialExpressionComment : public UMaterialExpression
 	UPROPERTY(EditAnywhere, Category=MaterialExpressionComment, meta=(ClampMin=1, ClampMax=1000))
 	int32 FontSize;
 
+	/** Whether to show a zoom-invariant comment bubble when zoomed out (making the comment readable at any distance). */
+	UPROPERTY(EditAnywhere, Category=MaterialExpressionComment, meta=(DisplayName="Show Bubble When Zoomed"))
+	uint32 bCommentBubbleVisible_InDetailsPanel:1;
+
+	/** Whether to use Comment Color to color the background of the comment bubble shown when zoomed out. */
+	UPROPERTY(EditAnywhere, Category=MaterialExpressionComment, meta=(DisplayName="Color Bubble", EditCondition=bCommentBubbleVisible_InDetailsPanel))
+	uint32 bColorCommentBubble:1;
+
+	/** Whether the comment should move any fully enclosed nodes around when it is moved */
+	UPROPERTY(EditAnywhere, Category=MaterialExpressionComment)
+	bool bGroupMode;
+	
 	//~ Begin UObject Interface
 #if WITH_EDITOR
 	virtual void PostEditChangeProperty(FPropertyChangedEvent& PropertyChangedEvent) override;
diff --git a/Engine/Source/Runtime/Engine/Private/Materials/MaterialExpressions.cpp b/Engine/Source/Runtime/Engine/Private/Materials/MaterialExpressions.cpp
--- a/Engine/Source/Runtime/Engine/Private/Materials/MaterialExpressions.cpp
+++ b/Engine/Source/Runtime/Engine/Private/Materials/MaterialExpressions.cpp
@@ -11636,6 +11636,10 @@ UMaterialExpressionComment::UMaterialExpressionComment(const FObjectInitializer&
 	: Super(ObjectInitializer)
 	, CommentColor(FLinearColor::White)
 	, FontSize(18)
+	, bCommentBubbleVisible_InDetailsPanel(false)
+	, bColorCommentBubble(false)
+	, bGroupMode(true)
+
 {
 	// Structure to hold one-time initialization
 	struct FConstructorStatics
@@ -11686,6 +11690,33 @@ void UMaterialExpressionComment::PostEditChangeProperty(FPropertyChangedEvent& P
 				CastChecked<UMaterialGraphNode_Comment>(GraphNode)->FontSize = FontSize;
 			}
 		}
+		else if (PropertyName == GET_MEMBER_NAME_CHECKED(UMaterialExpressionComment, bColorCommentBubble))
+		{
+			if (GraphNode)
+			{
+				GraphNode->Modify();
+				CastChecked<UMaterialGraphNode_Comment>(GraphNode)->bColorCommentBubble = bColorCommentBubble;
+			}
+		}
+		else if (PropertyName == GET_MEMBER_NAME_CHECKED(UMaterialExpressionComment, bCommentBubbleVisible_InDetailsPanel))
+		{
+			if (GraphNode)
+			{
+				GraphNode->Modify();
+				UMaterialGraphNode_Comment* CommentNode = CastChecked<UMaterialGraphNode_Comment>(GraphNode);
+				CommentNode->bCommentBubbleVisible_InDetailsPanel = bCommentBubbleVisible_InDetailsPanel;
+				CommentNode->bCommentBubbleVisible = bCommentBubbleVisible_InDetailsPanel;
+				CommentNode->bCommentBubblePinned = bCommentBubbleVisible_InDetailsPanel;
+			}
+		}
+		else if (PropertyName == GET_MEMBER_NAME_CHECKED(UMaterialExpressionComment, bGroupMode))
+		{
+			if (GraphNode)
+			{
+				GraphNode->Modify();
+				CastChecked<UMaterialGraphNode_Comment>(GraphNode)->MoveMode = bGroupMode ? ECommentBoxMode::GroupMovement : ECommentBoxMode::NoGroupMovement;
+			}
+		}
 
 		// Don't need to update preview after changing comments
 		bNeedToUpdatePreview = false;
