Bone buffer pool and cloth sim data pool owner name tracking

diff --git a/Engine/Source/Runtime/Engine/Private/GPUSkinCache.cpp b/Engine/Source/Runtime/Engine/Private/GPUSkinCache.cpp
--- a/Engine/Source/Runtime/Engine/Private/GPUSkinCache.cpp
+++ b/Engine/Source/Runtime/Engine/Private/GPUSkinCache.cpp
@@ -1553,8 +1553,7 @@ bool FGPUSkinCache::ProcessEntry(
 			}
 		}
 
-		const FName OwnerName(FString::Printf(TEXT("%s [LOD%d]"), *Skin->GetAssetPathName().ToString(), LODIndex));
-		FRWBuffersAllocation* NewPositionAllocation = TryAllocBuffer(TotalNumVertices, WithTangents, bEntryUseIntermediateTangents, InterAccumTangentBufferSize, RHICmdList, OwnerName);
+		FRWBuffersAllocation* NewPositionAllocation = TryAllocBuffer(TotalNumVertices, WithTangents, bEntryUseIntermediateTangents, InterAccumTangentBufferSize, RHICmdList, Skin->GetAssetPathName(LODIndex));
 		if (!NewPositionAllocation)
 		{
 			if (GSkinCachePrintMemorySummary > 0)
diff --git a/Engine/Source/Runtime/Engine/Private/GPUSkinVertexFactory.cpp b/Engine/Source/Runtime/Engine/Private/GPUSkinVertexFactory.cpp
--- a/Engine/Source/Runtime/Engine/Private/GPUSkinVertexFactory.cpp
+++ b/Engine/Source/Runtime/Engine/Private/GPUSkinVertexFactory.cpp
@@ -247,7 +247,7 @@ static bool DeferSkeletalLockAndFillToRHIThread()
 }
 
 bool FGPUBaseSkinVertexFactory::FShaderDataType::UpdateBoneData(FRHICommandListImmediate& RHICmdList, const TArray<FMatrix44f>& ReferenceToLocalMatrices,
-	const TArray<FBoneIndexType>& BoneMap, uint32 RevisionNumber, bool bPrevious, ERHIFeatureLevel::Type InFeatureLevel, bool bUseSkinCache, bool bForceUpdateImmediately)
+	const TArray<FBoneIndexType>& BoneMap, uint32 RevisionNumber, bool bPrevious, ERHIFeatureLevel::Type InFeatureLevel, bool bUseSkinCache, bool bForceUpdateImmediately, const FName& AssetPathName)
 {
 	// stat disabled by default due to low-value/high-frequency
 	//QUICK_SCOPE_CYCLE_COUNTER(STAT_FGPUBaseSkinVertexFactory_UpdateBoneData);
@@ -281,6 +281,7 @@ bool FGPUBaseSkinVertexFactory::FShaderDataType::UpdateBoneData(FRHICommandListI
 			}
 			*CurrentBoneBuffer = BoneBufferPool.CreatePooledResource(VectorArraySize);
 			check(IsValidRef(*CurrentBoneBuffer));
+			CurrentBoneBuffer->VertexBufferRHI->SetOwnerName(AssetPathName);
 		}
 		if(NumBones)
 		{
@@ -1160,7 +1161,7 @@ IMPLEMENT_TYPE_LAYOUT(TGPUSkinAPEXClothVertexFactoryShaderParameters);
 -----------------------------------------------------------------------------*/
 
 bool FGPUBaseSkinAPEXClothVertexFactory::ClothShaderType::UpdateClothSimulData(FRHICommandListImmediate& RHICmdList, const TArray<FVector3f>& InSimulPositions,
-	const TArray<FVector3f>& InSimulNormals, uint32 RevisionNumber, ERHIFeatureLevel::Type FeatureLevel, bool bForceUpdateImmediately)
+	const TArray<FVector3f>& InSimulNormals, uint32 RevisionNumber, ERHIFeatureLevel::Type FeatureLevel, bool bForceUpdateImmediately, const FName& AssetPathName)
 {
 	QUICK_SCOPE_CYCLE_COUNTER(STAT_FGPUBaseSkinAPEXClothVertexFactory_UpdateClothSimulData);
 
@@ -1183,6 +1184,7 @@ bool FGPUBaseSkinAPEXClothVertexFactory::ClothShaderType::UpdateClothSimulData(F
 		}
 		*CurrentClothBuffer = ClothSimulDataBufferPool.CreatePooledResource(VectorArraySize);
 		check(IsValidRef(*CurrentClothBuffer));
+		CurrentClothBuffer->VertexBufferRHI->SetOwnerName(AssetPathName);
 	}
 
 	if(NumSimulVerts)
diff --git a/Engine/Source/Runtime/Engine/Private/SkeletalRender.cpp b/Engine/Source/Runtime/Engine/Private/SkeletalRender.cpp
--- a/Engine/Source/Runtime/Engine/Private/SkeletalRender.cpp
+++ b/Engine/Source/Runtime/Engine/Private/SkeletalRender.cpp
@@ -236,10 +236,17 @@ float FSkeletalMeshObject::GetScreenSize(int32 LODIndex) const
 	return 0.f;
 }
 
-FName FSkeletalMeshObject::GetAssetPathName() const
+FName FSkeletalMeshObject::GetAssetPathName(int32 LODIndex) const
 {
 #if RHI_ENABLE_RESOURCE_INFO
-	return AssetPathName;
+	if (LODIndex > -1)
+	{
+		return FName(FString::Printf(TEXT("%s [LOD%d]"), *AssetPathName.ToString(), LODIndex));
+	}
+	else
+	{
+		return AssetPathName;
+	}
 #else
 	return NAME_None;
 #endif
diff --git a/Engine/Source/Runtime/Engine/Private/SkeletalRenderGPUSkin.cpp b/Engine/Source/Runtime/Engine/Private/SkeletalRenderGPUSkin.cpp
--- a/Engine/Source/Runtime/Engine/Private/SkeletalRenderGPUSkin.cpp
+++ b/Engine/Source/Runtime/Engine/Private/SkeletalRenderGPUSkin.cpp
@@ -688,16 +688,18 @@ void FSkeletalMeshObjectGPUSkin::ProcessUpdatedDynamicData(EGPUSkinCacheEntryMod
 			bool bShouldUseSeparateMatricesForRayTracing = false;
 		#endif
 
+			const FName OwnerName = GetAssetPathName(LODIndex);
+
 			// if we have previous reference to loca, we also update to previous frame
 			if (DynamicData->PreviousReferenceToLocal.Num() > 0)
 			{
 				TArray<FMatrix44f>& PreviousReferenceToLocalMatrices = bShouldUseSeparateMatricesForRayTracing ? DynamicData->PreviousReferenceToLocalForRayTracing : DynamicData->PreviousReferenceToLocal;
-				ShaderData.UpdateBoneData(RHICmdList, PreviousReferenceToLocalMatrices, Section.BoneMap, RevisionNumber, true, FeatureLevel, bUseSkinCache, DynamicData->bForceUpdateDynamicDataImmediately);
+				ShaderData.UpdateBoneData(RHICmdList, PreviousReferenceToLocalMatrices, Section.BoneMap, RevisionNumber, true, FeatureLevel, bUseSkinCache, DynamicData->bForceUpdateDynamicDataImmediately, OwnerName);
 			}
 
 			// Create a uniform buffer from the bone transforms.
 			TArray<FMatrix44f>& ReferenceToLocalMatrices = bShouldUseSeparateMatricesForRayTracing ? DynamicData->ReferenceToLocalForRayTracing : DynamicData->ReferenceToLocal;
-			bool bNeedFence = ShaderData.UpdateBoneData(RHICmdList, ReferenceToLocalMatrices, Section.BoneMap, RevisionNumber, false, FeatureLevel, bUseSkinCache, DynamicData->bForceUpdateDynamicDataImmediately);
+			bool bNeedFence = ShaderData.UpdateBoneData(RHICmdList, ReferenceToLocalMatrices, Section.BoneMap, RevisionNumber, false, FeatureLevel, bUseSkinCache, DynamicData->bForceUpdateDynamicDataImmediately, OwnerName);
 			ShaderData.UpdatedFrameNumber = FrameNumberToPrepare;
 
 			// Update uniform buffer for APEX cloth simulation mesh positions and normals
@@ -706,7 +708,7 @@ void FSkeletalMeshObjectGPUSkin::ProcessUpdatedDynamicData(EGPUSkinCacheEntryMod
 				FGPUBaseSkinAPEXClothVertexFactory::ClothShaderType& ClothShaderData = VertexFactoryData.ClothVertexFactories[SectionIdx]->GetClothShaderData();
 				ClothShaderData.ClothBlendWeight = DynamicData->ClothBlendWeight;
 
-				bNeedFence = ClothShaderData.UpdateClothSimulData(RHICmdList, SimData->Positions, SimData->Normals, RevisionNumber, FeatureLevel, DynamicData->bForceUpdateDynamicDataImmediately) || bNeedFence;
+				bNeedFence = ClothShaderData.UpdateClothSimulData(RHICmdList, SimData->Positions, SimData->Normals, RevisionNumber, FeatureLevel, DynamicData->bForceUpdateDynamicDataImmediately, OwnerName) || bNeedFence;
 				// Transform from cloth space to local space. Cloth space is relative to cloth root bone, local space is component space.
 				ClothShaderData.GetClothToLocalForWriting() = FMatrix44f(SimData->ComponentRelativeTransform.ToMatrixWithScale());
 			}
diff --git a/Engine/Source/Runtime/Engine/Public/GPUSkinVertexFactory.h b/Engine/Source/Runtime/Engine/Public/GPUSkinVertexFactory.h
--- a/Engine/Source/Runtime/Engine/Public/GPUSkinVertexFactory.h
+++ b/Engine/Source/Runtime/Engine/Public/GPUSkinVertexFactory.h
@@ -214,7 +214,8 @@ public:
 
 		// @param FrameTime from GFrameTime
 		bool UpdateBoneData(FRHICommandListImmediate& RHICmdList, const TArray<FMatrix44f>& ReferenceToLocalMatrices,
-			const TArray<FBoneIndexType>& BoneMap, uint32 RevisionNumber, bool bPrevious, ERHIFeatureLevel::Type FeatureLevel, bool bUseSkinCache, bool bForceUpdateImmediately);
+			const TArray<FBoneIndexType>& BoneMap, uint32 RevisionNumber, bool bPrevious, ERHIFeatureLevel::Type FeatureLevel, 
+			bool bUseSkinCache, bool bForceUpdateImmediately, const FName& AssetPathName);
 
 		void ReleaseBoneData()
 		{
@@ -623,7 +624,7 @@ public:
 		}
 
 		bool UpdateClothSimulData(FRHICommandListImmediate& RHICmdList, const TArray<FVector3f>& InSimulPositions, const TArray<FVector3f>& InSimulNormals, uint32 RevisionNumber, 
-									ERHIFeatureLevel::Type FeatureLevel, bool bForceUpdateImmediately);
+									ERHIFeatureLevel::Type FeatureLevel, bool bForceUpdateImmediately, const FName& AssetPathName);
 
 		void ReleaseClothSimulData()
 		{
diff --git a/Engine/Source/Runtime/Engine/Public/SkeletalRenderPublic.h b/Engine/Source/Runtime/Engine/Public/SkeletalRenderPublic.h
--- a/Engine/Source/Runtime/Engine/Public/SkeletalRenderPublic.h
+++ b/Engine/Source/Runtime/Engine/Public/SkeletalRenderPublic.h
@@ -213,7 +213,8 @@ public:
 
 	FColor GetSkinCacheVisualizationDebugColor(const FName& GPUSkinCacheVisualizationMode, uint32 SectionIndex) const;
 
-	FName GetAssetPathName() const;
+	/** Helper function to return the asset path name, optionally joined with the LOD index if LODIndex > -1. */
+	FName GetAssetPathName(int32 LODIndex = -1) const;
 
 #if RHI_RAYTRACING
 	/** Retrieve ray tracing geometry from the underlying mesh object */
