Remove complexity from audio streaming FChunkKey. This will allow a future change to remove the FSoundWaveProxy class and instead use FSoundWAveData directly by removing allowing FChunkKey to function w/o needing a valid FSoundWaveData.

- Moved the TWeakPtr<FSoundWaveData> from the chunk key to the FCacheElement
- Moved dependent functionality from the FChunkKey to the FCacheElement.

diff --git a/Engine/Source/Runtime/Engine/Private/AudioStreamingCache.cpp b/Engine/Source/Runtime/Engine/Private/AudioStreamingCache.cpp
--- a/Engine/Source/Runtime/Engine/Private/AudioStreamingCache.cpp
+++ b/Engine/Source/Runtime/Engine/Private/AudioStreamingCache.cpp
@@ -226,54 +226,36 @@ static FAutoConsoleCommand GDisableProfilingAudioCacheCommand(
 })
 );
 
+FAudioChunkCache::FChunkKey::FChunkKey(const FSoundWavePtr& InSoundWave, uint32 InChunkIndex
+#if WITH_EDITOR
+	, uint32 InChunkRevision
+#endif // #if WITH_EDITOR
+)
 
-FAudioChunkCache::FChunkKey::FChunkKey(const FChunkKey& Other)
-	: SoundWaveWeakPtr(Other.SoundWaveWeakPtr)
-	, ChunkIndex(Other.ChunkIndex)
-	, ObjectKey(Other.ObjectKey)
+	: ChunkIndex(InChunkIndex)
 #if WITH_EDITOR
-	, ChunkRevision(Other.ChunkRevision)
+	, ChunkRevision(InChunkRevision)
 #endif // #if WITH_EDITOR
 {
-	if (FSoundWavePtr SoundWaveDataPtr = SoundWaveWeakPtr.Pin())
+	if (InSoundWave.IsValid())
 	{
-		SoundWaveName = SoundWaveDataPtr->GetFName();
+		SoundWaveName = InSoundWave->GetFName();
+		ObjectKey = InSoundWave->GetFObjectKey();
 	}
 }
 
-
-FAudioChunkCache::FChunkKey& FAudioChunkCache::FChunkKey::operator=(const FChunkKey& Other)
-{
-	SoundWaveWeakPtr = Other.SoundWaveWeakPtr;
-	SoundWaveName = Other.SoundWaveName;
-	ChunkIndex = Other.ChunkIndex;
-	ObjectKey = Other.ObjectKey;
-
-#if WITH_EDITOR
-	ChunkRevision = Other.ChunkRevision;
-#endif // #if WITH_EDITOR
-
-
-	return *this;
-}
-
-FAudioChunkCache::FChunkKey::FChunkKey(const FSoundWavePtr& InSoundWave, uint32 InChunkIndex
+FAudioChunkCache::FChunkKey::FChunkKey(const FName& InSoundWaveName , const FObjectKey& InSoundWaveObjectKey , uint32 InChunkIndex
 #if WITH_EDITOR
-	, uint32 InChunkRevision
+		, uint32 InChunkRevision
 #endif // #if WITH_EDITOR
 )
-
-	: SoundWaveWeakPtr(InSoundWave)
+	: SoundWaveName(InSoundWaveName)
+	, ObjectKey(InSoundWaveObjectKey)
 	, ChunkIndex(InChunkIndex)
 #if WITH_EDITOR
 	, ChunkRevision(InChunkRevision)
 #endif // #if WITH_EDITOR
 {
-	if (FSoundWavePtr SoundWavePtr = SoundWaveWeakPtr.Pin())
-	{
-		SoundWaveName = SoundWavePtr->GetFName();
-		ObjectKey = SoundWavePtr->GetFObjectKey();
-	}
 }
 
 
@@ -286,48 +268,34 @@ bool FAudioChunkCache::FChunkKey::operator==(const FChunkKey& Other) const
 #endif
 }
 
-bool FAudioChunkCache::FChunkKey::IsChunkStale()
+#if DEBUG_STREAM_CACHE
+bool FAudioChunkCache::FCacheElement::IsBeingPlayed() const
 {
-	if(FSoundWavePtr SoundWaveDataPtr = SoundWaveWeakPtr.Pin())
-	{
-		return !SoundWaveDataPtr.IsValid()
-#if WITH_EDITOR
-			|| (SoundWaveDataPtr->GetCurrentChunkRevision() != ChunkRevision)
-#endif // #if WITH_EDITOR
-		;
+	const int32 NumActiveConsumers = NumConsumers.GetValue();
 
-	}
-
-	return true;
+	// if we 2 or more consumers, this chunk is being rendered.
+	// if we have 1 consumer, and we aren't Retained, then this chunk is being rendered
+	return (NumActiveConsumers > 1)
+		|| (NumActiveConsumers && (DebugInfo.LoadingBehavior != ESoundWaveLoadingBehavior::RetainOnLoad));
 }
 
-ESoundWaveLoadingBehavior FAudioChunkCache::FChunkKey::GetLoadingBehavior() const
+void FAudioChunkCache::FCacheElement::UpdateDebugInfoLoadingBehavior()
 {
-	if(FSoundWavePtr SoundWaveDataPtr = SoundWaveWeakPtr.Pin())
+	if(TSharedPtr<const FSoundWaveData> SoundWaveDataPtr = SoundWaveWeakPtr.Pin())
 	{
-		return SoundWaveDataPtr->GetLoadingBehavior();
+		// Recursing in no longer needed at this point since the inherited loading behavior has already been cached by the time this information is needed
+		DebugInfo.LoadingBehavior = SoundWaveDataPtr->GetLoadingBehavior();
+		DebugInfo.bLoadingBehaviorExternallyOverriden = SoundWaveDataPtr->WasLoadingBehaviorOverridden();
+	}
+	else
+	{
+		DebugInfo.LoadingBehavior = ESoundWaveLoadingBehavior::Uninitialized;
+		DebugInfo.bLoadingBehaviorExternallyOverriden = false;
 	}
-
-	return ESoundWaveLoadingBehavior::Uninitialized;
-}
-
-FStreamedAudioChunk& FAudioChunkCache::FChunkKey::GetChunk(uint32 InChunkIndex) const
-{
-	FSoundWavePtr SoundWaveDataPtr = SoundWaveWeakPtr.Pin();
-
-	// the Weakptr should be valid here since it's from a shared ptr up the stack
-	check(SoundWaveDataPtr.IsValid());
-
-	// This function shouldn't be called on audio marked "ForceInline."
-	ensureMsgf(SoundWaveDataPtr->GetLoadingBehavior() != ESoundWaveLoadingBehavior::ForceInline, TEXT("Calling GetNumChunks on a FSoundWaveProxy that is Force-Inlined is not allowed! SoundWave: %s - %s")
-		, *SoundWaveDataPtr->GetFName().ToString(), EnumToString(SoundWaveDataPtr->GetLoadingBehavior()));
-
-	check((ChunkIndex < (uint32)SoundWaveDataPtr->GetNumChunks()));
-
-	return SoundWaveDataPtr->GetChunk(ChunkIndex);
 }
+#endif // #if DEBUG_STREAM_CACHE
 
-uint32 FAudioChunkCache::FChunkKey::GetNumChunks() const
+uint32 FAudioChunkCache::FCacheElement::GetNumChunks() const
 {
 	if (FSoundWavePtr SoundWaveDataPtr = SoundWaveWeakPtr.Pin())
 	{
@@ -337,17 +305,39 @@ uint32 FAudioChunkCache::FChunkKey::GetNumChunks() const
 	return 0;
 }
 
-bool FAudioChunkCache::FChunkKey::IsRetainingAudio() const
+FStreamedAudioChunk* FAudioChunkCache::FCacheElement::GetChunk(uint32 InChunkIndex) const
 {
-	if (FSoundWavePtr SoundWaveDataPtr = SoundWaveWeakPtr.Pin())
+	FSoundWavePtr SoundWaveDataPtr = SoundWaveWeakPtr.Pin();
+
+	// the Weakptr should be valid here since it's from a shared ptr up the stack
+	if (ensure(SoundWaveDataPtr))
 	{
-		return SoundWaveDataPtr->IsRetainingAudio();
+		// This function shouldn't be called on audio marked "ForceInline."
+		checkf(SoundWaveDataPtr->GetLoadingBehavior() != ESoundWaveLoadingBehavior::ForceInline, TEXT("Calling GetNumChunks on a FSoundWaveProxy that is Force-Inlined is not allowed! SoundWave: %s - %s")
+			, *SoundWaveDataPtr->GetFName().ToString(), EnumToString(SoundWaveDataPtr->GetLoadingBehavior()));
+
+		if (ensureMsgf(InChunkIndex < (uint32)SoundWaveDataPtr->GetNumChunks(), TEXT("Attempt retrieve chunk %d when only %d exist on sound wave\"%s\"."), InChunkIndex, SoundWaveDataPtr->GetNumChunks(), *SoundWaveDataPtr->GetFName().ToString()))
+		{
+			return &SoundWaveDataPtr->GetChunk(InChunkIndex);
+		}
 	}
 
-	return false;
+	return nullptr;
 }
 
-void FAudioChunkCache::FChunkKey::ReleaseCompressedAudio()
+#if WITH_EDITOR
+bool FAudioChunkCache::FCacheElement::IsChunkStale()
+{
+	if(FSoundWavePtr SoundWaveDataPtr = SoundWaveWeakPtr.Pin())
+	{
+		return !SoundWaveDataPtr.IsValid() || (SoundWaveDataPtr->GetCurrentChunkRevision() != Key.ChunkRevision);
+	}
+
+	return true;
+}
+#endif // #if WITH_EDITOR
+
+void FAudioChunkCache::FCacheElement::ReleaseRetainedAudioOnSoundWave()
 {
 	if (FSoundWavePtr SoundWaveDataPtr = SoundWaveWeakPtr.Pin())
 	{
@@ -356,16 +346,16 @@ void FAudioChunkCache::FChunkKey::ReleaseCompressedAudio()
 	}
 }
 
-bool FAudioChunkCache::FChunkKey::WasLoadingBehaviorOverridden() const
+bool FAudioChunkCache::FCacheElement::IsSoundWaveRetainingAudio() const
 {
-	if(FSoundWavePtr SoundWaveDataPtr = SoundWaveWeakPtr.Pin())
+	if (FSoundWavePtr SoundWaveDataPtr = SoundWaveWeakPtr.Pin())
 	{
-		return SoundWaveDataPtr->WasLoadingBehaviorOverridden();
+		return SoundWaveDataPtr->IsRetainingAudio();
 	}
-
 	return false;
 }
 
+
 FCachedAudioStreamingManager::FCachedAudioStreamingManager(const FCachedAudioStreamingManagerParams& InitParams)
 {
 	LLM_SCOPE(ELLMTag::AudioStreamCache);
@@ -508,6 +498,12 @@ FAudioChunkHandle FCachedAudioStreamingManager::GetLoadedChunk(const FSoundWaveP
 		return {};
 	}
 
+	TSharedPtr<FSoundWaveData> SoundWaveData = SoundWave->GetSoundWaveData();
+	if (!ensure(SoundWaveData.IsValid()))
+	{
+		return {};
+	}
+
 	// If this sound wave is managed by a cache, use that to get the chunk:
 	FAudioChunkCache* Cache = GetCacheForWave(SoundWave);
 	if (Cache)
@@ -516,14 +512,14 @@ FAudioChunkHandle FCachedAudioStreamingManager::GetLoadedChunk(const FSoundWaveP
 		checkf(ChunkIndex != 0, TEXT("Decoder tried to access the zeroth chunk through the streaming manager. Use USoundWave::GetZerothChunk() instead."));
 
 		FAudioChunkCache::FChunkKey ChunkKey(
-			  SoundWave->GetSoundWaveData()
+			  SoundWaveData
 			, ChunkIndex
 #if WITH_EDITOR
 			, (uint32)SoundWave->GetCurrentChunkRevision()
 #endif
 		);
 
-		if (!FAudioChunkCache::IsKeyValid(ChunkKey))
+		if (!FAudioChunkCache::DoesKeyContainValidChunkIndex(ChunkKey, *SoundWaveData))
 		{
 			UE_LOG(LogAudioStreamCaching, Warning, TEXT("Invalid Chunk Index %d Requested for Wave %s!"), ChunkIndex, *SoundWave->GetFName().ToString());
 			return FAudioChunkHandle();
@@ -531,7 +527,7 @@ FAudioChunkHandle FCachedAudioStreamingManager::GetLoadedChunk(const FSoundWaveP
 
 		// The function call below increments the reference count to the internal chunk.
 		uint64 LookupIDForChunk = Cache->GetCacheLookupIDForChunk(ChunkKey);
-		TArrayView<uint8> LoadedChunk = Cache->GetChunk(ChunkKey, bBlockForLoad, (bForImmediatePlayback || bBlockForLoad), LookupIDForChunk);
+		TArrayView<uint8> LoadedChunk = Cache->GetChunk(ChunkKey, SoundWaveData, bBlockForLoad, (bForImmediatePlayback || bBlockForLoad), LookupIDForChunk);
 
 		// Ensure that, if we requested a synchronous load of this chunk, we didn't fail to load said chunk.
 		UE_CLOG(bBlockForLoad && !LoadedChunk.GetData(), LogAudioStreamCaching, Display, TEXT("Synchronous load of chunk index %d for SoundWave %s failed to return any data. Likely because the cache was blown."), ChunkIndex, *SoundWave->GetFName().ToString());
@@ -547,14 +543,14 @@ FAudioChunkHandle FCachedAudioStreamingManager::GetLoadedChunk(const FSoundWaveP
 		if (NextChunk != INDEX_NONE)
 		{
 			FAudioChunkCache::FChunkKey NextChunkKey(
-				SoundWave->GetSoundWaveData()
+				  SoundWaveData
 				, ((uint32)NextChunk)
 #if WITH_EDITOR
 				, (uint32)SoundWave->GetCurrentChunkRevision()
 #endif
 			);
 
-			uint64 LookupIDForNextChunk = Cache->AddOrTouchChunk(NextChunkKey, [](EAudioChunkLoadResult) {}, ENamedThreads::AnyThread, false);
+			uint64 LookupIDForNextChunk = Cache->AddOrTouchChunk(NextChunkKey, SoundWaveData, [](EAudioChunkLoadResult) {}, ENamedThreads::AnyThread, false);
 			if (LookupIDForNextChunk == InvalidAudioStreamCacheLookupID)
 			{
 				UE_LOG(LogAudioStreamCaching, Warning, TEXT("Failed to add follow-up chunk for %s. This should not happen with our current TrimMemory() logic"), *SoundWave->GetFName().ToString());
@@ -648,10 +644,11 @@ void FCachedAudioStreamingManager::AddReferenceToChunk(const FAudioChunkHandle&
 	check(Cache);
 
 	const FAudioChunkCache::FChunkKey ChunkKey(
-		  InHandle.CorrespondingWave.Pin()
+		  InHandle.CorrespondingWaveName
+		, InHandle.CorrespondingWaveObjectKey
 		, ((uint32)InHandle.ChunkIndex)
 #if WITH_EDITOR
-		, InHandle.ChunkGeneration
+		, InHandle.ChunkRevision
 #endif
 	);
 
@@ -665,10 +662,11 @@ void FCachedAudioStreamingManager::RemoveReferenceToChunk(const FAudioChunkHandl
 	check(Cache);
 
 	const FAudioChunkCache::FChunkKey ChunkKey(
-		  InHandle.CorrespondingWave.Pin()
+		  InHandle.CorrespondingWaveName
+		, InHandle.CorrespondingWaveObjectKey
 		, ((uint32)InHandle.ChunkIndex)
 #if WITH_EDITOR
-		, InHandle.ChunkGeneration
+		, InHandle.ChunkRevision
 #endif
 	);
 
@@ -690,7 +688,7 @@ bool FCachedAudioStreamingManager::RequestChunk(const FSoundWaveProxyPtr& SoundW
 #endif
 		);
 
-		uint64 LookupIDForChunk = Cache->AddOrTouchChunk(ChunkKey, OnLoadCompleted, ThreadToCallOnLoadCompletedOn, bForImmediatePlayback);
+		uint64 LookupIDForChunk = Cache->AddOrTouchChunk(ChunkKey, SoundWave->GetSoundWaveData(), OnLoadCompleted, ThreadToCallOnLoadCompletedOn, bForImmediatePlayback);
 
 		// Update the map entry through the streaming cache
 		Cache->SetCacheLookupIDForChunk(ChunkKey, LookupIDForChunk);
@@ -718,7 +716,7 @@ FAudioChunkCache::FAudioChunkCache(uint32 InMaxChunkSize, uint32 NumChunks, uint
 	CachePool.Reset(NumChunks);
 	for (uint32 Index = 0; Index < NumChunks; Index++)
 	{
-		CachePool.Emplace(MaxChunkSize, Index);
+		CachePool.Emplace(Index);
 	}
 	CacheOverflowCount.Set(0);
 }
@@ -731,7 +729,7 @@ FAudioChunkCache::~FAudioChunkCache()
 	check(NumberOfLoadsInFlight.GetValue() == 0);
 }
 
-uint64 FAudioChunkCache::AddOrTouchChunk(const FChunkKey& InKey, TFunction<void(EAudioChunkLoadResult)> OnLoadCompleted, ENamedThreads::Type CallbackThread, bool bNeededForPlayback)
+uint64 FAudioChunkCache::AddOrTouchChunk(const FChunkKey& InKey, const TSharedPtr<FSoundWaveData>& InSoundWaveData, TFunction<void(EAudioChunkLoadResult)> OnLoadCompleted, ENamedThreads::Type CallbackThread, bool bNeededForPlayback)
 {
 	// Update cache limit if needed.
 	if (!FMath::IsNearlyZero(StreamCacheSizeOverrideMBCVar) && StreamCacheSizeOverrideMBCVar > 0.0f)
@@ -739,7 +737,13 @@ uint64 FAudioChunkCache::AddOrTouchChunk(const FChunkKey& InKey, TFunction<void(
 		MemoryLimitBytes = ((uint64)(StreamCacheSizeOverrideMBCVar * 1024)) * 1024;
 	}
 
-	if (!IsKeyValid(InKey))
+	if (!InSoundWaveData.IsValid())
+	{
+		ExecuteOnLoadCompleteCallback(EAudioChunkLoadResult::ChunkOutOfBounds, OnLoadCompleted, CallbackThread);
+		return InvalidAudioStreamCacheLookupID;
+	}
+
+	if (!DoesKeyContainValidChunkIndex(InKey, *InSoundWaveData))
 	{
 		ensure(false);
 		ExecuteOnLoadCompleteCallback(EAudioChunkLoadResult::ChunkOutOfBounds, OnLoadCompleted, CallbackThread);
@@ -760,17 +764,14 @@ uint64 FAudioChunkCache::AddOrTouchChunk(const FChunkKey& InKey, TFunction<void(
 #if DEBUG_STREAM_CACHE
 		FoundElement->DebugInfo.NumTimesTouched++;
 
-		// Recursing in no longer needed at this point since the inherited loading behavior has already been cached by the time this information is needed
-		FoundElement->DebugInfo.LoadingBehavior = InKey.GetLoadingBehavior();
-
-		FoundElement->DebugInfo.bLoadingBehaviorExternallyOverriden = InKey.WasLoadingBehaviorOverridden();
+		FoundElement->UpdateDebugInfoLoadingBehavior();
 #endif
 
 		return FoundElement->CacheLookupID;
 	}
 	else
 	{
-		FCacheElement* CacheElement = InsertChunk(InKey);
+		FCacheElement* CacheElement = InsertChunk(InKey, InSoundWaveData);
 
 		if (!CacheElement)
 		{
@@ -780,27 +781,26 @@ uint64 FAudioChunkCache::AddOrTouchChunk(const FChunkKey& InKey, TFunction<void(
 
 #if DEBUG_STREAM_CACHE
 		CacheElement->DebugInfo.bWasCacheMiss = bNeededForPlayback;
-
-		CacheElement->DebugInfo.LoadingBehavior = InKey.GetLoadingBehavior();
-
-		CacheElement->DebugInfo.bLoadingBehaviorExternallyOverriden = InKey.WasLoadingBehaviorOverridden();
+		CacheElement->UpdateDebugInfoLoadingBehavior();
 #endif
-		const FStreamedAudioChunk& Chunk = InKey.GetChunk(InKey.ChunkIndex);
-		int32 ChunkDataSize = Chunk.AudioDataSize;
-
-		if (TrimCacheWhenOverBudgetCVar != 0 && (MemoryCounterBytes + ChunkDataSize) > MemoryLimitBytes)
+		if (const FStreamedAudioChunk* Chunk = CacheElement->GetChunk(InKey.ChunkIndex))
 		{
-			uint64 MemoryToTrim = 0;
-			if (MemoryLimitTrimPercentageCVar > 0.0f)
-			{
-				MemoryToTrim = MemoryLimitBytes * FMath::Min(MemoryLimitTrimPercentageCVar, 1.0f);
-			}
-			else
+			int32 ChunkDataSize = Chunk->AudioDataSize;
+
+			if (TrimCacheWhenOverBudgetCVar != 0 && (MemoryCounterBytes + ChunkDataSize) > MemoryLimitBytes)
 			{
-				MemoryToTrim = MemoryCounterBytes + ChunkDataSize - MemoryLimitBytes;
-			}
+				uint64 MemoryToTrim = 0;
+				if (MemoryLimitTrimPercentageCVar > 0.0f)
+				{
+					MemoryToTrim = MemoryLimitBytes * FMath::Min(MemoryLimitTrimPercentageCVar, 1.0f);
+				}
+				else
+				{
+					MemoryToTrim = MemoryCounterBytes + ChunkDataSize - MemoryLimitBytes;
+				}
 
-			TrimMemory(MemoryToTrim, true);
+				TrimMemory(MemoryToTrim, true);
+			}
 		}
 
 		KickOffAsyncLoad(CacheElement, InKey, OnLoadCompleted, CallbackThread, bNeededForPlayback);
@@ -808,7 +808,7 @@ uint64 FAudioChunkCache::AddOrTouchChunk(const FChunkKey& InKey, TFunction<void(
 		if (bNeededForPlayback && (bLogCacheMisses || AlwaysLogCacheMissesCVar))
 		{
 			// We missed 
-			const uint32 TotalNumChunksInWave = InKey.GetNumChunks();
+			const uint32 TotalNumChunksInWave = CacheElement->GetNumChunks();
 
 			FCacheMissInfo CacheMissInfo = { InKey.SoundWaveName, InKey.ChunkIndex, TotalNumChunksInWave, false };
 			CacheMissQueue.Enqueue(MoveTemp(CacheMissInfo));
@@ -818,7 +818,7 @@ uint64 FAudioChunkCache::AddOrTouchChunk(const FChunkKey& InKey, TFunction<void(
 	}
 }
 
-TArrayView<uint8> FAudioChunkCache::GetChunk(const FChunkKey& InKey, bool bBlockForLoadCompletion, bool bNeededForPlayback, uint64& OutCacheOffset)
+TArrayView<uint8> FAudioChunkCache::GetChunk(const FChunkKey& InKey, const TSharedPtr<FSoundWaveData>& InSoundWavePtr, bool bBlockForLoadCompletion, bool bNeededForPlayback, uint64& OutCacheOffset)
 {
 	FScopeLock ScopeLock(&CacheMutationCriticalSection);
 	FCacheElement* FoundElement = FindElementForKey(InKey);
@@ -847,7 +847,7 @@ TArrayView<uint8> FAudioChunkCache::GetChunk(const FChunkKey& InKey, bool bBlock
 	else
 	{
 		// If we missed it, kick off a new load with it.
-		FoundElement = InsertChunk(InKey);
+		FoundElement = InsertChunk(InKey, InSoundWavePtr);
 		if (!FoundElement)
 		{
 			OutCacheOffset = InvalidAudioStreamCacheLookupID;
@@ -867,62 +867,70 @@ TArrayView<uint8> FAudioChunkCache::GetChunk(const FChunkKey& InKey, bool bBlock
 
 		if (bBlockForLoadCompletion)
 		{
-			FStreamedAudioChunk& Chunk = InKey.GetChunk(InKey.ChunkIndex);
-			int32 ChunkAudioDataSize = Chunk.AudioDataSize;
+			if (FStreamedAudioChunk* Chunk = FoundElement->GetChunk(InKey.ChunkIndex))
+			{
+				int32 ChunkAudioDataSize = Chunk->AudioDataSize;
 #if DEBUG_STREAM_CACHE
-			FoundElement->DebugInfo.NumTotalChunks = InKey.GetNumChunks() - 1;
-			FoundElement->DebugInfo.TimeLoadStarted = FPlatformTime::Seconds();
+				FoundElement->DebugInfo.NumTotalChunks = FoundElement->GetNumChunks() - 1;
+				FoundElement->DebugInfo.TimeLoadStarted = FPlatformTime::Seconds();
 #endif
-			MemoryCounterBytes -= FoundElement->ChunkDataSize;
+				MemoryCounterBytes -= FoundElement->ChunkDataSize;
 
-			{
-				LLM_SCOPE(ELLMTag::AudioStreamCacheCompressedData);
+				{
+					LLM_SCOPE(ELLMTag::AudioStreamCacheCompressedData);
 
-				// Reallocate our chunk data This allows us to shrink if possible.
-				FoundElement->ChunkData = (uint8*)FMemory::Realloc(FoundElement->ChunkData, ChunkAudioDataSize);
-				void* DataDestPtr = FoundElement->ChunkData;
-				const bool Result = Chunk.GetCopy(&DataDestPtr);
+					// Reallocate our chunk data This allows us to shrink if possible.
+					FoundElement->ChunkData = (uint8*)FMemory::Realloc(FoundElement->ChunkData, ChunkAudioDataSize);
+					void* DataDestPtr = FoundElement->ChunkData;
+					const bool Result = Chunk->GetCopy(&DataDestPtr);
 
-				if (!Result)
-				{
-					UE_LOG(LogAudioStreamCaching, Warning, TEXT("Failed to retrieve chunk data from Bulk Data for soundwave: %s"), *InKey.SoundWaveName.ToString());
-					return TArrayView<uint8>();
+					if (!Result)
+					{
+						UE_LOG(LogAudioStreamCaching, Warning, TEXT("Failed to retrieve chunk data from Bulk Data for soundwave: %s"), *InKey.SoundWaveName.ToString());
+						return TArrayView<uint8>();
+					}
 				}
-			}
 
-			MemoryCounterBytes += ChunkAudioDataSize;
+				MemoryCounterBytes += ChunkAudioDataSize;
 
-			// Populate key and DataSize. The async read request was set up to write directly into CacheElement->ChunkData.
-			FoundElement->Key = InKey;
-			FoundElement->ChunkDataSize = ChunkAudioDataSize;
-			FoundElement->bIsLoaded = true;
+				// Populate key and DataSize. The async read request was set up to write directly into CacheElement->ChunkData.
+				FoundElement->Key = InKey;
+				FoundElement->ChunkDataSize = ChunkAudioDataSize;
+				FoundElement->bIsLoaded = true;
 #if DEBUG_STREAM_CACHE
-			FoundElement->DebugInfo.TimeToLoad = (FPlatformTime::Seconds() - FoundElement->DebugInfo.TimeLoadStarted) * 1000.0f;
+				FoundElement->DebugInfo.TimeToLoad = (FPlatformTime::Seconds() - FoundElement->DebugInfo.TimeLoadStarted) * 1000.0f;
 
 #endif
-			// If this value is ever negative, it means that we're decrementing more than we're incrementing:
-			if (ensureMsgf(FoundElement->NumConsumers.GetValue() >= 0, TEXT("NumConsumers was negative for FoundElement. Reseting to 1")))
-			{
-				FoundElement->NumConsumers.Increment();
+				// If this value is ever negative, it means that we're decrementing more than we're incrementing:
+				if (ensureMsgf(FoundElement->NumConsumers.GetValue() >= 0, TEXT("NumConsumers was negative for FoundElement. Reseting to 1")))
+				{
+					FoundElement->NumConsumers.Increment();
+				}
+				else
+				{
+					FoundElement->NumConsumers.Set(1);
+				}
+
+				return TArrayView<uint8>(FoundElement->ChunkData, ChunkAudioDataSize);
 			}
 			else
 			{
-				FoundElement->NumConsumers.Set(1);
+				UE_LOG(LogAudioStreamCaching, Error, TEXT("Failed to get chunk %d from soundwave: %s"), InKey.ChunkIndex, *InKey.SoundWaveName.ToString());
 			}
-
-			return TArrayView<uint8>(FoundElement->ChunkData, ChunkAudioDataSize);
 		}
 		else
 		{
 			KickOffAsyncLoad(FoundElement, InKey, [](EAudioChunkLoadResult InResult) {}, ENamedThreads::AnyThread, bNeededForPlayback);
 		}
+
 		if (bLogCacheMisses && !bBlockForLoadCompletion)
 		{
 			// Chunks missing. Log this as a miss.
-			const uint32 TotalNumChunksInWave = InKey.GetNumChunks();
+			const uint32 TotalNumChunksInWave = FoundElement->GetNumChunks();
 			FCacheMissInfo CacheMissInfo = { InKey.SoundWaveName, InKey.ChunkIndex, TotalNumChunksInWave, false };
 			CacheMissQueue.Enqueue(MoveTemp(CacheMissInfo));
 		}
+
 		// We missed, return an empty array view.
 		return TArrayView<uint8>();
 	}
@@ -964,7 +972,7 @@ void FAudioChunkCache::ClearCache()
 
 	for (uint32 Index = 0; Index < NumChunks; Index++)
 	{
-		CachePool.Emplace(MaxChunkSize, Index);
+		CachePool.Emplace(Index);
 	}
 
 	MostRecentElement = nullptr;
@@ -1029,10 +1037,10 @@ uint64 FAudioChunkCache::TrimMemory(uint64 BytesToFree, bool bInAllowRetainedChu
 		ElementToStopAt = MostRecentElement->LessRecentElement;
 		while (CurrentElement != ElementToStopAt && BytesFreed < BytesToFree)
 		{
-			if (CurrentElement->ChunkDataSize != 0 && CurrentElement->Key.IsSoundWaveValid() && CurrentElement->Key.IsRetainingAudio())
+			if (CurrentElement->ChunkDataSize != 0 && CurrentElement->IsSoundWaveRetainingAudio())
 			{
 				// Directly release the retained audio (TODO: this is on the audio thread right?)
-				CurrentElement->Key.ReleaseCompressedAudio();
+				CurrentElement->ReleaseRetainedAudioOnSoundWave();
 				if (CurrentElement->CanEvictChunk())
 				{
 					uint32 ChunkSize = CurrentElement->ChunkDataSize;
@@ -1363,7 +1371,7 @@ bool FAudioChunkCache::ShouldAddNewChunk() const
 	return (ChunksInUse < CachePool.Num()) && (MemoryCounterBytes.Load() < MemoryLimitBytes);
 }
 
-FAudioChunkCache::FCacheElement* FAudioChunkCache::InsertChunk(const FChunkKey& InKey)
+FAudioChunkCache::FCacheElement* FAudioChunkCache::InsertChunk(const FChunkKey& InKey, const TSharedPtr<FSoundWaveData>& InSoundWavePtr)
 {
 	FCacheElement* CacheElement = nullptr;
 
@@ -1406,6 +1414,7 @@ FAudioChunkCache::FCacheElement* FAudioChunkCache::InsertChunk(const FChunkKey&
 		check(CacheElement);
 		CacheElement->bIsLoaded = false;
 		CacheElement->Key = InKey;
+		CacheElement->SoundWaveWeakPtr = InSoundWavePtr;
 		TouchElement(CacheElement);
 
 		// If we've got multiple chunks, we can not cache the least recent chunk
@@ -1594,8 +1603,15 @@ void FAudioChunkCache::KickOffAsyncLoad(FCacheElement* CacheElement, const FChun
 {
 	check(CacheElement);
 
-	const FStreamedAudioChunk& Chunk = InKey.GetChunk(InKey.ChunkIndex);
-	int32 ChunkDataSize = Chunk.AudioDataSize;
+	const FStreamedAudioChunk* Chunk = CacheElement->GetChunk(InKey.ChunkIndex);
+
+	if (nullptr == Chunk)
+	{
+		UE_LOG(LogAudioStreamCaching, Error, TEXT("Failed to kick off async load of chunk %d in soundwave \"%s\". Could not find chunk."), InKey.ChunkIndex, *InKey.SoundWaveName.ToString());
+		return;
+	}
+
+	int32 ChunkDataSize = Chunk->AudioDataSize;
 
 	EAsyncIOPriorityAndFlags AsyncIOPriority = GetAsyncPriorityForChunk(InKey, bNeededForPlayback);
 
@@ -1605,21 +1621,20 @@ void FAudioChunkCache::KickOffAsyncLoad(FCacheElement* CacheElement, const FChun
 		LLM_SCOPE(ELLMTag::AudioStreamCacheCompressedData);
 
 		// Reallocate our chunk data This allows us to shrink if possible.
-		CacheElement->ChunkData = (uint8*)FMemory::Realloc(CacheElement->ChunkData, Chunk.AudioDataSize);
-		CacheElement->ChunkDataSize = Chunk.AudioDataSize;
+		CacheElement->ChunkData = (uint8*)FMemory::Realloc(CacheElement->ChunkData, Chunk->AudioDataSize);
+		CacheElement->ChunkDataSize = Chunk->AudioDataSize;
 	}
 
 	MemoryCounterBytes += CacheElement->ChunkDataSize;
 
 #if DEBUG_STREAM_CACHE
-	CacheElement->DebugInfo.NumTotalChunks = InKey.GetNumChunks() - 1;
-	CacheElement->DebugInfo.LoadingBehavior = InKey.GetLoadingBehavior();
-	CacheElement->DebugInfo.bLoadingBehaviorExternallyOverriden = InKey.WasLoadingBehaviorOverridden();
+	CacheElement->DebugInfo.NumTotalChunks = CacheElement->GetNumChunks() - 1;
+	CacheElement->UpdateDebugInfoLoadingBehavior();
 #endif
 
 	// In editor, we retrieve from the DDC. In non-editor situations, we read the chunk async from the pak file.
 #if WITH_EDITORONLY_DATA
-	if (Chunk.DerivedDataKey.IsEmpty() == false)
+	if (Chunk->DerivedDataKey.IsEmpty() == false)
 	{
 		CacheElement->ChunkDataSize = ChunkDataSize;
 
@@ -1665,7 +1680,7 @@ void FAudioChunkCache::KickOffAsyncLoad(FCacheElement* CacheElement, const FChun
 		NumberOfLoadsInFlight.Increment();
 
 		CacheElement->DDCTask.Reset(new FAsyncStreamDerivedChunkTask(
-			Chunk.DerivedDataKey,
+			Chunk->DerivedDataKey,
 			CacheElement->ChunkData,
 			ChunkDataSize,
 			&NumberOfLoadsInFlight,
@@ -1683,7 +1698,7 @@ void FAudioChunkCache::KickOffAsyncLoad(FCacheElement* CacheElement, const FChun
 		}
 
 		// Sanity check our bulk data against our currently allocated chunk size in the cache.
-		const int32 ChunkBulkDataSize = Chunk.BulkData.GetBulkDataSize();
+		const int32 ChunkBulkDataSize = Chunk->BulkData.GetBulkDataSize();
 		check(ChunkDataSize <= ChunkBulkDataSize);
 		check(((uint32)ChunkDataSize) <= CacheElement->ChunkDataSize);
 
@@ -1722,7 +1737,7 @@ void FAudioChunkCache::KickOffAsyncLoad(FCacheElement* CacheElement, const FChun
 #endif
 
 		CacheElement->ReadRequest = nullptr;
-		IBulkDataIORequest* LocalReadRequest = Chunk.BulkData.CreateStreamingRequest(0, ChunkDataSize, AsyncIOPriority | AIOP_FLAG_DONTCACHE, &AsyncFileCallBack, CacheElement->ChunkData);
+		IBulkDataIORequest* LocalReadRequest = Chunk->BulkData.CreateStreamingRequest(0, ChunkDataSize, AsyncIOPriority | AIOP_FLAG_DONTCACHE, &AsyncFileCallBack, CacheElement->ChunkData);
 		if (!LocalReadRequest)
 		{
 			UE_LOG(LogAudioStreamCaching, Error, TEXT("Chunk load in audio LRU cache failed."));
@@ -1815,9 +1830,9 @@ void FAudioChunkCache::ExecuteOnLoadCompleteCallback(EAudioChunkLoadResult Resul
 	}
 }
 
-bool FAudioChunkCache::IsKeyValid(const FChunkKey& InKey)
+bool FAudioChunkCache::DoesKeyContainValidChunkIndex(const FChunkKey& InKey, const FSoundWaveData& InSoundWaveData)
 {
-	return InKey.ChunkIndex < TNumericLimits<uint32>::Max() && ((int32)InKey.ChunkIndex) < InKey.GetNumChunks();
+	return InKey.ChunkIndex < TNumericLimits<uint32>::Max() && ((int32)InKey.ChunkIndex) < InSoundWaveData.GetNumChunks();
 }
 
 uint64 FAudioChunkCache::GetCacheLookupIDForChunk(const FChunkKey& InChunkKey) const
@@ -2005,7 +2020,7 @@ TPair<int, int> FAudioChunkCache::DebugDisplayLegacy(UWorld* World, FViewport* V
 
 #if WITH_EDITOR
 		// TODO: Worry about whether the sound wave is alive here. In most editor cases this is ok because the soundwave will always be loaded, but this may not be the case in the future.
-		bIsStaleChunk = CurrentElement->Key.IsChunkStale();
+		bIsStaleChunk = CurrentElement->IsChunkStale();
 #endif
 
 		const bool bWasTrimmed = CurrentElement->ChunkDataSize == 0;
@@ -2140,7 +2155,7 @@ FString FAudioChunkCache::DebugPrint()
 
 #if WITH_EDITOR
 		// TODO: Worry about whether the sound wave is alive here. In most editor cases this is ok because the soundwave will always be loaded, but this may not be the case in the future.
-		bIsStaleChunk = CurrentElement->Key.IsChunkStale();
+		bIsStaleChunk = CurrentElement->IsChunkStale();
 #endif
 
 		const bool bWasTrimmed = CurrentElement->ChunkDataSize == 0;
@@ -2555,7 +2570,7 @@ TPair<int, int> FAudioChunkCache::DebugVisualDisplay(UWorld* World, FViewport* V
 
 #if WITH_EDITOR
 		// TODO: Worry about whether the sound wave is alive here. In most editor cases this is ok because the soundwave will always be loaded, but this may not be the case in the future.
-		bIsStaleChunk = CurrentElement->Key.IsChunkStale();
+		bIsStaleChunk = CurrentElement->IsChunkStale();
 #endif
 		const bool bWasTrimmed = CurrentElement->ChunkDataSize == 0;
 
diff --git a/Engine/Source/Runtime/Engine/Private/ContentStreaming.cpp b/Engine/Source/Runtime/Engine/Private/ContentStreaming.cpp
--- a/Engine/Source/Runtime/Engine/Private/ContentStreaming.cpp
+++ b/Engine/Source/Runtime/Engine/Private/ContentStreaming.cpp
@@ -1447,11 +1447,10 @@ FArchive& operator<<( FArchive& Ar, FDynamicTextureInstance& TextureInstance )
 FAudioChunkHandle::FAudioChunkHandle()
 	: CachedData(nullptr)
 	, CachedDataNumBytes(0)
-	, CorrespondingWave(nullptr)
-	, CorrespondingWaveName()
 	, ChunkIndex(INDEX_NONE)
 #if WITH_EDITOR
-	, ChunkGeneration(INDEX_NONE)
+	, CorrespondingWave(nullptr)
+	, ChunkRevision(INDEX_NONE)
 #endif
 {
 }
@@ -1459,13 +1458,21 @@ FAudioChunkHandle::FAudioChunkHandle()
 FAudioChunkHandle::FAudioChunkHandle(const uint8* InData, uint32 NumBytes, const FSoundWaveProxyPtr&  InSoundWave, const FName& SoundWaveName, uint32 InChunkIndex, uint64 InCacheLookupID)
 	: CachedData(InData)
 	, CachedDataNumBytes(NumBytes)
-	, CorrespondingWave(InSoundWave->GetSoundWaveData())
 	, CorrespondingWaveName(SoundWaveName)
 	, ChunkIndex(InChunkIndex)
 #if WITH_EDITOR
-	, ChunkGeneration(InSoundWave.IsValid()? InSoundWave->GetCurrentChunkRevision() : 0)
+	, CorrespondingWave(InSoundWave->GetSoundWaveData())
+	, ChunkRevision(InSoundWave.IsValid()? InSoundWave->GetCurrentChunkRevision() : 0)
 #endif
 {
+	if (InSoundWave.IsValid())
+	{
+		TSharedPtr<FSoundWaveData> SoundWaveData = InSoundWave->GetSoundWaveData();
+		if (SoundWaveData.IsValid())
+		{
+			CorrespondingWaveObjectKey = SoundWaveData->GetFObjectKey();
+		}
+	}
 }
 
 FAudioChunkHandle::FAudioChunkHandle(const FAudioChunkHandle& Other)
@@ -1490,22 +1497,24 @@ FAudioChunkHandle& FAudioChunkHandle::operator=(FAudioChunkHandle&& Other)
 
 	CachedData = Other.CachedData;
 	CachedDataNumBytes = Other.CachedDataNumBytes;
-	CorrespondingWave = MoveTemp(Other.CorrespondingWave);
 	CorrespondingWaveName = Other.CorrespondingWaveName;
+	CorrespondingWaveObjectKey = Other.CorrespondingWaveObjectKey;
 	ChunkIndex = Other.ChunkIndex;
 #if WITH_EDITOR
-	ChunkGeneration = Other.ChunkGeneration;
+	CorrespondingWave = MoveTemp(Other.CorrespondingWave);
+	ChunkRevision = Other.ChunkRevision;
 #endif
 
 	// we don't need to call RemoveReferenceToChunk on Other, nor add a new reference to this chunk, since this is a move.
 	// Instead, we can simply null out the other chunk handle without invoking it's destructor.
 	Other.CachedData = nullptr;
 	Other.CachedDataNumBytes = 0;
-	Other.CorrespondingWave = nullptr;
 	Other.CorrespondingWaveName = FName();
+	Other.CorrespondingWaveObjectKey = FObjectKey();
 	Other.ChunkIndex = INDEX_NONE;
 #if WITH_EDITOR
-	Other.ChunkGeneration = INDEX_NONE;
+	Other.CorrespondingWave = nullptr;
+	Other.ChunkRevision = INDEX_NONE;
 #endif
 
 	return *this;
@@ -1517,21 +1526,16 @@ FAudioChunkHandle& FAudioChunkHandle::operator=(const FAudioChunkHandle& Other)
 	if (IsValid())
 	{
 		IStreamingManager::Get().GetAudioStreamingManager().RemoveReferenceToChunk(*this);
-		CorrespondingWave.Reset();
 	}
 
 	CachedData = Other.CachedData;
 	CachedDataNumBytes = Other.CachedDataNumBytes;
-
-	if (Other.CorrespondingWave.IsValid())
-	{
-		CorrespondingWave = Other.CorrespondingWave;
-	}
-
 	CorrespondingWaveName = Other.CorrespondingWaveName;
+	CorrespondingWaveObjectKey = Other.CorrespondingWaveObjectKey;
 	ChunkIndex = Other.ChunkIndex;
 #if WITH_EDITOR
-	ChunkGeneration = Other.ChunkGeneration;
+	CorrespondingWave = Other.CorrespondingWave;
+	ChunkRevision = Other.ChunkRevision;
 #endif
 
 	if (IsValid())
@@ -1563,7 +1567,7 @@ uint32 FAudioChunkHandle::Num() const
 
 bool FAudioChunkHandle::IsValid() const
 {
-	return GetData() && CorrespondingWave.Pin().IsValid();
+	return (nullptr != GetData());
 }
 
 #if WITH_EDITOR
@@ -1574,7 +1578,7 @@ bool FAudioChunkHandle::IsStale() const
 	if (SoundWaveDataPtr.IsValid())
 	{
 		// NOTE: While this is currently safe in editor, there's no guarantee the USoundWave will be kept alive during the lifecycle of this chunk handle.
-		return ChunkGeneration != SoundWaveDataPtr->GetCurrentChunkRevision();
+		return ChunkRevision != SoundWaveDataPtr->GetCurrentChunkRevision();
 	}
 	else
 	{
diff --git a/Engine/Source/Runtime/Engine/Private/SoundWave.cpp b/Engine/Source/Runtime/Engine/Private/SoundWave.cpp
--- a/Engine/Source/Runtime/Engine/Private/SoundWave.cpp
+++ b/Engine/Source/Runtime/Engine/Private/SoundWave.cpp
@@ -1926,15 +1926,6 @@ void USoundWave::BeginDestroy()
 	// Flush any async results so we don't leak them in the DDC
 	FlushAsyncLoadingDataFormats();
 #endif
-
-	// the FSoundWaveDataPtr is not able to release the handle itself
-	// (weakptr that is already invalid in ~FSoundWaveData)
-	// so we need to Release our handle to the chunk in the streaming cache
-	// if we are the last one holding a refernce to SoundWaveDataPtr
-	if (SoundWaveDataPtr.GetSharedReferenceCount() == 1)
-	{
-		ReleaseCompressedAudio();
-	}
 }
 
 void USoundWave::InitAudioResource(FByteBulkData& CompressedData)
@@ -4082,14 +4073,6 @@ FSoundWaveProxy::FSoundWaveProxy(USoundWave* InWave)
 
 FSoundWaveProxy::~FSoundWaveProxy()
 {
-	// the FSoundWaveDataPtr is not able to release the handle itself
-	// (weakptr that is already invalid in ~FSoundWaveData)
-	// so we need to Release our handle to the chunk in the streaming cache
-	// if we are the last one holding a refernce to SoundWaveDataPtr
-	if (SoundWaveDataPtr.GetSharedReferenceCount() == 1)
-	{
-		ReleaseCompressedAudio();
-	}
 }
 
 
diff --git a/Engine/Source/Runtime/Engine/Public/AudioStreamingCache.h b/Engine/Source/Runtime/Engine/Public/AudioStreamingCache.h
--- a/Engine/Source/Runtime/Engine/Public/AudioStreamingCache.h
+++ b/Engine/Source/Runtime/Engine/Public/AudioStreamingCache.h
@@ -32,9 +32,18 @@ public:
 	struct FChunkKey
 	{
 	public:
-		FChunkKey() {};
+		FChunkKey() = default;
 
-		FChunkKey(const FChunkKey& Other);
+		FChunkKey(const FChunkKey& Other) = default;
+
+		FChunkKey(
+			  const FName& InSoundWaveName
+			, const FObjectKey& InSoundWaveObjectKey
+			, uint32 InChunkIndex
+#if WITH_EDITOR
+			, uint32 InChunkRevision = 0
+#endif // #if WITH_EDITOR
+		);
 
 		FChunkKey(
 			  const FSoundWavePtr& InSoundWave
@@ -44,32 +53,19 @@ public:
 #endif // #if WITH_EDITOR
 		);
 
-		FChunkKey& operator=(const FChunkKey& Other);
+		FChunkKey& operator=(const FChunkKey& Other) = default;
 
 
-	private:
-		TWeakPtr<FSoundWaveData, ESPMode::ThreadSafe> SoundWaveWeakPtr;
-
 	public:
 		FName SoundWaveName = FName();
-		uint32 ChunkIndex = INDEX_NONE;
 		FObjectKey ObjectKey = FObjectKey();
+		uint32 ChunkIndex = INDEX_NONE;
 
 #if WITH_EDITOR
 		// This is used in the editor to invalidate stale compressed chunks.
 		uint32 ChunkRevision = 0;
 #endif
-		inline bool operator==(const FChunkKey& Other) const;
-
-		bool IsChunkStale();
-		bool IsSoundWaveValid() const { return SoundWaveWeakPtr.Pin().IsValid(); }
-		// Forward FSoundWaveProxy interface
-		ESoundWaveLoadingBehavior GetLoadingBehavior() const;
-		FStreamedAudioChunk& GetChunk(uint32 InChunkIndex) const;
-		uint32 GetNumChunks() const;
-		bool IsRetainingAudio() const;
-		void ReleaseCompressedAudio();
-		bool WasLoadingBehaviorOverridden() const;
+		bool operator==(const FChunkKey& Other) const;
 
 		/** Hash function */
 		friend uint32 GetTypeHash(const FChunkKey& InChunkKey)
@@ -115,12 +111,12 @@ public:
 
 	// Places chunk in cache, or puts this chunk back at the top of the cache if it's already loaded. Returns the static lookup ID of the chunk in the cache on success,
 	// or InvalidAudioStreamCacheLookupID on failiure.
-	uint64 AddOrTouchChunk(const FChunkKey& InKey, TFunction<void(EAudioChunkLoadResult) > OnLoadCompleted, ENamedThreads::Type CallbackThread, bool bNeededForPlayback);
+	uint64 AddOrTouchChunk(const FChunkKey& InKey, const TSharedPtr<FSoundWaveData>& InSoundWavePtr, TFunction<void(EAudioChunkLoadResult) > OnLoadCompleted, ENamedThreads::Type CallbackThread, bool bNeededForPlayback);
 
 	// Returns the chunk asked for, or an empty TArrayView if that chunk is not loaded.
 	// InOutCacheLookupID can optionally be set as a cache offset to use directly rather than searching the cache for a matching chunk.
 	// InOutCacheLookupID will be set to the offset the chunk is in the cache, which can be used for faster lookup in the future.
-	TArrayView<uint8> GetChunk(const FChunkKey& InKey, bool bBlockForLoadCompletion, bool bNeededForPlayback, uint64& InOutCacheLookupID);
+	TArrayView<uint8> GetChunk(const FChunkKey& InKey, const TSharedPtr<FSoundWaveData>& InSoundWavePtr, bool bBlockForLoadCompletion, bool bNeededForPlayback, uint64& InOutCacheLookupID);
 
 	// add an additional reference for a chunk.
 	void AddNewReferenceToChunk(const FChunkKey& InKey);
@@ -160,7 +156,7 @@ public:
 	FString FlushCacheMissLog();
 
 	// Static helper function to make sure a chunk is withing the bounds of a USoundWave.
-	static bool IsKeyValid(const FChunkKey& InKey);
+	static bool DoesKeyContainValidChunkIndex(const FChunkKey& InKey, const FSoundWaveData& InSoundWaveData);
 
 	// interface with the cache id lookup map
 	uint64 GetCacheLookupIDForChunk(const FChunkKey& InChunkKey) const;
@@ -257,6 +253,7 @@ private:
 			bWasCacheMiss = false;
 			AverageLocationInCacheWhenNeeded = 0.0f;
 		}
+
 	};
 #endif
 
@@ -273,7 +270,9 @@ private:
 		uint32 ChunkDataSize;
 		FCacheElement* MoreRecentElement;
 		FCacheElement* LessRecentElement;
-		uint64 CacheLookupID;
+		uint64 CacheLookupID = InvalidAudioStreamCacheLookupID;
+
+		TWeakPtr<FSoundWaveData, ESPMode::ThreadSafe> SoundWaveWeakPtr;
 
 		FThreadSafeBool bIsLoaded;
 		
@@ -291,7 +290,7 @@ private:
 		FCacheElementDebugInfo DebugInfo;
 #endif
 
-		FCacheElement(uint32 MaxChunkSize, uint32 InCacheIndex)
+		FCacheElement(uint32 InCacheIndex)
 			: ChunkData(nullptr)
 			, ChunkDataSize(0)
 			, MoreRecentElement(nullptr)
@@ -342,15 +341,8 @@ private:
 		}
 
 #if DEBUG_STREAM_CACHE
-		bool IsBeingPlayed() const
-		{
-			const int32 NumActiveConsumers = NumConsumers.GetValue();
-
-			// if we 2 or more consumers, this chunk is being rendered.
-			// if we have 1 consumer, and we aren't Retained, then this chunk is being rendered
-			return (NumActiveConsumers > 1)
-				|| (NumActiveConsumers && (DebugInfo.LoadingBehavior != ESoundWaveLoadingBehavior::RetainOnLoad));
-		}
+		bool IsBeingPlayed() const;
+		void UpdateDebugInfoLoadingBehavior();
 #endif
 
 		bool CanEvictChunk() const
@@ -358,14 +350,26 @@ private:
 			return !IsInUse() && !IsLoadInProgress();
 		}
 
+		uint32 GetNumChunks() const;
+		FStreamedAudioChunk* GetChunk(uint32 InChunkIndex) const;
+
+#if WITH_EDITOR
+		bool IsChunkStale();
+#endif
+		void ReleaseRetainedAudioOnSoundWave();
+		bool IsSoundWaveRetainingAudio() const;
+
 		~FCacheElement()
 		{
 			WaitForAsyncLoadCompletion(true);
 
 			if(NumConsumers.GetValue() != 0)
 			{
+				/** TODO: This is failing because of `Key.GetLoadingBehavior()`. We could get this from the FSoundWaveData or
+				 * store it, or ignore it here. 
 				UE_LOG(LogAudioStreamCaching, Error, TEXT("Tried to destroy streaming cache while the cached data was in use! (Sound: %s - Loading Behavior: %s - Chunk Index: %i)")
 					, *Key.SoundWaveName.ToString(), EnumToString(Key.GetLoadingBehavior()), Key.ChunkIndex);
+				*/
 			}
 			else if (ChunkData)
 			{
@@ -424,7 +428,7 @@ private:
 	void TouchElement(FCacheElement* InElement);
 
 	// Inserts a new element into the cache, potentially evicting the oldest element in the cache.
-	FCacheElement* InsertChunk(const FChunkKey& InKey);
+	FCacheElement* InsertChunk(const FChunkKey& InKey, const TSharedPtr<FSoundWaveData>& InSoundWaveWeakPtr);
 
 	// This is called once we have more than one chunk in our cache:
 	void SetUpLeastRecentChunk();
@@ -543,7 +547,4 @@ protected:
 
 	/** Audio chunk caches. These are set up on initialization. */
 	TArray<FAudioChunkCache> CacheArray;
-
-	
-
 };
diff --git a/Engine/Source/Runtime/Engine/Public/ContentStreaming.h b/Engine/Source/Runtime/Engine/Public/ContentStreaming.h
--- a/Engine/Source/Runtime/Engine/Public/ContentStreaming.h
+++ b/Engine/Source/Runtime/Engine/Public/ContentStreaming.h
@@ -8,13 +8,14 @@
 
 #include "Async/TaskGraphInterfaces.h"
 #include "CoreMinimal.h"
-#include "UObject/WeakObjectPtr.h"
 #include "RenderedTextureStats.h"
 #if UE_ENABLE_INCLUDE_ORDER_DEPRECATED_IN_5_2
 #include "CanvasTypes.h"
 #include "UnrealClient.h"
 #endif
 #include "Serialization/BulkData.h"
+#include "UObject/ObjectKey.h"
+#include "UObject/WeakObjectPtr.h"
 
 class AActor;
 class FSoundSource;
@@ -136,14 +137,15 @@ private:
 	const uint8*  CachedData;
 	int32 CachedDataNumBytes;
 
-	TWeakPtr<FSoundWaveData, ESPMode::ThreadSafe> CorrespondingWave;
 	FName CorrespondingWaveName;
+	FObjectKey CorrespondingWaveObjectKey;
 
 	// The index of this chunk in the sound wave's full set of chunks of compressed audio.
 	int32 ChunkIndex;
 
 #if WITH_EDITOR
-	uint32 ChunkGeneration;
+	TWeakPtr<FSoundWaveData, ESPMode::ThreadSafe> CorrespondingWave;
+	uint32 ChunkRevision;
 #endif
 
 	friend struct IAudioStreamingManager;
