Cache TextureFormat on the BuildSettings so we don't have to look through the TPM whenever we want it.

diff --git a/Engine/Source/Developer/TextureBuild/Private/TextureBuildFunction.cpp b/Engine/Source/Developer/TextureBuild/Private/TextureBuildFunction.cpp
--- a/Engine/Source/Developer/TextureBuild/Private/TextureBuildFunction.cpp
+++ b/Engine/Source/Developer/TextureBuild/Private/TextureBuildFunction.cpp
@@ -163,6 +163,20 @@ static FTextureBuildSettings ReadBuildSettingsFromCompactBinary(const FCbObjectV
 	ReadCbField(Object["TextureAddressModeY"], BuildSettings.TextureAddressModeY);
 	ReadCbField(Object["TextureAddressModeZ"], BuildSettings.TextureAddressModeZ);
 
+	ITextureFormatManagerModule* TFM = GetTextureFormatManager();
+	if (TFM == nullptr)
+	{
+		UE_LOG(LogTextureBuildFunction, Error, TEXT("Texture Format Manager Module not found!"));
+	}
+	else
+	{
+		BuildSettings.TextureFormat = TFM->FindTextureFormat(BuildSettings.TextureFormatName);
+		if (BuildSettings.TextureFormat == nullptr)
+		{
+			UE_LOG(LogTextureBuildFunction, Error, TEXT("Texture Format %s not found!"), *BuildSettings.TextureFormatName.ToString());
+		}
+	}
+
 	return BuildSettings;
 }
 
@@ -317,20 +331,13 @@ void FTextureBuildFunction::Build(UE::DerivedData::FBuildContext& Context) const
 	}
 
 	const FTextureBuildSettings BuildSettings = ReadBuildSettingsFromCompactBinary(Settings["Build"].AsObjectView());
-	
-	const uint16 RequiredTextureFormatVersion = Settings["FormatVersion"].AsUInt16();
-	const ITextureFormat* TextureFormat;
-	if (ITextureFormatManagerModule* TFM = GetTextureFormatManager())
-	{
-		TextureFormat = TFM->FindTextureFormat(BuildSettings.TextureFormatName);
-	}
-	else
+	if (BuildSettings.TextureFormat == nullptr)
 	{
-		UE_LOG(LogTextureBuildFunction, Error, TEXT("TextureFormatManager not found!"));
-		return;
+		return; // we had an error reading the build settings.
 	}
-
-	const uint16 CurrentTextureFormatVersion = TextureFormat ? TextureFormat->GetVersion(BuildSettings.TextureFormatName, &BuildSettings) : 0;
+	
+	const uint16 RequiredTextureFormatVersion = Settings["FormatVersion"].AsUInt16();
+	const uint16 CurrentTextureFormatVersion = BuildSettings.TextureFormat->GetVersion(BuildSettings.TextureFormatName, &BuildSettings);
 	if (CurrentTextureFormatVersion != RequiredTextureFormatVersion)
 	{
 		UE_LOG(LogTextureBuildFunction, Error, TEXT("%s has version %hu when version %hu is required."),
diff --git a/Engine/Source/Developer/TextureCompressor/Private/TextureCompressorModule.cpp b/Engine/Source/Developer/TextureCompressor/Private/TextureCompressorModule.cpp
--- a/Engine/Source/Developer/TextureCompressor/Private/TextureCompressorModule.cpp
+++ b/Engine/Source/Developer/TextureCompressor/Private/TextureCompressorModule.cpp
@@ -3468,14 +3468,7 @@ public:
 	{
 		//TRACE_CPUPROFILER_EVENT_SCOPE(Texture.BuildTexture);
 
-		const ITextureFormat* TextureFormat = nullptr;
-
-		ITextureFormatManagerModule* TFM = GetTextureFormatManager();
-		if (TFM)
-		{
-			TextureFormat = TFM->FindTextureFormat(BuildSettings.TextureFormatName);
-		}
-		if (TextureFormat == nullptr)
+		if (BuildSettings.TextureFormat == nullptr)
 		{
 			UE_LOG(LogTextureCompressor, Warning,
 				TEXT("Failed to find compressor for texture format '%s'. [%.*s]"),
@@ -3500,7 +3493,7 @@ public:
 
 		// allow to leave texture in sRGB in case compressor accepts other than non-F32 input source
 		// otherwise linearizing will force format to be RGBA32F
-		const bool bNeedLinearize = !TextureFormat->CanAcceptNonF32Source() || AssociatedNormalSourceMips.Num() != 0;
+		const bool bNeedLinearize = !BuildSettings.TextureFormat->CanAcceptNonF32Source() || AssociatedNormalSourceMips.Num() != 0;
 		if (!BuildTextureMips(SourceMips, BuildSettings, bNeedLinearize, IntermediateMipChain, DebugTexturePathName))
 		{
 			return false;
@@ -3581,7 +3574,7 @@ public:
 			OutMetadata->PreEncodeMipsHash = MipHashBuilder.Finalize();
 		}
 		
-		bool bCompressSucceeded = CompressMipChain(TextureFormat, IntermediateMipChain, BuildSettings, bImageHasAlphaChannel, DebugTexturePathName,
+		bool bCompressSucceeded = CompressMipChain(BuildSettings.TextureFormat, IntermediateMipChain, BuildSettings, bImageHasAlphaChannel, DebugTexturePathName,
 					OutTextureMips, OutNumMipsInTail, OutExtData);
 
 		return bCompressSucceeded;
diff --git a/Engine/Source/Developer/TextureCompressor/Public/TextureCompressorModule.h b/Engine/Source/Developer/TextureCompressor/Public/TextureCompressorModule.h
--- a/Engine/Source/Developer/TextureCompressor/Public/TextureCompressorModule.h
+++ b/Engine/Source/Developer/TextureCompressor/Public/TextureCompressorModule.h
@@ -232,6 +232,10 @@ struct FTextureBuildSettings
 	// If the target format is a tiled format and can leverage reusing the linear encoding, this is not nullptr.
 	const ITextureTiler* Tiler = nullptr;
 
+	// Cached copy of the texture format. If the texture format isn't loaded or found, this can be nullptr. The build fails on this
+	// right before the actual texture build, so if you're prior to that, check it!
+	const ITextureFormat* TextureFormat = nullptr;
+
 	// If shared linear is enabled _at all_ and this texture in involved with that _at all_ then we set
 	// this so we can segregate the derived data keys.
 	bool bAffectedBySharedLinearEncoding = false;
diff --git a/Engine/Source/Runtime/Engine/Private/TextureDerivedData.cpp b/Engine/Source/Runtime/Engine/Private/TextureDerivedData.cpp
--- a/Engine/Source/Runtime/Engine/Private/TextureDerivedData.cpp
+++ b/Engine/Source/Runtime/Engine/Private/TextureDerivedData.cpp
@@ -359,19 +359,11 @@ void GetTextureDerivedDataKeySuffix(const UTexture& Texture, const FTextureBuild
 	const FTextureBuildSettings& BuildSettings = BuildSettingsPerLayer[0];
 
 	// get the version for this texture's platform format
-	ITargetPlatformManagerModule* TPM = GetTargetPlatformManager();
-	const ITextureFormat* TextureFormat = NULL;
-	if (TPM)
+	if (BuildSettingsPerLayer->TextureFormat != nullptr)
 	{
-		TextureFormat = TPM->FindTextureFormat(BuildSettings.TextureFormatName);
-		if (TextureFormat)
-		{
-			Version = TextureFormat->GetVersion(BuildSettings.TextureFormatName, &BuildSettings);
-		}
-		// else error !?
+		Version = BuildSettingsPerLayer->TextureFormat->GetVersion(BuildSettings.TextureFormatName, BuildSettingsPerLayer);
 	}
-	// else error !?
-	
+
 	FString CompositeTextureStr;
 
 	if(IsValid(Texture.CompositeTexture) && Texture.CompositeTextureMode != CTM_Disabled)
@@ -388,7 +380,7 @@ void GetTextureDerivedDataKeySuffix(const UTexture& Texture, const FTextureBuild
 		*Texture.Source.GetIdString(),
 		*CompositeTextureStr,
 		(uint32)NUM_INLINE_DERIVED_MIPS,
-		(TextureFormat == NULL) ? TEXT("") : *TextureFormat->GetDerivedDataKeyString(BuildSettings)
+		(BuildSettingsPerLayer->TextureFormat == nullptr) ? TEXT("") : *BuildSettingsPerLayer->TextureFormat->GetDerivedDataKeyString(BuildSettings)
 		);
 
 	// Add key data for extra layers beyond the first
@@ -396,21 +388,15 @@ void GetTextureDerivedDataKeySuffix(const UTexture& Texture, const FTextureBuild
 	for (int32 LayerIndex = 1; LayerIndex < NumLayers; ++LayerIndex)
 	{
 		const FTextureBuildSettings& LayerBuildSettings = BuildSettingsPerLayer[LayerIndex];
-		const ITextureFormat* LayerTextureFormat = NULL;
-		if (TPM)
-		{
-			LayerTextureFormat = TPM->FindTextureFormat(LayerBuildSettings.TextureFormatName);
-		}
-
 		uint16 LayerVersion = 0;
-		if (LayerTextureFormat)
+		if (LayerBuildSettings.TextureFormat)
 		{
-			LayerVersion = LayerTextureFormat->GetVersion(LayerBuildSettings.TextureFormatName, &LayerBuildSettings);
+			LayerVersion = LayerBuildSettings.TextureFormat->GetVersion(LayerBuildSettings.TextureFormatName, &LayerBuildSettings);
 		}
 		OutKeySuffix.Append(FString::Printf(TEXT("%s%d%s_"),
 			*LayerBuildSettings.TextureFormatName.GetPlainNameString(),
 			LayerVersion,
-			(LayerTextureFormat == NULL) ? TEXT("") : *LayerTextureFormat->GetDerivedDataKeyString(LayerBuildSettings)));
+			(LayerBuildSettings.TextureFormat == nullptr) ? TEXT("") : *LayerBuildSettings.TextureFormat->GetDerivedDataKeyString(LayerBuildSettings)));
 	}
 
 	if (BuildSettings.bVirtualStreamable)
@@ -661,59 +647,61 @@ static void FinalizeBuildSettingsForLayer(
 	}
 
 	// Now that we know the texture format, we can make decisions based on it.
+	ITargetPlatformManagerModule* TPM = GetTargetPlatformManager();
+	check(TPM != nullptr);
+	if (TPM == nullptr)
+	{
+		OutSettings.TextureFormat = nullptr;
+		return;
+	}
+
+	// Can be null with first finalize (at the end of GetTextureBuildSettings), however we expect
+	// to get a valid texture format before the build settings are used!
+	const ITextureFormat* TextureFormat = TPM->FindTextureFormat(OutSettings.TextureFormatName);
+	OutSettings.TextureFormat = TextureFormat;
 
 	bool bSupportsEncodeSpeed = false;
+	if (TextureFormat != nullptr)
 	{
-		ITargetPlatformManagerModule* TPM = GetTargetPlatformManager();
-		if (TPM)
+		bSupportsEncodeSpeed = TextureFormat->SupportsEncodeSpeed(OutSettings.TextureFormatName);
+
+		if (OutBuildResultMetadata)
+		{
+			OutBuildResultMetadata->Encoder = TextureFormat->GetEncoderName(OutSettings.TextureFormatName);
+			OutBuildResultMetadata->bIsValid = true;
+			OutBuildResultMetadata->bSupportsEncodeSpeed = bSupportsEncodeSpeed;
+		}
+			
+		static auto CVarSharedLinearTextureEncoding = IConsoleManager::Get().FindTConsoleVariableDataInt(TEXT("r.SharedLinearTextureEncoding"));
+		if (CVarSharedLinearTextureEncoding->GetValueOnAnyThread())
 		{
-			// Can be null with first finalize (at the end of GetTextureBuildSettings)
-			const ITextureFormat* TextureFormat = TPM->FindTextureFormat(OutSettings.TextureFormatName);
-			if (TextureFormat)
+			//
+			// We want to separate out textures involved in shared linear encoding in order to facilitate
+			// fixing bugs without invalidating the world (even though we expect the exact same data to
+			// get generated). However, virtual textures never tile, and so are exempt from this separation.
+			//
+			if (OutSettings.bVirtualStreamable == false)
 			{
-				bSupportsEncodeSpeed = TextureFormat->SupportsEncodeSpeed(OutSettings.TextureFormatName);
+				OutSettings.bAffectedBySharedLinearEncoding = true;
+			}
 
-				if (OutBuildResultMetadata)
+			// Shared linear encoding can only work if the base texture format does not expect to
+			// do the tiling itself (SupportsTiling == false).
+			const FChildTextureFormat* ChildTextureFormat = TextureFormat->GetChildFormat();
+			if (ChildTextureFormat && ChildTextureFormat->GetBaseFormatObject(OutSettings.TextureFormatName)->SupportsTiling() == false)
+			{
+				OutSettings.Tiler = ChildTextureFormat->GetTiler();
+				if (OutSettings.Tiler && IsUsingNewDerivedData())
 				{
-					OutBuildResultMetadata->Encoder = TextureFormat->GetEncoderName(OutSettings.TextureFormatName);
-					OutBuildResultMetadata->bIsValid = true;
-					OutBuildResultMetadata->bSupportsEncodeSpeed = bSupportsEncodeSpeed;
+					// New derived data wants to treat everything as the base format and then have a separate tiling function
+					// afterwards.
+					OutSettings.TextureFormatName = ChildTextureFormat->GetBaseFormatName(OutSettings.TextureFormatName);
+					OutSettings.TextureFormat = ChildTextureFormat;
 				}
-			
-				
-				{
-					static auto CVarSharedLinearTextureEncoding = IConsoleManager::Get().FindTConsoleVariableDataInt(TEXT("r.SharedLinearTextureEncoding"));
-					if (CVarSharedLinearTextureEncoding->GetValueOnAnyThread())
-					{
-						//
-						// We want to separate out textures involved in shared linear encoding in order to facilitate
-						// fixing bugs without invalidating the world (even though we expect the exact same data to
-						// get generated). However, virtual textures never tile, and so are exempt from this separation.
-						//
-						if (OutSettings.bVirtualStreamable == false)
-						{
-							OutSettings.bAffectedBySharedLinearEncoding = true;
-						}
-
-						// Shared linear encoding can only work if the base texture format does not expect to
-						// do the tiling itself (SupportsTiling == false).
-						const FChildTextureFormat* ChildTextureFormat = TextureFormat->GetChildFormat();
-						if (ChildTextureFormat && ChildTextureFormat->GetBaseFormatObject(OutSettings.TextureFormatName)->SupportsTiling() == false)
-						{
-							OutSettings.Tiler = ChildTextureFormat->GetTiler();
-							if (OutSettings.Tiler && IsUsingNewDerivedData())
-							{
-								// New derived data wants to treat everything as the base format and then have a separate tiling function
-								// afterwards.
-								OutSettings.TextureFormatName = ChildTextureFormat->GetBaseFormatName(OutSettings.TextureFormatName);
-							}
-						}
-					} // end if enabled
-				} // end if ddc2
-			} // end if texture format found.
-		}
-	}
-
+			}
+		} // end if sle enabled
+	} // end if texture format found.
+	
 	if (bSupportsEncodeSpeed)
 	{
 		FTextureEncodeSpeedOptions Options;
