Modified VirtualTexture ProducePageData method to take a non-immediate command list.

diff --git a/Engine/Plugins/Experimental/GPULightmass/Source/GPULightmass/Private/LightmapPreviewVirtualTexture.cpp b/Engine/Plugins/Experimental/GPULightmass/Source/GPULightmass/Private/LightmapPreviewVirtualTexture.cpp
--- a/Engine/Plugins/Experimental/GPULightmass/Source/GPULightmass/Private/LightmapPreviewVirtualTexture.cpp
+++ b/Engine/Plugins/Experimental/GPULightmass/Source/GPULightmass/Private/LightmapPreviewVirtualTexture.cpp
@@ -71,7 +71,7 @@ FVTRequestPageResult FLightmapPreviewVirtualTexture::RequestPageData(
 }
 
 IVirtualTextureFinalizer* FLightmapPreviewVirtualTexture::ProducePageData(
-	FRHICommandListImmediate& RHICmdList,
+	FRHICommandList& RHICmdList,
 	ERHIFeatureLevel::Type FeatureLevel,
 	EVTProducePageFlags Flags,
 	const FVirtualTextureProducerHandle& InProducerHandle,
diff --git a/Engine/Plugins/Experimental/GPULightmass/Source/GPULightmass/Private/LightmapPreviewVirtualTexture.h b/Engine/Plugins/Experimental/GPULightmass/Source/GPULightmass/Private/LightmapPreviewVirtualTexture.h
--- a/Engine/Plugins/Experimental/GPULightmass/Source/GPULightmass/Private/LightmapPreviewVirtualTexture.h
+++ b/Engine/Plugins/Experimental/GPULightmass/Source/GPULightmass/Private/LightmapPreviewVirtualTexture.h
@@ -23,7 +23,7 @@ public:
 	) override;
 
 	virtual IVirtualTextureFinalizer* ProducePageData(
-		FRHICommandListImmediate& RHICmdList,
+		FRHICommandList& RHICmdList,
 		ERHIFeatureLevel::Type FeatureLevel,
 		EVTProducePageFlags Flags,
 		const FVirtualTextureProducerHandle& InProducerHandle,
diff --git a/Engine/Source/Runtime/Engine/Private/VT/RuntimeVirtualTexture.cpp b/Engine/Source/Runtime/Engine/Private/VT/RuntimeVirtualTexture.cpp
--- a/Engine/Source/Runtime/Engine/Private/VT/RuntimeVirtualTexture.cpp
+++ b/Engine/Source/Runtime/Engine/Private/VT/RuntimeVirtualTexture.cpp
@@ -51,7 +51,7 @@ namespace
 		}
 
 		virtual IVirtualTextureFinalizer* ProducePageData(
-			FRHICommandListImmediate& RHICmdList,
+			FRHICommandList& RHICmdList,
 			ERHIFeatureLevel::Type FeatureLevel,
 			EVTProducePageFlags Flags,
 			const FVirtualTextureProducerHandle& ProducerHandle,
diff --git a/Engine/Source/Runtime/Engine/Private/VT/UploadingVirtualTexture.cpp b/Engine/Source/Runtime/Engine/Private/VT/UploadingVirtualTexture.cpp
--- a/Engine/Source/Runtime/Engine/Private/VT/UploadingVirtualTexture.cpp
+++ b/Engine/Source/Runtime/Engine/Private/VT/UploadingVirtualTexture.cpp
@@ -133,7 +133,7 @@ FVTRequestPageResult FUploadingVirtualTexture::RequestPageData(const FVirtualTex
 	return StreamingManager->RequestTile(this, ProducerHandle, LayerMask, vLevel, (uint32)vAddress, Priority);
 }
 
-IVirtualTextureFinalizer* FUploadingVirtualTexture::ProducePageData(FRHICommandListImmediate& RHICmdList,
+IVirtualTextureFinalizer* FUploadingVirtualTexture::ProducePageData(FRHICommandList& RHICmdList,
 	ERHIFeatureLevel::Type FeatureLevel,
 	EVTProducePageFlags Flags,
 	const FVirtualTextureProducerHandle& ProducerHandle, uint8 LayerMask, uint8 vLevel, uint64 vAddress,
diff --git a/Engine/Source/Runtime/Engine/Private/VT/UploadingVirtualTexture.h b/Engine/Source/Runtime/Engine/Private/VT/UploadingVirtualTexture.h
--- a/Engine/Source/Runtime/Engine/Private/VT/UploadingVirtualTexture.h
+++ b/Engine/Source/Runtime/Engine/Private/VT/UploadingVirtualTexture.h
@@ -79,7 +79,7 @@ public:
 	// IVirtualTexture interface
 	virtual uint32 GetLocalMipBias(uint8 vLevel, uint32 vAddress) const override;
 	virtual FVTRequestPageResult RequestPageData(const FVirtualTextureProducerHandle& ProducerHandle, uint8 LayerMask, uint8 vLevel, uint64 vAddress, EVTRequestPagePriority Priority) override;
-	virtual IVirtualTextureFinalizer* ProducePageData(FRHICommandListImmediate& RHICmdList,
+	virtual IVirtualTextureFinalizer* ProducePageData(FRHICommandList& RHICmdList,
 		ERHIFeatureLevel::Type FeatureLevel,
 		EVTProducePageFlags Flags,
 		const FVirtualTextureProducerHandle& ProducerHandle, uint8 LayerMask, uint8 vLevel, uint64 vAddress,
diff --git a/Engine/Source/Runtime/Engine/Private/VT/VirtualTextureChunkManager.cpp b/Engine/Source/Runtime/Engine/Private/VT/VirtualTextureChunkManager.cpp
--- a/Engine/Source/Runtime/Engine/Private/VT/VirtualTextureChunkManager.cpp
+++ b/Engine/Source/Runtime/Engine/Private/VT/VirtualTextureChunkManager.cpp
@@ -142,7 +142,7 @@ FVTRequestPageResult FVirtualTextureChunkStreamingManager::RequestTile(FUploadin
 	return FVTRequestPageResult(EVTRequestPageStatus::Pending, TranscodeHandle.PackedData);
 }
 
-IVirtualTextureFinalizer* FVirtualTextureChunkStreamingManager::ProduceTile(FRHICommandListImmediate& RHICmdList, uint32 SkipBorderSize, uint8 NumLayers, uint8 LayerMask, uint64 RequestHandle, const FVTProduceTargetLayer* TargetLayers)
+IVirtualTextureFinalizer* FVirtualTextureChunkStreamingManager::ProduceTile(FRHICommandList& RHICmdList, uint32 SkipBorderSize, uint8 NumLayers, uint8 LayerMask, uint64 RequestHandle, const FVTProduceTargetLayer* TargetLayers)
 {
 	SCOPE_CYCLE_COUNTER(STAT_VTP_ProduceTile);
 
diff --git a/Engine/Source/Runtime/Engine/Private/VT/VirtualTextureChunkManager.h b/Engine/Source/Runtime/Engine/Private/VT/VirtualTextureChunkManager.h
--- a/Engine/Source/Runtime/Engine/Private/VT/VirtualTextureChunkManager.h
+++ b/Engine/Source/Runtime/Engine/Private/VT/VirtualTextureChunkManager.h
@@ -43,7 +43,7 @@ public:
 	// End IStreamingManager interface
 
 	FVTRequestPageResult RequestTile(FUploadingVirtualTexture* VTexture, const FVirtualTextureProducerHandle& ProducerHandle, uint8 LayerMask, uint8 vLevel, uint32 vAddress, EVTRequestPagePriority Priority);
-	IVirtualTextureFinalizer* ProduceTile(FRHICommandListImmediate& RHICmdList, uint32 SkipBorderSize, uint8 NumLayers, uint8 LayerMask, uint64 RequestHandle, const FVTProduceTargetLayer* TargetLayers);
+	IVirtualTextureFinalizer* ProduceTile(FRHICommandList& RHICmdList, uint32 SkipBorderSize, uint8 NumLayers, uint8 LayerMask, uint64 RequestHandle, const FVTProduceTargetLayer* TargetLayers);
 	void GatherProducePageDataTasks(FVirtualTextureProducerHandle const& ProducerHandle, FGraphEventArray& InOutTasks) const;
 	void GatherProducePageDataTasks(uint64 RequestHandle, FGraphEventArray& InOutTasks) const;
 
diff --git a/Engine/Source/Runtime/Engine/Private/VT/VirtualTextureLevelRedirector.cpp b/Engine/Source/Runtime/Engine/Private/VT/VirtualTextureLevelRedirector.cpp
--- a/Engine/Source/Runtime/Engine/Private/VT/VirtualTextureLevelRedirector.cpp
+++ b/Engine/Source/Runtime/Engine/Private/VT/VirtualTextureLevelRedirector.cpp
@@ -27,7 +27,7 @@ FVTRequestPageResult FVirtualTextureLevelRedirector::RequestPageData(
 }
 
 IVirtualTextureFinalizer* FVirtualTextureLevelRedirector::ProducePageData(
-	FRHICommandListImmediate& RHICmdList,
+	FRHICommandList& RHICmdList,
 	ERHIFeatureLevel::Type FeatureLevel,
 	EVTProducePageFlags Flags,
 	const FVirtualTextureProducerHandle& ProducerHandle,
diff --git a/Engine/Source/Runtime/Engine/Private/VT/VirtualTextureLevelRedirector.h b/Engine/Source/Runtime/Engine/Private/VT/VirtualTextureLevelRedirector.h
--- a/Engine/Source/Runtime/Engine/Private/VT/VirtualTextureLevelRedirector.h
+++ b/Engine/Source/Runtime/Engine/Private/VT/VirtualTextureLevelRedirector.h
@@ -25,7 +25,7 @@ public:
 	) override;
 
 	virtual IVirtualTextureFinalizer* ProducePageData(
-		FRHICommandListImmediate& RHICmdList,
+		FRHICommandList& RHICmdList,
 		ERHIFeatureLevel::Type FeatureLevel,
 		EVTProducePageFlags Flags,
 		const FVirtualTextureProducerHandle& ProducerHandle,
diff --git a/Engine/Source/Runtime/Engine/Private/VT/VirtualTextureUploadCache.cpp b/Engine/Source/Runtime/Engine/Private/VT/VirtualTextureUploadCache.cpp
--- a/Engine/Source/Runtime/Engine/Private/VT/VirtualTextureUploadCache.cpp
+++ b/Engine/Source/Runtime/Engine/Private/VT/VirtualTextureUploadCache.cpp
@@ -419,9 +419,9 @@ FVTUploadTileHandle FVirtualTextureUploadCache::PrepareTileForUpload(FVTUploadTi
 	return FVTUploadTileHandle(Index);
 }
 
-void FVirtualTextureUploadCache::SubmitTile(FRHICommandListImmediate& RHICmdList, const FVTUploadTileHandle& InHandle, FRHITexture2D* InDestTexture, int InDestX, int InDestY, int InSkipBorderSize)
+void FVirtualTextureUploadCache::SubmitTile(FRHICommandList& RHICmdList, const FVTUploadTileHandle& InHandle, FRHITexture2D* InDestTexture, int InDestX, int InDestY, int InSkipBorderSize)
 {
-	checkSlow(IsInRenderingThread());
+	checkSlow(IsInParallelRenderingThread());
 
 	check(PendingUpload.IsValidIndex(InHandle.Index));
 	FTileEntry& Entry = PendingUpload[InHandle.Index];
diff --git a/Engine/Source/Runtime/Engine/Private/VT/VirtualTextureUploadCache.h b/Engine/Source/Runtime/Engine/Private/VT/VirtualTextureUploadCache.h
--- a/Engine/Source/Runtime/Engine/Private/VT/VirtualTextureUploadCache.h
+++ b/Engine/Source/Runtime/Engine/Private/VT/VirtualTextureUploadCache.h
@@ -142,7 +142,7 @@ public:
 	 * Mark streamed upload data ready for upload to to the physical virtual texture.
 	 * Depending on the platform the upload might happen here, or be deferred to the Finalize() call.
 	 */
-	void SubmitTile(FRHICommandListImmediate& RHICmdList, const FVTUploadTileHandle& InHandle, FRHITexture2D* InDestTexture, int InDestX, int InDestY, int InSkipBorderSize);
+	void SubmitTile(FRHICommandList& RHICmdList, const FVTUploadTileHandle& InHandle, FRHITexture2D* InDestTexture, int InDestX, int InDestY, int InSkipBorderSize);
 	/** Cancel a tile that was already in flight. */
 	void CancelTile(const FVTUploadTileHandle& InHandle);
 
diff --git a/Engine/Source/Runtime/RenderCore/Public/VirtualTexturing.h b/Engine/Source/Runtime/RenderCore/Public/VirtualTexturing.h
--- a/Engine/Source/Runtime/RenderCore/Public/VirtualTexturing.h
+++ b/Engine/Source/Runtime/RenderCore/Public/VirtualTexturing.h
@@ -22,6 +22,7 @@
 #include "UObject/NameTypes.h"
 
 class FRDGBuilder;
+class FRHICommandList;
 class FRHICommandListImmediate;
 class FRHIShaderResourceView;
 class FRHITexture;
@@ -311,7 +312,7 @@ public:
 	* @param TargetLayers Array of 'FVTProduceTargetLayer' structs, gives location where each layer should write data
 	* @return a 'IVirtualTextureFinalizer' which must be finalized to complete the operation
 	*/
-	virtual IVirtualTextureFinalizer* ProducePageData(FRHICommandListImmediate& RHICmdList,
+	virtual IVirtualTextureFinalizer* ProducePageData(FRHICommandList& RHICmdList,
 		ERHIFeatureLevel::Type FeatureLevel,
 		EVTProducePageFlags Flags,
 		const FVirtualTextureProducerHandle& ProducerHandle, uint8 LayerMask, uint8 vLevel, uint64 vAddress,
diff --git a/Engine/Source/Runtime/Renderer/Private/VT/AdaptiveVirtualTexture.cpp b/Engine/Source/Runtime/Renderer/Private/VT/AdaptiveVirtualTexture.cpp
--- a/Engine/Source/Runtime/Renderer/Private/VT/AdaptiveVirtualTexture.cpp
+++ b/Engine/Source/Runtime/Renderer/Private/VT/AdaptiveVirtualTexture.cpp
@@ -80,7 +80,7 @@ public:
 	}
 
 	virtual IVirtualTextureFinalizer* ProducePageData(
-		FRHICommandListImmediate& RHICmdList,
+		FRHICommandList& RHICmdList,
 		ERHIFeatureLevel::Type FeatureLevel,
 		EVTProducePageFlags Flags,
 		const FVirtualTextureProducerHandle& ProducerHandle,
diff --git a/Engine/Source/Runtime/Renderer/Private/VT/RuntimeVirtualTextureProducer.cpp b/Engine/Source/Runtime/Renderer/Private/VT/RuntimeVirtualTextureProducer.cpp
--- a/Engine/Source/Runtime/Renderer/Private/VT/RuntimeVirtualTextureProducer.cpp
+++ b/Engine/Source/Runtime/Renderer/Private/VT/RuntimeVirtualTextureProducer.cpp
@@ -165,7 +165,7 @@ FVTRequestPageResult FRuntimeVirtualTextureProducer::RequestPageData(
 }
 
 IVirtualTextureFinalizer* FRuntimeVirtualTextureProducer::ProducePageData(
-	FRHICommandListImmediate& RHICmdList,
+	FRHICommandList& RHICmdList,
 	ERHIFeatureLevel::Type FeatureLevel,
 	EVTProducePageFlags Flags,
 	const FVirtualTextureProducerHandle& ProducerHandle,
diff --git a/Engine/Source/Runtime/Renderer/Private/VT/RuntimeVirtualTextureProducer.h b/Engine/Source/Runtime/Renderer/Private/VT/RuntimeVirtualTextureProducer.h
--- a/Engine/Source/Runtime/Renderer/Private/VT/RuntimeVirtualTextureProducer.h
+++ b/Engine/Source/Runtime/Renderer/Private/VT/RuntimeVirtualTextureProducer.h
@@ -93,7 +93,7 @@ public:
 	) override;
 
 	virtual IVirtualTextureFinalizer* ProducePageData(
-		FRHICommandListImmediate& RHICmdList,
+		FRHICommandList& RHICmdList,
 		ERHIFeatureLevel::Type FeatureLevel,
 		EVTProducePageFlags Flags,
 		const FVirtualTextureProducerHandle& ProducerHandle,
diff --git a/Engine/Source/Runtime/Renderer/Public/RuntimeVirtualTextureProducer.h b/Engine/Source/Runtime/Renderer/Public/RuntimeVirtualTextureProducer.h
--- a/Engine/Source/Runtime/Renderer/Public/RuntimeVirtualTextureProducer.h
+++ b/Engine/Source/Runtime/Renderer/Public/RuntimeVirtualTextureProducer.h
@@ -70,7 +70,7 @@ public:
 	) override;
 
 	virtual IVirtualTextureFinalizer* ProducePageData(
-		FRHICommandListImmediate& RHICmdList,
+		FRHICommandList& RHICmdList,
 		ERHIFeatureLevel::Type FeatureLevel,
 		EVTProducePageFlags Flags,
 		const FVirtualTextureProducerHandle& ProducerHandle,
