Updated the underlying types for regular and namespaced enums.

diff --git a/Engine/Plugins/2D/Paper2D/Source/Paper2D/Classes/PaperFlipbook.h b/Engine/Plugins/2D/Paper2D/Source/Paper2D/Classes/PaperFlipbook.h
--- a/Engine/Plugins/2D/Paper2D/Source/Paper2D/Classes/PaperFlipbook.h
+++ b/Engine/Plugins/2D/Paper2D/Source/Paper2D/Classes/PaperFlipbook.h
@@ -34,7 +34,7 @@ public:
 UENUM()
 namespace EFlipbookCollisionMode
 {
-	enum Type
+	enum Type : int
 	{
 		/** The flipbook has no collision */
 		NoCollision,
diff --git a/Engine/Plugins/2D/Paper2D/Source/Paper2D/Classes/SpriteEditorOnlyTypes.h b/Engine/Plugins/2D/Paper2D/Source/Paper2D/Classes/SpriteEditorOnlyTypes.h
--- a/Engine/Plugins/2D/Paper2D/Source/Paper2D/Classes/SpriteEditorOnlyTypes.h
+++ b/Engine/Plugins/2D/Paper2D/Source/Paper2D/Classes/SpriteEditorOnlyTypes.h
@@ -13,7 +13,7 @@ class UMaterialInterface;
 UENUM()
 namespace ESpriteCollisionMode
 {
-	enum Type
+	enum Type : int
 	{
 		/** Should this have no collison and not participate in physics? */
 		None,
@@ -145,7 +145,7 @@ public:
 UENUM()
 namespace ESpritePolygonMode
 {
-	enum Type
+	enum Type : int
 	{
 		/** Use the bounding box of the source sprite (no optimization) */
 		SourceBoundingBox,
@@ -319,7 +319,7 @@ public:
 UENUM()
 namespace ESpritePivotMode
 {
-	enum Type
+	enum Type : int
 	{
 		Top_Left,
 		Top_Center,
diff --git a/Engine/Plugins/Animation/ControlRig/Source/ControlRig/Public/Rigs/RigHierarchy.h b/Engine/Plugins/Animation/ControlRig/Source/ControlRig/Public/Rigs/RigHierarchy.h
--- a/Engine/Plugins/Animation/ControlRig/Source/ControlRig/Public/Rigs/RigHierarchy.h
+++ b/Engine/Plugins/Animation/ControlRig/Source/ControlRig/Public/Rigs/RigHierarchy.h
@@ -25,7 +25,7 @@ DECLARE_MULTICAST_DELEGATE_TwoParams(FRigHierarchyMetadataChangedDelegate, const
 DECLARE_MULTICAST_DELEGATE_ThreeParams(FRigHierarchyMetadataTagChangedDelegate, const FRigElementKey& /* Key */, const FName& /* Tag */, bool /* AddedOrRemoved */);
 
 UENUM()
-enum ERigTransformStackEntryType
+enum ERigTransformStackEntryType : int
 {
 	TransformPose,
 	ControlOffset,
diff --git a/Engine/Plugins/Cameras/GameplayCameras/Source/GameplayCameras/Public/LegacyCameraShake.h b/Engine/Plugins/Cameras/GameplayCameras/Source/GameplayCameras/Public/LegacyCameraShake.h
--- a/Engine/Plugins/Cameras/GameplayCameras/Source/GameplayCameras/Public/LegacyCameraShake.h
+++ b/Engine/Plugins/Cameras/GameplayCameras/Source/GameplayCameras/Public/LegacyCameraShake.h
@@ -33,7 +33,7 @@ enum class EOscillatorWaveform : uint8
 
 /** Shake start offset parameter */
 UENUM()
-enum EInitialOscillatorOffset
+enum EInitialOscillatorOffset : int
 {
 	/** Start with random offset (default). */
 	EOO_OffsetRandom UMETA(DisplayName = "Random"),
diff --git a/Engine/Plugins/Experimental/GeometryScripting/Source/GeometryScriptingCore/Public/GeometryScript/GeometryScriptTypes.h b/Engine/Plugins/Experimental/GeometryScripting/Source/GeometryScriptingCore/Public/GeometryScript/GeometryScriptTypes.h
--- a/Engine/Plugins/Experimental/GeometryScripting/Source/GeometryScriptingCore/Public/GeometryScript/GeometryScriptTypes.h
+++ b/Engine/Plugins/Experimental/GeometryScripting/Source/GeometryScriptingCore/Public/GeometryScript/GeometryScriptTypes.h
@@ -13,21 +13,21 @@ PREDECLARE_GEOMETRY(typedef TMeshAABBTree3<FDynamicMesh3> FDynamicMeshAABBTree3)
 
 
 UENUM(BlueprintType)
-enum EGeometryScriptOutcomePins
+enum EGeometryScriptOutcomePins : int
 {
 	Failure,
 	Success
 };
 
 UENUM(BlueprintType)
-enum EGeometryScriptSearchOutcomePins
+enum EGeometryScriptSearchOutcomePins : int
 {
 	Found,
 	NotFound
 };
 
 UENUM(BlueprintType)
-enum EGeometryScriptContainmentOutcomePins
+enum EGeometryScriptContainmentOutcomePins : int
 {
 	Inside, 
 	Outside
diff --git a/Engine/Plugins/Experimental/Mutable/Source/CustomizableObjectEditor/Classes/MuCOE/Nodes/CustomizableObjectNodeTextureLayer.h b/Engine/Plugins/Experimental/Mutable/Source/CustomizableObjectEditor/Classes/MuCOE/Nodes/CustomizableObjectNodeTextureLayer.h
--- a/Engine/Plugins/Experimental/Mutable/Source/CustomizableObjectEditor/Classes/MuCOE/Nodes/CustomizableObjectNodeTextureLayer.h
+++ b/Engine/Plugins/Experimental/Mutable/Source/CustomizableObjectEditor/Classes/MuCOE/Nodes/CustomizableObjectNodeTextureLayer.h
@@ -23,7 +23,7 @@ struct FPropertyChangedEvent;
 
 
 UENUM()
-enum ECustomizableObjectTextureLayerEffect
+enum ECustomizableObjectTextureLayerEffect : int
 {
 	COTLE_MODULATE			= 0 UMETA(DisplayName = "MODULATE"), 
 	COTLE_MULTIPLY				UMETA(DisplayName = "MULTIPLY"), 
diff --git a/Engine/Plugins/FX/Niagara/Source/Niagara/Public/NiagaraCommon.h b/Engine/Plugins/FX/Niagara/Source/Niagara/Public/NiagaraCommon.h
--- a/Engine/Plugins/FX/Niagara/Source/Niagara/Public/NiagaraCommon.h
+++ b/Engine/Plugins/FX/Niagara/Source/Niagara/Public/NiagaraCommon.h
@@ -903,7 +903,7 @@ enum class ENiagaraIterationSource : uint8
 };
 
 UENUM()
-enum ENiagaraBindingSource 
+enum ENiagaraBindingSource : int
 {
 	ImplicitFromSource = 0,
 	ExplicitParticles,
@@ -1520,7 +1520,7 @@ struct NIAGARA_API FSynchronizeWithParameterDefinitionsArgs
 UENUM()
 namespace ENiagaraGpuComputeTickStage
 {
-	enum Type
+	enum Type : int
 	{
 		PreInitViews,
 		PostInitViews,
diff --git a/Engine/Plugins/FX/Niagara/Source/Niagara/Public/NiagaraSettings.h b/Engine/Plugins/FX/Niagara/Source/Niagara/Public/NiagaraSettings.h
--- a/Engine/Plugins/FX/Niagara/Source/Niagara/Public/NiagaraSettings.h
+++ b/Engine/Plugins/FX/Niagara/Source/Niagara/Public/NiagaraSettings.h
@@ -15,7 +15,7 @@
 UENUM()
 namespace ENDISkelMesh_GpuMaxInfluences
 {
-	enum Type
+	enum Type : int
 	{
 		/** Allow up to 4 bones to be sampled. */
 		AllowMax4 = 0,
@@ -30,7 +30,7 @@ namespace ENDISkelMesh_GpuMaxInfluences
 UENUM()
 namespace ENDISkelMesh_GpuUniformSamplingFormat
 {
-	enum Type
+	enum Type : int
 	{
 		/** 64 bits per entry. Allow for the full int32 range of triangles (2 billion). */
 		Full = 0,
@@ -45,7 +45,7 @@ namespace ENDISkelMesh_GpuUniformSamplingFormat
 UENUM()
 namespace ENDISkelMesh_AdjacencyTriangleIndexFormat
 {
-	enum Type
+	enum Type : int
 	{
 		/** 32 bits per entry. Allow for the full int32 range of triangles (2 billion). */
 		Full = 0,
@@ -84,7 +84,7 @@ enum class ENiagaraDefaultGpuTranslucentLatency : uint8
 UENUM()
 namespace ENDICollisionQuery_AsyncGpuTraceProvider
 {
-	enum Type
+	enum Type : int
 	{
 		Default = 0 UMETA(DisplayName = "Project Default"),
 		HWRT = 1 UMETA(DisplayName = "HW Ray Tracing"),
diff --git a/Engine/Plugins/Online/OnlineFramework/Source/Party/Public/SocialManager.h b/Engine/Plugins/Online/OnlineFramework/Source/Party/Public/SocialManager.h
--- a/Engine/Plugins/Online/OnlineFramework/Source/Party/Public/SocialManager.h
+++ b/Engine/Plugins/Online/OnlineFramework/Source/Party/Public/SocialManager.h
@@ -23,7 +23,7 @@ class FOnlineSessionSearchResult;
 class FPartyPlatformSessionManager;
 class USocialDebugTools;
 
-enum ETravelType;
+enum ETravelType : int;
 
 #define ABORT_DURING_SHUTDOWN() if (IsEngineExitRequested() || bShutdownPending) { UE_LOG(LogParty, Log, TEXT("%s - Received callback during shutdown: IsEngineExitRequested=%s, bShutdownPending=%s."), ANSI_TO_TCHAR(__FUNCTION__), *LexToString(IsEngineExitRequested()), *LexToString(bShutdownPending)); return; }
 
diff --git a/Engine/Plugins/Online/OnlineSubsystem/Source/Public/Interfaces/OnlineTurnBasedInterface.h b/Engine/Plugins/Online/OnlineSubsystem/Source/Public/Interfaces/OnlineTurnBasedInterface.h
--- a/Engine/Plugins/Online/OnlineSubsystem/Source/Public/Interfaces/OnlineTurnBasedInterface.h
+++ b/Engine/Plugins/Online/OnlineSubsystem/Source/Public/Interfaces/OnlineTurnBasedInterface.h
@@ -66,7 +66,7 @@ DECLARE_DELEGATE_TwoParams(FEndMatchSignature, FString, bool);   // bool Success
 UENUM(BlueprintType)
 namespace EMPMatchOutcome
 {
-    enum Outcome
+    enum Outcome : int
     {
         None,
         Quit,
diff --git a/Engine/Plugins/Runtime/AndroidDeviceProfileSelector/Source/AndroidDeviceProfileSelector/Private/AndroidDeviceProfileMatchingRules.h b/Engine/Plugins/Runtime/AndroidDeviceProfileSelector/Source/AndroidDeviceProfileSelector/Private/AndroidDeviceProfileMatchingRules.h
--- a/Engine/Plugins/Runtime/AndroidDeviceProfileSelector/Source/AndroidDeviceProfileSelector/Private/AndroidDeviceProfileMatchingRules.h
+++ b/Engine/Plugins/Runtime/AndroidDeviceProfileSelector/Source/AndroidDeviceProfileSelector/Private/AndroidDeviceProfileMatchingRules.h
@@ -7,7 +7,7 @@
 #include "AndroidDeviceProfileMatchingRules.generated.h"
 
 UENUM()
-enum ESourceType
+enum ESourceType : int
 {
 	SRC_PreviousRegexMatch,
 	SRC_GpuFamily,
@@ -28,7 +28,7 @@ enum ESourceType
 };
 
 UENUM()
-enum ECompareType
+enum ECompareType : int
 {
 	CMP_Equal,
 	CMP_Less,
diff --git a/Engine/Plugins/Runtime/AudioWidgets/Source/AudioWidgets/Public/SAudioRadialSlider.h b/Engine/Plugins/Runtime/AudioWidgets/Source/AudioWidgets/Public/SAudioRadialSlider.h
--- a/Engine/Plugins/Runtime/AudioWidgets/Source/AudioWidgets/Public/SAudioRadialSlider.h
+++ b/Engine/Plugins/Runtime/AudioWidgets/Source/AudioWidgets/Public/SAudioRadialSlider.h
@@ -21,7 +21,7 @@
 class SRadialSlider;
 
 UENUM()
-enum EAudioRadialSliderLayout
+enum EAudioRadialSliderLayout : int
 {
 	/** Label above radial slider. */
 	Layout_LabelTop UMETA(DisplayName = "Label Top"),
diff --git a/Engine/Plugins/Runtime/GameplayAbilities/Source/GameplayAbilities/Public/Abilities/GameplayAbilityTargetTypes.h b/Engine/Plugins/Runtime/GameplayAbilities/Source/GameplayAbilities/Public/Abilities/GameplayAbilityTargetTypes.h
--- a/Engine/Plugins/Runtime/GameplayAbilities/Source/GameplayAbilities/Public/Abilities/GameplayAbilityTargetTypes.h
+++ b/Engine/Plugins/Runtime/GameplayAbilities/Source/GameplayAbilities/Public/Abilities/GameplayAbilityTargetTypes.h
@@ -26,7 +26,7 @@ UENUM(BlueprintType)
 namespace EGameplayTargetingConfirmation
 {
 	/** Describes how the targeting information is confirmed */
-	enum Type
+	enum Type : int
 	{
 		/** The targeting happens instantly without special logic or user input deciding when to 'fire' */
 		Instant,
@@ -166,7 +166,7 @@ UENUM(BlueprintType)
 namespace EGameplayAbilityTargetingLocationType
 {
 	/** What type of location calculation to use when an ability asks for our transform */
-	enum Type
+	enum Type : int
 	{
 		/** We report an actual raw transform. This is also the final fallback if other methods fail */
 		LiteralTransform		UMETA(DisplayName = "Literal Transform"),
@@ -641,7 +641,7 @@ DECLARE_MULTICAST_DELEGATE_TwoParams(FAbilityTargetDataSetDelegate, const FGamep
 UENUM()
 namespace EAbilityGenericReplicatedEvent
 {
-	enum Type
+	enum Type : int
 	{	
 		/** A generic confirmation to commit the ability */
 		GenericConfirm = 0,
diff --git a/Engine/Plugins/Runtime/GameplayAbilities/Source/GameplayAbilities/Public/Abilities/GameplayAbilityTypes.h b/Engine/Plugins/Runtime/GameplayAbilities/Source/GameplayAbilities/Public/Abilities/GameplayAbilityTypes.h
--- a/Engine/Plugins/Runtime/GameplayAbilities/Source/GameplayAbilities/Public/Abilities/GameplayAbilityTypes.h
+++ b/Engine/Plugins/Runtime/GameplayAbilities/Source/GameplayAbilities/Public/Abilities/GameplayAbilityTypes.h
@@ -44,7 +44,7 @@ namespace EGameplayAbilityInstancingPolicy
 	 *	How the ability is instanced when executed. This limits what an ability can do in its implementation. For example, a NonInstanced
 	 *	Ability cannot have state. It is probably unsafe for an InstancedPerActor ability to have latent actions, etc.
 	 */
-	enum Type
+	enum Type : int
 	{
 		// This ability is never instanced. Anything that executes the ability is operating on the CDO.
 		NonInstanced,
@@ -61,7 +61,7 @@ UENUM(BlueprintType)
 namespace EGameplayAbilityNetExecutionPolicy
 {
 	/** Where does an ability execute on the network. Does a client "ask and predict", "ask and wait", "don't ask (just do it)" */
-	enum Type
+	enum Type : int
 	{
 		// Part of this ability runs predictively on the local client if there is one
 		LocalPredicted		UMETA(DisplayName = "Local Predicted"),
@@ -81,7 +81,7 @@ UENUM(BlueprintType)
 namespace EGameplayAbilityNetSecurityPolicy
 {
 	/** What protections does this ability have? Should the client be allowed to request changes to the execution of the ability? */
-	enum Type
+	enum Type : int
 	{
 		// No security requirements. Client or server can trigger execution and termination of this ability freely.
 		ClientOrServer			UMETA(DisplayName = "Client Or Server"),
@@ -101,7 +101,7 @@ UENUM(BlueprintType)
 namespace EGameplayAbilityReplicationPolicy
 {
 	/** How an ability replicates state/events to everyone on the network */
-	enum Type
+	enum Type : int
 	{
 		// We don't replicate the instance of the ability to anyone.
 		ReplicateNo			UMETA(DisplayName = "Do Not Replicate"),
@@ -115,7 +115,7 @@ UENUM(BlueprintType)
 namespace EGameplayAbilityTriggerSource
 {
 	/**	Defines what type of trigger will activate the ability, paired to a tag */
-	enum Type
+	enum Type : int
 	{
 		// Triggered from a gameplay event, will come with payload
 		GameplayEvent,
diff --git a/Engine/Source/Developer/PhysicsUtilities/Public/PhysicsAssetUtils.h b/Engine/Source/Developer/PhysicsUtilities/Public/PhysicsAssetUtils.h
--- a/Engine/Source/Developer/PhysicsUtilities/Public/PhysicsAssetUtils.h
+++ b/Engine/Source/Developer/PhysicsUtilities/Public/PhysicsAssetUtils.h
@@ -14,7 +14,7 @@ class USkeletalMesh;
 class USkeletalMeshComponent;
 
 UENUM()
-enum EPhysAssetFitGeomType
+enum EPhysAssetFitGeomType : int
 {
 	EFG_Box					UMETA(DisplayName="Box"),
 	EFG_Sphyl				UMETA(DisplayName="Capsule"),
@@ -26,7 +26,7 @@ enum EPhysAssetFitGeomType
 };
 
 UENUM()
-enum EPhysAssetFitVertWeight
+enum EPhysAssetFitVertWeight : int
 {
 	EVW_AnyWeight			UMETA(DisplayName="Any Weight"),
 	EVW_DominantWeight		UMETA(DisplayName="Dominant Weight"),
diff --git a/Engine/Source/Editor/AnimGraph/Public/AnimPreviewInstance.h b/Engine/Source/Editor/AnimGraph/Public/AnimPreviewInstance.h
--- a/Engine/Source/Editor/AnimGraph/Public/AnimPreviewInstance.h
+++ b/Engine/Source/Editor/AnimGraph/Public/AnimPreviewInstance.h
@@ -15,7 +15,7 @@
 
 /** Enum to know how montage is being played */
 UENUM()
-enum EMontagePreviewType
+enum EMontagePreviewType : int
 {
 	/** Playing montage in usual way. */
 	EMPT_Normal, 
diff --git a/Engine/Source/Editor/AnimGraph/Public/AnimStateNode.h b/Engine/Source/Editor/AnimGraph/Public/AnimStateNode.h
--- a/Engine/Source/Editor/AnimGraph/Public/AnimStateNode.h
+++ b/Engine/Source/Editor/AnimGraph/Public/AnimStateNode.h
@@ -12,7 +12,7 @@ class UEdGraph;
 class UEdGraphPin;
 
 UENUM()
-enum EAnimStateType
+enum EAnimStateType : int
 {
 	AST_SingleAnimation UMETA(DisplayName="Single animation"),
 	AST_BlendGraph UMETA(DisplayName="Blend graph"),
diff --git a/Engine/Source/Editor/BlueprintGraph/Public/BlueprintEditorSettings.h b/Engine/Source/Editor/BlueprintGraph/Public/BlueprintEditorSettings.h
--- a/Engine/Source/Editor/BlueprintGraph/Public/BlueprintEditorSettings.h
+++ b/Engine/Source/Editor/BlueprintGraph/Public/BlueprintEditorSettings.h
@@ -35,7 +35,7 @@ struct FEditedDocumentInfo;
 struct FGuid;
 
 UENUM()
-enum ESaveOnCompile
+enum ESaveOnCompile : int
 {
 	SoC_Never UMETA(DisplayName="Never"),
 	SoC_SuccessOnly UMETA(DisplayName="On Success Only"),
diff --git a/Engine/Source/Editor/ConfigEditor/Public/ConfigPropertyHelper.h b/Engine/Source/Editor/ConfigEditor/Public/ConfigPropertyHelper.h
--- a/Engine/Source/Editor/ConfigEditor/Public/ConfigPropertyHelper.h
+++ b/Engine/Source/Editor/ConfigEditor/Public/ConfigPropertyHelper.h
@@ -19,7 +19,7 @@ class FProperty;
 template <typename T> struct TObjectPtr;
 
 UENUM()
-enum EConfigFileSourceControlStatus
+enum EConfigFileSourceControlStatus : int
 {
 	CFSCS_Unknown UMETA(DisplayName=Unknown),
 	CFSCS_Writable UMETA(DisplayName = "Available to edit"),
diff --git a/Engine/Source/Editor/Kismet/Public/FindInBlueprintManager.h b/Engine/Source/Editor/Kismet/Public/FindInBlueprintManager.h
--- a/Engine/Source/Editor/Kismet/Public/FindInBlueprintManager.h
+++ b/Engine/Source/Editor/Kismet/Public/FindInBlueprintManager.h
@@ -141,7 +141,7 @@ struct KISMET_API FFindInBlueprintSearchTags
 
 /** FiB data versioning */
 UENUM()
-enum EFiBVersion
+enum EFiBVersion : int
 {
 	FIB_VER_NONE = -1, // Unknown version (not set)
 
diff --git a/Engine/Source/Editor/PropertyEditor/Private/EditConditionParserTests.h b/Engine/Source/Editor/PropertyEditor/Private/EditConditionParserTests.h
--- a/Engine/Source/Editor/PropertyEditor/Private/EditConditionParserTests.h
+++ b/Engine/Source/Editor/PropertyEditor/Private/EditConditionParserTests.h
@@ -20,7 +20,7 @@ enum class EditConditionTestEnum
 };
 
 UENUM()
-enum EditConditionByteEnum
+enum EditConditionByteEnum : int
 {
 	First = 15,
 	Second = 31
diff --git a/Engine/Source/Editor/Sequencer/Public/ISequencerSection.h b/Engine/Source/Editor/Sequencer/Public/ISequencerSection.h
--- a/Engine/Source/Editor/Sequencer/Public/ISequencerSection.h
+++ b/Engine/Source/Editor/Sequencer/Public/ISequencerSection.h
@@ -44,7 +44,7 @@ namespace UE::Sequencer
 
 /** Enumerates which edge is being resized */
 UENUM()
-enum ESequencerSectionResizeMode
+enum ESequencerSectionResizeMode : int
 {
 	SSRM_LeadingEdge,
 	SSRM_TrailingEdge
diff --git a/Engine/Source/Editor/Sequencer/Public/SequencerSettings.h b/Engine/Source/Editor/Sequencer/Public/SequencerSettings.h
--- a/Engine/Source/Editor/Sequencer/Public/SequencerSettings.h
+++ b/Engine/Source/Editor/Sequencer/Public/SequencerSettings.h
@@ -16,7 +16,7 @@ enum class EKeyGroupMode : uint8;
 enum class EMovieSceneKeyInterpolation : uint8;
 
 UENUM()
-enum ESequencerSpawnPosition
+enum ESequencerSpawnPosition : int
 {
 	/** Origin. */
 	SSP_Origin UMETA(DisplayName="Origin"),
@@ -26,7 +26,7 @@ enum ESequencerSpawnPosition
 };
 
 UENUM()
-enum ESequencerZoomPosition
+enum ESequencerZoomPosition : int
 {
 	/** Current Time. */
 	SZP_CurrentTime UMETA(DisplayName="Playhead"),
@@ -36,7 +36,7 @@ enum ESequencerZoomPosition
 };
 
 UENUM()
-enum ESequencerLoopMode
+enum ESequencerLoopMode : int
 {
 	/** No Looping. */
 	SLM_NoLoop UMETA(DisplayName="No Looping"),
diff --git a/Engine/Source/Editor/StatsViewer/Classes/CookerStats.h b/Engine/Source/Editor/StatsViewer/Classes/CookerStats.h
--- a/Engine/Source/Editor/StatsViewer/Classes/CookerStats.h
+++ b/Engine/Source/Editor/StatsViewer/Classes/CookerStats.h
@@ -10,7 +10,7 @@
 
 /** Enum defining the object sets for this stats object */
 UENUM()
-enum ECookerStatsObjectSets
+enum ECookerStatsObjectSets : int
 {
 	CookerStatsObjectSets_Default UMETA(DisplayName="Default", ToolTip="View cooker statistics"),
 };
diff --git a/Engine/Source/Editor/StatsViewer/Classes/LightingBuildInfo.h b/Engine/Source/Editor/StatsViewer/Classes/LightingBuildInfo.h
--- a/Engine/Source/Editor/StatsViewer/Classes/LightingBuildInfo.h
+++ b/Engine/Source/Editor/StatsViewer/Classes/LightingBuildInfo.h
@@ -10,7 +10,7 @@
 
 /** Enum defining the object sets for this stats object */
 UENUM()
-enum ELightingBuildInfoObjectSets
+enum ELightingBuildInfoObjectSets : int
 {
 	LightingBuildInfoObjectSets_Default				UMETA( DisplayName = "Default" , ToolTip = "View lighting build statistics" ),
 };
diff --git a/Engine/Source/Editor/StatsViewer/Classes/PrimitiveStats.h b/Engine/Source/Editor/StatsViewer/Classes/PrimitiveStats.h
--- a/Engine/Source/Editor/StatsViewer/Classes/PrimitiveStats.h
+++ b/Engine/Source/Editor/StatsViewer/Classes/PrimitiveStats.h
@@ -12,7 +12,7 @@ class AActor;
 
 /** Enum defining the object sets for this stats object */
 UENUM()
-enum EPrimitiveObjectSets
+enum EPrimitiveObjectSets : int
 {
 	PrimitiveObjectSets_AllObjects			UMETA( DisplayName = "All Objects" , ToolTip = "View primitive statistics for all objects in all levels" ),
 	PrimitiveObjectSets_CurrentLevel		UMETA( DisplayName = "Current Level" , ToolTip = "View primitive statistics for objects in the current level" ),
diff --git a/Engine/Source/Editor/StatsViewer/Classes/ShaderCookerStats.h b/Engine/Source/Editor/StatsViewer/Classes/ShaderCookerStats.h
--- a/Engine/Source/Editor/StatsViewer/Classes/ShaderCookerStats.h
+++ b/Engine/Source/Editor/StatsViewer/Classes/ShaderCookerStats.h
@@ -10,7 +10,7 @@
 
 
 UENUM()
-enum EShaderCookerStatsSets
+enum EShaderCookerStatsSets : int
 {
 	EShaderCookerStatsSets_Default				UMETA(DisplayName = "Default", ToolTip = "Shader Cooker Sets"),
 };
diff --git a/Engine/Source/Editor/StatsViewer/Classes/StaticMeshLightingInfo.h b/Engine/Source/Editor/StatsViewer/Classes/StaticMeshLightingInfo.h
--- a/Engine/Source/Editor/StatsViewer/Classes/StaticMeshLightingInfo.h
+++ b/Engine/Source/Editor/StatsViewer/Classes/StaticMeshLightingInfo.h
@@ -14,7 +14,7 @@ class UStaticMeshComponent;
 
 /** Enum defining the object sets for this stats object */
 UENUM()
-enum EStaticMeshLightingInfoObjectSets
+enum EStaticMeshLightingInfoObjectSets : int
 {
 	StaticMeshLightingInfoObjectSets_CurrentLevel		UMETA( DisplayName = "Current Level" , ToolTip = "View static mesh lighting info for the current level" ),
 	StaticMeshLightingInfoObjectSets_SelectedLevels		UMETA( DisplayName = "Selected Levels" , ToolTip = "View lighting info for selected levels" ),
diff --git a/Engine/Source/Editor/StatsViewer/Classes/TextureStats.h b/Engine/Source/Editor/StatsViewer/Classes/TextureStats.h
--- a/Engine/Source/Editor/StatsViewer/Classes/TextureStats.h
+++ b/Engine/Source/Editor/StatsViewer/Classes/TextureStats.h
@@ -15,7 +15,7 @@ class UTexture;
 
 /** Enum defining the object sets for this stats object */
 UENUM()
-enum ETextureObjectSets
+enum ETextureObjectSets : int
 {
 	TextureObjectSet_CurrentStreamingLevel	UMETA( DisplayName = "Current Streaming Level" , ToolTip = "Display texture statistics for the current streaming level" ),
 	TextureObjectSet_AllStreamingLevels		UMETA( DisplayName = "All Streaming Levels" , ToolTip = "Display texture statistics for all streaming levels" ),
diff --git a/Engine/Source/Editor/TextureEditor/Classes/TextureEditorSettings.h b/Engine/Source/Editor/TextureEditor/Classes/TextureEditorSettings.h
--- a/Engine/Source/Editor/TextureEditor/Classes/TextureEditorSettings.h
+++ b/Engine/Source/Editor/TextureEditor/Classes/TextureEditorSettings.h
@@ -13,7 +13,7 @@
  * Enumerates background for the texture editor view port.
  */
 UENUM()
-enum ETextureEditorBackgrounds
+enum ETextureEditorBackgrounds : int
 {
 	TextureEditorBackground_SolidColor UMETA(DisplayName="Solid Color"),
 	TextureEditorBackground_Checkered UMETA(DisplayName="Checkered"),
@@ -21,21 +21,21 @@ enum ETextureEditorBackgrounds
 };
 
 UENUM()
-enum ETextureEditorSampling
+enum ETextureEditorSampling : int
 {
 	TextureEditorSampling_Default UMETA(DisplayName = "Default Sampling"),
 	TextureEditorSampling_Point UMETA(DisplayName = "Nearest-Point Sampling"),
 };
 
 UENUM()
-enum ETextureEditorVolumeViewMode
+enum ETextureEditorVolumeViewMode : int
 {
 	TextureEditorVolumeViewMode_DepthSlices UMETA(DisplayName="Depth Slices"),
 	TextureEditorVolumeViewMode_VolumeTrace UMETA(DisplayName="Trace Into Volume"),
 };
 
 UENUM()
-enum ETextureEditorCubemapViewMode
+enum ETextureEditorCubemapViewMode : int
 {
 	TextureEditorCubemapViewMode_2DView UMETA(DisplayName = "2D View"),
 	TextureEditorCubemapViewMode_3DView UMETA(DisplayName = "3D View"),
diff --git a/Engine/Source/Editor/UMGEditor/Public/Settings/WidgetDesignerSettings.h b/Engine/Source/Editor/UMGEditor/Public/Settings/WidgetDesignerSettings.h
--- a/Engine/Source/Editor/UMGEditor/Public/Settings/WidgetDesignerSettings.h
+++ b/Engine/Source/Editor/UMGEditor/Public/Settings/WidgetDesignerSettings.h
@@ -19,7 +19,7 @@ class UObject;
 class UWidgetPaletteFavorites;
 
 UENUM()
-enum EDisplayOnCompile
+enum EDisplayOnCompile : int
 {
 	DoC_ErrorsOrWarnings,
 	DoC_ErrorsOnly,
diff --git a/Engine/Source/Editor/UnrealEd/Classes/Builders/SheetBuilder.h b/Engine/Source/Editor/UnrealEd/Classes/Builders/SheetBuilder.h
--- a/Engine/Source/Editor/UnrealEd/Classes/Builders/SheetBuilder.h
+++ b/Engine/Source/Editor/UnrealEd/Classes/Builders/SheetBuilder.h
@@ -15,7 +15,7 @@
 class ABrush;
 
 UENUM()
-enum ESheetAxis
+enum ESheetAxis : int
 {
 	AX_Horizontal,
 	AX_XAxis,
diff --git a/Engine/Source/Editor/UnrealEd/Classes/Editor/EditorEngine.h b/Engine/Source/Editor/UnrealEd/Classes/Editor/EditorEngine.h
--- a/Engine/Source/Editor/UnrealEd/Classes/Editor/EditorEngine.h
+++ b/Engine/Source/Editor/UnrealEd/Classes/Editor/EditorEngine.h
@@ -93,7 +93,7 @@ class IProjectExternalContentInterface;
 // Things to set in mapSetBrush.
 //
 UENUM()
-enum EMapSetBrushFlags				
+enum EMapSetBrushFlags : int
 {
 	/** Set brush color. */
 	MSB_BrushColor	= 1,
@@ -106,7 +106,7 @@ enum EMapSetBrushFlags
 };
 
 UENUM()
-enum EPasteTo
+enum EPasteTo : int
 {
 	PT_OriginalLocation	= 0,
 	PT_Here				= 1,
diff --git a/Engine/Source/Editor/UnrealEd/Classes/Editor/PropertyEditorTestObject.h b/Engine/Source/Editor/UnrealEd/Classes/Editor/PropertyEditorTestObject.h
--- a/Engine/Source/Editor/UnrealEd/Classes/Editor/PropertyEditorTestObject.h
+++ b/Engine/Source/Editor/UnrealEd/Classes/Editor/PropertyEditorTestObject.h
@@ -31,7 +31,7 @@ class UTexture;
 template <typename ItemType> class SListView;
 
 UENUM()
-enum EPropertyEditorTestEnum
+enum EPropertyEditorTestEnum : int
 {	
 	/** This comment should appear above enum 1 */
 	PropertyEditorTest_Enum1 UMETA(Hidden),
@@ -59,7 +59,7 @@ enum class EPropertyEditorTestBitflags : uint8
 ENUM_CLASS_FLAGS(EPropertyEditorTestBitflags)
 
 UENUM()
-enum ArrayLabelEnum
+enum ArrayLabelEnum : int
 {
 	ArrayIndex0,
 	ArrayIndex1,
diff --git a/Engine/Source/Editor/UnrealEd/Classes/Editor/UnrealEdEngine.h b/Engine/Source/Editor/UnrealEd/Classes/Editor/UnrealEdEngine.h
--- a/Engine/Source/Editor/UnrealEd/Classes/Editor/UnrealEdEngine.h
+++ b/Engine/Source/Editor/UnrealEd/Classes/Editor/UnrealEdEngine.h
@@ -35,7 +35,7 @@ typedef FName FEditorModeID;
 struct FTypedElementSelectionOptions;
 
 UENUM()
-enum EPackageNotifyState
+enum EPackageNotifyState : int
 {
 	/** Updating the source control state of the package */
 	NS_Updating,
diff --git a/Engine/Source/Editor/UnrealEd/Classes/Editor/UnrealEdTypes.h b/Engine/Source/Editor/UnrealEd/Classes/Editor/UnrealEdTypes.h
--- a/Engine/Source/Editor/UnrealEd/Classes/Editor/UnrealEdTypes.h
+++ b/Engine/Source/Editor/UnrealEd/Classes/Editor/UnrealEdTypes.h
@@ -108,7 +108,7 @@ struct FLightmassParameterizedMaterialSettings
 //	ELevelViewportType
 //
 UENUM()
-enum ELevelViewportType
+enum ELevelViewportType : int
 {
 	/** Top */
 	LVT_OrthoXY = 0,
@@ -130,7 +130,7 @@ enum ELevelViewportType
 };
 
 UENUM()
-enum EDestructiveAssetActions
+enum EDestructiveAssetActions : int
 {
 	AssetDelete = 0,
 	AssetRename = 1 << 0,
diff --git a/Engine/Source/Editor/UnrealEd/Classes/Factories/FbxAnimSequenceImportData.h b/Engine/Source/Editor/UnrealEd/Classes/Factories/FbxAnimSequenceImportData.h
--- a/Engine/Source/Editor/UnrealEd/Classes/Factories/FbxAnimSequenceImportData.h
+++ b/Engine/Source/Editor/UnrealEd/Classes/Factories/FbxAnimSequenceImportData.h
@@ -18,7 +18,7 @@ b) As for FbxNode::GetAnimationInternval, this one will iterate through all prop
 
 /** Animation length type when importing */
 UENUM(BlueprintType)
-enum EFBXAnimationLengthImportType
+enum EFBXAnimationLengthImportType : int
 {
 	/** This option imports animation frames based on what is defined at the time of export */
 	FBXALIT_ExportedTime			UMETA(DisplayName = "Exported Time"),
diff --git a/Engine/Source/Editor/UnrealEd/Classes/Factories/FbxImportUI.h b/Engine/Source/Editor/UnrealEd/Classes/Factories/FbxImportUI.h
--- a/Engine/Source/Editor/UnrealEd/Classes/Factories/FbxImportUI.h
+++ b/Engine/Source/Editor/UnrealEd/Classes/Factories/FbxImportUI.h
@@ -19,7 +19,7 @@ namespace UnFbx
 
 /** Import mesh type */
 UENUM(BlueprintType)
-enum EFBXImportType
+enum EFBXImportType : int
 {
 	/** Select Static Mesh if you'd like to import static mesh. */
 	FBXIT_StaticMesh UMETA(DisplayName="Static Mesh"),
diff --git a/Engine/Source/Editor/UnrealEd/Classes/Factories/FbxMeshImportData.h b/Engine/Source/Editor/UnrealEd/Classes/Factories/FbxMeshImportData.h
--- a/Engine/Source/Editor/UnrealEd/Classes/Factories/FbxMeshImportData.h
+++ b/Engine/Source/Editor/UnrealEd/Classes/Factories/FbxMeshImportData.h
@@ -8,7 +8,7 @@
 #include "FbxMeshImportData.generated.h"
 
 UENUM(BlueprintType)
-enum EFBXNormalImportMethod
+enum EFBXNormalImportMethod : int
 {
 	FBXNIM_ComputeNormals UMETA(DisplayName="Compute Normals"),
 	FBXNIM_ImportNormals UMETA(DisplayName="Import Normals"),
@@ -19,7 +19,7 @@ enum EFBXNormalImportMethod
 UENUM(BlueprintType)
 namespace EFBXNormalGenerationMethod
 {
-	enum Type
+	enum Type : int
 	{
 		/** Use the legacy built in method to generate normals (faster in some cases) */
 		BuiltIn,
@@ -31,7 +31,7 @@ namespace EFBXNormalGenerationMethod
 UENUM(BlueprintType)
 namespace EVertexColorImportOption
 {
-	enum Type
+	enum Type : int
 	{
 		/** Import the static mesh using the vertex colors from the FBX file. */
 		Replace,
diff --git a/Engine/Source/Editor/UnrealEd/Classes/Factories/FbxSkeletalMeshImportData.h b/Engine/Source/Editor/UnrealEd/Classes/Factories/FbxSkeletalMeshImportData.h
--- a/Engine/Source/Editor/UnrealEd/Classes/Factories/FbxSkeletalMeshImportData.h
+++ b/Engine/Source/Editor/UnrealEd/Classes/Factories/FbxSkeletalMeshImportData.h
@@ -24,7 +24,7 @@ namespace SkeletalMeshImportData
 }
 
 UENUM(BlueprintType)
-enum EFBXImportContentType
+enum EFBXImportContentType : int
 {
 	FBXICT_All UMETA(DisplayName = "Geometry and Skinning Weights.", ToolTip = "Import all fbx content: geometry, skinning and weights."),
 	FBXICT_Geometry UMETA(DisplayName = "Geometry Only", ToolTip = "Import the skeletal mesh geometry only (will create a default skeleton, or map the geometry to the existing one). Morph and LOD can be imported with it."),
diff --git a/Engine/Source/Editor/UnrealEd/Classes/Settings/LevelEditorPlaySettings.h b/Engine/Source/Editor/UnrealEd/Classes/Settings/LevelEditorPlaySettings.h
--- a/Engine/Source/Editor/UnrealEd/Classes/Settings/LevelEditorPlaySettings.h
+++ b/Engine/Source/Editor/UnrealEd/Classes/Settings/LevelEditorPlaySettings.h
@@ -35,7 +35,7 @@ enum ELabelAnchorMode : int
 
 
 UENUM()
-enum ELaunchModeType
+enum ELaunchModeType : int
 {
 	/** Runs the map on a specified device. */
 	LaunchMode_OnDevice,
@@ -43,7 +43,7 @@ enum ELaunchModeType
 
 
 UENUM()
-enum EPlayModeLocations
+enum EPlayModeLocations : int
 {
 	/** Spawns the player at the current camera location. */
 	PlayLocation_CurrentCameraLocation,
@@ -54,7 +54,7 @@ enum EPlayModeLocations
 
 
 UENUM()
-enum EPlayModeType
+enum EPlayModeType : int
 {
 	/** Runs from within the editor. */
 	PlayMode_InViewPort = 0,
@@ -104,7 +104,7 @@ enum EPlayNetMode : int
  * Determines whether to build the executable when launching on device. Note the equivalence between these settings and EProjectPackagingBuild.
  */
 UENUM()
-enum EPlayOnBuildMode
+enum EPlayOnBuildMode : int
 {
 	/** Always build. */
 	PlayOnBuild_Always UMETA(DisplayName="Always"),
@@ -121,7 +121,7 @@ enum EPlayOnBuildMode
 
 /* Configuration to use when launching on device. */
 UENUM()
-enum EPlayOnLaunchConfiguration
+enum EPlayOnLaunchConfiguration : int
 {
 	/** Launch on device with the same build configuration as the editor. */
 	LaunchConfig_Default UMETA(DisplayName = "Same as Editor"),
diff --git a/Engine/Source/Editor/UnrealEd/Classes/Settings/LevelEditorViewportSettings.h b/Engine/Source/Editor/UnrealEd/Classes/Settings/LevelEditorViewportSettings.h
--- a/Engine/Source/Editor/UnrealEd/Classes/Settings/LevelEditorViewportSettings.h
+++ b/Engine/Source/Editor/UnrealEd/Classes/Settings/LevelEditorViewportSettings.h
@@ -30,7 +30,7 @@ enum ERotationGridMode : int
  * Enumerates camera control types for the W, A, S and D keys.
  */
 UENUM()
-enum EWASDType
+enum EWASDType : int
 {
 	WASD_Always  UMETA(DisplayName="Use WASD for Camera Controls"),
 	WASD_RMBOnly UMETA(DisplayName="Use WASD only when a Mouse Button is Pressed"),
@@ -53,7 +53,7 @@ enum class ELandscapeFoliageEditorControlType : uint8
  * Units used by measuring tool
  */
 UENUM()
-enum EMeasuringToolUnits
+enum EMeasuringToolUnits : int
 {
 	MeasureUnits_Centimeters UMETA(DisplayName="Centimeters"),
 	MeasureUnits_Meters      UMETA(DisplayName="Meters"),
diff --git a/Engine/Source/Editor/UnrealEd/Classes/TexAligner/TexAligner.h b/Engine/Source/Editor/UnrealEd/Classes/TexAligner/TexAligner.h
--- a/Engine/Source/Editor/UnrealEd/Classes/TexAligner/TexAligner.h
+++ b/Engine/Source/Editor/UnrealEd/Classes/TexAligner/TexAligner.h
@@ -19,7 +19,7 @@ class UModel;
 
 /** Alignment types. */
 UENUM()
-enum ETexAlign
+enum ETexAlign : int
 {
 	/** When passed to functions it means "use whatever the aligners default is". */
 	TEXALIGN_None,
diff --git a/Engine/Source/Editor/UnrealEd/Classes/ThumbnailRendering/ThumbnailManager.h b/Engine/Source/Editor/UnrealEd/Classes/ThumbnailRendering/ThumbnailManager.h
--- a/Engine/Source/Editor/UnrealEd/Classes/ThumbnailRendering/ThumbnailManager.h
+++ b/Engine/Source/Editor/UnrealEd/Classes/ThumbnailRendering/ThumbnailManager.h
@@ -24,7 +24,7 @@ class FAssetThumbnailPool;
  * Types of primitives for drawing thumbnails of resources.
  */
 UENUM()
-enum EThumbnailPrimType
+enum EThumbnailPrimType : int
 {
 	TPT_None,
 	TPT_Sphere,
diff --git a/Engine/Source/Editor/UnrealEd/Public/Tests/FbxAutomationCommon.h b/Engine/Source/Editor/UnrealEd/Public/Tests/FbxAutomationCommon.h
--- a/Engine/Source/Editor/UnrealEd/Public/Tests/FbxAutomationCommon.h
+++ b/Engine/Source/Editor/UnrealEd/Public/Tests/FbxAutomationCommon.h
@@ -18,7 +18,7 @@ class UFbxImportUI;
 
 /** Import mesh type */
 UENUM()
-enum EFBXExpectedResultPreset
+enum EFBXExpectedResultPreset : int
 {
 	/** Data should contain the number of error [int0]. */
 	Error_Number,
@@ -92,7 +92,7 @@ enum EFBXExpectedResultPreset
 
 /** Import mesh type */
 UENUM()
-enum EFBXTestPlanActionType
+enum EFBXTestPlanActionType : int
 {
 	/*Normal import*/
 	Import,
diff --git a/Engine/Source/Runtime/AIModule/Classes/AITypes.h b/Engine/Source/Runtime/AIModule/Classes/AITypes.h
--- a/Engine/Source/Runtime/AIModule/Classes/AITypes.h
+++ b/Engine/Source/Runtime/AIModule/Classes/AITypes.h
@@ -54,7 +54,7 @@ namespace FAISystem
 UENUM()
 namespace EAIOptionFlag
 {
-	enum Type
+	enum Type : int
 	{
 		Default,
 		Enable UMETA(DisplayName = "Yes"),
@@ -117,7 +117,7 @@ namespace EAILogicResuming
 UENUM()
 namespace EPawnActionAbortState
 {
-	enum Type
+	enum Type : int
 	{
 		NeverStarted,
 		NotBeingAborted,
@@ -133,7 +133,7 @@ namespace EPawnActionAbortState
 UENUM()
 namespace EPawnActionResult
 {
-	enum Type
+	enum Type : int
 	{
 		NotStarted,
 		InProgress,
@@ -146,7 +146,7 @@ namespace EPawnActionResult
 UENUM()
 namespace EPawnActionEventType
 {
-	enum Type
+	enum Type : int
 	{
 		Invalid,
 		FailedToStart,
@@ -160,7 +160,7 @@ namespace EPawnActionEventType
 UENUM()
 namespace EAIRequestPriority
 {
-	enum Type
+	enum Type : int
 	{
 		/** Actions requested by Level Designers by placing AI-hinting elements on the map. */
 		SoftScript,
@@ -185,7 +185,7 @@ namespace EAIRequestPriority
 UENUM()
 namespace EAILockSource
 {
-	enum Type
+	enum Type : int
 	{
 		Animation,
 		Logic,
diff --git a/Engine/Source/Runtime/AIModule/Classes/Actions/PawnAction.h b/Engine/Source/Runtime/AIModule/Classes/Actions/PawnAction.h
--- a/Engine/Source/Runtime/AIModule/Classes/Actions/PawnAction.h
+++ b/Engine/Source/Runtime/AIModule/Classes/Actions/PawnAction.h
@@ -19,7 +19,7 @@ struct FPawnActionStack;
 UENUM()
 namespace EPawnSubActionTriggeringPolicy
 {
-	enum Type
+	enum Type : int
 	{
 		CopyBeforeTriggering,
 		ReuseInstances,
@@ -32,7 +32,7 @@ DECLARE_DELEGATE_TwoParams(FPawnActionEventDelegate, UPawnAction&, EPawnActionEv
 UENUM()
 namespace EPawnActionFailHandling
 {
-	enum Type
+	enum Type : int
 	{
 		RequireSuccess,
 		IgnoreFailure
diff --git a/Engine/Source/Runtime/AIModule/Classes/EnvironmentQuery/EnvQueryTypes.h b/Engine/Source/Runtime/AIModule/Classes/EnvironmentQuery/EnvQueryTypes.h
--- a/Engine/Source/Runtime/AIModule/Classes/EnvironmentQuery/EnvQueryTypes.h
+++ b/Engine/Source/Runtime/AIModule/Classes/EnvironmentQuery/EnvQueryTypes.h
@@ -55,7 +55,7 @@ DECLARE_FLOAT_COUNTER_STAT_EXTERN(TEXT("Avg Instance Response Time (ms)"), STAT_
 UENUM()
 namespace EEnvTestPurpose
 {
-	enum Type
+	enum Type : int
 	{
 		Filter UMETA(DisplayName="Filter Only"),
 		Score UMETA(DisplayName="Score Only"),
@@ -66,7 +66,7 @@ namespace EEnvTestPurpose
 UENUM()
 namespace EEnvTestFilterType
 {
-	enum Type
+	enum Type : int
 	{
 		/** For numeric tests. */
 		Minimum,
@@ -82,7 +82,7 @@ namespace EEnvTestFilterType
 UENUM()
 namespace EEnvTestScoreEquation
 {
-	enum Type
+	enum Type : int
 	{
 		Linear,
 		Square,
@@ -107,7 +107,7 @@ namespace EEnvTestScoreEquation
 UENUM()
 namespace EEnvTestWeight
 {
-	enum Type
+	enum Type : int
 	{
 		None,
 		Square,
@@ -121,7 +121,7 @@ namespace EEnvTestWeight
 UENUM()
 namespace EEnvTestCost
 {
-	enum Type
+	enum Type : int
 	{
 		/** Reading data, math operations (e.g. distance). */
 		Low,
@@ -135,7 +135,7 @@ namespace EEnvTestCost
 UENUM()
 namespace EEnvTestFilterOperator
 {
-	enum Type
+	enum Type : int
 	{
 		AllPass			UMETA(Tooltip = "All contexts must pass condition"),
 		AnyPass			UMETA(Tooltip = "At least one context must pass condition"),
@@ -145,7 +145,7 @@ namespace EEnvTestFilterOperator
 UENUM()
 namespace EEnvTestScoreOperator
 {
-	enum Type
+	enum Type : int
 	{
 		AverageScore	UMETA(Tooltip = "Use average score from all contexts"),
 		MinScore		UMETA(Tooltip = "Use minimum score from all contexts"),
@@ -166,7 +166,7 @@ namespace EEnvItemStatus
 UENUM(BlueprintType)
 namespace EEnvQueryStatus
 {
-	enum Type
+	enum Type : int
 	{
 		Processing,
 		Success,
@@ -180,7 +180,7 @@ namespace EEnvQueryStatus
 UENUM()
 namespace EEnvQueryRunMode
 {
-	enum Type
+	enum Type : int
 	{
 		SingleResult	UMETA(Tooltip="Pick first item with the best score", DisplayName="Single Best Item"),
 		RandomBest5Pct	UMETA(Tooltip="Pick random item with score 95% .. 100% of max", DisplayName="Single Random Item from Best 5%"),
@@ -192,7 +192,7 @@ namespace EEnvQueryRunMode
 UENUM()
 namespace EEnvQueryParam
 {
-	enum Type
+	enum Type : int
 	{
 		Float,
 		Int,
@@ -212,7 +212,7 @@ enum class EAIParamType : uint8
 UENUM()
 namespace EEnvQueryTrace
 {
-	enum Type
+	enum Type : int
 	{
 		None,
 		Navigation,
@@ -225,7 +225,7 @@ namespace EEnvQueryTrace
 UENUM()
 namespace EEnvTraceShape
 {
-	enum Type
+	enum Type : int
 	{
 		Line,
 		Box,
@@ -237,7 +237,7 @@ namespace EEnvTraceShape
 UENUM()
 namespace EEnvOverlapShape
 {
-	enum Type
+	enum Type : int
 	{
 		Box,
 		Sphere,
@@ -248,7 +248,7 @@ namespace EEnvOverlapShape
 UENUM(meta=(ScriptName="EnvDirectionType"))
 namespace EEnvDirection
 {
-	enum Type
+	enum Type : int
 	{
 		TwoPoints	UMETA(DisplayName="Two Points",ToolTip="Direction from location of one context to another."),
 		Rotation	UMETA(ToolTip="Context's rotation will be used as a direction."),
@@ -258,7 +258,7 @@ namespace EEnvDirection
 UENUM()
 namespace EEnvQueryTestClamping
 {
-	enum Type
+	enum Type : int
 	{
 		None,			
 		/** Clamp to value specified in test. */
diff --git a/Engine/Source/Runtime/AnimGraphRuntime/Public/AnimNodes/AnimNode_RefPose.h b/Engine/Source/Runtime/AnimGraphRuntime/Public/AnimNodes/AnimNode_RefPose.h
--- a/Engine/Source/Runtime/AnimGraphRuntime/Public/AnimNodes/AnimNode_RefPose.h
+++ b/Engine/Source/Runtime/AnimGraphRuntime/Public/AnimNodes/AnimNode_RefPose.h
@@ -8,7 +8,7 @@
 #include "AnimNode_RefPose.generated.h"
 
 UENUM()
-enum ERefPoseType
+enum ERefPoseType : int
 {
 	EIT_LocalSpace, 
 	EIT_Additive
diff --git a/Engine/Source/Runtime/AnimGraphRuntime/Public/BoneControllers/AnimNode_ModifyBone.h b/Engine/Source/Runtime/AnimGraphRuntime/Public/BoneControllers/AnimNode_ModifyBone.h
--- a/Engine/Source/Runtime/AnimGraphRuntime/Public/BoneControllers/AnimNode_ModifyBone.h
+++ b/Engine/Source/Runtime/AnimGraphRuntime/Public/BoneControllers/AnimNode_ModifyBone.h
@@ -12,7 +12,7 @@
 class USkeletalMeshComponent;
 
 UENUM(BlueprintType)
-enum EBoneModificationMode
+enum EBoneModificationMode : int
 {
 	/** The modifier ignores this channel (keeps the existing bone translation, rotation, or scale). */
 	BMM_Ignore UMETA(DisplayName = "Ignore"),
diff --git a/Engine/Source/Runtime/Core/Public/Internationalization/Text.h b/Engine/Source/Runtime/Core/Public/Internationalization/Text.h
--- a/Engine/Source/Runtime/Core/Public/Internationalization/Text.h
+++ b/Engine/Source/Runtime/Core/Public/Internationalization/Text.h
@@ -122,7 +122,7 @@ CORE_API const TCHAR* LexToString(EDateTimeStyle::Type InValue);
 /** Redeclared in KismetTextLibrary for meta-data extraction purposes, be sure to update there as well */
 namespace EFormatArgumentType
 {
-	enum Type
+	enum Type : int
 	{
 		Int,
 		UInt,
@@ -143,7 +143,7 @@ typedef TSharedRef<const FTextFormatPatternDefinition, ESPMode::ThreadSafe> FTex
 typedef TSharedPtr<const FTextFormatPatternDefinition, ESPMode::ThreadSafe> FTextFormatPatternDefinitionConstPtr;
 
 /** Redeclared in KismetTextLibrary for meta-data extraction purposes, be sure to update there as well */
-enum ERoundingMode
+enum ERoundingMode : int
 {
 	/** Rounds to the nearest place, equidistant ties go to the value which is closest to an even value: 1.5 becomes 2, 0.5 becomes 0 */
 	HalfToEven,
diff --git a/Engine/Source/Runtime/CoreUObject/Public/UObject/CoreNetTypes.h b/Engine/Source/Runtime/CoreUObject/Public/UObject/CoreNetTypes.h
--- a/Engine/Source/Runtime/CoreUObject/Public/UObject/CoreNetTypes.h
+++ b/Engine/Source/Runtime/CoreUObject/Public/UObject/CoreNetTypes.h
@@ -13,7 +13,7 @@
 
 /** Secondary condition to check before considering the replication of a lifetime property. */
 UENUM(BlueprintType)
-enum ELifetimeCondition
+enum ELifetimeCondition : int
 {
 	COND_None = 0							UMETA(DisplayName = "None"),							// This property has no condition, and will send anytime it changes
 	COND_InitialOnly = 1					UMETA(DisplayName = "Initial Only"),					// This property will only attempt to send on the initial bunch
diff --git a/Engine/Source/Runtime/CoreUObject/Public/UObject/NoExportTypes.h b/Engine/Source/Runtime/CoreUObject/Public/UObject/NoExportTypes.h
--- a/Engine/Source/Runtime/CoreUObject/Public/UObject/NoExportTypes.h
+++ b/Engine/Source/Runtime/CoreUObject/Public/UObject/NoExportTypes.h
@@ -71,7 +71,7 @@
 UENUM()
 namespace ESearchCase
 {
-	enum Type
+	enum Type : int
 	{
 		CaseSensitive,
 		IgnoreCase,
@@ -85,7 +85,7 @@ namespace ESearchCase
 UENUM()
 namespace ESearchDir
 {
-	enum Type
+	enum Type : int
 	{
 		FromStart,
 		FromEnd,
@@ -99,7 +99,7 @@ namespace ESearchDir
 UENUM()
 namespace ELogTimes
 {
-	enum Type
+	enum Type : int
 	{
 		/** Do not display log timestamps. */
 		None UMETA(DisplayName = "None"),
@@ -119,7 +119,7 @@ namespace ELogTimes
 UENUM(meta=(ScriptName="AxisType"))
 namespace EAxis
 {
-	enum Type
+	enum Type : int
 	{
 		None,
 		X,
@@ -132,7 +132,7 @@ namespace EAxis
 UENUM()
 namespace EAxisList
 {
-	enum Type
+	enum Type : int
 	{
 		None = 0,
 		X = 1,
@@ -156,7 +156,7 @@ namespace EAxisList
 
 /** Describes shape of an interpolation curve (mirrored from InterpCurvePoint.h). */
 UENUM()
-enum EInterpCurveMode
+enum EInterpCurveMode : int
 {
 	/** A straight line between two keypoint values. */
 	CIM_Linear UMETA(DisplayName="Linear"),
@@ -186,7 +186,7 @@ enum EInterpCurveMode
  * @warning: The *Tools DLLs will also need to be recompiled if the ordering is changed, but should not need code changes.
  */
 UENUM()
-enum EPixelFormat
+enum EPixelFormat : int
 {
 	PF_Unknown,
 	PF_A32B32G32R32F,
@@ -298,7 +298,7 @@ enum EPixelFormat
 UENUM()
 namespace EMouseCursor
 {
-	enum Type
+	enum Type : int
 	{
 		/** Causes no mouse cursor to be visible. */
 		None,
@@ -411,7 +411,7 @@ enum class EPropertyAccessChangeNotifyMode : uint8
 UENUM(BlueprintType)
 namespace EAppReturnType
 {
-	enum Type
+	enum Type : int
 	{
 		No,
 		Yes,
@@ -434,7 +434,7 @@ namespace EAppMsgType
 	/**
 	 * Enumerates supported message dialog button types.
 	 */
-	enum Type
+	enum Type : int
 	{
 		Ok,
 		YesNo,
@@ -2164,7 +2164,7 @@ struct FPrimaryAssetId
 UENUM(BlueprintType)
 namespace ERangeBoundTypes
 {
-	enum Type
+	enum Type : int
 	{
 		/**
 		* The range excludes the bound.
diff --git a/Engine/Source/Runtime/Engine/Classes/Animation/AnimEnums.h b/Engine/Source/Runtime/Engine/Classes/Animation/AnimEnums.h
--- a/Engine/Source/Runtime/Engine/Classes/Animation/AnimEnums.h
+++ b/Engine/Source/Runtime/Engine/Classes/Animation/AnimEnums.h
@@ -10,7 +10,7 @@ UENUM()
 /** Root Bone Lock options when extracting Root Motion. */
 namespace ERootMotionRootLock
 {
-	enum Type
+	enum Type : int
 	{
 		/** Use reference pose root bone position. */
 		RefPose,
@@ -26,7 +26,7 @@ namespace ERootMotionRootLock
 UENUM()
 namespace ERootMotionMode
 {
-	enum Type
+	enum Type : int
 	{
 		/** Leave root motion in animation. */
 		NoRootMotionExtraction,
@@ -46,7 +46,7 @@ namespace ERootMotionMode
 * For an additive animation, indicates what the animation is relative to.
 */
 UENUM()
-enum EAdditiveBasePoseType
+enum EAdditiveBasePoseType : int
 {
 	/** Will be deprecated. */
 	ABPT_None UMETA(DisplayName = "None"),
diff --git a/Engine/Source/Runtime/Engine/Classes/Animation/AnimNode_TransitionPoseEvaluator.h b/Engine/Source/Runtime/Engine/Classes/Animation/AnimNode_TransitionPoseEvaluator.h
--- a/Engine/Source/Runtime/Engine/Classes/Animation/AnimNode_TransitionPoseEvaluator.h
+++ b/Engine/Source/Runtime/Engine/Classes/Animation/AnimNode_TransitionPoseEvaluator.h
@@ -13,7 +13,7 @@
 UENUM()
 namespace EEvaluatorDataSource
 {
-	enum Type
+	enum Type : int
 	{
 		EDS_SourcePose UMETA(DisplayName="Source Pose"),
 		EDS_DestinationPose UMETA(DisplayName="Destination Pose")
@@ -24,7 +24,7 @@ namespace EEvaluatorDataSource
 UENUM()
 namespace EEvaluatorMode
 {
-	enum Mode
+	enum Mode : int
 	{
 		/** DataSource is ticked and evaluated every frame. */
 		EM_Standard UMETA(DisplayName="Standard"),
diff --git a/Engine/Source/Runtime/Engine/Classes/Animation/AnimSequenceBase.h b/Engine/Source/Runtime/Engine/Classes/Animation/AnimSequenceBase.h
--- a/Engine/Source/Runtime/Engine/Classes/Animation/AnimSequenceBase.h
+++ b/Engine/Source/Runtime/Engine/Classes/Animation/AnimSequenceBase.h
@@ -19,7 +19,7 @@
 #include "AnimSequenceBase.generated.h"
 
 UENUM()
-enum ETypeAdvanceAnim
+enum ETypeAdvanceAnim : int
 {
 	ETAA_Default,
 	ETAA_Finished,
diff --git a/Engine/Source/Runtime/Engine/Classes/Animation/BlendSpace.h b/Engine/Source/Runtime/Engine/Classes/Animation/BlendSpace.h
--- a/Engine/Source/Runtime/Engine/Classes/Animation/BlendSpace.h
+++ b/Engine/Source/Runtime/Engine/Classes/Animation/BlendSpace.h
@@ -41,7 +41,7 @@ public:
 
 /** Interpolation data types. */
 UENUM()
-enum EBlendSpaceAxis
+enum EBlendSpaceAxis : int
 {
 	BSA_None UMETA(DisplayName = "None"),
 	BSA_X UMETA(DisplayName = "Horizontal (X) Axis"),
@@ -435,7 +435,7 @@ struct FPerBoneInterpolation
 UENUM()
 namespace ENotifyTriggerMode
 {
-	enum Type
+	enum Type : int
 	{
 		AllAnimations UMETA(DisplayName="All Animations"),
 		HighestWeightedAnimation UMETA(DisplayName="Highest Weighted Animation"),
diff --git a/Engine/Source/Runtime/Engine/Classes/Camera/PlayerCameraManager.h b/Engine/Source/Runtime/Engine/Classes/Camera/PlayerCameraManager.h
--- a/Engine/Source/Runtime/Engine/Classes/Camera/PlayerCameraManager.h
+++ b/Engine/Source/Runtime/Engine/Classes/Camera/PlayerCameraManager.h
@@ -26,7 +26,7 @@ DECLARE_DYNAMIC_MULTICAST_DELEGATE_TwoParams(FOnAudioFadeChangeSignature, bool,
  * @see FViewTargetTransitionParams, SetViewTarget 
  */
 UENUM()
-enum EViewTargetBlendFunction
+enum EViewTargetBlendFunction : int
 {
 	/** Camera does a simple linear interpolation. */
 	VTBlend_Linear,
diff --git a/Engine/Source/Runtime/Engine/Classes/Components/PrimitiveComponent.h b/Engine/Source/Runtime/Engine/Classes/Components/PrimitiveComponent.h
--- a/Engine/Source/Runtime/Engine/Classes/Components/PrimitiveComponent.h
+++ b/Engine/Source/Runtime/Engine/Classes/Components/PrimitiveComponent.h
@@ -57,7 +57,7 @@ namespace PrimitiveComponentCVars
 
 /** Determines whether a Character can attempt to step up onto a component when they walk in to it. */
 UENUM()
-enum ECanBeCharacterBase
+enum ECanBeCharacterBase : int
 {
 	/** Character cannot step up onto this Component. */
 	ECB_No UMETA(DisplayName="No"),
@@ -75,7 +75,7 @@ enum ECanBeCharacterBase
 UENUM()
 namespace EHasCustomNavigableGeometry
 {
-	enum Type
+	enum Type : int
 	{
 		/** Primitive doesn't have custom navigation geometry, if collision is enabled then its convex/trimesh collision will be used for generating the navmesh */
 		No,
diff --git a/Engine/Source/Runtime/Engine/Classes/Components/SceneComponent.h b/Engine/Source/Runtime/Engine/Classes/Components/SceneComponent.h
--- a/Engine/Source/Runtime/Engine/Classes/Components/SceneComponent.h
+++ b/Engine/Source/Runtime/Engine/Classes/Components/SceneComponent.h
@@ -33,7 +33,7 @@ namespace ERHIFeatureLevel { enum Type : int; }
 
 /** Detail mode for scene component rendering, corresponds with the integer value of UWorld::GetDetailMode() */
 UENUM()
-enum EDetailMode
+enum EDetailMode : int
 {
 	DM_Low UMETA(DisplayName="Low"),
 	DM_Medium UMETA(DisplayName="Medium"),
@@ -43,7 +43,7 @@ enum EDetailMode
 
 /** The space for the transform */
 UENUM()
-enum ERelativeTransformSpace
+enum ERelativeTransformSpace : int
 {
 	/** World space transform. */
 	RTS_World,
diff --git a/Engine/Source/Runtime/Engine/Classes/Components/SkinnedMeshComponent.h b/Engine/Source/Runtime/Engine/Classes/Components/SkinnedMeshComponent.h
--- a/Engine/Source/Runtime/Engine/Classes/Components/SkinnedMeshComponent.h
+++ b/Engine/Source/Runtime/Engine/Classes/Components/SkinnedMeshComponent.h
@@ -49,7 +49,7 @@ DECLARE_MULTICAST_DELEGATE_ThreeParams(FOnTickPose, USkinnedMeshComponent* /*Ski
 
 /** The valid BoneVisibilityStates values; A bone is only visible if it is *exactly* 1 */
 UENUM()
-enum EBoneVisibilityStatus
+enum EBoneVisibilityStatus : int
 {
 	/** Bone is hidden because it's parent is hidden. */
 	BVS_HiddenByParent,
@@ -62,7 +62,7 @@ enum EBoneVisibilityStatus
 
 /** PhysicsBody options when bone is hidden */
 UENUM()
-enum EPhysBodyOp
+enum EPhysBodyOp : int
 {
 	/** Don't do anything. */
 	PBO_None,
@@ -108,7 +108,7 @@ namespace EMeshComponentUpdateFlag
 UENUM()
 namespace EBoneSpaces
 {
-	enum Type
+	enum Type : int
 	{
 		/** Set absolute position of bone in world space. */
 		WorldSpace		UMETA(DisplayName = "World Space"),
diff --git a/Engine/Source/Runtime/Engine/Classes/Components/SkyLightComponent.h b/Engine/Source/Runtime/Engine/Classes/Components/SkyLightComponent.h
--- a/Engine/Source/Runtime/Engine/Classes/Components/SkyLightComponent.h
+++ b/Engine/Source/Runtime/Engine/Classes/Components/SkyLightComponent.h
@@ -79,7 +79,7 @@ private:
 };
 
 UENUM()
-enum ESkyLightSourceType
+enum ESkyLightSourceType : int
 {
 	/** Construct the sky light from the captured scene, anything further than SkyDistanceThreshold from the sky light position will be included. */
 	SLS_CapturedScene,
diff --git a/Engine/Source/Runtime/Engine/Classes/Components/StereoLayerComponent.h b/Engine/Source/Runtime/Engine/Classes/Components/StereoLayerComponent.h
--- a/Engine/Source/Runtime/Engine/Classes/Components/StereoLayerComponent.h
+++ b/Engine/Source/Runtime/Engine/Classes/Components/StereoLayerComponent.h
@@ -12,7 +12,7 @@ class UTexture;
 
 /** Used by IStereoLayer */
 UENUM()
-enum EStereoLayerType
+enum EStereoLayerType : int
 {
 	/** Location within the world */
 	SLT_WorldLocked		UMETA(DisplayName = "World Locked"),
@@ -28,7 +28,7 @@ enum EStereoLayerType
 
 /** The shape to use for the stereo layer.  Note that some shapes might not be supported on all platforms! */
 UENUM()
-enum EStereoLayerShape
+enum EStereoLayerShape : int
 {
 	/** Quad layer */
 	SLSH_QuadLayer		UMETA(DisplayName = "Quad Layer"),
diff --git a/Engine/Source/Runtime/Engine/Classes/Components/TextRenderComponent.h b/Engine/Source/Runtime/Engine/Classes/Components/TextRenderComponent.h
--- a/Engine/Source/Runtime/Engine/Classes/Components/TextRenderComponent.h
+++ b/Engine/Source/Runtime/Engine/Classes/Components/TextRenderComponent.h
@@ -14,7 +14,7 @@ class UFont;
 class UMaterialInterface;
 
 UENUM()
-enum EHorizTextAligment
+enum EHorizTextAligment : int
 {
 	/**  */
 	EHTA_Left UMETA(DisplayName="Left"),
@@ -25,7 +25,7 @@ enum EHorizTextAligment
 };
 
 UENUM()
-enum EVerticalTextAligment
+enum EVerticalTextAligment : int
 {
 	/**  */
 	EVRTA_TextTop UMETA(DisplayName="Text Top"),
diff --git a/Engine/Source/Runtime/Engine/Classes/Components/TimelineComponent.h b/Engine/Source/Runtime/Engine/Classes/Components/TimelineComponent.h
--- a/Engine/Source/Runtime/Engine/Classes/Components/TimelineComponent.h
+++ b/Engine/Source/Runtime/Engine/Classes/Components/TimelineComponent.h
@@ -31,7 +31,7 @@ DECLARE_DELEGATE_OneParam( FOnTimelineLinearColorStatic, FLinearColor );
 
 /** Whether or not the timeline should be finished after the specified length, or the last keyframe in the tracks */
 UENUM()
-enum ETimelineLengthMode
+enum ETimelineLengthMode : int
 {
 	TL_TimelineLength,
 	TL_LastKeyFrame
@@ -41,7 +41,7 @@ enum ETimelineLengthMode
 UENUM(BlueprintType)
 namespace ETimelineDirection
 {
-	enum Type
+	enum Type : int
 	{
 		Forward,
 		Backward,
diff --git a/Engine/Source/Runtime/Engine/Classes/Curves/RealCurve.h b/Engine/Source/Runtime/Engine/Classes/Curves/RealCurve.h
--- a/Engine/Source/Runtime/Engine/Classes/Curves/RealCurve.h
+++ b/Engine/Source/Runtime/Engine/Classes/Curves/RealCurve.h
@@ -9,7 +9,7 @@
 
 /** Method of interpolation between this key and the next. */
 UENUM(BlueprintType)
-enum ERichCurveInterpMode
+enum ERichCurveInterpMode : int
 {
 	/** Use linear interpolation between values. */
 	RCIM_Linear UMETA(DisplayName = "Linear"),
@@ -23,7 +23,7 @@ enum ERichCurveInterpMode
 
 /** Enumerates extrapolation options. */
 UENUM(BlueprintType)
-enum ERichCurveExtrapolation
+enum ERichCurveExtrapolation : int
 {
 	/** Repeat the curve without an offset. */
 	RCCE_Cycle UMETA(DisplayName = "Cycle"),
diff --git a/Engine/Source/Runtime/Engine/Classes/Curves/RichCurve.h b/Engine/Source/Runtime/Engine/Classes/Curves/RichCurve.h
--- a/Engine/Source/Runtime/Engine/Classes/Curves/RichCurve.h
+++ b/Engine/Source/Runtime/Engine/Classes/Curves/RichCurve.h
@@ -11,7 +11,7 @@
 
 /** If using RCIM_Cubic, this enum describes how the tangents should be controlled in editor. */
 UENUM(BlueprintType)
-enum ERichCurveTangentMode
+enum ERichCurveTangentMode : int
 {
 	/** Automatically calculates tangents to create smooth curves between values. */
 	RCTM_Auto UMETA(DisplayName="Auto"),
@@ -26,7 +26,7 @@ enum ERichCurveTangentMode
 
 /** Enumerates tangent weight modes. */
 UENUM(BlueprintType)
-enum ERichCurveTangentWeightMode
+enum ERichCurveTangentWeightMode : int
 {
 	/** Don't take tangent weights into account. */
 	RCTWM_WeightedNone UMETA(DisplayName="None"),
@@ -40,7 +40,7 @@ enum ERichCurveTangentWeightMode
 
 /** Enumerates curve compression options. */
 UENUM()
-enum ERichCurveCompressionFormat
+enum ERichCurveCompressionFormat : int
 {
 	/** No keys are present */
 	RCCF_Empty UMETA(DisplayName = "Empty"),
@@ -63,7 +63,7 @@ enum ERichCurveCompressionFormat
 
 /** Enumerates key time compression options. */
 UENUM()
-enum ERichCurveKeyTimeCompressionFormat
+enum ERichCurveKeyTimeCompressionFormat : int
 {
 	/** Key time is quantized to 16 bits */
 	RCKTCF_uint16 UMETA(DisplayName = "uint16"),
diff --git a/Engine/Source/Runtime/Engine/Classes/Distributions/Distribution.h b/Engine/Source/Runtime/Engine/Classes/Distributions/Distribution.h
--- a/Engine/Source/Runtime/Engine/Classes/Distributions/Distribution.h
+++ b/Engine/Source/Runtime/Engine/Classes/Distributions/Distribution.h
@@ -12,7 +12,7 @@
 extern ENGINE_API uint32 GDistributionType;
 
 UENUM()
-enum DistributionParamMode
+enum DistributionParamMode : int
 {
 	DPM_Normal,
 	DPM_Abs,
diff --git a/Engine/Source/Runtime/Engine/Classes/Distributions/DistributionVector.h b/Engine/Source/Runtime/Engine/Classes/Distributions/DistributionVector.h
--- a/Engine/Source/Runtime/Engine/Classes/Distributions/DistributionVector.h
+++ b/Engine/Source/Runtime/Engine/Classes/Distributions/DistributionVector.h
@@ -13,7 +13,7 @@
 struct FPropertyChangedEvent;
 
 UENUM()
-enum EDistributionVectorLockFlags
+enum EDistributionVectorLockFlags : int
 {
 	EDVLF_None UMETA(DisplayName="None"),
 	EDVLF_XY UMETA(DisplayName="XY"),
@@ -24,7 +24,7 @@ enum EDistributionVectorLockFlags
 };
 
 UENUM()
-enum EDistributionVectorMirrorFlags
+enum EDistributionVectorMirrorFlags : int
 {
 	EDVMF_Same UMETA(DisplayName="Same"),
 	EDVMF_Different UMETA(DisplayName="Different"),
diff --git a/Engine/Source/Runtime/Engine/Classes/EdGraph/EdGraphNode.h b/Engine/Source/Runtime/Engine/Classes/EdGraph/EdGraphNode.h
--- a/Engine/Source/Runtime/Engine/Classes/EdGraph/EdGraphNode.h
+++ b/Engine/Source/Runtime/Engine/Classes/EdGraph/EdGraphNode.h
@@ -93,7 +93,7 @@ struct ENGINE_API FEdGraphTerminalType
 
 /** Enum used to define which way data flows into or out of this pin. */
 UENUM()
-enum EEdGraphPinDirection
+enum EEdGraphPinDirection : int
 {
 	EGPD_Input,
 	EGPD_Output,
@@ -130,7 +130,7 @@ enum class EPinContainerType : uint8
 UENUM()
 namespace ENodeTitleType
 {
-	enum Type
+	enum Type : int
 	{
 		/** The full title, may be multiple lines. */
 		FullTitle,
@@ -149,7 +149,7 @@ namespace ENodeTitleType
 UENUM()
 namespace ENodeAdvancedPins
 {
-	enum Type
+	enum Type : int
 	{
 		/** No advanced pins. */
 		NoPins,
diff --git a/Engine/Source/Runtime/Engine/Classes/EdGraph/EdGraphPin.h b/Engine/Source/Runtime/Engine/Classes/EdGraph/EdGraphPin.h
--- a/Engine/Source/Runtime/Engine/Classes/EdGraph/EdGraphPin.h
+++ b/Engine/Source/Runtime/Engine/Classes/EdGraph/EdGraphPin.h
@@ -235,7 +235,7 @@ struct TStructOpsTypeTraits< FEdGraphPinType > : public TStructOpsTypeTraitsBase
 };
 
 UENUM()
-enum EBlueprintPinStyleType
+enum EBlueprintPinStyleType : int
 {
 	BPST_Original UMETA(DisplayName="Circles, Grid, Diamond"),
 	BPST_VariantA UMETA(DisplayName="Directional Circles")
diff --git a/Engine/Source/Runtime/Engine/Classes/EdGraph/EdGraphSchema.h b/Engine/Source/Runtime/Engine/Classes/EdGraph/EdGraphSchema.h
--- a/Engine/Source/Runtime/Engine/Classes/EdGraph/EdGraphSchema.h
+++ b/Engine/Source/Runtime/Engine/Classes/EdGraph/EdGraphSchema.h
@@ -22,7 +22,7 @@ struct FBPVariableDescription;
 
 /** Distinguishes between different graph types. Graphs can have different properties; for example: functions have one entry point, ubergraphs can have multiples. */
 UENUM()
-enum EGraphType
+enum EGraphType : int
 {
 	GT_Function,
 	GT_Ubergraph,
@@ -34,7 +34,7 @@ enum EGraphType
 
 /** This is the type of response the graph editor should take when making a connection */
 UENUM()
-enum ECanCreateConnectionResponse
+enum ECanCreateConnectionResponse : int
 {
 	/** Make the connection; there are no issues (message string is displayed if not empty). */
 	CONNECT_RESPONSE_MAKE,
diff --git a/Engine/Source/Runtime/Engine/Classes/Engine/BlendableInterface.h b/Engine/Source/Runtime/Engine/Classes/Engine/BlendableInterface.h
--- a/Engine/Source/Runtime/Engine/Classes/Engine/BlendableInterface.h
+++ b/Engine/Source/Runtime/Engine/Classes/Engine/BlendableInterface.h
@@ -12,7 +12,7 @@ class UMaterialInterface;
 
 /** Where to place a material node in the post processing graph. */
 UENUM()
-enum EBlendableLocation
+enum EBlendableLocation : int
 {
 	/** Input0:former pass color, Input1:SeparateTranslucency. */
 	BL_AfterTonemapping UMETA(DisplayName="After Tonemapping"),
diff --git a/Engine/Source/Runtime/Engine/Classes/Engine/Blueprint.h b/Engine/Source/Runtime/Engine/Classes/Engine/Blueprint.h
--- a/Engine/Source/Runtime/Engine/Classes/Engine/Blueprint.h
+++ b/Engine/Source/Runtime/Engine/Classes/Engine/Blueprint.h
@@ -35,7 +35,7 @@ struct FDiffResults;
  * Enumerates states a blueprint can be in.
  */
 UENUM()
-enum EBlueprintStatus
+enum EBlueprintStatus : int
 {
 	/** Blueprint is in an unknown state. */
 	BS_Unknown,
@@ -55,7 +55,7 @@ enum EBlueprintStatus
 
 /** Enumerates types of blueprints. */
 UENUM()
-enum EBlueprintType
+enum EBlueprintType : int
 {
 	/** Normal blueprint. */
 	BPTYPE_Normal				UMETA(DisplayName="Blueprint Class"),
diff --git a/Engine/Source/Runtime/Engine/Classes/Engine/Brush.h b/Engine/Source/Runtime/Engine/Classes/Engine/Brush.h
--- a/Engine/Source/Runtime/Engine/Classes/Engine/Brush.h
+++ b/Engine/Source/Runtime/Engine/Classes/Engine/Brush.h
@@ -13,7 +13,7 @@ class ULevel;
 //-----------------------------------------------------------------------------
 // Variables.
 UENUM()
-enum ECsgOper
+enum ECsgOper : int
 {
 	/** Active brush. (deprecated, do not use.) */
 	CSG_Active,
@@ -31,7 +31,7 @@ enum ECsgOper
 
 
 UENUM()
-enum EBrushType
+enum EBrushType : int
 {
 	/** Default/builder brush. */
 	Brush_Default UMETA(Hidden),
diff --git a/Engine/Source/Runtime/Engine/Classes/Engine/CloudStorageBase.h b/Engine/Source/Runtime/Engine/Classes/Engine/CloudStorageBase.h
--- a/Engine/Source/Runtime/Engine/Classes/Engine/CloudStorageBase.h
+++ b/Engine/Source/Runtime/Engine/Classes/Engine/CloudStorageBase.h
@@ -10,7 +10,7 @@
 
 /** All the types of delegate callbacks that a CloudStorage subclass may receive from C++ */
 UENUM()
-enum ECloudStorageDelegate
+enum ECloudStorageDelegate : int
 {
 	// @todo: Fill in the result descriptions for these (and the other PI subclasses)
 	CSD_KeyValueReadComplete,
diff --git a/Engine/Source/Runtime/Engine/Classes/Engine/Engine.h b/Engine/Source/Runtime/Engine/Classes/Engine/Engine.h
--- a/Engine/Source/Runtime/Engine/Classes/Engine/Engine.h
+++ b/Engine/Source/Runtime/Engine/Classes/Engine/Engine.h
@@ -80,7 +80,7 @@ enum class EGetWorldErrorMode
  * Enumerates types of fully loaded packages.
  */
 UENUM()
-enum EFullyLoadPackageType
+enum EFullyLoadPackageType : int
 {
 	/** Load the packages when the map in Tag is loaded. */
 	FULLYLOAD_Map,
diff --git a/Engine/Source/Runtime/Engine/Classes/Engine/EngineBaseTypes.h b/Engine/Source/Runtime/Engine/Classes/Engine/EngineBaseTypes.h
--- a/Engine/Source/Runtime/Engine/Classes/Engine/EngineBaseTypes.h
+++ b/Engine/Source/Runtime/Engine/Classes/Engine/EngineBaseTypes.h
@@ -30,7 +30,7 @@ struct FTickContext;
 //	EInputEvent
 //
 UENUM( BlueprintType, meta=(ScriptName="InputEventType"))
-enum EInputEvent
+enum EInputEvent : int
 {
 	IE_Pressed              =0,
 	IE_Released             =1,
@@ -83,7 +83,7 @@ enum ELevelTick
 
 /** Determines which ticking group a tick function belongs to. */
 UENUM(BlueprintType)
-enum ETickingGroup
+enum ETickingGroup : int
 {
 	/** Any item that needs to be executed before physics simulation starts. */
 	TG_PrePhysics UMETA(DisplayName="Pre Physics"),
@@ -523,7 +523,7 @@ struct TStructOpsTypeTraits<FActorComponentTickFunction> : public TStructOpsType
 UENUM()
 namespace ENetworkLagState
 {
-	enum Type
+	enum Type : int
 	{
 		/** The net driver is operating normally or it is not possible to tell if it is lagging */
 		NotLagging,
@@ -552,7 +552,7 @@ namespace ENetworkLagState
 UENUM(BlueprintType)
 namespace ETravelFailure
 {
-	enum Type
+	enum Type : int
 	{
 		/** No level found in the loaded package */
 		NoLevel,
@@ -618,7 +618,7 @@ namespace ETravelFailure
 
 // Traveling from server to server.
 UENUM()
-enum ETravelType
+enum ETravelType : int
 {
 	/** Absolute URL. */
 	TRAVEL_Absolute,
@@ -633,7 +633,7 @@ enum ETravelType
 UENUM(BlueprintType)
 namespace EDemoPlayFailure
 {
-	enum UE_DEPRECATED(5.1, "No longer used in favor of EReplayResult") Type
+	enum UE_DEPRECATED(5.1, "No longer used in favor of EReplayResult") Type : int
 	{
 		/** A Generic failure. */
 		Generic,
@@ -854,7 +854,7 @@ enum ENetMode
  * Don't change the order, the ID is serialized with the editor
  */
 UENUM()
-enum EViewModeIndex
+enum EViewModeIndex : int
 {
 	/** Wireframe w/ brushes. */
 	VMI_BrushWireframe = 0 UMETA(DisplayName = "Wireframe"),
diff --git a/Engine/Source/Runtime/Engine/Classes/Engine/EngineTypes.h b/Engine/Source/Runtime/Engine/Classes/Engine/EngineTypes.h
--- a/Engine/Source/Runtime/Engine/Classes/Engine/EngineTypes.h
+++ b/Engine/Source/Runtime/Engine/Classes/Engine/EngineTypes.h
@@ -39,7 +39,7 @@ enum { NumInlinedActorComponents = 24 };
 
 /** Enum describing how to constrain perspective view port FOV */
 UENUM()
-enum EAspectRatioAxisConstraint
+enum EAspectRatioAxisConstraint : int
 {
 	AspectRatio_MaintainYFOV UMETA(DisplayName="Maintain Y-Axis FOV"),
 	AspectRatio_MaintainXFOV UMETA(DisplayName="Maintain X-Axis FOV"),
@@ -167,7 +167,7 @@ struct ENGINE_API FDetachmentTransformRules
 UENUM()
 namespace EAttachLocation
 {
-	enum Type
+	enum Type : int
 	{
 		/** Keeps current relative transform as the relative transform to the new parent. */
 		KeepRelativeOffset,
@@ -188,7 +188,7 @@ namespace EAttachLocation
  * Elements with higher priority occlude elements with lower priority, disregarding distance.
  */
 UENUM()
-enum ESceneDepthPriorityGroup
+enum ESceneDepthPriorityGroup : int
 {
 	/** World scene DPG. */
 	SDPG_World,
@@ -199,7 +199,7 @@ enum ESceneDepthPriorityGroup
 
 /** Quality of indirect lighting for Movable primitives. This has a large effect on Indirect Lighting Cache update time. */
 UENUM()
-enum EIndirectLightingCacheQuality
+enum EIndirectLightingCacheQuality : int
 {
 	/** The indirect lighting cache will be disabled for this object, so no GI from stationary lights on movable objects. */
 	ILCQ_Off,
@@ -227,7 +227,7 @@ enum class ELightmapType : uint8
 
 /** Controls how occlusion from Distance Field Ambient Occlusion is combined with Screen Space Ambient Occlusion. */
 UENUM()
-enum EOcclusionCombineMode
+enum EOcclusionCombineMode : int
 {
 	/** Take the minimum occlusion value.  This is effective for avoiding over-occlusion from multiple methods, but can result in indoors looking too flat. */
 	OCM_Minimum,
@@ -246,7 +246,7 @@ enum EOcclusionCombineMode
  * @warning Check UMaterialInstance::Serialize if changed!!
  */
 UENUM(BlueprintType)
-enum EBlendMode
+enum EBlendMode : int
 {
 	BLEND_Opaque UMETA(DisplayName="Opaque"),
 	BLEND_Masked UMETA(DisplayName="Masked"),
@@ -262,7 +262,7 @@ enum EBlendMode
  * The blending mode for Strata materials
  */
 UENUM(BlueprintType)
-enum EStrataBlendMode
+enum EStrataBlendMode : int
 {
 	SBM_Opaque UMETA(DisplayName = "Opaque"),
 	SBM_Masked UMETA(DisplayName = "Masked"),
@@ -275,7 +275,7 @@ enum EStrataBlendMode
 
 /** The default float precision for material's pixel shaders on mobile devices*/
 UENUM()
-enum EMaterialFloatPrecisionMode
+enum EMaterialFloatPrecisionMode : int
 {
 	/** Uses project based precision mode setting */
 	MFPM_Default UMETA(DisplayName = "Default"),
@@ -290,7 +290,7 @@ enum EMaterialFloatPrecisionMode
 
 /** Controls where the sampler for different texture lookups comes from */
 UENUM()
-enum ESamplerSourceMode
+enum ESamplerSourceMode : int
 {
 	/** Get the sampler from the texture.  Every unique texture will consume a sampler slot, which are limited in number. */
 	SSM_FromTextureAsset UMETA(DisplayName="From texture asset"),
@@ -304,7 +304,7 @@ enum ESamplerSourceMode
 
 /** defines how MipValue is used */
 UENUM()
-enum ETextureMipValueMode
+enum ETextureMipValueMode : int
 {
 	/* Use hardware computed sample's mip level with automatic anisotropic filtering support. */
 	TMVM_None UMETA(DisplayName="None (use computed mip level)"),
@@ -323,7 +323,7 @@ enum ETextureMipValueMode
 
 /** Describes how to handle lighting of translucent objets */
 UENUM()
-enum ETranslucencyLightingMode
+enum ETranslucencyLightingMode : int
 {
 	/** 
 	 * Lighting will be calculated for a volume, without directionality.  Use this on particle effects like smoke and dust.
@@ -368,7 +368,7 @@ enum ETranslucencyLightingMode
 
 /** Determines how the refraction offset should be computed for the material. */
 UENUM()
-enum ERefractionMode
+enum ERefractionMode : int
 {
 	/** 
 	 * By default, when the root node refraction pin is unplugged, relies on the material IOR evaluated from F0.
@@ -405,7 +405,7 @@ enum ERefractionMode
 UENUM()
 namespace ETranslucentSortPolicy
 {
-	enum Type
+	enum Type : int
 	{
 		/** Sort based on distance from camera centerpoint to bounding sphere centerpoint. (Default, best for 3D games.) */
 		SortByDistance = 0,
@@ -422,7 +422,7 @@ namespace ETranslucentSortPolicy
 UENUM()
 namespace EDynamicGlobalIlluminationMethod
 {
-	enum Type
+	enum Type : int
 	{
 		/** No dynamic Global Illumination method will be used. Global Illumination can still be baked into lightmaps. */
 		None, 
@@ -445,7 +445,7 @@ namespace EDynamicGlobalIlluminationMethod
 UENUM()
 namespace EReflectionMethod
 {
-	enum Type
+	enum Type : int
 	{
 		/** No global reflection method will be used. Reflections can still come from Reflection Captures, Planar Reflections or a Skylight placed in the level. */
 		None, 
@@ -465,7 +465,7 @@ namespace EReflectionMethod
 UENUM()
 namespace EShadowMapMethod
 {
-	enum Type
+	enum Type : int
 	{
 		/** Render geometry into shadow depth maps for shadowing.  Requires manual setup of shadowing distances and only culls per-component, causing poor performance with high poly scenes.  Required to enable stationary baked shadows (but which is incompatible with Nanite geometry). */
 		ShadowMaps UMETA(DisplayName = "Shadow Maps"),
@@ -479,7 +479,7 @@ namespace EShadowMapMethod
 UENUM()
 namespace ECastRayTracedShadow 
 {
-	enum Type
+	enum Type : int
 	{
 		/** Ray traced shadows disabled for this light */
 		Disabled,
@@ -492,7 +492,7 @@ namespace ECastRayTracedShadow
 
 /** Specifies which component of the scene rendering should be output to the final render target. */
 UENUM()
-enum ESceneCaptureSource 
+enum ESceneCaptureSource : int
 { 
 	SCS_SceneColorHDR UMETA(DisplayName="SceneColor (HDR) in RGB, Inv Opacity in A"),
 	SCS_SceneColorHDRNoAlpha UMETA(DisplayName="SceneColor (HDR) in RGB, 0 in A"),
@@ -510,7 +510,7 @@ enum ESceneCaptureSource
 
 /** Specifies how scene captures are composited into render buffers */
 UENUM()
-enum ESceneCaptureCompositeMode
+enum ESceneCaptureCompositeMode : int
 { 
 	SCCM_Overwrite UMETA(DisplayName="Overwrite"),
 	SCCM_Additive UMETA(DisplayName="Additive"),
@@ -571,7 +571,7 @@ inline int32 GetFirstLightingChannelFromMask(uint8 Mask)
 UENUM()
 namespace EGBufferFormat
 {
-	enum Type
+	enum Type : int
 	{
 		/** Forces all GBuffers to 8 bits per channel. Intended as profiling for best performance. (Strata: Octahedral encoding as 2x11bits for simple and single materials, 2x16bits for complex materials) */
 		Force8BitsPerChannel = 0 UMETA(DisplayName = "Force 8 Bits Per Channel"),
@@ -586,7 +586,7 @@ namespace EGBufferFormat
 
 /** Controls the way that the width scale property affects animation trails. */
 UENUM()
-enum ETrailWidthMode
+enum ETrailWidthMode : int
 {
 	ETrailWidthMode_FromCentre UMETA(DisplayName = "From Centre"),
 	ETrailWidthMode_FromFirst UMETA(DisplayName = "From First Socket"),
@@ -597,7 +597,7 @@ enum ETrailWidthMode
 UENUM()
 namespace EParticleCollisionMode
 {
-	enum Type
+	enum Type : int
 	{
 		SceneDepth UMETA(DisplayName="Scene Depth"),
 		DistanceField UMETA(DisplayName="Distance Field")
@@ -609,7 +609,7 @@ namespace EParticleCollisionMode
  * @warning Check UMaterialInstance::Serialize if changed!
  */
 UENUM()
-enum EMaterialShadingModel
+enum EMaterialShadingModel : int
 {
 	MSM_Unlit					UMETA(DisplayName="Unlit"),
 	MSM_DefaultLit				UMETA(DisplayName="Default Lit"),
@@ -683,7 +683,7 @@ private:
  * Specifies the Strata runtime shading model summarized from the material graph
  */
 UENUM()
-enum EStrataShadingModel
+enum EStrataShadingModel : int
 {
 	SSM_Unlit					UMETA(DisplayName = "Unlit"),
 	SSM_DefaultLit				UMETA(DisplayName = "DefaultLit"),
@@ -754,7 +754,7 @@ private:
 
 /** Describes how textures are sampled for materials */
 UENUM(BlueprintType)
-enum EMaterialSamplerType
+enum EMaterialSamplerType : int
 {
 	SAMPLERTYPE_Color UMETA(DisplayName="Color"),
 	SAMPLERTYPE_Grayscale UMETA(DisplayName="Grayscale"),
@@ -785,7 +785,7 @@ inline bool IsVirtualSamplerType(EMaterialSamplerType Value)
 	return ((int32)Value >= (int32)SAMPLERTYPE_VirtualColor && (int32)Value <= (int32)SAMPLERTYPE_VirtualLinearGrayscale);
 }
 UENUM()
-enum EMaterialStencilCompare
+enum EMaterialStencilCompare : int
 {
 	MSC_Less			UMETA(DisplayName = "Less Than"),
 	MSC_LessEqual		UMETA(DisplayName = "Less Than or Equal"),
@@ -799,7 +799,7 @@ enum EMaterialStencilCompare
 };
 
 UENUM()
-enum EMaterialShadingRate
+enum EMaterialShadingRate : int
 {
 	MSR_1x1				UMETA(DisplayName = "1x1"),
 	MSR_2x1				UMETA(DisplayName = "2x1"),
@@ -814,7 +814,7 @@ enum EMaterialShadingRate
 
 /**	Lighting build quality enumeration */
 UENUM(BlueprintType)
-enum ELightingBuildQuality
+enum ELightingBuildQuality : int
 {
 	Quality_Preview		UMETA(DisplayName = "Preview"),
 	Quality_Medium		UMETA(DisplayName = "Medium"),
@@ -825,7 +825,7 @@ enum ELightingBuildQuality
 
 /** Movement modes for Characters. */
 UENUM(BlueprintType)
-enum EMovementMode
+enum EMovementMode : int
 {
 	/** None (movement is disabled). */
 	MOVE_None		UMETA(DisplayName="None"),
@@ -965,7 +965,7 @@ DECLARE_DELEGATE_OneParam(FOnPlasticDeformation, int32 /*ConstraintIndex*/);
  * @warning If you change this, change GetCollisionChannelFromOverlapFilter() to match 
  */
 UENUM(BlueprintType)
-enum EOverlapFilterOption
+enum EOverlapFilterOption : int
 {
 	/** Returns both overlaps with both dynamic and static components */
 	OverlapFilter_All UMETA(DisplayName="AllObjects"),
@@ -977,7 +977,7 @@ enum EOverlapFilterOption
 
 /** Specifies custom collision object types, overridable per game */
 UENUM(BlueprintType)
-enum EObjectTypeQuery
+enum EObjectTypeQuery : int
 {
 	ObjectTypeQuery1 UMETA(Hidden), 
 	ObjectTypeQuery2 UMETA(Hidden), 
@@ -1017,7 +1017,7 @@ enum EObjectTypeQuery
 
 /** Specifies custom collision trace types, overridable per game */
 UENUM(BlueprintType)
-enum ETraceTypeQuery
+enum ETraceTypeQuery : int
 {
 	TraceTypeQuery1 UMETA(Hidden), 
 	TraceTypeQuery2 UMETA(Hidden), 
@@ -1057,7 +1057,7 @@ enum ETraceTypeQuery
 
 /** Enum indicating how each type should respond */
 UENUM(BlueprintType, meta=(ScriptName="CollisionResponseType"))
-enum ECollisionResponse
+enum ECollisionResponse : int
 {
 	ECR_Ignore UMETA(DisplayName="Ignore"),
 	ECR_Overlap UMETA(DisplayName="Overlap"),
@@ -1067,7 +1067,7 @@ enum ECollisionResponse
 
 /** Interpolation method used by animation blending */
 UENUM()
-enum EFilterInterpolationType
+enum EFilterInterpolationType : int
 {
 	BSIT_Average UMETA(DisplayName = "Averaged"),
 	BSIT_Linear UMETA(DisplayName = "Linear"),
@@ -1383,7 +1383,7 @@ private:
 
 /** Enum used to indicate what type of timeline signature a function matches. */
 UENUM()
-enum ETimelineSigType
+enum ETimelineSigType : int
 {
 	ETS_EventSignature,
 	ETS_FloatSignature,
@@ -1397,7 +1397,7 @@ enum ETimelineSigType
 UENUM(BlueprintType)
 namespace ECollisionEnabled 
 { 
-	enum Type 
+	enum Type : int
 	{ 
 		/** Will not create any representation in the physics engine. Cannot be used for spatial queries (raycasts, sweeps, overlaps) or simulation (rigid body, constraints). Best performance possible (especially for moving objects) */
 		NoCollision UMETA(DisplayName="No Collision"), 
@@ -2119,7 +2119,7 @@ struct FSwarmDebugOptions
 
 /** Method for padding a light map in memory */
 UENUM()
-enum ELightMapPaddingType
+enum ELightMapPaddingType : int
 {
 	LMPT_NormalPadding,
 	LMPT_PrePadding,
@@ -2128,7 +2128,7 @@ enum ELightMapPaddingType
 
 /** Bit-field flags that affects storage (e.g. packing, streaming) and other info about a shadowmap. */
 UENUM()
-enum EShadowMapFlags
+enum EShadowMapFlags : int
 {
 	/** No flags. */
 	SMF_None			= 0,
@@ -2823,7 +2823,7 @@ struct FMeshNaniteSettings
 
 /** The network role of an actor on a local/remote network context */
 UENUM()
-enum ENetRole
+enum ENetRole : int
 {
 	/** No role at all. */
 	ROLE_None,
@@ -2838,7 +2838,7 @@ enum ENetRole
 
 /** Describes if an actor can enter a low network bandwidth dormant mode */
 UENUM(BlueprintType)
-enum ENetDormancy
+enum ENetDormancy : int
 {
 	/** This actor can never go network dormant. */
 	DORM_Never UMETA(DisplayName = "Never"),
@@ -2858,7 +2858,7 @@ enum ENetDormancy
 UENUM()
 namespace EAutoReceiveInput
 {
-	enum Type
+	enum Type : int
 	{
 		Disabled,
 		Player0,
@@ -2912,7 +2912,7 @@ namespace EEndPlayReason
  * @see UCharacterMovementComponent::GetWalkableFloorAngle(), UCharacterMovementComponent::SetWalkableFloorAngle()
  */
 UENUM(BlueprintType)
-enum EWalkableSlopeBehavior
+enum EWalkableSlopeBehavior : int
 {
 	/** Don't affect the walkable slope. Walkable slope angle will be ignored. */
 	WalkableSlope_Default		UMETA(DisplayName="Unchanged"),
@@ -3204,7 +3204,7 @@ struct ENGINE_API FSoftComponentReference : public FBaseComponentReference
 UENUM(BlueprintType)
 namespace EPhysicalMaterialMaskColor
 {
-	enum Type
+	enum Type : int
 	{
 		Red,
 		Green,
@@ -3222,7 +3222,7 @@ namespace EPhysicalMaterialMaskColor
 UENUM(BlueprintType)
 namespace EComponentMobility
 {
-	enum Type
+	enum Type : int
 	{
 		/**
 		 * Static objects cannot be moved or changed in game.
@@ -3274,7 +3274,7 @@ public:
 UENUM()
 namespace EComponentSocketType
 {
-	enum Type
+	enum Type : int
 	{
 		/** Not a valid socket or bone name. */
 		Invalid,
diff --git a/Engine/Source/Runtime/Engine/Classes/Engine/FontImportOptions.h b/Engine/Source/Runtime/Engine/Classes/Engine/FontImportOptions.h
--- a/Engine/Source/Runtime/Engine/Classes/Engine/FontImportOptions.h
+++ b/Engine/Source/Runtime/Engine/Classes/Engine/FontImportOptions.h
@@ -10,7 +10,7 @@
 
 /** Font character set type for importing TrueType fonts. */
 UENUM()
-enum EFontImportCharacterSet
+enum EFontImportCharacterSet : int
 {
 	FontICS_Default,
 	FontICS_Ansi,
diff --git a/Engine/Source/Runtime/Engine/Classes/Engine/InGameAdManager.h b/Engine/Source/Runtime/Engine/Classes/Engine/InGameAdManager.h
--- a/Engine/Source/Runtime/Engine/Classes/Engine/InGameAdManager.h
+++ b/Engine/Source/Runtime/Engine/Classes/Engine/InGameAdManager.h
@@ -14,7 +14,7 @@
 #include "InGameAdManager.generated.h"
 
 UENUM()
-enum EAdManagerDelegate
+enum EAdManagerDelegate : int
 {
 	AMD_ClickedBanner,
 	AMD_UserClosedAd,
diff --git a/Engine/Source/Runtime/Engine/Classes/Engine/LevelStreamingVolume.h b/Engine/Source/Runtime/Engine/Classes/Engine/LevelStreamingVolume.h
--- a/Engine/Source/Runtime/Engine/Classes/Engine/LevelStreamingVolume.h
+++ b/Engine/Source/Runtime/Engine/Classes/Engine/LevelStreamingVolume.h
@@ -14,7 +14,7 @@
 
 /** Enum for different usage cases of level streaming volumes. */
 UENUM()
-enum EStreamingVolumeUsage
+enum EStreamingVolumeUsage : int
 {
 	SVB_Loading,
 	SVB_LoadingAndVisibility,
diff --git a/Engine/Source/Runtime/Engine/Classes/Engine/MaterialMerging.h b/Engine/Source/Runtime/Engine/Classes/Engine/MaterialMerging.h
--- a/Engine/Source/Runtime/Engine/Classes/Engine/MaterialMerging.h
+++ b/Engine/Source/Runtime/Engine/Classes/Engine/MaterialMerging.h
@@ -10,7 +10,7 @@
 struct FMeshDescription;
 
 UENUM()
-enum ETextureSizingType
+enum ETextureSizingType : int
 {
 	TextureSizingType_UseSingleTextureSize UMETA(DisplayName = "Use TextureSize for all material properties"),
 	TextureSizingType_UseAutomaticBiasedSizes UMETA(DisplayName = "Use automatically biased texture sizes based on TextureSize"),
@@ -23,7 +23,7 @@ enum ETextureSizingType
 };
 
 UENUM()
-enum EMaterialMergeType
+enum EMaterialMergeType : int
 {
 	MaterialMergeType_Default,
 	MaterialMergeType_Simplygon
diff --git a/Engine/Source/Runtime/Engine/Classes/Engine/MicroTransactionBase.h b/Engine/Source/Runtime/Engine/Classes/Engine/MicroTransactionBase.h
--- a/Engine/Source/Runtime/Engine/Classes/Engine/MicroTransactionBase.h
+++ b/Engine/Source/Runtime/Engine/Classes/Engine/MicroTransactionBase.h
@@ -14,7 +14,7 @@
 
 /** All the types of delegate callbacks that a MicroTransaction subclass may receive from C++. */
 UENUM()
-enum EMicroTransactionDelegate
+enum EMicroTransactionDelegate : int
 {
 	/**
 	 * Data:None
@@ -35,7 +35,7 @@ enum EMicroTransactionDelegate
 
 /** Result of a purchase. */
 UENUM()
-enum EMicroTransactionResult
+enum EMicroTransactionResult : int
 {
 	MTR_Succeeded,
 	MTR_Failed,
diff --git a/Engine/Source/Runtime/Engine/Classes/Engine/PlatformInterfaceBase.h b/Engine/Source/Runtime/Engine/Classes/Engine/PlatformInterfaceBase.h
--- a/Engine/Source/Runtime/Engine/Classes/Engine/PlatformInterfaceBase.h
+++ b/Engine/Source/Runtime/Engine/Classes/Engine/PlatformInterfaceBase.h
@@ -18,7 +18,7 @@
 
 /** An enum for the types of data used in a PlatformInterfaceData struct. */
 UENUM()
-enum EPlatformInterfaceDataType
+enum EPlatformInterfaceDataType : int
 {
 	/** No data type specified. */
 	PIDT_None,
diff --git a/Engine/Source/Runtime/Engine/Classes/Engine/Scene.h b/Engine/Source/Runtime/Engine/Classes/Engine/Scene.h
--- a/Engine/Source/Runtime/Engine/Classes/Engine/Scene.h
+++ b/Engine/Source/Runtime/Engine/Classes/Engine/Scene.h
@@ -24,7 +24,7 @@ struct FPostProcessSettings;
 
 /** Used by FPostProcessSettings Depth of Fields */
 UENUM()
-enum EDepthOfFieldMethod
+enum EDepthOfFieldMethod : int
 {
 	DOFM_BokehDOF UMETA(DisplayName="BokehDOF"),
 	DOFM_Gaussian UMETA(DisplayName="GaussianDOF"),
@@ -34,7 +34,7 @@ enum EDepthOfFieldMethod
 
 /** Used by FPostProcessSettings Auto Exposure */
 UENUM()
-enum EAutoExposureMethod
+enum EAutoExposureMethod : int
 {
 	/** requires compute shader to construct 64 bin histogram */
 	AEM_Histogram  UMETA(DisplayName = "Auto Exposure Histogram"),
@@ -46,7 +46,7 @@ enum EAutoExposureMethod
 };
 
 UENUM()
-enum EBloomMethod
+enum EBloomMethod : int
 {
 	/** Sum of Gaussian formulation */
 	BM_SOG  UMETA(DisplayName = "Standard"),
@@ -57,7 +57,7 @@ enum EBloomMethod
 
 /** Used by FPostProcessSettings to determine Temperature calculation method. */
 UENUM()
-enum ETemperatureMethod
+enum ETemperatureMethod : int
 {
 	TEMP_WhiteBalance UMETA(DisplayName = "White Balance"),
 	TEMP_ColorTemperature UMETA(DisplayName = "Color Temperature"),
@@ -118,7 +118,7 @@ enum class EReflectedAndRefractedRayTracedShadows : uint8
 UENUM()
 namespace EMobilePlanarReflectionMode
 {
-	enum Type
+	enum Type : int
 	{
 		Usual = 0 UMETA(DisplayName = "Usual", ToolTip = "The PlanarReflection actor works as usual on all platforms."),
 		MobilePPRExclusive = 1 UMETA(DisplayName = "MobilePPR Exclusive", ToolTip = "The PlanarReflection actor is only used for mobile pixel projection reflection, it will not affect PC/Console. MobileMSAA will be disabled as a side effect."),
@@ -129,7 +129,7 @@ namespace EMobilePlanarReflectionMode
 UENUM()
 namespace EMobilePixelProjectedReflectionQuality
 {
-	enum Type
+	enum Type : int
 	{
 		Disabled = 0 UMETA(DisplayName = "Disabled", ToolTip = "Disabled."),
 		BestPerformance = 1 UMETA(DisplayName = "Best Performance", ToolTip = "Best performance but may have some artifacts in some view angles."),
diff --git a/Engine/Source/Runtime/Engine/Classes/Engine/StaticMesh.h b/Engine/Source/Runtime/Engine/Classes/Engine/StaticMesh.h
--- a/Engine/Source/Runtime/Engine/Classes/Engine/StaticMesh.h
+++ b/Engine/Source/Runtime/Engine/Classes/Engine/StaticMesh.h
@@ -223,7 +223,7 @@ struct FStaticMeshAsyncBuildTask : public FAsyncTask<FStaticMeshAsyncBuildWorker
 
 /** Optimization settings used to simplify mesh LODs. */
 UENUM()
-enum ENormalMode
+enum ENormalMode : int
 {
 	NM_PreserveSmoothingGroups,
 	NM_RecalculateNormals,
@@ -234,7 +234,7 @@ enum ENormalMode
 };
 
 UENUM()
-enum EImportanceLevel
+enum EImportanceLevel : int
 {
 	IL_Off,
 	IL_Lowest,
@@ -248,7 +248,7 @@ enum EImportanceLevel
 
 /** Enum specifying the reduction type to use when simplifying static meshes. */
 UENUM()
-enum EOptimizationType
+enum EOptimizationType : int
 {
 	OT_NumOfTriangles,
 	OT_MaxDeviation,
diff --git a/Engine/Source/Runtime/Engine/Classes/Engine/TextureDefines.h b/Engine/Source/Runtime/Engine/Classes/Engine/TextureDefines.h
--- a/Engine/Source/Runtime/Engine/Classes/Engine/TextureDefines.h
+++ b/Engine/Source/Runtime/Engine/Classes/Engine/TextureDefines.h
@@ -23,7 +23,7 @@
  *                         and want to have the highest settings
  */
 UENUM()
-enum TextureGroup
+enum TextureGroup : int
 {
 	TEXTUREGROUP_World UMETA(DisplayName="ini:World"),
 	TEXTUREGROUP_WorldNormalMap UMETA(DisplayName="ini:WorldNormalMap"),
@@ -93,7 +93,7 @@ enum TextureGroup
 };
 
 UENUM()
-enum TextureMipGenSettings
+enum TextureMipGenSettings : int
 {
 	/** Default for the "texture". */
 	TMGS_FromTextureGroup UMETA(DisplayName="FromTextureGroup"),
@@ -145,7 +145,7 @@ enum TextureMipGenSettings
 UENUM()
 namespace ETexturePowerOfTwoSetting
 {
-	enum Type
+	enum Type : int
 	{
 		/** Do not modify the texture dimensions. */
 		None,
@@ -218,7 +218,7 @@ enum class ETextureDownscaleOptions : uint8
 };
 
 UENUM()
-enum ETextureLossyCompressionAmount
+enum ETextureLossyCompressionAmount : int
 {
 	TLCA_Default		UMETA(DisplayName = "Default"),
 	TLCA_None			UMETA(DisplayName = "No lossy compression (Oodle RDO disabled)"),
@@ -267,7 +267,7 @@ enum class ETextureClass : uint8
 };
 
 UENUM()
-enum ECompositeTextureMode
+enum ECompositeTextureMode : int
 {
 	CTM_Disabled UMETA(DisplayName="Disabled"),
 	/** CompositingTexture needs to be a normal map with the same or larger size. */
@@ -286,7 +286,7 @@ enum ECompositeTextureMode
 };
 
 UENUM()
-enum ETextureSourceCompressionFormat
+enum ETextureSourceCompressionFormat : int
 {
 	TSCF_None	UMETA(DisplayName = "None"),
 	TSCF_PNG	UMETA(DisplayName = "PNG"),
@@ -297,7 +297,7 @@ enum ETextureSourceCompressionFormat
 
 // ETextureSourceFormat should map one-to-one to ImageCore ERawImageFormat::Type
 UENUM()
-enum ETextureSourceFormat
+enum ETextureSourceFormat : int
 {
 	TSF_Invalid,
 	TSF_G8,
@@ -326,7 +326,7 @@ enum ETextureSourceFormat
 
 // This needs to be mirrored in EditorFactories.cpp.
 UENUM()
-enum TextureCompressionSettings
+enum TextureCompressionSettings : int
 {
 	TC_Default					UMETA(DisplayName = "Default (DXT1/5, BC1/3 on DX11)"),
 	TC_Normalmap				UMETA(DisplayName = "Normalmap (DXT5, BC5 on DX11)"),
diff --git a/Engine/Source/Runtime/Engine/Classes/Engine/TwitterIntegrationBase.h b/Engine/Source/Runtime/Engine/Classes/Engine/TwitterIntegrationBase.h
--- a/Engine/Source/Runtime/Engine/Classes/Engine/TwitterIntegrationBase.h
+++ b/Engine/Source/Runtime/Engine/Classes/Engine/TwitterIntegrationBase.h
@@ -16,7 +16,7 @@
 /** The possible twitter request methods */
 
 UENUM()
-enum ETwitterRequestMethod
+enum ETwitterRequestMethod : int
 {
 	TRM_Get,
 	TRM_Post,
@@ -25,7 +25,7 @@ enum ETwitterRequestMethod
 };
 
 UENUM()
-enum ETwitterIntegrationDelegate
+enum ETwitterIntegrationDelegate : int
 {
 	TID_AuthorizeComplete,
 	TID_TweetUIComplete,
diff --git a/Engine/Source/Runtime/Engine/Classes/Engine/UserDefinedStruct.h b/Engine/Source/Runtime/Engine/Classes/Engine/UserDefinedStruct.h
--- a/Engine/Source/Runtime/Engine/Classes/Engine/UserDefinedStruct.h
+++ b/Engine/Source/Runtime/Engine/Classes/Engine/UserDefinedStruct.h
@@ -18,7 +18,7 @@ class UStructCookedMetaData;
 DECLARE_MULTICAST_DELEGATE_OneParam(FOnStructChanged, UUserDefinedStruct*);
 
 UENUM()
-enum EUserDefinedStructureStatus
+enum EUserDefinedStructureStatus : int
 {
 	/** Struct is in an unknown state. */
 	UDSS_UpToDate,
diff --git a/Engine/Source/Runtime/Engine/Classes/GameFramework/GameNetworkManager.h b/Engine/Source/Runtime/Engine/Classes/GameFramework/GameNetworkManager.h
--- a/Engine/Source/Runtime/Engine/Classes/GameFramework/GameNetworkManager.h
+++ b/Engine/Source/Runtime/Engine/Classes/GameFramework/GameNetworkManager.h
@@ -10,7 +10,7 @@
 
 /** Describes which standby detection event occured so the game can take appropriate action. */
 UENUM()
-enum EStandbyType
+enum EStandbyType : int
 {
 	STDBY_Rx,
 	STDBY_Tx,
diff --git a/Engine/Source/Runtime/Engine/Classes/GameFramework/WorldSettings.h b/Engine/Source/Runtime/Engine/Classes/GameFramework/WorldSettings.h
--- a/Engine/Source/Runtime/Engine/Classes/GameFramework/WorldSettings.h
+++ b/Engine/Source/Runtime/Engine/Classes/GameFramework/WorldSettings.h
@@ -22,7 +22,7 @@ class UNavigationSystemConfig;
 class UAISystemBase;
 
 UENUM()
-enum EVisibilityAggressiveness
+enum EVisibilityAggressiveness : int
 {
 	VIS_LeastAggressive UMETA(DisplayName = "Least Aggressive"),
 	VIS_ModeratelyAggressive UMETA(DisplayName = "Moderately Aggressive"),
@@ -31,7 +31,7 @@ enum EVisibilityAggressiveness
 };
 
 UENUM()
-enum EVolumeLightingMethod
+enum EVolumeLightingMethod : int
 {
 	/** 
 	 * Lighting samples are computed in an adaptive grid which covers the entire Lightmass Importance Volume.  Higher density grids are used near geometry.
diff --git a/Engine/Source/Runtime/Engine/Classes/Kismet/KismetTextLibrary.h b/Engine/Source/Runtime/Engine/Classes/Kismet/KismetTextLibrary.h
--- a/Engine/Source/Runtime/Engine/Classes/Kismet/KismetTextLibrary.h
+++ b/Engine/Source/Runtime/Engine/Classes/Kismet/KismetTextLibrary.h
@@ -14,7 +14,7 @@
 
 /** Provides rounding modes for converting numbers into localized text */
 UENUM(BlueprintType)
-enum ERoundingMode
+enum ERoundingMode : int
 {
 	/** Rounds to the nearest place, equidistant ties go to the value which is closest to an even value: 1.5 becomes 2, 0.5 becomes 0 */
 	HalfToEven,
@@ -43,7 +43,7 @@ enum class ETextGender : uint8
 UENUM(BlueprintType)
 namespace EFormatArgumentType
 {
-	enum Type
+	enum Type : int
 	{
 		Int,
 		UInt,
diff --git a/Engine/Source/Runtime/Engine/Classes/Materials/Material.h b/Engine/Source/Runtime/Engine/Classes/Materials/Material.h
--- a/Engine/Source/Runtime/Engine/Classes/Materials/Material.h
+++ b/Engine/Source/Runtime/Engine/Classes/Materials/Material.h
@@ -64,7 +64,7 @@ struct FMaterialsWithDirtyUsageFlags
 
 /** Defines how the GBuffer channels are getting manipulated by a decal material pass. Actual index is used to control shader parameters so don't change order. */
 UENUM()
-enum EDecalBlendMode
+enum EDecalBlendMode : int
 {
 	/** Blend full material, updating the GBuffer, does not work for baked lighting. */
 	DBM_Translucent UMETA(DisplayName="Translucent"),
@@ -108,7 +108,7 @@ enum EDecalBlendMode
 
 /** Defines how the material reacts on DBuffer decals, later we can expose more variants between None and Default. */
 UENUM()
-enum EMaterialDecalResponse
+enum EMaterialDecalResponse : int
 {
 	/** Do not receive decals (Later we still can read the DBuffer channels to customize the effect, this frees up some interpolators). */
 	MDR_None UMETA(DisplayName="None"),
@@ -132,7 +132,7 @@ enum EMaterialDecalResponse
 
 /** Specifies which separate translucency pass to render in. */
 UENUM()
-enum EMaterialTranslucencyPass
+enum EMaterialTranslucencyPass : int
 {
 	/** Render after depth of field. */
 	MTP_BeforeDOF UMETA(DisplayName="Before DOF"),
diff --git a/Engine/Source/Runtime/Engine/Classes/Materials/MaterialExpressionAntialiasedTextureMask.h b/Engine/Source/Runtime/Engine/Classes/Materials/MaterialExpressionAntialiasedTextureMask.h
--- a/Engine/Source/Runtime/Engine/Classes/Materials/MaterialExpressionAntialiasedTextureMask.h
+++ b/Engine/Source/Runtime/Engine/Classes/Materials/MaterialExpressionAntialiasedTextureMask.h
@@ -11,7 +11,7 @@
 class UTexture;
 
 UENUM()
-enum ETextureColorChannel
+enum ETextureColorChannel : int
 {
 	TCC_Red,
 	TCC_Green,
diff --git a/Engine/Source/Runtime/Engine/Classes/Materials/MaterialExpressionClamp.h b/Engine/Source/Runtime/Engine/Classes/Materials/MaterialExpressionClamp.h
--- a/Engine/Source/Runtime/Engine/Classes/Materials/MaterialExpressionClamp.h
+++ b/Engine/Source/Runtime/Engine/Classes/Materials/MaterialExpressionClamp.h
@@ -10,7 +10,7 @@
 #include "MaterialExpressionClamp.generated.h"
 
 UENUM()
-enum EClampMode
+enum EClampMode : int
 {
 	CMODE_Clamp,
 	CMODE_ClampMin,
diff --git a/Engine/Source/Runtime/Engine/Classes/Materials/MaterialExpressionCustom.h b/Engine/Source/Runtime/Engine/Classes/Materials/MaterialExpressionCustom.h
--- a/Engine/Source/Runtime/Engine/Classes/Materials/MaterialExpressionCustom.h
+++ b/Engine/Source/Runtime/Engine/Classes/Materials/MaterialExpressionCustom.h
@@ -12,7 +12,7 @@
 struct FPropertyChangedEvent;
 
 UENUM()
-enum ECustomMaterialOutputType
+enum ECustomMaterialOutputType : int
 {
 	CMOT_Float1,
 	CMOT_Float2,
diff --git a/Engine/Source/Runtime/Engine/Classes/Materials/MaterialExpressionDBufferTexture.h b/Engine/Source/Runtime/Engine/Classes/Materials/MaterialExpressionDBufferTexture.h
--- a/Engine/Source/Runtime/Engine/Classes/Materials/MaterialExpressionDBufferTexture.h
+++ b/Engine/Source/Runtime/Engine/Classes/Materials/MaterialExpressionDBufferTexture.h
@@ -8,7 +8,7 @@
 #include "MaterialExpressionDBufferTexture.generated.h"
 
 UENUM()
-enum EDBufferTextureId
+enum EDBufferTextureId : int
 {
 	/** DBuffer A (Base Color). RGB is pre-multiplied by alpha. A is one minus alpha. */
 	DBT_A UMETA(DisplayName = "DBufferA (BaseColor)"),
diff --git a/Engine/Source/Runtime/Engine/Classes/Materials/MaterialExpressionDepthOfFieldFunction.h b/Engine/Source/Runtime/Engine/Classes/Materials/MaterialExpressionDepthOfFieldFunction.h
--- a/Engine/Source/Runtime/Engine/Classes/Materials/MaterialExpressionDepthOfFieldFunction.h
+++ b/Engine/Source/Runtime/Engine/Classes/Materials/MaterialExpressionDepthOfFieldFunction.h
@@ -11,7 +11,7 @@
 
 // Note: The index is used to map the enum to different code in the shader
 UENUM()
-enum EDepthOfFieldFunctionValue
+enum EDepthOfFieldFunctionValue : int
 {
 	/** 0:in Focus .. 1:Near or Far. */
 	TDOF_NearAndFarMask,
diff --git a/Engine/Source/Runtime/Engine/Classes/Materials/MaterialExpressionFunctionInput.h b/Engine/Source/Runtime/Engine/Classes/Materials/MaterialExpressionFunctionInput.h
--- a/Engine/Source/Runtime/Engine/Classes/Materials/MaterialExpressionFunctionInput.h
+++ b/Engine/Source/Runtime/Engine/Classes/Materials/MaterialExpressionFunctionInput.h
@@ -15,7 +15,7 @@ struct FPropertyChangedEvent;
 
 /** Supported input types */
 UENUM(BlueprintType)
-enum EFunctionInputType
+enum EFunctionInputType : int
 {
 	FunctionInput_Scalar,
 	FunctionInput_Vector2,
diff --git a/Engine/Source/Runtime/Engine/Classes/Materials/MaterialExpressionNoise.h b/Engine/Source/Runtime/Engine/Classes/Materials/MaterialExpressionNoise.h
--- a/Engine/Source/Runtime/Engine/Classes/Materials/MaterialExpressionNoise.h
+++ b/Engine/Source/Runtime/Engine/Classes/Materials/MaterialExpressionNoise.h
@@ -10,7 +10,7 @@
 #include "MaterialExpressionNoise.generated.h"
 
 UENUM()
-enum ENoiseFunction
+enum ENoiseFunction : int
 {
 	/** High quality for direct use and bumps 
 	 * ~77 instructions per level, 4 texture lookups
diff --git a/Engine/Source/Runtime/Engine/Classes/Materials/MaterialExpressionRuntimeVirtualTextureSample.h b/Engine/Source/Runtime/Engine/Classes/Materials/MaterialExpressionRuntimeVirtualTextureSample.h
--- a/Engine/Source/Runtime/Engine/Classes/Materials/MaterialExpressionRuntimeVirtualTextureSample.h
+++ b/Engine/Source/Runtime/Engine/Classes/Materials/MaterialExpressionRuntimeVirtualTextureSample.h
@@ -13,7 +13,7 @@
  * Internally we will convert to ETextureMipValueMode which is used by internal APIs.
  */
 UENUM()
-enum ERuntimeVirtualTextureMipValueMode
+enum ERuntimeVirtualTextureMipValueMode : int
 {
 	/* 
 	 * Use default computed mip level. Takes into account UV scaling from using the WorldPosition pin.
@@ -46,7 +46,7 @@ enum ERuntimeVirtualTextureMipValueMode
  * Defines texture addressing behavior.
  */
 UENUM()
-enum ERuntimeVirtualTextureTextureAddressMode
+enum ERuntimeVirtualTextureTextureAddressMode : int
 {
 	/* Clamp mode. */
 	RVTTA_Clamp UMETA(DisplayName = "Clamp"),
diff --git a/Engine/Source/Runtime/Engine/Classes/Materials/MaterialExpressionSpeedTree.h b/Engine/Source/Runtime/Engine/Classes/Materials/MaterialExpressionSpeedTree.h
--- a/Engine/Source/Runtime/Engine/Classes/Materials/MaterialExpressionSpeedTree.h
+++ b/Engine/Source/Runtime/Engine/Classes/Materials/MaterialExpressionSpeedTree.h
@@ -8,7 +8,7 @@
 #include "MaterialExpressionSpeedTree.generated.h"
 
 UENUM()
-enum ESpeedTreeGeometryType
+enum ESpeedTreeGeometryType : int
 {
 	STG_Branch		UMETA(DisplayName="Branch"),
 	STG_Frond		UMETA(DisplayName="Frond"),
@@ -18,7 +18,7 @@ enum ESpeedTreeGeometryType
 };
 
 UENUM()
-enum ESpeedTreeWindType
+enum ESpeedTreeWindType : int
 {
 	STW_None		UMETA(DisplayName="None"),
 	STW_Fastest		UMETA(DisplayName="Fastest"),
@@ -30,7 +30,7 @@ enum ESpeedTreeWindType
 };
 
 UENUM()
-enum ESpeedTreeLODType
+enum ESpeedTreeLODType : int
 {
 	STLOD_Pop		UMETA(DisplayName="Pop"),
 	STLOD_Smooth	UMETA(DisplayName="Smooth")
diff --git a/Engine/Source/Runtime/Engine/Classes/Materials/MaterialExpressionTextureProperty.h b/Engine/Source/Runtime/Engine/Classes/Materials/MaterialExpressionTextureProperty.h
--- a/Engine/Source/Runtime/Engine/Classes/Materials/MaterialExpressionTextureProperty.h
+++ b/Engine/Source/Runtime/Engine/Classes/Materials/MaterialExpressionTextureProperty.h
@@ -13,7 +13,7 @@ class UTexture;
 
 /** Selects the texture property to output */
 UENUM()
-enum EMaterialExposedTextureProperty
+enum EMaterialExposedTextureProperty : int
 {
 	/* The texture's size. */
 	TMTM_TextureSize UMETA(DisplayName="Texture Size"),
diff --git a/Engine/Source/Runtime/Engine/Classes/Materials/MaterialExpressionTransform.h b/Engine/Source/Runtime/Engine/Classes/Materials/MaterialExpressionTransform.h
--- a/Engine/Source/Runtime/Engine/Classes/Materials/MaterialExpressionTransform.h
+++ b/Engine/Source/Runtime/Engine/Classes/Materials/MaterialExpressionTransform.h
@@ -10,7 +10,7 @@
 #include "MaterialExpressionTransform.generated.h"
 
 UENUM()
-enum EMaterialVectorCoordTransformSource
+enum EMaterialVectorCoordTransformSource : int
 {
 	/** Tangent space (relative to the surface) */
 	TRANSFORMSOURCE_Tangent UMETA(DisplayName="Tangent Space"),
@@ -37,7 +37,7 @@ enum EMaterialVectorCoordTransformSource
 };
 
 UENUM()
-enum EMaterialVectorCoordTransform
+enum EMaterialVectorCoordTransform : int
 {
 	/** Tangent space (relative to the surface) */
 	TRANSFORM_Tangent UMETA(DisplayName="Tangent Space"),
diff --git a/Engine/Source/Runtime/Engine/Classes/Materials/MaterialExpressionTransformPosition.h b/Engine/Source/Runtime/Engine/Classes/Materials/MaterialExpressionTransformPosition.h
--- a/Engine/Source/Runtime/Engine/Classes/Materials/MaterialExpressionTransformPosition.h
+++ b/Engine/Source/Runtime/Engine/Classes/Materials/MaterialExpressionTransformPosition.h
@@ -10,7 +10,7 @@
 #include "MaterialExpressionTransformPosition.generated.h"
 
 UENUM()
-enum EMaterialPositionTransformSource
+enum EMaterialPositionTransformSource : int
 {
 	/** Local space */
 	TRANSFORMPOSSOURCE_Local UMETA(DisplayName="Local Space"),
diff --git a/Engine/Source/Runtime/Engine/Classes/Materials/MaterialExpressionVectorNoise.h b/Engine/Source/Runtime/Engine/Classes/Materials/MaterialExpressionVectorNoise.h
--- a/Engine/Source/Runtime/Engine/Classes/Materials/MaterialExpressionVectorNoise.h
+++ b/Engine/Source/Runtime/Engine/Classes/Materials/MaterialExpressionVectorNoise.h
@@ -10,7 +10,7 @@
 #include "MaterialExpressionVectorNoise.generated.h"
 
 UENUM()
-enum EVectorNoiseFunction
+enum EVectorNoiseFunction : int
 {
 	/** Random color for each unit cell in 3D space.
 	  * RGB output range 0 to 1
