Add host machine name to OOM report in DXC backend and refactor global SCW error-code handling.

#rnx

diff --git a/Engine/Source/Developer/ShaderCompilerCommon/Private/ShaderConductorContext.cpp b/Engine/Source/Developer/ShaderCompilerCommon/Private/ShaderConductorContext.cpp
--- a/Engine/Source/Developer/ShaderCompilerCommon/Private/ShaderConductorContext.cpp
+++ b/Engine/Source/Developer/ShaderCompilerCommon/Private/ShaderConductorContext.cpp
@@ -10,8 +10,6 @@ THIRD_PARTY_INCLUDES_END
 #endif
 
 
-ESCWErrorCode GSCWErrorCode = ESCWErrorCode::NotSet;
-
 namespace CrossCompiler
 {
 
diff --git a/Engine/Source/Developer/ShaderCompilerCommon/Public/ShaderConductorContext.h b/Engine/Source/Developer/ShaderCompilerCommon/Public/ShaderConductorContext.h
--- a/Engine/Source/Developer/ShaderCompilerCommon/Public/ShaderConductorContext.h
+++ b/Engine/Source/Developer/ShaderCompilerCommon/Public/ShaderConductorContext.h
@@ -246,6 +246,3 @@ namespace CrossCompiler
 	};
 
 }
-
-// Error code for SCW to help track down crashes
-extern SHADERCOMPILERCOMMON_API ESCWErrorCode GSCWErrorCode;
diff --git a/Engine/Source/Developer/Windows/ShaderFormatD3D/Private/D3DShaderCompiler.cpp b/Engine/Source/Developer/Windows/ShaderFormatD3D/Private/D3DShaderCompiler.cpp
--- a/Engine/Source/Developer/Windows/ShaderFormatD3D/Private/D3DShaderCompiler.cpp
+++ b/Engine/Source/Developer/Windows/ShaderFormatD3D/Private/D3DShaderCompiler.cpp
@@ -483,7 +483,7 @@ static HRESULT D3DCompileWrapper(
 #if !PLATFORM_SEH_EXCEPTIONS_DISABLED
 	__except(D3DExceptionFilter(bCatchException))
 	{
-		GSCWErrorCode = ESCWErrorCode::CrashInsidePlatformCompiler;
+		FSCWErrorCode::Report(FSCWErrorCode::CrashInsidePlatformCompiler);
 		return E_FAIL;
 	}
 #endif
diff --git a/Engine/Source/Developer/Windows/ShaderFormatD3D/Private/D3DShaderCompilerDXC.cpp b/Engine/Source/Developer/Windows/ShaderFormatD3D/Private/D3DShaderCompilerDXC.cpp
--- a/Engine/Source/Developer/Windows/ShaderFormatD3D/Private/D3DShaderCompilerDXC.cpp
+++ b/Engine/Source/Developer/Windows/ShaderFormatD3D/Private/D3DShaderCompilerDXC.cpp
@@ -10,6 +10,7 @@
 #include "HAL/FileManager.h"
 #include "Serialization/MemoryWriter.h"
 #include "RayTracingDefinitions.h"
+#include "Math/UnitConversion.h"
 
 DEFINE_LOG_CATEGORY_STATIC(LogD3D12ShaderCompiler, Log, All);
 
@@ -127,11 +128,12 @@ static void LogFailedHRESULT(const TCHAR* FailedExpressionStr, HRESULT Result)
 {
 	if (Result == E_OUTOFMEMORY)
 	{
-		// Provide currently availble system memory to verify out-of-memory error is legitimate
 		const FPlatformMemoryStats MemoryStats = FPlatformMemory::GetStats();
-		FString SystemMemoryStr = FString::Printf(TEXT("%d MB"), MemoryStats.AvailablePhysical / 1024 / 1024);
-		FString UsedMemoryStr = FString::Printf(TEXT("%d MB"), MemoryStats.UsedPhysical / 1024 / 1024);
-		UE_LOG(LogD3D12ShaderCompiler, Fatal, TEXT("%s failed: Result=0x%08x (E_OUTOFMEMORY); System memory %s, used %s"), FailedExpressionStr, Result, *SystemMemoryStr, *UsedMemoryStr);
+		const FString AvailableMemoryStr = FString::Printf(TEXT("%d MB"), FUnitConversion::Convert(MemoryStats.AvailablePhysical, EUnit::Bytes, EUnit::Megabytes));
+		const FString UsedMemoryStr = FString::Printf(TEXT("%d MB"), FUnitConversion::Convert(MemoryStats.UsedPhysical, EUnit::Bytes, EUnit::Megabytes));
+		const FString ErrorReport = FString::Printf(TEXT("%s failed: Result=0x%08x (E_OUTOFMEMORY); Used physical memory %s of %s"), FailedExpressionStr, Result, *UsedMemoryStr, *AvailableMemoryStr);
+		FSCWErrorCode::Report(FSCWErrorCode::OutOfMemory, ErrorReport);
+		UE_LOG(LogD3D12ShaderCompiler, Fatal, TEXT("%s"), *ErrorReport);
 	}
 	else if (const TCHAR* ErrorCodeStr = DxcErrorCodeToString(Result))
 	{
@@ -197,17 +199,18 @@ public:
 
 	// IUnknown
 
-	ULONG STDCALL AddRef()
+	ULONG STDCALL AddRef() override
 	{
 		return ++RefCount;
 	}
 
-	ULONG STDCALL Release() {
-		check(RefCount != 0);
+	ULONG STDCALL Release() override
+	{
+		check(RefCount > 0);
 		return --RefCount;
 	}
 
-	HRESULT STDCALL QueryInterface(REFIID iid, void** ppvObject)
+	HRESULT STDCALL QueryInterface(REFIID iid, void** ppvObject) override
 	{
 		checkNoEntry(); // We do not expect or support QI on DXC allocator replacement
 		return ERROR_NOINTERFACE;
@@ -320,7 +323,7 @@ static HRESULT DXCCompileWrapper(
 
 	if (bExceptionError)
 	{
-		GSCWErrorCode = ESCWErrorCode::CrashInsidePlatformCompiler;
+		FSCWErrorCode::Report(FSCWErrorCode::CrashInsidePlatformCompiler);
 
 		FString ErrorMsg = TEXT("Internal error or exception inside dxcompiler.dll\n");
 		ErrorMsg += GDxcStackTrace;
@@ -603,7 +606,7 @@ static bool DXCRewriteWrapper(
 #if !PLATFORM_SEH_EXCEPTIONS_DISABLED
 	__except (EXCEPTION_EXECUTE_HANDLER)
 	{
-		GSCWErrorCode = ESCWErrorCode::CrashInsidePlatformCompiler;
+		FSCWErrorCode::Report(FSCWErrorCode::CrashInsidePlatformCompiler);
 		FMemory::Memzero(OutResultDesc);
 		bOutException = true;
 		return false;
diff --git a/Engine/Source/Programs/ShaderCompileWorker/Private/ShaderCompileWorker.cpp b/Engine/Source/Programs/ShaderCompileWorker/Private/ShaderCompileWorker.cpp
--- a/Engine/Source/Programs/ShaderCompileWorker/Private/ShaderCompileWorker.cpp
+++ b/Engine/Source/Programs/ShaderCompileWorker/Private/ShaderCompileWorker.cpp
@@ -15,6 +15,7 @@
 #include "RHIShaderFormatDefinitions.inl"
 #include "ShaderCompilerCommon.h"
 #include "Serialization/MemoryReader.h"
+#include "SocketSubsystem.h"
 
 #define DEBUG_USING_CONSOLE	0
 
@@ -35,7 +36,7 @@ inline bool IsUsingXGE()
 	return GXGEMode != EXGEMode::None;
 }
 
-static ESCWErrorCode GFailedErrorCode = ESCWErrorCode::Success;
+static FSCWErrorCode::ECode GFailedErrorCode = FSCWErrorCode::Success;
 
 static void OnXGEJobCompleted(const TCHAR* WorkingDirectory)
 {
@@ -48,10 +49,10 @@ static void OnXGEJobCompleted(const TCHAR* WorkingDirectory)
 }
 
 #if USING_CODE_ANALYSIS
-	UE_NORETURN static inline void ExitWithoutCrash(ESCWErrorCode ErrorCode, const FString& Message);
+	UE_NORETURN static inline void ExitWithoutCrash(FSCWErrorCode::ECode ErrorCode, const FString& Message);
 #endif
 
-static inline void ExitWithoutCrash(ESCWErrorCode ErrorCode, const FString& Message)
+static inline void ExitWithoutCrash(FSCWErrorCode::ECode ErrorCode, const FString& Message)
 {
 	GFailedErrorCode = ErrorCode;
 	FCString::Snprintf(GErrorExceptionDescription, sizeof(GErrorExceptionDescription), TEXT("%s"), *Message);
@@ -73,7 +74,7 @@ static const TArray<const IShaderFormat*>& GetShaderFormats()
 
 		if (!Modules.Num())
 		{
-			ExitWithoutCrash(ESCWErrorCode::NoTargetShaderFormatsFound, TEXT("No target shader formats found!"));
+			ExitWithoutCrash(FSCWErrorCode::NoTargetShaderFormatsFound, TEXT("No target shader formats found!"));
 		}
 
 		for (int32 Index = 0; Index < Modules.Num(); Index++)
@@ -189,14 +190,14 @@ static const IShaderFormat* FindShaderFormat(FName Name)
 
 
 /** Processes a compilation job. */
-static void ProcessCompilationJob(const FShaderCompilerInput& Input,FShaderCompilerOutput& Output,const FString& WorkingDirectory)
+static FORCENOINLINE void ProcessCompilationJob(const FShaderCompilerInput& Input,FShaderCompilerOutput& Output,const FString& WorkingDirectory)
 {
 	TRACE_CPUPROFILER_EVENT_SCOPE(ProcessCompilationJob);
 
 	const IShaderFormat* Compiler = FindShaderFormat(Input.ShaderFormat);
 	if (!Compiler)
 	{
-		ExitWithoutCrash(ESCWErrorCode::CantCompileForSpecificFormat, FString::Printf(TEXT("Can't compile shaders for format %s"), *Input.ShaderFormat.ToString()));
+		ExitWithoutCrash(FSCWErrorCode::CantCompileForSpecificFormat, FString::Printf(TEXT("Can't compile shaders for format %s"), *Input.ShaderFormat.ToString()));
 	}
 
 	// Apply the console variable values from the input environment before calling the platform shader compiler
@@ -258,6 +259,20 @@ static void UpdateFileSize(FArchive& OutputFile, int64 FileSizePosition)
 	OutputFile.Seek(Current);
 };
 
+static const TCHAR* GetLocalHostname(int32* OutHostnameLength = nullptr)
+{
+	static FString Hostname;
+	if (Hostname.IsEmpty())
+	{
+		ISocketSubsystem::Get()->GetHostName(Hostname);
+	}
+	if (OutHostnameLength)
+	{
+		*OutHostnameLength = Hostname.Len();
+	}
+	return *Hostname;
+}
+
 static int64 WriteOutputFileHeader(FArchive& OutputFile, int32 ErrorCode, int32 CallstackLength, const TCHAR* Callstack,
 	int32 ExceptionInfoLength, const TCHAR* ExceptionInfo)
 {
@@ -281,6 +296,10 @@ static int64 WriteOutputFileHeader(FArchive& OutputFile, int32 ErrorCode, int32
 
 	OutputFile << ExceptionInfoLength;
 
+	int32 HostnameLength = 0;
+	const TCHAR* Hostname = GetLocalHostname(&HostnameLength);
+	OutputFile << HostnameLength;
+
 	if (CallstackLength > 0)
 	{
 		OutputFile.Serialize((void*)Callstack, CallstackLength * sizeof(TCHAR));
@@ -291,6 +310,11 @@ static int64 WriteOutputFileHeader(FArchive& OutputFile, int32 ErrorCode, int32
 		OutputFile.Serialize((void*)ExceptionInfo, ExceptionInfoLength * sizeof(TCHAR));
 	}
 
+	if (HostnameLength > 0)
+	{
+		OutputFile.Serialize((void*)Hostname, HostnameLength * sizeof(TCHAR));
+	}
+
 	UpdateFileSize(OutputFile, FileSizePosition);
 	return FileSizePosition;
 }
@@ -461,7 +485,7 @@ private:
 			{
 				if (Pair.Value != *Found)
 				{
-					ExitWithoutCrash(ESCWErrorCode::BadShaderFormatVersion, FString::Printf(TEXT("Mismatched shader version for format %s: Found version %u but expected %u; did you forget to build ShaderCompilerWorker?"), *Pair.Key, *Found, Pair.Value));
+					ExitWithoutCrash(FSCWErrorCode::BadShaderFormatVersion, FString::Printf(TEXT("Mismatched shader version for format %s: Found version %u but expected %u; did you forget to build ShaderCompilerWorker?"), *Pair.Key, *Found, Pair.Value));
 				}
 			}
 		}
@@ -475,7 +499,7 @@ private:
 		*InputFilePtr << InputVersion;
 		if (ShaderCompileWorkerInputVersion != InputVersion)
 		{
-			ExitWithoutCrash(ESCWErrorCode::BadInputVersion, FString::Printf(TEXT("Exiting due to ShaderCompilerWorker expecting input version %d, got %d instead! Did you forget to build ShaderCompilerWorker?"), ShaderCompileWorkerInputVersion, InputVersion));
+			ExitWithoutCrash(FSCWErrorCode::BadInputVersion, FString::Printf(TEXT("Exiting due to ShaderCompilerWorker expecting input version %d, got %d instead! Did you forget to build ShaderCompilerWorker?"), ShaderCompileWorkerInputVersion, InputVersion));
 		}
 
 		FString CompressionFormatString;
@@ -492,7 +516,7 @@ private:
 
 			if (UncompressedDataSize == 0)
 			{
-				ExitWithoutCrash(ESCWErrorCode::BadInputFile, TEXT("Exiting due to bad input file to ShaderCompilerWorker (uncompressed size is 0)! Did you forget to build ShaderCompilerWorker?"));
+				ExitWithoutCrash(FSCWErrorCode::BadInputFile, TEXT("Exiting due to bad input file to ShaderCompilerWorker (uncompressed size is 0)! Did you forget to build ShaderCompilerWorker?"));
 				// unreachable
 				return;
 			}
@@ -502,7 +526,7 @@ private:
 			*InputFilePtr << CompressedData;
 			if (!FCompression::UncompressMemory(CompressionFormat, UncompressedData.GetData(), UncompressedDataSize, CompressedData.GetData(), CompressedData.Num()))
 			{
-				ExitWithoutCrash(ESCWErrorCode::BadInputFile, FString::Printf(TEXT("Exiting due to bad input file to ShaderCompilerWorker (cannot uncompress with the format %s)! Did you forget to build ShaderCompilerWorker?"), *CompressionFormatString));
+				ExitWithoutCrash(FSCWErrorCode::BadInputFile, FString::Printf(TEXT("Exiting due to bad input file to ShaderCompilerWorker (cannot uncompress with the format %s)! Did you forget to build ShaderCompilerWorker?"), *CompressionFormatString));
 				// unreachable
 				return;
 			}
@@ -715,7 +739,7 @@ private:
 			InputFile << SingleJobHeader;
 			if (ShaderCompileWorkerSingleJobHeader != SingleJobHeader)
 			{
-				ExitWithoutCrash(ESCWErrorCode::BadSingleJobHeader, FString::Printf(TEXT("Exiting due to ShaderCompilerWorker expecting job header %d, got %d instead! Did you forget to build ShaderCompilerWorker?"), ShaderCompileWorkerSingleJobHeader, SingleJobHeader));
+				ExitWithoutCrash(FSCWErrorCode::BadSingleJobHeader, FString::Printf(TEXT("Exiting due to ShaderCompilerWorker expecting job header %d, got %d instead! Did you forget to build ShaderCompilerWorker?"), ShaderCompileWorkerSingleJobHeader, SingleJobHeader));
 			}
 
 			int32 NumBatches = 0;
@@ -760,7 +784,7 @@ private:
 			InputFile << PipelineJobHeader;
 			if (ShaderCompileWorkerPipelineJobHeader != PipelineJobHeader)
 			{
-				ExitWithoutCrash(ESCWErrorCode::BadPipelineJobHeader, FString::Printf(TEXT("Exiting due to ShaderCompilerWorker expecting pipeline job header %d, got %d instead! Did you forget to build ShaderCompilerWorker?"), ShaderCompileWorkerSingleJobHeader, PipelineJobHeader));
+				ExitWithoutCrash(FSCWErrorCode::BadPipelineJobHeader, FString::Printf(TEXT("Exiting due to ShaderCompilerWorker expecting pipeline job header %d, got %d instead! Did you forget to build ShaderCompilerWorker?"), ShaderCompileWorkerSingleJobHeader, PipelineJobHeader));
 			}
 
 			int32 NumPipelines = 0;
@@ -874,7 +898,7 @@ private:
 
 			if (!bResult)
 			{
-				ExitWithoutCrash(ESCWErrorCode::CantDeleteInputFile, FString::Printf(TEXT("Couldn't delete input file %s, is it readonly?"), *InputFilePath));
+				ExitWithoutCrash(FSCWErrorCode::CantDeleteInputFile, FString::Printf(TEXT("Couldn't delete input file %s, is it readonly?"), *InputFilePath));
 			}
 		}
 
@@ -898,7 +922,7 @@ private:
 			
 		if (!OutputFilePtr)
 		{
-			ExitWithoutCrash(ESCWErrorCode::CantSaveOutputFile, FString::Printf(TEXT("Couldn't save output file %s"), *TempFilePath));
+			ExitWithoutCrash(FSCWErrorCode::CantSaveOutputFile, FString::Printf(TEXT("Couldn't save output file %s"), *TempFilePath));
 		}
 
 		return OutputFilePtr;
@@ -909,7 +933,9 @@ private:
 		TRACE_CPUPROFILER_EVENT_SCOPE(WriteToOutputArchive);
 
 		FArchive& OutputFile = *OutputFilePtr;
-		int64 FileSizePosition = WriteOutputFileHeader(OutputFile, (int32)ESCWErrorCode::Success, 0, nullptr, 0, nullptr);
+		const int64 FileSizePosition = FSCWErrorCode::IsSet()
+			? WriteOutputFileHeader(OutputFile, FSCWErrorCode::Get(), 0, nullptr, FSCWErrorCode::GetInfo().Len(), *FSCWErrorCode::GetInfo())
+			: WriteOutputFileHeader(OutputFile, FSCWErrorCode::Success, 0, nullptr, 0, nullptr);
 
 		{
 			int32 SingleJobHeader = ShaderCompileWorkerSingleJobHeader;
@@ -1387,21 +1413,21 @@ static int32 GuardedMainWrapper(int32 ArgC, TCHAR* ArgV[], const TCHAR* CrashOut
 		{
 			FArchive& OutputFile = *IFileManager::Get().CreateFileWriter(CrashOutputFile, FILEWRITE_EvenIfReadOnly);
 
-			if (GFailedErrorCode == ESCWErrorCode::Success)
+			if (GFailedErrorCode == FSCWErrorCode::Success)
 			{
-				if (GSCWErrorCode != ESCWErrorCode::NotSet)
+				if (FSCWErrorCode::IsSet())
 				{
 					// Use the value set inside the shader format
-					GFailedErrorCode = GSCWErrorCode;
+					GFailedErrorCode = FSCWErrorCode::Get();
 				}
 				else
 				{
 					// Something else failed before we could set the error code, so mark it as a General Crash
-					GFailedErrorCode = ESCWErrorCode::GeneralCrash;
+					GFailedErrorCode = FSCWErrorCode::GeneralCrash;
 				}
 			}
-			
-			int64 FileSizePosition = WriteOutputFileHeader(OutputFile, (int32)GFailedErrorCode, ExceptionCallStack.Len(), *ExceptionCallStack,
+
+			int64 FileSizePosition = WriteOutputFileHeader(OutputFile, GFailedErrorCode, ExceptionCallStack.Len(), *ExceptionCallStack,
 				ExceptionMsg.Len(), *ExceptionMsg);
 
 			int32 NumBatches = 0;
diff --git a/Engine/Source/Programs/ShaderCompileWorker/ShaderCompileWorker.Build.cs b/Engine/Source/Programs/ShaderCompileWorker/ShaderCompileWorker.Build.cs
--- a/Engine/Source/Programs/ShaderCompileWorker/ShaderCompileWorker.Build.cs
+++ b/Engine/Source/Programs/ShaderCompileWorker/ShaderCompileWorker.Build.cs
@@ -16,7 +16,8 @@ public class ShaderCompileWorker : ModuleRules
 				"TargetPlatform",
 				"ApplicationCore",
 				"TraceLog",
-				"ShaderCompilerCommon"
+				"ShaderCompilerCommon",
+				"Sockets",
 			});
 
 		PrivateIncludePathModuleNames.AddRange(
diff --git a/Engine/Source/Runtime/Engine/Private/ShaderCompiler/ShaderCompiler.cpp b/Engine/Source/Runtime/Engine/Private/ShaderCompiler/ShaderCompiler.cpp
--- a/Engine/Source/Runtime/Engine/Private/ShaderCompiler/ShaderCompiler.cpp
+++ b/Engine/Source/Runtime/Engine/Private/ShaderCompiler/ShaderCompiler.cpp
@@ -1245,7 +1245,7 @@ namespace ShaderCompilerCookStats
 #endif
 
 // Make functions so the crash reporter can disambiguate the actual error because of the different callstacks
-namespace SCWErrorCode
+namespace ShaderCompileWorkerError
 {
 	void HandleGeneralCrash(const TCHAR* ExceptionInfo, const TCHAR* Callstack)
 	{
@@ -1307,6 +1307,16 @@ namespace SCWErrorCode
 	{
 		ModalErrorOrLog(FString::Printf(TEXT("Crash inside the platform compiler!\n%s"), Data));
 	}
+
+	void HandleBadInputFile(const TCHAR* Data)
+	{
+		ModalErrorOrLog(FString::Printf(TEXT("ShaderCompileWorker failed with bad-input-file exception:\n%s\n"), Data));
+	}
+
+	void HandleOutOfMemory(const TCHAR* ExceptionInfo, const TCHAR* Hostname)
+	{
+		ModalErrorOrLog(FString::Printf(TEXT("ShaderCompileWorker failed with out-of-memory exception on machine \"%s\"\n%s\n"), Hostname, ExceptionInfo));
+	}
 }
 
 static TMap<FString, uint32> GetFormatVersionMap()
@@ -1851,17 +1861,22 @@ static void LogQueuedCompileJobs(const TArray<FShaderCommonCompileJobPtr>& Queue
 
 // Disable optimization for this crash handler to get full access to the entire stack frame when debugging a crash dump
 UE_DISABLE_OPTIMIZATION_SHIP
-static void HandleWorkerCrash(const TArray<FShaderCommonCompileJobPtr>& QueuedJobs, FArchive& OutputFile, int32 OutputVersion, int64 FileSize, ESCWErrorCode ErrorCode, int32 NumProcessedJobs, int32 CallstackLength, int32 ExceptionInfoLength)
+static void HandleWorkerCrash(const TArray<FShaderCommonCompileJobPtr>& QueuedJobs, FArchive& OutputFile, int32 OutputVersion, int64 FileSize, FSCWErrorCode::ECode ErrorCode, int32 NumProcessedJobs, int32 CallstackLength, int32 ExceptionInfoLength, int32 HostnameLength)
 {
 	TArray<TCHAR> Callstack;
 	Callstack.AddUninitialized(CallstackLength + 1);
 	OutputFile.Serialize(Callstack.GetData(), CallstackLength * sizeof(TCHAR));
-	Callstack[CallstackLength] = 0;
+	Callstack[CallstackLength] = TEXT('\0');
 
 	TArray<TCHAR> ExceptionInfo;
 	ExceptionInfo.AddUninitialized(ExceptionInfoLength + 1);
 	OutputFile.Serialize(ExceptionInfo.GetData(), ExceptionInfoLength * sizeof(TCHAR));
-	ExceptionInfo[ExceptionInfoLength] = 0;
+	ExceptionInfo[ExceptionInfoLength] = TEXT('\0');
+
+	TArray<TCHAR> Hostname;
+	Hostname.AddUninitialized(HostnameLength + 1);
+	OutputFile.Serialize(Hostname.GetData(), HostnameLength * sizeof(TCHAR));
+	Hostname[HostnameLength] = TEXT('\0');
 
 	// Store primary job information onto stack to make it part of a crash dump
 	static const int32 MaxNumCharsForSourcePaths = 8192;
@@ -1913,41 +1928,45 @@ static void HandleWorkerCrash(const TArray<FShaderCommonCompileJobPtr>& QueuedJo
 	switch (ErrorCode)
 	{
 	default:
-	case ESCWErrorCode::GeneralCrash:
-	{
+	case FSCWErrorCode::GeneralCrash:
 		LogQueuedCompileJobs(QueuedJobs, NumProcessedJobs);
-		SCWErrorCode::HandleGeneralCrash(ExceptionInfo.GetData(), Callstack.GetData());
-	}
-	break;
-	case ESCWErrorCode::BadShaderFormatVersion:
-		SCWErrorCode::HandleBadShaderFormatVersion(ExceptionInfo.GetData());
+		ShaderCompileWorkerError::HandleGeneralCrash(ExceptionInfo.GetData(), Callstack.GetData());
 		break;
-	case ESCWErrorCode::BadInputVersion:
-		SCWErrorCode::HandleBadInputVersion(ExceptionInfo.GetData());
+	case FSCWErrorCode::BadShaderFormatVersion:
+		ShaderCompileWorkerError::HandleBadShaderFormatVersion(ExceptionInfo.GetData());
 		break;
-	case ESCWErrorCode::BadSingleJobHeader:
-		SCWErrorCode::HandleBadSingleJobHeader(ExceptionInfo.GetData());
+	case FSCWErrorCode::BadInputVersion:
+		ShaderCompileWorkerError::HandleBadInputVersion(ExceptionInfo.GetData());
 		break;
-	case ESCWErrorCode::BadPipelineJobHeader:
-		SCWErrorCode::HandleBadPipelineJobHeader(ExceptionInfo.GetData());
+	case FSCWErrorCode::BadSingleJobHeader:
+		ShaderCompileWorkerError::HandleBadSingleJobHeader(ExceptionInfo.GetData());
 		break;
-	case ESCWErrorCode::CantDeleteInputFile:
-		SCWErrorCode::HandleCantDeleteInputFile(ExceptionInfo.GetData());
+	case FSCWErrorCode::BadPipelineJobHeader:
+		ShaderCompileWorkerError::HandleBadPipelineJobHeader(ExceptionInfo.GetData());
 		break;
-	case ESCWErrorCode::CantSaveOutputFile:
-		SCWErrorCode::HandleCantSaveOutputFile(ExceptionInfo.GetData());
+	case FSCWErrorCode::CantDeleteInputFile:
+		ShaderCompileWorkerError::HandleCantDeleteInputFile(ExceptionInfo.GetData());
 		break;
-	case ESCWErrorCode::NoTargetShaderFormatsFound:
-		SCWErrorCode::HandleNoTargetShaderFormatsFound(ExceptionInfo.GetData());
+	case FSCWErrorCode::CantSaveOutputFile:
+		ShaderCompileWorkerError::HandleCantSaveOutputFile(ExceptionInfo.GetData());
 		break;
-	case ESCWErrorCode::CantCompileForSpecificFormat:
-		SCWErrorCode::HandleCantCompileForSpecificFormat(ExceptionInfo.GetData());
+	case FSCWErrorCode::NoTargetShaderFormatsFound:
+		ShaderCompileWorkerError::HandleNoTargetShaderFormatsFound(ExceptionInfo.GetData());
+		break;
+	case FSCWErrorCode::CantCompileForSpecificFormat:
+		ShaderCompileWorkerError::HandleCantCompileForSpecificFormat(ExceptionInfo.GetData());
 		break;
 	case ESCWErrorCode::CrashInsidePlatformCompiler:
 		LogQueuedCompileJobs(QueuedJobs, NumProcessedJobs);
-		SCWErrorCode::HandleCrashInsidePlatformCompiler(ExceptionInfo.GetData());
+		ShaderCompileWorkerError::HandleCrashInsidePlatformCompiler(ExceptionInfo.GetData());
+		break;
+	case FSCWErrorCode::BadInputFile:
+		ShaderCompileWorkerError::HandleBadInputFile(ExceptionInfo.GetData());
 		break;
-	case ESCWErrorCode::Success:
+	case FSCWErrorCode::OutOfMemory:
+		ShaderCompileWorkerError::HandleOutOfMemory(ExceptionInfo.GetData(), Hostname.GetData());
+		break;
+	case FSCWErrorCode::Success:
 		// Can't get here...
 		break;
 	}
@@ -1959,7 +1978,7 @@ void FShaderCompileUtilities::DoReadTaskResults(const TArray<FShaderCommonCompil
 {
 	if (OutputFile.TotalSize() == 0)
 	{
-		SCWErrorCode::HandleOutputFileEmpty(*OutputFile.GetArchiveName());
+		ShaderCompileWorkerError::HandleOutputFileEmpty(*OutputFile.GetArchiveName());
 	}
 
 	int32 OutputVersion = ShaderCompileWorkerOutputVersion;
@@ -1977,7 +1996,7 @@ void FShaderCompileUtilities::DoReadTaskResults(const TArray<FShaderCommonCompil
 	// Check for corrupted output file
 	if (FileSize > OutputFile.TotalSize())
 	{
-		SCWErrorCode::HandleOutputFileCorrupted(*OutputFile.GetArchiveName(), FileSize, OutputFile.TotalSize());
+		ShaderCompileWorkerError::HandleOutputFileCorrupted(*OutputFile.GetArchiveName(), FileSize, OutputFile.TotalSize());
 	}
 
 	int32 ErrorCode = 0;
@@ -1992,10 +2011,13 @@ void FShaderCompileUtilities::DoReadTaskResults(const TArray<FShaderCommonCompil
 	int32 ExceptionInfoLength = 0;
 	OutputFile << ExceptionInfoLength;
 
+	int32 HostnameLength = 0;
+	OutputFile << HostnameLength;
+
 	// Worker crashed
-	if (ESCWErrorCode(ErrorCode) != ESCWErrorCode::Success)
+	if (ErrorCode != FSCWErrorCode::Success)
 	{
-		HandleWorkerCrash(QueuedJobs, OutputFile, OutputVersion, FileSize, ESCWErrorCode(ErrorCode), NumProcessedJobs, CallstackLength, ExceptionInfoLength);
+		HandleWorkerCrash(QueuedJobs, OutputFile, OutputVersion, FileSize, (FSCWErrorCode::ECode)ErrorCode, NumProcessedJobs, CallstackLength, ExceptionInfoLength, HostnameLength);
 	}
 
 	TArray<FShaderCompileJob*> QueuedSingleJobs;
@@ -4636,13 +4658,16 @@ void FShaderCompilingManager::ProcessCompiledShaderMaps(
 
 	// Process compiled shader maps in FIFO order, in case a shader map has been enqueued multiple times,
 	// Which can happen if a material is edited while a background compile is going on
-	for (TMap<int32, FShaderMapFinalizeResults>::TIterator ProcessIt(CompiledShaderMaps); ProcessIt; ++ProcessIt)
+	for (TMap<int32, FShaderMapFinalizeResults>::TIterator ShaderMapResultIter(CompiledShaderMaps); ShaderMapResultIter; ++ShaderMapResultIter)
 	{
-		const uint32 CompilingId = ProcessIt.Key();
-		FShaderMapFinalizeResults& CompileResults = ProcessIt.Value();
+		const uint32 CompilingId = ShaderMapResultIter.Key();
+
+		FShaderMapFinalizeResults& CompileResults = ShaderMapResultIter.Value();
+		TArray<FShaderCommonCompileJobPtr>& FinishedJobs = CompileResults.FinishedJobs;
+
 		if (CompileResults.bSkipResultProcessing)
 		{
-			ProcessIt.RemoveCurrent();
+			ShaderMapResultIter.RemoveCurrent();
 			continue;
 		}
 
@@ -4654,12 +4679,11 @@ void FShaderCompilingManager::ProcessCompiledShaderMaps(
 
 			TArray<FString> Errors;
 			FString DumpedSource;
-			TArray<FShaderCommonCompileJobPtr>& ResultArray = CompileResults.FinishedJobs;
 
 			bool bSuccess = true;
-			for (int32 JobIndex = 0; JobIndex < ResultArray.Num(); JobIndex++)
+			for (int32 JobIndex = 0; JobIndex < FinishedJobs.Num(); JobIndex++)
 			{
-				FShaderCommonCompileJob& CurrentJob = *ResultArray[JobIndex];
+				FShaderCommonCompileJob& CurrentJob = *FinishedJobs[JobIndex];
 
 				if (FShaderCompileJob* SingleJob = CurrentJob.GetSingleShaderJob())
 				{
@@ -4681,17 +4705,17 @@ void FShaderCompilingManager::ProcessCompiledShaderMaps(
 			if (bSuccess)
 			{
 				int32 JobIndex = 0;
-				if (ResultArray.Num() > 0)
+				if (FinishedJobs.Num() > 0)
 				{
-					CompilingShaderMap->ProcessCompilationResults(ResultArray, JobIndex, TimeBudget);
+					CompilingShaderMap->ProcessCompilationResults(FinishedJobs, JobIndex, TimeBudget);
 					{
 						FScopeLock Lock(&CompileQueueSection);
 						for (int32 i = 0; i < JobIndex; ++i)
 						{
-							ReleaseJob(ResultArray[i]);
+							ReleaseJob(FinishedJobs[i]);
 						}
 					}
-					ResultArray.RemoveAt(0, JobIndex);
+					FinishedJobs.RemoveAt(0, JobIndex);
 				}
 
 				// Make a clone of the compiling shader map to use for rendering
@@ -4701,9 +4725,9 @@ void FShaderCompilingManager::ProcessCompiledShaderMaps(
 				TimeBudget -= (FPlatformTime::Seconds() - StartTime);
 			}
 
-			if (!bSuccess || ResultArray.Num() == 0)
+			if (!bSuccess || FinishedJobs.Num() == 0)
 			{
-				ProcessIt.RemoveCurrent();
+				ShaderMapResultIter.RemoveCurrent();
 			}
 
 #if DEBUG_INFINITESHADERCOMPILE
@@ -4836,33 +4860,24 @@ void FShaderCompilingManager::ProcessCompiledShaderMaps(
 				break;
 			}
 		}
-		else if (CompilingId == GlobalShaderMapId)
+		else
 		{
-			ProcessCompiledGlobalShaders(CompileResults.FinishedJobs);
-			
-			PropagateGlobalShadersToAllPrimitives();
-
+			if (CompilingId == GlobalShaderMapId)
 			{
-				FScopeLock Lock(&CompileQueueSection);
-				for (auto& Job : CompileResults.FinishedJobs)
-				{
-					ReleaseJob(Job);
-				}
+				ProcessCompiledGlobalShaders(FinishedJobs);
+				PropagateGlobalShadersToAllPrimitives();
 			}
-			ProcessIt.RemoveCurrent();
-		}
-		else
-		{
+
 			// ShaderMap was removed from compiling list or is being used by another type of shader map which is maintaining a reference
 			// to the results, either way the job can be released
 			{
 				FScopeLock Lock(&CompileQueueSection);
-				for (auto& Job : CompileResults.FinishedJobs)
+				for (FShaderCommonCompileJobPtr& Job : FinishedJobs)
 				{
 					ReleaseJob(Job);
 				}
 			}
-			ProcessIt.RemoveCurrent();
+			ShaderMapResultIter.RemoveCurrent();
 		}
 	}
 
diff --git a/Engine/Source/Runtime/RenderCore/Private/ShaderCore.cpp b/Engine/Source/Runtime/RenderCore/Private/ShaderCore.cpp
--- a/Engine/Source/Runtime/RenderCore/Private/ShaderCore.cpp
+++ b/Engine/Source/Runtime/RenderCore/Private/ShaderCore.cpp
@@ -235,6 +235,38 @@ private:
 	bool bInitialized;	
 };
 
+
+static FSCWErrorCode::ECode GSCWErrorCode = FSCWErrorCode::NotSet;
+static FString GSCWErrorCodeInfo;
+
+void FSCWErrorCode::Report(ECode Code, const FStringView& Info)
+{
+	checkf(!FSCWErrorCode::IsSet(), TEXT("ShaderCompileWorker error code already set; this is only supposed to be set once to handle a worker crash:\n%s"), Info);
+	GSCWErrorCode = Code;
+	GSCWErrorCodeInfo = Info;
+}
+
+void FSCWErrorCode::Reset()
+{
+	GSCWErrorCode = FSCWErrorCode::NotSet;
+	GSCWErrorCodeInfo.Empty();
+}
+
+FSCWErrorCode::ECode FSCWErrorCode::Get()
+{
+	return GSCWErrorCode;
+}
+
+const FString& FSCWErrorCode::GetInfo()
+{
+	return GSCWErrorCodeInfo;
+}
+
+bool FSCWErrorCode::IsSet()
+{
+	return GSCWErrorCode != FSCWErrorCode::NotSet;
+}
+
 /** Protects GShaderHashCache from simultaneous modification by multiple threads. Note that it can cover more than one method of the class, e.g. a block of code doing Find() then Add() can be guarded */
 FRWLock GShaderHashAccessRWLock;
 
diff --git a/Engine/Source/Runtime/RenderCore/Public/ShaderCompilerCore.h b/Engine/Source/Runtime/RenderCore/Public/ShaderCompilerCore.h
--- a/Engine/Source/Runtime/RenderCore/Public/ShaderCompilerCore.h
+++ b/Engine/Source/Runtime/RenderCore/Public/ShaderCompilerCore.h
@@ -496,21 +496,44 @@ struct FShaderCompilerOutput
 	}
 };
 
-enum class ESCWErrorCode
+struct RENDERCORE_API FSCWErrorCode
 {
-	NotSet = -1,
-	Success,
-	GeneralCrash,
-	BadShaderFormatVersion,
-	BadInputVersion,
-	BadSingleJobHeader,
-	BadPipelineJobHeader,
-	CantDeleteInputFile,
-	CantSaveOutputFile,
-	NoTargetShaderFormatsFound,
-	CantCompileForSpecificFormat,
-	CrashInsidePlatformCompiler,
-	BadInputFile
+	enum ECode : int32
+	{
+		NotSet = -1,
+		Success,
+		GeneralCrash,
+		BadShaderFormatVersion,
+		BadInputVersion,
+		BadSingleJobHeader,
+		BadPipelineJobHeader,
+		CantDeleteInputFile,
+		CantSaveOutputFile,
+		NoTargetShaderFormatsFound,
+		CantCompileForSpecificFormat,
+		CrashInsidePlatformCompiler,
+		BadInputFile,
+		OutOfMemory,
+	};
+
+	/**
+	Sets the global SCW error code if it hasn't been set before.
+	Call Reset first before setting a new value.
+	Returns true on success, otherwise the error code has already been set.
+	*/
+	static void Report(ECode Code, const FStringView& Info = {});
+
+	/** Resets the global SCW error code to NotSet. */
+	static void Reset();
+
+	/** Returns the global SCW error code. */
+	static ECode Get();
+
+	/** Returns the global SCW error code information string. Empty string if not set. */
+	static const FString& GetInfo();
+
+	/** Returns true if the SCW global error code has been set. Equivalent to 'Get() != NotSet'. */
+	static bool IsSet();
 };
 
 UE_DEPRECATED(5.2, "Functionality has moved to UE::ShaderCompilerCommon::ShouldUseStableConstantBuffer")
