Fix compile error.

#rnx

diff --git a/Engine/Source/Runtime/Engine/Private/ShaderCompiler/ShaderCompiler.cpp b/Engine/Source/Runtime/Engine/Private/ShaderCompiler/ShaderCompiler.cpp
--- a/Engine/Source/Runtime/Engine/Private/ShaderCompiler/ShaderCompiler.cpp
+++ b/Engine/Source/Runtime/Engine/Private/ShaderCompiler/ShaderCompiler.cpp
@@ -1956,7 +1956,7 @@ static void HandleWorkerCrash(const TArray<FShaderCommonCompileJobPtr>& QueuedJo
 	case FSCWErrorCode::CantCompileForSpecificFormat:
 		ShaderCompileWorkerError::HandleCantCompileForSpecificFormat(ExceptionInfo.GetData());
 		break;
-	case ESCWErrorCode::CrashInsidePlatformCompiler:
+	case FSCWErrorCode::CrashInsidePlatformCompiler:
 		LogQueuedCompileJobs(QueuedJobs, NumProcessedJobs);
 		ShaderCompileWorkerError::HandleCrashInsidePlatformCompiler(ExceptionInfo.GetData());
 		break;
