Fixed CIS compile issue that came from enums being moved in a different branch.

diff --git a/Engine/Source/Runtime/Engine/Classes/Engine/TextureDefines.h b/Engine/Source/Runtime/Engine/Classes/Engine/TextureDefines.h
--- a/Engine/Source/Runtime/Engine/Classes/Engine/TextureDefines.h
+++ b/Engine/Source/Runtime/Engine/Classes/Engine/TextureDefines.h
@@ -406,7 +406,7 @@ enum class ETextureChromaticAdaptationMethod : uint8
 
 
 UENUM()
-enum TextureFilter
+enum TextureFilter : int
 {
 	TF_Nearest UMETA(DisplayName="Nearest"),
 	TF_Bilinear UMETA(DisplayName="Bi-linear"),
@@ -417,7 +417,7 @@ enum TextureFilter
 };
 
 UENUM()
-enum TextureAddress
+enum TextureAddress : int
 {
 	TA_Wrap UMETA(DisplayName="Wrap"),
 	TA_Clamp UMETA(DisplayName="Clamp"),
@@ -426,7 +426,7 @@ enum TextureAddress
 };
 
 UENUM()
-enum ETextureMipCount
+enum ETextureMipCount : int
 {
 	TMC_ResidentMips,
 	TMC_AllMips,
@@ -436,7 +436,7 @@ enum ETextureMipCount
 
 // TextureCompressionQuality is used for ASTC
 UENUM()
-enum ETextureCompressionQuality
+enum ETextureCompressionQuality : int
 {
 	TCQ_Default = 0		UMETA(DisplayName="Default"),
 	TCQ_Lowest = 1		UMETA(DisplayName="Lowest (ASTC 12x12)"),
