[Backout] - CL23401851
Original CL Desc
-----------------------------------------------------------------
Add host machine name to OOM report in DXC backend and refactor global SCW error-code handling.

#rnx

diff --git a/Engine/Source/Developer/ShaderCompilerCommon/Private/ShaderConductorContext.cpp b/Engine/Source/Developer/ShaderCompilerCommon/Private/ShaderConductorContext.cpp
--- a/Engine/Source/Developer/ShaderCompilerCommon/Private/ShaderConductorContext.cpp
+++ b/Engine/Source/Developer/ShaderCompilerCommon/Private/ShaderConductorContext.cpp
@@ -10,6 +10,8 @@ THIRD_PARTY_INCLUDES_END
 #endif
 
 
+ESCWErrorCode GSCWErrorCode = ESCWErrorCode::NotSet;
+
 namespace CrossCompiler
 {
 
diff --git a/Engine/Source/Developer/ShaderCompilerCommon/Public/ShaderConductorContext.h b/Engine/Source/Developer/ShaderCompilerCommon/Public/ShaderConductorContext.h
--- a/Engine/Source/Developer/ShaderCompilerCommon/Public/ShaderConductorContext.h
+++ b/Engine/Source/Developer/ShaderCompilerCommon/Public/ShaderConductorContext.h
@@ -246,3 +246,6 @@ namespace CrossCompiler
 	};
 
 }
+
+// Error code for SCW to help track down crashes
+extern SHADERCOMPILERCOMMON_API ESCWErrorCode GSCWErrorCode;
diff --git a/Engine/Source/Developer/Windows/ShaderFormatD3D/Private/D3DShaderCompiler.cpp b/Engine/Source/Developer/Windows/ShaderFormatD3D/Private/D3DShaderCompiler.cpp
--- a/Engine/Source/Developer/Windows/ShaderFormatD3D/Private/D3DShaderCompiler.cpp
+++ b/Engine/Source/Developer/Windows/ShaderFormatD3D/Private/D3DShaderCompiler.cpp
@@ -483,7 +483,7 @@ static HRESULT D3DCompileWrapper(
 #if !PLATFORM_SEH_EXCEPTIONS_DISABLED
 	__except(D3DExceptionFilter(bCatchException))
 	{
-		FSCWErrorCode::Report(FSCWErrorCode::CrashInsidePlatformCompiler);
+		GSCWErrorCode = ESCWErrorCode::CrashInsidePlatformCompiler;
 		return E_FAIL;
 	}
 #endif
diff --git a/Engine/Source/Developer/Windows/ShaderFormatD3D/Private/D3DShaderCompilerDXC.cpp b/Engine/Source/Developer/Windows/ShaderFormatD3D/Private/D3DShaderCompilerDXC.cpp
--- a/Engine/Source/Developer/Windows/ShaderFormatD3D/Private/D3DShaderCompilerDXC.cpp
+++ b/Engine/Source/Developer/Windows/ShaderFormatD3D/Private/D3DShaderCompilerDXC.cpp
@@ -10,7 +10,6 @@
 #include "HAL/FileManager.h"
 #include "Serialization/MemoryWriter.h"
 #include "RayTracingDefinitions.h"
-#include "Math/UnitConversion.h"
 
 DEFINE_LOG_CATEGORY_STATIC(LogD3D12ShaderCompiler, Log, All);
 
@@ -128,12 +127,11 @@ static void LogFailedHRESULT(const TCHAR* FailedExpressionStr, HRESULT Result)
 {
 	if (Result == E_OUTOFMEMORY)
 	{
+		// Provide currently availble system memory to verify out-of-memory error is legitimate
 		const FPlatformMemoryStats MemoryStats = FPlatformMemory::GetStats();
-		const FString AvailableMemoryStr = FString::Printf(TEXT("%d MB"), FUnitConversion::Convert(MemoryStats.AvailablePhysical, EUnit::Bytes, EUnit::Megabytes));
-		const FString UsedMemoryStr = FString::Printf(TEXT("%d MB"), FUnitConversion::Convert(MemoryStats.UsedPhysical, EUnit::Bytes, EUnit::Megabytes));
-		const FString ErrorReport = FString::Printf(TEXT("%s failed: Result=0x%08x (E_OUTOFMEMORY); Used physical memory %s of %s"), FailedExpressionStr, Result, *UsedMemoryStr, *AvailableMemoryStr);
-		FSCWErrorCode::Report(FSCWErrorCode::OutOfMemory, ErrorReport);
-		UE_LOG(LogD3D12ShaderCompiler, Fatal, TEXT("%s"), *ErrorReport);
+		FString SystemMemoryStr = FString::Printf(TEXT("%d MB"), MemoryStats.AvailablePhysical / 1024 / 1024);
+		FString UsedMemoryStr = FString::Printf(TEXT("%d MB"), MemoryStats.UsedPhysical / 1024 / 1024);
+		UE_LOG(LogD3D12ShaderCompiler, Fatal, TEXT("%s failed: Result=0x%08x (E_OUTOFMEMORY); System memory %s, used %s"), FailedExpressionStr, Result, *SystemMemoryStr, *UsedMemoryStr);
 	}
 	else if (const TCHAR* ErrorCodeStr = DxcErrorCodeToString(Result))
 	{
@@ -159,58 +157,55 @@ static void LogFailedHRESULT(const TCHAR* FailedExpressionStr, HRESULT Result)
 
 class FDxcMalloc final : public IMalloc
 {
-	ULONG RefCount = 1;
+	int32 RefCount = 1;
 
 public:
 
 	// IMalloc
 
-	void* STDCALL Alloc(SIZE_T cb) override
+	virtual void* STDMETHODCALLTYPE Alloc(SIZE_T cb) final override
 	{
-		cb = FMath::Max(SIZE_T(1), cb);
 		return FMemory::Malloc(cb);
 	}
 
-	void* STDCALL Realloc(void* pv, SIZE_T cb) override
+	virtual void* STDMETHODCALLTYPE Realloc(void* pv, SIZE_T cb) final override
 	{
-		cb = FMath::Max(SIZE_T(1), cb);
 		return FMemory::Realloc(pv, cb);
 	}
 
-	void STDCALL Free(void* pv) override
+	virtual void STDMETHODCALLTYPE Free(void* pv) final override
 	{
 		return FMemory::Free(pv);
 	}
 
-	SIZE_T STDCALL GetSize(void* pv) override
+	virtual SIZE_T STDMETHODCALLTYPE GetSize(void* pv) final override
 	{
 		return FMemory::GetAllocSize(pv);
 	}
 
-	int STDCALL DidAlloc(void* pv) override
+	virtual int STDMETHODCALLTYPE DidAlloc(void* pv) final override
 	{
 		return 1; // assume that all allocation queries coming from DXC belong to our allocator
 	}
 
-	void STDCALL HeapMinimize() override
+	virtual void STDMETHODCALLTYPE HeapMinimize() final override
 	{
 		// nothing
 	}
 
 	// IUnknown
 
-	ULONG STDCALL AddRef() override
+	ULONG STDMETHODCALLTYPE AddRef()
 	{
 		return ++RefCount;
 	}
 
-	ULONG STDCALL Release() override
-	{
-		check(RefCount > 0);
+	ULONG STDMETHODCALLTYPE Release() {
+		check(RefCount != 0);
 		return --RefCount;
 	}
 
-	HRESULT STDCALL QueryInterface(REFIID iid, void** ppvObject) override
+	STDMETHODIMP QueryInterface(REFIID iid, void** ppvObject)
 	{
 		checkNoEntry(); // We do not expect or support QI on DXC allocator replacement
 		return ERROR_NOINTERFACE;
@@ -323,7 +318,7 @@ static HRESULT DXCCompileWrapper(
 
 	if (bExceptionError)
 	{
-		FSCWErrorCode::Report(FSCWErrorCode::CrashInsidePlatformCompiler);
+		GSCWErrorCode = ESCWErrorCode::CrashInsidePlatformCompiler;
 
 		FString ErrorMsg = TEXT("Internal error or exception inside dxcompiler.dll\n");
 		ErrorMsg += GDxcStackTrace;
@@ -606,7 +601,7 @@ static bool DXCRewriteWrapper(
 #if !PLATFORM_SEH_EXCEPTIONS_DISABLED
 	__except (EXCEPTION_EXECUTE_HANDLER)
 	{
-		FSCWErrorCode::Report(FSCWErrorCode::CrashInsidePlatformCompiler);
+		GSCWErrorCode = ESCWErrorCode::CrashInsidePlatformCompiler;
 		FMemory::Memzero(OutResultDesc);
 		bOutException = true;
 		return false;
diff --git a/Engine/Source/Programs/ShaderCompileWorker/Private/ShaderCompileWorker.cpp b/Engine/Source/Programs/ShaderCompileWorker/Private/ShaderCompileWorker.cpp
--- a/Engine/Source/Programs/ShaderCompileWorker/Private/ShaderCompileWorker.cpp
+++ b/Engine/Source/Programs/ShaderCompileWorker/Private/ShaderCompileWorker.cpp
@@ -15,7 +15,6 @@
 #include "RHIShaderFormatDefinitions.inl"
 #include "ShaderCompilerCommon.h"
 #include "Serialization/MemoryReader.h"
-#include "SocketSubsystem.h"
 
 #define DEBUG_USING_CONSOLE	0
 
@@ -36,7 +35,7 @@ inline bool IsUsingXGE()
 	return GXGEMode != EXGEMode::None;
 }
 
-static FSCWErrorCode::ECode GFailedErrorCode = FSCWErrorCode::Success;
+static ESCWErrorCode GFailedErrorCode = ESCWErrorCode::Success;
 
 static void OnXGEJobCompleted(const TCHAR* WorkingDirectory)
 {
@@ -49,10 +48,10 @@ static void OnXGEJobCompleted(const TCHAR* WorkingDirectory)
 }
 
 #if USING_CODE_ANALYSIS
-	UE_NORETURN static inline void ExitWithoutCrash(FSCWErrorCode::ECode ErrorCode, const FString& Message);
+	UE_NORETURN static inline void ExitWithoutCrash(ESCWErrorCode ErrorCode, const FString& Message);
 #endif
 
-static inline void ExitWithoutCrash(FSCWErrorCode::ECode ErrorCode, const FString& Message)
+static inline void ExitWithoutCrash(ESCWErrorCode ErrorCode, const FString& Message)
 {
 	GFailedErrorCode = ErrorCode;
 	FCString::Snprintf(GErrorExceptionDescription, sizeof(GErrorExceptionDescription), TEXT("%s"), *Message);
@@ -74,7 +73,7 @@ static const TArray<const IShaderFormat*>& GetShaderFormats()
 
 		if (!Modules.Num())
 		{
-			ExitWithoutCrash(FSCWErrorCode::NoTargetShaderFormatsFound, TEXT("No target shader formats found!"));
+			ExitWithoutCrash(ESCWErrorCode::NoTargetShaderFormatsFound, TEXT("No target shader formats found!"));
 		}
 
 		for (int32 Index = 0; Index < Modules.Num(); Index++)
@@ -190,14 +189,14 @@ static const IShaderFormat* FindShaderFormat(FName Name)
 
 
 /** Processes a compilation job. */
-static FORCENOINLINE void ProcessCompilationJob(const FShaderCompilerInput& Input,FShaderCompilerOutput& Output,const FString& WorkingDirectory)
+static void ProcessCompilationJob(const FShaderCompilerInput& Input,FShaderCompilerOutput& Output,const FString& WorkingDirectory)
 {
 	TRACE_CPUPROFILER_EVENT_SCOPE(ProcessCompilationJob);
 
 	const IShaderFormat* Compiler = FindShaderFormat(Input.ShaderFormat);
 	if (!Compiler)
 	{
-		ExitWithoutCrash(FSCWErrorCode::CantCompileForSpecificFormat, FString::Printf(TEXT("Can't compile shaders for format %s"), *Input.ShaderFormat.ToString()));
+		ExitWithoutCrash(ESCWErrorCode::CantCompileForSpecificFormat, FString::Printf(TEXT("Can't compile shaders for format %s"), *Input.ShaderFormat.ToString()));
 	}
 
 	// Apply the console variable values from the input environment before calling the platform shader compiler
@@ -259,20 +258,6 @@ static void UpdateFileSize(FArchive& OutputFile, int64 FileSizePosition)
 	OutputFile.Seek(Current);
 };
 
-static const TCHAR* GetLocalHostname(int32* OutHostnameLength = nullptr)
-{
-	static FString Hostname;
-	if (Hostname.IsEmpty())
-	{
-		ISocketSubsystem::Get()->GetHostName(Hostname);
-	}
-	if (OutHostnameLength)
-	{
-		*OutHostnameLength = Hostname.Len();
-	}
-	return *Hostname;
-}
-
 static int64 WriteOutputFileHeader(FArchive& OutputFile, int32 ErrorCode, int32 CallstackLength, const TCHAR* Callstack,
 	int32 ExceptionInfoLength, const TCHAR* ExceptionInfo)
 {
@@ -296,10 +281,6 @@ static int64 WriteOutputFileHeader(FArchive& OutputFile, int32 ErrorCode, int32
 
 	OutputFile << ExceptionInfoLength;
 
-	int32 HostnameLength = 0;
-	const TCHAR* Hostname = GetLocalHostname(&HostnameLength);
-	OutputFile << HostnameLength;
-
 	if (CallstackLength > 0)
 	{
 		OutputFile.Serialize((void*)Callstack, CallstackLength * sizeof(TCHAR));
@@ -310,11 +291,6 @@ static int64 WriteOutputFileHeader(FArchive& OutputFile, int32 ErrorCode, int32
 		OutputFile.Serialize((void*)ExceptionInfo, ExceptionInfoLength * sizeof(TCHAR));
 	}
 
-	if (HostnameLength > 0)
-	{
-		OutputFile.Serialize((void*)Hostname, HostnameLength * sizeof(TCHAR));
-	}
-
 	UpdateFileSize(OutputFile, FileSizePosition);
 	return FileSizePosition;
 }
@@ -485,7 +461,7 @@ private:
 			{
 				if (Pair.Value != *Found)
 				{
-					ExitWithoutCrash(FSCWErrorCode::BadShaderFormatVersion, FString::Printf(TEXT("Mismatched shader version for format %s: Found version %u but expected %u; did you forget to build ShaderCompilerWorker?"), *Pair.Key, *Found, Pair.Value));
+					ExitWithoutCrash(ESCWErrorCode::BadShaderFormatVersion, FString::Printf(TEXT("Mismatched shader version for format %s: Found version %u but expected %u; did you forget to build ShaderCompilerWorker?"), *Pair.Key, *Found, Pair.Value));
 				}
 			}
 		}
@@ -499,7 +475,7 @@ private:
 		*InputFilePtr << InputVersion;
 		if (ShaderCompileWorkerInputVersion != InputVersion)
 		{
-			ExitWithoutCrash(FSCWErrorCode::BadInputVersion, FString::Printf(TEXT("Exiting due to ShaderCompilerWorker expecting input version %d, got %d instead! Did you forget to build ShaderCompilerWorker?"), ShaderCompileWorkerInputVersion, InputVersion));
+			ExitWithoutCrash(ESCWErrorCode::BadInputVersion, FString::Printf(TEXT("Exiting due to ShaderCompilerWorker expecting input version %d, got %d instead! Did you forget to build ShaderCompilerWorker?"), ShaderCompileWorkerInputVersion, InputVersion));
 		}
 
 		FString CompressionFormatString;
@@ -516,7 +492,7 @@ private:
 
 			if (UncompressedDataSize == 0)
 			{
-				ExitWithoutCrash(FSCWErrorCode::BadInputFile, TEXT("Exiting due to bad input file to ShaderCompilerWorker (uncompressed size is 0)! Did you forget to build ShaderCompilerWorker?"));
+				ExitWithoutCrash(ESCWErrorCode::BadInputFile, TEXT("Exiting due to bad input file to ShaderCompilerWorker (uncompressed size is 0)! Did you forget to build ShaderCompilerWorker?"));
 				// unreachable
 				return;
 			}
@@ -526,7 +502,7 @@ private:
 			*InputFilePtr << CompressedData;
 			if (!FCompression::UncompressMemory(CompressionFormat, UncompressedData.GetData(), UncompressedDataSize, CompressedData.GetData(), CompressedData.Num()))
 			{
-				ExitWithoutCrash(FSCWErrorCode::BadInputFile, FString::Printf(TEXT("Exiting due to bad input file to ShaderCompilerWorker (cannot uncompress with the format %s)! Did you forget to build ShaderCompilerWorker?"), *CompressionFormatString));
+				ExitWithoutCrash(ESCWErrorCode::BadInputFile, FString::Printf(TEXT("Exiting due to bad input file to ShaderCompilerWorker (cannot uncompress with the format %s)! Did you forget to build ShaderCompilerWorker?"), *CompressionFormatString));
 				// unreachable
 				return;
 			}
@@ -739,7 +715,7 @@ private:
 			InputFile << SingleJobHeader;
 			if (ShaderCompileWorkerSingleJobHeader != SingleJobHeader)
 			{
-				ExitWithoutCrash(FSCWErrorCode::BadSingleJobHeader, FString::Printf(TEXT("Exiting due to ShaderCompilerWorker expecting job header %d, got %d instead! Did you forget to build ShaderCompilerWorker?"), ShaderCompileWorkerSingleJobHeader, SingleJobHeader));
+				ExitWithoutCrash(ESCWErrorCode::BadSingleJobHeader, FString::Printf(TEXT("Exiting due to ShaderCompilerWorker expecting job header %d, got %d instead! Did you forget to build ShaderCompilerWorker?"), ShaderCompileWorkerSingleJobHeader, SingleJobHeader));
 			}
 
 			int32 NumBatches = 0;
@@ -784,7 +760,7 @@ private:
 			InputFile << PipelineJobHeader;
 			if (ShaderCompileWorkerPipelineJobHeader != PipelineJobHeader)
 			{
-				ExitWithoutCrash(FSCWErrorCode::BadPipelineJobHeader, FString::Printf(TEXT("Exiting due to ShaderCompilerWorker expecting pipeline job header %d, got %d instead! Did you forget to build ShaderCompilerWorker?"), ShaderCompileWorkerSingleJobHeader, PipelineJobHeader));
+				ExitWithoutCrash(ESCWErrorCode::BadPipelineJobHeader, FString::Printf(TEXT("Exiting due to ShaderCompilerWorker expecting pipeline job header %d, got %d instead! Did you forget to build ShaderCompilerWorker?"), ShaderCompileWorkerSingleJobHeader, PipelineJobHeader));
 			}
 
 			int32 NumPipelines = 0;
@@ -898,7 +874,7 @@ private:
 
 			if (!bResult)
 			{
-				ExitWithoutCrash(FSCWErrorCode::CantDeleteInputFile, FString::Printf(TEXT("Couldn't delete input file %s, is it readonly?"), *InputFilePath));
+				ExitWithoutCrash(ESCWErrorCode::CantDeleteInputFile, FString::Printf(TEXT("Couldn't delete input file %s, is it readonly?"), *InputFilePath));
 			}
 		}
 
@@ -922,7 +898,7 @@ private:
 			
 		if (!OutputFilePtr)
 		{
-			ExitWithoutCrash(FSCWErrorCode::CantSaveOutputFile, FString::Printf(TEXT("Couldn't save output file %s"), *TempFilePath));
+			ExitWithoutCrash(ESCWErrorCode::CantSaveOutputFile, FString::Printf(TEXT("Couldn't save output file %s"), *TempFilePath));
 		}
 
 		return OutputFilePtr;
@@ -933,9 +909,7 @@ private:
 		TRACE_CPUPROFILER_EVENT_SCOPE(WriteToOutputArchive);
 
 		FArchive& OutputFile = *OutputFilePtr;
-		const int64 FileSizePosition = FSCWErrorCode::IsSet()
-			? WriteOutputFileHeader(OutputFile, FSCWErrorCode::Get(), 0, nullptr, FSCWErrorCode::GetInfo().Len(), *FSCWErrorCode::GetInfo())
-			: WriteOutputFileHeader(OutputFile, FSCWErrorCode::Success, 0, nullptr, 0, nullptr);
+		int64 FileSizePosition = WriteOutputFileHeader(OutputFile, (int32)ESCWErrorCode::Success, 0, nullptr, 0, nullptr);
 
 		{
 			int32 SingleJobHeader = ShaderCompileWorkerSingleJobHeader;
@@ -1413,21 +1387,21 @@ static int32 GuardedMainWrapper(int32 ArgC, TCHAR* ArgV[], const TCHAR* CrashOut
 		{
 			FArchive& OutputFile = *IFileManager::Get().CreateFileWriter(CrashOutputFile, FILEWRITE_EvenIfReadOnly);
 
-			if (GFailedErrorCode == FSCWErrorCode::Success)
+			if (GFailedErrorCode == ESCWErrorCode::Success)
 			{
-				if (FSCWErrorCode::IsSet())
+				if (GSCWErrorCode != ESCWErrorCode::NotSet)
 				{
 					// Use the value set inside the shader format
-					GFailedErrorCode = FSCWErrorCode::Get();
+					GFailedErrorCode = GSCWErrorCode;
 				}
 				else
 				{
 					// Something else failed before we could set the error code, so mark it as a General Crash
-					GFailedErrorCode = FSCWErrorCode::GeneralCrash;
+					GFailedErrorCode = ESCWErrorCode::GeneralCrash;
 				}
 			}
-
-			int64 FileSizePosition = WriteOutputFileHeader(OutputFile, GFailedErrorCode, ExceptionCallStack.Len(), *ExceptionCallStack,
+			
+			int64 FileSizePosition = WriteOutputFileHeader(OutputFile, (int32)GFailedErrorCode, ExceptionCallStack.Len(), *ExceptionCallStack,
 				ExceptionMsg.Len(), *ExceptionMsg);
 
 			int32 NumBatches = 0;
diff --git a/Engine/Source/Programs/ShaderCompileWorker/ShaderCompileWorker.Build.cs b/Engine/Source/Programs/ShaderCompileWorker/ShaderCompileWorker.Build.cs
--- a/Engine/Source/Programs/ShaderCompileWorker/ShaderCompileWorker.Build.cs
+++ b/Engine/Source/Programs/ShaderCompileWorker/ShaderCompileWorker.Build.cs
@@ -16,8 +16,7 @@ public class ShaderCompileWorker : ModuleRules
 				"TargetPlatform",
 				"ApplicationCore",
 				"TraceLog",
-				"ShaderCompilerCommon",
-				"Sockets",
+				"ShaderCompilerCommon"
 			});
 
 		PrivateIncludePathModuleNames.AddRange(
diff --git a/Engine/Source/Runtime/Engine/Private/ShaderCompiler/ShaderCompiler.cpp b/Engine/Source/Runtime/Engine/Private/ShaderCompiler/ShaderCompiler.cpp
--- a/Engine/Source/Runtime/Engine/Private/ShaderCompiler/ShaderCompiler.cpp
+++ b/Engine/Source/Runtime/Engine/Private/ShaderCompiler/ShaderCompiler.cpp
@@ -1245,7 +1245,7 @@ namespace ShaderCompilerCookStats
 #endif
 
 // Make functions so the crash reporter can disambiguate the actual error because of the different callstacks
-namespace ShaderCompileWorkerError
+namespace SCWErrorCode
 {
 	void HandleGeneralCrash(const TCHAR* ExceptionInfo, const TCHAR* Callstack)
 	{
@@ -1307,16 +1307,6 @@ namespace ShaderCompileWorkerError
 	{
 		ModalErrorOrLog(FString::Printf(TEXT("Crash inside the platform compiler!\n%s"), Data));
 	}
-
-	void HandleBadInputFile(const TCHAR* Data)
-	{
-		ModalErrorOrLog(FString::Printf(TEXT("ShaderCompileWorker failed with bad-input-file exception:\n%s\n"), Data));
-	}
-
-	void HandleOutOfMemory(const TCHAR* ExceptionInfo, const TCHAR* Hostname)
-	{
-		ModalErrorOrLog(FString::Printf(TEXT("ShaderCompileWorker failed with out-of-memory exception on machine \"%s\"\n%s\n"), Hostname, ExceptionInfo));
-	}
 }
 
 static TMap<FString, uint32> GetFormatVersionMap()
@@ -1861,22 +1851,17 @@ static void LogQueuedCompileJobs(const TArray<FShaderCommonCompileJobPtr>& Queue
 
 // Disable optimization for this crash handler to get full access to the entire stack frame when debugging a crash dump
 UE_DISABLE_OPTIMIZATION_SHIP
-static void HandleWorkerCrash(const TArray<FShaderCommonCompileJobPtr>& QueuedJobs, FArchive& OutputFile, int32 OutputVersion, int64 FileSize, FSCWErrorCode::ECode ErrorCode, int32 NumProcessedJobs, int32 CallstackLength, int32 ExceptionInfoLength, int32 HostnameLength)
+static void HandleWorkerCrash(const TArray<FShaderCommonCompileJobPtr>& QueuedJobs, FArchive& OutputFile, int32 OutputVersion, int64 FileSize, ESCWErrorCode ErrorCode, int32 NumProcessedJobs, int32 CallstackLength, int32 ExceptionInfoLength)
 {
 	TArray<TCHAR> Callstack;
 	Callstack.AddUninitialized(CallstackLength + 1);
 	OutputFile.Serialize(Callstack.GetData(), CallstackLength * sizeof(TCHAR));
-	Callstack[CallstackLength] = TEXT('\0');
+	Callstack[CallstackLength] = 0;
 
 	TArray<TCHAR> ExceptionInfo;
 	ExceptionInfo.AddUninitialized(ExceptionInfoLength + 1);
 	OutputFile.Serialize(ExceptionInfo.GetData(), ExceptionInfoLength * sizeof(TCHAR));
-	ExceptionInfo[ExceptionInfoLength] = TEXT('\0');
-
-	TArray<TCHAR> Hostname;
-	Hostname.AddUninitialized(HostnameLength + 1);
-	OutputFile.Serialize(Hostname.GetData(), HostnameLength * sizeof(TCHAR));
-	Hostname[HostnameLength] = TEXT('\0');
+	ExceptionInfo[ExceptionInfoLength] = 0;
 
 	// Store primary job information onto stack to make it part of a crash dump
 	static const int32 MaxNumCharsForSourcePaths = 8192;
@@ -1928,45 +1913,41 @@ static void HandleWorkerCrash(const TArray<FShaderCommonCompileJobPtr>& QueuedJo
 	switch (ErrorCode)
 	{
 	default:
-	case FSCWErrorCode::GeneralCrash:
+	case ESCWErrorCode::GeneralCrash:
+	{
 		LogQueuedCompileJobs(QueuedJobs, NumProcessedJobs);
-		ShaderCompileWorkerError::HandleGeneralCrash(ExceptionInfo.GetData(), Callstack.GetData());
-		break;
-	case FSCWErrorCode::BadShaderFormatVersion:
-		ShaderCompileWorkerError::HandleBadShaderFormatVersion(ExceptionInfo.GetData());
+		SCWErrorCode::HandleGeneralCrash(ExceptionInfo.GetData(), Callstack.GetData());
+	}
+	break;
+	case ESCWErrorCode::BadShaderFormatVersion:
+		SCWErrorCode::HandleBadShaderFormatVersion(ExceptionInfo.GetData());
 		break;
-	case FSCWErrorCode::BadInputVersion:
-		ShaderCompileWorkerError::HandleBadInputVersion(ExceptionInfo.GetData());
+	case ESCWErrorCode::BadInputVersion:
+		SCWErrorCode::HandleBadInputVersion(ExceptionInfo.GetData());
 		break;
-	case FSCWErrorCode::BadSingleJobHeader:
-		ShaderCompileWorkerError::HandleBadSingleJobHeader(ExceptionInfo.GetData());
+	case ESCWErrorCode::BadSingleJobHeader:
+		SCWErrorCode::HandleBadSingleJobHeader(ExceptionInfo.GetData());
 		break;
-	case FSCWErrorCode::BadPipelineJobHeader:
-		ShaderCompileWorkerError::HandleBadPipelineJobHeader(ExceptionInfo.GetData());
+	case ESCWErrorCode::BadPipelineJobHeader:
+		SCWErrorCode::HandleBadPipelineJobHeader(ExceptionInfo.GetData());
 		break;
-	case FSCWErrorCode::CantDeleteInputFile:
-		ShaderCompileWorkerError::HandleCantDeleteInputFile(ExceptionInfo.GetData());
+	case ESCWErrorCode::CantDeleteInputFile:
+		SCWErrorCode::HandleCantDeleteInputFile(ExceptionInfo.GetData());
 		break;
-	case FSCWErrorCode::CantSaveOutputFile:
-		ShaderCompileWorkerError::HandleCantSaveOutputFile(ExceptionInfo.GetData());
+	case ESCWErrorCode::CantSaveOutputFile:
+		SCWErrorCode::HandleCantSaveOutputFile(ExceptionInfo.GetData());
 		break;
-	case FSCWErrorCode::NoTargetShaderFormatsFound:
-		ShaderCompileWorkerError::HandleNoTargetShaderFormatsFound(ExceptionInfo.GetData());
+	case ESCWErrorCode::NoTargetShaderFormatsFound:
+		SCWErrorCode::HandleNoTargetShaderFormatsFound(ExceptionInfo.GetData());
 		break;
-	case FSCWErrorCode::CantCompileForSpecificFormat:
-		ShaderCompileWorkerError::HandleCantCompileForSpecificFormat(ExceptionInfo.GetData());
+	case ESCWErrorCode::CantCompileForSpecificFormat:
+		SCWErrorCode::HandleCantCompileForSpecificFormat(ExceptionInfo.GetData());
 		break;
-	case FSCWErrorCode::CrashInsidePlatformCompiler:
+	case ESCWErrorCode::CrashInsidePlatformCompiler:
 		LogQueuedCompileJobs(QueuedJobs, NumProcessedJobs);
-		ShaderCompileWorkerError::HandleCrashInsidePlatformCompiler(ExceptionInfo.GetData());
-		break;
-	case FSCWErrorCode::BadInputFile:
-		ShaderCompileWorkerError::HandleBadInputFile(ExceptionInfo.GetData());
+		SCWErrorCode::HandleCrashInsidePlatformCompiler(ExceptionInfo.GetData());
 		break;
-	case FSCWErrorCode::OutOfMemory:
-		ShaderCompileWorkerError::HandleOutOfMemory(ExceptionInfo.GetData(), Hostname.GetData());
-		break;
-	case FSCWErrorCode::Success:
+	case ESCWErrorCode::Success:
 		// Can't get here...
 		break;
 	}
@@ -1978,7 +1959,7 @@ void FShaderCompileUtilities::DoReadTaskResults(const TArray<FShaderCommonCompil
 {
 	if (OutputFile.TotalSize() == 0)
 	{
-		ShaderCompileWorkerError::HandleOutputFileEmpty(*OutputFile.GetArchiveName());
+		SCWErrorCode::HandleOutputFileEmpty(*OutputFile.GetArchiveName());
 	}
 
 	int32 OutputVersion = ShaderCompileWorkerOutputVersion;
@@ -1996,7 +1977,7 @@ void FShaderCompileUtilities::DoReadTaskResults(const TArray<FShaderCommonCompil
 	// Check for corrupted output file
 	if (FileSize > OutputFile.TotalSize())
 	{
-		ShaderCompileWorkerError::HandleOutputFileCorrupted(*OutputFile.GetArchiveName(), FileSize, OutputFile.TotalSize());
+		SCWErrorCode::HandleOutputFileCorrupted(*OutputFile.GetArchiveName(), FileSize, OutputFile.TotalSize());
 	}
 
 	int32 ErrorCode = 0;
@@ -2011,13 +1992,10 @@ void FShaderCompileUtilities::DoReadTaskResults(const TArray<FShaderCommonCompil
 	int32 ExceptionInfoLength = 0;
 	OutputFile << ExceptionInfoLength;
 
-	int32 HostnameLength = 0;
-	OutputFile << HostnameLength;
-
 	// Worker crashed
-	if (ErrorCode != FSCWErrorCode::Success)
+	if (ESCWErrorCode(ErrorCode) != ESCWErrorCode::Success)
 	{
-		HandleWorkerCrash(QueuedJobs, OutputFile, OutputVersion, FileSize, (FSCWErrorCode::ECode)ErrorCode, NumProcessedJobs, CallstackLength, ExceptionInfoLength, HostnameLength);
+		HandleWorkerCrash(QueuedJobs, OutputFile, OutputVersion, FileSize, ESCWErrorCode(ErrorCode), NumProcessedJobs, CallstackLength, ExceptionInfoLength);
 	}
 
 	TArray<FShaderCompileJob*> QueuedSingleJobs;
@@ -4658,16 +4636,13 @@ void FShaderCompilingManager::ProcessCompiledShaderMaps(
 
 	// Process compiled shader maps in FIFO order, in case a shader map has been enqueued multiple times,
 	// Which can happen if a material is edited while a background compile is going on
-	for (TMap<int32, FShaderMapFinalizeResults>::TIterator ShaderMapResultIter(CompiledShaderMaps); ShaderMapResultIter; ++ShaderMapResultIter)
+	for (TMap<int32, FShaderMapFinalizeResults>::TIterator ProcessIt(CompiledShaderMaps); ProcessIt; ++ProcessIt)
 	{
-		const uint32 CompilingId = ShaderMapResultIter.Key();
-
-		FShaderMapFinalizeResults& CompileResults = ShaderMapResultIter.Value();
-		TArray<FShaderCommonCompileJobPtr>& FinishedJobs = CompileResults.FinishedJobs;
-
+		const uint32 CompilingId = ProcessIt.Key();
+		FShaderMapFinalizeResults& CompileResults = ProcessIt.Value();
 		if (CompileResults.bSkipResultProcessing)
 		{
-			ShaderMapResultIter.RemoveCurrent();
+			ProcessIt.RemoveCurrent();
 			continue;
 		}
 
@@ -4679,11 +4654,12 @@ void FShaderCompilingManager::ProcessCompiledShaderMaps(
 
 			TArray<FString> Errors;
 			FString DumpedSource;
+			TArray<FShaderCommonCompileJobPtr>& ResultArray = CompileResults.FinishedJobs;
 
 			bool bSuccess = true;
-			for (int32 JobIndex = 0; JobIndex < FinishedJobs.Num(); JobIndex++)
+			for (int32 JobIndex = 0; JobIndex < ResultArray.Num(); JobIndex++)
 			{
-				FShaderCommonCompileJob& CurrentJob = *FinishedJobs[JobIndex];
+				FShaderCommonCompileJob& CurrentJob = *ResultArray[JobIndex];
 
 				if (FShaderCompileJob* SingleJob = CurrentJob.GetSingleShaderJob())
 				{
@@ -4705,17 +4681,17 @@ void FShaderCompilingManager::ProcessCompiledShaderMaps(
 			if (bSuccess)
 			{
 				int32 JobIndex = 0;
-				if (FinishedJobs.Num() > 0)
+				if (ResultArray.Num() > 0)
 				{
-					CompilingShaderMap->ProcessCompilationResults(FinishedJobs, JobIndex, TimeBudget);
+					CompilingShaderMap->ProcessCompilationResults(ResultArray, JobIndex, TimeBudget);
 					{
 						FScopeLock Lock(&CompileQueueSection);
 						for (int32 i = 0; i < JobIndex; ++i)
 						{
-							ReleaseJob(FinishedJobs[i]);
+							ReleaseJob(ResultArray[i]);
 						}
 					}
-					FinishedJobs.RemoveAt(0, JobIndex);
+					ResultArray.RemoveAt(0, JobIndex);
 				}
 
 				// Make a clone of the compiling shader map to use for rendering
@@ -4725,9 +4701,9 @@ void FShaderCompilingManager::ProcessCompiledShaderMaps(
 				TimeBudget -= (FPlatformTime::Seconds() - StartTime);
 			}
 
-			if (!bSuccess || FinishedJobs.Num() == 0)
+			if (!bSuccess || ResultArray.Num() == 0)
 			{
-				ShaderMapResultIter.RemoveCurrent();
+				ProcessIt.RemoveCurrent();
 			}
 
 #if DEBUG_INFINITESHADERCOMPILE
@@ -4860,24 +4836,33 @@ void FShaderCompilingManager::ProcessCompiledShaderMaps(
 				break;
 			}
 		}
-		else
+		else if (CompilingId == GlobalShaderMapId)
 		{
-			if (CompilingId == GlobalShaderMapId)
+			ProcessCompiledGlobalShaders(CompileResults.FinishedJobs);
+			
+			PropagateGlobalShadersToAllPrimitives();
+
 			{
-				ProcessCompiledGlobalShaders(FinishedJobs);
-				PropagateGlobalShadersToAllPrimitives();
+				FScopeLock Lock(&CompileQueueSection);
+				for (auto& Job : CompileResults.FinishedJobs)
+				{
+					ReleaseJob(Job);
+				}
 			}
-
+			ProcessIt.RemoveCurrent();
+		}
+		else
+		{
 			// ShaderMap was removed from compiling list or is being used by another type of shader map which is maintaining a reference
 			// to the results, either way the job can be released
 			{
 				FScopeLock Lock(&CompileQueueSection);
-				for (FShaderCommonCompileJobPtr& Job : FinishedJobs)
+				for (auto& Job : CompileResults.FinishedJobs)
 				{
 					ReleaseJob(Job);
 				}
 			}
-			ShaderMapResultIter.RemoveCurrent();
+			ProcessIt.RemoveCurrent();
 		}
 	}
 
diff --git a/Engine/Source/Runtime/RenderCore/Private/ShaderCore.cpp b/Engine/Source/Runtime/RenderCore/Private/ShaderCore.cpp
--- a/Engine/Source/Runtime/RenderCore/Private/ShaderCore.cpp
+++ b/Engine/Source/Runtime/RenderCore/Private/ShaderCore.cpp
@@ -235,38 +235,6 @@ private:
 	bool bInitialized;	
 };
 
-
-static FSCWErrorCode::ECode GSCWErrorCode = FSCWErrorCode::NotSet;
-static FString GSCWErrorCodeInfo;
-
-void FSCWErrorCode::Report(ECode Code, const FStringView& Info)
-{
-	checkf(!FSCWErrorCode::IsSet(), TEXT("ShaderCompileWorker error code already set; this is only supposed to be set once to handle a worker crash:\n%s"), Info);
-	GSCWErrorCode = Code;
-	GSCWErrorCodeInfo = Info;
-}
-
-void FSCWErrorCode::Reset()
-{
-	GSCWErrorCode = FSCWErrorCode::NotSet;
-	GSCWErrorCodeInfo.Empty();
-}
-
-FSCWErrorCode::ECode FSCWErrorCode::Get()
-{
-	return GSCWErrorCode;
-}
-
-const FString& FSCWErrorCode::GetInfo()
-{
-	return GSCWErrorCodeInfo;
-}
-
-bool FSCWErrorCode::IsSet()
-{
-	return GSCWErrorCode != FSCWErrorCode::NotSet;
-}
-
 /** Protects GShaderHashCache from simultaneous modification by multiple threads. Note that it can cover more than one method of the class, e.g. a block of code doing Find() then Add() can be guarded */
 FRWLock GShaderHashAccessRWLock;
 
diff --git a/Engine/Source/Runtime/RenderCore/Public/ShaderCompilerCore.h b/Engine/Source/Runtime/RenderCore/Public/ShaderCompilerCore.h
--- a/Engine/Source/Runtime/RenderCore/Public/ShaderCompilerCore.h
+++ b/Engine/Source/Runtime/RenderCore/Public/ShaderCompilerCore.h
@@ -496,44 +496,21 @@ struct FShaderCompilerOutput
 	}
 };
 
-struct RENDERCORE_API FSCWErrorCode
+enum class ESCWErrorCode
 {
-	enum ECode : int32
-	{
-		NotSet = -1,
-		Success,
-		GeneralCrash,
-		BadShaderFormatVersion,
-		BadInputVersion,
-		BadSingleJobHeader,
-		BadPipelineJobHeader,
-		CantDeleteInputFile,
-		CantSaveOutputFile,
-		NoTargetShaderFormatsFound,
-		CantCompileForSpecificFormat,
-		CrashInsidePlatformCompiler,
-		BadInputFile,
-		OutOfMemory,
-	};
-
-	/**
-	Sets the global SCW error code if it hasn't been set before.
-	Call Reset first before setting a new value.
-	Returns true on success, otherwise the error code has already been set.
-	*/
-	static void Report(ECode Code, const FStringView& Info = {});
-
-	/** Resets the global SCW error code to NotSet. */
-	static void Reset();
-
-	/** Returns the global SCW error code. */
-	static ECode Get();
-
-	/** Returns the global SCW error code information string. Empty string if not set. */
-	static const FString& GetInfo();
-
-	/** Returns true if the SCW global error code has been set. Equivalent to 'Get() != NotSet'. */
-	static bool IsSet();
+	NotSet = -1,
+	Success,
+	GeneralCrash,
+	BadShaderFormatVersion,
+	BadInputVersion,
+	BadSingleJobHeader,
+	BadPipelineJobHeader,
+	CantDeleteInputFile,
+	CantSaveOutputFile,
+	NoTargetShaderFormatsFound,
+	CantCompileForSpecificFormat,
+	CrashInsidePlatformCompiler,
+	BadInputFile
 };
 
 UE_DEPRECATED(5.2, "Functionality has moved to UE::ShaderCompilerCommon::ShouldUseStableConstantBuffer")
