Iris
- Fixed issue with filtering of replicated objects spawned with streaming level as outer

diff --git a/Engine/Source/Runtime/Engine/Private/WorldPartition/WorldPartitionLevelStreamingDynamic.cpp b/Engine/Source/Runtime/Engine/Private/WorldPartition/WorldPartitionLevelStreamingDynamic.cpp
--- a/Engine/Source/Runtime/Engine/Private/WorldPartition/WorldPartitionLevelStreamingDynamic.cpp
+++ b/Engine/Source/Runtime/Engine/Private/WorldPartition/WorldPartitionLevelStreamingDynamic.cpp
@@ -691,11 +691,13 @@ UWorld* UWorldPartitionLevelStreamingDynamic::GetOuterWorld() const
 
 void UWorldPartitionLevelStreamingDynamic::UpdateShouldSkipMakingVisibilityTransactionRequest()
 {
-	// It is safe to skip client visibility transaction requests for cells without data layers when world partition server streaming is disabled
 	const UWorldPartitionRuntimeCell* Cell = GetWorldPartitionRuntimeCell();
 	if (ensure(Cell && OuterWorldPartition.IsValid()))
 	{
-		bSkipClientUseMakingVisibleTransactionRequest = bSkipClientUseMakingInvisibleTransactionRequest = !OuterWorldPartition->IsServerStreamingEnabled() && !Cell->HasDataLayers();
+		// It is safe to skip client MakingVisibility transaction requests for cells without data layers when world partition server streaming is disabled
+		bSkipClientUseMakingVisibleTransactionRequest = !(OuterWorldPartition->IsServerStreamingEnabled() || Cell->HasDataLayers());
+		// We always need the MakeInvisibleTransactionRequest for levels that might have replicated actors associated with them including dynamically spawned actors.
+		bSkipClientUseMakingInvisibleTransactionRequest = bSkipClientUseMakingVisibleTransactionRequest && Cell->GetClientOnlyVisible();
 	}
 }
 
