Adding missing ShaderPlatformCachedIniValue.h include.

diff --git a/Engine/Source/Runtime/Engine/Private/Materials/HLSLMaterialTranslator.cpp b/Engine/Source/Runtime/Engine/Private/Materials/HLSLMaterialTranslator.cpp
--- a/Engine/Source/Runtime/Engine/Private/Materials/HLSLMaterialTranslator.cpp
+++ b/Engine/Source/Runtime/Engine/Private/Materials/HLSLMaterialTranslator.cpp
@@ -8,6 +8,7 @@
 #include "Materials/MaterialExpressionStrata.h"
 #include "MaterialExpressionSettings.h"
 #include "Engine/RendererSettings.h"
+#include "ShaderPlatformCachedIniValue.h"
 #include <functional>
 
 static int32 GetLWCTruncateMode()
