WorldPartition:
- ForceLoadedActors outside of InstancingContext mapping build loop
- Only Force load actors for the main world partition (sub partitions actors get loaded by the AlwaysLoaded loader of the LevelInstance's WP)

#rnx

diff --git a/Engine/Source/Runtime/Engine/Private/WorldPartition/WorldPartition.cpp b/Engine/Source/Runtime/Engine/Private/WorldPartition/WorldPartition.cpp
--- a/Engine/Source/Runtime/Engine/Private/WorldPartition/WorldPartition.cpp
+++ b/Engine/Source/Runtime/Engine/Private/WorldPartition/WorldPartition.cpp
@@ -514,7 +514,7 @@ void UWorldPartition::Initialize(UWorld* InWorld, const FTransform& InTransform)
 		PinnedActors = new FLoaderAdapterPinnedActors(OuterWorld);
 		
 		IWorldPartitionEditorModule& WorldPartitionEditorModule = FModuleManager::LoadModuleChecked<IWorldPartitionEditorModule>("WorldPartitionEditor");
-		ForceLoadedActors = WorldPartitionEditorModule.GetDisableLoadingInEditor() ? new FLoaderAdapterActorList(OuterWorld) : nullptr;
+		ForceLoadedActors = WorldPartitionEditorModule.GetDisableLoadingInEditor() && IsMainWorldPartition() ? new FLoaderAdapterActorList(OuterWorld) : nullptr;
 	}
 
 	check(RuntimeHash);
@@ -540,13 +540,15 @@ void UWorldPartition::Initialize(UWorld* InWorld, const FTransform& InTransform)
 			// SoftObjectPaths: Specific case for new maps (/Temp/Untitled) where we need to remap the AssetPath and not just the Package name because the World gets renamed (See UWorld::PostLoad)
 			InstancingContext.AddPathMapping(
 				FSoftObjectPath(*FString::Format(TEXT("{0}.{1}"), {PackageName.ToString(), FPackageName::GetShortName(PackageName)})),
-				FSoftObjectPath(GetWorld())
+				FSoftObjectPath(OuterWorld)
 			);
 		}
 
 		ActorDescContainer = RegisterActorDescContainer(PackageName);
 
 		{
+			TArray<FGuid> ForceLoadedActorGuids;
+			
 			TRACE_CPUPROFILER_EVENT_SCOPE(UActorDescContainer::Hash);
 			for (FActorDescContainerCollection::TIterator<> ActorDescIterator(this); ActorDescIterator; ++ActorDescIterator)
 			{
@@ -564,7 +566,7 @@ void UWorldPartition::Initialize(UWorld* InWorld, const FTransform& InTransform)
 
 				if (ForceLoadedActors)
 				{
-					ForceLoadedActors->AddActors({ ActorDescIterator->GetGuid() });
+					ForceLoadedActorGuids.Add(ActorDescIterator->GetGuid());
 				}
 
 				if (bIsEditor && !bIsCooking)
@@ -572,6 +574,12 @@ void UWorldPartition::Initialize(UWorld* InWorld, const FTransform& InTransform)
 					HashActorDesc(*ActorDescIterator);
 				}
 			}
+
+			if (ForceLoadedActors && ForceLoadedActorGuids.Num() > 0)
+			{
+				TRACE_CPUPROFILER_EVENT_SCOPE(UActorDescContainer::ForceLoadedActors);
+				ForceLoadedActors->AddActors(ForceLoadedActorGuids);
+			}
 		}
 	}
 
