#ue5 - Fix ProjectileMovementComponent continuing to tick if an interpolated component is set but bInterpMovement=false, preventing bInterpolationComplete from ever becoming true. Also added check to avoid setting the interpolated component to the same as the UpdatedComponent, and make sure that setting a null interpolated component unregisters the tick if necessary (ie when there is no UpdatedComponent either).

diff --git a/Engine/Source/Runtime/Engine/Private/Components/ProjectileMovementComponent.cpp b/Engine/Source/Runtime/Engine/Private/Components/ProjectileMovementComponent.cpp
--- a/Engine/Source/Runtime/Engine/Private/Components/ProjectileMovementComponent.cpp
+++ b/Engine/Source/Runtime/Engine/Private/Components/ProjectileMovementComponent.cpp
@@ -136,7 +136,7 @@ void UProjectileMovementComponent::TickComponent(float DeltaTime, enum ELevelTic
 	CSV_SCOPED_TIMING_STAT_EXCLUSIVE(ProjectileMovement);
 
 	// Still need to finish interpolating after we've stopped simulating, so do that first.
-	if (bInterpMovement && !bInterpolationComplete)
+	if (!bInterpolationComplete)
 	{
 		QUICK_SCOPE_CYCLE_COUNTER(STAT_ProjectileMovementComponent_TickInterpolation);
 		TickInterpolation(DeltaTime);
@@ -708,6 +708,11 @@ void UProjectileMovementComponent::SetInterpolatedComponent(USceneComponent* Com
 
 	if (Component)
 	{
+		if (!ensureMsgf(Component != UpdatedComponent, TEXT("ProjectileMovement interpolated component should not be the same as the simulated component.")))
+		{
+			return;
+		}
+
 		ResetInterpolation();
 		InterpolatedComponentPtr = Component;
 		InterpInitialLocationOffset = Component->GetRelativeLocation();
@@ -721,6 +726,11 @@ void UProjectileMovementComponent::SetInterpolatedComponent(USceneComponent* Com
 		InterpInitialLocationOffset = FVector::ZeroVector;
 		InterpInitialRotationOffset = FQuat::Identity;
 		bInterpolationComplete = true;
+		// Disabling interpolation should stop our ticking if we are done simulating and just trying to finish interpolation.
+		if (bAutoUpdateTickRegistration && (UpdatedComponent == nullptr))
+		{
+			UpdateTickRegistration();
+		}
 	}
 }
 
@@ -809,7 +819,7 @@ void UProjectileMovementComponent::TickInterpolation(float DeltaTime)
 {
 	if (!bInterpolationComplete)
 	{
-		if (USceneComponent* InterpComponent = GetInterpolatedComponent())
+		if (bInterpMovement)
 		{
 			// Smooth location. Interp faster when stopping.
 			const float ActualInterpLocationTime = Velocity.IsZero() ? 0.5f * InterpLocationTime : InterpLocationTime;
@@ -842,18 +852,21 @@ void UProjectileMovementComponent::TickInterpolation(float DeltaTime)
 				bInterpolationComplete = true;
 			}
 
-			// Apply result
-			if (UpdatedComponent)
+			if (USceneComponent* InterpComponent = GetInterpolatedComponent())
 			{
-				const FVector NewRelTranslation = UpdatedComponent->GetComponentToWorld().InverseTransformVectorNoScale(InterpLocationOffset) + InterpInitialLocationOffset;
-				if (bInterpRotation)
+				// Apply result
+				if (UpdatedComponent)
 				{
-					const FQuat NewRelRotation = InterpRotationOffset * InterpInitialRotationOffset;
-					InterpComponent->SetRelativeLocationAndRotation(NewRelTranslation, NewRelRotation);
-				}
-				else
-				{
-					InterpComponent->SetRelativeLocation(NewRelTranslation);
+					const FVector NewRelTranslation = UpdatedComponent->GetComponentToWorld().InverseTransformVectorNoScale(InterpLocationOffset) + InterpInitialLocationOffset;
+					if (bInterpRotation)
+					{
+						const FQuat NewRelRotation = InterpRotationOffset * InterpInitialRotationOffset;
+						InterpComponent->SetRelativeLocationAndRotation(NewRelTranslation, NewRelRotation);
+					}
+					else
+					{
+						InterpComponent->SetRelativeLocation(NewRelTranslation);
+					}
 				}
 			}
 		}
@@ -862,12 +875,12 @@ void UProjectileMovementComponent::TickInterpolation(float DeltaTime)
 			ResetInterpolation();
 			bInterpolationComplete = true;
 		}
-	}
 
-	// Might be done interpolating and want to disable tick
-	if (bInterpolationComplete && bAutoUpdateTickRegistration && (UpdatedComponent == nullptr))
-	{
-		UpdateTickRegistration();
+		// Might be done interpolating and want to disable tick
+		if (bInterpolationComplete && bAutoUpdateTickRegistration && (UpdatedComponent == nullptr))
+		{
+			UpdateTickRegistration();
+		}
 	}
 }
 
