- Add a default Keyboard and Mouse hardware device Identifier to properly find out when the Keyboard and Mouse is used instead of a gamepad.

#rnx
#pf 638f8982bb6fefa472c2cc00

diff --git a/Engine/Source/Runtime/Engine/Classes/GameFramework/InputSettings.h b/Engine/Source/Runtime/Engine/Classes/GameFramework/InputSettings.h
--- a/Engine/Source/Runtime/Engine/Classes/GameFramework/InputSettings.h
+++ b/Engine/Source/Runtime/Engine/Classes/GameFramework/InputSettings.h
@@ -323,6 +323,9 @@ struct ENGINE_API FHardwareDeviceIdentifier
 	/** An Invalid Hardware Device Identifier. */
 	static FHardwareDeviceIdentifier Invalid;
 
+	/** Hardware device ID that represents a keyboard and mouse. This is what will be set when an Input Event's FKey is not a gamepad key. */
+	static FHardwareDeviceIdentifier DefaultKeyboardAndMouse;
+
 	bool operator==(const FHardwareDeviceIdentifier& Other) const
 	{
 		return Other.InputClassName == InputClassName && Other.HardwareDeviceIdentifier == HardwareDeviceIdentifier;
diff --git a/Engine/Source/Runtime/Engine/Private/GameFramework/InputDeviceSubsystem.cpp b/Engine/Source/Runtime/Engine/Private/GameFramework/InputDeviceSubsystem.cpp
--- a/Engine/Source/Runtime/Engine/Private/GameFramework/InputDeviceSubsystem.cpp
+++ b/Engine/Source/Runtime/Engine/Private/GameFramework/InputDeviceSubsystem.cpp
@@ -28,7 +28,7 @@ class FInputDeviceSubsystemProcessor : public IInputProcessor
 {
 	friend class UInputDeviceSubsystem;
 	
-	void UpdateLatestDevice(const FInputDeviceId InDeviceId)
+	void UpdateLatestDevice(const FKey Key, const FInputDeviceId InDeviceId)
 	{
 #if WITH_EDITOR
 		// If we're stopped at a breakpoint we need for this input preprocessor to just ignore all incoming input
@@ -47,7 +47,12 @@ class FInputDeviceSubsystemProcessor : public IInputProcessor
 			{
 				// TODO: Refactor FInputDeviceScope to use FName's instead of a FString for HardwareDeviceIdentifier
 				SubSystem->SetMostRecentlyUsedHardwareDevice(InDeviceId, { Scope->InputDeviceName, FName(*Scope->HardwareDeviceIdentifier) });
-			}	
+			}
+			// If there isn't a recent input device scope, then we can check if the key was from a keyboard and mouse
+			else if (!Key.IsGamepadKey())
+			{
+				SubSystem->SetMostRecentlyUsedHardwareDevice(InDeviceId, FHardwareDeviceIdentifier::DefaultKeyboardAndMouse);
+			}
 		}
 	}
 	
@@ -58,37 +63,46 @@ public:
 	
 	virtual bool HandleKeyDownEvent(FSlateApplication& SlateApp, const FKeyEvent& InEvent) override
 	{
-		UpdateLatestDevice(InEvent.GetInputDeviceId());
+		UpdateLatestDevice(InEvent.GetKey(), InEvent.GetInputDeviceId());
 		return false;
 	}
 
 	virtual bool HandleAnalogInputEvent(FSlateApplication& SlateApp, const FAnalogInputEvent& InEvent) override
 	{
-		UpdateLatestDevice(InEvent.GetInputDeviceId());
+		if (!FMath::IsNearlyZero(InEvent.GetAnalogValue(), UE_KINDA_SMALL_NUMBER))
+		{
+			UpdateLatestDevice(InEvent.GetKey(), InEvent.GetInputDeviceId());
+		}
+		
 		return false;
 	}
 
 	virtual bool HandleMouseMoveEvent(FSlateApplication& SlateApp, const FPointerEvent& InEvent) override
 	{
-		UpdateLatestDevice(InEvent.GetInputDeviceId());
+		// We only want input events that have made changes
+		if (!InEvent.GetCursorDelta().IsNearlyZero())
+		{
+			UpdateLatestDevice(InEvent.GetEffectingButton(), InEvent.GetInputDeviceId());	
+		}
+		
 		return false;
 	}
 
 	virtual bool HandleMouseButtonDownEvent(FSlateApplication& SlateApp, const FPointerEvent& InEvent) override
 	{
-		UpdateLatestDevice(InEvent.GetInputDeviceId());
+		UpdateLatestDevice(InEvent.GetEffectingButton(), InEvent.GetInputDeviceId());
 		return false;
 	}
 
 	virtual bool HandleMouseButtonDoubleClickEvent(FSlateApplication& SlateApp, const FPointerEvent& InEvent) override
 	{
-		UpdateLatestDevice(InEvent.GetInputDeviceId());
+		UpdateLatestDevice(InEvent.GetEffectingButton(), InEvent.GetInputDeviceId());
 		return false;
 	}
 	
 	virtual bool HandleMouseWheelOrGestureEvent(FSlateApplication& SlateApp, const FPointerEvent& InEvent, const FPointerEvent* InGestureEvent) override
 	{
-		UpdateLatestDevice(InEvent.GetInputDeviceId());
+		UpdateLatestDevice(InEvent.GetEffectingButton(), InEvent.GetInputDeviceId());
 		return false;
 	}
 };
diff --git a/Engine/Source/Runtime/Engine/Private/UserInterface/InputSettings.cpp b/Engine/Source/Runtime/Engine/Private/UserInterface/InputSettings.cpp
--- a/Engine/Source/Runtime/Engine/Private/UserInterface/InputSettings.cpp
+++ b/Engine/Source/Runtime/Engine/Private/UserInterface/InputSettings.cpp
@@ -495,6 +495,7 @@ void UInputSettings::SetDefaultInputComponentClass(TSubclassOf<UInputComponent>
 }
 
 FHardwareDeviceIdentifier FHardwareDeviceIdentifier::Invalid = { NAME_None, NAME_None };
+FHardwareDeviceIdentifier FHardwareDeviceIdentifier::DefaultKeyboardAndMouse = { TEXT("DefaultKeyboardAndMouse"), TEXT("KBM") };
 
 // Default to the invalid hardware device identifier
 FHardwareDeviceIdentifier::FHardwareDeviceIdentifier()
@@ -540,6 +541,9 @@ const TArray<FName>& UInputPlatformSettings::GetAllHardwareDeviceNames()
 	static TArray<FName> HardwareDevices;
 	HardwareDevices.Reset();
 
+	// Add the keyboard and mouse by default for everything
+	HardwareDevices.Add(FHardwareDeviceIdentifier::DefaultKeyboardAndMouse.HardwareDeviceIdentifier);
+
 	// Get every known platform's InputPlatformSettings and compile a list of them
 	TArray<UPlatformSettings*> AllInputSettings = UPlatformSettingsManager::Get().GetAllPlatformSettings<UInputPlatformSettings>();
 
