Chaos : fix bug with swing constraint limits.
[pull request 9711]

diff --git a/Engine/Source/Runtime/Engine/Private/PhysicsEngine/ConstraintTypes.cpp b/Engine/Source/Runtime/Engine/Private/PhysicsEngine/ConstraintTypes.cpp
--- a/Engine/Source/Runtime/Engine/Private/PhysicsEngine/ConstraintTypes.cpp
+++ b/Engine/Source/Runtime/Engine/Private/PhysicsEngine/ConstraintTypes.cpp
@@ -103,8 +103,8 @@ void FConeConstraint::UpdateConeLimit_AssumesLocked(const FPhysicsConstraintHand
 
 	const bool bSkipSoftLimits = bSoftConstraint && ShouldSkipSoftLimits(Stiffness, Damping, AverageMass);
 
-	FPhysicsInterface::SetAngularMotionLimitType_AssumesLocked(InConstraintRef, PhysicsInterfaceTypes::ELimitAxis::Swing2, (bSkipSoftLimits && Swing1Motion == ACM_Limited) ? EAngularConstraintMotion::ACM_Free : Swing1Motion.GetValue());
-	FPhysicsInterface::SetAngularMotionLimitType_AssumesLocked(InConstraintRef, PhysicsInterfaceTypes::ELimitAxis::Swing1, (bSkipSoftLimits && Swing1Motion == ACM_Limited) ? EAngularConstraintMotion::ACM_Free : Swing2Motion.GetValue());
+	FPhysicsInterface::SetAngularMotionLimitType_AssumesLocked(InConstraintRef, PhysicsInterfaceTypes::ELimitAxis::Swing1, (bSkipSoftLimits && Swing1Motion == ACM_Limited) ? EAngularConstraintMotion::ACM_Free : Swing1Motion.GetValue());
+	FPhysicsInterface::SetAngularMotionLimitType_AssumesLocked(InConstraintRef, PhysicsInterfaceTypes::ELimitAxis::Swing2, (bSkipSoftLimits && Swing2Motion == ACM_Limited) ? EAngularConstraintMotion::ACM_Free : Swing2Motion.GetValue());
 }
 
 void FTwistConstraint::UpdateTwistLimit_AssumesLocked(const FPhysicsConstraintHandle& InConstraintRef, float AverageMass) const
