Fixed unloading of an instanced partitioned level with streaming enabled.
- Server was not unloading its cells.
- Unloading could not complete due to the default behavior of RemoveFromWorld being incremental.

diff --git a/Engine/Source/Runtime/Engine/Classes/Engine/LevelStreaming.h b/Engine/Source/Runtime/Engine/Classes/Engine/LevelStreaming.h
--- a/Engine/Source/Runtime/Engine/Classes/Engine/LevelStreaming.h
+++ b/Engine/Source/Runtime/Engine/Classes/Engine/LevelStreaming.h
@@ -497,6 +497,8 @@ public:
 	 */
 	virtual bool ShouldBeVisible() const;
 
+	virtual bool ShouldBlockOnUnload() const { return bShouldBlockOnUnload; }
+
 	virtual bool ShouldBeAlwaysLoaded() const { return false; }
 
 	/** Get a bounding box around the streaming volumes associated with this LevelStreaming object */
diff --git a/Engine/Source/Runtime/Engine/Private/LevelStreaming.cpp b/Engine/Source/Runtime/Engine/Private/LevelStreaming.cpp
--- a/Engine/Source/Runtime/Engine/Private/LevelStreaming.cpp
+++ b/Engine/Source/Runtime/Engine/Private/LevelStreaming.cpp
@@ -966,7 +966,7 @@ void ULevelStreaming::UpdateStreamingState(bool& bOutUpdateAgain, bool& bOutRede
 			}
 
 			// Hide loaded level, incrementally if necessary
-			World->RemoveFromWorld(LoadedLevel, !bShouldBlockOnUnload && World->IsGameWorld(), TransactionId, this);
+			World->RemoveFromWorld(LoadedLevel, !ShouldBlockOnUnload() && World->IsGameWorld(), TransactionId, this);
 
 			// Hide loaded level immediately if bRequireFullVisibilityToRender is set
 			const bool LevelBecameInvisible = bWasVisible && !LoadedLevel->bIsVisible;
diff --git a/Engine/Source/Runtime/Engine/Private/WorldPartition/WorldPartitionLevelStreamingDynamic.cpp b/Engine/Source/Runtime/Engine/Private/WorldPartition/WorldPartitionLevelStreamingDynamic.cpp
--- a/Engine/Source/Runtime/Engine/Private/WorldPartition/WorldPartitionLevelStreamingDynamic.cpp
+++ b/Engine/Source/Runtime/Engine/Private/WorldPartition/WorldPartitionLevelStreamingDynamic.cpp
@@ -689,6 +689,23 @@ UWorld* UWorldPartitionLevelStreamingDynamic::GetOuterWorld() const
 	return OuterWorldPartition->GetTypedOuter<UWorld>();
 }
 
+bool UWorldPartitionLevelStreamingDynamic::ShouldBlockOnUnload() const
+{
+	if (Super::ShouldBlockOnUnload())
+	{
+		return true;
+	}
+
+	// When world partition cannot stream (anymore), return true so that RemoveFromWorld of this level is not incremental. 
+	// This guarantees that unloaded instanced wp levels fully unload their cell levels.
+	if (OuterWorldPartition.IsValid() && !OuterWorldPartition->CanStream())
+	{
+		return true;
+	}
+
+	return false;
+}
+
 void UWorldPartitionLevelStreamingDynamic::UpdateShouldSkipMakingVisibilityTransactionRequest()
 {
 	const UWorldPartitionRuntimeCell* Cell = GetWorldPartitionRuntimeCell();
diff --git a/Engine/Source/Runtime/Engine/Private/WorldPartition/WorldPartitionStreamingPolicy.cpp b/Engine/Source/Runtime/Engine/Private/WorldPartition/WorldPartitionStreamingPolicy.cpp
--- a/Engine/Source/Runtime/Engine/Private/WorldPartition/WorldPartitionStreamingPolicy.cpp
+++ b/Engine/Source/Runtime/Engine/Private/WorldPartition/WorldPartitionStreamingPolicy.cpp
@@ -402,6 +402,7 @@ void UWorldPartitionStreamingPolicy::UpdateStreamingState()
 		FrameLoadCells.Reset();
 	};
 
+	const bool bCanStream = WorldPartition->CanStream();
 	const bool bIsServer = WorldPartition->IsServer();
 	const bool bIsServerStreamingEnabled = WorldPartition->IsServerStreamingEnabled();
 	const bool bCanDeactivateOrUnloadCells = !bIsServer || WorldPartition->IsServerStreamingOutEnabled();
@@ -409,104 +410,107 @@ void UWorldPartitionStreamingPolicy::UpdateStreamingState()
 
 	const bool bIsStreamingInEnabled = WorldPartition->IsStreamingInEnabled();
 
-	if (!bIsServer || bIsServerStreamingEnabled || AWorldPartitionReplay::IsPlaybackEnabled(World))
+	if (bCanStream)
 	{
-		// When world partition can't stream, all cells must be unloaded
-		if (WorldPartition->CanStream() && WorldPartition->RuntimeHash)
+		if (!bIsServer || bIsServerStreamingEnabled || AWorldPartitionReplay::IsPlaybackEnabled(World))
 		{
-			TRACE_CPUPROFILER_EVENT_SCOPE(UWorldPartitionStreamingPolicy::UpdateStreamingState_ForEachStreamingCellsSources);
+			// When world partition can't stream, all cells must be unloaded
+			if (WorldPartition->RuntimeHash)
+			{
+				TRACE_CPUPROFILER_EVENT_SCOPE(UWorldPartitionStreamingPolicy::UpdateStreamingState_ForEachStreamingCellsSources);
 
-			UWorldPartitionRuntimeCell::DirtyStreamingSourceCacheEpoch();
+				UWorldPartitionRuntimeCell::DirtyStreamingSourceCacheEpoch();
 
-			WorldPartition->RuntimeHash->ForEachStreamingCellsSources(StreamingSources, [this](const UWorldPartitionRuntimeCell* Cell, EStreamingSourceTargetState TargetState)
-			{
-				switch (TargetState)
+				WorldPartition->RuntimeHash->ForEachStreamingCellsSources(StreamingSources, [this](const UWorldPartitionRuntimeCell* Cell, EStreamingSourceTargetState TargetState)
 				{
-				case EStreamingSourceTargetState::Loaded:
-					FrameLoadCells.Add(Cell);
-					break;
-				case EStreamingSourceTargetState::Activated:
-					FrameActivateCells.Add(Cell);
-					break;
-				default:
-					check(0);
-				}
+					switch (TargetState)
+					{
+					case EStreamingSourceTargetState::Loaded:
+						FrameLoadCells.Add(Cell);
+						break;
+					case EStreamingSourceTargetState::Activated:
+						FrameActivateCells.Add(Cell);
+						break;
+					default:
+						check(0);
+					}
 
-				return true;
-			});
+					return true;
+				});
+			}
 		}
-	}
 
-	if (bIsServer)
-	{
-		TRACE_CPUPROFILER_EVENT_SCOPE(UWorldPartitionStreamingPolicy::UpdateStreamingState_ServerUpdate);
-
-		// Server will activate all non data layer cells at first and then load/activate/unload data layer cells only when the data layer states change
-		if (!bIsServerStreamingEnabled && 
-			(ServerStreamingEnabledEpoch == NewServerStreamingEnabledEpoch) &&
-			(ContentBundleServerEpoch == FContentBundle::GetContentBundleEpoch()) &&
-			(DataLayersStatesServerEpoch == AWorldDataLayers::GetDataLayersStateEpoch()))
+		if (bIsServer)
 		{
-			// Server as nothing to do early out
-			return; 
-		}
+			TRACE_CPUPROFILER_EVENT_SCOPE(UWorldPartitionStreamingPolicy::UpdateStreamingState_ServerUpdate);
 
-		bUpdateEpoch = true;
-
-		// Gather Client visible level names
-		if (const UNetDriver* NetDriver = World->GetNetDriver())
-		{
-			for (UNetConnection* Connection : NetDriver->ClientConnections)
+			// Server will activate all non data layer cells at first and then load/activate/unload data layer cells only when the data layer states change
+			if (!bIsServerStreamingEnabled && 
+				(ServerStreamingEnabledEpoch == NewServerStreamingEnabledEpoch) &&
+				(ContentBundleServerEpoch == FContentBundle::GetContentBundleEpoch()) &&
+				(DataLayersStatesServerEpoch == AWorldDataLayers::GetDataLayersStateEpoch()))
 			{
-				ClientVisibleLevelNames.Add(Connection->GetClientWorldPackageName());
-				ClientVisibleLevelNames.Append(Connection->ClientVisibleLevelNames);
-				ClientVisibleLevelNames.Append(Connection->GetClientMakingVisibleLevelNames());
+				// Server as nothing to do early out
+				return; 
 			}
-		}
 
-		const UDataLayerSubsystem* DataLayerSubsystem = GetWorld()->GetSubsystem<UDataLayerSubsystem>();
-		TSet<FName> EffectiveActiveDataLayerNames = DataLayerSubsystem->GetEffectiveActiveDataLayerNames();
-		TSet<FName> EffectiveLoadedDataLayerNames = DataLayerSubsystem->GetEffectiveLoadedDataLayerNames();
+			bUpdateEpoch = true;
 
-		auto AddServerFrameCells = [this, &EffectiveLoadedDataLayerNames, &EffectiveActiveDataLayerNames](const UWorldPartitionRuntimeCell* Cell)
-		{
-			// Non Data Layer Cells + Active Data Layers
-			if (!Cell->HasDataLayers() || Cell->HasAnyDataLayer(EffectiveActiveDataLayerNames))
+			// Gather Client visible level names
+			if (const UNetDriver* NetDriver = World->GetNetDriver())
 			{
-				FrameActivateCells.Add(Cell);
+				for (UNetConnection* Connection : NetDriver->ClientConnections)
+				{
+					ClientVisibleLevelNames.Add(Connection->GetClientWorldPackageName());
+					ClientVisibleLevelNames.Append(Connection->ClientVisibleLevelNames);
+					ClientVisibleLevelNames.Append(Connection->GetClientMakingVisibleLevelNames());
+				}
 			}
 
-			// Loaded Data Layers Cells only
-			if (EffectiveLoadedDataLayerNames.Num())
+			const UDataLayerSubsystem* DataLayerSubsystem = GetWorld()->GetSubsystem<UDataLayerSubsystem>();
+			TSet<FName> EffectiveActiveDataLayerNames = DataLayerSubsystem->GetEffectiveActiveDataLayerNames();
+			TSet<FName> EffectiveLoadedDataLayerNames = DataLayerSubsystem->GetEffectiveLoadedDataLayerNames();
+
+			auto AddServerFrameCells = [this, &EffectiveLoadedDataLayerNames, &EffectiveActiveDataLayerNames](const UWorldPartitionRuntimeCell* Cell)
 			{
-				if (Cell->HasDataLayers() && Cell->HasAnyDataLayer(EffectiveLoadedDataLayerNames))
+				// Non Data Layer Cells + Active Data Layers
+				if (!Cell->HasDataLayers() || Cell->HasAnyDataLayer(EffectiveActiveDataLayerNames))
 				{
-					FrameLoadCells.Add(Cell);
+					FrameActivateCells.Add(Cell);
 				}
-			}
-		};
 
-		if (!bIsServerStreamingEnabled)
-		{
-			if (WorldPartition->RuntimeHash)
-			{
-				WorldPartition->RuntimeHash->ForEachStreamingCells([this, &AddServerFrameCells](const UWorldPartitionRuntimeCell* Cell)
+				// Loaded Data Layers Cells only
+				if (EffectiveLoadedDataLayerNames.Num())
 				{
-					AddServerFrameCells(Cell);
-					return true;
-				});
-			}
-		}
-		else if (!bCanDeactivateOrUnloadCells)
-		{
-			// When server streaming-out is disabled, revisit existing loaded/activated cells and add them in the proper FrameLoadCells/FrameActivateCells
-			for (const UWorldPartitionRuntimeCell* Cell : ActivatedCells.GetCells())
+					if (Cell->HasDataLayers() && Cell->HasAnyDataLayer(EffectiveLoadedDataLayerNames))
+					{
+						FrameLoadCells.Add(Cell);
+					}
+				}
+			};
+
+			if (!bIsServerStreamingEnabled)
 			{
-				AddServerFrameCells(Cell);
+				if (WorldPartition->RuntimeHash)
+				{
+					WorldPartition->RuntimeHash->ForEachStreamingCells([this, &AddServerFrameCells](const UWorldPartitionRuntimeCell* Cell)
+					{
+						AddServerFrameCells(Cell);
+						return true;
+					});
+				}
 			}
-			for (const UWorldPartitionRuntimeCell* Cell : LoadedCells)
+			else if (!bCanDeactivateOrUnloadCells)
 			{
-				AddServerFrameCells(Cell);
+				// When server streaming-out is disabled, revisit existing loaded/activated cells and add them in the proper FrameLoadCells/FrameActivateCells
+				for (const UWorldPartitionRuntimeCell* Cell : ActivatedCells.GetCells())
+				{
+					AddServerFrameCells(Cell);
+				}
+				for (const UWorldPartitionRuntimeCell* Cell : LoadedCells)
+				{
+					AddServerFrameCells(Cell);
+				}
 			}
 		}
 	}
@@ -536,6 +540,7 @@ void UWorldPartitionStreamingPolicy::UpdateStreamingState()
 
 	// Determine cells to activate (sorted by importance)
 	TArray<const UWorldPartitionRuntimeCell*> ToActivateCells;
+	if (bCanStream)
 	{
 		TRACE_CPUPROFILER_EVENT_SCOPE(UWorldPartitionStreamingPolicy::UpdateStreamingState_ToActivateCells);
 		for (const UWorldPartitionRuntimeCell* Cell : FrameActivateCells)
@@ -551,6 +556,7 @@ void UWorldPartitionStreamingPolicy::UpdateStreamingState()
 
 	// Determine cells to load (sorted by importance)
 	TArray<const UWorldPartitionRuntimeCell*> ToLoadCells;
+	if (bCanStream)
 	{
 		TRACE_CPUPROFILER_EVENT_SCOPE(UWorldPartitionStreamingPolicy::UpdateStreamingState_ToLoadCells);
 		for (const UWorldPartitionRuntimeCell* Cell : FrameLoadCells)
@@ -574,13 +580,13 @@ void UWorldPartitionStreamingPolicy::UpdateStreamingState()
 	TArray<const UWorldPartitionRuntimeCell*> ToUnloadCells;
 	{
 		TRACE_CPUPROFILER_EVENT_SCOPE(UWorldPartitionStreamingPolicy::UpdateStreamingState_ToUnloadCells);
-		auto BuildCellsToUnload = [this, &ToUnloadCells, bIsServer, ShouldWaitForClientVisibility](const TSet<const UWorldPartitionRuntimeCell*>& InCells)
+		auto BuildCellsToUnload = [this, &ToUnloadCells, bCanStream, bIsServer, ShouldWaitForClientVisibility](const TSet<const UWorldPartitionRuntimeCell*>& InCells)
 		{
 			for (const UWorldPartitionRuntimeCell* Cell : InCells)
 			{
 				if (!FrameActivateCells.Contains(Cell) && !FrameLoadCells.Contains(Cell))
 				{
-					if (!bIsServer || !ShouldWaitForClientVisibility(Cell))
+					if (!bCanStream || !bIsServer || !ShouldWaitForClientVisibility(Cell))
 					{
 						ToUnloadCells.Add(Cell);
 					}
diff --git a/Engine/Source/Runtime/Engine/Public/WorldPartition/WorldPartitionLevelStreamingDynamic.h b/Engine/Source/Runtime/Engine/Public/WorldPartition/WorldPartitionLevelStreamingDynamic.h
--- a/Engine/Source/Runtime/Engine/Public/WorldPartition/WorldPartitionLevelStreamingDynamic.h
+++ b/Engine/Source/Runtime/Engine/Public/WorldPartition/WorldPartitionLevelStreamingDynamic.h
@@ -34,6 +34,7 @@ class ENGINE_API UWorldPartitionLevelStreamingDynamic : public ULevelStreamingDy
 	const UWorldPartitionRuntimeCell* GetWorldPartitionRuntimeCell() const { return StreamingCell.Get(); }
 
 	virtual bool ShouldBeAlwaysLoaded() const override { return bShouldBeAlwaysLoaded; }
+	virtual bool ShouldBlockOnUnload() const override;
 	virtual bool ShouldRequireFullVisibilityToRender() const override { return true; }
 
 	virtual bool CanMakeVisible() override;
