Reduce editor hitches during shader compilation by avoiding a scan of all primitive to get their dependent material by using a reverse lookup table instead
  - Fix potential race conditions when async loading in editor is enabled by removing async objects from lookup results
  - Add LLM tracking support for reverse lookup caches

#rnx

diff --git a/Engine/Source/Runtime/Engine/Private/Materials/MaterialInstance.cpp b/Engine/Source/Runtime/Engine/Private/Materials/MaterialInstance.cpp
--- a/Engine/Source/Runtime/Engine/Private/Materials/MaterialInstance.cpp
+++ b/Engine/Source/Runtime/Engine/Private/Materials/MaterialInstance.cpp
@@ -3170,6 +3170,11 @@ bool UMaterialInstance::SetParentInternal(UMaterialInterface* NewParent, bool Re
 			Parent = NewParent;
 			bSetParent = true;
 
+#if WITH_EDITOR
+			// Important to notify when the parent change for Material -> Material relationship update
+			FObjectCacheEventSink::NotifyReferencedTextureChanged_Concurrent(this);
+#endif
+
 			if( Parent )
 			{
 				// It is possible to set a material's parent while post-loading. In
diff --git a/Engine/Source/Runtime/Engine/Private/ObjectCacheContext.cpp b/Engine/Source/Runtime/Engine/Private/ObjectCacheContext.cpp
--- a/Engine/Source/Runtime/Engine/Private/ObjectCacheContext.cpp
+++ b/Engine/Source/Runtime/Engine/Private/ObjectCacheContext.cpp
@@ -3,21 +3,37 @@
 #include "ObjectCacheContext.h"
 #include "UObject/UObjectIterator.h"
 #include "UObject/ObjectKey.h"
+#include "Materials/Material.h"
 #include "Materials/MaterialInstance.h"
+#include "Materials/MaterialInterface.h"
 #include "Engine/StaticMesh.h"
 #include "Engine/SkeletalMesh.h"
 #include "Engine/Texture.h"
 #include "Components/PrimitiveComponent.h"
 #include "Components/StaticMeshComponent.h"
 #include "Components/SkinnedMeshComponent.h"
+#include "Containers/Queue.h"
+#include "Logging/LogMacros.h"
+#include "HAL/LowLevelMemStats.h"
 #include "Misc/ScopeRWLock.h"
 #include "Logging/LogCategory.h"
 #include "Templates/UniquePtr.h"
 
 #define LOCTEXT_NAMESPACE "ObjectCache"
 
-DECLARE_LOG_CATEGORY_EXTERN(LogObjectCache, Log, All);
-DEFINE_LOG_CATEGORY(LogObjectCache);
+LLM_DEFINE_TAG(ObjectReverseLookupCache);
+
+DEFINE_LOG_CATEGORY_STATIC(LogObjectCache, Log, All);
+
+namespace ObjectCacheContextImpl {
+
+EInternalObjectFlags GetObjectCacheInternalFlagsExclusion()
+{
+	// We never want to return objects that are invalid or still being worked on by other threads
+	return EInternalObjectFlags::Unreachable | EInternalObjectFlags::Garbage | EInternalObjectFlags::AsyncLoading | EInternalObjectFlags::Async;
+}
+
+} // namespace ObjectCacheContextImpl
 
 #if WITH_EDITOR
 #include "ObjectCacheEventSink.h"
@@ -46,6 +62,8 @@ static FAutoConsoleCommand CVarObjectReverseLookupValidate(
 	)
 );
 
+namespace ObjectCacheContextImpl {
+
 enum class EObjectReverseLookupMode
 {
 	Temporary = 0,
@@ -70,6 +88,21 @@ EObjectReverseLookupMode GetObjectReverseLookupMode()
 	return (EObjectReverseLookupMode)CVarObjectReverseLookupMode.GetValueOnAnyThread();
 }
 
+struct FDefaultObjectSearchPredicate
+{
+	EInternalObjectFlags ExcludedInternalFlags;
+
+	FDefaultObjectSearchPredicate()
+		: ExcludedInternalFlags(GetObjectCacheInternalFlagsExclusion())
+	{
+	}
+
+	bool operator() (UObject* Object) const
+	{
+		return !Object->HasAnyInternalFlags(ExcludedInternalFlags);
+	}
+};
+
 template <typename FromType, typename ToType>
 class FObjectReverseLookupCache
 {
@@ -97,6 +130,7 @@ public:
 
 	void Update(ToType* InTo, const TArray<FromType*>& NewFromObjects)
 	{
+		LLM_SCOPE_BYTAG(ObjectReverseLookupCache);
 		FWriteScopeLock ScopeLock(Lock);
 
 		// Remove the old reverse lookup mappings as they can't be read from the
@@ -121,9 +155,12 @@ public:
 
 			for (FromType* From : NewFromObjects)
 			{
-				UE_LOG(LogObjectCache, VeryVerbose, TEXT("Lookup mapping added %s -> %s"), *From->GetFullName(), *InTo->GetFullName());
-				FromToMapping.FindOrAdd(From).Add(InTo);
-				FromObjects.Add(From);
+				if (From != nullptr)
+				{
+					UE_LOG(LogObjectCache, VeryVerbose, TEXT("Lookup mapping added %s -> %s"), *From->GetFullName(), *InTo->GetFullName());
+					FromToMapping.FindOrAdd(From).FindOrAdd(InTo);
+					FromObjects.FindOrAdd(From);
+				}
 			}
 		}
 		else
@@ -132,7 +169,8 @@ public:
 		}
 	}
 
-	TObjectCacheIterator<ToType> GetFrom(FromType* InFrom) const
+	template <typename PredicateType = FDefaultObjectSearchPredicate>
+	TObjectCacheIterator<ToType> GetFrom(FromType* InFrom, const PredicateType& Predicate = PredicateType()) const
 	{
 		FReadScopeLock ScopeLock(Lock);
 
@@ -144,14 +182,18 @@ public:
 			OutTo.Reserve(Set->Num());
 			for (ToType* To : *Set)
 			{
-				OutTo.Add(To);
+				if (Predicate(To))
+				{
+					OutTo.Add(To);
+				}
 			}
 		}
 
 		return TObjectCacheIterator<ToType>(MoveTemp(OutTo));
 	}
 
-	TObjectCacheIterator<FromType> GetTo(ToType* InTo) const
+	template <typename PredicateType = FDefaultObjectSearchPredicate>
+	TObjectCacheIterator<FromType> GetTo(ToType* InTo, const PredicateType& Predicate = PredicateType()) const
 	{
 		FReadScopeLock ScopeLock(Lock);
 
@@ -163,7 +205,10 @@ public:
 			OutFrom.Reserve(Set->Num());
 			for (FromType* From : *Set)
 			{
-				OutFrom.Add(From);
+				if (Predicate(From))
+				{
+					OutFrom.Add(From);
+				}
 			}
 		}
 
@@ -242,60 +287,72 @@ private:
 	TRawMap<ToType*, TRawSet<FromType*>> ToFromMapping;
 };
 
+FObjectReverseLookupCache<UMaterialInterface, UMaterialInstance>   GMaterialToMaterialInstanceLookupCache; // Parent -> Children
 FObjectReverseLookupCache<UTexture, UMaterialInterface>            GTextureToMaterialLookupCache;
 FObjectReverseLookupCache<UStaticMesh, UStaticMeshComponent>       GStaticMeshToComponentLookupCache;
 FObjectReverseLookupCache<UMaterialInterface, UPrimitiveComponent> GMaterialToPrimitiveLookupCache;
 
-namespace ObjectCacheContextImpl {
-
-	void GetReferencedTextures(UMaterialInterface* MaterialInterface, TSet<UTexture*>& OutReferencedTextures);
-	void Validate()
+void GetReferencedTextures(UMaterialInterface* MaterialInterface, TSet<UTexture*>& OutReferencedTextures);
+void Validate()
+{
+	// Scan and compare UStaticMesh -> UStaticMeshComponent
+	int32 ErrorCount = 0;
 	{
-		// Scan and compare UStaticMesh -> UStaticMeshComponent
-		int32 ErrorCount = 0;
+		FObjectReverseLookupCache<UStaticMesh, UStaticMeshComponent> TempLookup;
+		for (TObjectIterator<UStaticMeshComponent> It(RF_ClassDefaultObject, true /*bIncludeDerivedClasses*/, GetObjectCacheInternalFlagsExclusion()); It; ++It)
 		{
-			FObjectReverseLookupCache<UStaticMesh, UStaticMeshComponent> TempLookup;
-			for (TObjectIterator<UStaticMeshComponent> It; It; ++It)
+			if (It->GetStaticMesh())
 			{
-				if (It->GetStaticMesh())
-				{
-					TempLookup.Update(*It, { It->GetStaticMesh() } );
-				}
+				TempLookup.Update(*It, { It->GetStaticMesh() } );
 			}
-			ErrorCount += TempLookup.Compare(GStaticMeshToComponentLookupCache);
 		}
+		ErrorCount += TempLookup.Compare(GStaticMeshToComponentLookupCache);
+	}
 
-		// Scan and compare UTexture -> UMaterialInterface
+	// Scan and compare UTexture -> UMaterialInterface
+	{
+		FObjectReverseLookupCache<UTexture, UMaterialInterface> TempLookup;
+		for (TObjectIterator<UMaterialInterface> It(RF_ClassDefaultObject, true /*bIncludeDerivedClasses*/, GetObjectCacheInternalFlagsExclusion()); It; ++It)
 		{
-			FObjectReverseLookupCache<UTexture, UMaterialInterface> TempLookup;
-			for (TObjectIterator<UMaterialInterface> It; It; ++It)
+			TSet<UTexture*> ReferencedTextures;
+			GetReferencedTextures(*It, ReferencedTextures);
+			TempLookup.Update(*It, ReferencedTextures.Array());
+		}
+		ErrorCount += TempLookup.Compare(GTextureToMaterialLookupCache);
+	}
+
+	// Scan and compare UMaterialInterface -> UPrimitiveComponent
+	{
+		FObjectReverseLookupCache<UMaterialInterface, UPrimitiveComponent> TempLookup;
+		for (TObjectIterator<UPrimitiveComponent> It(RF_ClassDefaultObject, true /*bIncludeDerivedClasses*/, GetObjectCacheInternalFlagsExclusion()); It; ++It)
+		{
+			UPrimitiveComponent* Component = *It;
+			if (Component->IsRenderStateCreated() && Component->SceneProxy && !Component->IsRenderStateDirty())
 			{
-				TSet<UTexture*> ReferencedTextures;
-				GetReferencedTextures(*It, ReferencedTextures);
-				TempLookup.Update(*It, ReferencedTextures.Array());
+				TArray<UMaterialInterface*> UsedMaterials;
+				Component->GetUsedMaterials(UsedMaterials);
+				TempLookup.Update(Component, UsedMaterials);
 			}
-			ErrorCount += TempLookup.Compare(GTextureToMaterialLookupCache);
 		}
+		ErrorCount += TempLookup.Compare(GMaterialToPrimitiveLookupCache);
+	}
 
-		// Scan and compare UMaterialInterface -> UPrimitiveComponent
+	// Scan and compare UMaterialInterface -> UMaterialInterface
+	{
+		FObjectReverseLookupCache<UMaterialInterface, UMaterialInstance> TempLookup;
+		for (TObjectIterator<UMaterialInstance> It(RF_ClassDefaultObject, true /*bIncludeDerivedClasses*/, GetObjectCacheInternalFlagsExclusion()); It; ++It)
 		{
-			FObjectReverseLookupCache<UMaterialInterface, UPrimitiveComponent> TempLookup;
-			for (TObjectIterator<UPrimitiveComponent> It; It; ++It)
+			if (It->Parent)
 			{
-				UPrimitiveComponent* Component = *It;
-				if (Component->IsRenderStateCreated() && Component->SceneProxy && !Component->IsRenderStateDirty())
-				{
-					TArray<UMaterialInterface*> UsedMaterials;
-					Component->GetUsedMaterials(UsedMaterials);
-					TempLookup.Update(Component, UsedMaterials);
-				}
+				TempLookup.Update(*It, { It->Parent });
 			}
-			ErrorCount += TempLookup.Compare(GMaterialToPrimitiveLookupCache);
 		}
-
-		UE_LOG(LogObjectCache, Display, TEXT("Object Cache verification found %d discrepenties"), ErrorCount);
+		ErrorCount += TempLookup.Compare(GMaterialToMaterialInstanceLookupCache);
 	}
 
+	UE_LOG(LogObjectCache, Display, TEXT("Object Cache verification found %d discrepenties"), ErrorCount);
+}
+
 } // namespace ObjectCacheContextImpl
 
 #endif // #if WITH_EDITOR
@@ -310,7 +367,7 @@ namespace ObjectCacheContextImpl {
 			UTexture* Texture = Cast<UTexture>(TextureObject);
 			if (Texture)
 			{
-				OutReferencedTextures.Add(Texture);
+				OutReferencedTextures.FindOrAdd(Texture);
 			}
 		}
 
@@ -322,7 +379,7 @@ namespace ObjectCacheContextImpl {
 			{
 				if (TextureParam.ParameterValue)
 				{
-					OutReferencedTextures.Add(TextureParam.ParameterValue);
+					OutReferencedTextures.FindOrAdd(TextureParam.ParameterValue);
 				}
 			}
 
@@ -330,17 +387,142 @@ namespace ObjectCacheContextImpl {
 		}
 	}
 
+	// This is the old method of finding dependent materials by scanning all MaterialInstance currently reachable
+	// Adapted from FMaterialUpdateContext::~FMaterialUpdateContext() / MaterialShared.cpp
+	void GetMaterialsAffectedByMaterials_Iteration(TArrayView<UMaterialInterface*> InMaterials, TSet<UMaterialInterface*>& OutAffectedMaterials)
+	{
+		TRACE_CPUPROFILER_EVENT_SCOPE(GetMaterialsAffectedByMaterials_Iteration);
+
+		TSet<UMaterialInterface*> Materials;
+		Materials.Reserve(InMaterials.Num());
+		for (UMaterialInterface* Material : InMaterials)
+		{
+			OutAffectedMaterials.FindOrAdd(Material);
+			Materials.FindOrAdd(Material);
+		}
+
+		TSet<UMaterialInterface*> AffectedMaterials;
+		for (TObjectIterator<UMaterialInstance> It(RF_ClassDefaultObject, true /*bIncludeDerivedClasses*/, GetObjectCacheInternalFlagsExclusion()); It; ++It)
+		{
+			UMaterialInstance* MaterialInstance = *It;
+			UMaterial* BaseMaterial = MaterialInstance->GetMaterial();
+
+			if (Materials.Contains(BaseMaterial))
+			{
+				// Check to see if this instance is dependent on any of the material interfaces we directly updated.
+				for (UMaterialInterface* MaterialInterface : InMaterials)
+				{
+					if (MaterialInstance->IsDependent(MaterialInterface))
+					{
+						OutAffectedMaterials.FindOrAdd(MaterialInstance);
+						break;
+					}
+				}
+			}
+		}
+	}
+
+#if WITH_EDITOR
+	// New method by recursively iterating on a reverse lookup from parent -> children relationships until the graph is fully resolved
+	void GetMaterialsAffectedByMaterials_Lookup(TArrayView<UMaterialInterface*> InMaterials, TSet<UMaterialInterface*>& OutAffectedMaterials)
+	{
+		TRACE_CPUPROFILER_EVENT_SCOPE(GetMaterialsAffectedByMaterials_Lookup);
+
+		TQueue<UMaterialInterface*> Todo;
+		for (UMaterialInterface* Material : InMaterials)
+		{
+			Todo.Enqueue(Material);
+		}
+
+		UMaterialInterface* NextItem;
+		while (Todo.Dequeue(NextItem))
+		{
+			OutAffectedMaterials.FindOrAdd(NextItem);
+			for (UMaterialInstance* Child : GMaterialToMaterialInstanceLookupCache.GetFrom(NextItem))
+			{
+				bool bIsAlreadyPresent = false;
+				OutAffectedMaterials.FindOrAdd(Child, &bIsAlreadyPresent);
+				if (!bIsAlreadyPresent)
+				{
+					Todo.Enqueue(Child);
+				}
+			}
+		}
+	}
+#endif // #if WITH_EDITOR
+
+	bool DoesPrimitiveDependsOnMaterials(FObjectCacheContext& Context, UPrimitiveComponent* PrimitiveComponent, TArrayView<UMaterialInterface*> InMaterials)
+	{
+		// Note: relying on GetUsedMaterials to be accurate, or else we won't propagate to the right primitives and the renderer will crash later
+		// FPrimitiveSceneProxy::VerifyUsedMaterial is used to make sure that all materials used for rendering are reported in GetUsedMaterials
+		TObjectCacheIterator<UMaterialInterface> UsedMaterials = Context.GetUsedMaterials(PrimitiveComponent);
+		if (!UsedMaterials.IsEmpty())
+		{
+			for (UMaterialInterface* UpdatedMaterial : InMaterials)
+			{
+				for (UMaterialInterface* UsedMaterial : UsedMaterials)
+				{
+					if (UsedMaterial && (UsedMaterial == UpdatedMaterial || UsedMaterial->IsDependent(UpdatedMaterial)))
+					{
+						return true;
+					}
+				}
+			}
+		}
+
+		return false;
+	}
+
+	// This is the old method of finding primitives that depends on a set of materials by scanning all currently reachable Primitives' UsedMaterials
+	// Adapted from FShaderCompilingManager::PropagateMaterialChangesToPrimitives / ShaderCompiler.cpp
+	void GetPrimitivesAffectedByMaterials_Iteration(FObjectCacheContext& Context, TArrayView<UMaterialInterface*> InMaterials, TSet<UPrimitiveComponent*>& OutAffectedPrimitives)
+	{
+		TRACE_CPUPROFILER_EVENT_SCOPE(GetPrimitivesAffectedByMaterials_Iteration);
+
+		for (UPrimitiveComponent* PrimitiveComponent : Context.GetPrimitiveComponents())
+		{
+			if (PrimitiveComponent->IsRenderStateCreated() && PrimitiveComponent->SceneProxy)
+			{
+				if (DoesPrimitiveDependsOnMaterials(Context, PrimitiveComponent, InMaterials))
+				{
+					OutAffectedPrimitives.FindOrAdd(PrimitiveComponent);
+				}
+			}
+		}
+	}
+
+#if WITH_EDITOR
+	// This is the new method of finding primitives that depends on a set of materials by using reverse lookups
+	void GetPrimitivesAffectedByMaterials_Lookup(TArrayView<UMaterialInterface*> InMaterials, TSet<UPrimitiveComponent*>& OutAffectedPrimitives)
+	{
+		TRACE_CPUPROFILER_EVENT_SCOPE(GetPrimitivesAffectedByMaterials_Lookup);
+
+		TSet<UMaterialInterface*> AffectedMaterials;
+		GetMaterialsAffectedByMaterials_Lookup(InMaterials, AffectedMaterials);
+
+		for (UMaterialInterface* MaterialInterface : AffectedMaterials)
+		{
+			for (UPrimitiveComponent* Primitive : GMaterialToPrimitiveLookupCache.GetFrom(MaterialInterface))
+			{
+				OutAffectedPrimitives.FindOrAdd(Primitive);
+			}
+		}
+	}
+#endif // #if WITH_EDITOR
+
 } // namespace ObjectCacheContextImpl
 
 TObjectCacheIterator<UPrimitiveComponent> FObjectCacheContext::GetPrimitiveComponents()
 {
+	using namespace ObjectCacheContextImpl;
+
 	if (!PrimitiveComponents.IsSet())
 	{
 		TRACE_CPUPROFILER_EVENT_SCOPE(ComputePrimitiveComponents);
 
 		TArray<UPrimitiveComponent*> Array;
 		Array.Reserve(4096);
-		for (TObjectIterator<UPrimitiveComponent> It; It; ++It)
+		for (TObjectIterator<UPrimitiveComponent> It(RF_ClassDefaultObject, true /*bIncludeDerivedClasses*/, GetObjectCacheInternalFlagsExclusion()); It; ++It)
 		{
 			Array.Add(*It);
 		}
@@ -352,13 +534,15 @@ TObjectCacheIterator<UPrimitiveComponent> FObjectCacheContext::GetPrimitiveCompo
 
 TObjectCacheIterator<UStaticMeshComponent> FObjectCacheContext::GetStaticMeshComponents()
 {
+	using namespace ObjectCacheContextImpl;
+
 	if (!StaticMeshComponents.IsSet())
 	{
 		TRACE_CPUPROFILER_EVENT_SCOPE(ComputeStaticMeshComponents);
 
 		TArray<UStaticMeshComponent*> Array;
 		Array.Reserve(4096);
-		for (TObjectIterator<UStaticMeshComponent> It; It; ++It)
+		for (TObjectIterator<UStaticMeshComponent> It(RF_ClassDefaultObject, true /*bIncludeDerivedClasses*/, GetObjectCacheInternalFlagsExclusion()); It; ++It)
 		{
 			Array.Add(*It);
 		}
@@ -369,6 +553,7 @@ TObjectCacheIterator<UStaticMeshComponent> FObjectCacheContext::GetStaticMeshCom
 
 TObjectCacheIterator<UTexture> FObjectCacheContext::GetUsedTextures(UMaterialInterface* MaterialInterface)
 {
+	using namespace ObjectCacheContextImpl;
 #if WITH_EDITOR
 	const EObjectReverseLookupMode ObjectCacheContextMode = GetObjectReverseLookupMode();
 	if (ObjectCacheContextMode == EObjectReverseLookupMode::Permanent)
@@ -390,7 +575,7 @@ TObjectCacheIterator<UTexture> FObjectCacheContext::GetUsedTextures(UMaterialInt
 		TSet<UTexture*> LookupResult;
 		for (UTexture* Texture : GTextureToMaterialLookupCache.GetTo(MaterialInterface))
 		{
-			LookupResult.Add(Texture);
+			LookupResult.FindOrAdd(Texture);
 			checkf(Textures->Contains(Texture), TEXT("Permanent Object Cache has an additionnal texture %s on material %s"), *Texture->GetFullName(), *MaterialInterface->GetFullName());
 		}
 
@@ -406,13 +591,15 @@ TObjectCacheIterator<UTexture> FObjectCacheContext::GetUsedTextures(UMaterialInt
 
 TObjectCacheIterator<USkinnedMeshComponent> FObjectCacheContext::GetSkinnedMeshComponents()
 {
+	using namespace ObjectCacheContextImpl;
+
 	if (!SkinnedMeshComponents.IsSet())
 	{
 		TRACE_CPUPROFILER_EVENT_SCOPE(ComputeSkinnedMeshComponents);
 
 		TArray<USkinnedMeshComponent*> Array;
 		Array.Reserve(4096);
-		for (TObjectIterator<USkinnedMeshComponent> It; It; ++It)
+		for (TObjectIterator<USkinnedMeshComponent> It(RF_ClassDefaultObject, true /*bIncludeDerivedClasses*/, GetObjectCacheInternalFlagsExclusion()); It; ++It)
 		{
 			Array.Add(*It);
 		}
@@ -424,6 +611,7 @@ TObjectCacheIterator<USkinnedMeshComponent> FObjectCacheContext::GetSkinnedMeshC
 
 TObjectCacheIterator<UMaterialInterface> FObjectCacheContext::GetUsedMaterials(UPrimitiveComponent* Component)
 {
+	using namespace ObjectCacheContextImpl;
 #if WITH_EDITOR
 	const EObjectReverseLookupMode ObjectCacheContextMode = GetObjectReverseLookupMode();
 	if (ObjectCacheContextMode == EObjectReverseLookupMode::Permanent)
@@ -437,29 +625,33 @@ TObjectCacheIterator<UMaterialInterface> FObjectCacheContext::GetUsedMaterials(U
 	{
 		Materials = &PrimitiveComponentToMaterial.Add(Component);
 		TArray<UMaterialInterface*> OutMaterials;
+		// GetUsedMaterials can sometimes return nullptr... need to filter them out
 		Component->GetUsedMaterials(OutMaterials, true);
-		Materials->Append(OutMaterials);
+		for (UMaterialInterface* MaterialInterface : OutMaterials)
+		{
+			if (MaterialInterface)
+			{
+				Materials->FindOrAdd(MaterialInterface);
+			}
+		}
 	}
 
 #if WITH_EDITOR
+	// The only case where reverse lookup is allowed to diverge is when the render state is dirty meaning
+	// UsedMaterials might have been modified but the SceneProxy hasn't been recreated yet.
+	// This is not a problem as we're using the reverse proxy to know which component needs to be dirtied anyway.
 	if (ObjectCacheContextMode == EObjectReverseLookupMode::Comparison && Component->IsRenderStateCreated() && Component->SceneProxy && !Component->IsRenderStateDirty())
 	{
 		TSet<UMaterialInterface*> LookupResult;
 		for (UMaterialInterface* Material : GMaterialToPrimitiveLookupCache.GetTo(Component))
 		{
-			LookupResult.Add(Material);
-			if (!Materials->Contains(Material))
-			{
-				UE_LOG(LogObjectCache, Warning, TEXT("Permanent Object Cache has an additionnal material %s for component %"), *Material->GetFullName(), *Component->GetFullName());
-			}
+			LookupResult.FindOrAdd(Material);
+			checkf(Materials->Contains(Material), TEXT("Permanent Object Cache has an additionnal material %s for component %"), *Material->GetFullName(), *Component->GetFullName());
 		}
 
 		for (UMaterialInterface* Material : *Materials)
 		{
-			if (!LookupResult.Contains(Material))
-			{
-				UE_LOG(LogObjectCache, Warning, TEXT("Permanent Object Cache is missing a material %s on component %s"), *Material->GetFullName(), *Component->GetFullName());
-			}
+			checkf(LookupResult.Contains(Material), TEXT("Permanent Object Cache is missing a material %s on component %s"), *Material->GetFullName(), *Component->GetFullName());
 		}
 	}
 #endif
@@ -469,6 +661,7 @@ TObjectCacheIterator<UMaterialInterface> FObjectCacheContext::GetUsedMaterials(U
 
 TObjectCacheIterator<UStaticMeshComponent> FObjectCacheContext::GetStaticMeshComponents(UStaticMesh* InStaticMesh)
 {
+	using namespace ObjectCacheContextImpl;
 #if WITH_EDITOR
 	const EObjectReverseLookupMode ObjectCacheContextMode = GetObjectReverseLookupMode();
 	if (ObjectCacheContextMode == EObjectReverseLookupMode::Permanent)
@@ -485,7 +678,7 @@ TObjectCacheIterator<UStaticMeshComponent> FObjectCacheContext::GetStaticMeshCom
 		TempMap.Reserve(8192);
 		for (UStaticMeshComponent* Component : GetStaticMeshComponents())
 		{
-			TempMap.FindOrAdd(Component->GetStaticMesh()).Add(Component);
+			TempMap.FindOrAdd(Component->GetStaticMesh()).FindOrAdd(Component);
 		}
 		StaticMeshToComponents = MoveTemp(TempMap);
 	}
@@ -500,7 +693,7 @@ TObjectCacheIterator<UStaticMeshComponent> FObjectCacheContext::GetStaticMeshCom
 		TSet<UStaticMeshComponent*> LookupResult;
 		for (UStaticMeshComponent* Component : GStaticMeshToComponentLookupCache.GetFrom(InStaticMesh))
 		{
-			LookupResult.Add(Component);
+			LookupResult.FindOrAdd(Component);
 			checkf(ComputeResult.Contains(Component), TEXT("Permanent Object Cache has an additionnal component %s for staticmesh %s"), *Component->GetFullName(), *InStaticMesh->GetFullName());
 		}
 
@@ -516,6 +709,7 @@ TObjectCacheIterator<UStaticMeshComponent> FObjectCacheContext::GetStaticMeshCom
 
 TObjectCacheIterator<UMaterialInterface> FObjectCacheContext::GetMaterialsAffectedByTexture(UTexture* InTexture)
 {
+	using namespace ObjectCacheContextImpl;
 #if WITH_EDITOR
 	const EObjectReverseLookupMode ObjectCacheContextMode = GetObjectReverseLookupMode();
 	if (ObjectCacheContextMode == EObjectReverseLookupMode::Permanent)
@@ -530,12 +724,12 @@ TObjectCacheIterator<UMaterialInterface> FObjectCacheContext::GetMaterialsAffect
 
 		TMap<TObjectKey<UTexture>, TSet<UMaterialInterface*>> TempMap;
 		TempMap.Reserve(8192);
-		for (TObjectIterator<UMaterialInterface> It; It; ++It)
+		for (TObjectIterator<UMaterialInterface> It(RF_ClassDefaultObject, true /*bIncludeDerivedClasses*/, GetObjectCacheInternalFlagsExclusion()); It; ++It)
 		{
 			UMaterialInterface* MaterialInterface = *It;
 			for (UTexture* Texture : GetUsedTextures(MaterialInterface))
 			{
-				TempMap.FindOrAdd(Texture).Add(MaterialInterface);
+				TempMap.FindOrAdd(Texture).FindOrAdd(MaterialInterface);
 			}
 		}
 		TextureToMaterials = MoveTemp(TempMap);
@@ -551,7 +745,7 @@ TObjectCacheIterator<UMaterialInterface> FObjectCacheContext::GetMaterialsAffect
 		TSet<UMaterialInterface*> LookupResult;
 		for (UMaterialInterface* Material : GTextureToMaterialLookupCache.GetFrom(InTexture))
 		{
-			LookupResult.Add(Material);
+			LookupResult.FindOrAdd(Material);
 			checkf(ComputeResult.Contains(Material), TEXT("Permanent Object Cache has an additionnal material %s for texture %s"), *Material->GetFullName(), *InTexture->GetFullName());
 		}
 
@@ -569,13 +763,13 @@ TObjectCacheIterator<USkinnedMeshComponent> FObjectCacheContext::GetSkinnedMeshC
 {
 	if (!SkinnedAssetToComponents.IsSet())
 	{
-		TRACE_CPUPROFILER_EVENT_SCOPE(ComputeSkeletalMeshToComponents);
+		TRACE_CPUPROFILER_EVENT_SCOPE(ComputeSkinnedMeshToComponents);
 
 		TMap<TObjectKey<USkinnedAsset>, TSet<USkinnedMeshComponent*>> TempMap;
 		TempMap.Reserve(8192);
 		for (USkinnedMeshComponent* Component : GetSkinnedMeshComponents())
 		{
-			TempMap.FindOrAdd(Component->GetSkinnedAsset()).Add(Component);
+			TempMap.FindOrAdd(Component->GetSkinnedAsset()).FindOrAdd(Component);
 		}
 		SkinnedAssetToComponents = MoveTemp(TempMap);
 	}
@@ -587,69 +781,97 @@ TObjectCacheIterator<USkinnedMeshComponent> FObjectCacheContext::GetSkinnedMeshC
 	return TObjectCacheIterator<USkinnedMeshComponent>(Result.Array());
 }
 
-TObjectCacheIterator<UPrimitiveComponent> FObjectCacheContext::GetPrimitivesAffectedByMaterial(UMaterialInterface* InMaterialInterface)
+TObjectCacheIterator<UMaterialInterface> FObjectCacheContext::GetMaterialsAffectedByMaterials(TArrayView<UMaterialInterface*> InMaterials)
 {
+	using namespace ObjectCacheContextImpl;
+	TRACE_CPUPROFILER_EVENT_SCOPE(FObjectCacheContext::GetMaterialsAffectedByMaterials);
+
 #if WITH_EDITOR
 	const EObjectReverseLookupMode ObjectCacheContextMode = GetObjectReverseLookupMode();
 	if (ObjectCacheContextMode == EObjectReverseLookupMode::Permanent)
 	{
-		return GMaterialToPrimitiveLookupCache.GetFrom(InMaterialInterface);
+		TSet<UMaterialInterface*> LookupResults;
+		ObjectCacheContextImpl::GetMaterialsAffectedByMaterials_Lookup(InMaterials, LookupResults);
+		return TObjectCacheIterator<UMaterialInterface>(LookupResults.Array());
 	}
 #endif
 
-	if (!MaterialToPrimitives.IsSet())
+	TSet<UMaterialInterface*> IterationResults;
+	ObjectCacheContextImpl::GetMaterialsAffectedByMaterials_Iteration(InMaterials, IterationResults);
+
+#if WITH_EDITOR
+	if (ObjectCacheContextMode == EObjectReverseLookupMode::Comparison)
 	{
-		TRACE_CPUPROFILER_EVENT_SCOPE(ComputePrimitivesAffectedByMaterial);
+		TSet<UMaterialInterface*> LookupResults;
+		ObjectCacheContextImpl::GetMaterialsAffectedByMaterials_Lookup(InMaterials, LookupResults);
+		
+		for (UMaterialInterface* MaterialInterface : LookupResults)
+		{
+			checkf(IterationResults.Contains(MaterialInterface), TEXT("Permanent Object Cache has an additionnal %s"), *MaterialInterface->GetFullName());
+		}
 
-		TMap<TObjectKey<UMaterialInterface>, TSet<UPrimitiveComponent*>> TempMap;
-		for (UPrimitiveComponent* Component : GetPrimitiveComponents())
+		for (UMaterialInterface* MaterialInterface : IterationResults)
 		{
-			if (Component->IsRegistered() && Component->IsRenderStateCreated() && Component->SceneProxy)
-			{
-				for (UMaterialInterface* MaterialInterface : GetUsedMaterials(Component))
-				{
-					if (MaterialInterface)
-					{
-						TempMap.FindOrAdd(MaterialInterface).Add(Component);
-					}
-				}
-			}
+			checkf(LookupResults.Contains(MaterialInterface), TEXT("Permanent Object Cache is missing a %s"), *MaterialInterface->GetFullName());
 		}
+	}
+#endif
+
+	return TObjectCacheIterator<UMaterialInterface>(IterationResults.Array());
+}
+
+TObjectCacheIterator<UPrimitiveComponent> FObjectCacheContext::GetPrimitivesAffectedByMaterials(TArrayView<UMaterialInterface*> InMaterials)
+{
+	using namespace ObjectCacheContextImpl;
+	TRACE_CPUPROFILER_EVENT_SCOPE(FObjectCacheContext::GetPrimitivesAffectedByMaterials);
 
-		MaterialToPrimitives = MoveTemp(TempMap);
+#if WITH_EDITOR
+	const EObjectReverseLookupMode ObjectCacheContextMode = GetObjectReverseLookupMode();
+	if (ObjectCacheContextMode == EObjectReverseLookupMode::Permanent)
+	{
+		TSet<UPrimitiveComponent*> AffectedPrimitives;
+		ObjectCacheContextImpl::GetPrimitivesAffectedByMaterials_Lookup(InMaterials, AffectedPrimitives);
+		return TObjectCacheIterator<UPrimitiveComponent>(AffectedPrimitives.Array());
 	}
+#endif
+
+	TSet<UPrimitiveComponent*> IterationResults;
+	ObjectCacheContextImpl::GetPrimitivesAffectedByMaterials_Iteration(*this, InMaterials, IterationResults);
 
-	static TSet<UPrimitiveComponent*> EmptySet;
-	TSet<UPrimitiveComponent*>* Set = MaterialToPrimitives.GetValue().Find(InMaterialInterface);
-	const TSet<UPrimitiveComponent*>& ComputeResult = Set ? *Set : EmptySet;
-	
 #if WITH_EDITOR
 	if (ObjectCacheContextMode == EObjectReverseLookupMode::Comparison)
 	{
-		TSet<UPrimitiveComponent*> LookupResult;
-		for (UPrimitiveComponent* Component : GMaterialToPrimitiveLookupCache.GetFrom(InMaterialInterface))
+		TSet<UPrimitiveComponent*> LookupResults;
+		ObjectCacheContextImpl::GetPrimitivesAffectedByMaterials_Lookup(InMaterials, LookupResults);
+
+		// The only case where reverse lookup is allowed to diverge is when the render state is dirty meaning
+		// UsedMaterials might have been modified but the SceneProxy hasn't been recreated yet.
+		// This is not a problem as we're using the reverse proxy to know which component needs to be dirtied anyway.
+		for (UPrimitiveComponent* Component : LookupResults)
 		{
-			LookupResult.Add(Component);
-			if(!ComputeResult.Contains(Component))
+			if (Component->IsRenderStateCreated() && Component->SceneProxy && !Component->IsRenderStateDirty())
 			{
-				UE_LOG(LogObjectCache, Warning, TEXT("Permanent Object Cache has an additionnal component %s for material %"), *Component->GetFullName(), *InMaterialInterface->GetFullName());
+				checkf(IterationResults.Contains(Component), TEXT("Permanent Object Cache has an additionnal %s"), *Component->GetFullName());
 			}
 		}
 
-		for (UPrimitiveComponent* Component : ComputeResult)
+		for (UPrimitiveComponent* Component : IterationResults)
 		{
 			if (Component->IsRenderStateCreated() && Component->SceneProxy && !Component->IsRenderStateDirty())
 			{
-				if (!LookupResult.Contains(Component))
-				{
-					UE_LOG(LogObjectCache, Warning, TEXT("Permanent Object Cache is missing a component %s for material %s"), *Component->GetFullName(), *InMaterialInterface->GetFullName());
-				}
+				checkf(LookupResults.Contains(Component), TEXT("Permanent Object Cache is missing a %s"), *Component->GetFullName());
 			}
 		}
 	}
 #endif
 
-	return TObjectCacheIterator<UPrimitiveComponent>(ComputeResult.Array());
+	return TObjectCacheIterator<UPrimitiveComponent>(IterationResults.Array());
+}
+
+TObjectCacheIterator<UPrimitiveComponent> FObjectCacheContext::GetPrimitivesAffectedByMaterial(UMaterialInterface* InMaterialInterface)
+{
+	UMaterialInterface* InplaceArray[1] = { InMaterialInterface };
+	return GetPrimitivesAffectedByMaterials(InplaceArray);
 }
 
 namespace ObjectCacheContextScopeImpl
@@ -684,6 +906,8 @@ FObjectCacheContext& FObjectCacheContextScope::GetContext()
 
 #if WITH_EDITOR
 
+namespace ObjectCacheContextImpl {
+
 struct FObjectCacheEventSinkPrivate
 {
 	static void BeginQueueNotifyEvents();
@@ -723,23 +947,33 @@ struct FObjectCacheEventSinkPrivate
 	static void AddNotifyEvent(FObjectCacheEventSinkPrivate::ECacheEventType EventType, UMaterialInterface* MaterialInterface, const UPrimitiveComponent* PrimitiveComponent, UStaticMeshComponent* StaticMeshComponent, const TArray<UMaterialInterface*>* UsedMaterials);
 };
 
+} // namespace ObjectCacheContextImpl
+
 void FObjectCacheEventSink::BeginQueueNotifyEvents()
 {
+	using namespace ObjectCacheContextImpl;
+
 	FObjectCacheEventSinkPrivate::BeginQueueNotifyEvents();
 }
 
 void FObjectCacheEventSink::ProcessQueuedNotifyEvents()
 {
+	using namespace ObjectCacheContextImpl;
+
 	FObjectCacheEventSinkPrivate::ProcessQueuedNotifyEvents();
 }
 
 void FObjectCacheEventSink::EndQueueNotifyEvents()
 {
+	using namespace ObjectCacheContextImpl;
+
 	FObjectCacheEventSinkPrivate::EndQueueNotifyEvents();
 }
 
 void FObjectCacheEventSink::NotifyMaterialDestroyed_Concurrent(UMaterialInterface* MaterialInterface)
 {
+	using namespace ObjectCacheContextImpl;
+
 	if (GetObjectReverseLookupMode() != EObjectReverseLookupMode::Temporary)
 	{
 		if (FObjectCacheEventSinkPrivate::bShouldQueueSinkEvents)
@@ -755,6 +989,8 @@ void FObjectCacheEventSink::NotifyMaterialDestroyed_Concurrent(UMaterialInterfac
 
 void FObjectCacheEventSink::NotifyUsedMaterialsChanged_Concurrent(const UPrimitiveComponent* PrimitiveComponent, const TArray<UMaterialInterface*>& UsedMaterials)
 {
+	using namespace ObjectCacheContextImpl;
+
 	if (GetObjectReverseLookupMode() != EObjectReverseLookupMode::Temporary)
 	{
 		if (FObjectCacheEventSinkPrivate::bShouldQueueSinkEvents)
@@ -770,6 +1006,8 @@ void FObjectCacheEventSink::NotifyUsedMaterialsChanged_Concurrent(const UPrimiti
 
 void FObjectCacheEventSink::NotifyRenderStateChanged_Concurrent(UActorComponent* ActorComponent)
 {
+	using namespace ObjectCacheContextImpl;
+
 	if (GetObjectReverseLookupMode() != EObjectReverseLookupMode::Temporary)
 	{
 		if (UPrimitiveComponent* Primitive = Cast<UPrimitiveComponent>(ActorComponent))
@@ -788,6 +1026,8 @@ void FObjectCacheEventSink::NotifyRenderStateChanged_Concurrent(UActorComponent*
 
 void FObjectCacheEventSink::NotifyReferencedTextureChanged_Concurrent(UMaterialInterface* MaterialInterface)
 {
+	using namespace ObjectCacheContextImpl;
+
 	if (GetObjectReverseLookupMode() != EObjectReverseLookupMode::Temporary)
 	{
 		if (FObjectCacheEventSinkPrivate::bShouldQueueSinkEvents)
@@ -803,6 +1043,8 @@ void FObjectCacheEventSink::NotifyReferencedTextureChanged_Concurrent(UMaterialI
 
 void FObjectCacheEventSink::NotifyStaticMeshChanged_Concurrent(UStaticMeshComponent* StaticMeshComponent)
 {
+	using namespace ObjectCacheContextImpl;
+
 	if (GetObjectReverseLookupMode() != EObjectReverseLookupMode::Temporary)
 	{
 		if (FObjectCacheEventSinkPrivate::bShouldQueueSinkEvents)
@@ -816,6 +1058,8 @@ void FObjectCacheEventSink::NotifyStaticMeshChanged_Concurrent(UStaticMeshCompon
 	}
 }
 
+namespace ObjectCacheContextImpl {
+
 std::atomic<bool> FObjectCacheEventSinkPrivate::bShouldQueueSinkEvents(false);
 
 FObjectCacheEventSinkPrivate::TEventAllocator& FObjectCacheEventSinkPrivate::GetAllocator()
@@ -895,6 +1139,8 @@ void FObjectCacheEventSinkPrivate::EndQueueNotifyEvents()
 
 void FObjectCacheEventSinkPrivate::NotifyMaterialDestroyed_Concurrent(UMaterialInterface* MaterialInterface)
 {
+	using namespace ObjectCacheContextImpl;
+
 	// If a material is destroyed, remove it from the cache so we don't return it anymore
 	// This can happen if a component doesn't update its render state after modifying its used materials
 	// (i.e. LandscapeComponent won't update it's render state when modifying MobileMaterialInterfaces during cook since they aren't meaningful for rendering)
@@ -903,35 +1149,61 @@ void FObjectCacheEventSinkPrivate::NotifyMaterialDestroyed_Concurrent(UMaterialI
 
 	// Remove any Material to Texture that might be present in the cache for this Material
 	GTextureToMaterialLookupCache.Update(MaterialInterface, {});
+
+	// Cleanup up any mapping that this material could have with other materials
+	if (UMaterialInstance* MaterialInstance = Cast<UMaterialInstance>(MaterialInterface))
+	{
+		GMaterialToMaterialInstanceLookupCache.Update(MaterialInstance, {});
+	}
 }
 
 void FObjectCacheEventSinkPrivate::NotifyUsedMaterialsChanged_Concurrent(const UPrimitiveComponent* PrimitiveComponent, const TArray<UMaterialInterface*>& UsedMaterials)
 {
+	using namespace ObjectCacheContextImpl;
+
 	GMaterialToPrimitiveLookupCache.Update(const_cast<UPrimitiveComponent*>(PrimitiveComponent), UsedMaterials);
 }
 
 void FObjectCacheEventSinkPrivate::NotifyRenderStateChanged_Concurrent(const UPrimitiveComponent* PrimitiveComponent)
 {
+	using namespace ObjectCacheContextImpl;
+
 	// Clear mappings whenever the render state changes until NotifyUsedMaterialsChanged is called during SceneProxy creation
 	GMaterialToPrimitiveLookupCache.Update(const_cast<UPrimitiveComponent*>(PrimitiveComponent), {});
 }
 
 void FObjectCacheEventSinkPrivate::NotifyReferencedTextureChanged_Concurrent(UMaterialInterface* MaterialInterface)
 {
+	using namespace ObjectCacheContextImpl;
+
 	if (MaterialInterface->HasAnyFlags(RF_BeginDestroyed))
 	{
-		GTextureToMaterialLookupCache.Update(MaterialInterface, {});
+		NotifyMaterialDestroyed_Concurrent(MaterialInterface);
 	}
 	else
 	{
 		TSet<UTexture*> Textures;
 		ObjectCacheContextImpl::GetReferencedTextures(MaterialInterface, Textures);
 		GTextureToMaterialLookupCache.Update(MaterialInterface, Textures.Array());
+
+		if (UMaterialInstance* MaterialInstance = Cast<UMaterialInstance>(MaterialInterface))
+		{
+			if (MaterialInstance->Parent)
+			{
+				GMaterialToMaterialInstanceLookupCache.Update(MaterialInstance, { MaterialInstance->Parent });
+			}
+			else
+			{
+				GMaterialToMaterialInstanceLookupCache.Update(MaterialInstance, { });
+			}
+		}
 	}
 }
 
 void FObjectCacheEventSinkPrivate::NotifyStaticMeshChanged_Concurrent(UStaticMeshComponent* StaticMeshComponent)
 {
+	using namespace ObjectCacheContextImpl;
+
 	// Stop tracking components that are being destroyed
 	if (StaticMeshComponent->HasAnyFlags(RF_BeginDestroyed) || StaticMeshComponent->GetStaticMesh() == nullptr)
 	{
@@ -943,6 +1215,8 @@ void FObjectCacheEventSinkPrivate::NotifyStaticMeshChanged_Concurrent(UStaticMes
 	}
 }
 
+} // namespace ObjectCacheContextImpl 
+
 #endif
 
 #undef LOCTEXT_NAMESPACE
\ No newline at end of file
diff --git a/Engine/Source/Runtime/Engine/Private/ShaderCompiler/ShaderCompiler.cpp b/Engine/Source/Runtime/Engine/Private/ShaderCompiler/ShaderCompiler.cpp
--- a/Engine/Source/Runtime/Engine/Private/ShaderCompiler/ShaderCompiler.cpp
+++ b/Engine/Source/Runtime/Engine/Private/ShaderCompiler/ShaderCompiler.cpp
@@ -4912,44 +4912,24 @@ void FShaderCompilingManager::PropagateMaterialChangesToPrimitives(TMap<TRefCoun
 		return;
 	}
 
-	TRACE_CPUPROFILER_EVENT_SCOPE(FShaderCompilingManager::PropagateMaterialChangesToPrimitives);
-
 	TSet<FSceneInterface*> ScenesToUpdate;
 	FObjectCacheContextScope ObjectCacheScope;
 	TIndirectArray<FComponentRecreateRenderStateContext> ComponentContexts;
-	for (UPrimitiveComponent* PrimitiveComponent : ObjectCacheScope.GetContext().GetPrimitiveComponents())
 	{
-		if (PrimitiveComponent->IsRenderStateCreated())
+		TRACE_CPUPROFILER_EVENT_SCOPE(FShaderCompilingManager::PropagateMaterialChangesToPrimitives);
+
+		TArray<UMaterialInterface*> UpdatedMaterials;
+		for (TMap<TRefCountPtr<FMaterial>, TRefCountPtr<FMaterialShaderMap>>::TConstIterator MaterialIt(MaterialsToUpdate); MaterialIt; ++MaterialIt)
 		{
-			bool bPrimitiveIsDependentOnMaterial = false;
+			FMaterial* UpdatedMaterial = MaterialIt.Key();
+			UpdatedMaterials.Add(UpdatedMaterial->GetMaterialInterface());
+		}
 
-			// Note: relying on GetUsedMaterials to be accurate, or else we won't propagate to the right primitives and the renderer will crash later
-			// FPrimitiveSceneProxy::VerifyUsedMaterial is used to make sure that all materials used for rendering are reported in GetUsedMaterials
-			TObjectCacheIterator<UMaterialInterface> UsedMaterials = ObjectCacheScope.GetContext().GetUsedMaterials(PrimitiveComponent);
-			if (!UsedMaterials.IsEmpty())
+		for (UPrimitiveComponent* PrimitiveComponent : ObjectCacheScope.GetContext().GetPrimitivesAffectedByMaterials(UpdatedMaterials))
+		{
+			if (PrimitiveComponent->IsRenderStateCreated() && PrimitiveComponent->SceneProxy)
 			{
-				for (TMap<TRefCountPtr<FMaterial>, TRefCountPtr<FMaterialShaderMap>>::TConstIterator MaterialIt(MaterialsToUpdate); MaterialIt; ++MaterialIt)
-				{
-					FMaterial* UpdatedMaterial = MaterialIt.Key();
-					UMaterialInterface* UpdatedMaterialInterface = UpdatedMaterial->GetMaterialInterface();
-
-					if (UpdatedMaterialInterface)
-					{
-						for (UMaterialInterface* TestMaterial : UsedMaterials)
-						{
-							if (TestMaterial && (TestMaterial == UpdatedMaterialInterface || TestMaterial->IsDependent(UpdatedMaterialInterface)))
-							{
-								bPrimitiveIsDependentOnMaterial = true;
-								break;
-							}
-						}
-					}
-				}
-
-				if (bPrimitiveIsDependentOnMaterial)
-				{
-					ComponentContexts.Add(new FComponentRecreateRenderStateContext(PrimitiveComponent, &ScenesToUpdate));
-				}
+				ComponentContexts.Add(new FComponentRecreateRenderStateContext(PrimitiveComponent, &ScenesToUpdate));
 			}
 		}
 	}
diff --git a/Engine/Source/Runtime/Engine/Public/ObjectCacheContext.h b/Engine/Source/Runtime/Engine/Public/ObjectCacheContext.h
--- a/Engine/Source/Runtime/Engine/Public/ObjectCacheContext.h
+++ b/Engine/Source/Runtime/Engine/Public/ObjectCacheContext.h
@@ -73,10 +73,11 @@ public:
 	TObjectCacheIterator<UStaticMeshComponent>  GetStaticMeshComponents(UStaticMesh* InStaticMesh);
 	TObjectCacheIterator<UMaterialInterface>    GetMaterialsAffectedByTexture(UTexture* InTexture);
 	TObjectCacheIterator<UPrimitiveComponent>   GetPrimitivesAffectedByMaterial(UMaterialInterface* InMaterial);
+	TObjectCacheIterator<UPrimitiveComponent>   GetPrimitivesAffectedByMaterials(TArrayView<UMaterialInterface*> InMaterials);
 	TObjectCacheIterator<UTexture>              GetUsedTextures(UMaterialInterface* InMaterial);
 	TObjectCacheIterator<UMaterialInterface>    GetUsedMaterials(UPrimitiveComponent* InComponent);
+	TObjectCacheIterator<UMaterialInterface>    GetMaterialsAffectedByMaterials(TArrayView<UMaterialInterface*> InMaterials);
 
-	
 private:
 	friend class FObjectCacheContextScope;
 	FObjectCacheContext() = default;
