Remove usages of TFunctionRef vars on stack causing stack corruption on clang
Fix double-register ensure when FlushDormancyForObject  new replicator overwrites  an existing replicator

diff --git a/Engine/Source/Runtime/Engine/Private/DemoNetDriver.cpp b/Engine/Source/Runtime/Engine/Private/DemoNetDriver.cpp
--- a/Engine/Source/Runtime/Engine/Private/DemoNetDriver.cpp
+++ b/Engine/Source/Runtime/Engine/Private/DemoNetDriver.cpp
@@ -3018,7 +3018,7 @@ void UDemoNetDriver::FinalizeFastForward(const double StartTime)
 			}
 		}
 
-		UE::Net::FExecuteForEachDormantReplicator CallRepNotifies = [](AActor* OwnerActor, FObjectKey ObjectKey, const TSharedRef<FObjectReplicator>& ReplicatorRef)
+		auto CallRepNotifies = [](AActor* OwnerActor, FObjectKey ObjectKey, const TSharedRef<FObjectReplicator>& ReplicatorRef)
 		{
 			ReplicatorRef->CallRepNotifies(true);
 		};
@@ -3921,7 +3921,7 @@ bool UDemoNetDriver::FastForwardLevels(const FGotoResult& GotoResult)
 			}
 		}
 
-		UE::Net::FExecuteForEachDormantReplicator CallRepNotifies = [](AActor* OwnerActor, FObjectKey ObjectKey, const TSharedRef<FObjectReplicator>& ReplicatorRef)
+		auto CallRepNotifies = [](AActor* OwnerActor, FObjectKey ObjectKey, const TSharedRef<FObjectReplicator>& ReplicatorRef)
 		{
 			ReplicatorRef->CallRepNotifies(true);
 		};
diff --git a/Engine/Source/Runtime/Engine/Private/Net/NetDormantHolder.cpp b/Engine/Source/Runtime/Engine/Private/Net/NetDormantHolder.cpp
--- a/Engine/Source/Runtime/Engine/Private/Net/NetDormantHolder.cpp
+++ b/Engine/Source/Runtime/Engine/Private/Net/NetDormantHolder.cpp
@@ -78,7 +78,7 @@ TSharedPtr<FObjectReplicator> FDormantReplicatorHolder::FindAndRemoveReplicator(
 	return ReplicatorPtr;
 }
 
-const TSharedRef<FObjectReplicator>& FDormantReplicatorHolder::CreateAndStoreReplicator(AActor* DormantActor, UObject* ReplicatedObject)
+const TSharedRef<FObjectReplicator>& FDormantReplicatorHolder::CreateAndStoreReplicator(AActor* DormantActor, UObject* ReplicatedObject, bool& bOverwroteExistingReplicator)
 {
 	FActorDormantReplicators& ActorReplicators = ActorReplicatorSet.FindOrAdd(FActorDormantReplicators(DormantActor));
 
@@ -86,7 +86,7 @@ const TSharedRef<FObjectReplicator>& FDormantReplicatorHolder::CreateAndStoreRep
 
 	// Add a new replicator tied to this object. 
 	// If there was already a replicator for the same object in the set, it will be destroyed and overwritten by this new one.
-	FSetElementId Index = ActorReplicators.DormantReplicators.Add(FDormantObjectReplicator(SubObjectKey));
+	FSetElementId Index = ActorReplicators.DormantReplicators.Add(FDormantObjectReplicator(SubObjectKey), &bOverwroteExistingReplicator);
 
 	return ActorReplicators.DormantReplicators[Index].Replicator;
 }
diff --git a/Engine/Source/Runtime/Engine/Private/NetConnection.cpp b/Engine/Source/Runtime/Engine/Private/NetConnection.cpp
--- a/Engine/Source/Runtime/Engine/Private/NetConnection.cpp
+++ b/Engine/Source/Runtime/Engine/Private/NetConnection.cpp
@@ -4834,14 +4834,16 @@ void UNetConnection::FlushDormancyForObject(AActor* DormantActor, UObject* Repli
 	// If we need to create a new replicator when flushing
 	if (!bReuseReplicators)
 	{
-		const TSharedRef<FObjectReplicator>& ObjectReplicatorRef = DormantReplicatorSet.CreateAndStoreReplicator(DormantActor, ReplicatedObject);
+		bool bOverwroteExistingReplicator = false;
+		const TSharedRef<FObjectReplicator>& ObjectReplicatorRef = DormantReplicatorSet.CreateAndStoreReplicator(DormantActor, ReplicatedObject, bOverwroteExistingReplicator);
 		
 		// Init using the objects current state
 		constexpr bool bUseDefaultState = false; 
 		ObjectReplicatorRef->InitWithObject(ReplicatedObject, this, bUseDefaultState);
 
 #if UE_REPLICATED_OBJECT_REFCOUNTING
-		if (Driver && DormantActor != ReplicatedObject)
+		// Add a refcount only when we did not create a replicator on top of an existing one.
+		if (Driver && DormantActor != ReplicatedObject && !bOverwroteExistingReplicator)
 		{
 			Driver->GetNetworkObjectList().AddSubObjectChannelReference(DormantActor, ReplicatedObject, this);
 		}
@@ -5095,7 +5097,7 @@ void UNetConnection::CleanupDormantReplicatorsForActor(AActor* Actor)
 		if (Driver)
 		{
 			TArray<TWeakObjectPtr<UObject>, TInlineAllocator<16>> RemovedObjects;
-			UE::Net::FExecuteForEachDormantReplicator ExecuteFunction = [&RemovedObjects](AActor* OwnerActor, FObjectKey ObjectKey, const TSharedRef<FObjectReplicator>& ReplicatorRef)
+			auto ExecuteFunction = [&RemovedObjects](AActor* OwnerActor, FObjectKey ObjectKey, const TSharedRef<FObjectReplicator>& ReplicatorRef)
 			{
 				if (OwnerActor != ReplicatorRef->GetObject())
 				{
diff --git a/Engine/Source/Runtime/Engine/Private/NetDriver.cpp b/Engine/Source/Runtime/Engine/Private/NetDriver.cpp
--- a/Engine/Source/Runtime/Engine/Private/NetDriver.cpp
+++ b/Engine/Source/Runtime/Engine/Private/NetDriver.cpp
@@ -5171,7 +5171,7 @@ int32 UNetDriver::ServerReplicateActors(float DeltaSeconds)
 		// net.DormancyValidate can be set to 2 to validate all dormant actors against last known state before going dormant
 		if ( GNetDormancyValidate == 2 )
 		{
-			UE::Net::FExecuteForEachDormantReplicator ValidateFunction = [](AActor* OwnerActor, FObjectKey ObjectKey, const TSharedRef<FObjectReplicator>& ReplicatorRef)
+			auto ValidateFunction = [](AActor* OwnerActor, FObjectKey ObjectKey, const TSharedRef<FObjectReplicator>& ReplicatorRef)
 			{
 				FObjectReplicator& Replicator = ReplicatorRef.Get();
 
diff --git a/Engine/Source/Runtime/Engine/Public/Net/NetDormantHolder.h b/Engine/Source/Runtime/Engine/Public/Net/NetDormantHolder.h
--- a/Engine/Source/Runtime/Engine/Public/Net/NetDormantHolder.h
+++ b/Engine/Source/Runtime/Engine/Public/Net/NetDormantHolder.h
@@ -109,9 +109,10 @@ struct FDormantReplicatorHolder
 	* 
 	* @param DormantActor The dormant actor that owns the replicated object
 	* @param ReplicatedObject The object that will be tied to the new ObjectReplicator.
+	* @param bOverwroteExistingReplicator When it returns true it means we erased an existing replicator.
 	* @return Return a reference to the object replicator we created and are now storing.
 	*/
-	const TSharedRef<FObjectReplicator>& CreateAndStoreReplicator(AActor* DormantActor, UObject* ReplicatedObject);
+	const TSharedRef<FObjectReplicator>& CreateAndStoreReplicator(AActor* DormantActor, UObject* ReplicatedObject, bool& bOverwroteExistingReplicator);
 
 	/**
 	* Store an existing replicator tied to the given replicated object
