Bulk Data - handle I/O error when streaming Static/Skeletal mesh

#rnx

diff --git a/Engine/Source/Runtime/Engine/Private/Streaming/SkeletalMeshUpdate.cpp b/Engine/Source/Runtime/Engine/Private/Streaming/SkeletalMeshUpdate.cpp
--- a/Engine/Source/Runtime/Engine/Private/Streaming/SkeletalMeshUpdate.cpp
+++ b/Engine/Source/Runtime/Engine/Private/Streaming/SkeletalMeshUpdate.cpp
@@ -454,7 +454,7 @@ void FSkeletalMeshStreamIn_IO::SetIORequest(const FContext& Context)
 			// At this point task synchronization would hold the number of pending requests.
 			TaskSynchronization.Decrement();
 
-			if (FBulkDataRequest::EStatus::Cancelled == Status)
+			if (FBulkDataRequest::EStatus::Ok != Status)
 			{
 				// If IO requests was cancelled but the streaming request wasn't, this is an IO error.
 				if (!bIsCancelled)
diff --git a/Engine/Source/Runtime/Engine/Private/Streaming/StaticMeshUpdate.cpp b/Engine/Source/Runtime/Engine/Private/Streaming/StaticMeshUpdate.cpp
--- a/Engine/Source/Runtime/Engine/Private/Streaming/StaticMeshUpdate.cpp
+++ b/Engine/Source/Runtime/Engine/Private/Streaming/StaticMeshUpdate.cpp
@@ -493,7 +493,7 @@ void FStaticMeshStreamIn_IO::SetIORequest(const FContext& Context)
 		{
 			TaskSynchronization.Decrement();
 
-			if (FBulkDataRequest::EStatus::Cancelled == Status)
+			if (FBulkDataRequest::EStatus::Ok != Status)
 			{
 				// If IO requests was cancelled but the streaming request wasn't, this is an IO error.
 				if (!bIsCancelled)
