World Partition - HLOD: Added component & actor names in the EHLODBatchingPolicy::MeshSection warning message

diff --git a/Engine/Source/Runtime/Engine/Private/WorldPartition/HLOD/HLODBuilder.cpp b/Engine/Source/Runtime/Engine/Private/WorldPartition/HLOD/HLODBuilder.cpp
--- a/Engine/Source/Runtime/Engine/Private/WorldPartition/HLOD/HLODBuilder.cpp
+++ b/Engine/Source/Runtime/Engine/Private/WorldPartition/HLOD/HLODBuilder.cpp
@@ -6,6 +6,7 @@
 #include "Engine/HLODProxy.h"
 #include "Engine/StaticMesh.h"
 #include "Engine/Texture.h"
+#include "Engine/Level.h"
 #include "Components/InstancedStaticMeshComponent.h"
 #include "ISMPartition/ISMComponentBatcher.h"
 #include "ISMPartition/ISMComponentDescriptor.h"
@@ -217,7 +218,7 @@ static bool ShouldBatchComponent(UActorComponent* ActorComponent)
 			break;
 		case EHLODBatchingPolicy::MeshSection:
 			bShouldBatch = true;
-			UE_LOG(LogHLODBuilder, Warning, TEXT("EHLODBatchingPolicy::MeshSection is not yet supported by the HLOD builder, falling back to EHLODBatchingPolicy::Instancing."));
+			UE_LOG(LogHLODBuilder, Warning, TEXT("EHLODBatchingPolicy::MeshSection is not yet supported by the HLOD builder, falling back to EHLODBatchingPolicy::Instancing for component %s (from actor %s)."), *ActorComponent->GetName(), *ActorComponent->GetOwner()->GetActorLabel());
 			break;
 		default:
 			checkNoEntry();
