#ue5 - Prevent ConvertTransformToRelative from being added to new blueprints, and allow the old redirect in BaseEngine.ini to work properly now that the old function is no longer a UFUNCTION. When it was originally deprecated, it must have been assumed that it would prevent it from being used and that the redirect would fix it, however keeping the UFUNCTION alive meant the deprecation didn't affect blueprints or redirects.

#udn https://udn.unrealengine.com/s/question/0D54z00008SPDQoCAP/bug-convert-transform-to-relative-is-broken
#tests Before this change, added a ConvertTransformToRelative node to a level blueprint. After compiling the new change, verified that you cannot add the old nodes and that the existing ones fixed up (with params correctly swapped) to MakeRelativeTransform.

diff --git a/Engine/Source/Runtime/Engine/Classes/Kismet/KismetMathLibrary.h b/Engine/Source/Runtime/Engine/Classes/Kismet/KismetMathLibrary.h
--- a/Engine/Source/Runtime/Engine/Classes/Kismet/KismetMathLibrary.h
+++ b/Engine/Source/Runtime/Engine/Classes/Kismet/KismetMathLibrary.h
@@ -3847,7 +3847,6 @@ class ENGINE_API UKismetMathLibrary : public UBlueprintFunctionLibrary
 	static FTransform MakeRelativeTransform(const FTransform& A, const FTransform& RelativeTo);
 
 	UE_DEPRECATED(4.22, "Use MakeRelativeTransform instead, with reversed order of arguments.")
-	UFUNCTION(BlueprintPure, meta=(Keywords="cast convert"), Category="Math|Transform")
 	static FTransform ConvertTransformToRelative(const FTransform& Transform, const FTransform& ParentTransform);
 
 	/** 
