Expose the FPhysicsObject objects via the UPrimitiveComponent and specialize for USkeletalMeshComponent and UGeometryCollectionComponent

diff --git a/Engine/Source/Runtime/Engine/Classes/Components/PrimitiveComponent.h b/Engine/Source/Runtime/Engine/Classes/Components/PrimitiveComponent.h
--- a/Engine/Source/Runtime/Engine/Classes/Components/PrimitiveComponent.h
+++ b/Engine/Source/Runtime/Engine/Classes/Components/PrimitiveComponent.h
@@ -10,6 +10,7 @@
 #include "UObject/UObjectGlobals.h"
 #include "Misc/Guid.h"
 #include "InputCoreTypes.h"
+#include "Interfaces/IPhysicsComponent.h"
 #include "Templates/SubclassOf.h"
 #include "Engine/EngineTypes.h"
 #include "Engine/ScopedMovementUpdate.h"
@@ -255,7 +256,7 @@ DECLARE_DYNAMIC_MULTICAST_SPARSE_DELEGATE_TwoParams( FComponentEndTouchOverSigna
  * ShapeComponents generate geometry that is used for collision detection but are not rendered, while StaticMeshComponents and SkeletalMeshComponents contain pre-built geometry that is rendered, but can also be used for collision detection.
  */
 UCLASS(abstract, HideCategories=(Mobility, VirtualTexture), ShowCategories=(PhysicsVolume))
-class ENGINE_API UPrimitiveComponent : public USceneComponent, public INavRelevantInterface, public IInterface_AsyncCompilation
+class ENGINE_API UPrimitiveComponent : public USceneComponent, public INavRelevantInterface, public IInterface_AsyncCompilation, public IPhysicsComponent
 {
 	GENERATED_BODY()
 
@@ -2739,6 +2740,13 @@ public:
 	void DispatchOnReleased(FKey ButtonReleased = EKeys::LeftMouseButton);
 	void DispatchOnInputTouchBegin(const ETouchIndex::Type Key);
 	void DispatchOnInputTouchEnd(const ETouchIndex::Type Key);
+
+	//~ Begin IPhysicsComponent Interface.
+public:
+	virtual Chaos::FPhysicsObject* GetPhysicsObjectById(int32 Id) const override;
+	virtual Chaos::FPhysicsObject* GetPhysicsObjectByName(const FName& Name) const override;
+	virtual TArray<Chaos::FPhysicsObject*> GetAllPhysicsObjects() const override;
+	//~ End IPhysicsComponent Interface.
 };
 
 /** 
diff --git a/Engine/Source/Runtime/Engine/Classes/Components/SkeletalMeshComponent.h b/Engine/Source/Runtime/Engine/Classes/Components/SkeletalMeshComponent.h
--- a/Engine/Source/Runtime/Engine/Classes/Components/SkeletalMeshComponent.h
+++ b/Engine/Source/Runtime/Engine/Classes/Components/SkeletalMeshComponent.h
@@ -2586,6 +2586,12 @@ public:
 
 	// Are we currently within PostAnimEvaluation
 	bool IsPostEvaluatingAnimation() const { return bPostEvaluatingAnimation; }
+
+	//~ Begin IPhysicsComponent Interface.
+	virtual Chaos::FPhysicsObject* GetPhysicsObjectById(int32 Id) const override;
+	virtual Chaos::FPhysicsObject* GetPhysicsObjectByName(const FName& Name) const override;
+	virtual TArray<Chaos::FPhysicsObject*> GetAllPhysicsObjects() const override;
+	//~ End IPhysicsComponent Interface.
 };
 
 #if WITH_EDITOR
diff --git a/Engine/Source/Runtime/Engine/Classes/Interfaces/IPhysicsComponent.h b/Engine/Source/Runtime/Engine/Classes/Interfaces/IPhysicsComponent.h
new file mode 100644
--- /dev/null
+++ b/Engine/Source/Runtime/Engine/Classes/Interfaces/IPhysicsComponent.h
@@ -0,0 +1,25 @@
+// Copyright Epic Games, Inc. All Rights Reserved.
+
+#pragma once
+
+#include "Chaos/PhysicsObject.h"
+#include "UObject/ObjectMacros.h"
+#include "UObject/Interface.h"
+#include "IPhysicsComponent.generated.h"
+
+/** Interface for components that contains physics bodies. **/
+UINTERFACE(MinimalApi, Experimental, meta = (CannotImplementInterfaceInBlueprint))
+class UPhysicsComponent : public UInterface
+{
+	GENERATED_BODY()
+};
+
+class IPhysicsComponent
+{
+	GENERATED_BODY()
+
+public:
+	virtual Chaos::FPhysicsObject* GetPhysicsObjectById(int32 Id) const = 0;
+	virtual Chaos::FPhysicsObject* GetPhysicsObjectByName(const FName& Name) const = 0;
+	virtual TArray<Chaos::FPhysicsObject*> GetAllPhysicsObjects() const = 0;
+};
\ No newline at end of file
diff --git a/Engine/Source/Runtime/Engine/Classes/PhysicsEngine/PhysicsObjectBlueprintLibrary.h b/Engine/Source/Runtime/Engine/Classes/PhysicsEngine/PhysicsObjectBlueprintLibrary.h
new file mode 100644
--- /dev/null
+++ b/Engine/Source/Runtime/Engine/Classes/PhysicsEngine/PhysicsObjectBlueprintLibrary.h
@@ -0,0 +1,24 @@
+// Copyright Epic Games, Inc. All Rights Reserved.
+#pragma once
+
+#include "Kismet/BlueprintFunctionLibrary.h"
+
+#include "Chaos/PhysicsObjectInterface.h"
+#include "Math/MathFwd.h"
+
+#include "PhysicsObjectBlueprintLibrary.generated.h"
+
+UCLASS()
+class UPhysicsObjectBlueprintLibrary : public UBlueprintFunctionLibrary
+{
+	GENERATED_BODY()
+public:
+	UFUNCTION(BlueprintPure, Category = "Physics Object")
+	static FClosestPhysicsObjectResult GetClosestPhysicsObjectFromWorldLocation(UPrimitiveComponent* Component, const FVector& WorldLocation);
+
+	UFUNCTION(BlueprintPure, Category = "Physics Object")
+	static bool ExtractClosestPhysicsObjectResults(const FClosestPhysicsObjectResult& Result, FName& OutName);
+
+	UFUNCTION(BlueprintPure, Category = "Physics Object")
+	static FTransform GetPhysicsObjectWorldTransform(UPrimitiveComponent* Component, FName BoneName);
+};
\ No newline at end of file
diff --git a/Engine/Source/Runtime/Engine/Private/PhysicsEngine/PhysicsObjectBlueprintLibrary.cpp b/Engine/Source/Runtime/Engine/Private/PhysicsEngine/PhysicsObjectBlueprintLibrary.cpp
new file mode 100644
--- /dev/null
+++ b/Engine/Source/Runtime/Engine/Private/PhysicsEngine/PhysicsObjectBlueprintLibrary.cpp
@@ -0,0 +1,48 @@
+// Copyright Epic Games, Inc. All Rights Reserved.
+#include "PhysicsEngine/PhysicsObjectBlueprintLibrary.h"
+#include "PhysicsEngine/PhysicsObjectExternalInterface.h"
+
+#include "Components/PrimitiveComponent.h"
+
+FClosestPhysicsObjectResult UPhysicsObjectBlueprintLibrary::GetClosestPhysicsObjectFromWorldLocation(UPrimitiveComponent* Component, const FVector& WorldLocation)
+{
+	if (!Component)
+	{
+		return {};
+	}
+
+	TArray<Chaos::FPhysicsObject*> PhysicsObjects = Component->GetAllPhysicsObjects();
+	FLockedReadPhysicsObjectExternalInterface Interface = FPhysicsObjectExternalInterface::LockRead(PhysicsObjects);
+	PhysicsObjects = PhysicsObjects.FilterByPredicate(
+		// Maybe allow users to customize the predicate?
+		[&Interface](Chaos::FPhysicsObject* Object) {
+			// We need to filter so we only get the leaf bodies. Mainly in consideration of the geometry collection.
+			return !Interface->AreAllDisabled({&Object, 1}) && !Chaos::FPhysicsObjectInterface::HasChildren(Object);
+		}
+	);
+	return Interface->GetClosestPhysicsBodyFromLocation(PhysicsObjects, WorldLocation);
+}
+bool UPhysicsObjectBlueprintLibrary::ExtractClosestPhysicsObjectResults(const FClosestPhysicsObjectResult& Result, FName& OutName)
+{
+	if (!Result)
+	{
+		return false;
+	}
+	OutName = Result.HitName();
+	return true;
+}
+FTransform UPhysicsObjectBlueprintLibrary::GetPhysicsObjectWorldTransform(UPrimitiveComponent* Component, FName BoneName)
+{
+	if (!Component)
+	{
+		return FTransform::Identity;
+	}
+
+	if (Chaos::FPhysicsObject* PhysicsObject = Component->GetPhysicsObjectByName(BoneName))
+	{
+		FLockedReadPhysicsObjectExternalInterface Interface = FPhysicsObjectExternalInterface::LockRead({&PhysicsObject, 1});
+		return Interface->GetTransform(PhysicsObject);
+	}
+	
+	return FTransform::Identity;
+}
diff --git a/Engine/Source/Runtime/Engine/Private/PrimitiveComponentPhysics.cpp b/Engine/Source/Runtime/Engine/Private/PrimitiveComponentPhysics.cpp
--- a/Engine/Source/Runtime/Engine/Private/PrimitiveComponentPhysics.cpp
+++ b/Engine/Source/Runtime/Engine/Private/PrimitiveComponentPhysics.cpp
@@ -15,7 +15,9 @@
 #include "Physics/PhysicsInterfaceCore.h"
 #include "UObject/UObjectThreadContext.h"
 #include "Physics/PhysicsInterfaceScene.h"
+#include "PhysicsEngine/PhysicsObjectExternalInterface.h"
 #include "PhysicsProxy/SingleParticlePhysicsProxy.h"
+#include "Chaos/PhysicsObjectInterface.h"
 
 //////////////// PRIMITIVECOMPONENT ///////////////
 
@@ -329,6 +331,23 @@ void UPrimitiveComponent::WakeRigidBody(FName BoneName)
 	{
 		BI->WakeInstance();
 	}
+	else if (BoneName == NAME_None)
+	{
+		TArray<Chaos::FPhysicsObject*> PhysicsObjects = GetAllPhysicsObjects();
+		FLockedWritePhysicsObjectExternalInterface Interface = FPhysicsObjectExternalInterface::LockWrite(PhysicsObjects);
+		
+		PhysicsObjects = PhysicsObjects.FilterByPredicate(
+			[&Interface](Chaos::FPhysicsObject* Object) {
+				return !Interface->AreAllDisabled({ &Object, 1 });
+			}
+		);
+		Interface->WakeUp(PhysicsObjects);
+	}
+	else if (Chaos::FPhysicsObject* Object = GetPhysicsObjectByName(BoneName))
+	{
+		FLockedWritePhysicsObjectExternalInterface Interface = FPhysicsObjectExternalInterface::LockWrite({&Object, 1});
+		Interface->WakeUp({ &Object, 1 });
+	}
 }
 
 void UPrimitiveComponent::WakeAllRigidBodies()
@@ -502,6 +521,23 @@ void UPrimitiveComponent::PutRigidBodyToSleep(FName BoneName)
 	{
 		BI->PutInstanceToSleep();
 	}
+	else if (BoneName == NAME_None)
+	{
+		TArray<Chaos::FPhysicsObject*> PhysicsObjects = GetAllPhysicsObjects();
+		FLockedWritePhysicsObjectExternalInterface Interface = FPhysicsObjectExternalInterface::LockWrite(PhysicsObjects);
+
+		PhysicsObjects = PhysicsObjects.FilterByPredicate(
+			[&Interface](Chaos::FPhysicsObject* Object) {
+				return !Interface->AreAllDisabled({ &Object, 1 });
+			}
+		);
+		Interface->PutToSleep(PhysicsObjects);
+	}
+	else if (Chaos::FPhysicsObject* Object = GetPhysicsObjectByName(BoneName))
+	{
+		FLockedWritePhysicsObjectExternalInterface Interface = FPhysicsObjectExternalInterface::LockWrite({ &Object, 1 });
+		Interface->PutToSleep({ &Object, 1 });
+	}
 }
 
 void UPrimitiveComponent::PutAllRigidBodiesToSleep()
@@ -860,15 +896,24 @@ float UPrimitiveComponent::GetClosestPointOnCollision(const FVector& Point, FVec
 
 bool UPrimitiveComponent::IsSimulatingPhysics(FName BoneName) const
 {
-	FBodyInstance* BodyInst = GetBodyInstance(BoneName);
-	if(BodyInst != NULL)
+	if(FBodyInstance* BodyInst = GetBodyInstance(BoneName))
 	{
 		return BodyInst->IsInstanceSimulatingPhysics();
 	}
+	else if (BoneName != NAME_None)
+	{
+		if (Chaos::FPhysicsObject* PhysicsObject = GetPhysicsObjectByName(BoneName))
+		{
+			return true;
+		}
+	}
 	else
 	{
-		return false;
+		TArray<Chaos::FPhysicsObject*> PhysicsObjects = GetAllPhysicsObjects();
+		return !PhysicsObjects.IsEmpty();
 	}
+
+	return false;
 }
 
 FVector UPrimitiveComponent::GetComponentVelocity() const
@@ -1124,4 +1169,25 @@ void UPrimitiveComponent::UpdatePhysicsToRBChannels()
 	}
 }
 
+Chaos::FPhysicsObject* UPrimitiveComponent::GetPhysicsObjectById(int32 Id) const
+{
+	if (!BodyInstance.IsValidBodyInstance())
+	{
+		return nullptr;
+	}
+
+	return BodyInstance.ActorHandle->GetPhysicsObject();
+}
+
+Chaos::FPhysicsObject* UPrimitiveComponent::GetPhysicsObjectByName(const FName& Name) const
+{
+	return GetPhysicsObjectById(0);
+}
+
+TArray<Chaos::FPhysicsObject*> UPrimitiveComponent::GetAllPhysicsObjects() const
+{
+	TArray<Chaos::FPhysicsObject*> Bodies = { GetPhysicsObjectById(INDEX_NONE) };
+	return Bodies;
+}
+
 #undef LOCTEXT_NAMESPACE
diff --git a/Engine/Source/Runtime/Engine/Private/SkeletalMeshComponentPhysics.cpp b/Engine/Source/Runtime/Engine/Private/SkeletalMeshComponentPhysics.cpp
--- a/Engine/Source/Runtime/Engine/Private/SkeletalMeshComponentPhysics.cpp
+++ b/Engine/Source/Runtime/Engine/Private/SkeletalMeshComponentPhysics.cpp
@@ -30,6 +30,7 @@
 #include "Rendering/SkeletalMeshRenderData.h"
 #include "Physics/PhysicsInterfaceCore.h"
 #include "Physics/PhysicsInterfaceScene.h"
+#include "PhysicsProxy/SingleParticlePhysicsProxy.h"
 #include "AnimationRuntime.h"
 #include "ClothCollisionData.h"
 #include "ClothingSimulationInteractor.h"
@@ -59,6 +60,7 @@
 #include "Chaos/ImplicitObject.h"
 #include "Chaos/ImplicitObjectScaled.h"
 #include "Chaos/ImplicitObjectTransformed.h"
+#include "Chaos/PhysicsObjectInterface.h"
 
 #define LOCTEXT_NAMESPACE "SkeletalMeshComponentPhysics"
 
@@ -740,6 +742,23 @@ void USkeletalMeshComponent::InitArticulated(FPhysScene* PhysScene)
 	}
 
 	InstantiatePhysicsAsset(*PhysicsAsset, Scale3D, Bodies, Constraints, PhysScene, this, RootBodyIndex, Aggregate);
+	for (int32 BodyIndex = 0; BodyIndex < NumBodies; ++BodyIndex)
+	{
+		if (FBodyInstance* Body = Bodies[BodyIndex])
+		{
+			if (!Body->ActorHandle)
+			{
+				continue;
+			}
+
+			Chaos::FPhysicsObject* PhysicsObject = Body->ActorHandle->GetPhysicsObject();
+			if (UBodySetup* Setup = Body->GetBodySetup())
+			{
+				Chaos::FPhysicsObjectInterface::SetName(PhysicsObject, Setup->BoneName);
+			}
+			Chaos::FPhysicsObjectInterface::SetId(PhysicsObject, BodyIndex);
+		}
+	}
 
 	// now update root body index because body has BodySetup now
 	SetRootBodyIndex(RootBodyIndex);
@@ -4036,4 +4055,37 @@ FTransform USkeletalMeshComponent::GetComponentTransformFromBodyInstance(FBodyIn
 		return GetComponentTransform();
 	}
 }
+
+
+Chaos::FPhysicsObject* USkeletalMeshComponent::GetPhysicsObjectById(int32 Id) const
+{
+	if (!Bodies.IsValidIndex(Id) || !Bodies[Id] || !Bodies[Id]->ActorHandle)
+	{
+		return nullptr;
+	}
+	return Bodies[Id]->ActorHandle->GetPhysicsObject();
+}
+
+Chaos::FPhysicsObject* USkeletalMeshComponent::GetPhysicsObjectByName(const FName& Name) const
+{
+	FBodyInstance* Body = GetBodyInstance(Name);
+	if (!Body || !Body->ActorHandle)
+	{
+		return nullptr;
+	}
+
+	return Body->ActorHandle->GetPhysicsObject();
+}
+
+TArray<Chaos::FPhysicsObject*> USkeletalMeshComponent::GetAllPhysicsObjects() const
+{
+	TArray<Chaos::FPhysicsObject*> Objects;
+	Objects.Reserve(Bodies.Num());
+	for (int32 Index = 0; Index < Bodies.Num(); ++Index)
+	{
+		Objects.Add(GetPhysicsObjectById(Index));
+	}
+	return Objects;
+}
+
 #undef LOCTEXT_NAMESPACE
diff --git a/Engine/Source/Runtime/Experimental/GeometryCollectionEngine/Private/GeometryCollection/GeometryCollectionComponent.cpp b/Engine/Source/Runtime/Experimental/GeometryCollectionEngine/Private/GeometryCollection/GeometryCollectionComponent.cpp
--- a/Engine/Source/Runtime/Experimental/GeometryCollectionEngine/Private/GeometryCollection/GeometryCollectionComponent.cpp
+++ b/Engine/Source/Runtime/Experimental/GeometryCollectionEngine/Private/GeometryCollection/GeometryCollectionComponent.cpp
@@ -785,6 +785,11 @@ void UGeometryCollectionComponent::SetSimulatePhysics(bool bEnabled)
 	{
 		RegisterAndInitializePhysicsProxy();
 	}
+
+	if (bEnabled)
+	{
+		RemoveAllAnchors();
+	}
 }
 
 void UGeometryCollectionComponent::AddForce(FVector Force, FName BoneName, bool bAccelChange)
@@ -4630,3 +4635,37 @@ void UGeometryCollectionComponent::PostLoad()
 	}
 }
 
+Chaos::FPhysicsObject* UGeometryCollectionComponent::GetPhysicsObjectById(int32 Id) const
+{
+	if (!PhysicsProxy)
+	{
+		return nullptr;
+	}
+	return PhysicsProxy->GetPhysicsObjectByIndex(Id);
+}
+
+Chaos::FPhysicsObject* UGeometryCollectionComponent::GetPhysicsObjectByName(const FName& Name) const
+{
+	if (!RestCollection)
+	{
+		return nullptr;
+	}
+	const int32 Index = RestCollection->GetGeometryCollection()->BoneName.Find(Name.ToString());
+	return GetPhysicsObjectById(Index);
+}
+
+TArray<Chaos::FPhysicsObject*> UGeometryCollectionComponent::GetAllPhysicsObjects() const
+{
+	if (!PhysicsProxy)
+	{
+		return {};
+	}
+	TArray<Chaos::FPhysicsObject*> Objects;
+	Objects.Reserve(PhysicsProxy->GetNumParticles());
+	
+	for (int32 Index = 0; Index < PhysicsProxy->GetNumParticles(); ++Index)
+	{
+		Objects.Add(GetPhysicsObjectById(Index));
+	}
+	return Objects;
+}
\ No newline at end of file
diff --git a/Engine/Source/Runtime/Experimental/GeometryCollectionEngine/Public/GeometryCollection/GeometryCollectionComponent.h b/Engine/Source/Runtime/Experimental/GeometryCollectionEngine/Public/GeometryCollection/GeometryCollectionComponent.h
--- a/Engine/Source/Runtime/Experimental/GeometryCollectionEngine/Public/GeometryCollection/GeometryCollectionComponent.h
+++ b/Engine/Source/Runtime/Experimental/GeometryCollectionEngine/Public/GeometryCollection/GeometryCollectionComponent.h
@@ -1197,4 +1197,11 @@ private:
 
 	/** True if GeometryCollection transforms have changed from previous tick. */
 	bool bIsMoving;
+
+	//~ Begin IPhysicsComponent Interface.
+public:
+	virtual Chaos::FPhysicsObject* GetPhysicsObjectById(int32 Id) const override;
+	virtual Chaos::FPhysicsObject* GetPhysicsObjectByName(const FName& Name) const override;
+	virtual TArray<Chaos::FPhysicsObject*> GetAllPhysicsObjects() const override;
+	//~ End IPhysicsComponent Interface.
 };
