World Partition: avoid generating time-stamped streaming generation logs on buildmachines.

diff --git a/Engine/Source/Runtime/Engine/Private/WorldPartition/WorldPartitionStreamingGeneration.cpp b/Engine/Source/Runtime/Engine/Private/WorldPartition/WorldPartitionStreamingGeneration.cpp
--- a/Engine/Source/Runtime/Engine/Private/WorldPartition/WorldPartitionStreamingGeneration.cpp
+++ b/Engine/Source/Runtime/Engine/Private/WorldPartition/WorldPartitionStreamingGeneration.cpp
@@ -805,7 +805,15 @@ public:
 
 	static TUniquePtr<FArchive> CreateDumpStateLogArchive(const TCHAR* Suffix)
 	{
-		const FString StateLogOutputFilename = FPaths::ProjectSavedDir() / TEXT("WorldPartition") / FString::Printf(TEXT("StreamingGeneration-%s-%08x-%s.log"), Suffix, FPlatformProcess::GetCurrentProcessId(), *FDateTime::Now().ToString());
+		FString StateLogOutputFilename = FPaths::ProjectSavedDir() / TEXT("WorldPartition") / FString::Printf(TEXT("StreamingGeneration-%s"), Suffix);
+		
+		if (!GIsBuildMachine)
+		{
+			StateLogOutputFilename += FString::Printf(TEXT("-%08x-%s"), FPlatformProcess::GetCurrentProcessId(), *FDateTime::Now().ToString());
+		}
+
+		StateLogOutputFilename += TEXT(".log");
+
 		return TUniquePtr<FArchive>(IFileManager::Get().CreateFileWriter(*StateLogOutputFilename));
 	}
 
