Moving enhanced input binding for actor components into component registration instead of the general add component functionality.
* Fixes bug with components that are added via manual creation not having their inputs bound.

diff --git a/Engine/Source/Runtime/Engine/Private/ActorConstruction.cpp b/Engine/Source/Runtime/Engine/Private/ActorConstruction.cpp
--- a/Engine/Source/Runtime/Engine/Private/ActorConstruction.cpp
+++ b/Engine/Source/Runtime/Engine/Private/ActorConstruction.cpp
@@ -33,8 +33,6 @@
 #include "ProfilingDebugging/CsvProfiler.h"
 #include "Algo/Transform.h"
 #include "Misc/ScopeExit.h"
-#include "Engine/InputDelegateBinding.h"
-#include "GameFramework/InputSettings.h"
 
 #if WITH_EDITOR
 #include "Editor.h"
@@ -1264,11 +1262,6 @@ void AActor::FinishAddComponent(UActorComponent* NewActorComp, bool bManualAttac
 				World->UpdateCullDistanceVolumes(this, NewPrimitiveComponent);
 			}
 		}
-		
-		if (InputComponent && GetDefault<UInputSettings>()->bEnableDynamicComponentInputBinding)
-		{
-			UInputDelegateBinding::BindInputDelegates(NewActorComp->GetClass(), InputComponent, NewActorComp);
-		}
 	}
 }
 
diff --git a/Engine/Source/Runtime/Engine/Private/Components/ActorComponent.cpp b/Engine/Source/Runtime/Engine/Private/Components/ActorComponent.cpp
--- a/Engine/Source/Runtime/Engine/Private/Components/ActorComponent.cpp
+++ b/Engine/Source/Runtime/Engine/Private/Components/ActorComponent.cpp
@@ -38,6 +38,8 @@
 #include "UObject/FrameworkObjectVersion.h"
 #include "UObject/FortniteReleaseBranchCustomObjectVersion.h"
 #include "Async/ParallelFor.h"
+#include "Engine/InputDelegateBinding.h"
+#include "GameFramework/InputSettings.h"
 
 #if WITH_EDITOR
 #include "Kismet2/ComponentEditorUtils.h"
@@ -1389,6 +1391,11 @@ void UActorComponent::RegisterComponentWithWorld(UWorld* InWorld, FRegisterCompo
 		}
 
 	}
+
+	if (MyOwner && MyOwner->InputComponent && GetDefault<UInputSettings>()->bEnableDynamicComponentInputBinding)
+	{
+		UInputDelegateBinding::BindInputDelegates(GetClass(), MyOwner->InputComponent, this);
+	}
 }
 
 void UActorComponent::RegisterComponent()
