#ue5 - Fix EndComponentOverlap events not firing when bGenerateOverlapEventsDuringLevelStreaming is enabled and the component begins and then ends an overlap during streaming. Mirrors a similar fix from 12366833 for BeginComponentOverlap.

#udn https://udn.unrealengine.com/s/question/0D54z00008NKXkPCAX/aactorbgenerateoverlapeventsduringlevelstreaming-doesnt-notify-endoverlap-event-during-beginplay

diff --git a/Engine/Source/Runtime/Engine/Private/Components/PrimitiveComponent.cpp b/Engine/Source/Runtime/Engine/Private/Components/PrimitiveComponent.cpp
--- a/Engine/Source/Runtime/Engine/Private/Components/PrimitiveComponent.cpp
+++ b/Engine/Source/Runtime/Engine/Private/Components/PrimitiveComponent.cpp
@@ -3106,11 +3106,12 @@ void UPrimitiveComponent::EndComponentOverlap(const FOverlapInfo& OtherOverlap,
 		GlobalOverlapEventsCounter++;
 		OverlappingComponents.RemoveAtSwap(OverlapIdx, 1, false);
 
+		AActor* const MyActor = GetOwner();
 		const UWorld* World = GetWorld();
-		if (bDoNotifies && World && World->HasBegunPlay())
+		const bool bLevelStreamingOverlap = (bDoNotifies && MyActor && MyActor->bGenerateOverlapEventsDuringLevelStreaming && MyActor->IsActorBeginningPlayFromLevelStreaming());
+		if (bDoNotifies && ((World && World->HasBegunPlay()) || bLevelStreamingOverlap))
 		{
 			AActor* const OtherActor = OtherComp->GetOwner();
-			AActor* const MyActor = GetOwner();
 			if (OtherActor)
 			{
 				if (!bSkipNotifySelf && IsPrimCompValidAndAlive(this))
