Texture provide Blueprint functions to help detect nearly-empty textures

diff --git a/Engine/Source/Runtime/Engine/Classes/Engine/Texture.h b/Engine/Source/Runtime/Engine/Classes/Engine/Texture.h
--- a/Engine/Source/Runtime/Engine/Classes/Engine/Texture.h
+++ b/Engine/Source/Runtime/Engine/Classes/Engine/Texture.h
@@ -1344,10 +1344,30 @@ public:
 	uint8 bNotOfflineProcessed : 1;
 
 	/**
-	 * Gets the memory size of the texture, in bytes
+	 * Gets the memory size of the texture, in bytes.
+	 * This is the size in GPU memory of the built platformdata, accounting for LODBias, etc.
+	 * Returns zero for error.
 	 */
 	UFUNCTION(BlueprintCallable, meta = (DisplayName = "GetMemorySize"), Category = "Rendering|Texture")
 	int64 Blueprint_GetMemorySize() const;
+	
+	/**
+	 * Gets the memory size of the texture source top mip, in bytes, and the size on disk of the asset, which may be compressed.
+	 * Uses texture source, not available in runtime games.
+	 * Does not cause texture source to be loaded, queries cached values.
+	 * Returns zero for error.
+	 */
+	UFUNCTION(BlueprintCallable, meta = (DisplayName = "GetTextureSourceDiskAndMemorySize"), Category = "Rendering|Texture")
+	void Blueprint_GetTextureSourceDiskAndMemorySize(int64 & OutDiskSize,int64 & OutMemorySize) const;
+
+	/**
+	 * Scan the texture source pixels to compute the min & max values of the RGBA channels.
+	 * Uses texture source, not available in runtime games.
+	 * Causes texture source data to be loaded, is computed by scanning pixels when called.
+	 * Will set Min=Max=zero and return false on failure
+	 */
+	UFUNCTION(BlueprintCallable, Category = "Rendering|Texture")
+	bool ComputeTextureSourceChannelMinMax(FLinearColor & OutColorMin, FLinearColor & OutColorMax) const;
 
 private:
 	/** Whether the async resource release process has already been kicked off or not */
@@ -1732,9 +1752,10 @@ public:
 
 	/**
 	 * Calculates the size of this texture if it had MipCount miplevels streamed in.
+	 * This is the size in GPU memory of the built platformdata, accounting for LODBias, etc.
 	 *
 	 * @param	Enum	Which mips to calculate size for.
-	 * @return	Total size of all specified mips, in bytes
+	 * @return	Total size of all specified mips, in bytes.  Returns 0 for error.
 	 */
 	virtual uint32 CalcTextureMemorySizeEnum( ETextureMipCount Enum ) const
 	{
diff --git a/Engine/Source/Runtime/Engine/Private/Texture.cpp b/Engine/Source/Runtime/Engine/Private/Texture.cpp
--- a/Engine/Source/Runtime/Engine/Private/Texture.cpp
+++ b/Engine/Source/Runtime/Engine/Private/Texture.cpp
@@ -3600,6 +3600,39 @@ int64 UTexture::Blueprint_GetMemorySize() const
 	return CalcTextureMemorySizeEnum(TMC_ResidentMips);
 }
 
+void UTexture::Blueprint_GetTextureSourceDiskAndMemorySize(int64 & OutDiskSize,int64 & OutMemorySize) const
+{
+#if WITH_EDITORONLY_DATA
+	OutMemorySize = Source.CalcMipSize(0);
+	OutDiskSize = Source.GetSizeOnDisk();
+#else
+	OutDiskSize = OutMemorySize = 0;
+	UE_LOG(LogTexture, Error, TEXT("Blueprint_GetTextureSourceDiskAndMemorySize can only be called WITH_EDITORONLY_DATA. (%s)"), *GetName());
+#endif
+}
+
+bool UTexture::ComputeTextureSourceChannelMinMax(FLinearColor & OutColorMin, FLinearColor & OutColorMax) const
+{
+	// make sure we fill the outputs if we return failure :
+	OutColorMin = FLinearColor(ForceInit);
+	OutColorMax = FLinearColor(ForceInit);
+
+#if WITH_EDITORONLY_DATA
+	FImage Image;
+	if ( ! const_cast<UTexture *>(this)->Source.GetMipImage(Image,0) )
+	{
+		UE_LOG(LogTexture, Error, TEXT("ComputeTextureSourceChannelMinMax failed to GetMipImage. (%s)"), *GetName());
+		return false;
+	}
+
+	FImageCore::ComputeChannelLinearMinMax(Image,OutColorMin,OutColorMax);
+	return true;
+#else
+	UE_LOG(LogTexture, Error, TEXT("ComputeTextureSourceChannelMinMax can only be called WITH_EDITORONLY_DATA. (%s)"), *GetName());
+	return false;
+#endif
+}
+
 #if WITH_EDITOR
 
 FTextureSource::FMipData::FMipData(const FTextureSource& InSource, FSharedBuffer InData)
diff --git a/Engine/Source/Runtime/ImageCore/Private/ImageCore.cpp b/Engine/Source/Runtime/ImageCore/Private/ImageCore.cpp
--- a/Engine/Source/Runtime/ImageCore/Private/ImageCore.cpp
+++ b/Engine/Source/Runtime/ImageCore/Private/ImageCore.cpp
@@ -887,6 +887,7 @@ void FImage::CopyTo(FImage& DestImage, ERawImageFormat::Type DestFormat, EGammaS
 		DestGammaSpace = EGammaSpace::sRGB;
 	}
 
+	// existing contents of DestImage are freed and replaced
 	DestImage.Init(SizeX,SizeY,NumSlices,DestFormat,DestGammaSpace);
 	FImageCore::CopyImage(*this, DestImage);
 }
@@ -1506,3 +1507,36 @@ void FImageCore::SetAlphaOpaque(const FImageView & InImage)
 		check(0);
 	}
 }
+
+
+void FImageCore::ComputeChannelLinearMinMax(const FImage & InImage, FLinearColor & OutMin, FLinearColor & OutMax)
+{
+	// @todo Oodle : for speed, we should ideally scan the image for min/max in its native pixel format
+	//	then only convert the min/max colors to float linear after the scan
+	//	don't convert the whole image
+
+	FImage ImageLinear;
+	InImage.CopyTo(ImageLinear,ERawImageFormat::RGBA32F,EGammaSpace::Linear);
+	
+	TArrayView64<FLinearColor> Colors = ImageLinear.AsRGBA32F();
+	if ( Colors.Num() == 0 )
+	{
+		OutMin = FLinearColor(ForceInit);
+		OutMax = FLinearColor(ForceInit);
+		return;
+	}
+	
+	VectorRegister4Float VMin = VectorLoad(&Colors[0].Component(0));
+	VectorRegister4Float VMax = VMin;
+	
+	for ( const FLinearColor & Color : Colors )
+	{
+		VectorRegister4Float VCur = VectorLoad(&Color.Component(0));
+
+		VMin = VectorMin(VMin,VCur);
+		VMax = VectorMax(VMax,VCur);
+	}
+
+	VectorStore(VMin,&OutMin.Component(0));
+	VectorStore(VMax,&OutMax.Component(0));
+}
diff --git a/Engine/Source/Runtime/ImageCore/Public/ImageCore.h b/Engine/Source/Runtime/ImageCore/Public/ImageCore.h
--- a/Engine/Source/Runtime/ImageCore/Public/ImageCore.h
+++ b/Engine/Source/Runtime/ImageCore/Public/ImageCore.h
@@ -415,7 +415,7 @@ public:
 	/**
 	 * Copies the image to a destination image with the specified format.
 	 *
-	 * @param DestImage - The destination image.
+	 * @param DestImage - The destination image.  Will be allocated.  Any existing contents are replaced.
 	 * @param DestFormat - The destination image format.
 	 * @param DestSRGB - Whether the destination image is in SRGB format.
 	 */
@@ -704,6 +704,16 @@ IMAGECORE_API void SetAlphaOpaque(const FImageView & InImage);
 	*/
 IMAGECORE_API void ResizeTo(const FImageView & SourceImage,FImage& DestImage, int32 DestSizeX, int32 DestSizeY, ERawImageFormat::Type DestFormat, EGammaSpace DestGammaSpace);
 
+/**
+ * Compute the min/max of each channel to get value ranges
+ * Colors are converted to float Linear Gamma
+ * 
+ * @param InImage - The image to scan
+ * @param OutMin - filled with the minimum of the color channels
+ * @param OutMax - filled with the maximum of the color channels
+ */
+IMAGECORE_API void ComputeChannelLinearMinMax(const FImage & InImage, FLinearColor & OutMin, FLinearColor & OutMax);
+
 };
 
 
