Change the joint constraints to alternatively use the FPhysicsObject and switch over the existing FBodyInstance paths

diff --git a/Engine/Source/Runtime/Engine/Classes/PhysicsEngine/ConstraintInstance.h b/Engine/Source/Runtime/Engine/Classes/PhysicsEngine/ConstraintInstance.h
--- a/Engine/Source/Runtime/Engine/Classes/PhysicsEngine/ConstraintInstance.h
+++ b/Engine/Source/Runtime/Engine/Classes/PhysicsEngine/ConstraintInstance.h
@@ -5,6 +5,7 @@
 #include "Misc/EnumClassFlags.h"
 #include "UObject/ObjectMacros.h"
 #include "UObject/Class.h"
+#include "Chaos/PhysicsObject.h"
 #include "EngineDefines.h"
 #include "PhysicsEngine/ConstraintTypes.h"
 #include "PhysicsEngine/ConstraintDrives.h"
@@ -875,9 +876,11 @@ public:
 
 	/** Create physics engine constraint. */
 	void InitConstraint(FBodyInstance* Body1, FBodyInstance* Body2, float Scale, UObject* DebugOwner, FOnConstraintBroken InConstraintBrokenDelegate = FOnConstraintBroken(), FOnPlasticDeformation InPlasticDeformationDelegate = FOnPlasticDeformation());
+	void InitConstraint(Chaos::FPhysicsObject* Body1, Chaos::FPhysicsObject* Body2, float Scale, UObject* DebugOwner, FOnConstraintBroken InConstraintBrokenDelegate = FOnConstraintBroken(), FOnPlasticDeformation InPlasticDeformationDelegate = FOnPlasticDeformation());
 
 	/** Create physics engine constraint using physx actors. */
 	void InitConstraint_AssumesLocked(const FPhysicsActorHandle& ActorRef1, const FPhysicsActorHandle& ActorRef2, float InScale, FOnConstraintBroken InConstraintBrokenDelegate = FOnConstraintBroken(), FOnPlasticDeformation InPlasticDeformationDelegate = FOnPlasticDeformation());
+	void InitConstraint_AssumesLocked(Chaos::FPhysicsObject* Body1, Chaos::FPhysicsObject* Body2, float InScale, FOnConstraintBroken InConstraintBrokenDelegate = FOnConstraintBroken(), FOnPlasticDeformation InPlasticDeformationDelegate = FOnPlasticDeformation());
 
 	/** Terminate physics engine constraint */
 	void TermConstraint();
@@ -992,8 +995,8 @@ public:
 
 private:
 
-	bool CreateJoint_AssumesLocked(const FPhysicsActorHandle& InActorRef1, const FPhysicsActorHandle& InActorRef2);
-	void UpdateAverageMass_AssumesLocked(const FPhysicsActorHandle& InActorRef1, const FPhysicsActorHandle& InActorRef2);
+	bool CreateJoint_AssumesLocked(Chaos::FPhysicsObject* Body1, Chaos::FPhysicsObject* Body2);
+	void UpdateAverageMass_AssumesLocked(Chaos::FPhysicsObject* Body1, Chaos::FPhysicsObject* Body2);
 
 	struct FPDIOrCollector
 	{
diff --git a/Engine/Source/Runtime/Engine/Classes/PhysicsEngine/PhysicsConstraintComponent.h b/Engine/Source/Runtime/Engine/Classes/PhysicsEngine/PhysicsConstraintComponent.h
--- a/Engine/Source/Runtime/Engine/Classes/PhysicsEngine/PhysicsConstraintComponent.h
+++ b/Engine/Source/Runtime/Engine/Classes/PhysicsEngine/PhysicsConstraintComponent.h
@@ -362,6 +362,9 @@ protected:
 	/** Get the body instance that we want to constrain to */
 	FBodyInstance* GetBodyInstance(EConstraintFrame::Type Frame) const;
 
+	/** Get the physics object that we want to constrain to. */
+	Chaos::FPhysicsObject* GetPhysicsObject(EConstraintFrame::Type Frame) const;
+
 	/** Internal util to get body transform from actor/component name/bone name information */
 	FTransform GetBodyTransformInternal(EConstraintFrame::Type Frame, FName InBoneName) const;
 	/** Internal util to get body box from actor/component name/bone name information */
diff --git a/Engine/Source/Runtime/Engine/Private/PhysicsEngine/ConstraintInstance.cpp b/Engine/Source/Runtime/Engine/Private/PhysicsEngine/ConstraintInstance.cpp
--- a/Engine/Source/Runtime/Engine/Private/PhysicsEngine/ConstraintInstance.cpp
+++ b/Engine/Source/Runtime/Engine/Private/PhysicsEngine/ConstraintInstance.cpp
@@ -5,12 +5,15 @@
 #include "UObject/FrameworkObjectVersion.h"
 #include "UObject/AnimPhysObjectVersion.h"
 #include "HAL/IConsoleManager.h"
+#include "Chaos/PhysicsObjectInterface.h"
 #include "Components/PrimitiveComponent.h"
 #include "Components/SkeletalMeshComponent.h"
 #include "PhysicsEngine/PhysicsConstraintComponent.h"
+#include "PhysicsEngine/PhysicsObjectExternalInterface.h"
 #include "PhysicsEngine/PhysicsAsset.h"
 #include "PhysicsPublic.h"
 #include "Physics/PhysicsInterfaceTypes.h"
+#include "PhysicsProxy/SingleParticlePhysicsProxy.h"
 
 #include "Logging/TokenizedMessage.h"
 #include "Logging/MessageLog.h"
@@ -429,22 +432,24 @@ void FConstraintInstance::SetDisableCollision(bool InDisableCollision)
 	});
 }
 
-float ComputeAverageMass_AssumesLocked(const FPhysicsActorHandle& InActor1, const FPhysicsActorHandle& InActor2)
+float ComputeAverageMass_AssumesLocked(Chaos::FPhysicsObject* Body1, Chaos::FPhysicsObject* Body2)
 {
 	float AverageMass = 0;
 
 	float TotalMass = 0;
 	int NumDynamic = 0;
 
-	if (FPhysicsInterface::IsValid(InActor1) && FPhysicsInterface::IsRigidBody(InActor1))
+	Chaos::FReadPhysicsObjectInterface_External Interface = FPhysicsObjectExternalInterface::GetRead_AssumesLocked();
+
+	if (Interface.AreAllRigidBody({ &Body1, 1 }))
 	{
-		TotalMass += FPhysicsInterface::GetMass_AssumesLocked(InActor1);
+		TotalMass += Interface.GetMass({ &Body1, 1 });
 		++NumDynamic;
 	}
 
-	if(FPhysicsInterface::IsValid(InActor2) && FPhysicsInterface::IsRigidBody(InActor2))
+	if (Interface.AreAllRigidBody({ &Body2, 1 }))
 	{
-		TotalMass += FPhysicsInterface::GetMass_AssumesLocked(InActor2);
+		TotalMass += Interface.GetMass({ &Body2, 1 });
 		++NumDynamic;
 	}
 
@@ -512,12 +517,14 @@ bool GetActorRefs(FBodyInstance* Body1, FBodyInstance* Body2, FPhysicsActorHandl
 	return true;
 }
 
-bool FConstraintInstance::CreateJoint_AssumesLocked(const FPhysicsActorHandle& InActorRef1, const FPhysicsActorHandle& InActorRef2)
+bool FConstraintInstance::CreateJoint_AssumesLocked(Chaos::FPhysicsObject* Body1, Chaos::FPhysicsObject* Body2)
 {
 	LLM_SCOPE(ELLMTag::ChaosConstraint);
 
+	Chaos::FReadPhysicsObjectInterface_External Interface = FPhysicsObjectExternalInterface::GetRead_AssumesLocked();
+
 	FTransform Local1 = GetRefFrame(EConstraintFrame::Frame1);
-	if(FPhysicsInterface::IsValid(InActorRef1))
+	if (Interface.AreAllValid({ &Body1, 1 }))
 	{
 		Local1.ScaleTranslation(FVector(LastKnownScale));
 	}
@@ -525,7 +532,7 @@ bool FConstraintInstance::CreateJoint_AssumesLocked(const FPhysicsActorHandle& I
 	checkf(Local1.IsValid() && !Local1.ContainsNaN(), TEXT("%s"), *Local1.ToString());
 
 	FTransform Local2 = GetRefFrame(EConstraintFrame::Frame2);
-	if(FPhysicsInterface::IsValid(InActorRef2))
+	if (Interface.AreAllValid({ &Body2, 1 }))
 	{
 		Local2.ScaleTranslation(FVector(LastKnownScale));
 	}
@@ -534,7 +541,7 @@ bool FConstraintInstance::CreateJoint_AssumesLocked(const FPhysicsActorHandle& I
 
 	if (bEnableSkeletalMeshConstraints)
 	{
-		ConstraintHandle = FPhysicsInterface::CreateConstraint(InActorRef1, InActorRef2, Local1, Local2);
+		ConstraintHandle = FPhysicsInterface::CreateConstraint(Body1, Body2, Local1, Local2);
 	}
 	if(!ConstraintHandle.IsValid())
 	{
@@ -561,9 +568,9 @@ void FConstraintProfileProperties::UpdateConstraintFlags_AssumesLocked(const FPh
 }
 
 
-void FConstraintInstance::UpdateAverageMass_AssumesLocked(const FPhysicsActorHandle& InActorRef1, const FPhysicsActorHandle& InActorRef2)
+void FConstraintInstance::UpdateAverageMass_AssumesLocked(Chaos::FPhysicsObject* Body1, Chaos::FPhysicsObject* Body2)
 {
-	AverageMass = ComputeAverageMass_AssumesLocked(InActorRef1, InActorRef2);
+	AverageMass = ComputeAverageMass_AssumesLocked(Body1, Body2);
 }
 
 /** 
@@ -574,27 +581,43 @@ void FConstraintInstance::InitConstraint(FBodyInstance* Body1, FBodyInstance* Bo
 	FPhysicsActorHandle Actor1;
 	FPhysicsActorHandle Actor2;
 
+	const bool bValidActors = GetActorRefs(Body1, Body2, Actor1, Actor2, DebugOwner);
+	if (!bValidActors)
 	{
-		const bool bValidActors = GetActorRefs(Body1, Body2, Actor1, Actor2, DebugOwner);
-		if (!bValidActors)
-		{
-			return;
-		}
+		return;
+	}
+
+	InitConstraint(Actor1->GetPhysicsObject(), Actor2->GetPhysicsObject(), InScale, DebugOwner, InConstraintBrokenDelegate, InPlasticDeformationDelegate);
+}
+
+void FConstraintInstance::InitConstraint(Chaos::FPhysicsObject* Body1, Chaos::FPhysicsObject* Body2, float InScale, UObject* DebugOwner, FOnConstraintBroken InConstraintBrokenDelegate, FOnPlasticDeformation InPlasticDeformationDelegate)
+{
+	{
+		Chaos::FPhysicsObject* Bodies[2] = { Body1, Body2 };
+		FLockedReadPhysicsObjectExternalInterface Interface = FPhysicsObjectExternalInterface::LockRead({ Bodies, 2 });
 
-		if (!bAllowKinematicKinematicConstraints && (!FPhysicsInterface::IsValid(Actor1) || FPhysicsInterface::IsKinematic(Actor1)) && (!FPhysicsInterface::IsValid(Actor2) || FPhysicsInterface::IsKinematic(Actor2)))
+		const bool bBody1Valid = Interface->AreAllValid({ &Body1, 1 });
+		const bool bBody2Valid = Interface->AreAllValid({ &Body2, 1 });
+		if (!bAllowKinematicKinematicConstraints && (!bBody1Valid || Interface->AreAllKinematic({ &Body1, 1 })) && (!bBody2Valid || Interface->AreAllKinematic({ &Body2, 1 })))
 		{
 			return;
 		}
-
-		FPhysicsCommand::ExecuteWrite(Actor1, Actor2, [&](const FPhysicsActorHandle& ActorA, const FPhysicsActorHandle& ActorB)
-		{
-			InitConstraint_AssumesLocked(ActorA, ActorB, InScale, InConstraintBrokenDelegate, InPlasticDeformationDelegate);
-		});
 	}
 
+	FPhysicsCommand::ExecuteWrite(Body1, Body2, [&](Chaos::FPhysicsObject* ActorA, Chaos::FPhysicsObject* ActorB)
+	{
+		InitConstraint_AssumesLocked(ActorA, ActorB, InScale, InConstraintBrokenDelegate, InPlasticDeformationDelegate);
+	});
 }
 
 void FConstraintInstance::InitConstraint_AssumesLocked(const FPhysicsActorHandle& ActorRef1, const FPhysicsActorHandle& ActorRef2, float InScale, FOnConstraintBroken InConstraintBrokenDelegate, FOnPlasticDeformation InPlasticDeformationDelegate)
+{
+	Chaos::FPhysicsObject* Body1 = ActorRef1 ? ActorRef1->GetPhysicsObject() : nullptr;
+	Chaos::FPhysicsObject* Body2 = ActorRef2 ? ActorRef2->GetPhysicsObject() : nullptr;
+	InitConstraint_AssumesLocked(Body1, Body2, InScale, InConstraintBrokenDelegate, InPlasticDeformationDelegate);
+}
+
+void FConstraintInstance::InitConstraint_AssumesLocked(Chaos::FPhysicsObject* Body1, Chaos::FPhysicsObject* Body2, float InScale, FOnConstraintBroken InConstraintBrokenDelegate, FOnPlasticDeformation InPlasticDeformationDelegate)
 {
 	OnConstraintBrokenDelegate = InConstraintBrokenDelegate;
 	OnPlasticDeformationDelegate = InPlasticDeformationDelegate;
@@ -602,9 +625,10 @@ void FConstraintInstance::InitConstraint_AssumesLocked(const FPhysicsActorHandle
 
 	UserData = FChaosUserData(this);
 
+	Chaos::FWritePhysicsObjectInterface_External Interface = FPhysicsObjectExternalInterface::GetWrite_AssumesLocked();
 	// Creating/Destroying a joint between two bodies will wake them, so we may want to re-sleep them
-	const bool bActor1WasAsleep = FPhysicsInterface::IsValid(ActorRef1) && FPhysicsInterface::IsSleeping(ActorRef1);
-	const bool bActor2WasAsleep = FPhysicsInterface::IsValid(ActorRef2) && FPhysicsInterface::IsSleeping(ActorRef2);
+	const bool bActor1WasAsleep = Interface.AreAllValid({ &Body1, 1 }) && Interface.AreAllSleeping({ &Body1, 1 });
+	const bool bActor2WasAsleep = Interface.AreAllValid({ &Body2, 1 }) && Interface.AreAllSleeping({ &Body2, 1 });
 
 	// if there's already a constraint, get rid of it first
 	if (ConstraintHandle.IsValid())
@@ -612,27 +636,27 @@ void FConstraintInstance::InitConstraint_AssumesLocked(const FPhysicsActorHandle
 		TermConstraint();
 	}
 
-	if (!CreateJoint_AssumesLocked(ActorRef1, ActorRef2))
+	if (!CreateJoint_AssumesLocked(Body1, Body2))
 	{
 		return;
 	}
-	
+
 	// update mass
-	UpdateAverageMass_AssumesLocked(ActorRef1, ActorRef2);
+	UpdateAverageMass_AssumesLocked(Body1, Body2);
 
 	ProfileInstance.Update_AssumesLocked(ConstraintHandle, AverageMass, bScaleLinearLimits ? LastKnownScale : 1.f, true);
 
 	// Put the bodies back to sleep both bodies were asleep
 	if (bActor1WasAsleep && bActor2WasAsleep)
 	{
-		if(!FPhysicsInterface::IsKinematic_AssumesLocked(ActorRef1))
+		if (!Interface.AreAllKinematic({ &Body1, 1 }))
 		{
-			FPhysicsInterface::PutToSleep_AssumesLocked(ActorRef1);
+			Interface.PutToSleep({ &Body1, 1 });
 		}
 
-		if(!FPhysicsInterface::IsKinematic_AssumesLocked(ActorRef2))
+		if (!Interface.AreAllKinematic({ &Body2, 1 }))
 		{
-			FPhysicsInterface::PutToSleep_AssumesLocked(ActorRef2);
+			Interface.PutToSleep({ &Body2, 1 });
 		}
 	}
 }
diff --git a/Engine/Source/Runtime/Engine/Private/PhysicsEngine/Experimental/PhysInterface_Chaos.cpp b/Engine/Source/Runtime/Engine/Private/PhysicsEngine/Experimental/PhysInterface_Chaos.cpp
--- a/Engine/Source/Runtime/Engine/Private/PhysicsEngine/Experimental/PhysInterface_Chaos.cpp
+++ b/Engine/Source/Runtime/Engine/Private/PhysicsEngine/Experimental/PhysInterface_Chaos.cpp
@@ -3,6 +3,7 @@
 #include "Physics/Experimental/PhysInterface_Chaos.h"
 #include "Physics/Experimental/PhysScene_Chaos.h"
 #include "Physics/Experimental/ChaosInterfaceUtils.h"
+#include "Physics/Experimental/ChaosScopedSceneLock.h"
 #include "Physics/PhysicsInterfaceTypes.h"
 #include "PhysicsEngine/AggregateGeom.h"
 #include "PhysicsEngine/ConstraintInstance.h"
@@ -32,6 +33,7 @@
 #include "Chaos/PBDJointConstraints.h"
 #include "Chaos/PBDJointConstraintData.h"
 #include "Chaos/ChaosConstraintSettings.h"
+#include "Chaos/PhysicsObjectInterface.h"
 
 #include "Async/ParallelFor.h"
 #include "Collision/CollisionConversions.h"
@@ -566,6 +568,13 @@ bool FPhysInterface_Chaos::ExecuteRead(FPhysScene* InScene, TFunctionRef<void()>
 	return false;
 }
 
+bool FPhysInterface_Chaos::ExecuteRead(Chaos::FPhysicsObject* InObjectA, Chaos::FPhysicsObject* InObjectB, TFunctionRef<void(Chaos::FPhysicsObject* InObjectA, Chaos::FPhysicsObject* InObjectB)> InCallable)
+{
+	FScopedSceneLock_Chaos SceneLock(InObjectA, InObjectB, EPhysicsInterfaceScopedLockType::Read);
+	InCallable(InObjectA, InObjectB);
+	return true;
+}
+
 bool FPhysInterface_Chaos::ExecuteWrite(const FPhysicsActorHandle& InActorReference, TFunctionRef<void(const FPhysicsActorHandle& Actor)> InCallable)
 {
 	//why do we have a write that takes in a const handle?
@@ -640,6 +649,13 @@ bool FPhysInterface_Chaos::ExecuteWrite(FPhysScene* InScene, TFunctionRef<void(F
 	return false;
 }
 
+bool FPhysInterface_Chaos::ExecuteWrite(Chaos::FPhysicsObject* InObjectA, Chaos::FPhysicsObject* InObjectB, TFunctionRef<void(Chaos::FPhysicsObject* InObjectA, Chaos::FPhysicsObject* InObjectB)> InCallable)
+{
+	FScopedSceneLock_Chaos SceneLock(InObjectA, InObjectB, EPhysicsInterfaceScopedLockType::Write);
+	InCallable(InObjectA, InObjectB);
+	return true;
+}
+
 void FPhysInterface_Chaos::ExecuteShapeWrite(FBodyInstance* InInstance, FPhysicsShapeHandle& InShape, TFunctionRef<void(FPhysicsShapeHandle& InShape)> InCallable)
 {
 	if(InInstance && InShape.IsValid())
diff --git a/Engine/Source/Runtime/Engine/Private/PhysicsEngine/PhysicsConstraintComponent.cpp b/Engine/Source/Runtime/Engine/Private/PhysicsEngine/PhysicsConstraintComponent.cpp
--- a/Engine/Source/Runtime/Engine/Private/PhysicsEngine/PhysicsConstraintComponent.cpp
+++ b/Engine/Source/Runtime/Engine/Private/PhysicsEngine/PhysicsConstraintComponent.cpp
@@ -10,7 +10,10 @@
 #include "PhysicsEngine/PhysicsConstraintTemplate.h"
 #include "PhysicsEngine/PhysicsAsset.h"
 #include "PhysicsEngine/ConstraintUtils.h"
+#include "PhysicsEngine/PhysicsObjectExternalInterface.h"
+#include "PhysicsProxy/SingleParticlePhysicsProxy.h"
 #include "Components/BillboardComponent.h"
+#include "Chaos/PhysicsObjectInterface.h"
 
 #include UE_INLINE_GENERATED_CPP_BY_NAME(PhysicsConstraintComponent)
 
@@ -134,16 +137,16 @@ int32 GetBoneIndexHelper(FName InBoneName, const USkeletalMeshComponent& SkelCom
 FTransform UPhysicsConstraintComponent::GetBodyTransformInternal(EConstraintFrame::Type Frame, FName InBoneName) const
 {
 	UPrimitiveComponent* PrimComp = GetComponentInternal(Frame);
-	if(!PrimComp)
+	if (!PrimComp)
 	{
 		return FTransform::Identity;
 	}
-	  
+
 	//Use GetComponentTransform() by default for all components
-	FTransform ResultTM = PrimComp->GetComponentTransform();
-		
+	FTransform ResultTM = FTransform::Identity;
+
 	// Skeletal case
-	if(const USkeletalMeshComponent* SkelComp = Cast<USkeletalMeshComponent>(PrimComp))
+	if (const USkeletalMeshComponent* SkelComp = Cast<USkeletalMeshComponent>(PrimComp))
 	{
 		const int32 BoneIndex = GetBoneIndexHelper(InBoneName, *SkelComp);
 		if (BoneIndex != INDEX_NONE)
@@ -158,6 +161,22 @@ FTransform UPhysicsConstraintComponent::GetBodyTransformInternal(EConstraintFram
 		}
 #endif // !(UE_BUILD_SHIPPING || UE_BUILD_TEST)
 	}
+	else if (Chaos::FPhysicsObject* InitialObject = PrimComp->GetPhysicsObjectByName(InBoneName))
+	{
+		FLockedReadPhysicsObjectExternalInterface Interface = FPhysicsObjectExternalInterface::LockRead({ &InitialObject, 1 });
+		if (Chaos::FPhysicsObject* PhysicsObject = Interface->GetRootObject(InitialObject))
+		{
+			ResultTM = Interface->GetTransform(PhysicsObject);
+		}
+		else
+		{
+			ResultTM = Interface->GetTransform(InitialObject);
+		}
+	}
+	else
+	{
+		ResultTM = PrimComp->GetComponentTransform();
+	}
 
 	return ResultTM;
 }
@@ -239,6 +258,17 @@ FBodyInstance* UPhysicsConstraintComponent::GetBodyInstance(EConstraintFrame::Ty
 	return Instance;
 }
 
+Chaos::FPhysicsObject* UPhysicsConstraintComponent::GetPhysicsObject(EConstraintFrame::Type Frame) const
+{
+	UPrimitiveComponent* PrimComp = GetComponentInternal(Frame);
+	if (!PrimComp)
+	{
+		return nullptr;
+	}
+
+	const FName BoneName = (Frame == EConstraintFrame::Frame1) ? ConstraintInstance.ConstraintBone1 : ConstraintInstance.ConstraintBone2;
+	return PrimComp->GetPhysicsObjectByName(BoneName);
+}
 
 /** Wrapper that calls our constraint broken delegate */
 void UPhysicsConstraintComponent::OnConstraintBrokenWrapper(int32 ConstraintIndex)
@@ -257,13 +287,18 @@ void UPhysicsConstraintComponent::InitComponentConstraint()
 	// First we convert world space position of constraint into local space frames
 	UpdateConstraintFrames();
 
-	// Then we init the constraint
+	// Normally, we'd want to init the constraint using the FPhysicsObject that we get from the component.
+	// However, to ensure that we can preserve backward compatible behavior, we first check to see if the
+	// FBodyInstance exists first. This way the behavior w.r.t. grabbing welded bodies remains the same.
 	FBodyInstance* Body1 = GetBodyInstance(EConstraintFrame::Frame1);
 	FBodyInstance* Body2 = GetBodyInstance(EConstraintFrame::Frame2);
 
-	if (Body1 != nullptr || Body2 != nullptr)
+	Chaos::FPhysicsObject* Object1 = (Body1 && Body1->IsValidBodyInstance()) ? Body1->ActorHandle->GetPhysicsObject() : GetPhysicsObject(EConstraintFrame::Frame1);
+	Chaos::FPhysicsObject* Object2 = (Body2 && Body2->IsValidBodyInstance()) ? Body2->ActorHandle->GetPhysicsObject() : GetPhysicsObject(EConstraintFrame::Frame2);
+
+	if (Object1 || Object2)
 	{
-		ConstraintInstance.InitConstraint(Body1, Body2, GetConstraintScale(), this, FOnConstraintBroken::CreateUObject(this, &UPhysicsConstraintComponent::OnConstraintBrokenWrapper));
+		ConstraintInstance.InitConstraint(Object1, Object2, GetConstraintScale(), this, FOnConstraintBroken::CreateUObject(this, &UPhysicsConstraintComponent::OnConstraintBrokenWrapper));
 	}
 }
 
@@ -515,12 +550,12 @@ void UPhysicsConstraintComponent::UpdateConstraintFrames()
 
 	//Note that in the case where there is no body instance, the position is given in world space and there is no scaling.
 	const float RefScale = FMath::Max(GetConstraintScale(), 0.01f);
-	if(GetBodyInstance(EConstraintFrame::Frame1))
+	if(GetPhysicsObject(EConstraintFrame::Frame1))
 	{
 		ConstraintInstance.Pos1 /= RefScale;
 	}
 
-	if (GetBodyInstance(EConstraintFrame::Frame2))
+	if (GetPhysicsObject(EConstraintFrame::Frame2))
 	{
 		ConstraintInstance.Pos2 /= RefScale;
 	}
diff --git a/Engine/Source/Runtime/Engine/Public/Physics/Experimental/PhysInterface_Chaos.h b/Engine/Source/Runtime/Engine/Public/Physics/Experimental/PhysInterface_Chaos.h
--- a/Engine/Source/Runtime/Engine/Public/Physics/Experimental/PhysInterface_Chaos.h
+++ b/Engine/Source/Runtime/Engine/Public/Physics/Experimental/PhysInterface_Chaos.h
@@ -7,6 +7,7 @@
 #include "Engine/Engine.h"
 #endif
 #include "Chaos/Declares.h"
+#include "Chaos/PhysicsObject.h"
 #include "PhysicsEngine/ConstraintTypes.h"
 #include "PhysicsInterfaceWrapperShared.h"
 #include "Physics/PhysicsInterfaceDeclares.h"
@@ -73,6 +74,7 @@ public:
     static bool ExecuteRead(const FPhysicsActorHandle& InActorReferenceA, const FPhysicsActorHandle& InActorReferenceB, TFunctionRef<void(const FPhysicsActorHandle& ActorA, const FPhysicsActorHandle& ActorB)> InCallable);
     static bool ExecuteRead(const FPhysicsConstraintHandle& InConstraintRef, TFunctionRef<void(const FPhysicsConstraintHandle& Constraint)> InCallable);
     static bool ExecuteRead(FPhysScene* InScene, TFunctionRef<void()> InCallable);
+	static bool ExecuteRead(Chaos::FPhysicsObject* InObjectA, Chaos::FPhysicsObject* InObjectB, TFunctionRef<void(Chaos::FPhysicsObject* InObjectA, Chaos::FPhysicsObject* InObjectB)> InCallable);
 
     static bool ExecuteWrite(const FPhysicsActorHandle& InActorReference, TFunctionRef<void(const FPhysicsActorHandle& Actor)> InCallable);
 	static bool ExecuteWrite(FPhysicsActorHandle& InActorReference, TFunctionRef<void(FPhysicsActorHandle& Actor)> InCallable);
@@ -81,6 +83,7 @@ public:
     static bool ExecuteWrite(const FPhysicsConstraintHandle& InConstraintRef, TFunctionRef<void(const FPhysicsConstraintHandle& Constraint)> InCallable);
     static bool ExecuteWrite(FPhysScene* InScene, TFunctionRef<void()> InCallable);
 	static bool ExecuteWrite(FPhysScene* InScene, TFunctionRef<void(FPhysScene* Scene)> InCallable);
+	static bool ExecuteWrite(Chaos::FPhysicsObject* InObjectA, Chaos::FPhysicsObject* InObjectB, TFunctionRef<void(Chaos::FPhysicsObject* InObjectA, Chaos::FPhysicsObject* InObjectB)> InCallable);
 
     static void ExecuteShapeWrite(FBodyInstance* InInstance, FPhysicsShapeHandle& InShape, TFunctionRef<void(FPhysicsShapeHandle& InShape)> InCallable);
 
diff --git a/Engine/Source/Runtime/Experimental/Chaos/Private/Chaos/PBDJointConstraintData.cpp b/Engine/Source/Runtime/Experimental/Chaos/Private/Chaos/PBDJointConstraintData.cpp
--- a/Engine/Source/Runtime/Experimental/Chaos/Private/Chaos/PBDJointConstraintData.cpp
+++ b/Engine/Source/Runtime/Experimental/Chaos/Private/Chaos/PBDJointConstraintData.cpp
@@ -2,6 +2,8 @@
 
 #include "Chaos/PBDJointConstraintData.h"
 #include "PBDRigidsSolver.h"
+#include "PhysicsProxy/SingleParticlePhysicsProxy.h"
+#include "Chaos/PhysicsObjectInterface.h"
 
 namespace Chaos
 {
@@ -47,5 +49,37 @@ namespace Chaos
 		}
 	}
 
+	void FJointConstraint::SetParticleProxies(const FProxyBasePair& InJointParticles)
+	{
+		JointProxies.Modify(/*bInvalidate=*/true, DirtyFlags, Proxy, [&InJointParticles](FProxyBasePairProperty& Data)
+		{
+			Data.ParticleProxies[0] = InJointParticles[0];
+			Data.ParticleProxies[1] = InJointParticles[1];
+		});
+
+		// This should work fine since this is a legacy endpoint and was primarily only used with single particle physics proxies.
+		JointBodies.Modify(/*bInvalidate=*/true, DirtyFlags, Proxy, [&InJointParticles](FPhysicsObjectPairProperty& Data)
+		{
+			Data.PhysicsBodies[0] = (InJointParticles[0]->GetType() == EPhysicsProxyType::SingleParticleProxy) ? static_cast<FSingleParticlePhysicsProxy*>(InJointParticles[0])->GetPhysicsObject() : nullptr;
+			Data.PhysicsBodies[1] = (InJointParticles[1]->GetType() == EPhysicsProxyType::SingleParticleProxy) ? static_cast<FSingleParticlePhysicsProxy*>(InJointParticles[1])->GetPhysicsObject() : nullptr;
+		});
+	}
+
+	void FJointConstraint::SetPhysicsBodies(const FPhysicsObjectPair& InBodies)
+	{
+		JointProxies.Modify(/*bInvalidate=*/true, DirtyFlags, Proxy, [&InBodies](FProxyBasePairProperty& Data)
+		{
+			Chaos::FPhysicsObject* Object1 = InBodies[0];
+			Chaos::FPhysicsObject* Object2 = InBodies[1];
+			Data.ParticleProxies[0] = Chaos::FPhysicsObjectInterface::GetProxy({ &Object1, 1 });
+			Data.ParticleProxies[1] = Chaos::FPhysicsObjectInterface::GetProxy({ &Object2, 1 });
+		});
+
+		JointBodies.Modify(/*bInvalidate=*/true, DirtyFlags, Proxy, [&InBodies](FPhysicsObjectPairProperty& Data)
+		{
+			Data.PhysicsBodies[0] = InBodies[0];
+			Data.PhysicsBodies[1] = InBodies[1];
+		});
+	}
 
 } // Chaos
diff --git a/Engine/Source/Runtime/Experimental/Chaos/Private/PhysicsProxy/JointConstraintProxy.cpp b/Engine/Source/Runtime/Experimental/Chaos/Private/PhysicsProxy/JointConstraintProxy.cpp
--- a/Engine/Source/Runtime/Experimental/Chaos/Private/PhysicsProxy/JointConstraintProxy.cpp
+++ b/Engine/Source/Runtime/Experimental/Chaos/Private/PhysicsProxy/JointConstraintProxy.cpp
@@ -14,6 +14,7 @@
 #include "PhysicsProxy/SingleParticlePhysicsProxy.h"
 #include "PhysicsSolver.h"
 #include "Chaos/PullPhysicsDataImp.h"
+#include "Chaos/PhysicsObjectInternal.h"
 
 namespace Chaos
 {
@@ -81,10 +82,11 @@ void FJointConstraintPhysicsProxy::InitializeOnPhysicsThread(FPBDRigidsSolver* I
 	if (Handles.Size())
 	{
 		auto& JointConstraints = InSolver->GetJointConstraints();
-		if(const FProxyBasePairProperty* BasePairs = RemoteData.FindJointParticleProxies(Manager, DataIdx))
+		if(const FPhysicsObjectPairProperty* BodyPairs = RemoteData.FindJointPhysicsObjects(Manager, DataIdx))
 		{
-			FGeometryParticleHandle* Handle0 = GetParticleHandleFromProxy(BasePairs->ParticleProxies[0]);
-			FGeometryParticleHandle* Handle1 = GetParticleHandleFromProxy(BasePairs->ParticleProxies[1]);
+			FGeometryParticleHandle* Handle0 = BodyPairs->PhysicsBodies[0]->GetRootParticle<Chaos::EThreadContext::Internal>();
+			FGeometryParticleHandle* Handle1 = BodyPairs->PhysicsBodies[1]->GetRootParticle<Chaos::EThreadContext::Internal>();
+
 			if (Handle0 && Handle1)
 			{
 				if (const FPBDJointSettings* JointSettings = RemoteData.FindJointSettings(Manager, DataIdx))
diff --git a/Engine/Source/Runtime/Experimental/Chaos/Public/Chaos/PBDJointConstraintData.h b/Engine/Source/Runtime/Experimental/Chaos/Public/Chaos/PBDJointConstraintData.h
--- a/Engine/Source/Runtime/Experimental/Chaos/Public/Chaos/PBDJointConstraintData.h
+++ b/Engine/Source/Runtime/Experimental/Chaos/Public/Chaos/PBDJointConstraintData.h
@@ -73,26 +73,22 @@ namespace Chaos
 		virtual void SyncRemoteDataImp(FDirtyPropertiesManager& Manager, int32 DataIdx, FDirtyChaosProperties& RemoteData) override
 		{
 			JointProxies.SyncRemote(Manager, DataIdx, RemoteData);
+			JointBodies.SyncRemote(Manager, DataIdx, RemoteData);
 			JointSettings.SyncRemote(Manager, DataIdx, RemoteData);
 		}
 
-		void SetParticleProxies(const FProxyBasePair& InJointParticles)
-		{
-			JointProxies.Modify(/*bInvalidate=*/true, DirtyFlags, Proxy, [&InJointParticles](FProxyBasePairProperty& Data)
-			{
-				Data.ParticleProxies[0] = InJointParticles[0];
-				Data.ParticleProxies[1] = InJointParticles[1];
-			});
-		}
-
+		void SetParticleProxies(const FProxyBasePair& InJointParticles);
+		void SetPhysicsBodies(const FPhysicsObjectPair& InBodies);
 
 		const FProxyBasePair& GetParticleProxies() const { return  JointProxies.Read().ParticleProxies; }
+		const FPhysicsObjectPair& GetPhysicsBodies() const { return  JointBodies.Read().PhysicsBodies; }
 
 	protected:
 
 		void ReleaseKinematicEndPoint(FPBDRigidsSolver* Solver);
 		
 		TChaosProperty<FProxyBasePairProperty, EChaosProperty::JointParticleProxies> JointProxies;
+		TChaosProperty<FPhysicsObjectPairProperty, EChaosProperty::JointPhysicsObjects> JointBodies;
 		TChaosProperty<FPBDJointSettings, EChaosProperty::JointSettings> JointSettings;
 
 
diff --git a/Engine/Source/Runtime/Experimental/Chaos/Public/Chaos/ParticleDirtyFlags.h b/Engine/Source/Runtime/Experimental/Chaos/Public/Chaos/ParticleDirtyFlags.h
--- a/Engine/Source/Runtime/Experimental/Chaos/Public/Chaos/ParticleDirtyFlags.h
+++ b/Engine/Source/Runtime/Experimental/Chaos/Public/Chaos/ParticleDirtyFlags.h
@@ -215,6 +215,14 @@ struct FProxyBasePairProperty
 	FProxyBasePair ParticleProxies = { nullptr, nullptr };
 };
 
+struct FPhysicsObject;
+typedef TVector<FPhysicsObject*, 2> FPhysicsObjectPair;
+
+struct FPhysicsObjectPairProperty
+{
+	FPhysicsObjectPair PhysicsBodies = { nullptr, nullptr };
+};
+
 inline FChaosArchive& operator<<(FChaosArchive& Ar, FParticleDynamics& Data)
 {
 	Data.Serialize(Ar);
diff --git a/Engine/Source/Runtime/Experimental/Chaos/Public/Chaos/ParticleProperties.inl b/Engine/Source/Runtime/Experimental/Chaos/Public/Chaos/ParticleProperties.inl
--- a/Engine/Source/Runtime/Experimental/Chaos/Public/Chaos/ParticleProperties.inl
+++ b/Engine/Source/Runtime/Experimental/Chaos/Public/Chaos/ParticleProperties.inl
@@ -19,6 +19,7 @@ CHAOS_PROPERTY(KinematicTarget,FKinematicTarget, EPhysicsProxyType::SinglePartic
 
 CHAOS_PROPERTY(JointSettings, FPBDJointSettings, EPhysicsProxyType::JointConstraintType)
 CHAOS_PROPERTY(JointParticleProxies, FProxyBasePairProperty, EPhysicsProxyType::JointConstraintType)
+CHAOS_PROPERTY(JointPhysicsObjects, FPhysicsObjectPairProperty, EPhysicsProxyType::JointConstraintType)
 
 CHAOS_PROPERTY(SuspensionSettings, FPBDSuspensionSettings, EPhysicsProxyType::SuspensionConstraintType)
 CHAOS_PROPERTY(SuspensionParticleProxy, FParticleProxyProperty, EPhysicsProxyType::SuspensionConstraintType)
diff --git a/Engine/Source/Runtime/PhysicsCore/Private/ChaosEngineInterface.cpp b/Engine/Source/Runtime/PhysicsCore/Private/ChaosEngineInterface.cpp
--- a/Engine/Source/Runtime/PhysicsCore/Private/ChaosEngineInterface.cpp
+++ b/Engine/Source/Runtime/PhysicsCore/Private/ChaosEngineInterface.cpp
@@ -8,6 +8,7 @@
 #include "Chaos/ChaosScene.h"
 #include "Chaos/KinematicTargets.h"
 #include "PhysicsInterfaceDeclaresCore.h"
+#include "PhysicsObjectPhysicsCoreInterface.h"
 
 FPhysicsDelegatesCore::FOnUpdatePhysXMaterial FPhysicsDelegatesCore::OnUpdatePhysXMaterial;
 
@@ -23,6 +24,7 @@ FPhysicsDelegatesCore::FOnUpdatePhysXMaterial FPhysicsDelegatesCore::OnUpdatePhy
 #include "Chaos/Collision/CollisionConstraintFlags.h"
 #include "PhysicsProxy/SingleParticlePhysicsProxy.h"
 #include "Chaos/ImplicitObject.h"
+#include "Chaos/PhysicsObjectInterface.h"
 #include "PBDRigidsSolver.h"
 #include "PhysicsProxy/SingleParticlePhysicsProxy.h"
 #include "Chaos/CastingUtilities.h"
@@ -988,43 +990,45 @@ SIZE_T FChaosEngineInterface::GetResourceSizeEx(const FPhysicsActorHandle& InAct
 }
 
 // Constraints
-FPhysicsConstraintHandle FChaosEngineInterface::CreateConstraint(const FPhysicsActorHandle& InActorRef1,const FPhysicsActorHandle& InActorRef2,const FTransform& InLocalFrame1,const FTransform& InLocalFrame2)
+FPhysicsConstraintHandle FChaosEngineInterface::CreateConstraint(Chaos::FPhysicsObject* Body1, Chaos::FPhysicsObject* Body2, const FTransform& InLocalFrame1, const FTransform& InLocalFrame2)
 {
 	FPhysicsConstraintHandle ConstraintRef;
 
-	if(bEnableChaosJointConstraints)
+	if (bEnableChaosJointConstraints)
 	{
-		if(InActorRef1 && InActorRef2 && InActorRef1->GetSolverBase() && InActorRef2->GetSolverBase())
+		Chaos::FPhysicsSolver* Solver1 = Chaos::FPhysicsObjectInterface::GetSolver({ &Body1, 1 });
+		Chaos::FPhysicsSolver* Solver2 = Chaos::FPhysicsObjectInterface::GetSolver({ &Body2, 1 });
+
+		if (Body1 && Body2 && Solver1 && Solver2)
 		{
-			if(InActorRef1->GetSolverBase() && InActorRef2->GetSolverBase())
-			{
-				LLM_SCOPE(ELLMTag::ChaosConstraint);
+			LLM_SCOPE(ELLMTag::ChaosConstraint);
 
-				auto* JointConstraint = new Chaos::FJointConstraint();
-				ConstraintRef.Constraint = JointConstraint;
+			auto* JointConstraint = new Chaos::FJointConstraint();
+			ConstraintRef.Constraint = JointConstraint;
 
-				JointConstraint->SetParticleProxies({ InActorRef1,InActorRef2 });
-				JointConstraint->SetJointTransforms({ InLocalFrame1,InLocalFrame2 });
+			JointConstraint->SetPhysicsBodies({ Body1, Body2 });
+			JointConstraint->SetJointTransforms({ InLocalFrame1,InLocalFrame2 });
 
-				Chaos::FPhysicsSolver* Solver = InActorRef1->GetSolver<Chaos::FPhysicsSolver>();
-				checkSlow(Solver == InActorRef2->GetSolver<Chaos::FPhysicsSolver>());
-				Solver->RegisterObject(JointConstraint);
-			}
+			checkSlow(Solver1 == Solver2);
+			Solver1->RegisterObject(JointConstraint);
 		}
-		else if (InActorRef1 != nullptr || InActorRef2 != nullptr)
+		else if (Body1 || Body2)
 		{
 			LLM_SCOPE(ELLMTag::ChaosConstraint);
 
-			FPhysicsActorHandle ValidParticle = InActorRef1;
+			Chaos::FPhysicsObject* ValidObject = Body1;
+			Chaos::FPhysicsSolver* ValidSolver = Solver1;
 			bool bSwapped = false;
-			if (ValidParticle == nullptr)
+			if (!ValidObject || !ValidSolver)
 			{
 				bSwapped = true;
-				ValidParticle = InActorRef2;
+				ValidObject = Body2;
+				ValidSolver = Solver2;
 			}
-			if(ValidParticle->GetSolverBase())
+
+			if (ValidSolver)
 			{
-				FChaosScene* Scene = FChaosEngineInterface::GetCurrentScene(ValidParticle);
+				FChaosScene* Scene = PhysicsObjectPhysicsCoreInterface::GetScene({ &ValidObject, 1 });
 
 				// Create kinematic actor to attach to joint
 				FPhysicsActorHandle KinematicEndPoint;
@@ -1045,17 +1049,16 @@ FPhysicsConstraintHandle FChaosEngineInterface::CreateConstraint(const FPhysicsA
 				JointConstraint->SetKinematicEndPoint(KinematicEndPoint, Scene->GetSolver());
 				ConstraintRef.Constraint = JointConstraint;
 
-			// Disable collision on shape to ensure it is not added to acceleration structure.
-			for (const TUniquePtr<Chaos::FPerShapeData>& Shape : KinematicEndPoint->GetGameThreadAPI().ShapesArray())
-			{
-				Chaos::FCollisionData CollisionData = Shape->GetCollisionData();
-				CollisionData.bQueryCollision = false;
-				CollisionData.bSimCollision = false;
-				Shape->SetCollisionData(CollisionData);
-			}
-
+				// Disable collision on shape to ensure it is not added to acceleration structure.
+				for (const TUniquePtr<Chaos::FPerShapeData>& Shape : KinematicEndPoint->GetGameThreadAPI().ShapesArray())
+				{
+					Chaos::FCollisionData CollisionData = Shape->GetCollisionData();
+					CollisionData.bQueryCollision = false;
+					CollisionData.bSimCollision = false;
+					Shape->SetCollisionData(CollisionData);
+				}
 
-				JointConstraint->SetParticleProxies({ ValidParticle, KinematicEndPoint });
+				JointConstraint->SetPhysicsBodies({ ValidObject, KinematicEndPoint->GetPhysicsObject() });
 
 				Chaos::FTransformPair TransformPair = { InLocalFrame1, InLocalFrame2 };
 				if (bSwapped)
@@ -1064,15 +1067,21 @@ FPhysicsConstraintHandle FChaosEngineInterface::CreateConstraint(const FPhysicsA
 				}
 				JointConstraint->SetJointTransforms(TransformPair);
 
-				Chaos::FPhysicsSolver* Solver = ValidParticle->GetSolver<Chaos::FPhysicsSolver>();
-				checkSlow(Solver == KinematicEndPoint->GetSolver<Chaos::FPhysicsSolver>());
-				Solver->RegisterObject(JointConstraint);
+				checkSlow(ValidSolver == KinematicEndPoint->GetSolver<Chaos::FPhysicsSolver>());
+				ValidSolver->RegisterObject(JointConstraint);
 			}
 		}
 	}
+
 	return ConstraintRef;
 }
 
+FPhysicsConstraintHandle FChaosEngineInterface::CreateConstraint(const FPhysicsActorHandle& InActorRef1,const FPhysicsActorHandle& InActorRef2,const FTransform& InLocalFrame1,const FTransform& InLocalFrame2)
+{
+	Chaos::FPhysicsObject* Body1 = InActorRef1 ? InActorRef1->GetPhysicsObject() : nullptr;
+	Chaos::FPhysicsObject* Body2 = InActorRef2 ? InActorRef2->GetPhysicsObject() : nullptr;
+	return FChaosEngineInterface::CreateConstraint(Body1, Body2, InLocalFrame1, InLocalFrame2);
+}
 
 FPhysicsConstraintHandle FChaosEngineInterface::CreateSuspension(const FPhysicsActorHandle& InActorRef, const FVector& InLocalFrame)
 {
diff --git a/Engine/Source/Runtime/PhysicsCore/Public/Chaos/ChaosEngineInterface.h b/Engine/Source/Runtime/PhysicsCore/Public/Chaos/ChaosEngineInterface.h
--- a/Engine/Source/Runtime/PhysicsCore/Public/Chaos/ChaosEngineInterface.h
+++ b/Engine/Source/Runtime/PhysicsCore/Public/Chaos/ChaosEngineInterface.h
@@ -4,6 +4,7 @@
 
 #include "CoreMinimal.h"
 #include "Chaos/Declares.h"
+#include "Chaos/PhysicsObject.h"
 #include "PhysicsInterfaceDeclaresCore.h"
 #include "PhysicsInterfaceWrapperShared.h"
 #include "PhysicsInterfaceTypesCore.h"
@@ -478,7 +479,9 @@ public:
 
 	static SIZE_T GetResourceSizeEx(const FPhysicsActorHandle& InActorRef);
 
+	static FPhysicsConstraintHandle CreateConstraint(Chaos::FPhysicsObject* Body1, Chaos::FPhysicsObject* Body2, const FTransform& InLocalFrame1, const FTransform& InLocalFrame2);
 	static FPhysicsConstraintHandle CreateConstraint(const FPhysicsActorHandle& InActorRef1,const FPhysicsActorHandle& InActorRef2,const FTransform& InLocalFrame1,const FTransform& InLocalFrame2);
+
 	static FPhysicsConstraintHandle CreateSuspension(const FPhysicsActorHandle& InActorRef, const FVector& InLocalFrame);
 	static void SetConstraintUserData(const FPhysicsConstraintHandle& InConstraintRef,void* InUserData);
 	static void ReleaseConstraint(FPhysicsConstraintHandle& InConstraintRef);
