EngineTest fix for audio bus subsystem
- Move default audio bus lifetime management to subsystem Initialize/Deinitialize

#rnx

diff --git a/Engine/Source/Runtime/AudioMixer/Private/AudioBusSubsystem.cpp b/Engine/Source/Runtime/AudioMixer/Private/AudioBusSubsystem.cpp
--- a/Engine/Source/Runtime/AudioMixer/Private/AudioBusSubsystem.cpp
+++ b/Engine/Source/Runtime/AudioMixer/Private/AudioBusSubsystem.cpp
@@ -12,18 +12,19 @@ UAudioBusSubsystem::UAudioBusSubsystem()
 {
 }
 
-void UAudioBusSubsystem::PostInitProperties()
+bool UAudioBusSubsystem::ShouldCreateSubsystem(UObject* Outer) const
 {
-	Super::PostInitProperties();
-	if (this != GetClass()->ClassDefaultObject)
-	{
-		InitDefaultAudioBuses();
-	}
+	return !IsRunningDedicatedServer();
 }
 
-bool UAudioBusSubsystem::ShouldCreateSubsystem(UObject* Outer) const
+void UAudioBusSubsystem::Initialize(FSubsystemCollectionBase& Collection)
 {
-	return !IsRunningDedicatedServer();
+	InitDefaultAudioBuses();
+}
+
+void UAudioBusSubsystem::Deinitialize()
+{
+	ShutdownDefaultAudioBuses();
 }
 
 void UAudioBusSubsystem::StartAudioBus(Audio::FAudioBusKey InAudioBusKey, int32 InNumChannels, bool bInIsAutomatic)
diff --git a/Engine/Source/Runtime/AudioMixer/Public/AudioBusSubsystem.h b/Engine/Source/Runtime/AudioMixer/Public/AudioBusSubsystem.h
--- a/Engine/Source/Runtime/AudioMixer/Public/AudioBusSubsystem.h
+++ b/Engine/Source/Runtime/AudioMixer/Public/AudioBusSubsystem.h
@@ -72,10 +72,11 @@ class AUDIOMIXER_API UAudioBusSubsystem : public UAudioEngineSubsystem
 public:
 	UAudioBusSubsystem();
 	virtual ~UAudioBusSubsystem() = default;
-	virtual void PostInitProperties() override;
 
 	//~ Begin USubsystem interface
 	virtual bool ShouldCreateSubsystem(UObject* Outer) const override;
+	virtual void Initialize(FSubsystemCollectionBase& Collection) override;
+	virtual void Deinitialize() override;
 	//~ End USubsystem interface
 
 	// Audio bus API from FMixerDevice
diff --git a/Engine/Source/Runtime/Engine/Private/AudioDeviceManager.cpp b/Engine/Source/Runtime/Engine/Private/AudioDeviceManager.cpp
--- a/Engine/Source/Runtime/Engine/Private/AudioDeviceManager.cpp
+++ b/Engine/Source/Runtime/Engine/Private/AudioDeviceManager.cpp
@@ -2,7 +2,6 @@
 #include "AudioDeviceManager.h"
 
 #include "Audio.h"
-#include "AudioBusSubsystem.h"
 #include "Audio/AudioDebug.h"
 #include "AudioAnalytics.h"
 #include "AudioDefines.h"
@@ -531,12 +530,6 @@ void FAudioDeviceManager::UnregisterWorld(UWorld* InWorld, Audio::FDeviceId Devi
 	{
 		if (DeviceContainer->WorldsUsingThisDevice.Contains(InWorld))
 		{
-			if (FAudioDevice* AudioDevice = DeviceContainer->Device)
-			{
-				UAudioBusSubsystem* AudioBusSubsystem = AudioDevice->GetSubsystem<UAudioBusSubsystem>();
-				check(AudioBusSubsystem);
-				AudioBusSubsystem->ShutdownDefaultAudioBuses();
-			}
 			UE_LOG(LogAudio, Display, TEXT("Audio Device unregistered from world '%s'."), *InWorld->GetName());
 			DeviceContainer->WorldsUsingThisDevice.Remove(InWorld);
 			FAudioDeviceWorldDelegates::OnWorldUnregisteredWithAudioDevice.Broadcast(InWorld, DeviceId);
