fix potential nullptr exception in InitConstraint if passing in a nullptr for FBodyInstance*

diff --git a/Engine/Source/Runtime/Engine/Private/PhysicsEngine/ConstraintInstance.cpp b/Engine/Source/Runtime/Engine/Private/PhysicsEngine/ConstraintInstance.cpp
--- a/Engine/Source/Runtime/Engine/Private/PhysicsEngine/ConstraintInstance.cpp
+++ b/Engine/Source/Runtime/Engine/Private/PhysicsEngine/ConstraintInstance.cpp
@@ -587,7 +587,7 @@ void FConstraintInstance::InitConstraint(FBodyInstance* Body1, FBodyInstance* Bo
 		return;
 	}
 
-	InitConstraint(Actor1->GetPhysicsObject(), Actor2->GetPhysicsObject(), InScale, DebugOwner, InConstraintBrokenDelegate, InPlasticDeformationDelegate);
+	InitConstraint(Actor1 ? Actor1->GetPhysicsObject() : nullptr, Actor2 ? Actor2->GetPhysicsObject() : nullptr, InScale, DebugOwner, InConstraintBrokenDelegate, InPlasticDeformationDelegate);
 }
 
 void FConstraintInstance::InitConstraint(Chaos::FPhysicsObject* Body1, Chaos::FPhysicsObject* Body2, float InScale, UObject* DebugOwner, FOnConstraintBroken InConstraintBrokenDelegate, FOnPlasticDeformation InPlasticDeformationDelegate)
