Reduced WorldPartitionHLOD replication cost
* Changed NetDormancy of HLODActor to DORM_DormantAll
* Flush net dormancy on health change from UWorldPartitionDestructibleHLODMeshComponent

diff --git a/Engine/Source/Runtime/Engine/Private/WorldPartition/HLOD/DestructibleHLODComponent.cpp b/Engine/Source/Runtime/Engine/Private/WorldPartition/HLOD/DestructibleHLODComponent.cpp
--- a/Engine/Source/Runtime/Engine/Private/WorldPartition/HLOD/DestructibleHLODComponent.cpp
+++ b/Engine/Source/Runtime/Engine/Private/WorldPartition/HLOD/DestructibleHLODComponent.cpp
@@ -87,6 +87,8 @@ void FWorldPartitionDestructibleHLODState::SetActorHealth(int32 InActorIndex, ui
 
 		DamagedActor.ActorHealth = InActorHealth;
 		MarkItemDirty(DamagedActor);
+
+		OwnerComponent->GetOwner()->FlushNetDormancy();
 	}
 
 	// If we are the client, directly update the visibility buffer
diff --git a/Engine/Source/Runtime/Engine/Private/WorldPartition/HLOD/HLODActor.cpp b/Engine/Source/Runtime/Engine/Private/WorldPartition/HLOD/HLODActor.cpp
--- a/Engine/Source/Runtime/Engine/Private/WorldPartition/HLOD/HLODActor.cpp
+++ b/Engine/Source/Runtime/Engine/Private/WorldPartition/HLOD/HLODActor.cpp
@@ -51,6 +51,8 @@ AWorldPartitionHLOD::AWorldPartitionHLOD(const FObjectInitializer& ObjectInitial
 	// needs to be replicated.
 	bReplicates = true;
 
+	NetDormancy = DORM_DormantAll;
+
 #if WITH_EDITORONLY_DATA
 	HLODHash = 0;
 	HLODBounds = FBox(EForceInit::ForceInit);
