Strata - IsThin are only editable when strata is enabled.

diff --git a/Engine/Source/Runtime/Engine/Classes/Materials/Material.h b/Engine/Source/Runtime/Engine/Classes/Materials/Material.h
--- a/Engine/Source/Runtime/Engine/Classes/Materials/Material.h
+++ b/Engine/Source/Runtime/Engine/Classes/Materials/Material.h
@@ -515,7 +515,7 @@ public:
 	UPROPERTY(EditAnywhere, Category=Material)
 	uint8 TwoSided : 1;
 
-	/** Indicates that the material should be rendered as a thin surface (i.e., without inner volume). */
+	/** Indicates that the material should be rendered as a thin surface (i.e., without inner volume). Only used by Strata materials. */
 	UPROPERTY(EditAnywhere, Category = Material)
 	uint8 bIsThinSurface : 1;
 
diff --git a/Engine/Source/Runtime/Engine/Private/Materials/Material.cpp b/Engine/Source/Runtime/Engine/Private/Materials/Material.cpp
--- a/Engine/Source/Runtime/Engine/Private/Materials/Material.cpp
+++ b/Engine/Source/Runtime/Engine/Private/Materials/Material.cpp
@@ -4322,7 +4322,6 @@ bool UMaterial::CanEditChange(const FProperty* InProperty) const
 		if (PropertyName == GET_MEMBER_NAME_STRING_CHECKED(UMaterial, bFullyRough) ||
 			PropertyName == GET_MEMBER_NAME_STRING_CHECKED(UMaterial, bNormalCurvatureToRoughness) ||
 			PropertyName == GET_MEMBER_NAME_STRING_CHECKED(UMaterial, TwoSided) ||
-			PropertyName == GET_MEMBER_NAME_STRING_CHECKED(UMaterial, bIsThinSurface) ||
 			PropertyName == GET_MEMBER_NAME_STRING_CHECKED(UMaterial, bUseLightmapDirectionality) ||
 			PropertyName == GET_MEMBER_NAME_STRING_CHECKED(UMaterial, bUseHQForwardReflections) ||
 			PropertyName == GET_MEMBER_NAME_STRING_CHECKED(UMaterial, bForwardBlendsSkyLightCubemaps) ||
@@ -4360,6 +4359,11 @@ bool UMaterial::CanEditChange(const FProperty* InProperty) const
 			return !bStrataEnabled && MaterialDomain == MD_Surface;
 		}
 
+		if (PropertyName == GET_MEMBER_NAME_STRING_CHECKED(UMaterial, bIsThinSurface))
+		{
+			return bStrataEnabled && MaterialDomain == MD_Surface;
+		}
+
 		if (FCString::Strncmp(*PropertyName, TEXT("bUsedWith"), 9) == 0)
 		{
 			return MaterialDomain == MD_DeferredDecal || MaterialDomain == MD_Surface || MaterialDomain == MD_Volume;
