Fix bone transform revision number not updated when a lead pose component cannot render and hence not have a valid mesh object.

diff --git a/Engine/Source/Runtime/Engine/Private/Components/SkinnedMeshComponent.cpp b/Engine/Source/Runtime/Engine/Private/Components/SkinnedMeshComponent.cpp
--- a/Engine/Source/Runtime/Engine/Private/Components/SkinnedMeshComponent.cpp
+++ b/Engine/Source/Runtime/Engine/Private/Components/SkinnedMeshComponent.cpp
@@ -919,8 +919,10 @@ void USkinnedMeshComponent::CreateRenderState_Concurrent(FRegisterComponentConte
 
 	if (GetSkinnedAsset())
 	{
-		// Update dynamic data
+		// Update revision number before mesh object update. We need to update regardless if MeshObject exists.
+		UpdateBoneTransformRevisionNumber();
 
+		// Update dynamic data
 		if(MeshObject)
 		{
 			// Clamp LOD within the VALID range
@@ -955,7 +957,6 @@ void USkinnedMeshComponent::CreateRenderState_Concurrent(FRegisterComponentConte
 				}
 
 				RefreshExternalMorphTargetWeights();
-				UpdateBoneTransformRevisionNumber();
 				MeshObject->Update(ModifiedLODLevel, this, ActiveMorphTargets, MorphTargetWeights, EPreviousBoneTransformUpdateMode::UpdatePrevious, GetExternalMorphWeights(ModifiedLODLevel));  // send to rendering thread
 			}
 		}
@@ -1040,6 +1041,9 @@ void USkinnedMeshComponent::SendRenderDynamicData_Concurrent()
 	}
 #endif
 
+	// Update revision number before mesh object update. We need to update regardless if MeshObject exists.
+	EPreviousBoneTransformUpdateMode PreviousBoneTransformUpdateMode = UpdateBoneTransformRevisionNumber();
+
 	// if we have not updated the transforms then no need to send them to the rendering thread
 	// @todo GIsEditor used to be bUpdateSkelWhenNotRendered. Look into it further to find out why it doesn't update animations in the AnimSetViewer, when a level is loaded in UED (like POC_Cover.gear).
 	if( MeshObject && GetSkinnedAsset() && (bForceMeshObjectUpdate || (bRecentlyRendered || VisibilityBasedAnimTickOption == EVisibilityBasedAnimTickOption::AlwaysTickPoseAndRefreshBones || GIsEditor || MeshObject->bHasBeenUpdatedAtLeastOnce == false)) )
@@ -1065,8 +1069,6 @@ void USkinnedMeshComponent::SendRenderDynamicData_Concurrent()
 				ActiveMorphTargets.Empty();
 			}
 			
-			EPreviousBoneTransformUpdateMode PreviousBoneTransformUpdateMode = UpdateBoneTransformRevisionNumber();
-
 			MeshObject->Update(
 				UseLOD,
 				this,
