Content Bundle Cooking: shortened paths by removing the Content Bundle GUID from the path, since it's already appended to the package name.

diff --git a/Engine/Source/Runtime/Engine/Private/WorldPartition/ContentBundle/ContentBundle.cpp b/Engine/Source/Runtime/Engine/Private/WorldPartition/ContentBundle/ContentBundle.cpp
--- a/Engine/Source/Runtime/Engine/Private/WorldPartition/ContentBundle/ContentBundle.cpp
+++ b/Engine/Source/Runtime/Engine/Private/WorldPartition/ContentBundle/ContentBundle.cpp
@@ -30,7 +30,7 @@ void FContentBundle::DoInitialize()
 #if WITH_EDITOR
 	InitializeForPIE();
 #else
-	ExternalStreamingObjectPackage = LoadPackage(nullptr, *GetExternalStreamingObjectPackageName(), LOAD_None);
+	ExternalStreamingObjectPackage = LoadPackage(nullptr, *GetExternalStreamingObjectPackagePath(), LOAD_None);
 	if (ExternalStreamingObjectPackage != nullptr)
 	{
 		if (UObject* Object = StaticFindObjectFast(URuntimeHashExternalStreamingObjectBase::StaticClass(), ExternalStreamingObjectPackage, *GetExternalStreamingObjectName()))
@@ -40,7 +40,7 @@ void FContentBundle::DoInitialize()
 		}
 		else
 		{
-			UE_LOG(LogContentBundle, Error, TEXT("[CB: %s] No streaming object found in package %s."), *GetDescriptor()->GetDisplayName(), *GetExternalStreamingObjectPackageName());
+			UE_LOG(LogContentBundle, Error, TEXT("[CB: %s] No streaming object found in package %s."), *GetDescriptor()->GetDisplayName(), *GetExternalStreamingObjectPackagePath());
 		}
 	}
 	else
diff --git a/Engine/Source/Runtime/Engine/Private/WorldPartition/ContentBundle/ContentBundleBase.cpp b/Engine/Source/Runtime/Engine/Private/WorldPartition/ContentBundle/ContentBundleBase.cpp
--- a/Engine/Source/Runtime/Engine/Private/WorldPartition/ContentBundle/ContentBundleBase.cpp
+++ b/Engine/Source/Runtime/Engine/Private/WorldPartition/ContentBundle/ContentBundleBase.cpp
@@ -94,7 +94,12 @@ const UContentBundleDescriptor* FContentBundleBase::GetDescriptor() const
 
 FString FContentBundleBase::GetExternalStreamingObjectPackageName() const
 {
-	return ContentBundlePaths::GetCookedContentBundleLevelFolder(*this) + ContentBundlePaths::GetGeneratedFolderName() + TEXT("/StreamingObject");
+	return TEXT("StreamingObject_") + GetDescriptor()->GetGuid().ToString(EGuidFormats::Base36Encoded);
+}
+
+FString FContentBundleBase::GetExternalStreamingObjectPackagePath() const
+{
+	return ContentBundlePaths::GetCookedContentBundleLevelFolder(*this) + ContentBundlePaths::GetGeneratedFolderName() + TEXT("/") + GetExternalStreamingObjectPackageName();
 }
 
 FString FContentBundleBase::GetExternalStreamingObjectName() const
diff --git a/Engine/Source/Runtime/Engine/Private/WorldPartition/ContentBundle/ContentBundleEditor.cpp b/Engine/Source/Runtime/Engine/Private/WorldPartition/ContentBundle/ContentBundleEditor.cpp
--- a/Engine/Source/Runtime/Engine/Private/WorldPartition/ContentBundle/ContentBundleEditor.cpp
+++ b/Engine/Source/Runtime/Engine/Private/WorldPartition/ContentBundle/ContentBundleEditor.cpp
@@ -373,7 +373,7 @@ bool FContentBundleEditor::GatherPackagesToCook(class IWorldPartitionCookPackage
 			}
 		});
 		
-		const FWorldPartitionCookPackage* CookPackage = CookContext.AddGenericPackageToGenerate(this, ContentBundlePaths::GetCookedContentBundleLevelFolder(*this), TEXT("StreamingObject"));
+		const FWorldPartitionCookPackage* CookPackage = CookContext.AddGenericPackageToGenerate(this, ContentBundlePaths::GetCookedContentBundleLevelFolder(*this), GetExternalStreamingObjectPackageName());
 		if (CookPackage == nullptr)
 		{
 			UE_LOG(LogContentBundle, Error, TEXT("[CB: %s][Cook] Failed to add streaming object package in cook context %s."), *GetDescriptor()->GetDisplayName());
diff --git a/Engine/Source/Runtime/Engine/Private/WorldPartition/ContentBundle/ContentBundlePaths.cpp b/Engine/Source/Runtime/Engine/Private/WorldPartition/ContentBundle/ContentBundlePaths.cpp
--- a/Engine/Source/Runtime/Engine/Private/WorldPartition/ContentBundle/ContentBundlePaths.cpp
+++ b/Engine/Source/Runtime/Engine/Private/WorldPartition/ContentBundle/ContentBundlePaths.cpp
@@ -15,25 +15,12 @@
 #include "PackageTools.h"
 #endif
 
-namespace ContentBundlePaths
-{
-	// ContentBundleCookedFolder is reduce to avoid reaching MAX_PATH
-	constexpr FStringView ContentBundleCookedFolder = TEXTVIEW("/CB/");
-}
-
 FString ContentBundlePaths::GetCookedContentBundleLevelFolder(const FContentBundleBase& ContentBundle)
 {
 	TStringBuilderWithBuffer<TCHAR, NAME_SIZE> GeneratedContentBundleLevelFolder;
 	GeneratedContentBundleLevelFolder += TEXT("/");
 	GeneratedContentBundleLevelFolder += ContentBundle.GetDescriptor()->GetPackageRoot();
-	GeneratedContentBundleLevelFolder += TEXT("/");
-	GeneratedContentBundleLevelFolder += ContentBundlePaths::ContentBundleCookedFolder;
-	GeneratedContentBundleLevelFolder += TEXT("/");
-	
-	// ContentBundleGuid is reduced to avoid reaching MAX_PATH
-	ContentBundle.GetDescriptor()->GetGuid().AppendString(GeneratedContentBundleLevelFolder, EGuidFormats::Short); 
-
-	GeneratedContentBundleLevelFolder += TEXT("/");
+	GeneratedContentBundleLevelFolder += TEXT("/CB/");
 	GeneratedContentBundleLevelFolder += GetRelativeLevelFolder(ContentBundle);
 	GeneratedContentBundleLevelFolder += TEXT("/");
 
@@ -42,7 +29,6 @@ FString ContentBundlePaths::GetCookedContentBundleLevelFolder(const FContentBund
 	return Result;
 }
 
-
 FString ContentBundlePaths::GetRelativeLevelFolder(const FContentBundleBase& ContentBundle)
 {
 	const UWorld* InjectedWorld = ContentBundle.GetInjectedWorld();
diff --git a/Engine/Source/Runtime/Engine/Public/WorldPartition/ContentBundle/ContentBundleBase.h b/Engine/Source/Runtime/Engine/Public/WorldPartition/ContentBundle/ContentBundleBase.h
--- a/Engine/Source/Runtime/Engine/Public/WorldPartition/ContentBundle/ContentBundleBase.h
+++ b/Engine/Source/Runtime/Engine/Public/WorldPartition/ContentBundle/ContentBundleBase.h
@@ -31,6 +31,7 @@ public:
 	const UContentBundleDescriptor* GetDescriptor() const;
 	EContentBundleStatus GetStatus() const { return Status; }
 	FString GetExternalStreamingObjectPackageName() const;
+	FString GetExternalStreamingObjectPackagePath() const;
 	FString GetExternalStreamingObjectName() const;
 
 protected:
