Add option to skip tracking of certain PSO precache requests if the pass can provide a fallback rendering

#rnx

diff --git a/Engine/Source/Runtime/Engine/Private/PSOPrecache.cpp b/Engine/Source/Runtime/Engine/Private/PSOPrecache.cpp
--- a/Engine/Source/Runtime/Engine/Private/PSOPrecache.cpp
+++ b/Engine/Source/Runtime/Engine/Private/PSOPrecache.cpp
@@ -70,7 +70,7 @@ FGraphEventArray PrecachePSOs(const TArray<FPSOPrecacheData>& PSOInitializers)
 #endif // PSO_PRECACHING_VALIDATE
 
 			FGraphEventRef GraphEvent = PipelineStateCache::PrecacheGraphicsPipelineState(PrecacheData.GraphicsPSOInitializer);
-			if (GraphEvent)
+			if (GraphEvent && PrecacheData.bRequired)
 			{
 				GraphEvents.Add(GraphEvent);
 			}
@@ -83,7 +83,7 @@ FGraphEventArray PrecachePSOs(const TArray<FPSOPrecacheData>& PSOInitializers)
 #endif // PSO_PRECACHING_VALIDATE
 
 			FGraphEventRef GraphEvent = PipelineStateCache::PrecacheComputePipelineState(PrecacheData.ComputeShader);
-			if (GraphEvent)
+			if (GraphEvent && PrecacheData.bRequired)
 			{
 				GraphEvents.Add(GraphEvent);
 			}
diff --git a/Engine/Source/Runtime/Engine/Public/PSOPrecache.h b/Engine/Source/Runtime/Engine/Public/PSOPrecache.h
--- a/Engine/Source/Runtime/Engine/Public/PSOPrecache.h
+++ b/Engine/Source/Runtime/Engine/Public/PSOPrecache.h
@@ -116,7 +116,7 @@ struct FPSOPrecacheParams
  */
 struct FPSOPrecacheData
 {
-	FPSOPrecacheData() : Type(EType::Graphics)
+	FPSOPrecacheData() : Type(EType::Graphics), bRequired(true)
 #if PSO_PRECACHING_VALIDATE
 		, MeshPassType(0)
 		, VertexFactoryType(nullptr)
@@ -131,6 +131,10 @@ struct FPSOPrecacheData
 	};
 	EType Type;
 
+	// Is the PSO required to be able render the object or can it provide a fallback path
+	// (proxy creation won't wait for these PSOs if enabled)
+	bool bRequired;
+
 	union
 	{
 		FGraphicsPipelineStateInitializer GraphicsPSOInitializer;
