Strata - improved strata debug view cvar documentation.
- now gives reason as to why it is not working.
- Menu entries unchanged,

diff --git a/Engine/Source/Runtime/Engine/Private/StrataVisualizationData.cpp b/Engine/Source/Runtime/Engine/Private/StrataVisualizationData.cpp
--- a/Engine/Source/Runtime/Engine/Private/StrataVisualizationData.cpp
+++ b/Engine/Source/Runtime/Engine/Private/StrataVisualizationData.cpp
@@ -35,8 +35,13 @@ static FString ConfigureConsoleCommand(FStrataVisualizationData::TModeMap& ModeM
 	for (FStrataVisualizationData::TModeMap::TConstIterator It = ModeMap.CreateConstIterator(); It; ++It)
 	{
 		const FStrataVisualizationData::FModeRecord& Record = It.Value();
-		AvailableVisualizationModes += FString(TEXT("\n  "));
+		AvailableVisualizationModes += FString(TEXT("\n  Value="));
+		AvailableVisualizationModes += FString::Printf(TEXT("%d: "), uint8(Record.Mode));
 		AvailableVisualizationModes += Record.ModeString;
+		if (!Record.bAvailableCommand)
+		{
+			AvailableVisualizationModes += FString::Printf(TEXT(" --- Unavailable, reason: %s"), *Record.UnavailableReason.ToString());
+		}
 	}
 
 	FString Out;
@@ -58,7 +63,9 @@ static void AddVisualizationMode(
 	const FText& ModeText,
 	const FText& ModeDesc,
 	const FStrataVisualizationData::FViewMode Mode,
-	bool DefaultComposited
+	bool bDefaultComposited,
+	bool bAvailableCommand,
+	const FText& UnavailableReason
 )
 {
 	const FName ModeName = FName(ModeString);
@@ -69,79 +76,97 @@ static void AddVisualizationMode(
 	Record.ModeText				= ModeText;
 	Record.ModeDesc				= ModeDesc;
 	Record.Mode					= Mode;
-	Record.DefaultComposited	= DefaultComposited;
+	Record.bDefaultComposited	= bDefaultComposited;
+	Record.bAvailableCommand	= bAvailableCommand;
+	Record.UnavailableReason	= UnavailableReason;
 }
 
 void FStrataVisualizationData::Initialize()
 {
 	if (!bIsInitialized && Engine_IsStrataEnabled())
 	{
+		TModeMap AllModeMap;
+
 		AddVisualizationMode(
-			ModeMap,
+			AllModeMap,
 			TEXT("MaterialProperties"),
-			LOCTEXT("MaterialProperties", "Material Properites"),
+			LOCTEXT("MaterialProperties", "Material Properties"),
 			LOCTEXT("MaterialPropertiesDesc", "Visualizes Strata material properties under mouse cursor"),
 			FViewMode::MaterialProperties,
-			true);
+			true,
+			true,
+			LOCTEXT("None", "None"));
 
 		AddVisualizationMode(
-			ModeMap,
+			AllModeMap,
 			TEXT("MaterialCount"),
 			LOCTEXT("MaterialCount", "Material Count"),
 			LOCTEXT("MaterialCountDesc", "Visualizes Strata material count per pixel"),
 			FViewMode::MaterialCount,
-			true);
+			true,
+			true,
+			LOCTEXT("None", "None"));
 
-		if (Engine_IsStrataAdvancedDebugShaderEnabled())
-		{
-			AddVisualizationMode(
-				ModeMap,
-				TEXT("AdvancedMaterialProperties"),
-				LOCTEXT("AdvancedMaterialProperties", "Advanced Material Properties"),
-				LOCTEXT("AdvancedMaterialPropertiesDesc", "Visualizes Strata advanced material properties"),
-				FViewMode::AdvancedMaterialProperties,
-				true);
-		}
+		AddVisualizationMode(
+			AllModeMap,
+			TEXT("AdvancedMaterialProperties"),
+			LOCTEXT("AdvancedMaterialProperties", "Advanced Material Properties"),
+			LOCTEXT("AdvancedMaterialPropertiesDesc", "Visualizes Strata advanced material properties"),
+			FViewMode::AdvancedMaterialProperties,
+			true,
+			Engine_IsStrataAdvancedDebugShaderEnabled(),
+			LOCTEXT("IsStrataAdvancedDebugShaderEnabled", "Strata advanced debugging r.Strata.Debug.AdvancedVisualizationShaders is disabled"));
 
 		AddVisualizationMode(
-			ModeMap,
+			AllModeMap,
 			TEXT("MaterialClassification"),
 			LOCTEXT("MaterialClassification", "Material Classification"),
 			LOCTEXT("MaterialClassificationDesc", "Visualizes Strata material classification"),
 			FViewMode::MaterialClassification,
-			true);
+			true,
+			true,
+			LOCTEXT("None", "None"));
 
-		if (Engine_IsStrataDBufferPassEnabled())
-		{
-			AddVisualizationMode(
-				ModeMap,
-				TEXT("DecalClassification"),
-				LOCTEXT("DecalClassification", "Decal classification"),
-				LOCTEXT("DecalClassificationDesc", "Visualizes Strata decal classification"),
-				FViewMode::DecalClassification,
-				true);
-		}
+		AddVisualizationMode(
+			AllModeMap,
+			TEXT("DecalClassification"),
+			LOCTEXT("DecalClassification", "Decal classification"),
+			LOCTEXT("DecalClassificationDesc", "Visualizes Strata decal classification"),
+			FViewMode::DecalClassification,
+			true,
+			Engine_IsStrataDBufferPassEnabled(),
+			LOCTEXT("IsStrataDBufferPassEnabled", "Strata tiled DBuffer pass (r.Strata.DBufferPass and r.Strata.DBufferPass.DedicatedTiles) is disabled"));
 
-		if (Engine_IsStrataRoughRefractionEnabled())
-		{
-			AddVisualizationMode(
-				ModeMap,
-				TEXT("RoughRefractionClassification"),
-				LOCTEXT("RoughRefractionClassification", "Material Count"),
-				LOCTEXT("RoughRefractionClassificationDesc", "Visualizes Strata rough refraction classification"),
-				FViewMode::RoughRefractionClassification,
-				true);
-		}
+		AddVisualizationMode(
+			AllModeMap,
+			TEXT("RoughRefractionClassification"),
+			LOCTEXT("RoughRefractionClassification", "Material Count"),
+			LOCTEXT("RoughRefractionClassificationDesc", "Visualizes Strata rough refraction classification"),
+			FViewMode::RoughRefractionClassification,
+			true,
+			Engine_IsStrataRoughRefractionEnabled(),
+			LOCTEXT("IsStrataRoughRefractionEnabled", "Strata rough refraction r.Strata.OpaqueMaterialRoughRefraction is disabled"));
 
 		AddVisualizationMode(
-			ModeMap,
+			AllModeMap,
 			TEXT("StrataInfo"),
 			LOCTEXT("StrataInfo", "Strata Info"),
 			LOCTEXT("StrataInfoDesc", "Visualizes Strata info"),
 			FViewMode::StrataInfo,
-			true);
+			true,
+			true,
+			LOCTEXT("None", "None"));
 
-		ConsoleDocumentationVisualizationMode = ConfigureConsoleCommand(ModeMap);
+		ConsoleDocumentationVisualizationMode = ConfigureConsoleCommand(AllModeMap);
+
+		// Now only copy the available modes for the menu to not overload it with useless entries.
+		for (auto& Mode : AllModeMap)
+		{
+			if(Mode.Value.bAvailableCommand)
+			{
+				ModeMap.Emplace(Mode.Key) = Mode.Value;
+			}
+		}
 	}
 	bIsInitialized = true;
 }
@@ -174,7 +199,7 @@ bool FStrataVisualizationData::GetModeDefaultComposited(const FName& InModeName)
 {
 	if (const FModeRecord* Record = ModeMap.Find(InModeName))
 	{
-		return Record->DefaultComposited;
+		return Record->bDefaultComposited;
 	}
 	else
 	{
diff --git a/Engine/Source/Runtime/Engine/Public/StrataVisualizationData.h b/Engine/Source/Runtime/Engine/Public/StrataVisualizationData.h
--- a/Engine/Source/Runtime/Engine/Public/StrataVisualizationData.h
+++ b/Engine/Source/Runtime/Engine/Public/StrataVisualizationData.h
@@ -29,7 +29,10 @@ public:
 		FViewMode Mode;
 
 		// Whether or not this mode (by default) composites with regular scene depth.
-		bool      DefaultComposited;
+		bool      bDefaultComposited;
+
+		bool      bAvailableCommand;
+		FText	  UnavailableReason;
 	};
 
 	/** Mapping of FName to a visualization mode record. */
