Fixed crashes related to undoing a parameter change in the Material Editor.

diff --git a/Engine/Source/Editor/MaterialEditor/Private/MaterialEditor.cpp b/Engine/Source/Editor/MaterialEditor/Private/MaterialEditor.cpp
--- a/Engine/Source/Editor/MaterialEditor/Private/MaterialEditor.cpp
+++ b/Engine/Source/Editor/MaterialEditor/Private/MaterialEditor.cpp
@@ -6383,6 +6383,17 @@ void FMaterialEditor::PostUndo(bool bSuccess)
 			FocusedGraphEd->ClearSelectionSet();
 		}
 		
+		// After an undo operation, the reference contained in graph node to its material expression is lost
+		// (the serializer sets it to null). Here we relink the material expressions graph nodes back.
+		for (UMaterialExpression* MaterialExpression : Material->GetExpressions())
+		{
+			UMaterialGraphNode* GraphNode = Cast<UMaterialGraphNode>(MaterialExpression->GraphNode);
+			if (GraphNode)
+			{
+				GraphNode->MaterialExpression = MaterialExpression;
+			}
+		}
+
 		Material->BuildEditorParameterList();
 
 		// Update the current preview material.
diff --git a/Engine/Source/Editor/UnrealEd/Private/MaterialGraphNode.cpp b/Engine/Source/Editor/UnrealEd/Private/MaterialGraphNode.cpp
--- a/Engine/Source/Editor/UnrealEd/Private/MaterialGraphNode.cpp
+++ b/Engine/Source/Editor/UnrealEd/Private/MaterialGraphNode.cpp
@@ -74,7 +74,7 @@ void UMaterialGraphNode::PostCopyNode()
 
 FMaterialRenderProxy* UMaterialGraphNode::GetExpressionPreview()
 {
-	return FMaterialEditorUtilities::GetExpressionPreview(GetGraph(), MaterialExpression);
+	return MaterialExpression ? FMaterialEditorUtilities::GetExpressionPreview(GetGraph(), MaterialExpression) : nullptr;
 }
 
 void UMaterialGraphNode::RecreateAndLinkNode()
diff --git a/Engine/Source/Runtime/Engine/Private/Materials/MaterialExpressions.cpp b/Engine/Source/Runtime/Engine/Private/Materials/MaterialExpressions.cpp
--- a/Engine/Source/Runtime/Engine/Private/Materials/MaterialExpressions.cpp
+++ b/Engine/Source/Runtime/Engine/Private/Materials/MaterialExpressions.cpp
@@ -1507,48 +1507,51 @@ void UMaterialExpression::PostEditChangeProperty(FPropertyChangedEvent& Property
 			const FString& NewDefaultValue = GetInputPinDefaultValue(PinIndex);
 
 			// Update the pin value of the expression input
-			UEdGraphPin* Pin = GraphNode->GetPinAt(PinIndex);
-			if (Pin)
+			if (UEdGraphPin* Pin = GraphNode->GetPinAt(PinIndex))
 			{
 				Pin->Modify();
 				Pin->DefaultValue = NewDefaultValue;
-			}
 
-			// If the property is linked as inline toggle to another property, both pins need updating to reflect the change.
-			bool bInlineEditConditionToggle = MemberPropertyThatChanged->HasMetaData(TEXT("InlineEditConditionToggle"));
-			bool bEditCondition = MemberPropertyThatChanged->HasMetaData(TEXT("EditCondition"));
-			if (bInlineEditConditionToggle || bEditCondition)
-			{
-				CastChecked<UMaterialGraphNode>(GraphNode)->RecreateAndLinkNode();
-			}
-
-			// If this expression refers to a parameter, we need to keep the pin state in sync with all other nodes of the same type as this node.
-			if (IsA<UMaterialExpressionParameter>())
-			{
-				// Remember this expression parameter name.
-				const FName& ParameterName = static_cast<UMaterialExpressionParameter*>(this)->ParameterName;
+				// If this expression refers to a parameter, we need to keep the pin state in sync with all other nodes of the same type as this node.
+				if (IsA<UMaterialExpressionParameter>())
+				{
+					// Remember this expression parameter name.
+					const FName& ParameterName = static_cast<UMaterialExpressionParameter*>(this)->ParameterName;
 
-				// Fetch all nodes in the material that refer to a parameter.
-				TArray<UMaterialExpressionParameter*> ParameterExpressions;
-				Material->GetAllExpressionsInMaterialAndFunctionsOfType<UMaterialExpressionParameter>(ParameterExpressions);
+					// Fetch all nodes in the material that refer to a parameter.
+					TArray<UMaterialExpressionParameter*> ParameterExpressions;
+					Material->GetAllExpressionsInMaterialAndFunctionsOfType<UMaterialExpressionParameter>(ParameterExpressions);
 
-				for (UMaterialExpressionParameter* ExpressionParameter : ParameterExpressions)
-				{
-					// If the other expression type and parameter name are the same as this expression's...
-					if (ExpressionParameter->GraphNode && ExpressionParameter->GetArchetype() == GetArchetype() && ExpressionParameter->ParameterName == ParameterName)
+					for (UMaterialExpressionParameter* ExpressionParameter : ParameterExpressions)
 					{
-						// ...modify the pin on other parameter expression node with the new value.
-						UEdGraphPin* OtherPin = ExpressionParameter->GraphNode->GetPinAt(PinIndex);
-						if (ensure(OtherPin && Pin && OtherPin->GetName() == Pin->GetName()))
+						// If the other expression type and parameter name are the same as this expression's...
+						if (ExpressionParameter->GraphNode && ExpressionParameter->GetArchetype() == GetArchetype() && ExpressionParameter->ParameterName == ParameterName)
 						{
-							OtherPin->Modify();
-							OtherPin->DefaultValue = NewDefaultValue;
+							// ...modify the pin on other parameter expression node with the new value.
+							UEdGraphPin* OtherPin = ExpressionParameter->GraphNode->GetPinAt(PinIndex);
+							if (ensure(OtherPin && Pin && OtherPin->GetName() == Pin->GetName()))
+							{
+								OtherPin->Modify();
+								OtherPin->DefaultValue = NewDefaultValue;
+							}
 						}
 					}
+
+					// Propagate he parameter value change so that it updates the other caches.
+					// Note: since this could create a transaction, avoid creating a secondary nested transition.
+					if (!GIsTransacting)
+					{
+						Material->PropagateExpressionParameterChanges(this);
+					}
 				}
+			}
 
-				// Propagate he parameter value change so that it updates the other caches.
-				Material->PropagateExpressionParameterChanges(this);
+			// If the property is linked as inline toggle to another property, both pins need updating to reflect the change.
+			bool bInlineEditConditionToggle = MemberPropertyThatChanged->HasMetaData(TEXT("InlineEditConditionToggle"));
+			bool bEditCondition = MemberPropertyThatChanged->HasMetaData(TEXT("EditCondition"));
+			if (bInlineEditConditionToggle || bEditCondition)
+			{
+				CastChecked<UMaterialGraphNode>(GraphNode)->RecreateAndLinkNode();
 			}
 		}
 	}
