AssetManager.cpp: SortIncludes
#rnx

diff --git a/Engine/Source/Runtime/Engine/Private/AssetManager.cpp b/Engine/Source/Runtime/Engine/Private/AssetManager.cpp
--- a/Engine/Source/Runtime/Engine/Private/AssetManager.cpp
+++ b/Engine/Source/Runtime/Engine/Private/AssetManager.cpp
@@ -1,47 +1,48 @@
 // Copyright Epic Games, Inc. All Rights Reserved.
 
 #include "Engine/AssetManager.h"
-#include "Engine/AssetManagerSettings.h"
-#include "Engine/BlueprintCore.h"
-#include "Engine/GameInstance.h"
-#include "Engine/PrimaryAssetLabel.h"
-#include "AssetRegistry/AssetData.h"
+
 #include "AssetRegistry/ARFilter.h"
+#include "AssetRegistry/AssetData.h"
 #include "AssetRegistry/AssetRegistryHelpers.h"
+#include "AssetRegistry/AssetRegistryModule.h"
+#include "AssetRegistry/AssetRegistryState.h"
 #include "Containers/StringView.h"
-#include "Engine/Engine.h"
+#include "Engine/AssetManagerSettings.h"
+#include "Engine/BlueprintCore.h"
 #include "Engine/BlueprintGeneratedClass.h"
+#include "Engine/Engine.h"
+#include "Engine/GameInstance.h"
 #include "Engine/ICookInfo.h"
+#include "Engine/PrimaryAssetLabel.h"
+#include "HAL/PlatformApplicationMisc.h"
+#include "HAL/PlatformFileManager.h"
+#include "HAL/PlatformMisc.h"
 #include "Interfaces/IPluginManager.h"
-#include "MoviePlayerProxy.h"
-#include "UObject/ConstructorHelpers.h"
-#include "UObject/LinkerLoad.h"
-#include "UObject/ObjectSaveContext.h"
-#include "UObject/UObjectHash.h"
+#include "Internationalization/PackageLocalizationManager.h"
+#include "IPlatformFilePak.h"
 #include "Misc/AssertionMacros.h"
 #include "Misc/FileHelper.h"
-#include "Misc/ScopedSlowTask.h"
 #include "Misc/Paths.h"
+#include "Misc/ScopedSlowTask.h"
 #include "Misc/StringBuilder.h"
+#include "MoviePlayerProxy.h"
 #include "Serialization/MemoryReader.h"
-#include "AssetRegistry/AssetRegistryState.h"
-#include "HAL/PlatformFileManager.h"
-#include "IPlatformFilePak.h"
 #include "Stats/StatsMisc.h"
 #include "String/Find.h"
-#include "Internationalization/PackageLocalizationManager.h"
-#include "HAL/PlatformMisc.h"
-#include "HAL/PlatformApplicationMisc.h"
-#include "AssetRegistry/AssetRegistryModule.h"
+#include "UObject/ConstructorHelpers.h"
+#include "UObject/LinkerLoad.h"
+#include "UObject/ObjectSaveContext.h"
+#include "UObject/UObjectHash.h"
 
 #include UE_INLINE_GENERATED_CPP_BY_NAME(AssetManager)
 
 #if WITH_EDITOR
 #include "Editor.h"
-#include "Widgets/Notifications/SNotificationList.h"
-#include "Framework/Notifications/NotificationManager.h"
 #include "Commandlets/ChunkDependencyInfo.h"
+#include "Framework/Notifications/NotificationManager.h"
 #include "Settings/ProjectPackagingSettings.h"
+#include "Widgets/Notifications/SNotificationList.h"
 #endif
 
 #define LOCTEXT_NAMESPACE "AssetManager"
