Added support for setting r.GPUSkin.UnlimitedBoneInfluences and r.GPUSkin.UnlimitedBoneInfluencesThreshold per platform

diff --git a/Engine/Source/Developer/MeshUtilities/Private/MeshUtilities.cpp b/Engine/Source/Developer/MeshUtilities/Private/MeshUtilities.cpp
--- a/Engine/Source/Developer/MeshUtilities/Private/MeshUtilities.cpp
+++ b/Engine/Source/Developer/MeshUtilities/Private/MeshUtilities.cpp
@@ -4074,7 +4074,7 @@ public:
 						InfluenceCount = FMath::Min<uint32>(InfluenceCount, InfluenceLimit);
 					}
 
-					if (InfluenceCount > EXTRA_BONE_INFLUENCES && !FGPUBaseSkinVertexFactory::UseUnlimitedBoneInfluences(InfluenceCount))
+					if (InfluenceCount > EXTRA_BONE_INFLUENCES && !FGPUBaseSkinVertexFactory::UseUnlimitedBoneInfluences(InfluenceCount, BuildData.BuildOptions.TargetPlatform))
 					{
 						// Add a single warning message, there could be too many
 						if (bIsFirstBoneInfluenceTruncation)
diff --git a/Engine/Source/Developer/SkeletalMeshUtilitiesCommon/Private/LODUtilities.cpp b/Engine/Source/Developer/SkeletalMeshUtilitiesCommon/Private/LODUtilities.cpp
--- a/Engine/Source/Developer/SkeletalMeshUtilitiesCommon/Private/LODUtilities.cpp
+++ b/Engine/Source/Developer/SkeletalMeshUtilitiesCommon/Private/LODUtilities.cpp
@@ -2714,7 +2714,7 @@ void FLODUtilities::GenerateImportedSkinWeightProfileData(
 	TArray<FRawSkinWeight>& SkinWeights = ImportedProfileData.SkinWeights;
 	SkinWeights.Empty(DestinationSoftVertices.Num());
 
-	const int32 MaxBoneInfluencesFromProjectSettings = FGPUBaseSkinVertexFactory::UseUnlimitedBoneInfluences(MAX_TOTAL_INFLUENCES) ? MAX_TOTAL_INFLUENCES : EXTRA_BONE_INFLUENCES;
+	const int32 MaxBoneInfluencesFromProjectSettings = FGPUBaseSkinVertexFactory::UseUnlimitedBoneInfluences(MAX_TOTAL_INFLUENCES, TargetPlatform) ? MAX_TOTAL_INFLUENCES : EXTRA_BONE_INFLUENCES;
 	const int32 MaxBoneInfluencesFromAsset = FGPUBaseSkinVertexFactory::GetBoneInfluenceLimitForAsset(BoneInfluenceLimit, TargetPlatform);
 
 	//Get the maximum allow bone influence, so we can cut lowest weight properly and get the same result has the sk build
diff --git a/Engine/Source/Runtime/Engine/Private/GPUSkinVertexFactory.cpp b/Engine/Source/Runtime/Engine/Private/GPUSkinVertexFactory.cpp
--- a/Engine/Source/Runtime/Engine/Private/GPUSkinVertexFactory.cpp
+++ b/Engine/Source/Runtime/Engine/Private/GPUSkinVertexFactory.cpp
@@ -19,6 +19,7 @@
 #include "RenderGraphResources.h"
 #include "ShaderPlatformCachedIniValue.h"
 #include "Engine/RendererSettings.h"
+#include "RHI.h"
 
 #if INTEL_ISPC
 #include "GPUSkinVertexFactory.ispc.generated.h"
@@ -448,15 +449,42 @@ int32 FGPUBaseSkinVertexFactory::GetMaxGPUSkinBones(const ITargetPlatform* Targe
 	}
 }
 
-bool FGPUBaseSkinVertexFactory::UseUnlimitedBoneInfluences(uint32 MaxBoneInfluences)
+bool FGPUBaseSkinVertexFactory::UseUnlimitedBoneInfluences(uint32 MaxBoneInfluences, const ITargetPlatform* TargetPlatform)
 {
-	const bool bUnlimitedBoneInfluence = (GCVarUnlimitedBoneInfluences!=0);
-	const uint32 UnlimitedBoneInfluencesThreshold = (uint32) GCVarUnlimitedBoneInfluencesThreshold;
-	return bUnlimitedBoneInfluence && MaxBoneInfluences > UnlimitedBoneInfluencesThreshold;
+	if (!GetUnlimitedBoneInfluences(TargetPlatform))
+	{
+		return false;
+	}
+
+	uint32 UnlimitedBoneInfluencesThreshold = (uint32)GCVarUnlimitedBoneInfluencesThreshold;
+
+#if ALLOW_OTHER_PLATFORM_CONFIG
+	if (TargetPlatform)
+	{
+		TSharedPtr<IConsoleVariable> VariablePtr = CVarUnlimitedBoneInfluencesThreshold->GetPlatformValueVariable(*TargetPlatform->IniPlatformName());
+		if (VariablePtr.IsValid())
+		{
+			UnlimitedBoneInfluencesThreshold = (uint32)VariablePtr->GetInt();
+		}
+	}
+#endif
+	
+	return MaxBoneInfluences > UnlimitedBoneInfluencesThreshold;
 }
 
-bool FGPUBaseSkinVertexFactory::GetUnlimitedBoneInfluences()
+bool FGPUBaseSkinVertexFactory::GetUnlimitedBoneInfluences(const ITargetPlatform* TargetPlatform)
 {
+#if ALLOW_OTHER_PLATFORM_CONFIG
+	if (TargetPlatform)
+	{
+		TSharedPtr<IConsoleVariable> VariablePtr = CVarUnlimitedBoneInfluences->GetPlatformValueVariable(*TargetPlatform->IniPlatformName());
+		if (VariablePtr.IsValid())
+		{
+			return VariablePtr->GetBool();
+		}
+	}
+#endif
+	
 	return (GCVarUnlimitedBoneInfluences!=0);
 }
 
@@ -543,7 +571,9 @@ TGlobalResource<FBoneBufferPool> FGPUBaseSkinVertexFactory::BoneBufferPool;
 template <GPUSkinBoneInfluenceType BoneInfluenceType>
 bool TGPUSkinVertexFactory<BoneInfluenceType>::ShouldCompilePermutation(const FVertexFactoryShaderPermutationParameters& Parameters)
 {
-	bool bUnlimitedBoneInfluences = (BoneInfluenceType == UnlimitedBoneInfluence && GCVarUnlimitedBoneInfluences);
+	static FShaderPlatformCachedIniValue<int32> PerPlatformCVar(TEXT("r.GPUSkin.UnlimitedBoneInfluences"));
+	const bool bUnlimitedBoneInfluences = BoneInfluenceType == UnlimitedBoneInfluence && PerPlatformCVar.Get(Parameters.Platform) != 0;
+
 	return ShouldWeCompileGPUSkinVFShaders(Parameters.Platform, Parameters.MaterialParameters.FeatureLevel) &&
 		  (((Parameters.MaterialParameters.bIsUsedWithSkeletalMesh || Parameters.MaterialParameters.bIsUsedWithMorphTargets) && (BoneInfluenceType != UnlimitedBoneInfluence || bUnlimitedBoneInfluences)) 
 			  || Parameters.MaterialParameters.bIsSpecialEngineMaterial);
diff --git a/Engine/Source/Runtime/Engine/Private/SkeletalMesh.cpp b/Engine/Source/Runtime/Engine/Private/SkeletalMesh.cpp
--- a/Engine/Source/Runtime/Engine/Private/SkeletalMesh.cpp
+++ b/Engine/Source/Runtime/Engine/Private/SkeletalMesh.cpp
@@ -4514,7 +4514,7 @@ FString USkeletalMesh::BuildDerivedDataKey(const ITargetPlatform* TargetPlatform
 	}
 
 	IMeshBuilderModule::GetForPlatform(TargetPlatform).AppendToDDCKey(KeySuffix, true);
-	const bool bUnlimitedBoneInfluences = FGPUBaseSkinVertexFactory::GetUnlimitedBoneInfluences();
+	const bool bUnlimitedBoneInfluences = FGPUBaseSkinVertexFactory::GetUnlimitedBoneInfluences(TargetPlatform);
 	KeySuffix += bUnlimitedBoneInfluences ? "1" : "0";
 
 	// Include the global default bone influences limit in case any LODs don't set an explicit limit (highly likely)
diff --git a/Engine/Source/Runtime/Engine/Public/GPUSkinVertexFactory.h b/Engine/Source/Runtime/Engine/Public/GPUSkinVertexFactory.h
--- a/Engine/Source/Runtime/Engine/Public/GPUSkinVertexFactory.h
+++ b/Engine/Source/Runtime/Engine/Public/GPUSkinVertexFactory.h
@@ -376,8 +376,8 @@ public:
 
 	static const uint32 GHardwareMaxGPUSkinBones = 65536;
 	
-	ENGINE_API static bool UseUnlimitedBoneInfluences(uint32 MaxBoneInfluences);
-	ENGINE_API static bool GetUnlimitedBoneInfluences();
+	ENGINE_API static bool UseUnlimitedBoneInfluences(uint32 MaxBoneInfluences, const ITargetPlatform* TargetPlatform = nullptr);
+	ENGINE_API static bool GetUnlimitedBoneInfluences(const ITargetPlatform* TargetPlatform = nullptr);
 
 	/*
 	 * Returns the maximum number of bone influences that should be used for a skeletal mesh, given
