SkeletalMesh GetAssetRegistryTags: When compiling, return the current tags from the AssetRegistry rather than returning the empty set, so that queries for those tags during the compilation period will continue to work.
#rnx

diff --git a/Engine/Source/Runtime/Engine/Private/SkeletalMesh.cpp b/Engine/Source/Runtime/Engine/Private/SkeletalMesh.cpp
--- a/Engine/Source/Runtime/Engine/Private/SkeletalMesh.cpp
+++ b/Engine/Source/Runtime/Engine/Private/SkeletalMesh.cpp
@@ -69,6 +69,8 @@
 #include "Algo/MaxElement.h"
 
 #if WITH_EDITOR
+#include "AssetRegistry/AssetData.h"
+#include "AssetRegistry/IAssetRegistry.h"
 #include "Async/ParallelFor.h"
 #include "Rendering/SkeletalMeshModel.h"
 #include "Rendering/SkeletalMeshLODImporterData.h"
@@ -3481,6 +3483,15 @@ void USkeletalMesh::GetAssetRegistryTags(TArray<FAssetRegistryTag>& OutTags) con
 	// Avoid accessing properties being compiled, this function will get called again after compilation is finished.
 	if (IsCompiling())
 	{
+		IAssetRegistry* AssetRegistry = IAssetRegistry::Get();
+		if (AssetRegistry)
+		{
+			FAssetData AssetData = AssetRegistry->GetAssetByObjectPath(FSoftObjectPath(this), true /* bIncludeOnlyOnDiskAssets */);
+			AssetData.EnumerateTags([&OutTags](const TPair<FName, FAssetTagValueRef>& Pair)
+				{
+					OutTags.Add(FAssetRegistryTag(Pair.Key, Pair.Value.GetStorageString(), FAssetRegistryTag::TT_Alphabetical));
+				});
+		}
 		return;
 	}
 #endif
