StaticMesh LOGGroup can now specify absolute reduction criterion.
New keywords supported by lodgroup in ini files:
- TerminationCriterion=XXXX[EStaticMeshReductionTerimationCriterion::(Triangles, Vertices or Any)] Represent the termination criterion.
- LODPercentVertices=XXXX where XXXX [0 - 100] Same as LODPercentTriangles if the termination criterion is Vertices or Any.
- LOD[0, 1, 2, 3 ... 7]MaxNumOfTriangles=XXXX[2 +] Represent the maximum absolute number of triangles for the specified LOD if the termination criterion is Triangles or Any
- LOD[0, 1, 2, 3 ... 7]MaxNumOfVertices=XXXX[4 +] Represent the maximum absolute number of vertices for the specified LOD if the termination criterion is Vertices or Any
- BasePercentVerticesMult=XXXX[0 - 100] Same as BasePercentTrianglesMult if the termination criterion is Vertices or Any.
- LODPercentVerticesMult=XXXX[0 - 100] Same as LODPercentTrianglesMult if the termination criterion is Vertices or Any.

diff --git a/Engine/Source/Runtime/Engine/Private/StaticMesh.cpp b/Engine/Source/Runtime/Engine/Private/StaticMesh.cpp
--- a/Engine/Source/Runtime/Engine/Private/StaticMesh.cpp
+++ b/Engine/Source/Runtime/Engine/Private/StaticMesh.cpp
@@ -2174,12 +2174,36 @@ void FStaticMeshLODSettings::Initialize(const ITargetPlatform* TargetPlatform)
 	for (TMap<FName,FStaticMeshLODGroup>::TIterator It(Groups); It; ++It)
 	{
 		FStaticMeshLODGroup& Group = It.Value();
-		float PercentTrianglesPerLOD = Group.DefaultSettings[1].PercentTriangles;
+		EStaticMeshReductionTerimationCriterion LODTerminationCriterion = Group.DefaultSettings[0].TerminationCriterion;
+		float PercentTrianglesPerLODRatio = Group.DefaultSettings[1].PercentTriangles;
+		float PercentVerticesPerLODRatio = Group.DefaultSettings[1].PercentVertices;
 		for (int32 LODIndex = 1; LODIndex < MAX_STATIC_MESH_LODS; ++LODIndex)
 		{
-			float PercentTriangles = Group.DefaultSettings[LODIndex-1].PercentTriangles;
+			//Set the termination criteria
+			Group.DefaultSettings[LODIndex].TerminationCriterion = LODTerminationCriterion;
+			float PercentTriangles = Group.DefaultSettings[LODIndex - 1].PercentTriangles;
+			float PercentVertices = Group.DefaultSettings[LODIndex - 1].PercentVertices;
+
+			//Clamp Absolute value so every LOD is equal or less the previous LOD
+			uint32 MaxNumOfTriangles = FMath::Clamp<uint32>(
+				Group.DefaultSettings[LODIndex].MaxNumOfTriangles
+				, 2
+				, Group.DefaultSettings[LODIndex - 1].MaxNumOfTriangles);
+			uint32 MaxNumOfVerts = FMath::Clamp<uint32>(
+				Group.DefaultSettings[LODIndex].MaxNumOfVerts
+				, 4
+				, Group.DefaultSettings[LODIndex - 1].MaxNumOfVerts);
+
+			//Copy the previous LOD
 			Group.DefaultSettings[LODIndex] = Group.DefaultSettings[LODIndex - 1];
-			Group.DefaultSettings[LODIndex].PercentTriangles = PercentTriangles * PercentTrianglesPerLOD;
+
+			//Reduce the data from the previous LOD using the ratios
+			Group.DefaultSettings[LODIndex].PercentTriangles = PercentTriangles * PercentTrianglesPerLODRatio;
+			Group.DefaultSettings[LODIndex].PercentVertices = PercentVertices * PercentVerticesPerLODRatio;
+
+			//Put back the absolute criterion after the LOD copy
+			Group.DefaultSettings[LODIndex].MaxNumOfTriangles = MaxNumOfTriangles;
+			Group.DefaultSettings[LODIndex].MaxNumOfVerts = MaxNumOfVerts;
 		}
 	}
 }
@@ -2226,6 +2250,11 @@ void FStaticMeshLODSettings::ReadEntry(FStaticMeshLODGroup& Group, FString Entry
 		Group.DefaultLightMapResolution = (Group.DefaultLightMapResolution + 3) & (~3);
 	}
 
+	FString TerminationCriterion = StaticEnum<EStaticMeshReductionTerimationCriterion>()->GetValueAsString(EStaticMeshReductionTerimationCriterion::Triangles);
+	if (FParse::Value(*Entry, TEXT("TerminationCriterion="), TerminationCriterion))
+	{
+		Group.DefaultSettings[0].TerminationCriterion = static_cast<EStaticMeshReductionTerimationCriterion>(StaticEnum<EStaticMeshReductionTerimationCriterion>()->GetValueByNameString(TerminationCriterion));
+	}
 	float BasePercentTriangles = 100.0f;
 	if (FParse::Value(*Entry, TEXT("BasePercentTriangles="), BasePercentTriangles))
 	{
@@ -2233,6 +2262,13 @@ void FStaticMeshLODSettings::ReadEntry(FStaticMeshLODGroup& Group, FString Entry
 		Group.DefaultSettings[0].PercentTriangles = BasePercentTriangles * 0.01f;
 	}
 
+	float BasePercentVertices = 100.0f;
+	if (FParse::Value(*Entry, TEXT("BasePercentVertices="), BasePercentVertices))
+	{
+		BasePercentVertices = FMath::Clamp<float>(BasePercentVertices, 0.0f, 100.0f);
+		Group.DefaultSettings[0].PercentVertices = BasePercentVertices * 0.01f;
+	}
+
 	float LODPercentTriangles = 100.0f;
 	if (FParse::Value(*Entry, TEXT("LODPercentTriangles="), LODPercentTriangles))
 	{
@@ -2240,6 +2276,30 @@ void FStaticMeshLODSettings::ReadEntry(FStaticMeshLODGroup& Group, FString Entry
 		Group.DefaultSettings[1].PercentTriangles = LODPercentTriangles * 0.01f;
 	}
 
+	float LODPercentVertices = 100.0f;
+	if (FParse::Value(*Entry, TEXT("LODPercentVertices="), LODPercentVertices))
+	{
+		LODPercentVertices = FMath::Clamp<float>(LODPercentVertices, 0.0f, 100.0f);
+		Group.DefaultSettings[1].PercentVertices = LODPercentVertices * 0.01f;
+	}
+
+	for (int32 LodIndex = 0; LodIndex < Group.DefaultNumLODs; ++LodIndex)
+	{
+		uint32 LODMaxNumOfTriangles = MAX_uint32;
+		FString KeySearch = FString::Printf(TEXT("LOD%dMaxNumOfTriangles="), LodIndex);
+		if (FParse::Value(*Entry, *KeySearch, LODMaxNumOfTriangles))
+		{
+			Group.DefaultSettings[LodIndex].MaxNumOfTriangles = LODMaxNumOfTriangles;
+		}
+
+		uint32 LODMaxNumOfVerts = MAX_uint32;
+		KeySearch = FString::Printf(TEXT("LOD%dMaxNumOfVertices="), LodIndex);
+		if (FParse::Value(*Entry, *KeySearch, LODMaxNumOfVerts))
+		{
+			Group.DefaultSettings[LodIndex].MaxNumOfVerts = LODMaxNumOfVerts;
+		}
+	}
+
 	if (FParse::Value(*Entry, TEXT("MaxDeviation="), Settings.MaxDeviation))
 	{
 		Settings.MaxDeviation = FMath::Clamp<float>(Settings.MaxDeviation, 0.0f, 1000.0f);
@@ -2282,6 +2342,13 @@ void FStaticMeshLODSettings::ReadEntry(FStaticMeshLODGroup& Group, FString Entry
 		Group.BasePercentTrianglesMult = BasePercentTrianglesMult * 0.01f;
 	}
 
+	float BasePercentVerticesMult = 100.0f;
+	if (FParse::Value(*Entry, TEXT("BasePercentVerticesMult="), BasePercentVerticesMult))
+	{
+		BasePercentVerticesMult = FMath::Clamp<float>(BasePercentVerticesMult, 0.0f, 100.0f);
+		Group.BasePercentVerticesMult = BasePercentVerticesMult * 0.01f;
+	}
+
 	float LODPercentTrianglesMult = 100.0f;
 	if (FParse::Value(*Entry, TEXT("LODPercentTrianglesMult="), LODPercentTrianglesMult))
 	{
@@ -2289,6 +2356,13 @@ void FStaticMeshLODSettings::ReadEntry(FStaticMeshLODGroup& Group, FString Entry
 		Bias.PercentTriangles = LODPercentTrianglesMult * 0.01f;
 	}
 
+	float LODPercentVerticesMult = 100.0f;
+	if (FParse::Value(*Entry, TEXT("LODPercentVerticesMult="), LODPercentVerticesMult))
+	{
+		LODPercentVerticesMult = FMath::Clamp<float>(LODPercentVerticesMult, 0.0f, 100.0f);
+		Bias.PercentVertices = LODPercentVerticesMult * 0.01f;
+	}
+
 	if (FParse::Value(*Entry, TEXT("MaxDeviationBias="), Bias.MaxDeviation))
 	{
 		Bias.MaxDeviation = FMath::Clamp<float>(Bias.MaxDeviation, -1000.0f, 1000.0f);
@@ -2351,6 +2425,9 @@ FMeshReductionSettings FStaticMeshLODGroup::GetSettings(const FMeshReductionSett
 	float PercentTrianglesMult = (LODIndex == 0) ? BasePercentTrianglesMult : SettingsBias.PercentTriangles;
 	FinalSettings.PercentTriangles = FMath::Clamp(InSettings.PercentTriangles * PercentTrianglesMult, 0.0f, 1.0f);
 
+	float PercentVerticesMult = (LODIndex == 0) ? BasePercentVerticesMult : SettingsBias.PercentVertices;
+	FinalSettings.PercentVertices = FMath::Clamp(InSettings.PercentVertices * PercentVerticesMult, 0.0f, 1.0f);
+
 	// Bias the remaining settings.
 	FinalSettings.MaxDeviation = FMath::Max(InSettings.MaxDeviation + SettingsBias.MaxDeviation, 0.0f);
 	FinalSettings.PixelError = FMath::Max(InSettings.PixelError + SettingsBias.PixelError, 1.0f);
diff --git a/Engine/Source/Runtime/Engine/Public/StaticMeshResources.h b/Engine/Source/Runtime/Engine/Public/StaticMeshResources.h
--- a/Engine/Source/Runtime/Engine/Public/StaticMeshResources.h
+++ b/Engine/Source/Runtime/Engine/Public/StaticMeshResources.h
@@ -62,11 +62,13 @@ public:
 		, DefaultMaxNumOptionalLODs(0)
 		, DefaultLightMapResolution(64)
 		, BasePercentTrianglesMult(1.0f)
+		, BasePercentVerticesMult(1.0f)
 		, bSupportLODStreaming(false)
 		, DisplayName( NSLOCTEXT( "UnrealEd", "None", "None" ) )
 	{
 		FMemory::Memzero(SettingsBias);
 		SettingsBias.PercentTriangles = 1.0f;
+		SettingsBias.PercentVertices = 1.0f;
 	}
 
 	/** Returns the default number of LODs to build. */
@@ -122,6 +124,8 @@ private:
 	int32 DefaultLightMapResolution;
 	/** An additional reduction of base meshes in this group. */
 	float BasePercentTrianglesMult;
+	/** An additional reduction of base meshes in this group. */
+	float BasePercentVerticesMult;
 	/** Whether static meshes in this LOD group can be streamed. */
 	bool bSupportLODStreaming;
 	/** Display name. */
