Remove unneeded ensure. FAudioChunkHandle() no longer needs a valid FSoundWaveData ptr to release retained audio from the stream cache.

diff --git a/Engine/Source/Runtime/Engine/Private/SoundWave.cpp b/Engine/Source/Runtime/Engine/Private/SoundWave.cpp
--- a/Engine/Source/Runtime/Engine/Private/SoundWave.cpp
+++ b/Engine/Source/Runtime/Engine/Private/SoundWave.cpp
@@ -141,9 +141,6 @@ namespace SoundWaveCookStats
 
 FSoundWaveData::~FSoundWaveData()
 {
-	// at this point, first chunk data should have been cleared by the last reference holder
-	// (see USoundWave::BeginDestroy and ~FSoundWaveProxy implmentations)
-	ensure(FirstChunk.GetData() == nullptr);
 }
 
 void FSoundWaveData::InitializeDataFromSoundWave(USoundWave& InWave)
