Fixed instance edited spline properties from being reset to class defaults during replay scrubbing

#Splines
#rnx

diff --git a/Engine/Source/Runtime/Engine/Classes/Components/SplineComponent.h b/Engine/Source/Runtime/Engine/Classes/Components/SplineComponent.h
--- a/Engine/Source/Runtime/Engine/Classes/Components/SplineComponent.h
+++ b/Engine/Source/Runtime/Engine/Classes/Components/SplineComponent.h
@@ -195,7 +195,7 @@ class ENGINE_API USplineComponent : public UPrimitiveComponent
 {
 	GENERATED_UCLASS_BODY()
 
-	UPROPERTY(EditAnywhere, Category=Points)
+	UPROPERTY(EditAnywhere, Replicated, Category=Points)
 	FSplineCurves SplineCurves;
 
 	/** Deprecated - please use GetSplinePointsPosition() to fetch this FInterpCurve */
@@ -217,7 +217,7 @@ class ENGINE_API USplineComponent : public UPrimitiveComponent
 	bool bAllowSplineEditingPerInstance_DEPRECATED;
 
 	/** Number of steps per spline segment to place in the reparameterization table */
-	UPROPERTY(EditAnywhere, AdvancedDisplay, Category = Spline, meta=(ClampMin=4, UIMin=4, ClampMax=100, UIMax=100))
+	UPROPERTY(EditAnywhere, Replicated, AdvancedDisplay, Category = Spline, meta=(ClampMin=4, UIMin=4, ClampMax=100, UIMax=100))
 	int32 ReparamStepsPerSegment;
 
 	/** Specifies the duration of the spline in seconds */
@@ -225,11 +225,11 @@ class ENGINE_API USplineComponent : public UPrimitiveComponent
 	float Duration;
 
 	/** Whether the endpoints of the spline are considered stationary when traversing the spline at non-constant velocity.  Essentially this sets the endpoints' tangents to zero vectors. */
-	UPROPERTY(EditAnywhere, BlueprintReadWrite, AdvancedDisplay, Category = Spline, meta=(EditCondition="!bClosedLoop"))
+	UPROPERTY(EditAnywhere, Replicated, BlueprintReadWrite, AdvancedDisplay, Category = Spline, meta=(EditCondition="!bClosedLoop"))
 	bool bStationaryEndpoints;
 
 	/** Whether the spline has been edited from its default by the spline component visualizer */
-	UPROPERTY(EditAnywhere, Category = Spline, meta=(DisplayName="Override Construction Script"))
+	UPROPERTY(EditAnywhere, Replicated, Category = Spline, meta=(DisplayName="Override Construction Script"))
 	bool bSplineHasBeenEdited;
 
 	UPROPERTY()
@@ -240,7 +240,7 @@ class ENGINE_API USplineComponent : public UPrimitiveComponent
 	 * Whether the spline points should be passed to the User Construction Script so they can be further manipulated by it.
 	 * If false, they will not be visible to it, and it will not be able to influence the per-instance positions set in the editor.
 	 */
-	UPROPERTY(EditAnywhere, BlueprintReadWrite, Category = Spline)
+	UPROPERTY(EditAnywhere, Replicated, BlueprintReadWrite, Category = Spline)
 	bool bInputSplinePointsToConstructionScript;
 
 	/** If true, the spline will be rendered if the Splines showflag is set. */
@@ -252,18 +252,18 @@ private:
 	 * Whether the spline is to be considered as a closed loop.
 	 * Use SetClosedLoop() to set this property, and IsClosedLoop() to read it.
 	 */
-	UPROPERTY(EditAnywhere, Category = Spline)
+	UPROPERTY(EditAnywhere, Replicated, Category = Spline)
 	bool bClosedLoop;
 
-	UPROPERTY(EditAnywhere, Category = Spline, meta=(InlineEditConditionToggle=true))
+	UPROPERTY(EditAnywhere, Replicated, Category = Spline, meta=(InlineEditConditionToggle=true))
 	bool bLoopPositionOverride;
 
-	UPROPERTY(EditAnywhere, Category = Spline, meta=(EditCondition="bLoopPositionOverride"))
+	UPROPERTY(EditAnywhere, Replicated, Category = Spline, meta=(EditCondition="bLoopPositionOverride"))
 	float LoopPosition;
 
 public:
 	/** Default up vector in local space to be used when calculating transforms along the spline */
-	UPROPERTY(EditAnywhere, BlueprintReadWrite, Category = Spline)
+	UPROPERTY(EditAnywhere, Replicated, BlueprintReadWrite, Category = Spline)
 	FVector DefaultUpVector;
 
 #if WITH_EDITORONLY_DATA
@@ -297,6 +297,7 @@ public:
 #endif
 
 	//~ Begin UObject Interface
+	virtual void GetLifetimeReplicatedProps(TArray<FLifetimeProperty>& OutLifetimeProps) const override;
 	virtual void Serialize(FArchive& Ar) override;
 	virtual void PostLoad() override;
 #if WITH_EDITOR
diff --git a/Engine/Source/Runtime/Engine/Private/Components/SplineComponent.cpp b/Engine/Source/Runtime/Engine/Private/Components/SplineComponent.cpp
--- a/Engine/Source/Runtime/Engine/Private/Components/SplineComponent.cpp
+++ b/Engine/Source/Runtime/Engine/Private/Components/SplineComponent.cpp
@@ -10,6 +10,8 @@
 #include "PrimitiveSceneProxy.h"
 #include "SceneManagement.h"
 #include "UnrealEngine.h"
+#include "Net/UnrealNetwork.h"
+#include "Net/Core/PushModel/PushModel.h"
 #include "Styling/StyleColors.h"
 
 #include UE_INLINE_GENERATED_CPP_BY_NAME(SplineComponent)
@@ -79,6 +81,14 @@ USplineComponent::USplineComponent(const FObjectInitializer& ObjectInitializer)
 	SplineReparamTable_DEPRECATED = SplineCurves.ReparamTable;
 }
 
+void USplineComponent::GetLifetimeReplicatedProps(TArray<FLifetimeProperty>& OutLifetimeProps) const
+{
+	Super::GetLifetimeReplicatedProps(OutLifetimeProps);
+
+	// This is a workaround for UE-129807 so that scrubbing a replay doesn't cause instance edited properties to be reset to class defaults
+	// If you encounter this issue, reset the relevant replicated properties of this class with COND_ReplayOnly
+	DISABLE_ALL_CLASS_REPLICATED_PROPERTIES(USplineComponent, EFieldIteratorFlags::ExcludeSuper);
+}
 
 EInterpCurveMode ConvertSplinePointTypeToInterpCurveMode(ESplinePointType::Type SplinePointType)
 {
@@ -217,6 +227,15 @@ void FSplineCurves::UpdateSpline(bool bClosedLoop, bool bStationaryEndpoints, in
 void USplineComponent::UpdateSpline()
 {
 	SplineCurves.UpdateSpline(bClosedLoop, bStationaryEndpoints, ReparamStepsPerSegment, bLoopPositionOverride, LoopPosition, GetComponentTransform().GetScale3D());
+	MARK_PROPERTY_DIRTY_FROM_NAME(USplineComponent, SplineCurves, this);
+	MARK_PROPERTY_DIRTY_FROM_NAME(USplineComponent, bClosedLoop, this);
+	MARK_PROPERTY_DIRTY_FROM_NAME(USplineComponent, bStationaryEndpoints, this);
+	MARK_PROPERTY_DIRTY_FROM_NAME(USplineComponent, ReparamStepsPerSegment, this);
+	MARK_PROPERTY_DIRTY_FROM_NAME(USplineComponent, bLoopPositionOverride, this);
+	MARK_PROPERTY_DIRTY_FROM_NAME(USplineComponent, LoopPosition, this);
+	MARK_PROPERTY_DIRTY_FROM_NAME(USplineComponent, DefaultUpVector, this);
+	MARK_PROPERTY_DIRTY_FROM_NAME(USplineComponent, bSplineHasBeenEdited, this);
+	MARK_PROPERTY_DIRTY_FROM_NAME(USplineComponent, bInputSplinePointsToConstructionScript, this);
 
 #if !UE_BUILD_SHIPPING
 	if (bDrawDebug)
