Removed all code under !UE_AUDIO_THREAD_AS_PIPE. Just a cleanup.

diff --git a/Engine/Source/Runtime/Core/Private/Async/TaskGraph.cpp b/Engine/Source/Runtime/Core/Private/Async/TaskGraph.cpp
--- a/Engine/Source/Runtime/Core/Private/Async/TaskGraph.cpp
+++ b/Engine/Source/Runtime/Core/Private/Async/TaskGraph.cpp
@@ -1438,8 +1438,6 @@ PRAGMA_ENABLE_DEPRECATION_WARNINGS
 #endif
 	}
 
-#if UE_AUDIO_THREAD_AS_PIPE
-
 	// AudioThread was removed and replaced by a pipe. While all engine code was converted where's a tiny possibilty that 
 	// an external code is using it directly. Redirect all audio tasks to the pipe
 	static void RedirectAudioTasksToPipe(FBaseGraphTask* Task)
@@ -1461,8 +1459,6 @@ PRAGMA_ENABLE_DEPRECATION_WARNINGS
 		}
 	}
 
-#endif
-
 	// API inherited from FTaskGraphInterface
 
 	/** 
@@ -1475,12 +1471,8 @@ PRAGMA_ENABLE_DEPRECATION_WARNINGS
 	{
 		TASKGRAPH_SCOPE_CYCLE_COUNTER(2, STAT_TaskGraph_QueueTask);
 
-#if UE_STATS_THREAD_AS_PIPE
 		RedirectStatsTasksToPipe(Task);
-#endif
-#if UE_AUDIO_THREAD_AS_PIPE
 		RedirectAudioTasksToPipe(Task);
-#endif
 
 		if (ENamedThreads::GetThreadIndex(ThreadToExecuteOn) == ENamedThreads::AnyThread)
 		{
@@ -2024,12 +2016,8 @@ public:
 private:
 	void QueueTask(class FBaseGraphTask* Task, bool bWakeUpWorker, ENamedThreads::Type InThreadToExecuteOn, ENamedThreads::Type InCurrentThreadIfKnown) override
 	{
-#if UE_STATS_THREAD_AS_PIPE
 		FTaskGraphImplementation::RedirectStatsTasksToPipe(Task);
-#endif
-#if UE_AUDIO_THREAD_AS_PIPE
 		FTaskGraphImplementation::RedirectAudioTasksToPipe(Task);
-#endif
 
 		if (ENamedThreads::GetThreadIndex(InThreadToExecuteOn) == ENamedThreads::AnyThread)
 		{
diff --git a/Engine/Source/Runtime/Core/Private/HAL/ThreadingBase.cpp b/Engine/Source/Runtime/Core/Private/HAL/ThreadingBase.cpp
--- a/Engine/Source/Runtime/Core/Private/HAL/ThreadingBase.cpp
+++ b/Engine/Source/Runtime/Core/Private/HAL/ThreadingBase.cpp
@@ -228,11 +228,6 @@ CORE_API std::atomic<bool> GIsAudioThreadRunning{ false };
 
 CORE_API std::atomic<bool> GIsAudioThreadSuspended{ false };
 
-#if !UE_AUDIO_THREAD_AS_PIPE
-CORE_API FRunnableThread* GAudioThread = nullptr;
-#endif
-
-
 CORE_API bool IsAudioThreadRunning()
 {
 	return GIsAudioThreadRunning.load(std::memory_order_acquire) && !GIsAudioThreadSuspended.load(std::memory_order_acquire);
@@ -240,31 +235,7 @@ CORE_API bool IsAudioThreadRunning()
 
 CORE_API bool IsInAudioThread()
 {
-#if UE_AUDIO_THREAD_AS_PIPE
-
 	return GIsAudioThreadRunning.load(std::memory_order_acquire) && !GIsAudioThreadSuspended.load(std::memory_order_acquire) ? GAudioPipe.IsInContext() : IsInGameThread();
-
-#else
-
-PRAGMA_DISABLE_DEPRECATION_WARNINGS
-	// True if this is the audio thread or if there is no audio thread, then if it is the game thread
-	bool newValue = (nullptr == GAudioThread || GIsAudioThreadSuspended.load(std::memory_order_relaxed))
-		? FTaskTagScope::IsCurrentTag(ETaskTag::EGameThread)
-		: FTaskTagScope::IsCurrentTag(ETaskTag::EAudioThread);
-#if !UE_BUILD_SHIPPING && !UE_BUILD_TEST
-	if (!LowLevelTasks::FSchedulerTls::IsBusyWaiting() &&
-		!CoroTask_Detail::FCoroLocalState::IsCoroLaunchedTask() &&
-		!UE::Tasks::Private::IsThreadRetractingTask())
-	{
-		bool oldValue = FPlatformTLS::GetCurrentThreadId() == ((nullptr == GAudioThread || GIsAudioThreadSuspended.load(std::memory_order_relaxed)) ? GGameThreadId : GAudioThread->GetThreadID());
-		ensureMsgf(oldValue == newValue, TEXT("oldValue(%i) newValue(%i) If this check fails make sure that there is a FTaskTagScope(ETaskTag::EAudioThread) as deep as possible on the current callstack, you can see the current value in ActiveNamedThreads(%x)"), oldValue, newValue, FTaskTagScope::GetCurrentTag());
-		newValue = oldValue;
-	}
-#endif
-	return newValue;
-PRAGMA_ENABLE_DEPRECATION_WARNINGS
-
-#endif
 }
 
 CORE_API TAtomic<int32> GIsRenderingThreadSuspended(0);
diff --git a/Engine/Source/Runtime/Core/Public/Async/ParallelFor.h b/Engine/Source/Runtime/Core/Public/Async/ParallelFor.h
--- a/Engine/Source/Runtime/Core/Public/Async/ParallelFor.h
+++ b/Engine/Source/Runtime/Core/Public/Async/ParallelFor.h
@@ -168,9 +168,6 @@ namespace ParallelForImpl
 			ETaskTag::EStaticInit | 
 			ETaskTag::EGameThread | 
 			ETaskTag::ESlateThread | 
-#if !UE_AUDIO_THREAD_AS_PIPE
-			ETaskTag::EAudioThread | 
-#endif
 			ETaskTag::ERenderingThread | 
 			ETaskTag::ERhiThread;
 
diff --git a/Engine/Source/Runtime/Core/Public/Async/TaskGraphInterfaces.h b/Engine/Source/Runtime/Core/Public/Async/TaskGraphInterfaces.h
--- a/Engine/Source/Runtime/Core/Public/Async/TaskGraphInterfaces.h
+++ b/Engine/Source/Runtime/Core/Public/Async/TaskGraphInterfaces.h
@@ -63,11 +63,7 @@ namespace ENamedThreads
 #endif
 #endif
 		RHIThread,
-#if UE_AUDIO_THREAD_AS_PIPE
 		AudioThread UE_DEPRECATED(5.0, "`AudioThread` has been removed. Please use `FAudioThread` API"),
-#else
-		AudioThread,
-#endif
 		GameThread,
 		// The render thread is sometimes the game thread and is sometimes the actual rendering thread
 		ActualRenderingThread = GameThread + 1,
diff --git a/Engine/Source/Runtime/Core/Public/CoreGlobals.h b/Engine/Source/Runtime/Core/Public/CoreGlobals.h
--- a/Engine/Source/Runtime/Core/Public/CoreGlobals.h
+++ b/Engine/Source/Runtime/Core/Public/CoreGlobals.h
@@ -580,11 +580,7 @@ enum class ETaskTag : int32
 	EStaticInit					= 1 << 0,
 	EGameThread					= 1 << 1,
 	ESlateThread				= 1 << 2,
-#if UE_AUDIO_THREAD_AS_PIPE	
 	EAudioThread UE_DEPRECATED(5.0, "AudioThread was removed and ETaskTag::EAudioThread is not used anymore. Please remove it.") = 1 << 3,
-#else
-	EAudioThread = 1 << 3,
-#endif
 	ERenderingThread			= 1 << 4,
 	ERhiThread					= 1 << 5,
 	EAsyncLoadingThread			= 1 << 6,
@@ -686,14 +682,6 @@ extern CORE_API bool IsAudioThreadRunning();
 /** @return True if called from the audio thread, and not merely a thread calling audio functions. */
 extern CORE_API bool IsInAudioThread();
 
-#if !UE_AUDIO_THREAD_AS_PIPE
-
-/** Thread used for audio */
-UE_DEPRECATED(4.26, "Please use `IsAudioThreadRunning()` or `IsInAudioThread()`")
-extern CORE_API FRunnableThread* GAudioThread;
-
-#endif
-
 /** @return True if called from the slate thread, and not merely a thread calling slate functions. */
 extern CORE_API bool IsInSlateThread();
 
diff --git a/Engine/Source/Runtime/Core/Public/Experimental/Coroutine/CoroParallelFor.h b/Engine/Source/Runtime/Core/Public/Experimental/Coroutine/CoroParallelFor.h
--- a/Engine/Source/Runtime/Core/Public/Experimental/Coroutine/CoroParallelFor.h
+++ b/Engine/Source/Runtime/Core/Public/Experimental/Coroutine/CoroParallelFor.h
@@ -78,9 +78,6 @@ inline void CoroParallelFor(const TCHAR* DebugName, int32 Num, BodyType Body, EP
 		ETaskTag::EStaticInit | 
 		ETaskTag::EGameThread | 
 		ETaskTag::ESlateThread | 
-#if !UE_AUDIO_THREAD_AS_PIPE
-		ETaskTag::EAudioThread | 
-#endif
 		ETaskTag::ERenderingThread | 
 		ETaskTag::ERhiThread;
 
diff --git a/Engine/Source/Runtime/Core/Public/Misc/CoreDefines.h b/Engine/Source/Runtime/Core/Public/Misc/CoreDefines.h
--- a/Engine/Source/Runtime/Core/Public/Misc/CoreDefines.h
+++ b/Engine/Source/Runtime/Core/Public/Misc/CoreDefines.h
@@ -66,10 +66,6 @@ Can be overridden by setting to 1 or 0 in the project's .Target.cs files.
 	#define NO_CVARS 0
 #endif
 
-#if !defined(UE_AUDIO_THREAD_AS_PIPE)
-	#define UE_AUDIO_THREAD_AS_PIPE 1
-#endif
-
 #if !defined(UE_STATS_THREAD_AS_PIPE)
 	#define UE_STATS_THREAD_AS_PIPE 1
 #endif
diff --git a/Engine/Source/Runtime/Engine/Private/AudioThread.cpp b/Engine/Source/Runtime/Engine/Private/AudioThread.cpp
--- a/Engine/Source/Runtime/Engine/Private/AudioThread.cpp
+++ b/Engine/Source/Runtime/Engine/Private/AudioThread.cpp
@@ -123,8 +123,6 @@ TStatId FAudioThread::CurrentAudioThreadStatId;
 TStatId FAudioThread::LongestAudioThreadStatId;
 double FAudioThread::LongestAudioThreadTimeMsec = 0.0;
 
-#if UE_AUDIO_THREAD_AS_PIPE
-
 TUniquePtr<UE::Tasks::FTaskEvent> FAudioThread::ResumeEvent;
 int32 FAudioThread::SuspendCount{ 0 };
 
@@ -184,132 +182,6 @@ void FAudioThread::ResumeAudioThread()
 	ResumeEvent = MakeUnique<UE::Tasks::FTaskEvent>(UE_SOURCE_LOCATION);
 }
 
-#else // UE_AUDIO_THREAD_AS_PIPE
-
-FRunnable* FAudioThread::AudioThreadRunnable = nullptr;
-
-/** The audio thread main loop */
-void AudioThreadMain( FEvent* TaskGraphBoundSyncEvent )
-{
-	FTaskGraphInterface::Get().AttachToThread(ENamedThreads::AudioThread);
-	FPlatformMisc::MemoryBarrier();
-
-	// Inform main thread that the audio thread has been attached to the taskgraph and is ready to receive tasks
-	if( TaskGraphBoundSyncEvent != nullptr )
-	{
-		TaskGraphBoundSyncEvent->Trigger();
-	}
-
-	FTaskGraphInterface::Get().ProcessThreadUntilRequestReturn(ENamedThreads::AudioThread);
-	FPlatformMisc::MemoryBarrier();
-}
-
-FAudioThread::FAudioThread()
-{
-	TaskGraphBoundSyncEvent	= FPlatformProcess::GetSynchEventFromPool(true);
-
-	FCoreUObjectDelegates::GetPreGarbageCollectDelegate().AddRaw(this, &FAudioThread::OnPreGarbageCollect);
-	FCoreUObjectDelegates::GetPostGarbageCollect().AddRaw(this, &FAudioThread::OnPostGarbageCollect);
-
-	FCoreUObjectDelegates::PreGarbageCollectConditionalBeginDestroy.AddRaw(this, &FAudioThread::OnPreGarbageCollect);
-	FCoreUObjectDelegates::PostGarbageCollectConditionalBeginDestroy.AddRaw(this, &FAudioThread::OnPostGarbageCollect);
-}
-
-FAudioThread::~FAudioThread()
-{
-	FCoreUObjectDelegates::GetPreGarbageCollectDelegate().RemoveAll(this);
-	FCoreUObjectDelegates::GetPostGarbageCollect().RemoveAll(this);
-
-	FCoreUObjectDelegates::PreGarbageCollectConditionalBeginDestroy.RemoveAll(this);
-	FCoreUObjectDelegates::PostGarbageCollectConditionalBeginDestroy.RemoveAll(this);
-
-	FPlatformProcess::ReturnSynchEventToPool(TaskGraphBoundSyncEvent);
-	TaskGraphBoundSyncEvent = nullptr;
-}
-
-static int32 AudioThreadSuspendCount = 0;
-
-void FAudioThread::SuspendAudioThread()
-{
-	TRACE_CPUPROFILER_EVENT_SCOPE(SuspendAudioThread);
-	check(FPlatformTLS::GetCurrentThreadId() == GGameThreadId);
-	check(!GIsAudioThreadSuspended.load() || GCVarSuspendAudioThread != 0);
-
-	if (IsAudioThreadRunning())
-	{
-		// Make GC wait on the audio thread finishing processing
-		FAudioCommandFence AudioFence;
-		AudioFence.BeginFence();
-		AudioFence.Wait();
-
-		GIsAudioThreadSuspended = true;
-		FPlatformMisc::MemoryBarrier();
-	}
-	check(!IsAudioThreadRunning());
-}
-
-void FAudioThread::ResumeAudioThread()
-{
-PRAGMA_DISABLE_DEPRECATION_WARNINGS
-
-	check(FPlatformTLS::GetCurrentThreadId() == GGameThreadId);
-	if(GIsAudioThreadSuspended.load() && GCVarSuspendAudioThread == 0)
-	{
-		GIsAudioThreadSuspended = false;
-		FPlatformMisc::MemoryBarrier();
-	}
-	ProcessAllCommands();
-
-PRAGMA_ENABLE_DEPRECATION_WARNINGS
-}
-
-void FAudioThread::OnPreGarbageCollect()
-{
-	AudioThreadSuspendCount++;
-	if (AudioThreadSuspendCount == 1)
-	{
-		SuspendAudioThread();
-	}
-}
-
-void FAudioThread::OnPostGarbageCollect()
-{
-	AudioThreadSuspendCount--;
-	if (AudioThreadSuspendCount == 0)
-	{
-		ResumeAudioThread();
-	}
-}
-
-#if !UE_AUDIO_THREAD_AS_PIPE
-
-bool FAudioThread::Init()
-{ 
-	GIsAudioThreadRunning.store(true, std::memory_order_release);
-	return true;
-}
-
-void FAudioThread::Exit()
-{
-	GIsAudioThreadRunning.store(false, std::memory_order_release);
-}
-
-#endif
-
-uint32 FAudioThread::Run()
-{
-	LLM_SCOPE(ELLMTag::AudioMisc);
-	SCOPED_NAMED_EVENT(FAudioThread_Run, FColor::Blue);
-
-	FMemory::SetupTLSCachesOnCurrentThread();
-	FPlatformProcess::SetupAudioThread();
-	AudioThreadMain( TaskGraphBoundSyncEvent );
-	FMemory::ClearAndDisableTLSCachesOnCurrentThread();
-	return 0;
-}
-
-#endif // UE_AUDIO_THREAD_AS_PIPE
-
 void FAudioThread::SetUseThreadedAudio(const bool bInUseThreadedAudio)
 {
 	if (IsAudioThreadRunning() && !bInUseThreadedAudio)
@@ -327,8 +199,6 @@ bool FAudioThread::IsUsingThreadedAudio()
 	return bUseThreadedAudio;
 }
 
-#if UE_AUDIO_THREAD_AS_PIPE
-
 // batching audio commands allows to avoid the overhead of launching a task per command when resources are limited.
 // We assume that resources are limited if the previous batch is not completed yet (potentially waiting for execution due to the CPU being busy
 // with something else). Otherwise we don't wait until we collect a full batch.
@@ -382,59 +252,6 @@ struct FAudioAsyncBatcher
 	}
 };
 
-#else // UE_AUDIO_THREAD_AS_PIPE
-
-struct FAudioAsyncBatcher
-{
-	FGraphEventArray DispatchEvent;
-	int32 NumBatched = 0;
-
-
-	FGraphEventArray* GetAsyncPrereq()
-	{
-		check(IsInGameThread());
-#if !WITH_EDITOR
-		if (GCVarEnableBatchProcessing)
-		{
-			if (NumBatched >= GBatchAudioAsyncBatchSize || !DispatchEvent.Num() || !DispatchEvent[0].GetReference() || DispatchEvent[0]->IsComplete())
-			{
-				Flush();
-			}
-			if (DispatchEvent.Num() == 0)
-			{
-				check(NumBatched == 0);
-				DispatchEvent.Add(FGraphEvent::CreateGraphEvent());
-			}
-			NumBatched++;
-			return &DispatchEvent;
-		}
-#endif
-		return nullptr;
-	}
-
-	void Flush()
-	{
-		check(IsInGameThread());
-		if (NumBatched)
-		{
-			check(DispatchEvent.Num() && DispatchEvent[0].GetReference() && !DispatchEvent[0]->IsComplete());
-			FGraphEventRef Dispatch = DispatchEvent[0];
-			TFunction<void()> FlushAudioCommands = [Dispatch]()
-			{
-				Dispatch->DispatchSubsequents();
-			};
-
-			FFunctionGraphTask::CreateAndDispatchWhenReady(MoveTemp(FlushAudioCommands), TStatId(), nullptr, ENamedThreads::AudioThread);
-
-			DispatchEvent.Empty();
-			NumBatched = 0;
-		}
-	}
-
-};
-
-#endif // UE_AUDIO_THREAD_AS_PIPE
-
 static FAudioAsyncBatcher GAudioAsyncBatcher;
 
 TUniqueFunction<void()> FAudioThread::GetCommandWrapper(TUniqueFunction<void()> InFunction, const TStatId InStatId)
@@ -443,9 +260,6 @@ TUniqueFunction<void()> FAudioThread::GetCommandWrapper(TUniqueFunction<void()>
 	{
 		return [Function = MoveTemp(InFunction), InStatId]()
 		{
-#if !UE_AUDIO_THREAD_AS_PIPE
-			FTaskTagScope Scope(ETaskTag::EAudioThread);
-#endif
 			FScopeCycleCounter ScopeCycleCounter(InStatId);
 			FAudioThread::SetCurrentAudioThreadStatId(InStatId);
 
@@ -467,9 +281,6 @@ TUniqueFunction<void()> FAudioThread::GetCommandWrapper(TUniqueFunction<void()>
 	{
 		return [Function = MoveTemp(InFunction), InStatId]()
 		{
-#if !UE_AUDIO_THREAD_AS_PIPE
-			FTaskTagScope Scope(ETaskTag::EAudioThread);
-#endif
 			FScopeCycleCounter ScopeCycleCounter(InStatId);
 			Function();
 		};
@@ -478,8 +289,6 @@ TUniqueFunction<void()> FAudioThread::GetCommandWrapper(TUniqueFunction<void()>
 
 void FAudioThread::RunCommandOnAudioThread(TUniqueFunction<void()> InFunction, const TStatId InStatId)
 {
-#if UE_AUDIO_THREAD_AS_PIPE
-
 	TUniqueFunction<void()> CommandWrapper{ GetCommandWrapper(MoveTemp(InFunction), InStatId) };
 	if (IsInAudioThread())
 	{
@@ -501,35 +310,6 @@ void FAudioThread::RunCommandOnAudioThread(TUniqueFunction<void()> InFunction, c
 		// the command must be executed on the game thread
 		AsyncTask(ENamedThreads::GameThread, MoveTemp(CommandWrapper));
 	}
-
-#else
-
-	if (IsInAudioThread())
-	{
-		// it's audio-thread-safe so execute the command in-place
-		FScopeCycleCounter ScopeCycleCounter(InStatId);
-		InFunction();
-		return;
-	}
-
-	TUniqueFunction<void()> CommandWrapper{ GetCommandWrapper(MoveTemp(InFunction), InStatId) };
-	if (IsInGameThread())
-	{
-		// batch commands to minimise game thread overhead
-		FFunctionGraphTask::CreateAndDispatchWhenReady(MoveTemp(CommandWrapper), TStatId{}, GAudioAsyncBatcher.GetAsyncPrereq(), ENamedThreads::AudioThread);
-	}
-	// we are on an unknown thread
-	else if (IsUsingThreadedAudio())
-	{
-		FFunctionGraphTask::CreateAndDispatchWhenReady(MoveTemp(CommandWrapper), TStatId{}, nullptr, ENamedThreads::AudioThread);
-	}
-	else
-	{
-		// the command must be executed on the game thread
-		FFunctionGraphTask::CreateAndDispatchWhenReady(MoveTemp(CommandWrapper), TStatId{}, nullptr, ENamedThreads::GameThread);
-	}
-
-#endif
 }
 
 void FAudioThread::SetCurrentAudioThreadStatId(TStatId InStatId)
@@ -581,11 +361,7 @@ void FAudioThread::ProcessAllCommands()
 	}
 	else
 	{
-#if UE_AUDIO_THREAD_AS_PIPE
 		check(GAudioAsyncBatcher.WorkItems.IsEmpty());
-#else
-		check(FPlatformTLS::GetCurrentThreadId() == GGameThreadId);
-#endif
 	}
 }
 
@@ -604,8 +380,6 @@ void FAudioThread::RunCommandOnGameThread(TUniqueFunction<void()> InFunction, co
 	}
 }
 
-#if UE_AUDIO_THREAD_AS_PIPE
-
 FDelegateHandle FAudioThread::PreGC;
 FDelegateHandle FAudioThread::PostGC;
 FDelegateHandle FAudioThread::PreGCDestroy;
@@ -692,198 +466,3 @@ void FAudioCommandFence::Wait(bool bProcessGameThreadTasks) const
 	Fence.Wait();
 	Fence = UE::Tasks::FTask{}; // release the task as it can hold some references
 }
-
-#else // UE_AUDIO_THREAD_AS_PIPE
-
-void FAudioThread::StartAudioThread()
-{
-	check(FPlatformTLS::GetCurrentThreadId() == GGameThreadId);
-
-	check(!IsAudioThreadRunning());
-	check(!GIsAudioThreadSuspended);
-	if (bUseThreadedAudio)
-	{
-PRAGMA_DISABLE_DEPRECATION_WARNINGS
-		check(GAudioThread == nullptr);
-PRAGMA_ENABLE_DEPRECATION_WARNINGS
-
-		static uint32 ThreadCount = 0;
-		check(!ThreadCount); // we should not stop and restart the audio thread; it is complexity we don't need.
-
-		// Create the audio thread.
-		AudioThreadRunnable = new FAudioThread();
-
-PRAGMA_DISABLE_DEPRECATION_WARNINGS
-		GAudioThread = 
-PRAGMA_ENABLE_DEPRECATION_WARNINGS
-			FRunnableThread::Create(AudioThreadRunnable, *FName(NAME_AudioThread).GetPlainNameString(), 0, (GCVarAboveNormalAudioThreadPri == 0) ? TPri_BelowNormal : TPri_AboveNormal, FPlatformAffinity::GetAudioThreadMask());
-
-		// Wait for audio thread to have taskgraph bound before we dispatch any tasks for it.
-		((FAudioThread*)AudioThreadRunnable)->TaskGraphBoundSyncEvent->Wait();
-
-		// ensure the thread has actually started and is idling
-		FAudioCommandFence Fence;
-		Fence.BeginFence();
-		Fence.Wait();
-
-		ThreadCount++;
-		
-		if (GCVarSuspendAudioThread != 0)
-		{
-			SuspendAudioThread();
-		}
-	}
-}
-
-void FAudioThread::StopAudioThread()
-{
-	check(FPlatformTLS::GetCurrentThreadId() == GGameThreadId);
-	check(!GIsAudioThreadSuspended.load() || GCVarSuspendAudioThread != 0);
-
-	if (!IsAudioThreadRunning())
-	{
-		return;
-	}
-
-	FAudioCommandFence Fence;
-	Fence.BeginFence();
-	Fence.Wait();
-	FGraphEventRef QuitTask = TGraphTask<FReturnGraphTask>::CreateTask(nullptr, ENamedThreads::GameThread).ConstructAndDispatchWhenReady(ENamedThreads::AudioThread);
-
-	{
-		QUICK_SCOPE_CYCLE_COUNTER(STAT_StopAudioThread);
-		FTaskGraphInterface::Get().WaitUntilTaskCompletes(QuitTask, ENamedThreads::GameThread_Local);
-	}
-
-PRAGMA_DISABLE_DEPRECATION_WARNINGS
-	// Wait for the audio thread to return.
-	GAudioThread->WaitForCompletion();
-
-	// Destroy the audio thread objects.
-	delete GAudioThread;
-	GAudioThread = nullptr;
-PRAGMA_ENABLE_DEPRECATION_WARNINGS
-
-	delete AudioThreadRunnable;
-	AudioThreadRunnable = nullptr;
-}
-
-FAudioCommandFence::FAudioCommandFence()
-	: FenceDoneEvent(nullptr)
-{
-}
-
-FAudioCommandFence::~FAudioCommandFence()
-{
-	if (FenceDoneEvent)
-	{
-		FenceDoneEvent->Wait();
-
-		FPlatformProcess::ReturnSynchEventToPool(FenceDoneEvent);
-		FenceDoneEvent = nullptr;
-	}
-}
-
-void FAudioCommandFence::BeginFence()
-{
-	if (IsAudioThreadRunning())
-	{
-		DECLARE_CYCLE_STAT(TEXT("FNullGraphTask.FenceAudioCommand"),
-			STAT_FNullGraphTask_FenceAudioCommand,
-			STATGROUP_TaskGraphTasks);
-
-		CompletionEvent = TGraphTask<FNullGraphTask>::CreateTask(GAudioAsyncBatcher.GetAsyncPrereq(), ENamedThreads::GameThread).ConstructAndDispatchWhenReady(
-			GET_STATID(STAT_FNullGraphTask_FenceAudioCommand), ENamedThreads::AudioThread);
-
-		if (FenceDoneEvent)
-		{
-			FenceDoneEvent->Wait();
-
-			FPlatformProcess::ReturnSynchEventToPool(FenceDoneEvent);
-			FenceDoneEvent = nullptr;
-		}
-
-		FenceDoneEvent = FPlatformProcess::GetSynchEventFromPool(true);
-
-		FTaskGraphInterface::Get().TriggerEventWhenTaskCompletes(FenceDoneEvent, CompletionEvent, ENamedThreads::GameThread, ENamedThreads::AudioThread);
-
-		FAudioThread::ProcessAllCommands();
-	}
-	else
-	{
-		CompletionEvent = nullptr;
-	}
-}
-
-bool FAudioCommandFence::IsFenceComplete() const
-{
-	FAudioThread::ProcessAllCommands();
-	
-	if (!CompletionEvent.GetReference() || CompletionEvent->IsComplete())
-	{
-		CompletionEvent = nullptr; // this frees the handle for other uses, the NULL state is considered completed
-		return true;
-	}
-
-	check(IsAudioThreadRunning());
-
-	return FenceDoneEvent->Wait(0);
-}
-
-/**
- * Waits for pending fence commands to retire.
- */
-void FAudioCommandFence::Wait(bool bProcessGameThreadTasks) const
-{
-	FAudioThread::ProcessAllCommands();
-
-	if (!IsFenceComplete()) // this checks the current thread
-	{
-		const double StartTime = FPlatformTime::Seconds();
-		QUICK_SCOPE_CYCLE_COUNTER(STAT_FAudioCommandFence_Wait);
-
-		bool bDone = false;
-
-		do
-		{
-			if (FenceDoneEvent)
-			{
-				bDone = FenceDoneEvent->Wait(GAudioCommandFenceWaitTimeMs);
-			}
-			else
-			{
-				bDone = true;
-			}
-			
-			if(bDone && FenceDoneEvent)
-			{
-				FPlatformProcess::ReturnSynchEventToPool(FenceDoneEvent);
-				FenceDoneEvent = nullptr;
-			}
-
-			// Log how long we've been waiting for the audio thread:
-			float ThisTime = FPlatformTime::Seconds() - StartTime;
- 			if (ThisTime > static_cast<float>(GAudioCommandFenceWaitTimeMs) / 1000.0f + SMALL_NUMBER)
-			{
-				if (GCVarEnableAudioCommandLogging == 1)
-				{
-					FString CurrentTask = FAudioThread::GetCurrentAudioThreadStatId();
-
-					FString LongestTask;
-					double LongestTaskTimeMs;
-					FAudioThread::GetLongestTaskInfo(LongestTask, LongestTaskTimeMs);
-
-					UE_LOG(LogAudio, Display, TEXT("Waited %.2f ms for audio thread. (Current Task: %s, Longest task: %s %.2f ms)"), ThisTime * 1000.0f, *CurrentTask, *LongestTask, LongestTaskTimeMs);
-				}
-				else
-				{
-					UE_LOG(LogAudio, Display,  TEXT("Waited %f ms for audio thread."), ThisTime * 1000.0f);
-				}
-			}
-		} while (!bDone);
-
-		FAudioThread::ResetAudioThreadTimers();
-	}
-}
-
-#endif // UE_AUDIO_THREAD_AS_PIPE
\ No newline at end of file
diff --git a/Engine/Source/Runtime/Engine/Public/AudioThread.h b/Engine/Source/Runtime/Engine/Public/AudioThread.h
--- a/Engine/Source/Runtime/Engine/Public/AudioThread.h
+++ b/Engine/Source/Runtime/Engine/Public/AudioThread.h
@@ -11,15 +11,6 @@
 
 #include "Tasks/Pipe.h"
 
-
-#if !UE_AUDIO_THREAD_AS_PIPE
-
-#include "Async/TaskGraphInterfaces.h"
-#include "HAL/Runnable.h"
-#include "Containers/Queue.h"
-
-#endif
-
 ////////////////////////////////////
 // Audio thread API
 // (the naming is outdated as the thread was replaced by piped tasks)
@@ -28,16 +19,11 @@
 DECLARE_STATS_GROUP(TEXT("Audio Thread Commands"), STATGROUP_AudioThreadCommands, STATCAT_Advanced);
 
 class FAudioThread
-#if !UE_AUDIO_THREAD_AS_PIPE
-	: public FRunnable
-#endif
 {
 private:
 	/** Whether to run with an audio thread */
 	static bool bUseThreadedAudio;
 
-#if UE_AUDIO_THREAD_AS_PIPE
-
 	static TUniquePtr<UE::Tasks::FTaskEvent> ResumeEvent;
 	static int32 SuspendCount;  // accessed only from GT
 
@@ -47,27 +33,6 @@ private:
 	static FDelegateHandle PreGCDestroy;
 	static FDelegateHandle PostGCDestroy;
 
-#else
-
-	/** 
-	* Sync event to make sure that audio thread is bound to the task graph before main thread queues work against it.
-	*/
-	FEvent* TaskGraphBoundSyncEvent;
-
-	/**
-	* Whether the audio thread is currently running
-	* If this is false, then we have no audio thread and audio commands will be issued directly on the game thread
-	*/
-	static TAtomic<bool> bIsAudioThreadRunning;
-
-	/** The audio thread itself. */
-	static FRunnable* AudioThreadRunnable;
-
-	void OnPreGarbageCollect();
-	void OnPostGarbageCollect();
-
-#endif
-
 	/** Stat id of the currently executing audio thread command. */
 	static TStatId CurrentAudioThreadStatId;
 	static TStatId LongestAudioThreadStatId;
@@ -83,19 +48,6 @@ private:
 	static TUniqueFunction<void()> GetCommandWrapper(TUniqueFunction<void()> InFunction, const TStatId InStatId);
 
 public:
-
-#if !UE_AUDIO_THREAD_AS_PIPE
-
-	FAudioThread();
-	virtual ~FAudioThread();
-
-	// FRunnable interface.
-	virtual bool Init() override;
-	virtual void Exit() override;
-	virtual uint32 Run() override;
-
-#endif
-
 	/** Starts the audio thread. */
 	static ENGINE_API void StartAudioThread();
 
@@ -150,9 +102,6 @@ struct FAudioThreadSuspendContext
 class ENGINE_API FAudioCommandFence
 {
 public:
-#if !UE_AUDIO_THREAD_AS_PIPE
-	FAudioCommandFence();
-#endif
 	~FAudioCommandFence();
 
 	/**
@@ -172,19 +121,7 @@ public:
 	bool IsFenceComplete() const;
 
 private:
-
-#if UE_AUDIO_THREAD_AS_PIPE
-
 	/** The last audio batch task **/
 	mutable UE::Tasks::FTask Fence;
-
-#else
-
-	/** Graph event that represents completion of this fence **/
-	mutable FGraphEventRef CompletionEvent;
-	/** Event that fires when CompletionEvent is done. **/
-	mutable FEvent* FenceDoneEvent;
-
-#endif
 };
 
