Add an additional flavor of "OnEndGenerate" virtual function to USoundWave so that derivations like UMetasoundSource can know WHICH generator that it created in the past is now finished.

diff --git a/Engine/Source/Runtime/AudioMixer/Private/AudioMixerSourceBuffer.cpp b/Engine/Source/Runtime/AudioMixer/Private/AudioMixerSourceBuffer.cpp
--- a/Engine/Source/Runtime/AudioMixer/Private/AudioMixerSourceBuffer.cpp
+++ b/Engine/Source/Runtime/AudioMixer/Private/AudioMixerSourceBuffer.cpp
@@ -617,6 +617,10 @@ namespace Audio
 		if (SoundGenerator.IsValid())
 		{
 			SoundGenerator->OnEndGenerate();
+			if (SoundWave)
+			{
+				SoundWave->OnEndGenerate(SoundGenerator);
+			}
 		}
 		else
 		{
diff --git a/Engine/Source/Runtime/Engine/Classes/Sound/SoundWave.h b/Engine/Source/Runtime/Engine/Classes/Sound/SoundWave.h
--- a/Engine/Source/Runtime/Engine/Classes/Sound/SoundWave.h
+++ b/Engine/Source/Runtime/Engine/Classes/Sound/SoundWave.h
@@ -914,6 +914,7 @@ public:
 
 	// Called when the procedural sound wave is done generating on the render thread. Only used in the audio mixer and when bProcedural is true..
 	virtual void OnEndGenerate() {};
+	virtual void OnEndGenerate(ISoundGeneratorPtr Generator) {};
 
 	void AddPlayingSource(const FSoundWaveClientPtr& Source);
 	void RemovePlayingSource(const FSoundWaveClientPtr& Source);
