AssetManager: Optimize by changing container types in ModifyCook and UpdateManagementDatabase to reduce runtime in StartCookByTheBook when adding 100,000 assets.
#rnx

diff --git a/Engine/Source/Runtime/Engine/Private/AssetManager.cpp b/Engine/Source/Runtime/Engine/Private/AssetManager.cpp
--- a/Engine/Source/Runtime/Engine/Private/AssetManager.cpp
+++ b/Engine/Source/Runtime/Engine/Private/AssetManager.cpp
@@ -2,6 +2,8 @@
 
 #include "Engine/AssetManager.h"
 
+#include "Algo/Sort.h"
+#include "Algo/Unique.h"
 #include "AssetRegistry/ARFilter.h"
 #include "AssetRegistry/AssetData.h"
 #include "AssetRegistry/AssetRegistryHelpers.h"
@@ -3763,6 +3765,7 @@ void UAssetManager::UpdateManagementDatabase(bool bForceRefresh)
 
 	// List of packages that need to have their chunks updated
 	TSet<FName> PackagesToUpdateChunksFor;
+	TArray<FName> AssetPackagesReferenced;
 
 	for (const TPair<FName, TSharedRef<FPrimaryAssetTypeData>>& TypePair : AssetTypeMap)
 	{
@@ -3776,7 +3779,7 @@ void UAssetManager::UpdateManagementDatabase(bool bForceRefresh)
 			FPrimaryAssetRules Rules = GetPrimaryAssetRules(PrimaryAssetId);
 
 			// Get the list of directly referenced assets, the registry wants it as FNames
-			TArray<FName> AssetPackagesReferenced;
+			AssetPackagesReferenced.Reset();
 
 			const FSoftObjectPath& AssetRef = NameData.AssetPtr.ToSoftObjectPath();
 
@@ -3790,7 +3793,7 @@ void UAssetManager::UpdateManagementDatabase(bool bForceRefresh)
 				}
 				else
 				{
-					AssetPackagesReferenced.AddUnique(PackageName);
+					AssetPackagesReferenced.Add(PackageName);
 					PackagesToUpdateChunksFor.Add(PackageName);
 				}
 			}
@@ -3811,13 +3814,15 @@ void UAssetManager::UpdateManagementDatabase(bool bForceRefresh)
 						}
 						else
 						{
-							AssetPackagesReferenced.AddUnique(PackageName);
+							AssetPackagesReferenced.Add(PackageName);
 							PackagesToUpdateChunksFor.Add(PackageName);
 						}
 					}
 				}
 			}
 
+			Algo::Sort(AssetPackagesReferenced, FNameLexicalLess());
+			AssetPackagesReferenced.SetNum(Algo::Unique(AssetPackagesReferenced), false /* bAllowShrinking */);
 			for (const FName& AssetPackage : AssetPackagesReferenced)
 			{
 				TMultiMap<FAssetIdentifier, FAssetIdentifier>& ManagerMap = Rules.bApplyRecursively ? PriorityManagementMap.FindOrAdd(Rules.Priority) : NoReferenceManagementMap;
@@ -4021,14 +4026,19 @@ void UAssetManager::ModifyCook(TArray<FName>& PackagesToCook, TArray<FName>& Pac
 
 	bool bIncludeDevelopmentAssets = !bOnlyCookProductionAssets || bTargetPlatformsAllowDevelopmentObjects;
 
+	// Uniquely append packages we need that are not already in PackagesToCook and PackagesToNeverCook
+	TSet<FName> PackagesToCookSet(PackagesToCook);
+	TSet<FName> PackagesToNeverCookSet(PackagesToNeverCook);
+
 	// Get package names in the libraries that we care about for cooking. Only get ones that are needed in production
+	TArray<FName> AssetPackages;
 	for (const FPrimaryAssetTypeInfo& TypeInfo : TypeList)
 	{
 		// Cook these types
 		TArray<FPrimaryAssetId> AssetIdList;
 		GetPrimaryAssetIdList(TypeInfo.PrimaryAssetType, AssetIdList);
+		AssetPackages.Reset();
 
-		TArray<FName> AssetPackages;
 		for (const FPrimaryAssetId& PrimaryAssetId : AssetIdList)
 		{
 			FAssetData AssetData;
@@ -4047,11 +4057,13 @@ void UAssetManager::ModifyCook(TArray<FName>& PackagesToCook, TArray<FName>& Pac
 					for (const FTopLevelAssetPath& FoundReference : FoundEntry.AssetPaths)
 					{
 						FName PackageName = FoundReference.GetPackageName();
-						AssetPackages.AddUnique(PackageName);
+						AssetPackages.Add(PackageName);
 					}
 				}
 			}
 		}
+		Algo::Sort(AssetPackages, FNameFastLess());
+		AssetPackages.SetNum(Algo::Unique(AssetPackages), false /* bAllowShrinking */);
 
 		for (FName PackageName : AssetPackages)
 		{
@@ -4065,12 +4077,22 @@ void UAssetManager::ModifyCook(TArray<FName>& PackagesToCook, TArray<FName>& Pac
 			if (bAlwaysCook && bCanCook && !TypeInfo.bIsEditorOnly)
 			{
 				// If this is always cook, not excluded, and not editor only, cook it
-				PackagesToCook.AddUnique(PackageName);
+				bool bAlreadyInSet;
+				PackagesToCookSet.Add(PackageName, &bAlreadyInSet);
+				if (!bAlreadyInSet)
+				{
+					PackagesToCook.Add(PackageName);
+				}
 			}
 			else if (!bCanCook)
 			{
 				// If this package cannot be cooked, add to exclusion list
-				PackagesToNeverCook.AddUnique(PackageName);
+				bool bAlreadyInSet;
+				PackagesToNeverCookSet.Add(PackageName, &bAlreadyInSet);
+				if (!bAlreadyInSet)
+				{
+					PackagesToNeverCookSet.Add(PackageName);
+				}
 			}
 		}
 	}
