Changing more of ParticleGpuSimulation to use shader parameter structs.

diff --git a/Engine/Shaders/Private/ParticleSimulationShader.usf b/Engine/Shaders/Private/ParticleSimulationShader.usf
--- a/Engine/Shaders/Private/ParticleSimulationShader.usf
+++ b/Engine/Shaders/Private/ParticleSimulationShader.usf
@@ -124,32 +124,21 @@ SamplerState CurveTextureSampler;
 #if MAX_VECTOR_FIELDS != 4 && MAX_VECTOR_FIELDS != 1
 #error This must match MAX_VECTOR_FIELDS in C++ land
 #endif
-Texture3D VectorFieldTextures0;
+Texture3D VectorFieldTextures_0;
 
 #if MAX_VECTOR_FIELDS > 1
-Texture3D VectorFieldTextures1;
-Texture3D VectorFieldTextures2;
-Texture3D VectorFieldTextures3;
+Texture3D VectorFieldTextures_1;
+Texture3D VectorFieldTextures_2;
+Texture3D VectorFieldTextures_3;
 #endif
-//Texture3D VectorFieldTextures4;
-//Texture3D VectorFieldTextures5;
-//Texture3D VectorFieldTextures6;
-//Texture3D VectorFieldTextures7;
-//Texture3D VectorFieldTextures8;
-//Texture3D VectorFieldTextures9;
-
-SamplerState VectorFieldTexturesSampler0;
-#if MAX_VECTOR_FIELDS > 1
-SamplerState VectorFieldTexturesSampler1;
-SamplerState VectorFieldTexturesSampler2;
-SamplerState VectorFieldTexturesSampler3;
-#endif
-//SamplerState VectorFieldTexturesSampler4;
-//SamplerState VectorFieldTexturesSampler5;
-//SamplerState VectorFieldTexturesSampler6;
-//SamplerState VectorFieldTexturesSampler7;
-//SamplerState VectorFieldTexturesSampler8;
-//SamplerState VectorFieldTexturesSampler9;
+//Texture3D VectorFieldTextures_4;
+//Texture3D VectorFieldTextures_5;
+//Texture3D VectorFieldTextures_6;
+//Texture3D VectorFieldTextures_7;
+//Texture3D VectorFieldTextures_8;
+//Texture3D VectorFieldTextures_9;
+
+SamplerState VectorFieldTexturesSampler;
 
 /**
  * Computes the orbit velocity to apply to the particle based on time.
@@ -182,22 +171,22 @@ float3 ComputeOrbitVelocity(float Time, float RandomOrbit)
  float3 SampleVectorFieldTexture(int Index, float3 UV)
  {
 #if MAX_VECTOR_FIELDS == 1
-	 return Texture3DSample(VectorFieldTextures0, VectorFieldTexturesSampler0, UV).xyz;
+	 return Texture3DSample(VectorFieldTextures_0, VectorFieldTexturesSampler, UV).xyz;
 #else
-	if (Index == 0) return Texture3DSample(VectorFieldTextures0, VectorFieldTexturesSampler0, UV).xyz;
-	if (Index == 1) return Texture3DSample(VectorFieldTextures1, VectorFieldTexturesSampler1, UV).xyz;
-	if (Index == 2) return Texture3DSample(VectorFieldTextures2, VectorFieldTexturesSampler2, UV).xyz;
-	return Texture3DSample(VectorFieldTextures3, VectorFieldTexturesSampler3, UV).xyz;
+	if (Index == 0) return Texture3DSample(VectorFieldTextures_0, VectorFieldTexturesSampler, UV).xyz;
+	if (Index == 1) return Texture3DSample(VectorFieldTextures_1, VectorFieldTexturesSampler, UV).xyz;
+	if (Index == 2) return Texture3DSample(VectorFieldTextures_2, VectorFieldTexturesSampler, UV).xyz;
+	return Texture3DSample(VectorFieldTextures_3, VectorFieldTexturesSampler, UV).xyz;
 #endif
 
-//	if (Index == 3) return Texture3DSample(VectorFieldTextures3, VectorFieldTexturesSampler3, UV).xyz;
-//	if (Index == 4) return Texture3DSample(VectorFieldTextures4, VectorFieldTexturesSampler4, UV).xyz;
-//	if (Index == 5) return Texture3DSample(VectorFieldTextures5, VectorFieldTexturesSampler5, UV).xyz;
-//	if (Index == 6) return Texture3DSample(VectorFieldTextures6, VectorFieldTexturesSampler6, UV).xyz;
-//	if (Index == 7) return Texture3DSample(VectorFieldTextures7, VectorFieldTexturesSampler7, UV).xyz;
-//	if (Index == 8) return Texture3DSample(VectorFieldTextures8, VectorFieldTexturesSampler8, UV).xyz;
+//	if (Index == 3) return Texture3DSample(VectorFieldTextures_3, VectorFieldTexturesSampler, UV).xyz;
+//	if (Index == 4) return Texture3DSample(VectorFieldTextures_4, VectorFieldTexturesSampler, UV).xyz;
+//	if (Index == 5) return Texture3DSample(VectorFieldTextures_5, VectorFieldTexturesSampler, UV).xyz;
+//	if (Index == 6) return Texture3DSample(VectorFieldTextures_6, VectorFieldTexturesSampler, UV).xyz;
+//	if (Index == 7) return Texture3DSample(VectorFieldTextures_7, VectorFieldTexturesSampler, UV).xyz;
+//	if (Index == 8) return Texture3DSample(VectorFieldTextures_8, VectorFieldTexturesSampler, UV).xyz;
 //	
-//	return Texture3DSample(VectorFieldTextures9, VectorFieldTexturesSampler9, UV).xyz;
+//	return Texture3DSample(VectorFieldTextures_9, VectorFieldTexturesSampler, UV).xyz;
  }
 
 /**
diff --git a/Engine/Source/Runtime/Engine/Private/Particles/FXSystem.cpp b/Engine/Source/Runtime/Engine/Private/Particles/FXSystem.cpp
--- a/Engine/Source/Runtime/Engine/Private/Particles/FXSystem.cpp
+++ b/Engine/Source/Runtime/Engine/Private/Particles/FXSystem.cpp
@@ -509,7 +509,7 @@ void FFXSystem::PreRender(FRDGBuilder& GraphBuilder, TConstArrayView<FViewInfo>
 		RDG_GPU_STAT_SCOPE(GraphBuilder, FXSystemPreRender);
 		RDG_CSV_STAT_EXCLUSIVE_SCOPE(GraphBuilder, FXSystem);
 
-		FRHIUniformBuffer* ViewUniformBuffer = GetReferenceViewUniformBuffer(Views);
+		TUniformBufferRef<FViewUniformShaderParameters> ViewUniformBuffer = GetReferenceViewUniformBuffer(Views);
 		const FGlobalDistanceFieldParameterData* GlobalDistanceFieldParameterData = GetReferenceGlobalDistanceFieldData(Views);
 
 		AddPass(
@@ -526,7 +526,7 @@ void FFXSystem::PreRender(FRDGBuilder& GraphBuilder, TConstArrayView<FViewInfo>
 				PrepareGPUSimulation(RHICmdList);
 
 				RHICmdList.SetCurrentStat(GET_STATID(STAT_CLM_FXPreRender_Simulate));
-				SimulateGPUParticles(RHICmdList, EParticleSimulatePhase::Main, nullptr, nullptr);
+				SimulateGPUParticles(RHICmdList, EParticleSimulatePhase::Main, {}, nullptr);
 
 				RHICmdList.SetCurrentStat(GET_STATID(STAT_CLM_FXPreRender_Finalize));
 				FinalizeGPUSimulation(RHICmdList);
@@ -556,7 +556,7 @@ void FFXSystem::PostRenderOpaque(FRDGBuilder& GraphBuilder, TConstArrayView<FVie
 		RDG_GPU_STAT_SCOPE(GraphBuilder, FXSystemPostRenderOpaque);
 		RDG_CSV_STAT_EXCLUSIVE_SCOPE(GraphBuilder, FXSystem);
 
-		FRHIUniformBuffer* ViewUniformBuffer = GetReferenceViewUniformBuffer(Views);
+		TUniformBufferRef<FViewUniformShaderParameters> ViewUniformBuffer = GetReferenceViewUniformBuffer(Views);
 
 		AddPass(GraphBuilder, RDG_EVENT_NAME("FFXSystem::PostRenderOpaque"), 
 			[this, ViewUniformBuffer](FRHICommandListImmediate& RHICmdList)
diff --git a/Engine/Source/Runtime/Engine/Private/Particles/FXSystemPrivate.h b/Engine/Source/Runtime/Engine/Private/Particles/FXSystemPrivate.h
--- a/Engine/Source/Runtime/Engine/Private/Particles/FXSystemPrivate.h
+++ b/Engine/Source/Runtime/Engine/Private/Particles/FXSystemPrivate.h
@@ -269,7 +269,7 @@ public:
 	/** Get the shared SortManager, used in the rendering loop to call FGPUSortManager::OnPreRender() and FGPUSortManager::OnPostRenderOpaque() */
 	virtual FGPUSortManager* GetGPUSortManager() const override;
 
-	virtual void SetSceneTexturesUniformBuffer(FRHIUniformBuffer* InSceneTexturesUniformParams) override { SceneTexturesUniformParams = InSceneTexturesUniformParams; }
+	virtual void SetSceneTexturesUniformBuffer(const TUniformBufferRef<FSceneTextureUniformParameters>& InSceneTexturesUniformParams) override { SceneTexturesUniformParams = InSceneTexturesUniformParams; }
 
 private:
 
@@ -340,7 +340,7 @@ private:
 	void SimulateGPUParticles(
 		FRHICommandListImmediate& RHICmdList,
 		EParticleSimulatePhase::Type Phase,
-		FRHIUniformBuffer* ViewUniformBuffer,
+		const TUniformBufferRef<FViewUniformShaderParameters>& ViewUniformBuffer,
 		const FGlobalDistanceFieldParameterData* GlobalDistanceFieldParameterData
 		);
 
@@ -356,7 +356,7 @@ private:
 	void SimulateGPUParticles_Internal(
 		FRHICommandListImmediate& RHICmdList,
 		EParticleSimulatePhase::Type Phase,
-		FRHIUniformBuffer* ViewUniformBuffer,
+		const TUniformBufferRef<FViewUniformShaderParameters>& ViewUniformBuffer,
 		const FGlobalDistanceFieldParameterData* GlobalDistanceFieldParameterData,
 		FRHITexture2D* SceneDepthTexture,
 		FRHITexture2D* GBufferATexture
@@ -396,6 +396,6 @@ private:
 	EParticleSimulatePhase::Type PhaseToBroadcastTemporalEffect = EParticleSimulatePhase::First;
 #endif
 
-	FRHIUniformBuffer* SceneTexturesUniformParams = nullptr;
+	TUniformBufferRef<FSceneTextureUniformParameters> SceneTexturesUniformParams;
 };
 
diff --git a/Engine/Source/Runtime/Engine/Private/Particles/FXSystemSet.cpp b/Engine/Source/Runtime/Engine/Private/Particles/FXSystemSet.cpp
--- a/Engine/Source/Runtime/Engine/Private/Particles/FXSystemSet.cpp
+++ b/Engine/Source/Runtime/Engine/Private/Particles/FXSystemSet.cpp
@@ -203,7 +203,7 @@ void FFXSystemSet::PreRender(class FRDGBuilder& GraphBuilder, TConstArrayView<FV
 	}
 }
 
-void FFXSystemSet::SetSceneTexturesUniformBuffer(FRHIUniformBuffer* InSceneTexturesUniformParams)
+void FFXSystemSet::SetSceneTexturesUniformBuffer(const TUniformBufferRef<FSceneTextureUniformParameters>& InSceneTexturesUniformParams)
 {
 	for (FFXSystemInterface* FXSystem : FXSystems)
 	{
diff --git a/Engine/Source/Runtime/Engine/Private/Particles/FXSystemSet.h b/Engine/Source/Runtime/Engine/Private/Particles/FXSystemSet.h
--- a/Engine/Source/Runtime/Engine/Private/Particles/FXSystemSet.h
+++ b/Engine/Source/Runtime/Engine/Private/Particles/FXSystemSet.h
@@ -54,7 +54,7 @@ public:
 	/** Get the shared SortManager, used in the rendering loop to call FGPUSortManager::OnPreRender() and FGPUSortManager::OnPostRenderOpaque() */
 	virtual FGPUSortManager* GetGPUSortManager() const override;
 
-	virtual void SetSceneTexturesUniformBuffer(FRHIUniformBuffer* InSceneTexturesUniformParams) override;
+	virtual void SetSceneTexturesUniformBuffer(const TUniformBufferRef<FSceneTextureUniformParameters>& InSceneTexturesUniformParams) override;
 
 protected:
 
diff --git a/Engine/Source/Runtime/Engine/Private/Particles/ParticleGpuSimulation.cpp b/Engine/Source/Runtime/Engine/Private/Particles/ParticleGpuSimulation.cpp
--- a/Engine/Source/Runtime/Engine/Private/Particles/ParticleGpuSimulation.cpp
+++ b/Engine/Source/Runtime/Engine/Private/Particles/ParticleGpuSimulation.cpp
@@ -51,6 +51,8 @@
 #include "Misc/CoreDelegates.h"
 #include "PipelineStateCache.h"
 #include "MeshMaterialShader.h"
+#include "ShaderParameterStruct.h"
+#include "SceneRendering.h"
 
 DECLARE_CYCLE_STAT(TEXT("GPUSpriteEmitterInstance Init GT"), STAT_GPUSpriteEmitterInstance_Init, STATGROUP_Particles);
 DECLARE_GPU_STAT_NAMED(ParticleSimulation, TEXT("Particle Simulation"));
@@ -842,57 +844,32 @@ struct FParticlePerFrameSimulationParameters
 
 };
 
-/**
- * Per-frame shader parameters for particle simulation.
- */
-struct FParticlePerFrameSimulationShaderParameters
+BEGIN_SHADER_PARAMETER_STRUCT(FParticlePerFrameSimulationShaderParameters, )
+	SHADER_PARAMETER(FVector4f, PointAttractor)
+	SHADER_PARAMETER(FVector4f, PositionOffsetAndAttractorStrength)
+	SHADER_PARAMETER(FVector2f, LocalToWorldScale)
+	SHADER_PARAMETER(float, DeltaSeconds)
+	SHADER_PARAMETER(int32, NumIterations)
+	SHADER_PARAMETER(FVector3f, LWCTile)
+END_SHADER_PARAMETER_STRUCT()
+
+FParticlePerFrameSimulationShaderParameters GetParticlePerFrameSimulationShaderParameters(const FParticlePerFrameSimulationParameters& Parameters, bool bUseFixDT)
 {
-	DECLARE_INLINE_TYPE_LAYOUT(FParticlePerFrameSimulationShaderParameters, NonVirtual);
-public:
-	void Bind(const FShaderParameterMap& ParameterMap)
-	{
-		PointAttractor.Bind(ParameterMap,TEXT("PointAttractor"));
-		PositionOffsetAndAttractorStrength.Bind(ParameterMap,TEXT("PositionOffsetAndAttractorStrength"));
-		LocalToWorldScale.Bind(ParameterMap,TEXT("LocalToWorldScale"));
-		DeltaSeconds.Bind(ParameterMap,TEXT("DeltaSeconds"));
-		NumIterations.Bind(ParameterMap,TEXT("NumIterations"));
-		LWCTile.Bind(ParameterMap,TEXT("LWCTile"));
-	}
+	FParticlePerFrameSimulationShaderParameters Result{};
 
-	template <typename ShaderRHIParamRef>
-	void Set(FRHICommandList& RHICmdList, const ShaderRHIParamRef& ShaderRHI, const FParticlePerFrameSimulationParameters& Parameters, bool bUseFixDT) const
-	{
-		// The offset must only be applied once in the frame, and be stored in the persistent data (not the interpolated one).
-		const float FixDeltaSeconds = CVarGPUParticleFixDeltaSeconds.GetValueOnRenderThread();
-		const bool bApplyOffset = FixDeltaSeconds <= 0 || bUseFixDT;
-		const FVector4f OnlyAttractorStrength = FVector4f(0, 0, 0, Parameters.PositionOffsetAndAttractorStrength.W);
-
-		SetShaderValue(RHICmdList,ShaderRHI,PointAttractor,Parameters.PointAttractor);
-		SetShaderValue(RHICmdList,ShaderRHI,PositionOffsetAndAttractorStrength, bApplyOffset ? Parameters.PositionOffsetAndAttractorStrength : OnlyAttractorStrength);
-		SetShaderValue(RHICmdList,ShaderRHI,LocalToWorldScale,Parameters.LocalToWorldScale);
-		SetShaderValue(RHICmdList,ShaderRHI,DeltaSeconds, bUseFixDT ? Parameters.DeltaSecondsInFix : Parameters.DeltaSecondsInVar);
-		SetShaderValue(RHICmdList,ShaderRHI,NumIterations, bUseFixDT ? Parameters.NumIterationsInFix : Parameters.NumIterationsInVar);
-		SetShaderValue(RHICmdList,ShaderRHI,LWCTile, Parameters.LWCTile);
-	}
-
-	LAYOUT_FIELD(FShaderParameter, PointAttractor);
-	LAYOUT_FIELD(FShaderParameter, PositionOffsetAndAttractorStrength);
-	LAYOUT_FIELD(FShaderParameter, LocalToWorldScale);
-	LAYOUT_FIELD(FShaderParameter, DeltaSeconds);
-	LAYOUT_FIELD(FShaderParameter, NumIterations);
-	LAYOUT_FIELD(FShaderParameter, LWCTile);
-	
-};
+	// The offset must only be applied once in the frame, and be stored in the persistent data (not the interpolated one).
+	const float FixDeltaSeconds = CVarGPUParticleFixDeltaSeconds.GetValueOnRenderThread();
+	const bool bApplyOffset = FixDeltaSeconds <= 0 || bUseFixDT;
+	const FVector4f OnlyAttractorStrength = FVector4f(0, 0, 0, Parameters.PositionOffsetAndAttractorStrength.W);
 
-FArchive& operator<<(FArchive& Ar, FParticlePerFrameSimulationShaderParameters& PerFrameParameters)
-{
-	Ar << PerFrameParameters.PointAttractor;
-	Ar << PerFrameParameters.PositionOffsetAndAttractorStrength;
-	Ar << PerFrameParameters.LocalToWorldScale;
-	Ar << PerFrameParameters.DeltaSeconds;
-	Ar << PerFrameParameters.NumIterations;
-	Ar << PerFrameParameters.LWCTile;
-	return Ar;
+	Result.PointAttractor = Parameters.PointAttractor;
+	Result.PositionOffsetAndAttractorStrength = bApplyOffset ? Parameters.PositionOffsetAndAttractorStrength : OnlyAttractorStrength;
+	Result.LocalToWorldScale = Parameters.LocalToWorldScale;
+	Result.DeltaSeconds = bUseFixDT ? Parameters.DeltaSecondsInFix : Parameters.DeltaSecondsInVar;
+	Result.NumIterations = bUseFixDT ? Parameters.NumIterationsInFix : Parameters.NumIterationsInVar;
+	Result.LWCTile = Parameters.LWCTile;
+
+	return Result;
 }
 
 /**
@@ -917,9 +894,9 @@ typedef TUniformBufferRef<FVectorFieldUniformParameters> FVectorFieldUniformBuff
  */
 class FParticleTileVS : public FGlobalShader
 {
-	DECLARE_SHADER_TYPE(FParticleTileVS,Global);
-
 public:
+	DECLARE_GLOBAL_SHADER(FParticleTileVS);
+	SHADER_USE_PARAMETER_STRUCT(FParticleTileVS, FGlobalShader)
 
 	static bool ShouldCompilePermutation(const FGlobalShaderPermutationParameters& Parameters)
 	{
@@ -932,44 +909,58 @@ public:
 		OutEnvironment.SetDefine(TEXT("TILES_PER_INSTANCE"), TILES_PER_INSTANCE);
 	}
 
-	/** Default constructor. */
-	FParticleTileVS()
-	{
-	}
+	BEGIN_SHADER_PARAMETER_STRUCT(FParameters, )
+		SHADER_PARAMETER(float, TileSizeX)
+		SHADER_PARAMETER(float, TileSizeY)
+		SHADER_PARAMETER_SRV(FShaderResourceParameter, TileOffsets)
+	END_SHADER_PARAMETER_STRUCT()
 
-	/** Initialization constructor. */
-	explicit FParticleTileVS( const ShaderMetaType::CompiledShaderInitializerType& Initializer )
-		: FGlobalShader(Initializer)
-	{
-		TileOffsets.Bind(Initializer.ParameterMap, TEXT("TileOffsets"));
-		TileSizeX.Bind(Initializer.ParameterMap, TEXT("TileSizeX"));
-		TileSizeY.Bind(Initializer.ParameterMap, TEXT("TileSizeY"));
-	}
-	/** Set parameters. */
-	void SetParameters(FRHICommandList& RHICmdList, FParticleShaderParamRef TileOffsetsRef)
+	FParameters GetParameters(FParticleShaderParamRef TileOffsetsRef)
 	{
-		FRHIVertexShader* VertexShaderRHI = RHICmdList.GetBoundVertexShader();
-		SetSRVParameter(RHICmdList, VertexShaderRHI, TileOffsets, TileOffsetsRef);
-		SetShaderValue(RHICmdList, VertexShaderRHI, TileSizeX, (float)GParticleSimulationTileSize / (float)GParticleSimulationTextureSizeX);
-		SetShaderValue(RHICmdList, VertexShaderRHI, TileSizeY, (float)GParticleSimulationTileSize / (float)GParticleSimulationTextureSizeY);
+		FParameters Result{};
+		Result.TileSizeX = (float)GParticleSimulationTileSize / (float)GParticleSimulationTextureSizeX;
+		Result.TileSizeY = (float)GParticleSimulationTileSize / (float)GParticleSimulationTextureSizeY;
+		Result.TileOffsets = TileOffsetsRef;
+		return Result;
 	}
-
-private:
-	/** Buffer from which to read tile offsets. */
-	LAYOUT_FIELD(FShaderResourceParameter, TileOffsets);
-	LAYOUT_FIELD(FShaderParameter, TileSizeX);
-	LAYOUT_FIELD(FShaderParameter, TileSizeY);
 };
 
+BEGIN_SHADER_PARAMETER_STRUCT(FParticleSimulationPsParameters, )
+	SHADER_PARAMETER(float, CollisionDepthBounds)
+
+	SHADER_PARAMETER_TEXTURE(Texture2D, PositionTexture)
+	SHADER_PARAMETER_SAMPLER(SamplerState, PositionTextureSampler)
+	SHADER_PARAMETER_TEXTURE(Texture2D, VelocityTexture)
+	SHADER_PARAMETER_SAMPLER(SamplerState, VelocityTextureSampler)
+	SHADER_PARAMETER_TEXTURE(Texture2D, AttributesTexture)
+	SHADER_PARAMETER_SAMPLER(SamplerState, AttributesTextureSampler)
+	SHADER_PARAMETER_TEXTURE(Texture2D, RenderAttributesTexture)
+	SHADER_PARAMETER_SAMPLER(SamplerState, RenderAttributesTextureSampler)
+	SHADER_PARAMETER_TEXTURE(Texture2D, CurveTexture)
+	SHADER_PARAMETER_SAMPLER(SamplerState, CurveTextureSampler)
+
+	SHADER_PARAMETER_TEXTURE_ARRAY(Texture3D, VectorFieldTextures, [MAX_VECTOR_FIELDS])
+	SHADER_PARAMETER_SAMPLER(SamplerState, VectorFieldTexturesSampler)
+
+	SHADER_PARAMETER_STRUCT_INCLUDE(FParticlePerFrameSimulationShaderParameters, PerFrameParameters)
+	SHADER_PARAMETER_STRUCT_INCLUDE(FGlobalDistanceFieldParameters2, GlobalDistanceFieldParameters)
+
+	SHADER_PARAMETER_STRUCT_REF(FViewUniformShaderParameters, View)
+	SHADER_PARAMETER_STRUCT_REF(FParticleSimulationParameters, ParticleSimulationParameters)
+	SHADER_PARAMETER_STRUCT_REF(FVectorFieldUniformParameters, VectorFieldUniformParameters)
+END_SHADER_PARAMETER_STRUCT()
+
 /**
  * Pixel shader for simulating particles on the GPU.
  */
 template <EParticleCollisionShaderMode CollisionMode>
 class TParticleSimulationPS : public FGlobalShader
 {
-	DECLARE_SHADER_TYPE(TParticleSimulationPS,Global);
-
 public:
+	DECLARE_GLOBAL_SHADER(TParticleSimulationPS);
+
+	using FParameters = FParticleSimulationPsParameters;
+	SHADER_USE_PARAMETER_STRUCT(TParticleSimulationPS, FGlobalShader)
 
 	static bool ShouldCompilePermutation(const FGlobalShaderPermutationParameters& Parameters)
 	{
@@ -985,155 +976,6 @@ public:
 		OutEnvironment.SetDefine(TEXT("DISTANCE_FIELD_COLLISION"), CollisionMode == PCM_DistanceField);
 		OutEnvironment.SetRenderTargetOutputFormat(0, PF_A32B32G32R32F);
 	}
-
-	/** Default constructor. */
-	TParticleSimulationPS()
-	{
-	}
-
-	/** Initialization constructor. */
-	explicit TParticleSimulationPS(const ShaderMetaType::CompiledShaderInitializerType& Initializer)
-		: FGlobalShader(Initializer)
-	{
-		PositionTexture.Bind(Initializer.ParameterMap, TEXT("PositionTexture"));
-		PositionTextureSampler.Bind(Initializer.ParameterMap, TEXT("PositionTextureSampler"));
-		VelocityTexture.Bind(Initializer.ParameterMap, TEXT("VelocityTexture"));
-		VelocityTextureSampler.Bind(Initializer.ParameterMap, TEXT("VelocityTextureSampler"));
-		AttributesTexture.Bind(Initializer.ParameterMap, TEXT("AttributesTexture"));
-		AttributesTextureSampler.Bind(Initializer.ParameterMap, TEXT("AttributesTextureSampler"));
-		RenderAttributesTexture.Bind(Initializer.ParameterMap, TEXT("RenderAttributesTexture"));
-		RenderAttributesTextureSampler.Bind(Initializer.ParameterMap, TEXT("RenderAttributesTextureSampler"));
-		CurveTexture.Bind(Initializer.ParameterMap, TEXT("CurveTexture"));
-		CurveTextureSampler.Bind(Initializer.ParameterMap, TEXT("CurveTextureSampler"));
-		for (int32 i = 0; i < MAX_VECTOR_FIELDS; ++i)
-		{
-			VectorFieldTextures[i].Bind(Initializer.ParameterMap, *FString::Printf(TEXT("VectorFieldTextures%d"), i));
-			VectorFieldTexturesSamplers[i].Bind(Initializer.ParameterMap, *FString::Printf(TEXT("VectorFieldTexturesSampler%d"), i));
-		}
-		CollisionDepthBounds.Bind(Initializer.ParameterMap,TEXT("CollisionDepthBounds"));
-		PerFrameParameters.Bind(Initializer.ParameterMap);
-		GlobalDistanceFieldParameters.Bind(Initializer.ParameterMap);
-	}
-
-	/**
-	 * Set parameters for this shader.
-	 */
-	void SetParameters(
-		FRHICommandList& RHICmdList, 
-		const FParticleStateTextures& TextureResources,
-		const FParticleAttributesTexture& InAttributesTexture,
-		const FParticleAttributesTexture& InRenderAttributesTexture,
-		FRHIUniformBuffer* ViewUniformBuffer,
-		ERHIFeatureLevel::Type FeatureLevel,
-		const FGlobalDistanceFieldParameterData* GlobalDistanceFieldParameterData
-		)
-	{
-		FRHIPixelShader* PixelShaderRHI = RHICmdList.GetBoundPixelShader();
-		FRHISamplerState* SamplerStatePoint = TStaticSamplerState<SF_Point>::GetRHI();
-		FRHISamplerState* SamplerStateLinear = TStaticSamplerState<SF_Bilinear,AM_Clamp,AM_Clamp,AM_Clamp>::GetRHI();
-		SetTextureParameter(RHICmdList, PixelShaderRHI, PositionTexture, PositionTextureSampler, SamplerStatePoint, TextureResources.PositionTextureRHI);
-		SetTextureParameter(RHICmdList, PixelShaderRHI, VelocityTexture, VelocityTextureSampler, SamplerStatePoint, TextureResources.VelocityTextureRHI);
-		SetTextureParameter(RHICmdList, PixelShaderRHI, AttributesTexture, AttributesTextureSampler, SamplerStatePoint, InAttributesTexture.TextureRHI);
-		SetTextureParameter(RHICmdList, PixelShaderRHI, CurveTexture, CurveTextureSampler, SamplerStateLinear, GParticleCurveTexture.GetCurveTexture());
-
-		if (CollisionMode == PCM_DepthBuffer)
-		{
-			check(ViewUniformBuffer != NULL);
-			FGlobalShader::SetParameters<FViewUniformShaderParameters>(RHICmdList, PixelShaderRHI, ViewUniformBuffer);
-
-			SetTextureParameter(
-				RHICmdList, 
-				PixelShaderRHI,
-				RenderAttributesTexture,
-				RenderAttributesTextureSampler,
-				SamplerStatePoint,
-				InRenderAttributesTexture.TextureRHI
-				);
-			SetShaderValue(RHICmdList, PixelShaderRHI, CollisionDepthBounds, FXConsoleVariables::GPUCollisionDepthBounds);
-		}
-		else if (CollisionMode == PCM_DistanceField)
-		{
-			check(ViewUniformBuffer != NULL);
-			FGlobalShader::SetParameters<FViewUniformShaderParameters>(RHICmdList, PixelShaderRHI, ViewUniformBuffer);
-
-			GlobalDistanceFieldParameters.Set(RHICmdList, PixelShaderRHI, *GlobalDistanceFieldParameterData);
-
-			SetTextureParameter(
-				RHICmdList, 
-				PixelShaderRHI,
-				RenderAttributesTexture,
-				RenderAttributesTextureSampler,
-				SamplerStatePoint,
-				InRenderAttributesTexture.TextureRHI
-				);
-		}
-	}
-
-	/**
-	 * Set parameters for the vector fields sampled by this shader.
-	 * @param VectorFieldParameters -Parameters needed to sample local vector fields.
-	 */
-	void SetVectorFieldParameters(FRHICommandList& RHICmdList, const FVectorFieldUniformBufferRef& UniformBuffer, FRHITexture3D* const* VolumeTexturesRHI)
-	{
-		FRHIPixelShader* PixelShaderRHI = RHICmdList.GetBoundPixelShader();
-		SetUniformBufferParameter(RHICmdList, PixelShaderRHI, GetUniformBufferParameter<FVectorFieldUniformParameters>(), UniformBuffer);
-		
-		FRHISamplerState* SamplerStateLinear = TStaticSamplerState<SF_Bilinear,AM_Clamp,AM_Clamp,AM_Clamp>::GetRHI();
-
-		for (int32 i = 0; i < MAX_VECTOR_FIELDS; ++i)
-		{
-			SetSamplerParameter(RHICmdList, PixelShaderRHI, VectorFieldTexturesSamplers[i], SamplerStateLinear);
-			SetTextureParameter(RHICmdList, PixelShaderRHI, VectorFieldTextures[i], VolumeTexturesRHI[i]);
-		}
-	}
-
-	/**
-	 * Set per-instance parameters for this shader.
-	 */
-	void SetInstanceParameters(FRHICommandList& RHICmdList, FRHIUniformBuffer* UniformBuffer, const FParticlePerFrameSimulationParameters& InPerFrameParameters, bool bUseFixDT)
-	{
-		FRHIPixelShader* PixelShaderRHI = RHICmdList.GetBoundPixelShader();
-		SetUniformBufferParameter(RHICmdList, PixelShaderRHI, GetUniformBufferParameter<FParticleSimulationParameters>(), UniformBuffer);
-		PerFrameParameters.Set(RHICmdList, PixelShaderRHI, InPerFrameParameters, bUseFixDT);
-	}
-
-	/**
-	 * Unbinds buffers that may need to be bound as UAVs.
-	 */
-	void UnbindBuffers(FRHICommandList& RHICmdList)
-	{
-		FRHIPixelShader* PixelShaderRHI = RHICmdList.GetBoundPixelShader();
-		FRHIShaderResourceView* NullSRV = nullptr;
-		for (int32 i = 0; i < MAX_VECTOR_FIELDS; ++i)
-		{
-			SetSRVParameter(RHICmdList, PixelShaderRHI, VectorFieldTextures[i], NullSRV);
-		}
-	}
-
-private:
-	/** The position texture parameter. */
-	LAYOUT_FIELD(FShaderResourceParameter, PositionTexture);
-	LAYOUT_FIELD(FShaderResourceParameter, PositionTextureSampler);
-	/** The velocity texture parameter. */
-	LAYOUT_FIELD(FShaderResourceParameter, VelocityTexture);
-	LAYOUT_FIELD(FShaderResourceParameter, VelocityTextureSampler);
-	/** The simulation attributes texture parameter. */
-	LAYOUT_FIELD(FShaderResourceParameter, AttributesTexture);
-	LAYOUT_FIELD(FShaderResourceParameter, AttributesTextureSampler);
-	/** The render attributes texture parameter. */
-	LAYOUT_FIELD(FShaderResourceParameter, RenderAttributesTexture);
-	LAYOUT_FIELD(FShaderResourceParameter, RenderAttributesTextureSampler);
-	/** The curve texture parameter. */
-	LAYOUT_FIELD(FShaderResourceParameter, CurveTexture);
-	LAYOUT_FIELD(FShaderResourceParameter, CurveTextureSampler);
-	/** Vector fields. */
-	LAYOUT_ARRAY(FShaderResourceParameter, VectorFieldTextures, MAX_VECTOR_FIELDS);
-	LAYOUT_ARRAY(FShaderResourceParameter, VectorFieldTexturesSamplers, MAX_VECTOR_FIELDS);
-	/** Per frame simulation parameters. */
-	LAYOUT_FIELD(FParticlePerFrameSimulationShaderParameters, PerFrameParameters);
-	/** Collision depth bounds. */
-	LAYOUT_FIELD(FShaderParameter, CollisionDepthBounds);
-	LAYOUT_FIELD(FGlobalDistanceFieldParameters, GlobalDistanceFieldParameters);
 };
 
 /**
@@ -1339,7 +1181,7 @@ struct FSimulationCommandGPU
 	FParticleShaderParamRef TileOffsetsShaderRef;
 	FParticleBufferParamRef TileOffsetsBufferRef;
 	/** Uniform buffer containing simulation parameters. */
-	FRHIUniformBuffer* UniformBuffer;
+	FParticleSimulationBufferRef UniformBuffer;
 	/** Uniform buffer containing per-frame simulation parameters. */
 	FParticlePerFrameSimulationParameters PerFrameParameters;
 	/** Parameters to sample the local vector field for this simulation. */
@@ -1350,7 +1192,7 @@ struct FSimulationCommandGPU
 	int32 UnalignedTileCount;
 
 	/** Initialization constructor. */
-	FSimulationCommandGPU(FParticleShaderParamRef InTileOffsetsShaderRef, FParticleBufferParamRef InTileOffsetsBufferRef, FRHIUniformBuffer* InUniformBuffer, const FParticlePerFrameSimulationParameters& InPerFrameParameters, FVectorFieldUniformBufferRef& InVectorFieldsUniformBuffer, int32 InTileCount)
+	FSimulationCommandGPU(FParticleShaderParamRef InTileOffsetsShaderRef, FParticleBufferParamRef InTileOffsetsBufferRef, const FParticleSimulationBufferRef& InUniformBuffer, const FParticlePerFrameSimulationParameters& InPerFrameParameters, FVectorFieldUniformBufferRef& InVectorFieldsUniformBuffer, int32 InTileCount)
 		: TileOffsetsShaderRef(InTileOffsetsShaderRef)
 		, TileOffsetsBufferRef(InTileOffsetsBufferRef)
 		, UniformBuffer(InUniformBuffer)
@@ -1381,7 +1223,7 @@ void ExecuteSimulationCommands(
 	ERHIFeatureLevel::Type FeatureLevel,
 	const TArray<FSimulationCommandGPU>& SimulationCommands,
 	FParticleSimulationResources* ParticleSimulationResources,
-	FRHIUniformBuffer* ViewUniformBuffer,
+	const TUniformBufferRef<FViewUniformShaderParameters>& ViewUniformBuffer,
 	const FGlobalDistanceFieldParameterData* GlobalDistanceFieldParameterData,
 	FRHIUniformBuffer* SceneTexturesUniformBuffer,
 	bool bUseFixDT)
@@ -1400,6 +1242,10 @@ void ExecuteSimulationCommands(
 	{
 		StaticUniformBuffers.AddUniformBuffer(SceneTexturesUniformBuffer);
 	}
+	if (ViewUniformBuffer)
+	{
+		StaticUniformBuffers.AddUniformBuffer(ViewUniformBuffer);
+	}
 	SCOPED_UNIFORM_BUFFER_STATIC_BINDINGS(RHICmdList, StaticUniformBuffers);
 
 	const float FixDeltaSeconds = CVarGPUParticleFixDeltaSeconds.GetValueOnRenderThread();
@@ -1411,6 +1257,44 @@ void ExecuteSimulationCommands(
 	TShaderMapRef<FParticleTileVS> VertexShader(GetGlobalShaderMap(FeatureLevel));
 	TShaderMapRef<TParticleSimulationPS<CollisionMode> > PixelShader(GetGlobalShaderMap(FeatureLevel));
 
+	FParticleSimulationPsParameters PsParameters{};
+	{
+		FRHISamplerState* SamplerStatePoint = TStaticSamplerState<SF_Point>::GetRHI();
+		FRHISamplerState* SamplerStateLinear = TStaticSamplerState<SF_Bilinear, AM_Clamp, AM_Clamp, AM_Clamp>::GetRHI();
+
+		PsParameters.PositionTexture = TextureResources.PositionTextureRHI;
+		PsParameters.PositionTextureSampler = SamplerStatePoint;
+
+		PsParameters.VelocityTexture = TextureResources.VelocityTextureRHI;
+		PsParameters.VelocityTextureSampler = SamplerStatePoint;
+
+		PsParameters.AttributesTexture = AttributeTexture.TextureRHI;
+		PsParameters.AttributesTextureSampler = SamplerStatePoint;
+
+		PsParameters.CurveTexture = GParticleCurveTexture.GetCurveTexture();
+		PsParameters.CurveTextureSampler = SamplerStateLinear;
+
+		PsParameters.VectorFieldTexturesSampler = SamplerStateLinear;
+
+		if (CollisionMode != PCM_None)
+		{
+			PsParameters.RenderAttributesTexture = RenderAttributeTexture.TextureRHI;
+			PsParameters.RenderAttributesTextureSampler = SamplerStatePoint;
+
+			check(ViewUniformBuffer != NULL);
+			PsParameters.View = ViewUniformBuffer;
+
+			if (CollisionMode == PCM_DepthBuffer)
+			{
+				PsParameters.CollisionDepthBounds = FXConsoleVariables::GPUCollisionDepthBounds;
+			}
+			else if (CollisionMode == PCM_DistanceField)
+			{
+				PsParameters.GlobalDistanceFieldParameters = SetupGlobalDistanceFieldParameters_Minimal(*GlobalDistanceFieldParameterData);
+			}
+		}
+	}
+
 	GraphicsPSOInit.BoundShaderState.VertexDeclarationRHI = GetParticleTileVertexDeclaration(FeatureLevel);
 	GraphicsPSOInit.BoundShaderState.VertexShaderRHI = VertexShader.GetVertexShader();
 	GraphicsPSOInit.BoundShaderState.PixelShaderRHI = PixelShader.GetPixelShader();
@@ -1418,20 +1302,25 @@ void ExecuteSimulationCommands(
 
 	SetGraphicsPipelineState(RHICmdList, GraphicsPSOInit, 0);
 
-	PixelShader->SetParameters(RHICmdList, TextureResources, AttributeTexture, RenderAttributeTexture, ViewUniformBuffer, FeatureLevel, GlobalDistanceFieldParameterData);
-
 	// Draw tiles to perform the simulation step.
 	const int32 CommandCount = SimulationCommands.Num();
 	for (int32 CommandIndex = 0; CommandIndex < CommandCount; ++CommandIndex)
 	{
 		const FSimulationCommandGPU& Command = SimulationCommands[CommandIndex];
-		VertexShader->SetParameters(RHICmdList, Command.TileOffsetsShaderRef);
-		PixelShader->SetInstanceParameters(RHICmdList, Command.UniformBuffer, Command.PerFrameParameters, bUseFixDT);
-		PixelShader->SetVectorFieldParameters(
-			RHICmdList, 
-			Command.VectorFieldsUniformBuffer,
-			Command.VectorFieldTexturesRHI
-		);
+
+		const FParticleTileVS::FParameters VsParameters = VertexShader->GetParameters(Command.TileOffsetsShaderRef);
+
+		PsParameters.PerFrameParameters = GetParticlePerFrameSimulationShaderParameters(Command.PerFrameParameters, bUseFixDT);
+		for (int32 i = 0; i < MAX_VECTOR_FIELDS; ++i)
+		{
+			PsParameters.VectorFieldTextures[i] = Command.VectorFieldTexturesRHI[i];
+		}
+
+		PsParameters.ParticleSimulationParameters = Command.UniformBuffer;
+		PsParameters.VectorFieldUniformParameters = Command.VectorFieldsUniformBuffer;
+
+		SetShaderParameters(RHICmdList, VertexShader, VertexShader.GetVertexShader(), VsParameters);
+		SetShaderParameters(RHICmdList, PixelShader, PixelShader.GetPixelShader(), PsParameters);
 		
 		if (FeatureLevel <= ERHIFeatureLevel::ES3_1)
 		{
@@ -1445,7 +1334,7 @@ void ExecuteSimulationCommands(
 	}
 
 	// Unbind input buffers.
-	PixelShader->UnbindBuffers(RHICmdList);
+	UnsetShaderSRVs(RHICmdList, PixelShader, PixelShader.GetPixelShader());
 }
 
 
@@ -1455,7 +1344,7 @@ void ExecuteSimulationCommands(
 	ERHIFeatureLevel::Type FeatureLevel,
 	const TArray<FSimulationCommandGPU>& SimulationCommands,
 	FParticleSimulationResources* ParticleSimulationResources,
-	FRHIUniformBuffer* ViewUniformBuffer,
+	const TUniformBufferRef<FViewUniformShaderParameters>& ViewUniformBuffer,
 	const FGlobalDistanceFieldParameterData* GlobalDistanceFieldParameterData,
 	FRHIUniformBuffer* SceneTexturesUniformBuffer,
 	EParticleSimulatePhase::Type Phase,
@@ -1495,7 +1384,7 @@ void ExecuteSimulationCommands(
 			FeatureLevel,
 			SimulationCommands,
 			ParticleSimulationResources,
-			NULL,
+			{},
 			GlobalDistanceFieldParameterData,
 			SceneTexturesUniformBuffer,
 			bUseFixDT);
@@ -1546,14 +1435,20 @@ void ClearTiles(FRHICommandList& RHICmdList, FGraphicsPipelineStateInitializer&
 		if (FeatureLevel <= ERHIFeatureLevel::ES3_1)
 		{
 			BuildTileVertexBuffer(BufferParam, TilesPtr, TilesThisDrawCall, TilesThisDrawCall);
-			VertexShader->SetParameters(RHICmdList, ShaderParam);
+
+			const FParticleTileVS::FParameters VsParameters = VertexShader->GetParameters(ShaderParam);
+			SetShaderParameters(RHICmdList, VertexShader, VertexShader.GetVertexShader(), VsParameters);
+
 			DrawParticleTiles(RHICmdList, BufferParam, TilesThisDrawCall);
 		}
 		else
 		{
 			const int32 AlignedTilesThisDrawCall = ComputeAlignedTileCount(TilesThisDrawCall);
 			BuildTileVertexBuffer(BufferParam, TilesPtr, TilesThisDrawCall, AlignedTilesThisDrawCall);
-			VertexShader->SetParameters(RHICmdList, ShaderParam);
+
+			const FParticleTileVS::FParameters VsParameters = VertexShader->GetParameters(ShaderParam);
+			SetShaderParameters(RHICmdList, VertexShader, VertexShader.GetVertexShader(), VsParameters);
+
 			DrawAlignedParticleTiles(RHICmdList, AlignedTilesThisDrawCall);
 		}
 		
@@ -1804,34 +1699,18 @@ typedef TUniformBufferRef<FParticleSimVisualizeParameters> FParticleSimVisualize
  */
 class FParticleSimVisualizeVS : public FGlobalShader
 {
-	DECLARE_SHADER_TYPE(FParticleSimVisualizeVS,Global);
-
 public:
+	DECLARE_GLOBAL_SHADER(FParticleSimVisualizeVS);
+	SHADER_USE_PARAMETER_STRUCT(FParticleSimVisualizeVS, FGlobalShader)
+
+	BEGIN_SHADER_PARAMETER_STRUCT(FParameters, )
+		SHADER_PARAMETER_STRUCT_REF(FParticleSimVisualizeParameters, ParticleSimVisualizeParameters)
+	END_GLOBAL_SHADER_PARAMETER_STRUCT()
 
 	static bool ShouldCompilePermutation(const FGlobalShaderPermutationParameters& Parameters)
 	{
 		return SupportsGPUParticles(Parameters.Platform);
 	}
-
-	/** Default constructor. */
-	FParticleSimVisualizeVS()
-	{
-	}
-
-	/** Initialization constructor. */
-	explicit FParticleSimVisualizeVS( const ShaderMetaType::CompiledShaderInitializerType& Initializer )
-		: FGlobalShader(Initializer)
-	{
-	}
-
-	/**
-	 * Set parameters for this shader.
-	 */
-	void SetParameters(FRHICommandList& RHICmdList, const FParticleSimVisualizeBufferRef& UniformBuffer )
-	{
-		FRHIVertexShader* VertexShader = RHICmdList.GetBoundVertexShader();
-		SetUniformBufferParameter(RHICmdList, VertexShader, GetUniformBufferParameter<FParticleSimVisualizeParameters>(), UniformBuffer );
-	}
 };
 
 /**
@@ -1839,50 +1718,22 @@ public:
  */
 class FParticleSimVisualizePS : public FGlobalShader
 {
-	DECLARE_SHADER_TYPE(FParticleSimVisualizePS,Global);
-
 public:
+	DECLARE_GLOBAL_SHADER(FParticleSimVisualizePS);
+	SHADER_USE_PARAMETER_STRUCT(FParticleSimVisualizePS, FGlobalShader)
+
+	BEGIN_SHADER_PARAMETER_STRUCT(FParameters, )
+		SHADER_PARAMETER(int32, VisualizationMode)
+		SHADER_PARAMETER_TEXTURE(Texture2D, PositionTexture)
+		SHADER_PARAMETER_SAMPLER(SamplerState, PositionTextureSampler)
+		SHADER_PARAMETER_TEXTURE(Texture2D, CurveTexture)
+		SHADER_PARAMETER_SAMPLER(SamplerState, CurveTextureSampler)
+	END_GLOBAL_SHADER_PARAMETER_STRUCT()
 
 	static bool ShouldCompilePermutation(const FGlobalShaderPermutationParameters& Parameters)
 	{
 		return SupportsGPUParticles(Parameters.Platform);
 	}
-
-	/** Default constructor. */
-	FParticleSimVisualizePS()
-	{
-	}
-
-	/** Initialization constructor. */
-	explicit FParticleSimVisualizePS( const ShaderMetaType::CompiledShaderInitializerType& Initializer )
-		: FGlobalShader(Initializer)
-	{
-		VisualizationMode.Bind( Initializer.ParameterMap, TEXT("VisualizationMode") );
-		PositionTexture.Bind( Initializer.ParameterMap, TEXT("PositionTexture") );
-		PositionTextureSampler.Bind( Initializer.ParameterMap, TEXT("PositionTextureSampler") );
-		CurveTexture.Bind( Initializer.ParameterMap, TEXT("CurveTexture") );
-		CurveTextureSampler.Bind( Initializer.ParameterMap, TEXT("CurveTextureSampler") );
-	}
-
-	/**
-	 * Set parameters for this shader.
-	 */
-	void SetParameters(FRHICommandList& RHICmdList, int32 InVisualizationMode, FRHITexture2D* PositionTextureRHI, FRHITexture2D* CurveTextureRHI )
-	{
-		FRHIPixelShader* PixelShader = RHICmdList.GetBoundPixelShader();
-		SetShaderValue(RHICmdList, PixelShader, VisualizationMode, InVisualizationMode );
-		FRHISamplerState* SamplerStatePoint = TStaticSamplerState<SF_Point>::GetRHI();
-		SetTextureParameter(RHICmdList, PixelShader, PositionTexture, PositionTextureSampler, SamplerStatePoint, PositionTextureRHI );
-		SetTextureParameter(RHICmdList, PixelShader, CurveTexture, CurveTextureSampler, SamplerStatePoint, CurveTextureRHI );
-	}
-
-private:
-
-	LAYOUT_FIELD(FShaderParameter, VisualizationMode);
-	LAYOUT_FIELD(FShaderResourceParameter, PositionTexture);
-	LAYOUT_FIELD(FShaderResourceParameter, PositionTextureSampler);
-	LAYOUT_FIELD(FShaderResourceParameter, CurveTexture);
-	LAYOUT_FIELD(FShaderResourceParameter, CurveTextureSampler);
 };
 
 /** Implementation for all shaders used for visualization. */
@@ -1964,14 +1815,31 @@ static void VisualizeGPUSimulation(
 		SetGraphicsPipelineState(RHICmdList, GraphicsPSOInit, 0);
 
 		// Parameters for the visualization.
-		FParticleSimVisualizeParameters Parameters;
-		Parameters.ScaleBias.X = 2.0f * DisplaySizeX / (float)TargetSize.X;
-		Parameters.ScaleBias.Y = 2.0f * DisplaySizeY / (float)TargetSize.Y;
-		Parameters.ScaleBias.Z = 2.0f * DisplayOffsetX / (float)TargetSize.X - 1.0f;
-		Parameters.ScaleBias.W = 2.0f * DisplayOffsetY / (float)TargetSize.Y - 1.0f;
-		FParticleSimVisualizeBufferRef UniformBuffer = FParticleSimVisualizeBufferRef::CreateUniformBufferImmediate(Parameters, UniformBuffer_SingleDraw);
-		VertexShader->SetParameters(RHICmdList, UniformBuffer);
-		PixelShader->SetParameters(RHICmdList, VisualizationMode, StateTextures.PositionTextureRHI, CurveTextureRHI);
+		{
+			FParticleSimVisualizeParameters Parameters;
+			Parameters.ScaleBias.X = 2.0f * DisplaySizeX / (float)TargetSize.X;
+			Parameters.ScaleBias.Y = 2.0f * DisplaySizeY / (float)TargetSize.Y;
+			Parameters.ScaleBias.Z = 2.0f * DisplayOffsetX / (float)TargetSize.X - 1.0f;
+			Parameters.ScaleBias.W = 2.0f * DisplayOffsetY / (float)TargetSize.Y - 1.0f;
+
+			FParticleSimVisualizeVS::FParameters VsParameters{};
+			VsParameters.ParticleSimVisualizeParameters = FParticleSimVisualizeBufferRef::CreateUniformBufferImmediate(Parameters, UniformBuffer_SingleDraw);
+
+			SetShaderParameters(RHICmdList, VertexShader, VertexShader.GetVertexShader(), VsParameters);
+		}
+
+		{
+			FRHISamplerState* SamplerStatePoint = TStaticSamplerState<SF_Point>::GetRHI();
+
+			FParticleSimVisualizePS::FParameters PsParameters{};
+			PsParameters.VisualizationMode = VisualizationMode;
+			PsParameters.PositionTexture = StateTextures.PositionTextureRHI;
+			PsParameters.PositionTextureSampler = SamplerStatePoint;
+			PsParameters.CurveTexture = CurveTextureRHI;
+			PsParameters.CurveTextureSampler = SamplerStatePoint;
+
+			SetShaderParameters(RHICmdList, PixelShader, PixelShader.GetPixelShader(), PsParameters);
+		}
 
 		const int32 VertexStride = sizeof(FVector2f);
 
@@ -2063,9 +1931,20 @@ typedef TUniformBufferRef<FParticleBoundsParameters> FParticleBoundsUniformBuffe
  */
 class FParticleBoundsCS : public FGlobalShader
 {
-	DECLARE_SHADER_TYPE(FParticleBoundsCS,Global);
-
 public:
+	DECLARE_GLOBAL_SHADER(FParticleBoundsCS);
+	SHADER_USE_PARAMETER_STRUCT(FParticleBoundsCS, FGlobalShader)
+
+	BEGIN_SHADER_PARAMETER_STRUCT(FParameters, )
+		SHADER_PARAMETER(int32, TextureSizeX)
+		SHADER_PARAMETER(int32, TextureSizeY)
+
+		SHADER_PARAMETER_UAV(RWBuffer<float4>, OutBounds)
+		SHADER_PARAMETER_SRV(Buffer<float2>, InParticleIndices)
+		SHADER_PARAMETER_TEXTURE(Texture2D, PositionTexture)
+
+		SHADER_PARAMETER_STRUCT_REF(FParticleBoundsParameters, ParticleBoundsParameters)
+	END_SHADER_PARAMETER_STRUCT()
 
 	static bool ShouldCompilePermutation(const FGlobalShaderPermutationParameters& Parameters)
 	{
@@ -2078,73 +1957,6 @@ public:
 		OutEnvironment.SetDefine( TEXT("THREAD_COUNT"), PARTICLE_BOUNDS_THREADS );
 		OutEnvironment.CompilerFlags.Add( CFLAG_StandardOptimization );
 	}
-
-	/** Default constructor. */
-	FParticleBoundsCS()
-	{
-	}
-
-	/** Initialization constructor. */
-	explicit FParticleBoundsCS( const ShaderMetaType::CompiledShaderInitializerType& Initializer )
-		: FGlobalShader(Initializer)
-	{
-		InParticleIndices.Bind( Initializer.ParameterMap, TEXT("InParticleIndices") );
-		PositionTexture.Bind( Initializer.ParameterMap, TEXT("PositionTexture") );
-		PositionTextureSampler.Bind( Initializer.ParameterMap, TEXT("PositionTextureSampler") );
-		OutBounds.Bind( Initializer.ParameterMap, TEXT("OutBounds") );
-		TextureSizeX.Bind(Initializer.ParameterMap, TEXT("TextureSizeX"));
-		TextureSizeY.Bind(Initializer.ParameterMap, TEXT("TextureSizeY"));
-	}
-
-	/**
-	 * Set output buffers for this shader.
-	 */
-	void SetOutput(FRHICommandList& RHICmdList, FRHIUnorderedAccessView* OutBoundsUAV )
-	{
-		FRHIComputeShader* ComputeShaderRHI = RHICmdList.GetBoundComputeShader();
-		SetUAVParameter(RHICmdList, ComputeShaderRHI, OutBounds, OutBoundsUAV);
-	}
-
-	/**
-	 * Set input parameters.
-	 */
-	void SetParameters(
-		FRHICommandList& RHICmdList,
-		FParticleBoundsUniformBufferRef& UniformBuffer,
-		FRHIShaderResourceView* InIndicesSRV,
-		FRHITexture2D* PositionTextureRHI
-		)
-	{
-		FRHIComputeShader* ComputeShaderRHI = RHICmdList.GetBoundComputeShader();
-		SetUniformBufferParameter(RHICmdList, ComputeShaderRHI, GetUniformBufferParameter<FParticleBoundsParameters>(), UniformBuffer );
-
-		SetSRVParameter(RHICmdList, ComputeShaderRHI, InParticleIndices, InIndicesSRV);
-		SetTextureParameter(RHICmdList, ComputeShaderRHI, PositionTexture, PositionTextureRHI);
-
-		SetShaderValue(RHICmdList, ComputeShaderRHI, TextureSizeX, GParticleSimulationTextureSizeX);
-		SetShaderValue(RHICmdList, ComputeShaderRHI, TextureSizeY, GParticleSimulationTextureSizeY);
-	}
-
-	/**
-	 * Unbinds any buffers that have been bound.
-	 */
-	void UnbindBuffers(FRHICommandList& RHICmdList)
-	{
-		FRHIComputeShader* ComputeShaderRHI = RHICmdList.GetBoundComputeShader();
-		SetSRVParameter(RHICmdList, ComputeShaderRHI, InParticleIndices, nullptr);
-		SetUAVParameter(RHICmdList, ComputeShaderRHI, OutBounds, nullptr);
-	}
-
-private:
-	/** Input buffer containing particle indices. */
-	LAYOUT_FIELD(FShaderResourceParameter, InParticleIndices);
-	/** Texture containing particle positions. */
-	LAYOUT_FIELD(FShaderResourceParameter, PositionTexture);
-	LAYOUT_FIELD(FShaderResourceParameter, PositionTextureSampler);
-	/** Output key buffer. */
-	LAYOUT_FIELD(FShaderResourceParameter, OutBounds);
-	LAYOUT_FIELD(FShaderParameter, TextureSizeX);
-	LAYOUT_FIELD(FShaderParameter, TextureSizeY);
 };
 IMPLEMENT_SHADER_TYPE(,FParticleBoundsCS,TEXT("/Engine/Private/ParticleBoundsShader.usf"),TEXT("ComputeParticleBounds"),SF_Compute);
 
@@ -2171,8 +1983,7 @@ static FBox ComputeParticleBounds(
 	int32 ParticleCount )
 {
 	FBox BoundingBox;
-	FParticleBoundsParameters Parameters;
-	FParticleBoundsUniformBufferRef UniformBuffer;
+	FParticleBoundsParameters ParticleBoundsParameters;
 
 	if (ParticleCount > 0 && FeatureLevel >= ERHIFeatureLevel::SM5)
 	{
@@ -2183,10 +1994,9 @@ static FBox ComputeParticleBounds(
 		const uint32 GroupCount = FMath::Clamp<uint32>( ChunkCount, 1, MaxGroupCount );
 
 		// Create the uniform buffer.
-		Parameters.ChunksPerGroup = ChunkCount / GroupCount;
-		Parameters.ExtraChunkCount = ChunkCount % GroupCount;
-		Parameters.ParticleCount = ParticleCount;
-		UniformBuffer = FParticleBoundsUniformBufferRef::CreateUniformBufferImmediate( Parameters, UniformBuffer_SingleFrame );
+		ParticleBoundsParameters.ChunksPerGroup = ChunkCount / GroupCount;
+		ParticleBoundsParameters.ExtraChunkCount = ChunkCount % GroupCount;
+		ParticleBoundsParameters.ParticleCount = ParticleCount;
 
 		// Create a buffer for storing bounds.
 		const int32 BufferSize = GroupCount * 2 * sizeof(FVector4f);
@@ -2203,16 +2013,29 @@ static FBox ComputeParticleBounds(
 		TShaderMapRef<FParticleBoundsCS> ParticleBoundsCS(GetGlobalShaderMap(FeatureLevel));
 		SetComputePipelineState(RHICmdList, ParticleBoundsCS.GetComputeShader());
 
+		{
+			FParticleBoundsCS::FParameters Parameters{};
+			Parameters.TextureSizeX = GParticleSimulationTextureSizeX;
+			Parameters.TextureSizeY = GParticleSimulationTextureSizeY;
+			Parameters.OutBounds = BoundsVertexBufferUAV;
+			Parameters.InParticleIndices = VertexBufferSRV;
+			Parameters.PositionTexture = PositionTextureRHI;
+			Parameters.ParticleBoundsParameters = FParticleBoundsUniformBufferRef::CreateUniformBufferImmediate(ParticleBoundsParameters, UniformBuffer_SingleFrame);
+
+			SetShaderParameters(RHICmdList, ParticleBoundsCS, ParticleBoundsCS.GetComputeShader(), Parameters);
+		}
+
 		// Dispatch shader to compute bounds.
-		ParticleBoundsCS->SetOutput(RHICmdList, BoundsVertexBufferUAV);
-		ParticleBoundsCS->SetParameters(RHICmdList, UniformBuffer, VertexBufferSRV, PositionTextureRHI);
+
 		DispatchComputeShader(
 			RHICmdList, 
 			ParticleBoundsCS.GetShader(), 
 			GroupCount,
 			1,
 			1 );
-		ParticleBoundsCS->UnbindBuffers(RHICmdList);
+
+		UnsetShaderSRVs(RHICmdList, ParticleBoundsCS, ParticleBoundsCS.GetComputeShader());
+		UnsetShaderUAVs(RHICmdList, ParticleBoundsCS, ParticleBoundsCS.GetComputeShader());
 
 		// Read back bounds.
 		FVector4f* GroupBounds = (FVector4f*)RHILockBuffer( BoundsVertexBufferRHI, 0, BufferSize, RLM_ReadOnly );
@@ -4618,7 +4441,7 @@ void FFXSystem::FinalizeGPUSimulation(FRHICommandListImmediate& RHICmdList)
 void FFXSystem::SimulateGPUParticles(
 	FRHICommandListImmediate& RHICmdList,
 	EParticleSimulatePhase::Type Phase,
-	FRHIUniformBuffer* ViewUniformBuffer,
+	const TUniformBufferRef<FViewUniformShaderParameters>& ViewUniformBuffer,
 	const FGlobalDistanceFieldParameterData* GlobalDistanceFieldParameterData
 	)
 {
diff --git a/Engine/Source/Runtime/Engine/Public/FXSystem.h b/Engine/Source/Runtime/Engine/Public/FXSystem.h
--- a/Engine/Source/Runtime/Engine/Public/FXSystem.h
+++ b/Engine/Source/Runtime/Engine/Public/FXSystem.h
@@ -264,7 +264,7 @@ public:
 	/** Get the shared SortManager, used in the rendering loop to call FGPUSortManager::OnPreRender() and FGPUSortManager::OnPostRenderOpaque() */
 	virtual FGPUSortManager* GetGPUSortManager() const = 0;
 
-	virtual void SetSceneTexturesUniformBuffer(FRHIUniformBuffer* InSceneTexturesUniformParams) {}
+	virtual void SetSceneTexturesUniformBuffer(const TUniformBufferRef<FSceneTextureUniformParameters>& InSceneTexturesUniformParams) {}
 
 	FORCEINLINE FScene* GetScene()const { return Scene; }
 	FORCEINLINE void SetScene(FScene* InScene) { Scene = InScene; }
diff --git a/Engine/Source/Runtime/RenderCore/Public/ShaderParameterStruct.h b/Engine/Source/Runtime/RenderCore/Public/ShaderParameterStruct.h
--- a/Engine/Source/Runtime/RenderCore/Public/ShaderParameterStruct.h
+++ b/Engine/Source/Runtime/RenderCore/Public/ShaderParameterStruct.h
@@ -193,6 +193,23 @@ inline void UnsetShaderSRVs(TRHICmdList& RHICmdList, const TShaderRef<TShaderCla
 	RHICmdList.SetShaderParameters(ShaderRHI, {}, {}, ResourceSRVsToUnset, BindlessSRVsToUnset);
 }
 
+/** Unset compute shader SRVs. */
+template<typename TRHICmdList, typename TShaderClass>
+inline void UnsetShaderSRVs(TRHICmdList& RHICmdList, const TShaderRef<TShaderClass>& Shader, FRHIGraphicsShader* ShaderRHI)
+{
+	const FShaderParameterBindings& Bindings = Shader->Bindings;
+
+	checkf(Bindings.RootParameterBufferIndex == FShaderParameterBindings::kInvalidBufferIndex, TEXT("Can't use UnsetShaderSRVs() for root parameter buffer index."));
+
+	TArray<FRHIShaderParameterResource, TInlineAllocator<16>> ResourceSRVsToUnset;
+	TArray<FRHIShaderParameterResource, TInlineAllocator<16>> BindlessSRVsToUnset;
+
+	CollectSRVsToUnset(ResourceSRVsToUnset, Bindings.ResourceParameters);
+	CollectSRVsToUnset(BindlessSRVsToUnset, Bindings.BindlessResourceParameters);
+
+	RHICmdList.SetShaderParameters(ShaderRHI, {}, {}, ResourceSRVsToUnset, BindlessSRVsToUnset);
+}
+
 RENDERCORE_API void SetShaderParameters(
 	FRHIComputeCommandList& RHICmdList,
 	FRHIComputeShader* ShaderRHI,
diff --git a/Engine/Source/Runtime/Renderer/Private/DeferredShadingRenderer.cpp b/Engine/Source/Runtime/Renderer/Private/DeferredShadingRenderer.cpp
--- a/Engine/Source/Runtime/Renderer/Private/DeferredShadingRenderer.cpp
+++ b/Engine/Source/Runtime/Renderer/Private/DeferredShadingRenderer.cpp
@@ -628,7 +628,7 @@ static void RenderOpaqueFX(
 		// Clear the scene textures UB pointer on the FX system. Use the same pass parameters to extend resource lifetimes.
 		GraphBuilder.AddPass(RDG_EVENT_NAME("UnsetSceneTexturesUniformBuffer"), ExtractUBPassParameters, UBPassFlags, [FXSystem](FRHICommandListImmediate&)
 		{
-			FXSystem->SetSceneTexturesUniformBuffer(nullptr);
+			FXSystem->SetSceneTexturesUniformBuffer({});
 		});
 
 		if (FGPUSortManager* GPUSortManager = FXSystem->GetGPUSortManager())
diff --git a/Engine/Source/Runtime/Renderer/Public/GlobalDistanceFieldParameters.h b/Engine/Source/Runtime/Renderer/Public/GlobalDistanceFieldParameters.h
--- a/Engine/Source/Runtime/Renderer/Public/GlobalDistanceFieldParameters.h
+++ b/Engine/Source/Runtime/Renderer/Public/GlobalDistanceFieldParameters.h
@@ -69,7 +69,36 @@ END_SHADER_PARAMETER_STRUCT()
 
 FGlobalDistanceFieldParameters2 SetupGlobalDistanceFieldParameters(const FGlobalDistanceFieldParameterData& ParameterData);
 
-class FGlobalDistanceFieldParameters
+inline FGlobalDistanceFieldParameters2 SetupGlobalDistanceFieldParameters_Minimal(const FGlobalDistanceFieldParameterData& ParameterData)
+{
+	FGlobalDistanceFieldParameters2 ShaderParameters{};
+
+	ShaderParameters.GlobalDistanceFieldPageAtlasTexture = ParameterData.PageAtlasTexture ? ParameterData.PageAtlasTexture : GBlackVolumeTexture->TextureRHI.GetReference();
+	ShaderParameters.GlobalDistanceFieldPageTableTexture = ParameterData.PageTableTexture ? ParameterData.PageTableTexture : GBlackUintVolumeTexture->TextureRHI.GetReference();
+	ShaderParameters.GlobalDistanceFieldMipTexture = ParameterData.MipTexture ? ParameterData.MipTexture : GBlackVolumeTexture->TextureRHI.GetReference();
+
+	for (int32 Index = 0; Index < GlobalDistanceField::MaxClipmaps; Index++)
+	{
+		ShaderParameters.GlobalVolumeTranslatedCenterAndExtent[Index] = ParameterData.TranslatedCenterAndExtent[Index];
+		ShaderParameters.GlobalVolumeTranslatedWorldToUVAddAndMul[Index] = ParameterData.TranslatedWorldToUVAddAndMul[Index];
+		ShaderParameters.GlobalDistanceFieldMipTranslatedWorldToUVScale[Index] = ParameterData.MipTranslatedWorldToUVScale[Index];
+		ShaderParameters.GlobalDistanceFieldMipTranslatedWorldToUVBias[Index] = ParameterData.MipTranslatedWorldToUVBias[Index];
+	}
+
+	ShaderParameters.GlobalDistanceFieldMipFactor = ParameterData.MipFactor;
+	ShaderParameters.GlobalDistanceFieldMipTransition = ParameterData.MipTransition;
+	ShaderParameters.GlobalDistanceFieldClipmapSizeInPages = ParameterData.ClipmapSizeInPages;
+	ShaderParameters.GlobalDistanceFieldInvPageAtlasSize = (FVector3f)ParameterData.InvPageAtlasSize;
+	ShaderParameters.GlobalDistanceFieldInvCoverageAtlasSize = (FVector3f)ParameterData.InvCoverageAtlasSize;
+	ShaderParameters.GlobalVolumeDimension = ParameterData.GlobalDFResolution;
+	ShaderParameters.GlobalVolumeTexelSize = 1.0f / ParameterData.GlobalDFResolution;
+	ShaderParameters.MaxGlobalDFAOConeDistance = ParameterData.MaxDFAOConeDistance;
+	ShaderParameters.NumGlobalSDFClipmaps = ParameterData.NumGlobalSDFClipmaps;
+
+	return ShaderParameters;
+}
+
+class UE_DEPRECATED(5.2, "FGlobalDistanceFieldParameters2 should be used from now on.") FGlobalDistanceFieldParameters
 {
 	DECLARE_INLINE_TYPE_LAYOUT(FGlobalDistanceFieldParameters, NonVirtual);
 public:
