Fixing shader warning caused by a bad SRV type.

diff --git a/Engine/Source/Runtime/Engine/Private/Particles/ParticleGpuSimulation.cpp b/Engine/Source/Runtime/Engine/Private/Particles/ParticleGpuSimulation.cpp
--- a/Engine/Source/Runtime/Engine/Private/Particles/ParticleGpuSimulation.cpp
+++ b/Engine/Source/Runtime/Engine/Private/Particles/ParticleGpuSimulation.cpp
@@ -912,7 +912,7 @@ public:
 	BEGIN_SHADER_PARAMETER_STRUCT(FParameters, )
 		SHADER_PARAMETER(float, TileSizeX)
 		SHADER_PARAMETER(float, TileSizeY)
-		SHADER_PARAMETER_SRV(FShaderResourceParameter, TileOffsets)
+		SHADER_PARAMETER_SRV(Buffer<float2>, TileOffsets)
 	END_SHADER_PARAMETER_STRUCT()
 
 	FParameters GetParameters(FParticleShaderParamRef TileOffsetsRef)
