Adding a overridable method to WorldSettings to get the Physics Handler Class

#rnx

diff --git a/Engine/Source/Runtime/Engine/Classes/GameFramework/WorldSettings.h b/Engine/Source/Runtime/Engine/Classes/GameFramework/WorldSettings.h
--- a/Engine/Source/Runtime/Engine/Classes/GameFramework/WorldSettings.h
+++ b/Engine/Source/Runtime/Engine/Classes/GameFramework/WorldSettings.h
@@ -653,6 +653,9 @@ public:
 	UPROPERTY(EditAnywhere, noclear, BlueprintReadOnly, Category=Physics, AdvancedDisplay)
 	TSubclassOf<class ADefaultPhysicsVolume> DefaultPhysicsVolumeClass;
 
+	// Method to get the Physics Collision Handler Class
+	virtual TSubclassOf<class UPhysicsCollisionHandler> GetPhysicsCollisionHandlerClass() { return PhysicsCollisionHandlerClass; }
+
 	// optional level specific collision handler
 	UPROPERTY(EditAnywhere, BlueprintReadOnly, Category=Physics, AdvancedDisplay)
 	TSubclassOf<class UPhysicsCollisionHandler>	PhysicsCollisionHandlerClass;
diff --git a/Engine/Source/Runtime/Engine/Private/World.cpp b/Engine/Source/Runtime/Engine/Private/World.cpp
--- a/Engine/Source/Runtime/Engine/Private/World.cpp
+++ b/Engine/Source/Runtime/Engine/Private/World.cpp
@@ -1945,7 +1945,7 @@ void UWorld::InitWorld(const InitializationValues IVS)
 	if (IVS.bCreatePhysicsScene)
 	{
 		// First look for world override
-		TSubclassOf<UPhysicsCollisionHandler> PhysHandlerClass = (WorldSettings ? WorldSettings->PhysicsCollisionHandlerClass : nullptr);
+		TSubclassOf<UPhysicsCollisionHandler> PhysHandlerClass = (WorldSettings ? WorldSettings->GetPhysicsCollisionHandlerClass() : nullptr);
 		// Then fall back to engine default
 		if(PhysHandlerClass == NULL)
 		{
