physics object blueprint library: finding closest physics object should not check disabled state since otherwise we wouldn't be able to find the closest physics object on a non-simulating geometry collection

diff --git a/Engine/Source/Runtime/Engine/Private/PhysicsEngine/PhysicsObjectBlueprintLibrary.cpp b/Engine/Source/Runtime/Engine/Private/PhysicsEngine/PhysicsObjectBlueprintLibrary.cpp
--- a/Engine/Source/Runtime/Engine/Private/PhysicsEngine/PhysicsObjectBlueprintLibrary.cpp
+++ b/Engine/Source/Runtime/Engine/Private/PhysicsEngine/PhysicsObjectBlueprintLibrary.cpp
@@ -17,7 +17,7 @@ FClosestPhysicsObjectResult UPhysicsObjectBlueprintLibrary::GetClosestPhysicsObj
 		// Maybe allow users to customize the predicate?
 		[&Interface](Chaos::FPhysicsObject* Object) {
 			// We need to filter so we only get the leaf bodies. Mainly in consideration of the geometry collection.
-			return !Interface->AreAllDisabled({&Object, 1}) && !Interface->HasChildren(Object);
+			return !Interface->HasChildren(Object);
 		}
 	);
 	return Interface->GetClosestPhysicsBodyFromLocation(PhysicsObjects, WorldLocation);
