more correct check for UPrimitiveComponent::IsSimulatingPhysics to check the non-disabled particles on whether or not they're dynamic

diff --git a/Engine/Source/Runtime/Engine/Private/PrimitiveComponentPhysics.cpp b/Engine/Source/Runtime/Engine/Private/PrimitiveComponentPhysics.cpp
--- a/Engine/Source/Runtime/Engine/Private/PrimitiveComponentPhysics.cpp
+++ b/Engine/Source/Runtime/Engine/Private/PrimitiveComponentPhysics.cpp
@@ -900,17 +900,31 @@ bool UPrimitiveComponent::IsSimulatingPhysics(FName BoneName) const
 	{
 		return BodyInst->IsInstanceSimulatingPhysics();
 	}
-	else if (BoneName != NAME_None)
+	else
 	{
-		if (Chaos::FPhysicsObject* PhysicsObject = GetPhysicsObjectByName(BoneName))
+		TArray<Chaos::FPhysicsObjectHandle> PhysicsObjects;
+		if (BoneName != NAME_None)
 		{
-			return true;
+			PhysicsObjects.Add(GetPhysicsObjectByName(BoneName));
 		}
-	}
-	else
-	{
-		TArray<Chaos::FPhysicsObject*> PhysicsObjects = GetAllPhysicsObjects();
-		return !PhysicsObjects.IsEmpty();
+		else
+		{
+			PhysicsObjects = GetAllPhysicsObjects();
+		}
+
+		if (PhysicsObjects.IsEmpty())
+		{
+			return false;
+		}
+
+		FLockedReadPhysicsObjectExternalInterface Interface = FPhysicsObjectExternalInterface::LockRead(PhysicsObjects);
+		PhysicsObjects = PhysicsObjects.FilterByPredicate(
+			[&Interface](Chaos::FPhysicsObjectHandle Object) {
+				return !Interface->AreAllDisabled({ &Object, 1 });
+			}
+		);
+
+		return Interface->AreAllDynamic(PhysicsObjects);
 	}
 
 	return false;
