Fix LevelSequence resolving on Instanced WorldPartition cells

#rnx

diff --git a/Engine/Source/Runtime/Engine/Classes/Engine/LevelStreaming.h b/Engine/Source/Runtime/Engine/Classes/Engine/LevelStreaming.h
--- a/Engine/Source/Runtime/Engine/Classes/Engine/LevelStreaming.h
+++ b/Engine/Source/Runtime/Engine/Classes/Engine/LevelStreaming.h
@@ -534,6 +534,9 @@ public:
 
 	virtual UWorld* GetWorld() const override final;
 
+	/** Returns the UWorld that triggered the streaming of this streaming level. */
+	virtual UWorld* GetStreamingWorld() const;
+
 	/** Returns whether streaming level is visible */
 	UFUNCTION(BlueprintCallable, Category="Game")
 	bool IsLevelVisible() const;
diff --git a/Engine/Source/Runtime/Engine/Private/LevelInstance/LevelInstanceSubsystem.cpp b/Engine/Source/Runtime/Engine/Private/LevelInstance/LevelInstanceSubsystem.cpp
--- a/Engine/Source/Runtime/Engine/Private/LevelInstance/LevelInstanceSubsystem.cpp
+++ b/Engine/Source/Runtime/Engine/Private/LevelInstance/LevelInstanceSubsystem.cpp
@@ -513,7 +513,7 @@ ILevelInstanceInterface* ULevelInstanceSubsystem::GetOwningLevelInstance(const U
 		}
 		else if (UWorldPartitionLevelStreamingDynamic* WorldPartitionLevelStreaming = Cast<UWorldPartitionLevelStreamingDynamic>(BaseLevelStreaming))
 		{
-			return GetOwningLevelInstance(WorldPartitionLevelStreaming->GetOuterWorld()->PersistentLevel);
+			return GetOwningLevelInstance(WorldPartitionLevelStreaming->GetStreamingWorld()->PersistentLevel);
 		}
 	}
 
diff --git a/Engine/Source/Runtime/Engine/Private/LevelStreaming.cpp b/Engine/Source/Runtime/Engine/Private/LevelStreaming.cpp
--- a/Engine/Source/Runtime/Engine/Private/LevelStreaming.cpp
+++ b/Engine/Source/Runtime/Engine/Private/LevelStreaming.cpp
@@ -458,6 +458,12 @@ UWorld* ULevelStreaming::GetWorld() const
 	}
 }
 
+UWorld* ULevelStreaming::GetStreamingWorld() const
+{
+	check(!IsTemplate());
+	return GetWorld();
+}
+
 bool ULevelStreaming::IsLevelVisible() const
 {
 	return LoadedLevel && LoadedLevel->bIsVisible;
diff --git a/Engine/Source/Runtime/Engine/Private/WorldPartition/WorldPartitionLevelStreamingDynamic.cpp b/Engine/Source/Runtime/Engine/Private/WorldPartition/WorldPartitionLevelStreamingDynamic.cpp
--- a/Engine/Source/Runtime/Engine/Private/WorldPartition/WorldPartitionLevelStreamingDynamic.cpp
+++ b/Engine/Source/Runtime/Engine/Private/WorldPartition/WorldPartitionLevelStreamingDynamic.cpp
@@ -627,6 +627,14 @@ void UWorldPartitionLevelStreamingDynamic::Deactivate()
 	SetShouldBeVisible(false);
 }
 
+UWorld* UWorldPartitionLevelStreamingDynamic::GetStreamingWorld() const
+{
+	// For UWorldPartitionLevelStreamingDynamic the StreamingWorld is the world to which the OuterWorldPartition is outered.
+	// This World can be used to resolved SoftObjectPaths between cells.
+	check(OuterWorldPartition.IsValid());
+	return OuterWorldPartition->GetTypedOuter<UWorld>();
+}
+
 bool UWorldPartitionLevelStreamingDynamic::CanMakeVisible()
 {
 	const ENetMode NetMode = GetWorld()->GetNetMode();
@@ -683,12 +691,6 @@ bool UWorldPartitionLevelStreamingDynamic::CanMakeInvisible()
 	return Super::CanMakeInvisible();
 }
 
-UWorld* UWorldPartitionLevelStreamingDynamic::GetOuterWorld() const
-{
-	check(OuterWorldPartition.IsValid());
-	return OuterWorldPartition->GetTypedOuter<UWorld>();
-}
-
 bool UWorldPartitionLevelStreamingDynamic::ShouldBlockOnUnload() const
 {
 	if (Super::ShouldBlockOnUnload())
diff --git a/Engine/Source/Runtime/Engine/Public/WorldPartition/WorldPartitionLevelStreamingDynamic.h b/Engine/Source/Runtime/Engine/Public/WorldPartition/WorldPartitionLevelStreamingDynamic.h
--- a/Engine/Source/Runtime/Engine/Public/WorldPartition/WorldPartitionLevelStreamingDynamic.h
+++ b/Engine/Source/Runtime/Engine/Public/WorldPartition/WorldPartitionLevelStreamingDynamic.h
@@ -29,7 +29,7 @@ class ENGINE_API UWorldPartitionLevelStreamingDynamic : public ULevelStreamingDy
 	void Unload();
 	void Activate();
 	void Deactivate();
-	UWorld* GetOuterWorld() const;
+	UWorld* GetStreamingWorld() const override;
 	void SetShouldBeAlwaysLoaded(bool bInShouldBeAlwaysLoaded) { bShouldBeAlwaysLoaded = bInShouldBeAlwaysLoaded; }
 	const UWorldPartitionRuntimeCell* GetWorldPartitionRuntimeCell() const { return StreamingCell.Get(); }
 
diff --git a/Engine/Source/Runtime/LevelSequence/Private/LevelSequenceBindingReference.cpp b/Engine/Source/Runtime/LevelSequence/Private/LevelSequenceBindingReference.cpp
--- a/Engine/Source/Runtime/LevelSequence/Private/LevelSequenceBindingReference.cpp
+++ b/Engine/Source/Runtime/LevelSequence/Private/LevelSequenceBindingReference.cpp
@@ -67,8 +67,12 @@ UObject* FLevelSequenceBindingReference::Resolve(UObject* InContext, const FTopL
 			return nullptr;
 		}
 
-		// ExternalObjectPath.GetSubPathString() specifies the path from the package (so includes PersistentLevel.) so we must do a FindObject from its outer
-		return FindObject<UObject>(InContext->GetOuter(), *ExternalObjectPath.GetSubPathString());
+		UObject* ResolvedObject = nullptr;
+		// ExternalObjectPath.GetSubPathString() specifies the path from the package (so includes PersistentLevel.) so we must do a ResolveSubObject from its outer
+		UObject* ContextOuter = InContext->GetOuter();
+		check(ContextOuter);
+		ContextOuter->ResolveSubobject(*ExternalObjectPath.GetSubPathString(), ResolvedObject, /*bLoadIfExists*/false);
+		return ResolvedObject;
 	}
 	else
 	{
diff --git a/Engine/Source/Runtime/LevelSequence/Private/LevelSequencePlayer.cpp b/Engine/Source/Runtime/LevelSequence/Private/LevelSequencePlayer.cpp
--- a/Engine/Source/Runtime/LevelSequence/Private/LevelSequencePlayer.cpp
+++ b/Engine/Source/Runtime/LevelSequence/Private/LevelSequencePlayer.cpp
@@ -94,6 +94,17 @@ void ULevelSequencePlayer::Initialize(ULevelSequence* InLevelSequence, ULevel* I
 
 	// Construct the path to the level asset that the streamed level relates to
 	ULevelStreaming* LevelStreaming = FLevelUtils::FindStreamingLevel(InLevel);
+	if (LevelStreaming)
+	{
+		// All ULevelStreaming objects live in the owning world but if we are streaming a World Partition persistent level it will be returned as the StreamingWorld for all it's 
+		// ULevelStreaming cells. This streaming world should be used to resolve bindings.
+		if (UWorld* StreamingWorld = LevelStreaming->GetStreamingWorld(); StreamingWorld && (StreamingWorld != World))
+		{
+			Level = StreamingWorld->PersistentLevel;
+			LevelStreaming = FLevelUtils::FindStreamingLevel(Level.Get());
+		}
+	}
+		
 	if (LevelStreaming)
 	{
 		// StreamedLevelPackage is a package name of the form /Game/Folder/MapName, not a full asset path
