Add check for parallel game thread for SkinWeightVertexBuffer SetOwnerName

diff --git a/Engine/Source/Runtime/Engine/Private/Rendering/SkinWeightVertexBuffer.cpp b/Engine/Source/Runtime/Engine/Private/Rendering/SkinWeightVertexBuffer.cpp
--- a/Engine/Source/Runtime/Engine/Private/Rendering/SkinWeightVertexBuffer.cpp
+++ b/Engine/Source/Runtime/Engine/Private/Rendering/SkinWeightVertexBuffer.cpp
@@ -805,7 +805,7 @@ void FSkinWeightVertexBuffer::ResetVertexBoneWeights(uint32 VertexIndex)
 
 void FSkinWeightVertexBuffer::SetOwnerName(const FName& OwnerName)
 {
-	check(IsInGameThread());
+	check(IsInGameThread() || IsInParallelGameThread());
 	LookupVertexBuffer.SetOwnerName(OwnerName);
 	DataVertexBuffer.SetOwnerName(OwnerName);
 }
