Fix crash when pooling shader from DDC.

diff --git a/Engine/Source/Runtime/Engine/Private/ShaderCompiler/ShaderCompiler.cpp b/Engine/Source/Runtime/Engine/Private/ShaderCompiler/ShaderCompiler.cpp
--- a/Engine/Source/Runtime/Engine/Private/ShaderCompiler/ShaderCompiler.cpp
+++ b/Engine/Source/Runtime/Engine/Private/ShaderCompiler/ShaderCompiler.cpp
@@ -8728,7 +8728,7 @@ FShaderJobCache::FJobCachedOutput* FShaderJobCache::Find(const FJobInputHash& Ha
 
 					TRACE_COUNTER_ADD(Shaders_JobCacheDDCBytesReceived, Results.GetSize());
 
-					check(NewStoredOutput->JobOutput.GetAllocatedSize() == Results.GetSize());
+					check(Results.GetSize() <= NewStoredOutput->JobOutput.GetAllocatedSize());
 					FMemory::Memcpy(NewStoredOutput->JobOutput.GetData(),  Results.GetData(), Results.GetSize());
 
 					// Generate an output hash and cache the result in the FShaderJobCache
