Improved the accuracy of some comments about how the Support Ray Tracing setting on Skeletal Meshes affects their usage of the skin cache

diff --git a/Engine/Source/Runtime/Engine/Classes/Engine/RendererSettings.h b/Engine/Source/Runtime/Engine/Classes/Engine/RendererSettings.h
--- a/Engine/Source/Runtime/Engine/Classes/Engine/RendererSettings.h
+++ b/Engine/Source/Runtime/Engine/Classes/Engine/RendererSettings.h
@@ -992,7 +992,7 @@ class ENGINE_API URendererSettings : public UDeveloperSettings
 
 	UPROPERTY(config, EditAnywhere, Category = Optimizations, meta = (
 		ConsoleVariable = "r.SkinCache.DefaultBehavior", DisplayName = "Default Skin Cache Behavior",
-		ToolTip = "Default behavior if all skeletal meshes are included/excluded from the skin cache. If Support Ray Tracing is enabled on a mesh, will force inclusive behavior on that mesh."))
+		ToolTip = "Default behavior if all skeletal meshes are included/excluded from the skin cache. If Support Ray Tracing is enabled on a mesh, the skin cache will be used for Ray Tracing updates on that mesh regardless of this setting."))
 	ESkinCacheDefaultBehavior DefaultSkinCacheBehavior;
 
 	UPROPERTY(config, EditAnywhere, Category = Optimizations, meta = (
diff --git a/Engine/Source/Runtime/Engine/Classes/Engine/SkinnedAssetCommon.h b/Engine/Source/Runtime/Engine/Classes/Engine/SkinnedAssetCommon.h
--- a/Engine/Source/Runtime/Engine/Classes/Engine/SkinnedAssetCommon.h
+++ b/Engine/Source/Runtime/Engine/Classes/Engine/SkinnedAssetCommon.h
@@ -25,10 +25,10 @@ struct FSkeletalMeshLODGroupSettings;
 UENUM()
 enum class ESkinCacheUsage : uint8
 {
-	// Auto will defer to child or global behavior based on context. If Support Ray Tracing is enabled on the mesh, will imply Enabled
+	// Auto will defer to child or global behavior based on context
 	Auto		= 0,
 
-	// Mesh will not use the skin cache. If Support Ray Tracing is enabled on the mesh, will imply Enabled
+	// Mesh will not use the skin cache. However, if Support Ray Tracing is enabled on the mesh, the skin cache will still be used for Ray Tracing updates
 	Disabled	= uint8(-1),
 
 	// Mesh will use the skin cache
