Fixed crash in Virtual Texture upload path due to use of immediate command list.

diff --git a/Engine/Plugins/Experimental/GPULightmass/Source/GPULightmass/Private/LightmapPreviewVirtualTexture.cpp b/Engine/Plugins/Experimental/GPULightmass/Source/GPULightmass/Private/LightmapPreviewVirtualTexture.cpp
--- a/Engine/Plugins/Experimental/GPULightmass/Source/GPULightmass/Private/LightmapPreviewVirtualTexture.cpp
+++ b/Engine/Plugins/Experimental/GPULightmass/Source/GPULightmass/Private/LightmapPreviewVirtualTexture.cpp
@@ -61,6 +61,7 @@ FLightmapPreviewVirtualTexture::FLightmapPreviewVirtualTexture(FLightmapRenderSt
 }
 
 FVTRequestPageResult FLightmapPreviewVirtualTexture::RequestPageData(
+	FRHICommandList& RHICmdList,
 	const FVirtualTextureProducerHandle& InProducerHandle,
 	uint8 LayerMask,
 	uint8 vLevel,
diff --git a/Engine/Plugins/Experimental/GPULightmass/Source/GPULightmass/Private/LightmapPreviewVirtualTexture.h b/Engine/Plugins/Experimental/GPULightmass/Source/GPULightmass/Private/LightmapPreviewVirtualTexture.h
--- a/Engine/Plugins/Experimental/GPULightmass/Source/GPULightmass/Private/LightmapPreviewVirtualTexture.h
+++ b/Engine/Plugins/Experimental/GPULightmass/Source/GPULightmass/Private/LightmapPreviewVirtualTexture.h
@@ -15,6 +15,7 @@ public:
 	FLightmapPreviewVirtualTexture(FLightmapRenderStateRef LightmapRenderState, FLightmapRenderer* Renderer);
 
 	virtual FVTRequestPageResult RequestPageData(
+		FRHICommandList& RHICmdList,
 		const FVirtualTextureProducerHandle& ProducerHandle,
 		uint8 LayerMask,
 		uint8 vLevel,
diff --git a/Engine/Source/Runtime/Engine/Private/Texture2D.cpp b/Engine/Source/Runtime/Engine/Private/Texture2D.cpp
--- a/Engine/Source/Runtime/Engine/Private/Texture2D.cpp
+++ b/Engine/Source/Runtime/Engine/Private/Texture2D.cpp
@@ -1406,7 +1406,7 @@ void FVirtualTexture2DResource::InitializeEditorResources(IVirtualTexture* InVir
 			for (uint32 TileX = 0u; TileX < MipWidthInTiles; ++TileX)
 			{
 				const uint32 vAddress = FMath::MortonCode2(TileX) | (FMath::MortonCode2(TileY) << 1);
-				const FVTRequestPageResult RequestResult = InVirtualTexture->RequestPageData(ProducerHandle, LayerMask, MipLevel, vAddress, EVTRequestPagePriority::High);
+				const FVTRequestPageResult RequestResult = InVirtualTexture->RequestPageData(FRHICommandListExecutor::GetImmediateCommandList(), ProducerHandle, LayerMask, MipLevel, vAddress, EVTRequestPagePriority::High);
 				
 				// High priority request should never be Saturated
 				// It's possible for status to be Invalid, if requesting data from a mip level that doesn't exist for the given producer (when using sparse UDIMs)
diff --git a/Engine/Source/Runtime/Engine/Private/VT/RuntimeVirtualTexture.cpp b/Engine/Source/Runtime/Engine/Private/VT/RuntimeVirtualTexture.cpp
--- a/Engine/Source/Runtime/Engine/Private/VT/RuntimeVirtualTexture.cpp
+++ b/Engine/Source/Runtime/Engine/Private/VT/RuntimeVirtualTexture.cpp
@@ -40,6 +40,7 @@ namespace
 
 		//~ Begin IVirtualTexture Interface.
 		virtual FVTRequestPageResult RequestPageData(
+			FRHICommandList& RHICmdList,
 			const FVirtualTextureProducerHandle& ProducerHandle,
 			uint8 LayerMask,
 			uint8 vLevel,
diff --git a/Engine/Source/Runtime/Engine/Private/VT/UploadingVirtualTexture.cpp b/Engine/Source/Runtime/Engine/Private/VT/UploadingVirtualTexture.cpp
--- a/Engine/Source/Runtime/Engine/Private/VT/UploadingVirtualTexture.cpp
+++ b/Engine/Source/Runtime/Engine/Private/VT/UploadingVirtualTexture.cpp
@@ -125,12 +125,12 @@ uint32 FUploadingVirtualTexture::GetLocalMipBias(uint8 vLevel, uint32 vAddress)
 	return Current_vLevel - vLevel;
 }
 
-FVTRequestPageResult FUploadingVirtualTexture::RequestPageData(const FVirtualTextureProducerHandle& ProducerHandle, uint8 LayerMask, uint8 vLevel, uint64 vAddress, EVTRequestPagePriority Priority)
+FVTRequestPageResult FUploadingVirtualTexture::RequestPageData(FRHICommandList& RHICmdList, const FVirtualTextureProducerHandle& ProducerHandle, uint8 LayerMask, uint8 vLevel, uint64 vAddress, EVTRequestPagePriority Priority)
 {
 	vLevel += FirstMipOffset;
 
 	check(vAddress <= MAX_uint32); // Not supporting 64 bit vAddress here. Only currently supported for adaptive runtime virtual texture.
-	return StreamingManager->RequestTile(this, ProducerHandle, LayerMask, vLevel, (uint32)vAddress, Priority);
+	return StreamingManager->RequestTile(RHICmdList, this, ProducerHandle, LayerMask, vLevel, (uint32)vAddress, Priority);
 }
 
 IVirtualTextureFinalizer* FUploadingVirtualTexture::ProducePageData(FRHICommandList& RHICmdList,
diff --git a/Engine/Source/Runtime/Engine/Private/VT/UploadingVirtualTexture.h b/Engine/Source/Runtime/Engine/Private/VT/UploadingVirtualTexture.h
--- a/Engine/Source/Runtime/Engine/Private/VT/UploadingVirtualTexture.h
+++ b/Engine/Source/Runtime/Engine/Private/VT/UploadingVirtualTexture.h
@@ -78,7 +78,7 @@ public:
 
 	// IVirtualTexture interface
 	virtual uint32 GetLocalMipBias(uint8 vLevel, uint32 vAddress) const override;
-	virtual FVTRequestPageResult RequestPageData(const FVirtualTextureProducerHandle& ProducerHandle, uint8 LayerMask, uint8 vLevel, uint64 vAddress, EVTRequestPagePriority Priority) override;
+	virtual FVTRequestPageResult RequestPageData(FRHICommandList& RHICmdList, const FVirtualTextureProducerHandle& ProducerHandle, uint8 LayerMask, uint8 vLevel, uint64 vAddress, EVTRequestPagePriority Priority) override;
 	virtual IVirtualTextureFinalizer* ProducePageData(FRHICommandList& RHICmdList,
 		ERHIFeatureLevel::Type FeatureLevel,
 		EVTProducePageFlags Flags,
diff --git a/Engine/Source/Runtime/Engine/Private/VT/VirtualTextureChunkManager.cpp b/Engine/Source/Runtime/Engine/Private/VT/VirtualTextureChunkManager.cpp
--- a/Engine/Source/Runtime/Engine/Private/VT/VirtualTextureChunkManager.cpp
+++ b/Engine/Source/Runtime/Engine/Private/VT/VirtualTextureChunkManager.cpp
@@ -43,8 +43,8 @@ void FVirtualTextureChunkStreamingManager::UpdateResourceStreaming(float DeltaTi
 #if WITH_EDITOR
 				GetVirtualTextureChunkDDCCache()->UpdateRequests();
 #endif // WITH_EDITOR
-				StreamingManager->TranscodeCache.RetireOldTasks(StreamingManager->UploadCache);
-				StreamingManager->UploadCache.UpdateFreeList();
+				StreamingManager->TranscodeCache.RetireOldTasks(RHICmdList, StreamingManager->UploadCache);
+				StreamingManager->UploadCache.UpdateFreeList(RHICmdList);
 
 				FVirtualTextureCodec::RetireOldCodecs();
 			});
@@ -61,7 +61,7 @@ void FVirtualTextureChunkStreamingManager::CancelForcedResources()
 {
 }
 
-FVTRequestPageResult FVirtualTextureChunkStreamingManager::RequestTile(FUploadingVirtualTexture* VTexture, const FVirtualTextureProducerHandle& ProducerHandle, uint8 LayerMask, uint8 vLevel, uint32 vAddress, EVTRequestPagePriority Priority)
+FVTRequestPageResult FVirtualTextureChunkStreamingManager::RequestTile(FRHICommandList& RHICmdList, FUploadingVirtualTexture* VTexture, const FVirtualTextureProducerHandle& ProducerHandle, uint8 LayerMask, uint8 vLevel, uint32 vAddress, EVTRequestPagePriority Priority)
 {
 	SCOPE_CYCLE_COUNTER(STAT_VTP_RequestTile);
 
@@ -138,7 +138,7 @@ FVTRequestPageResult FVirtualTextureChunkStreamingManager::RequestTile(FUploadin
 	TranscodeParams.LayerMask = LayerMask;
 	TranscodeParams.Codec = CodecResult.Codec;
 	TranscodeParams.Name = VTexture->GetName();
-	const FVTTranscodeTileHandle TranscodeHandle = TranscodeCache.SubmitTask(UploadCache, TranscodeKey, ProducerHandle, TranscodeParams, &GraphCompletionEvents);
+	const FVTTranscodeTileHandle TranscodeHandle = TranscodeCache.SubmitTask(RHICmdList, UploadCache, TranscodeKey, ProducerHandle, TranscodeParams, &GraphCompletionEvents);
 	return FVTRequestPageResult(EVTRequestPageStatus::Pending, TranscodeHandle.PackedData);
 }
 
diff --git a/Engine/Source/Runtime/Engine/Private/VT/VirtualTextureChunkManager.h b/Engine/Source/Runtime/Engine/Private/VT/VirtualTextureChunkManager.h
--- a/Engine/Source/Runtime/Engine/Private/VT/VirtualTextureChunkManager.h
+++ b/Engine/Source/Runtime/Engine/Private/VT/VirtualTextureChunkManager.h
@@ -42,7 +42,7 @@ public:
 	virtual void NotifyLevelOffset(class ULevel* Level, const FVector& Offset) override {}
 	// End IStreamingManager interface
 
-	FVTRequestPageResult RequestTile(FUploadingVirtualTexture* VTexture, const FVirtualTextureProducerHandle& ProducerHandle, uint8 LayerMask, uint8 vLevel, uint32 vAddress, EVTRequestPagePriority Priority);
+	FVTRequestPageResult RequestTile(FRHICommandList& RHICmdList, FUploadingVirtualTexture* VTexture, const FVirtualTextureProducerHandle& ProducerHandle, uint8 LayerMask, uint8 vLevel, uint32 vAddress, EVTRequestPagePriority Priority);
 	IVirtualTextureFinalizer* ProduceTile(FRHICommandList& RHICmdList, uint32 SkipBorderSize, uint8 NumLayers, uint8 LayerMask, uint64 RequestHandle, const FVTProduceTargetLayer* TargetLayers);
 	void GatherProducePageDataTasks(FVirtualTextureProducerHandle const& ProducerHandle, FGraphEventArray& InOutTasks) const;
 	void GatherProducePageDataTasks(uint64 RequestHandle, FGraphEventArray& InOutTasks) const;
diff --git a/Engine/Source/Runtime/Engine/Private/VT/VirtualTextureLevelRedirector.cpp b/Engine/Source/Runtime/Engine/Private/VT/VirtualTextureLevelRedirector.cpp
--- a/Engine/Source/Runtime/Engine/Private/VT/VirtualTextureLevelRedirector.cpp
+++ b/Engine/Source/Runtime/Engine/Private/VT/VirtualTextureLevelRedirector.cpp
@@ -15,6 +15,7 @@ FVirtualTextureLevelRedirector::~FVirtualTextureLevelRedirector()
 }
 
 FVTRequestPageResult FVirtualTextureLevelRedirector::RequestPageData(
+	FRHICommandList& RHICmdList,
 	const FVirtualTextureProducerHandle& ProducerHandle,
 	uint8 LayerMask,
 	uint8 vLevel,
@@ -23,7 +24,7 @@ FVTRequestPageResult FVirtualTextureLevelRedirector::RequestPageData(
 {
 	int32 VirtualTextureIndex = vLevel < TransitionLevel ? 0 : 1;
 	int32 vLevelOffset = vLevel < TransitionLevel ? 0 : TransitionLevel;
-	return VirtualTextures[VirtualTextureIndex]->RequestPageData(ProducerHandle, LayerMask, vLevel - vLevelOffset, vAddress, Priority);
+	return VirtualTextures[VirtualTextureIndex]->RequestPageData(RHICmdList, ProducerHandle, LayerMask, vLevel - vLevelOffset, vAddress, Priority);
 }
 
 IVirtualTextureFinalizer* FVirtualTextureLevelRedirector::ProducePageData(
diff --git a/Engine/Source/Runtime/Engine/Private/VT/VirtualTextureLevelRedirector.h b/Engine/Source/Runtime/Engine/Private/VT/VirtualTextureLevelRedirector.h
--- a/Engine/Source/Runtime/Engine/Private/VT/VirtualTextureLevelRedirector.h
+++ b/Engine/Source/Runtime/Engine/Private/VT/VirtualTextureLevelRedirector.h
@@ -17,6 +17,7 @@ public:
 
 	//~ Begin IVirtualTexture Interface.
 	virtual FVTRequestPageResult RequestPageData(
+		FRHICommandList& RHICmdList,
 		const FVirtualTextureProducerHandle& ProducerHandle,
 		uint8 LayerMask,
 		uint8 vLevel,
diff --git a/Engine/Source/Runtime/Engine/Private/VT/VirtualTextureTranscodeCache.cpp b/Engine/Source/Runtime/Engine/Private/VT/VirtualTextureTranscodeCache.cpp
--- a/Engine/Source/Runtime/Engine/Private/VT/VirtualTextureTranscodeCache.cpp
+++ b/Engine/Source/Runtime/Engine/Private/VT/VirtualTextureTranscodeCache.cpp
@@ -344,6 +344,7 @@ const FVTUploadTileHandle* FVirtualTextureTranscodeCache::AcquireTaskResult(FVTT
 }
 
 FVTTranscodeTileHandle FVirtualTextureTranscodeCache::SubmitTask(
+	FRHICommandList& RHICmdList,
 	FVirtualTextureUploadCache& InUploadCache,
 	const FVTTranscodeKey& InKey,
 	const FVirtualTextureProducerHandle& InProducerHandle,
@@ -417,7 +418,7 @@ FVTTranscodeTileHandle FVirtualTextureTranscodeCache::SubmitTask(
 				TileBaseOffset = TileLayerOffset;
 			}
 
-			TaskEntry.StageTileHandle[LayerIndex] = InUploadCache.PrepareTileForUpload(StagingBufferForLayer, LayerFormat, TilePixelSize);
+			TaskEntry.StageTileHandle[LayerIndex] = InUploadCache.PrepareTileForUpload(RHICmdList, StagingBufferForLayer, LayerFormat, TilePixelSize);
 
 			// If there aren't any prerequisites, it's not worth launching a task for a raw memcpy, just do the work now
 			if (Prerequisites || Chunk.CodecType[LayerIndex] != EVirtualTextureCodec::RawGPU)
@@ -441,7 +442,7 @@ FVTTranscodeTileHandle FVirtualTextureTranscodeCache::SubmitTask(
 	return FVTTranscodeTileHandle(TaskIndex, TaskEntry.Magic);
 }
 
-void FVirtualTextureTranscodeCache::RetireOldTasks(FVirtualTextureUploadCache& InUploadCache)
+void FVirtualTextureTranscodeCache::RetireOldTasks(FRHICommandList& RHICmdList, FVirtualTextureUploadCache& InUploadCache)
 {
 	const uint32 CurrentFrame = GFrameNumberRenderThread;
 
@@ -477,7 +478,7 @@ void FVirtualTextureTranscodeCache::RetireOldTasks(FVirtualTextureUploadCache& I
 			const FVTUploadTileHandle StageTileHandle = TaskEntry.StageTileHandle[LayerIndex];
 			if (Key.LayerMask & (1u << LayerIndex))
 			{
-				InUploadCache.CancelTile(StageTileHandle);
+				InUploadCache.CancelTile(RHICmdList, StageTileHandle);
 			}
 		}
 	
diff --git a/Engine/Source/Runtime/Engine/Private/VT/VirtualTextureTranscodeCache.h b/Engine/Source/Runtime/Engine/Private/VT/VirtualTextureTranscodeCache.h
--- a/Engine/Source/Runtime/Engine/Private/VT/VirtualTextureTranscodeCache.h
+++ b/Engine/Source/Runtime/Engine/Private/VT/VirtualTextureTranscodeCache.h
@@ -74,6 +74,7 @@ public:
 	FVTTranscodeTileHandleAndStatus FindTask(const FVTTranscodeKey& InKey) const;
 
 	FVTTranscodeTileHandle SubmitTask(
+		FRHICommandList& RHICmdList,
 		FVirtualTextureUploadCache& InUploadCache,
 		const FVTTranscodeKey& InKey,
 		const FVirtualTextureProducerHandle& InProducerHandle,
@@ -87,7 +88,7 @@ public:
 	FGraphEventRef GetTaskEvent(FVTTranscodeTileHandle InHandle) const;
 	void GatherProducePageDataTasks(FVirtualTextureProducerHandle const& ProducerHandle, FGraphEventArray& InOutTasks) const;
 
-	void RetireOldTasks(FVirtualTextureUploadCache& InUploadCache);
+	void RetireOldTasks(FRHICommandList& RHICmdList, FVirtualTextureUploadCache& InUploadCache);
 
 private:
 	enum ListType
diff --git a/Engine/Source/Runtime/Engine/Private/VT/VirtualTextureUploadCache.cpp b/Engine/Source/Runtime/Engine/Private/VT/VirtualTextureUploadCache.cpp
--- a/Engine/Source/Runtime/Engine/Private/VT/VirtualTextureUploadCache.cpp
+++ b/Engine/Source/Runtime/Engine/Private/VT/VirtualTextureUploadCache.cpp
@@ -39,7 +39,7 @@ static TAutoConsoleVariable<int32> CVarMaxUploadRequests(
 	ECVF_RenderThreadSafe);
 
 
-uint32 FVTUploadTileAllocator::Allocate(EPixelFormat InFormat, uint32 InTileSize)
+uint32 FVTUploadTileAllocator::Allocate(FRHICommandList& RHICmdList, EPixelFormat InFormat, uint32 InTileSize)
 {
 	// Find matching FormatBuffer.
 	const FPixelFormatInfo& FormatInfo = GPixelFormats[InFormat];
@@ -96,7 +96,7 @@ uint32 FVTUploadTileAllocator::Allocate(EPixelFormat InFormat, uint32 InTileSize
 	if (StagingBuffer.Memory == nullptr)
 	{
 		// Staging buffer needs underlying buffer allocating.
-		StagingBuffer.Init(Desc.BlockBytes, Desc.MemorySize);
+		StagingBuffer.Init(RHICmdList, Desc.BlockBytes, Desc.MemorySize);
 		NumAllocatedBytes += StagingBuffer.TileSizeAligned * StagingBuffer.NumTiles;
 	}
 
@@ -110,7 +110,7 @@ uint32 FVTUploadTileAllocator::Allocate(EPixelFormat InFormat, uint32 InTileSize
 	return Handle.PackedValue;
 }
 
-void FVTUploadTileAllocator::Free(uint32 InHandle)
+void FVTUploadTileAllocator::Free(FRHICommandList& RHICmdList, uint32 InHandle)
 {
 	FHandle Handle;
 	Handle.PackedValue = InHandle;
@@ -125,7 +125,7 @@ void FVTUploadTileAllocator::Free(uint32 InHandle)
 		check(NumAllocatedBytes >= StagingBuffer.TileSizeAligned * StagingBuffer.NumTiles);
 		NumAllocatedBytes -= StagingBuffer.TileSizeAligned * StagingBuffer.NumTiles;
 
-		StagingBuffer.Release();
+		StagingBuffer.Release(&RHICmdList);
 	}
 }
 
@@ -162,10 +162,10 @@ FVTUploadTileBufferExt FVTUploadTileAllocator::GetBufferFromHandleExt(uint32 InH
 
 FVTUploadTileAllocator::FStagingBuffer::~FStagingBuffer()
 {
-	Release();
+	Release(nullptr);
 }
 
-void FVTUploadTileAllocator::FStagingBuffer::Init(uint32 InBufferStrideBytes, uint32 InTileSizeBytes)
+void FVTUploadTileAllocator::FStagingBuffer::Init(FRHICommandList& RHICmdList, uint32 InBufferStrideBytes, uint32 InTileSizeBytes)
 {
 	TileSize = InTileSizeBytes;
 	TileSizeAligned = Align(InTileSizeBytes, 128u);
@@ -184,11 +184,8 @@ void FVTUploadTileAllocator::FStagingBuffer::Init(uint32 InBufferStrideBytes, ui
 	if (GRHISupportsDirectGPUMemoryLock && GRHISupportsUpdateFromBufferTexture)
 	{
 		// Allocate staging buffer directly in GPU memory.
-		checkSlow(IsInRenderingThread());
-		FRHICommandListImmediate& RHICmdList = FRHICommandListExecutor::GetImmediateCommandList();
-
 		FRHIResourceCreateInfo CreateInfo(TEXT("StagingBuffer"));
-		RHIBuffer = RHICreateStructuredBuffer(InBufferStrideBytes, BufferSize, BUF_ShaderResource | BUF_Static | BUF_KeepCPUAccessible, CreateInfo);
+		RHIBuffer = RHICmdList.CreateBuffer(BufferSize, BUF_ShaderResource | BUF_Static | BUF_KeepCPUAccessible | BUF_StructuredBuffer, InBufferStrideBytes, ERHIAccess::SRVMask, CreateInfo);
 
 		// Here we bypass 'normal' RHI operations in order to get a persistent pointer to GPU memory, on supported platforms
 		// This should be encapsulated into a proper RHI method at some point
@@ -205,14 +202,16 @@ void FVTUploadTileAllocator::FStagingBuffer::Init(uint32 InBufferStrideBytes, ui
 	}
 }
 
-void FVTUploadTileAllocator::FStagingBuffer::Release()
+void FVTUploadTileAllocator::FStagingBuffer::Release(FRHICommandList* RHICmdList)
 {
 	const uint32 BufferSize = TileSizeAligned * NumTiles;
 
 	if (RHIBuffer.IsValid())
 	{
+		check(RHICmdList);
+
 		// Unmap and release the GPU buffer if present.
-		RHIUnlockBuffer(RHIBuffer);
+		RHICmdList->UnlockBuffer(RHIBuffer);
 		RHIBuffer.SafeRelease();
 		// In this case 'Memory' was the mapped pointer, so release it.
 		Memory = nullptr;
@@ -341,7 +340,7 @@ void FVirtualTextureUploadCache::Finalize(FRDGBuilder& GraphBuilder)
 			}
 
 			// Can release upload buffer.
-			TileAllocator.Free(Entry.TileHandle);
+			TileAllocator.Free(RHICmdList, Entry.TileHandle);
 		}
 
 		RHICmdList.UnlockTexture2D(StagingTexture.RHITexture, 0u, false, false);
@@ -390,22 +389,24 @@ void FVirtualTextureUploadCache::Finalize(FRDGBuilder& GraphBuilder)
 
 void FVirtualTextureUploadCache::ReleaseRHI()
 {
+	check(IsInRenderingThread());
+	FRHICommandListImmediate& RHICmdList = FRHICommandListExecutor::GetImmediateCommandList();
+
 	// Complete/Cancel all work will release allocated staging buffers.
-	UpdateFreeList(true);
+	UpdateFreeList(RHICmdList, true);
 	for (TSparseArray<FTileEntry>::TIterator It(PendingUpload); It; ++It)
 	{
-		CancelTile(FVTUploadTileHandle(It.GetIndex()));
+		CancelTile(RHICmdList, FVTUploadTileHandle(It.GetIndex()));
 	}
 	// Release staging textures.
 	Pools.Empty();
 }
 
-FVTUploadTileHandle FVirtualTextureUploadCache::PrepareTileForUpload(FVTUploadTileBuffer& OutBuffer, EPixelFormat InFormat, uint32 InTileSize)
+FVTUploadTileHandle FVirtualTextureUploadCache::PrepareTileForUpload(FRHICommandList& RHICmdList, FVTUploadTileBuffer& OutBuffer, EPixelFormat InFormat, uint32 InTileSize)
 {
 	SCOPE_CYCLE_COUNTER(STAT_VTP_StageTile)
-	checkSlow(IsInRenderingThread());
 
-	uint32 TileHandle = TileAllocator.Allocate(InFormat, InTileSize);
+	uint32 TileHandle = TileAllocator.Allocate(RHICmdList, InFormat, InTileSize);
 	OutBuffer = TileAllocator.GetBufferFromHandle(TileHandle);
 
 	const int32 PoolIndex = GetOrCreatePoolIndex(InFormat, InTileSize);
@@ -459,7 +460,7 @@ void FVirtualTextureUploadCache::SubmitTile(FRHICommandList& RHICmdList, const F
 		RHICmdList.UpdateTexture2D(InDestTexture, 0u, UpdateRegion, UploadBuffer.Stride, (uint8*)UploadBuffer.BufferMemory + UploadBuffer.BufferOffset);
 
 		// UpdateTexture2D makes internal copy of data, no need to wait before re-using tile
-		TileAllocator.Free(Entry.TileHandle);
+		TileAllocator.Free(RHICmdList, Entry.TileHandle);
 	}
 	else
 	{
@@ -475,20 +476,16 @@ void FVirtualTextureUploadCache::SubmitTile(FRHICommandList& RHICmdList, const F
 	PendingUpload.RemoveAt(InHandle.Index);
 }
 
-void FVirtualTextureUploadCache::CancelTile(const FVTUploadTileHandle& InHandle)
+void FVirtualTextureUploadCache::CancelTile(FRHICommandList& RHICmdList, const FVTUploadTileHandle& InHandle)
 {
-	checkSlow(IsInRenderingThread());
-
 	check(PendingUpload.IsValidIndex(InHandle.Index));
 	FTileEntry& Entry = PendingUpload[InHandle.Index];
-	TileAllocator.Free(Entry.TileHandle);
+	TileAllocator.Free(RHICmdList, Entry.TileHandle);
 	PendingUpload.RemoveAt(InHandle.Index);
 }
 
-void FVirtualTextureUploadCache::UpdateFreeList(bool bForceFreeAll)
+void FVirtualTextureUploadCache::UpdateFreeList(FRHICommandList& RHICmdList, bool bForceFreeAll)
 {
-	checkSlow(IsInRenderingThread());
-
 	const uint32 CurrentFrame = GFrameNumberRenderThread;
 
 	// Iterate tiles pending release and free them if they are old enough.
@@ -501,7 +498,7 @@ void FVirtualTextureUploadCache::UpdateFreeList(bool bForceFreeAll)
 			break;
 		}
 
-		TileAllocator.Free(It->TileHandle);
+		TileAllocator.Free(RHICmdList, It->TileHandle);
 		It.RemoveCurrent();
 	}
 }
diff --git a/Engine/Source/Runtime/Engine/Private/VT/VirtualTextureUploadCache.h b/Engine/Source/Runtime/Engine/Private/VT/VirtualTextureUploadCache.h
--- a/Engine/Source/Runtime/Engine/Private/VT/VirtualTextureUploadCache.h
+++ b/Engine/Source/Runtime/Engine/Private/VT/VirtualTextureUploadCache.h
@@ -52,9 +52,9 @@ class FVTUploadTileAllocator
 {
 public:
 	/** Allocate a tile. Sometimes does an allocation of the backing CPU/GPU block of memory. */
-	uint32 Allocate(EPixelFormat InFormat, uint32 InTileSize);
+	uint32 Allocate(FRHICommandList& RHICmdList, EPixelFormat InFormat, uint32 InTileSize);
 	/** Free a tile. Sometimes does a free of the backing CPU/GPU block of memory. */
-	void Free(uint32 InHandle);
+	void Free(FRHICommandList& RHICmdList, uint32 InHandle);
 
 	/** Get upload buffer description from handle. */
 	FVTUploadTileBuffer GetBufferFromHandle(uint32 InHandle) const;
@@ -86,8 +86,8 @@ private:
 	{
 		~FStagingBuffer();
 
-		void Init(uint32 InBufferStrideBytes, uint32 InTileSizeBytes);
-		void Release();
+		void Init(FRHICommandList& RHICmdList, uint32 InBufferStrideBytes, uint32 InTileSizeBytes);
+		void Release(FRHICommandList* RHICmdList);
 
 		/** GPU buffer if used on platform. */
 		TRefCountPtr<FRHIBuffer> RHIBuffer;
@@ -137,21 +137,21 @@ public:
 	//~ End FRenderResource Interface.
 
 	/** Get a staging upload buffer for streaming texture data into. */
-	FVTUploadTileHandle PrepareTileForUpload(FVTUploadTileBuffer& OutBuffer, EPixelFormat InFormat, uint32 InTileSize);
+	FVTUploadTileHandle PrepareTileForUpload(FRHICommandList& RHICmdList, FVTUploadTileBuffer& OutBuffer, EPixelFormat InFormat, uint32 InTileSize);
 	/** 
 	 * Mark streamed upload data ready for upload to to the physical virtual texture.
 	 * Depending on the platform the upload might happen here, or be deferred to the Finalize() call.
 	 */
 	void SubmitTile(FRHICommandList& RHICmdList, const FVTUploadTileHandle& InHandle, FRHITexture2D* InDestTexture, int InDestX, int InDestY, int InSkipBorderSize);
 	/** Cancel a tile that was already in flight. */
-	void CancelTile(const FVTUploadTileHandle& InHandle);
+	void CancelTile(FRHICommandList& RHICmdList, const FVTUploadTileHandle& InHandle);
 
 	//~ Begin IVirtualTextureFinalizer Interface.
 	virtual void Finalize(FRDGBuilder& GraphBuilder) override;
 	//~ End IVirtualTextureFinalizer Interface.
 
 	/** Call on a tick to recycle submitted staging buffers. */
-	void UpdateFreeList(bool bForceFreeAll = false);
+	void UpdateFreeList(FRHICommandList& RHICmdList, bool bForceFreeAll = false);
 
 	/** Returns true if underlying allocator within the budget set by r.VT.MaxUploadMemory.  */
 	uint32 IsInMemoryBudget() const;
diff --git a/Engine/Source/Runtime/RenderCore/Public/VirtualTexturing.h b/Engine/Source/Runtime/RenderCore/Public/VirtualTexturing.h
--- a/Engine/Source/Runtime/RenderCore/Public/VirtualTexturing.h
+++ b/Engine/Source/Runtime/RenderCore/Public/VirtualTexturing.h
@@ -294,7 +294,10 @@ public:
 	* @param Priority Priority of the request, used to drive async IO/task priority needed to generate data for request
 	* @return FVTRequestPageResult describing the availability of the request
 	*/
-	virtual FVTRequestPageResult RequestPageData(const FVirtualTextureProducerHandle& ProducerHandle, uint8 LayerMask, uint8 vLevel, uint64 vAddress, EVTRequestPagePriority Priority) = 0;
+	virtual FVTRequestPageResult RequestPageData(FRHICommandList& RHICmdList, const FVirtualTextureProducerHandle& ProducerHandle, uint8 LayerMask, uint8 vLevel, uint64 vAddress, EVTRequestPagePriority Priority) = 0;
+
+	UE_DEPRECATED(5.2, "RequestPageData now requires an RHI command list.")
+	virtual FVTRequestPageResult RequestPageData(const FVirtualTextureProducerHandle& ProducerHandle, uint8 LayerMask, uint8 vLevel, uint64 vAddress, EVTRequestPagePriority Priority) { return {}; }
 
 	/**
 	* Upload page data to the cache, data must have been previously requested, and reported either 'Available' or 'Pending'
diff --git a/Engine/Source/Runtime/Renderer/Private/VT/AdaptiveVirtualTexture.cpp b/Engine/Source/Runtime/Renderer/Private/VT/AdaptiveVirtualTexture.cpp
--- a/Engine/Source/Runtime/Renderer/Private/VT/AdaptiveVirtualTexture.cpp
+++ b/Engine/Source/Runtime/Renderer/Private/VT/AdaptiveVirtualTexture.cpp
@@ -64,6 +64,7 @@ public:
 	}
 
 	virtual FVTRequestPageResult RequestPageData(
+		FRHICommandList& RHICmdList,
 		const FVirtualTextureProducerHandle& ProducerHandle,
 		uint8 LayerMask,
 		uint8 vLevel,
@@ -76,7 +77,7 @@ public:
 		vAddress = FMath::MortonCode2_64(X) | (FMath::MortonCode2_64(Y) << 1);
 		vLevel = (uint8)(FMath::Max((int32)vLevel + LevelOffset, 0));
 
-		return VirtualTexture->RequestPageData(ProducerHandle, LayerMask, vLevel, vAddress, Priority);
+		return VirtualTexture->RequestPageData(RHICmdList, ProducerHandle, LayerMask, vLevel, vAddress, Priority);
 	}
 
 	virtual IVirtualTextureFinalizer* ProducePageData(
diff --git a/Engine/Source/Runtime/Renderer/Private/VT/RuntimeVirtualTextureProducer.cpp b/Engine/Source/Runtime/Renderer/Private/VT/RuntimeVirtualTextureProducer.cpp
--- a/Engine/Source/Runtime/Renderer/Private/VT/RuntimeVirtualTextureProducer.cpp
+++ b/Engine/Source/Runtime/Renderer/Private/VT/RuntimeVirtualTextureProducer.cpp
@@ -149,6 +149,7 @@ FRuntimeVirtualTextureProducer::FRuntimeVirtualTextureProducer(
 }
 
 FVTRequestPageResult FRuntimeVirtualTextureProducer::RequestPageData(
+	FRHICommandList& RHICmdList,
 	const FVirtualTextureProducerHandle& ProducerHandle,
 	uint8 LayerMask,
 	uint8 vLevel,
diff --git a/Engine/Source/Runtime/Renderer/Private/VT/RuntimeVirtualTextureProducer.h b/Engine/Source/Runtime/Renderer/Private/VT/RuntimeVirtualTextureProducer.h
--- a/Engine/Source/Runtime/Renderer/Private/VT/RuntimeVirtualTextureProducer.h
+++ b/Engine/Source/Runtime/Renderer/Private/VT/RuntimeVirtualTextureProducer.h
@@ -85,6 +85,7 @@ public:
 
 	//~ Begin IVirtualTexture Interface.
 	virtual FVTRequestPageResult RequestPageData(
+		FRHICommandList& RHICmdList,
 		const FVirtualTextureProducerHandle& ProducerHandle,
 		uint8 LayerMask,
 		uint8 vLevel,
diff --git a/Engine/Source/Runtime/Renderer/Private/VT/VirtualTextureSystem.cpp b/Engine/Source/Runtime/Renderer/Private/VT/VirtualTextureSystem.cpp
--- a/Engine/Source/Runtime/Renderer/Private/VT/VirtualTextureSystem.cpp
+++ b/Engine/Source/Runtime/Renderer/Private/VT/VirtualTextureSystem.cpp
@@ -1972,7 +1972,7 @@ void FVirtualTextureSystem::SubmitRequestsFromLocalTileList(FRHICommandList& RHI
 		}
 
 		FVTRequestPageResult RequestPageResult = Producer->GetVirtualTexture()->RequestPageData(
-			ProducerHandle, LayerMask, Tile.Local_vLevel, Tile.Local_vAddress, EVTRequestPagePriority::High);
+			RHICmdList, ProducerHandle, LayerMask, Tile.Local_vLevel, Tile.Local_vAddress, EVTRequestPagePriority::High);
 
 		if (RequestPageResult.Status != EVTRequestPageStatus::Available)
 		{
@@ -2080,7 +2080,7 @@ void FVirtualTextureSystem::BeginSubmitRequests(FRHICommandList& RHICmdList, ERH
 			}
 
 			const EVTRequestPagePriority Priority = bLockTile ? EVTRequestPagePriority::High : EVTRequestPagePriority::Normal;
-			FVTRequestPageResult RequestPageResult = Producer.GetVirtualTexture()->RequestPageData(ProducerHandle, ProducerTextureLayerMask, TileToLoad.Local_vLevel, TileToLoad.Local_vAddress, Priority);
+			FVTRequestPageResult RequestPageResult = Producer.GetVirtualTexture()->RequestPageData(RHICmdList, ProducerHandle, ProducerTextureLayerMask, TileToLoad.Local_vLevel, TileToLoad.Local_vAddress, Priority);
 			if (RequestPageResult.Status == EVTRequestPageStatus::Pending && bForceProduceTile)
 			{
 				// If we're forcing production of this tile, we're OK producing data now (and possibly waiting) as long as data is pending
diff --git a/Engine/Source/Runtime/Renderer/Private/VT/VirtualTextureTest.cpp b/Engine/Source/Runtime/Renderer/Private/VT/VirtualTextureTest.cpp
--- a/Engine/Source/Runtime/Renderer/Private/VT/VirtualTextureTest.cpp
+++ b/Engine/Source/Runtime/Renderer/Private/VT/VirtualTextureTest.cpp
@@ -88,7 +88,7 @@ FVirtualTextureTest::FVirtualTextureTest( uint32 InSizeX, uint32 InSizeY, uint32
 FVirtualTextureTest::~FVirtualTextureTest()
 {}
 
-bool FVirtualTextureTest::RequestPageData( uint8 vLevel, uint64 vAddress, void* RESTRICT& Location ) /*const*/
+bool FVirtualTextureTest::RequestPageData(FRHICommandList& RHICmdList, uint8 vLevel, uint64 vAddress, void* RESTRICT& Location ) /*const*/
 {
 	/*
 	const size_t VTHeaderSize = 0;
diff --git a/Engine/Source/Runtime/Renderer/Private/VT/VirtualTextureTest.h b/Engine/Source/Runtime/Renderer/Private/VT/VirtualTextureTest.h
--- a/Engine/Source/Runtime/Renderer/Private/VT/VirtualTextureTest.h
+++ b/Engine/Source/Runtime/Renderer/Private/VT/VirtualTextureTest.h
@@ -19,7 +19,7 @@ public:
 					FVirtualTextureTest( uint32 SizeX, uint32 SizeY, uint32 SizeZ );
 	virtual			~FVirtualTextureTest();
 
-	virtual bool	RequestPageData( uint8 vLevel, uint64 vAddress, void* RESTRICT& Location ) /*const*/ override;
+	virtual bool	RequestPageData( FRHICommandList& RHICmdList, uint8 vLevel, uint64 vAddress, void* RESTRICT& Location ) /*const*/ override;
 	virtual void	ProducePageData( FRHICommandList& RHICmdList, ERHIFeatureLevel::Type FeatureLevel, uint8 vLevel, uint64 vAddress, uint16 pAddress, void* RESTRICT Location ) /*const*/ override;
 	virtual void    DumpToConsole(bool verbose) override;
 };
