Reduce include dependencies for NavigationTypes.h

diff --git a/Engine/Plugins/Runtime/MassGameplay/Source/MassCommon/Public/MassCommonFragments.h b/Engine/Plugins/Runtime/MassGameplay/Source/MassCommon/Public/MassCommonFragments.h
--- a/Engine/Plugins/Runtime/MassGameplay/Source/MassCommon/Public/MassCommonFragments.h
+++ b/Engine/Plugins/Runtime/MassGameplay/Source/MassCommon/Public/MassCommonFragments.h
@@ -3,7 +3,9 @@
 #pragma once
 
 #include "MassEntityTypes.h"
+#if UE_ENABLE_INCLUDE_ORDER_DEPRECATED_IN_5_2
 #include "AI/Navigation/NavigationTypes.h"
+#endif //UE_ENABLE_INCLUDE_ORDER_DEPRECATED_IN_5_2
 #include "MassCommonFragments.generated.h"
 
 
diff --git a/Engine/Source/Runtime/AIModule/Classes/EnvironmentQuery/EQSTestingPawn.h b/Engine/Source/Runtime/AIModule/Classes/EnvironmentQuery/EQSTestingPawn.h
--- a/Engine/Source/Runtime/AIModule/Classes/EnvironmentQuery/EQSTestingPawn.h
+++ b/Engine/Source/Runtime/AIModule/Classes/EnvironmentQuery/EQSTestingPawn.h
@@ -7,6 +7,7 @@
 #include "EnvironmentQuery/EnvQueryTypes.h"
 #include "EnvironmentQuery/EQSQueryResultSourceInterface.h"
 #include "GameFramework/Character.h"
+#include "AI/Navigation/NavigationTypes.h"
 #include "EQSTestingPawn.generated.h"
 
 class UEnvQuery;
diff --git a/Engine/Source/Runtime/AIModule/Classes/EnvironmentQuery/Generators/EnvQueryGenerator_OnCircle.h b/Engine/Source/Runtime/AIModule/Classes/EnvironmentQuery/Generators/EnvQueryGenerator_OnCircle.h
--- a/Engine/Source/Runtime/AIModule/Classes/EnvironmentQuery/Generators/EnvQueryGenerator_OnCircle.h
+++ b/Engine/Source/Runtime/AIModule/Classes/EnvironmentQuery/Generators/EnvQueryGenerator_OnCircle.h
@@ -5,7 +5,9 @@
 #include "CoreMinimal.h"
 #include "UObject/ObjectMacros.h"
 #include "Templates/SubclassOf.h"
+#if UE_ENABLE_INCLUDE_ORDER_DEPRECATED_IN_5_2
 #include "AI/Navigation/NavigationTypes.h"
+#endif //UE_ENABLE_INCLUDE_ORDER_DEPRECATED_IN_5_2
 #include "EnvironmentQuery/EnvQueryTypes.h"
 #include "DataProviders/AIDataProvider.h"
 #include "EnvironmentQuery/Generators/EnvQueryGenerator_ProjectedPoints.h"
diff --git a/Engine/Source/Runtime/Engine/Classes/AI/Navigation/NavAgentSelector.h b/Engine/Source/Runtime/Engine/Classes/AI/Navigation/NavAgentSelector.h
new file mode 100644
--- /dev/null
+++ b/Engine/Source/Runtime/Engine/Classes/AI/Navigation/NavAgentSelector.h
@@ -0,0 +1,109 @@
+﻿// Copyright Epic Games, Inc. All Rights Reserved.
+
+#pragma once
+
+#include "NavAgentSelector.generated.h"
+
+USTRUCT()
+struct ENGINE_API FNavAgentSelector
+{
+	GENERATED_USTRUCT_BODY()
+
+	static const uint32 InitializedBit = 0x80000000;
+
+#if CPP
+	union
+	{
+		struct
+		{
+#endif
+			UPROPERTY(EditAnywhere, Category = Default)
+			uint32 bSupportsAgent0 : 1;
+			UPROPERTY(EditAnywhere, Category = Default)
+			uint32 bSupportsAgent1 : 1;
+			UPROPERTY(EditAnywhere, Category = Default)
+			uint32 bSupportsAgent2 : 1;
+			UPROPERTY(EditAnywhere, Category = Default)
+			uint32 bSupportsAgent3 : 1;
+			UPROPERTY(EditAnywhere, Category = Default)
+			uint32 bSupportsAgent4 : 1;
+			UPROPERTY(EditAnywhere, Category = Default)
+			uint32 bSupportsAgent5 : 1;
+			UPROPERTY(EditAnywhere, Category = Default)
+			uint32 bSupportsAgent6 : 1;
+			UPROPERTY(EditAnywhere, Category = Default)
+			uint32 bSupportsAgent7 : 1;
+			UPROPERTY(EditAnywhere, Category = Default)
+			uint32 bSupportsAgent8 : 1;
+			UPROPERTY(EditAnywhere, Category = Default)
+			uint32 bSupportsAgent9 : 1;
+			UPROPERTY(EditAnywhere, Category = Default)
+			uint32 bSupportsAgent10 : 1;
+			UPROPERTY(EditAnywhere, Category = Default)
+			uint32 bSupportsAgent11 : 1;
+			UPROPERTY(EditAnywhere, Category = Default)
+			uint32 bSupportsAgent12 : 1;
+			UPROPERTY(EditAnywhere, Category = Default)
+			uint32 bSupportsAgent13 : 1;
+			UPROPERTY(EditAnywhere, Category = Default)
+			uint32 bSupportsAgent14 : 1;
+			UPROPERTY(EditAnywhere, Category = Default)
+			uint32 bSupportsAgent15 : 1;
+#if CPP
+		};
+		uint32 PackedBits;
+	};
+#endif
+
+	explicit FNavAgentSelector(const uint32 InBits = 0x7fffffff);
+
+	FORCEINLINE bool Contains(int32 AgentIndex) const
+	{
+		return (AgentIndex >= 0 && AgentIndex < 16) ? !!(PackedBits & (1 << AgentIndex)) : false;
+	}
+
+	FORCEINLINE void Set(int32 AgentIndex)
+	{
+		if (AgentIndex >= 0 && AgentIndex < 16)
+		{
+			PackedBits |= (1 << AgentIndex);
+		}
+	}
+
+	FORCEINLINE bool IsInitialized() const
+	{
+		return (PackedBits & InitializedBit) != 0;
+	}
+
+	FORCEINLINE void MarkInitialized()
+	{
+		PackedBits |= InitializedBit;
+	}
+
+	FORCEINLINE void Empty()
+	{
+		PackedBits = 0;
+	}
+
+	bool IsSame(const FNavAgentSelector& Other) const
+	{
+		return (~InitializedBit & PackedBits) == (~InitializedBit & Other.PackedBits);
+	}
+
+	bool Serialize(FArchive& Ar);
+
+	uint32 GetAgentBits() const 
+	{
+		return (~InitializedBit & PackedBits);
+	}
+};
+
+template<>
+struct TStructOpsTypeTraits< FNavAgentSelector > : public TStructOpsTypeTraitsBase2< FNavAgentSelector >
+{
+	enum
+	{
+		WithSerializer = true,
+	};
+};
+
diff --git a/Engine/Source/Runtime/Engine/Classes/AI/Navigation/NavAreaBase.h b/Engine/Source/Runtime/Engine/Classes/AI/Navigation/NavAreaBase.h
--- a/Engine/Source/Runtime/Engine/Classes/AI/Navigation/NavAreaBase.h
+++ b/Engine/Source/Runtime/Engine/Classes/AI/Navigation/NavAreaBase.h
@@ -3,9 +3,13 @@
 #pragma once
 
 #include "UObject/Object.h"
+#include "Templates/SubclassOf.h"
+#if UE_ENABLE_INCLUDE_ORDER_DEPRECATED_IN_5_2
 #include "AI/Navigation/NavigationTypes.h"
+#endif //UE_ENABLE_INCLUDE_ORDER_DEPRECATED_IN_5_2
 #include "NavAreaBase.generated.h"
 
+struct FNavAgentProperties;
 
 // a stub class. Actual implementation in NavigationSystem module.
 UCLASS(DefaultToInstanced, abstract, Config = Engine)
diff --git a/Engine/Source/Runtime/Engine/Classes/AI/Navigation/NavLinkDefinition.h b/Engine/Source/Runtime/Engine/Classes/AI/Navigation/NavLinkDefinition.h
--- a/Engine/Source/Runtime/Engine/Classes/AI/Navigation/NavLinkDefinition.h
+++ b/Engine/Source/Runtime/Engine/Classes/AI/Navigation/NavLinkDefinition.h
@@ -7,7 +7,10 @@
 #include "UObject/Object.h"
 #include "UObject/Class.h"
 #include "Templates/SubclassOf.h"
+#if UE_ENABLE_INCLUDE_ORDER_DEPRECATED_IN_5_2
 #include "AI/Navigation/NavigationTypes.h"
+#endif //UE_ENABLE_INCLUDE_ORDER_DEPRECATED_IN_5_2
+#include "AI/Navigation/NavAgentSelector.h"
 #include "NavLinkDefinition.generated.h"
 
 class UNavAreaBase;
diff --git a/Engine/Source/Runtime/Engine/Classes/AI/Navigation/NavigationBounds.h b/Engine/Source/Runtime/Engine/Classes/AI/Navigation/NavigationBounds.h
new file mode 100644
--- /dev/null
+++ b/Engine/Source/Runtime/Engine/Classes/AI/Navigation/NavigationBounds.h
@@ -0,0 +1,37 @@
+﻿// Copyright Epic Games, Inc. All Rights Reserved.
+
+#pragma once
+
+#include "AI/Navigation/NavAgentSelector.h"
+
+struct FNavigationBounds
+{
+	uint32 UniqueID;
+	FBox AreaBox;
+	FNavAgentSelector SupportedAgents;
+	TWeakObjectPtr<ULevel> Level;		// The level this bounds belongs to
+
+	bool operator==(const FNavigationBounds& Other) const 
+	{ 
+		return UniqueID == Other.UniqueID; 
+	}
+
+	friend uint32 GetTypeHash(const FNavigationBounds& NavBounds)
+	{
+		return GetTypeHash(NavBounds.UniqueID);
+	}
+};
+
+struct FNavigationBoundsUpdateRequest 
+{
+	FNavigationBounds NavBounds;
+	
+	enum Type
+	{
+		Added,
+		Removed,
+		Updated,
+	};
+
+	Type UpdateRequest;
+};
\ No newline at end of file
diff --git a/Engine/Source/Runtime/Engine/Classes/AI/Navigation/NavigationDataResolution.h b/Engine/Source/Runtime/Engine/Classes/AI/Navigation/NavigationDataResolution.h
new file mode 100644
--- /dev/null
+++ b/Engine/Source/Runtime/Engine/Classes/AI/Navigation/NavigationDataResolution.h
@@ -0,0 +1,15 @@
+﻿// Copyright Epic Games, Inc. All Rights Reserved.
+
+#pragma once
+
+#include "NavigationDataResolution.generated.h"
+
+UENUM()
+enum class ENavigationDataResolution : uint8
+{
+	Low = 0,
+	Default = 1,
+	High = 2,
+	Invalid = 3 UMETA(DisplayName="None"),
+	MAX = 3,
+};
\ No newline at end of file
diff --git a/Engine/Source/Runtime/Engine/Classes/AI/Navigation/NavigationTypes.h b/Engine/Source/Runtime/Engine/Classes/AI/Navigation/NavigationTypes.h
--- a/Engine/Source/Runtime/Engine/Classes/AI/Navigation/NavigationTypes.h
+++ b/Engine/Source/Runtime/Engine/Classes/AI/Navigation/NavigationTypes.h
@@ -17,6 +17,10 @@
 #include "Misc/CoreStats.h"
 #include "UObject/SoftObjectPath.h"
 #include "GameFramework/Actor.h"
+#if UE_ENABLE_INCLUDE_ORDER_DEPRECATED_IN_5_2
+#include "AI/Navigation/NavAgentSelector.h"
+#include "AI/Navigation/NavigationBounds.h"
+#endif //UE_ENABLE_INCLUDE_ORDER_DEPRECATED_IN_5_2
 #include "NavigationTypes.generated.h"
 
 #define INVALID_NAVNODEREF (0)
@@ -65,16 +69,6 @@ namespace FNavigationSystem
 	ENGINE_API bool IsInBaseNavmesh(const UObject* Object);
 }
 
-UENUM()
-enum class ENavigationDataResolution : uint8
-{
-	Low = 0,
-	Default = 1,
-	High = 2,
-	Invalid = 3 UMETA(DisplayName="None"),
-	MAX = 3,
-};
-
 UENUM()
 namespace ENavigationOptionFlag
 {
@@ -126,141 +120,6 @@ struct FNavigationDirtyArea
 	}
 };
 
-USTRUCT()
-struct ENGINE_API FNavAgentSelector
-{
-	GENERATED_USTRUCT_BODY()
-
-	static const uint32 InitializedBit = 0x80000000;
-
-#if CPP
-	union
-	{
-		struct
-		{
-#endif
-			UPROPERTY(EditAnywhere, Category = Default)
-			uint32 bSupportsAgent0 : 1;
-			UPROPERTY(EditAnywhere, Category = Default)
-			uint32 bSupportsAgent1 : 1;
-			UPROPERTY(EditAnywhere, Category = Default)
-			uint32 bSupportsAgent2 : 1;
-			UPROPERTY(EditAnywhere, Category = Default)
-			uint32 bSupportsAgent3 : 1;
-			UPROPERTY(EditAnywhere, Category = Default)
-			uint32 bSupportsAgent4 : 1;
-			UPROPERTY(EditAnywhere, Category = Default)
-			uint32 bSupportsAgent5 : 1;
-			UPROPERTY(EditAnywhere, Category = Default)
-			uint32 bSupportsAgent6 : 1;
-			UPROPERTY(EditAnywhere, Category = Default)
-			uint32 bSupportsAgent7 : 1;
-			UPROPERTY(EditAnywhere, Category = Default)
-			uint32 bSupportsAgent8 : 1;
-			UPROPERTY(EditAnywhere, Category = Default)
-			uint32 bSupportsAgent9 : 1;
-			UPROPERTY(EditAnywhere, Category = Default)
-			uint32 bSupportsAgent10 : 1;
-			UPROPERTY(EditAnywhere, Category = Default)
-			uint32 bSupportsAgent11 : 1;
-			UPROPERTY(EditAnywhere, Category = Default)
-			uint32 bSupportsAgent12 : 1;
-			UPROPERTY(EditAnywhere, Category = Default)
-			uint32 bSupportsAgent13 : 1;
-			UPROPERTY(EditAnywhere, Category = Default)
-			uint32 bSupportsAgent14 : 1;
-			UPROPERTY(EditAnywhere, Category = Default)
-			uint32 bSupportsAgent15 : 1;
-#if CPP
-		};
-		uint32 PackedBits;
-	};
-#endif
-
-	explicit FNavAgentSelector(const uint32 InBits = 0x7fffffff);
-
-	FORCEINLINE bool Contains(int32 AgentIndex) const
-	{
-		return (AgentIndex >= 0 && AgentIndex < 16) ? !!(PackedBits & (1 << AgentIndex)) : false;
-	}
-
-	FORCEINLINE void Set(int32 AgentIndex)
-	{
-		if (AgentIndex >= 0 && AgentIndex < 16)
-		{
-			PackedBits |= (1 << AgentIndex);
-		}
-	}
-
-	FORCEINLINE bool IsInitialized() const
-	{
-		return (PackedBits & InitializedBit) != 0;
-	}
-
-	FORCEINLINE void MarkInitialized()
-	{
-		PackedBits |= InitializedBit;
-	}
-
-	FORCEINLINE void Empty()
-	{
-		PackedBits = 0;
-	}
-
-	bool IsSame(const FNavAgentSelector& Other) const
-	{
-		return (~InitializedBit & PackedBits) == (~InitializedBit & Other.PackedBits);
-	}
-
-	bool Serialize(FArchive& Ar);
-
-	uint32 GetAgentBits() const 
-	{
-		return (~InitializedBit & PackedBits);
-	}
-};
-
-template<>
-struct TStructOpsTypeTraits< FNavAgentSelector > : public TStructOpsTypeTraitsBase2< FNavAgentSelector >
-{
-	enum
-	{
-		WithSerializer = true,
-	};
-};
-
-struct FNavigationBounds
-{
-	uint32 UniqueID;
-	FBox AreaBox;
-	FNavAgentSelector SupportedAgents;
-	TWeakObjectPtr<ULevel> Level;		// The level this bounds belongs to
-
-	bool operator==(const FNavigationBounds& Other) const 
-	{ 
-		return UniqueID == Other.UniqueID; 
-	}
-
-	friend uint32 GetTypeHash(const FNavigationBounds& NavBounds)
-	{
-		return GetTypeHash(NavBounds.UniqueID);
-	}
-};
-
-struct FNavigationBoundsUpdateRequest 
-{
-	FNavigationBounds NavBounds;
-	
-	enum Type
-	{
-		Added,
-		Removed,
-		Updated,
-	};
-
-	Type UpdateRequest;
-};
-
 UENUM()
 enum class ENavDataGatheringMode : uint8
 {
diff --git a/Engine/Source/Runtime/Engine/Classes/AI/NavigationSystemConfig.h b/Engine/Source/Runtime/Engine/Classes/AI/NavigationSystemConfig.h
--- a/Engine/Source/Runtime/Engine/Classes/AI/NavigationSystemConfig.h
+++ b/Engine/Source/Runtime/Engine/Classes/AI/NavigationSystemConfig.h
@@ -3,6 +3,7 @@
 #pragma once
 
 #include "AI/Navigation/NavigationTypes.h"
+#include "AI/Navigation/NavAgentSelector.h"
 #include "NavigationSystemConfig.generated.h"
 
 class UNavigationSystemBase;
diff --git a/Engine/Source/Runtime/Engine/Public/AI/NavigationModifier.h b/Engine/Source/Runtime/Engine/Public/AI/NavigationModifier.h
--- a/Engine/Source/Runtime/Engine/Public/AI/NavigationModifier.h
+++ b/Engine/Source/Runtime/Engine/Public/AI/NavigationModifier.h
@@ -7,6 +7,8 @@
 #include "Templates/SubclassOf.h"
 #include "UObject/WeakObjectPtrTemplates.h"
 #include "AI/Navigation/NavLinkDefinition.h"
+#include "AI/Navigation/NavigationDataResolution.h"
+#include "AI/Navigation/NavigationTypes.h" //FNavDataPerInstanceTransformDelegate
 
 class UBrushComponent;
 class UPrimitiveComponent;
diff --git a/Engine/Source/Runtime/NavigationSystem/Public/AbstractNavData.h b/Engine/Source/Runtime/NavigationSystem/Public/AbstractNavData.h
--- a/Engine/Source/Runtime/NavigationSystem/Public/AbstractNavData.h
+++ b/Engine/Source/Runtime/NavigationSystem/Public/AbstractNavData.h
@@ -10,7 +10,9 @@
 #include "Math/Box.h"
 #include "Math/UnrealMathSSE.h"
 #include "Math/Vector.h"
+#if UE_ENABLE_INCLUDE_ORDER_DEPRECATED_IN_5_2
 #include "NavFilters/NavigationQueryFilter.h"
+#endif //UE_ENABLE_INCLUDE_ORDER_DEPRECATED_IN_5_2
 #include "NavigationData.h"
 #include "UObject/ObjectMacros.h"
 #include "UObject/UObjectGlobals.h"
diff --git a/Engine/Source/Runtime/NavigationSystem/Public/NavAreas/NavArea.h b/Engine/Source/Runtime/NavigationSystem/Public/NavAreas/NavArea.h
--- a/Engine/Source/Runtime/NavigationSystem/Public/NavAreas/NavArea.h
+++ b/Engine/Source/Runtime/NavigationSystem/Public/NavAreas/NavArea.h
@@ -7,6 +7,7 @@
 #include "UObject/Object.h"
 #include "Templates/SubclassOf.h"
 #include "AI/Navigation/NavAreaBase.h"
+#include "AI/Navigation/NavAgentSelector.h"
 #include "NavArea.generated.h"
 
 /** Class containing definition of a navigation area */
diff --git a/Engine/Source/Runtime/NavigationSystem/Public/NavLinkCustomComponent.h b/Engine/Source/Runtime/NavigationSystem/Public/NavLinkCustomComponent.h
--- a/Engine/Source/Runtime/NavigationSystem/Public/NavLinkCustomComponent.h
+++ b/Engine/Source/Runtime/NavigationSystem/Public/NavLinkCustomComponent.h
@@ -8,6 +8,7 @@
 #include "Engine/EngineTypes.h"
 #include "NavAreas/NavArea.h"
 #include "AI/Navigation/NavLinkDefinition.h"
+#include "AI/Navigation/NavAgentSelector.h"
 #include "NavLinkCustomInterface.h"
 #include "NavRelevantComponent.h"
 #include "NavLinkCustomComponent.generated.h"
diff --git a/Engine/Source/Runtime/NavigationSystem/Public/NavMesh/NavMeshBoundsVolume.h b/Engine/Source/Runtime/NavigationSystem/Public/NavMesh/NavMeshBoundsVolume.h
--- a/Engine/Source/Runtime/NavigationSystem/Public/NavMesh/NavMeshBoundsVolume.h
+++ b/Engine/Source/Runtime/NavigationSystem/Public/NavMesh/NavMeshBoundsVolume.h
@@ -11,6 +11,7 @@
 #include "UObject/ObjectMacros.h"
 #include "AI/Navigation/NavigationTypes.h"
 #include "GameFramework/Volume.h"
+#include "AI/Navigation/NavAgentSelector.h"
 #include "NavMeshBoundsVolume.generated.h"
 
 
diff --git a/Engine/Source/Runtime/NavigationSystem/Public/NavMesh/PImplRecastNavMesh.h b/Engine/Source/Runtime/NavigationSystem/Public/NavMesh/PImplRecastNavMesh.h
--- a/Engine/Source/Runtime/NavigationSystem/Public/NavMesh/PImplRecastNavMesh.h
+++ b/Engine/Source/Runtime/NavigationSystem/Public/NavMesh/PImplRecastNavMesh.h
@@ -291,7 +291,7 @@ public:
 
 protected:
 	/** 
-	 *	@param ForbiddenFlags polys with flags matching the fillter will get added to 
+	 *	@param ForbiddenFlags polys with flags matching the filter will get added to 
 	 */
 	int32 GetTilesDebugGeometry(const FRecastNavMeshGenerator* Generator, const dtMeshTile& Tile, int32 VertBase, FRecastDebugGeometry& OutGeometry, int32 TileIdx = INDEX_NONE, uint16 ForbiddenFlags = 0) const;
 
diff --git a/Engine/Source/Runtime/NavigationSystem/Public/NavMesh/RecastQueryFilter.h b/Engine/Source/Runtime/NavigationSystem/Public/NavMesh/RecastQueryFilter.h
--- a/Engine/Source/Runtime/NavigationSystem/Public/NavMesh/RecastQueryFilter.h
+++ b/Engine/Source/Runtime/NavigationSystem/Public/NavMesh/RecastQueryFilter.h
@@ -3,7 +3,9 @@
 #pragma once 
 
 #include "CoreMinimal.h"
+#if UE_ENABLE_INCLUDE_ORDER_DEPRECATED_IN_5_2
 #include "NavFilters/NavigationQueryFilter.h"
+#endif //UE_ENABLE_INCLUDE_ORDER_DEPRECATED_IN_5_2
 #include "AI/Navigation/NavigationTypes.h"
 
 #if WITH_RECAST
diff --git a/Engine/Source/Runtime/NavigationSystem/Public/NavigationSystem.h b/Engine/Source/Runtime/NavigationSystem/Public/NavigationSystem.h
--- a/Engine/Source/Runtime/NavigationSystem/Public/NavigationSystem.h
+++ b/Engine/Source/Runtime/NavigationSystem/Public/NavigationSystem.h
@@ -20,6 +20,7 @@
 #include "NavigationOctreeController.h"
 #include "NavigationDirtyAreasController.h"
 #include "Math/MovingWindowAverageFast.h"
+#include "AI/Navigation/NavigationBounds.h"
 #if UE_ENABLE_INCLUDE_ORDER_DEPRECATED_IN_5_2
 #if WITH_EDITOR
 #include "UnrealEdMisc.h"
