#ue5 - Add SlerpQuats(), FindQuatBetweenVectors(), FindQuatBetweenNormals() to Blueprints. Add SlerpVectorToDirection() and SlerpNormals() to both C++ and Blueprints.

diff --git a/Engine/Source/Runtime/Core/Private/Math/UnrealMath.cpp b/Engine/Source/Runtime/Core/Private/Math/UnrealMath.cpp
--- a/Engine/Source/Runtime/Core/Private/Math/UnrealMath.cpp
+++ b/Engine/Source/Runtime/Core/Private/Math/UnrealMath.cpp
@@ -220,6 +220,39 @@ UE::Math::TRotator<T> UE::Math::TVector<T>::ToOrientationRotator() const
 	return R;
 }
 
+template<typename T>
+UE::Math::TVector<T> UE::Math::TVector<T>::SlerpVectorToDirection(UE::Math::TVector<T>& V, UE::Math::TVector<T>& Direction, T Alpha)
+{
+	using TVector = UE::Math::TVector<T>;
+	using TQuat = UE::Math::TQuat<T>;
+	
+	// Find rotation from A to B
+	const TQuat RotationQuat = TQuat::FindBetweenVectors(V, Direction);
+	const TVector Axis = RotationQuat.GetRotationAxis();
+	const T AngleRads = RotationQuat.GetAngle();
+
+	// Rotate from A toward B using portion of the angle specified by Alpha.
+	const TQuat DeltaQuat(Axis, AngleRads * Alpha);
+	TVector Result = DeltaQuat.RotateVector(V);
+	return Result;
+}
+
+template<typename T>
+UE::Math::TVector<T> UE::Math::TVector<T>::SlerpNormals(UE::Math::TVector<T>& NormalA, UE::Math::TVector<T>& NormalB, T Alpha)
+{
+	using TVector = UE::Math::TVector<T>;
+	using TQuat = UE::Math::TQuat<T>;
+
+	// Find rotation from A to B
+	const TQuat RotationQuat = TQuat::FindBetweenNormals(NormalA, NormalB);
+	const TVector Axis = RotationQuat.GetRotationAxis();
+	const T AngleRads = RotationQuat.GetAngle();
+
+	// Rotate from A toward B using portion of the angle specified by Alpha.
+	const TQuat DeltaQuat(Axis, AngleRads * Alpha);
+	TVector Result = DeltaQuat.RotateVector(NormalA);
+	return Result;
+}
 
 template<typename T>
 UE::Math::TRotator<T> UE::Math::TVector4<T>::ToOrientationRotator() const
diff --git a/Engine/Source/Runtime/Core/Public/Math/Vector.h b/Engine/Source/Runtime/Core/Public/Math/Vector.h
--- a/Engine/Source/Runtime/Core/Public/Math/Vector.h
+++ b/Engine/Source/Runtime/Core/Public/Math/Vector.h
@@ -893,6 +893,26 @@ public:
      */
     T HeadingAngle() const;
 
+	/**
+	 * Interpolate from a vector to the direction of another vector along a spherical path.
+	 * 
+	 * @param V Vector we interpolate from
+	 * @param Direction Target direction we interpolate to
+	 * @param Alpha interpolation amount, usually between 0-1
+	 * @return Vector after interpolating between Vector and Direction along a spherical path. The magnitude will remain the length of the starting vector.
+	 */
+	static CORE_API TVector<T> SlerpVectorToDirection(TVector<T>& V, TVector<T>& Direction, T Alpha);
+
+	/**
+	 * Interpolate from normalized vector A to normalized vector B along a spherical path.
+	 *
+	 * @param NormalA Start direction of interpolation, must be normalized.
+	 * @param NormalB End target direction of interpolation, must be normalized.
+	 * @param Alpha interpolation amount, usually between 0-1
+	 * @return Normalized vector after interpolating between NormalA and NormalB along a spherical path.
+	 */
+	static CORE_API TVector<T> SlerpNormals(TVector<T>& NormalA, TVector<T>& NormalB, T Alpha);
+
     /**
      * Create an orthonormal basis from a basis with at least two orthogonal vectors.
      * It may change the directions of the X and Y axes to make the basis orthogonal,
diff --git a/Engine/Source/Runtime/Engine/Classes/Kismet/KismetMathLibrary.h b/Engine/Source/Runtime/Engine/Classes/Kismet/KismetMathLibrary.h
--- a/Engine/Source/Runtime/Engine/Classes/Kismet/KismetMathLibrary.h
+++ b/Engine/Source/Runtime/Engine/Classes/Kismet/KismetMathLibrary.h
@@ -1370,13 +1370,35 @@ class ENGINE_API UKismetMathLibrary : public UBlueprintFunctionLibrary
 	 * Similar to the FRotator version, returns a result without roll such that it preserves the up vector.
 	 *
 	 * @note If you don't care about preserving the up vector and just want the most direct rotation, you can use the faster
-	 * 'FQuat::FindBetweenVectors(FVector::ForwardVector, YourVector)' or 'FQuat::FindBetweenNormals(...)' if you know the vector is of unit length.
+	 * 'FindBetweenVectors(ForwardVector, YourVector)' or 'FindBetweenNormals(...)' if you know the vector is of unit length.
 	 *
 	 * @return Quaternion from the Vector's direction, without any roll.
 	 */
 	UFUNCTION(BlueprintPure, meta=(DisplayName = "To Quaternion (Vector)", ScriptMethod = "Quaternion", Keywords="rotation rotate cast convert", BlueprintAutocast), Category="Math|Conversions")
 	static FQuat Conv_VectorToQuaternion(FVector InVec);
 
+	/**
+	 * Interpolate from a vector to the direction of another vector along a spherical path.
+	 *
+	 * @param Vector Vector we interpolate from
+	 * @param Direction Target direction we interpolate to
+	 * @param Alpha interpolation amount, usually between 0-1
+	 * @return Vector after interpolating between Vector and Direction along a spherical path. The magnitude will remain the length of the starting vector.
+	 */
+	UFUNCTION(BlueprintPure, meta = (DisplayName = "SlerpVectorToDirection", ScriptMethod = "SlerpVectors", Keywords = "spherical lerpvector vectorslerp interpolate"), Category = "Math|Vector")
+	static FVector Vector_SlerpVectorToDirection(FVector Vector, FVector Direction, double Alpha);
+
+	/**
+	 * Interpolate from normalized vector A to normalized vector B along a spherical path.
+	 *
+	 * @param NormalA Start direction of interpolation, must be normalized.
+	 * @param NormalB End target direction of interpolation, must be normalized.
+	 * @param Alpha interpolation amount, usually between 0-1
+	 * @return Vector after interpolating between A and B along a spherical path.
+	 */
+	UFUNCTION(BlueprintPure, meta = (DisplayName = "SlerpNormals", ScriptMethod = "SlerpNormals", Keywords = "spherical lerpnormal normalslerp interpolate"), Category = "Math|Vector")
+	static FVector Vector_SlerpNormals(FVector NormalA, FVector NormalB, double Alpha);
+
 	/** Vector addition */
 	UFUNCTION(BlueprintPure, meta=(DisplayName = "vector + vector", CompactNodeTitle = "+", ScriptMethod = "Add", ScriptOperator = "+;+=", Keywords = "+ add plus", CommutativeAssociativeBinaryOperator = "true"), Category="Math|Vector")
 	static FVector Add_VectorVector(FVector A, FVector B);
@@ -2081,7 +2103,7 @@ class ENGINE_API UKismetMathLibrary : public UBlueprintFunctionLibrary
 	 * Similar to the FRotator version, returns a result without roll such that it preserves the up vector.
 	 *
 	 * @note If you don't care about preserving the up vector and just want the most direct rotation, you can use the faster
-	 * 'FQuat::FindBetweenVectors(FVector::ForwardVector, YourVector)' or 'FQuat::FindBetweenNormals(...)' if you know the vector is of unit length.
+	 * 'FindBetweenVectors(ForwardVector, YourVector)' or 'FindBetweenNormals(...)' if you know the vector is of unit length.
 	 *
 	 * @return Quaternion from the Vector's direction, without any roll.
 	 */
@@ -2957,6 +2979,36 @@ class ENGINE_API UKismetMathLibrary : public UBlueprintFunctionLibrary
 	UFUNCTION(BlueprintPure, meta = (DisplayName = "Unrotate Vector (Quat)", ScriptMethod = "UnrotateVector"), Category = "Math|Quat")
 	static FVector Quat_UnrotateVector(const FQuat& Q, const FVector& V);
 
+	/**
+	 * Spherical interpolation between Quaternions. Result is normalized.
+	 * 
+	 * @param A The starting quat we interp from
+	 * @param B The target quat we interp to
+	 * @param Alpha The interpolation amount, usually 0 to 1.
+	 * @return Quat after spherical interpolation
+	 */
+	UFUNCTION(BlueprintPure, meta = (DisplayName = "Slerp (Quat)", ScriptMethod = "SlerpQuat", Keywords = "spherical interpolate"), Category = "Math|Quat")
+	static FQuat Quat_Slerp(const FQuat& A, const FQuat& B, double Alpha);
+
+	/**
+	 * Generates the 'smallest' (geodesic) rotation around a sphere between two vectors of arbitrary length.
+	 * 
+	 * @param Start Vector the rotation starts from
+	 * @param End Vector the rotation ends at
+	 * @return Quat that will rotate from Start to End
+	 */
+	UFUNCTION(BlueprintPure, meta = (DisplayName = "Find Quat Between Vectors", ScriptMethod = "FindQuatBetweenVectors", Keywords = "FindQuat FindBetween"), Category = "Math|Quat")
+	static FQuat Quat_FindBetweenVectors(FVector Start, FVector End);
+
+	/**
+	 * Generates the 'smallest' (geodesic) rotation around a sphere between two normals (assumed to be unit length).
+	 * 
+	 * @param Start Normalized vector the rotation starts from
+	 * @param End Normalized vector the rotation ends at
+	 * @return Quat that will rotate from Start to End
+	 */
+	UFUNCTION(BlueprintPure, meta = (DisplayName = "Find Quat Between Normals", ScriptMethod = "FindQuatBetweenNormals", Keywords = "FindQuat FindBetween"), Category = "Math|Quat")
+	static FQuat Quat_FindBetweenNormals(FVector StartNormal, FVector EndNormal);
 
 	//
 	// LinearColor constants - exposed for scripting
diff --git a/Engine/Source/Runtime/Engine/Classes/Kismet/KismetMathLibrary.inl b/Engine/Source/Runtime/Engine/Classes/Kismet/KismetMathLibrary.inl
--- a/Engine/Source/Runtime/Engine/Classes/Kismet/KismetMathLibrary.inl
+++ b/Engine/Source/Runtime/Engine/Classes/Kismet/KismetMathLibrary.inl
@@ -1528,6 +1528,18 @@ FQuat UKismetMathLibrary::Conv_VectorToQuaternion(FVector InVec)
 	return InVec.ToOrientationQuat();
 }
 
+KISMET_MATH_FORCEINLINE
+FVector UKismetMathLibrary::Vector_SlerpVectorToDirection(FVector Vector, FVector Direction, double Alpha)
+{
+	return FVector::SlerpVectorToDirection(Vector, Direction, Alpha);
+}
+
+KISMET_MATH_FORCEINLINE
+FVector UKismetMathLibrary::Vector_SlerpNormals(FVector NormalA, FVector NormalB, double Alpha)
+{
+	return FVector::SlerpNormals(NormalA, NormalB, Alpha);
+}
+
 KISMET_MATH_FORCEINLINE
 FLinearColor UKismetMathLibrary::Conv_VectorToLinearColor(FVector InVec)
 {
@@ -3495,6 +3507,24 @@ FVector UKismetMathLibrary::Quat_UnrotateVector(const FQuat& Q, const FVector& V
 	return Q.UnrotateVector(V);
 }
 
+KISMET_MATH_FORCEINLINE
+FQuat UKismetMathLibrary::Quat_Slerp(const FQuat& A, const FQuat& B, double Alpha)
+{
+	return FQuat::Slerp(A, B, Alpha);
+}
+
+KISMET_MATH_FORCEINLINE
+FQuat UKismetMathLibrary::Quat_FindBetweenVectors(FVector Start, FVector End)
+{
+	return FQuat::FindBetweenVectors(Start, End);
+}
+
+KISMET_MATH_FORCEINLINE
+FQuat UKismetMathLibrary::Quat_FindBetweenNormals(FVector StartNormal, FVector EndNormal)
+{
+	return FQuat::FindBetweenNormals(StartNormal, EndNormal);
+}
+
 
 /* FLinearColor
 *****************************************************************************/
