Add owner name to track skeletal & static mesh DX12 ray tracing buffers

diff --git a/Engine/Source/Runtime/D3D12RHI/Private/D3D12RayTracing.cpp b/Engine/Source/Runtime/D3D12RHI/Private/D3D12RayTracing.cpp
--- a/Engine/Source/Runtime/D3D12RHI/Private/D3D12RayTracing.cpp
+++ b/Engine/Source/Runtime/D3D12RHI/Private/D3D12RayTracing.cpp
@@ -3245,6 +3245,7 @@ FD3D12RayTracingGeometry::FD3D12RayTracingGeometry(FD3D12Adapter* Adapter, const
 	INC_DWORD_STAT(STAT_D3D12RayTracingAllocatedBLAS);
 
 	DebugName = !Initializer.DebugName.IsNone() ? Initializer.DebugName : FDebugName(FName(TEXT("BLAS")));
+	OwnerName = Initializer.OwnerName;
 	
 	FMemory::Memzero(bHasPendingCompactionRequests);
 	FMemory::Memzero(bRegisteredAsRenameListener);
@@ -3319,6 +3320,7 @@ FD3D12RayTracingGeometry::FD3D12RayTracingGeometry(FD3D12Adapter* Adapter, const
 	for (uint32 GPUIndex = 0; GPUIndex < MAX_NUM_GPUS && GPUIndex < GNumExplicitGPUsForRendering; ++GPUIndex)
 	{
 		AccelerationStructureBuffers[GPUIndex] = CreateRayTracingBuffer(Adapter, GPUIndex, SizeInfo.ResultSize, ERayTracingBufferType::AccelerationStructure, DebugName);
+		AccelerationStructureBuffers[GPUIndex]->SetOwnerName(OwnerName);
 
 		INC_MEMORY_STAT_BY(STAT_D3D12RayTracingUsedVideoMemory, AccelerationStructureBuffers[GPUIndex]->GetSize());
 		INC_MEMORY_STAT_BY(STAT_D3D12RayTracingBLASMemory, AccelerationStructureBuffers[GPUIndex]->GetSize());
@@ -3863,6 +3865,7 @@ void FD3D12RayTracingGeometry::CompactAccelerationStructure(FD3D12CommandContext
 	TRefCountPtr<FD3D12Buffer> OldAccelerationStructure = AccelerationStructureBuffers[InGPUIndex];
 
 	AccelerationStructureBuffers[InGPUIndex] = CreateRayTracingBuffer(CommandContext.GetParentAdapter(), InGPUIndex, InSizeAfterCompaction, ERayTracingBufferType::AccelerationStructure, DebugName);
+	AccelerationStructureBuffers[InGPUIndex]->SetOwnerName(OwnerName);
 
 	INC_MEMORY_STAT_BY(STAT_D3D12RayTracingUsedVideoMemory, AccelerationStructureBuffers[InGPUIndex]->GetSize());
 	INC_MEMORY_STAT_BY(STAT_D3D12RayTracingBLASMemory, AccelerationStructureBuffers[InGPUIndex]->GetSize());
diff --git a/Engine/Source/Runtime/D3D12RHI/Private/D3D12RayTracing.h b/Engine/Source/Runtime/D3D12RHI/Private/D3D12RayTracing.h
--- a/Engine/Source/Runtime/D3D12RHI/Private/D3D12RayTracing.h
+++ b/Engine/Source/Runtime/D3D12RHI/Private/D3D12RayTracing.h
@@ -82,6 +82,7 @@ public:
 	TArray<FHitGroupSystemParameters> HitGroupSystemParameters[MAX_NUM_GPUS];
 
 	FDebugName DebugName;
+	FName OwnerName;		// Store the path name of the owner object for resource tracking
 
 	// Array of geometry descriptions, one per segment (single-segment geometry is a common case).
 	// Only references CPU-accessible structures (no GPU resources).
diff --git a/Engine/Source/Runtime/Engine/Classes/Engine/StaticMesh.h b/Engine/Source/Runtime/Engine/Classes/Engine/StaticMesh.h
--- a/Engine/Source/Runtime/Engine/Classes/Engine/StaticMesh.h
+++ b/Engine/Source/Runtime/Engine/Classes/Engine/StaticMesh.h
@@ -1908,6 +1908,9 @@ public:
 
 	inline bool AreRenderingResourcesInitialized() const { return bRenderingResourcesInitialized; }
 
+	/** Helper function for resource tracking, construct a name using the mesh's path name and LOD index . */
+	static FName GetLODPathName(const UStaticMesh* Mesh, int32 LODIndex);
+
 #if WITH_EDITOR
 
 	/**
diff --git a/Engine/Source/Runtime/Engine/Private/SkeletalRenderGPUSkin.cpp b/Engine/Source/Runtime/Engine/Private/SkeletalRenderGPUSkin.cpp
--- a/Engine/Source/Runtime/Engine/Private/SkeletalRenderGPUSkin.cpp
+++ b/Engine/Source/Runtime/Engine/Private/SkeletalRenderGPUSkin.cpp
@@ -824,6 +824,7 @@ void FSkeletalMeshObjectGPUSkin::UpdateRayTracingGeometry(FSkeletalMeshLODRender
 				Initializer.DebugName = FName(DefaultDebugName, DebugNumber++);
 			}
 
+			Initializer.OwnerName = GetAssetPathName(LODIndex);
 			Initializer.IndexBuffer = IndexBufferRHI;
 			Initializer.TotalPrimitiveCount = TotalNumTriangles;
 			Initializer.GeometryType = RTGT_Triangles;
diff --git a/Engine/Source/Runtime/Engine/Private/StaticMesh.cpp b/Engine/Source/Runtime/Engine/Private/StaticMesh.cpp
--- a/Engine/Source/Runtime/Engine/Private/StaticMesh.cpp
+++ b/Engine/Source/Runtime/Engine/Private/StaticMesh.cpp
@@ -1375,16 +1375,16 @@ void FStaticMeshLODResources::InitResources(UStaticMesh* Parent, int32 LODIndex)
 	{
 		const bool bProceduralPrimitive = Parent->HasValidNaniteData() && Nanite::GetSupportsRayTracingProceduralPrimitive(GMaxRHIShaderPlatform);
 		ENQUEUE_RENDER_COMMAND(InitStaticMeshRayTracingGeometry)(
-			[this, DebugName = Parent->GetFName(), bProceduralPrimitive](FRHICommandListImmediate& RHICmdList)
+			[this, DebugName = Parent->GetFName(), bProceduralPrimitive, OwnerName](FRHICommandListImmediate& RHICmdList)
 			{
 				FRayTracingGeometryInitializer Initializer;
 				if (bProceduralPrimitive)
 				{
-					SetupRayTracingProceduralGeometryInitializer(Initializer, DebugName);
+					SetupRayTracingProceduralGeometryInitializer(Initializer, DebugName, OwnerName);
 				}
 				else
 				{
-					SetupRayTracingGeometryInitializer(Initializer, DebugName);
+					SetupRayTracingGeometryInitializer(Initializer, DebugName, OwnerName);
 				}
 				RayTracingGeometry.SetInitializer(Initializer);
 			}
@@ -1438,9 +1438,10 @@ void FStaticMeshLODResources::InitResources(UStaticMesh* Parent, int32 LODIndex)
 }
 
 #if RHI_RAYTRACING
-void FStaticMeshLODResources::SetupRayTracingGeometryInitializer(FRayTracingGeometryInitializer& Initializer, const FName& DebugName)
+void FStaticMeshLODResources::SetupRayTracingGeometryInitializer(FRayTracingGeometryInitializer& Initializer, const FName& DebugName, const FName& OwnerName)
 {
 	Initializer.DebugName = DebugName;
+	Initializer.OwnerName = OwnerName;
 	Initializer.IndexBuffer = IndexBuffer.IndexBufferRHI;
 	Initializer.TotalPrimitiveCount = 0; // This is calculated below based on static mesh section data
 	Initializer.GeometryType = RTGT_Triangles;
@@ -1466,9 +1467,10 @@ void FStaticMeshLODResources::SetupRayTracingGeometryInitializer(FRayTracingGeom
 	Initializer.Segments = GeometrySections;
 }
 
-void FStaticMeshLODResources::SetupRayTracingProceduralGeometryInitializer(FRayTracingGeometryInitializer& Initializer, const FName& DebugName)
+void FStaticMeshLODResources::SetupRayTracingProceduralGeometryInitializer(FRayTracingGeometryInitializer& Initializer, const FName& DebugName, const FName& OwnerName)
 {
 	Initializer.DebugName = DebugName;
+	Initializer.OwnerName = OwnerName;
 	Initializer.IndexBuffer = nullptr;
 	Initializer.TotalPrimitiveCount = 1; // one AABB
 	Initializer.GeometryType = RTGT_Procedural;
@@ -5432,6 +5434,15 @@ void UStaticMesh::CalculateExtendedBounds()
 	SetExtendedBounds(Bounds);
 }
 
+FName UStaticMesh::GetLODPathName(const UStaticMesh* Mesh, int32 LODIndex)
+{
+#if RHI_ENABLE_RESOURCE_INFO
+	return FName(FString::Printf(TEXT("%s [LOD%d]"), Mesh ? *Mesh->GetPathName() : TEXT("UnknownStaticMesh"), LODIndex));
+#else
+	return NAME_None;
+#endif
+}
+
 #if WITH_EDITORONLY_DATA
 FUObjectAnnotationSparseBool GStaticMeshesThatNeedMaterialFixup;
 #endif // #if WITH_EDITORONLY_DATA
diff --git a/Engine/Source/Runtime/Engine/Private/Streaming/StaticMeshUpdate.cpp b/Engine/Source/Runtime/Engine/Private/Streaming/StaticMeshUpdate.cpp
--- a/Engine/Source/Runtime/Engine/Private/Streaming/StaticMeshUpdate.cpp
+++ b/Engine/Source/Runtime/Engine/Private/Streaming/StaticMeshUpdate.cpp
@@ -182,14 +182,16 @@ void FStaticMeshStreamIn::CreateBuffers_Internal(const FContext& Context)
 			if (IsRayTracingEnabled() && Context.Mesh->bSupportRayTracing &&
 				LODResource.VertexBuffers.StaticMeshVertexBuffer.GetNumVertices() > 0)
 			{
+				const FName OwnerName = UStaticMesh::GetLODPathName(Context.Mesh, LODIdx);
+
 				FRayTracingGeometryInitializer Initializer;
 				if (Context.Mesh->HasValidNaniteData() && Nanite::GetSupportsRayTracingProceduralPrimitive(GMaxRHIShaderPlatform))
 				{
-					FStaticMeshLODResources::SetupRayTracingProceduralGeometryInitializer(Initializer, Context.Mesh->GetFName());
+					FStaticMeshLODResources::SetupRayTracingProceduralGeometryInitializer(Initializer, Context.Mesh->GetFName(), OwnerName);
 				}
 				else
 				{
-					Context.LODResourcesView[LODIdx]->SetupRayTracingGeometryInitializer(Initializer, Context.Mesh->GetFName());
+					Context.LODResourcesView[LODIdx]->SetupRayTracingGeometryInitializer(Initializer, Context.Mesh->GetFName(), OwnerName);
 				}
 				Initializer.Type = ERayTracingGeometryInitializerType::StreamingSource;
 				IntermediateRayTracingGeometry[LODIdx].SetInitializer(Initializer);
@@ -268,15 +270,17 @@ void FStaticMeshStreamIn::DoFinishUpdate(const FContext& Context)
 				// Skip LODs that have their render data stripped
 				if (LODResource.VertexBuffers.StaticMeshVertexBuffer.GetNumVertices() > 0)
 				{
+					const FName OwnerName = UStaticMesh::GetLODPathName(Context.Mesh, LODIndex);
+
 					// Rebuild the initializer because it could have been reset during a previous release
 					FRayTracingGeometryInitializer Initializer;
 					if (Context.Mesh->HasValidNaniteData() && Nanite::GetSupportsRayTracingProceduralPrimitive(GMaxRHIShaderPlatform))
 					{
-						FStaticMeshLODResources::SetupRayTracingProceduralGeometryInitializer(Initializer, Context.Mesh->GetFName());
+						FStaticMeshLODResources::SetupRayTracingProceduralGeometryInitializer(Initializer, Context.Mesh->GetFName(), OwnerName);
 					}
 					else
 					{
-						LODResource.SetupRayTracingGeometryInitializer(Initializer, Context.Mesh->GetFName());
+						LODResource.SetupRayTracingGeometryInitializer(Initializer, Context.Mesh->GetFName(), OwnerName);
 					}
 					LODResource.RayTracingGeometry.SetInitializer(Initializer);
 					LODResource.RayTracingGeometry.SetAsStreamedIn();			
diff --git a/Engine/Source/Runtime/Engine/Public/StaticMeshResources.h b/Engine/Source/Runtime/Engine/Public/StaticMeshResources.h
--- a/Engine/Source/Runtime/Engine/Public/StaticMeshResources.h
+++ b/Engine/Source/Runtime/Engine/Public/StaticMeshResources.h
@@ -471,8 +471,8 @@ public:
 	void GetResourceSizeEx(FResourceSizeEx& CumulativeResourceSize) const;
 
 #if RHI_RAYTRACING
-	void SetupRayTracingGeometryInitializer(FRayTracingGeometryInitializer& Initializer, const FName& DebugName);
-	static void SetupRayTracingProceduralGeometryInitializer(FRayTracingGeometryInitializer& Initializer, const FName& DebugName);
+	void SetupRayTracingGeometryInitializer(FRayTracingGeometryInitializer& Initializer, const FName& DebugName, const FName& OwnerName);
+	static void SetupRayTracingProceduralGeometryInitializer(FRayTracingGeometryInitializer& Initializer, const FName& DebugName, const FName& OwnerName);
 #endif // RHI_RAYTRACING
 
 	/** Get the estimated memory overhead of buffers marked as NeedsCPUAccess. */
diff --git a/Engine/Source/Runtime/RHI/Public/RHIResources.h b/Engine/Source/Runtime/RHI/Public/RHIResources.h
--- a/Engine/Source/Runtime/RHI/Public/RHIResources.h
+++ b/Engine/Source/Runtime/RHI/Public/RHIResources.h
@@ -2465,7 +2465,10 @@ public:
 	LAYOUT_FIELD_INITIALIZED(bool, bAllowCompaction, true);
 	LAYOUT_FIELD_INITIALIZED(ERayTracingGeometryInitializerType, Type, ERayTracingGeometryInitializerType::Rendering);
 
+	// Use FDebugName for auto-generated debug names with numbered suffixes, it is a variation of FMemoryImageName with optional number postfix.
 	LAYOUT_FIELD(FDebugName, DebugName);
+	// Store the path name of the owner object for resource tracking. FMemoryImageName allows a conversion to/from FName.
+	LAYOUT_FIELD(FMemoryImageName, OwnerName);
 };
 
 enum ERayTracingSceneLifetime
