Set the client world package name on recording demo connections

diff --git a/Engine/Source/Runtime/Engine/Private/DemoNetDriver.cpp b/Engine/Source/Runtime/Engine/Private/DemoNetDriver.cpp
--- a/Engine/Source/Runtime/Engine/Private/DemoNetDriver.cpp
+++ b/Engine/Source/Runtime/Engine/Private/DemoNetDriver.cpp
@@ -1103,6 +1103,8 @@ bool UDemoNetDriver::InitListen(FNetworkNotify* InNotify, FURL& ListenURL, bool
 
 	AddClientConnection(Connection);
 
+	Connection->SetClientWorldPackageName(World->GetOutermost()->GetFName());
+
 	// Technically, NetDriver's can be renamed so this could become stale.
 	// However, it's only used for logging and DemoNetDriver's are typically given a special name.
 	BudgetLogHelper = MakeUnique<FDemoBudgetLogHelper>(NetDriverName.ToString());
