Include more HLOD properties from PrimitiveComponents in cooked data
* Compacted ExcludeForSpecificHLODLevels from a TArray<> to a bit mask (16 bytes -> 1 byte)

diff --git a/Engine/Plugins/Experimental/VirtualHeightfieldMesh/Source/VirtualHeightfieldMesh/Private/VirtualHeightfieldMeshComponent.cpp b/Engine/Plugins/Experimental/VirtualHeightfieldMesh/Source/VirtualHeightfieldMesh/Private/VirtualHeightfieldMeshComponent.cpp
--- a/Engine/Plugins/Experimental/VirtualHeightfieldMesh/Source/VirtualHeightfieldMesh/Private/VirtualHeightfieldMeshComponent.cpp
+++ b/Engine/Plugins/Experimental/VirtualHeightfieldMesh/Source/VirtualHeightfieldMesh/Private/VirtualHeightfieldMeshComponent.cpp
@@ -21,9 +21,7 @@ UVirtualHeightfieldMeshComponent::UVirtualHeightfieldMeshComponent(const FObject
 	bAffectDynamicIndirectLighting = false;
 	bAffectDistanceFieldLighting = false;
 	bNeverDistanceCull = true;
-#if WITH_EDITORONLY_DATA
 	bEnableAutoLODGeneration = false;
-#endif // WITH_EDITORONLY_DATA
 	Mobility = EComponentMobility::Static;
 }
 
diff --git a/Engine/Plugins/Experimental/Water/Source/Runtime/Private/WaterBodyComponent.cpp b/Engine/Plugins/Experimental/Water/Source/Runtime/Private/WaterBodyComponent.cpp
--- a/Engine/Plugins/Experimental/Water/Source/Runtime/Private/WaterBodyComponent.cpp
+++ b/Engine/Plugins/Experimental/Water/Source/Runtime/Private/WaterBodyComponent.cpp
@@ -128,11 +128,7 @@ UWaterBodyComponent::UWaterBodyComponent(const FObjectInitializer& ObjectInitial
 
 bool UWaterBodyComponent::IsHLODRelevant() const
 {
-#if WITH_EDITOR
 	return bEnableAutoLODGeneration;
-#else
-	return false;
-#endif
 }
 
 void UWaterBodyComponent::OnVisibilityChanged()
diff --git a/Engine/Plugins/Experimental/Water/Source/Runtime/Private/WaterBodyCustomComponent.cpp b/Engine/Plugins/Experimental/Water/Source/Runtime/Private/WaterBodyCustomComponent.cpp
--- a/Engine/Plugins/Experimental/Water/Source/Runtime/Private/WaterBodyCustomComponent.cpp
+++ b/Engine/Plugins/Experimental/Water/Source/Runtime/Private/WaterBodyCustomComponent.cpp
@@ -72,9 +72,7 @@ void UWaterBodyCustomComponent::OnUpdateBody(bool bWithExclusionVolumes)
 		MeshComp = NewObject<UStaticMeshComponent>(OwnerActor, TEXT("CustomMeshComponent"), RF_Transactional);
 		MeshComp->SetNetAddressable(); // it's deterministically named so it's addressable over network (needed for collision)
 		MeshComp->SetupAttachment(this);
-#if WITH_EDITORONLY_DATA
 		MeshComp->bEnableAutoLODGeneration = false; // The water body component is already responsible for generating HLOD of the WaterBody actor.
-#endif // WITH_EDITORONLY_DATA
 
 		if(IsRegistered())
 		{
diff --git a/Engine/Source/Runtime/Engine/Classes/Components/PrimitiveComponent.h b/Engine/Source/Runtime/Engine/Classes/Components/PrimitiveComponent.h
--- a/Engine/Source/Runtime/Engine/Classes/Components/PrimitiveComponent.h
+++ b/Engine/Source/Runtime/Engine/Classes/Components/PrimitiveComponent.h
@@ -27,6 +27,7 @@
 #include "HitProxies.h"
 #include "Interfaces/Interface_AsyncCompilation.h"
 #include "HLOD/HLODBatchingPolicy.h"
+#include "HLOD/HLODLevelExclusion.h"
 #include "Stats/Stats2.h"
 #if UE_ENABLE_INCLUDE_ORDER_DEPRECATED_IN_5_1
 #include "Engine/OverlapInfo.h"
@@ -305,34 +306,29 @@ public:
 	UPROPERTY(EditAnywhere, AdvancedDisplay, BlueprintReadOnly, Category=Lighting)
 	ELightmapType LightmapType;
 
-#if WITH_EDITORONLY_DATA
-	/** Whether to include this component in HLODs or not. */
-	UPROPERTY(EditAnywhere, BlueprintReadWrite, Category=HLOD, meta=(DisplayName="Include Component in HLOD"))
-	uint8 bEnableAutoLODGeneration : 1;
-
 	/** Which specific HLOD levels this component should be excluded from */
-	UPROPERTY(EditAnywhere, AdvancedDisplay, BlueprintReadWrite, Category=HLOD, meta=(DisplayName="Exclude from HLOD Levels", EditConditionHides, EditCondition="bEnableAutoLODGeneration"))
-	TArray<int32> ExcludeForSpecificHLODLevels;
+	UPROPERTY(EditAnywhere, AdvancedDisplay, BlueprintReadWrite, Category=HLOD, meta=(Bitmask, BitmaskEnum = "EHLODLevelExclusion", DisplayName="Exclude from HLOD Levels", DisplayAfter="bEnableAutoLODGeneration", EditConditionHides, EditCondition="bEnableAutoLODGeneration"))
+	uint8 ExcludeFromSpecificHLODLevels;
 
 	/** Determines how the geometry of a component will be incorporated in proxy (simplified) HLODs. */
-	UPROPERTY(EditAnywhere, AdvancedDisplay, BlueprintReadWrite, Category=HLOD, meta=(DisplayName="HLOD Batching Policy", EditConditionHides, EditCondition="bEnableAutoLODGeneration"))
+	UPROPERTY(EditAnywhere, AdvancedDisplay, BlueprintReadWrite, Category=HLOD, meta=(DisplayName="HLOD Batching Policy", DisplayAfter="bEnableAutoLODGeneration", EditConditionHides, EditCondition="bEnableAutoLODGeneration"))
 	EHLODBatchingPolicy HLODBatchingPolicy;
 
+	/** Whether to include this component in HLODs or not. */
+	UPROPERTY(EditAnywhere, BlueprintReadWrite, Category = HLOD, meta=(DisplayName="Include Component in HLOD"))
+	uint8 bEnableAutoLODGeneration : 1;
 
 	/** Indicates that the texture streaming built data is local to the Actor (see UActorTextureStreamingBuildDataComponent). */
 	UPROPERTY()
 	uint8 bIsActorTextureStreamingBuiltData : 1;
-#endif 
 
 	/** Indicates to the texture streaming wether it can use the pre-built texture streaming data (even if empty). */
 	UPROPERTY()
 	uint8 bIsValidTextureStreamingBuiltData : 1;
 
-	/**
-	 * When enabled this object will not be culled by distance. This is ignored if a child of a HLOD.
-	 */
+	/** When enabled this object will not be culled by distance. This is ignored if a child of a HLOD. */
 	UPROPERTY(EditAnywhere, AdvancedDisplay, BlueprintReadOnly, Category=LOD)
-	uint8 bNeverDistanceCull:1;
+	uint8 bNeverDistanceCull : 1;
 
 	/** Whether this primitive is referenced by a FLevelRenderAssetManager  */
 	mutable uint8 bAttachedToStreamingManagerAsStatic : 1;
@@ -687,6 +683,9 @@ private:
 	/** Deprecated - represented by HLODBatchingPolicy == EHLODBatchingPolicy::Instancing */
 	UPROPERTY()
 	uint8 bBatchImpostersAsInstances_DEPRECATED : 1;
+
+	UPROPERTY()
+	TArray<int32> ExcludeForSpecificHLODLevels_DEPRECATED;
 #endif
 
 	FMaskFilter MoveIgnoreMask;
diff --git a/Engine/Source/Runtime/Engine/Private/Components/PrimitiveComponent.cpp b/Engine/Source/Runtime/Engine/Private/Components/PrimitiveComponent.cpp
--- a/Engine/Source/Runtime/Engine/Private/Components/PrimitiveComponent.cpp
+++ b/Engine/Source/Runtime/Engine/Private/Components/PrimitiveComponent.cpp
@@ -323,8 +323,12 @@ UPrimitiveComponent::UPrimitiveComponent(const FObjectInitializer& ObjectInitial
 
 	LDMaxDrawDistance = 0.f;
 	CachedMaxDrawDistance = 0.f;
-#if WITH_EDITORONLY_DATA
+
+	bEnableAutoLODGeneration = true;
 	HLODBatchingPolicy = EHLODBatchingPolicy::None;
+	ExcludeFromSpecificHLODLevels = 0;
+
+#if WITH_EDITORONLY_DATA
 	bUseMaxLODAsImposter_DEPRECATED = false;
 	bBatchImpostersAsInstances_DEPRECATED = false;
 #endif
@@ -361,7 +365,6 @@ UPrimitiveComponent::UPrimitiveComponent(const FObjectInitializer& ObjectInitial
 #endif
 
 #if WITH_EDITORONLY_DATA
-	bEnableAutoLODGeneration = true;
 	HitProxyPriority = HPP_World;
 #endif // WITH_EDITORONLY_DATA
 
@@ -1007,6 +1010,18 @@ void UPrimitiveComponent::Serialize(FArchive& Ar)
 			HLODBatchingPolicy = EHLODBatchingPolicy::Instancing;
 		}
 	}
+
+	if (!ExcludeForSpecificHLODLevels_DEPRECATED.IsEmpty())
+	{
+		ExcludeFromSpecificHLODLevels = 0;
+		for (int32 ExcludeFromLevel : ExcludeForSpecificHLODLevels_DEPRECATED)
+		{
+			if (ExcludeFromLevel < CHAR_BIT)
+			{
+				ExcludeFromSpecificHLODLevels |= (1 << ExcludeFromLevel);
+			}
+		}
+	}
 #endif
 }
 
@@ -4164,7 +4179,7 @@ const bool UPrimitiveComponent::ShouldGenerateAutoLOD(const int32 HierarchicalLe
 
 	// bAllowSpecificExclusion
 	bool bExcluded = false;
-	if (ExcludeForSpecificHLODLevels.Contains(HierarchicalLevelIndex))
+	if ((HierarchicalLevelIndex < CHAR_BIT) && (ExcludeFromSpecificHLODLevels & (1 << HierarchicalLevelIndex)))
 	{
 		const TArray<struct FHierarchicalSimplification>& HLODSetup = GetOwner()->GetLevel()->GetWorldSettings()->GetHierarchicalLODSetup();
 		if (HLODSetup.IsValidIndex(HierarchicalLevelIndex))
diff --git a/Engine/Source/Runtime/Engine/Public/HLOD/HLODLevelExclusion.h b/Engine/Source/Runtime/Engine/Public/HLOD/HLODLevelExclusion.h
new file mode 100644
--- /dev/null
+++ b/Engine/Source/Runtime/Engine/Public/HLOD/HLODLevelExclusion.h
@@ -0,0 +1,19 @@
+// Copyright Epic Games, Inc. All Rights Reserved.
+
+#pragma once
+
+#include "CoreMinimal.h"
+
+/** Bitflag enum to allow editing of UPrimitiveComponent::ExcludeFromSpecificHLODLevels as a bitmask in the properties */
+UENUM(meta = (Bitflags, UseEnumValuesAsMaskValuesInEditor = "true"))
+enum class EHLODLevelExclusion : uint8
+{
+	HLOD0 = 1 << 0,
+	HLOD1 = 1 << 1,
+	HLOD2 = 1 << 2,
+	HLOD3 = 1 << 3,
+	HLOD4 = 1 << 4,
+	HLOD5 = 1 << 5,
+	HLOD6 = 1 << 6,
+	HLOD7 = 1 << 7
+};
diff --git a/Engine/Source/Runtime/Foliage/Private/InstancedFoliage.cpp b/Engine/Source/Runtime/Foliage/Private/InstancedFoliage.cpp
--- a/Engine/Source/Runtime/Foliage/Private/InstancedFoliage.cpp
+++ b/Engine/Source/Runtime/Foliage/Private/InstancedFoliage.cpp
@@ -5871,9 +5871,7 @@ UFoliageInstancedStaticMeshComponent::UFoliageInstancedStaticMeshComponent(const
 	// * warnings and a fixup option for components which have instances with very large coordinates.
 	bUseTranslatedInstanceSpace = true;
 
-#if WITH_EDITORONLY_DATA
 	bEnableAutoLODGeneration = false;
-#endif
 
 	ViewRelevanceType = EHISMViewRelevanceType::Foliage;
 }
