Level Instance: don't enable world partition for packed actor levels.

diff --git a/Engine/Source/Runtime/Engine/Private/LevelInstance/LevelInstanceSubsystem.cpp b/Engine/Source/Runtime/Engine/Private/LevelInstance/LevelInstanceSubsystem.cpp
--- a/Engine/Source/Runtime/Engine/Private/LevelInstance/LevelInstanceSubsystem.cpp
+++ b/Engine/Source/Runtime/Engine/Private/LevelInstance/LevelInstanceSubsystem.cpp
@@ -1007,7 +1007,7 @@ ILevelInstanceInterface* ULevelInstanceSubsystem::CreateLevelInstanceFrom(const
 
 	ULevelStreamingLevelInstanceEditor* LevelStreaming = nullptr;
 	{
-		const bool bIsPartitioned = GetWorld()->IsPartitionedWorld();
+		const bool bIsPartitioned = (CreationParams.Type != ELevelInstanceCreationType::PackedLevelActor) && GetWorld()->IsPartitionedWorld();
 		LevelStreaming = StaticCast<ULevelStreamingLevelInstanceEditor*>(EditorLevelUtils::CreateNewStreamingLevelForWorld(
 		*GetWorld(), ULevelStreamingLevelInstanceEditor::StaticClass(), CreationParams.UseExternalActors(), LevelFilename, &ActorsToMove, CreationParams.TemplateWorld, /*bUseSaveAs*/true, bIsPartitioned, [this, bIsPartitioned, &ActorsToMove](ULevel* InLevel)
 		{
