Fixed CIS : Short type name "EHLODLevelExclusion" provided for TryFindType

diff --git a/Engine/Source/Runtime/Engine/Classes/Components/PrimitiveComponent.h b/Engine/Source/Runtime/Engine/Classes/Components/PrimitiveComponent.h
--- a/Engine/Source/Runtime/Engine/Classes/Components/PrimitiveComponent.h
+++ b/Engine/Source/Runtime/Engine/Classes/Components/PrimitiveComponent.h
@@ -307,7 +307,7 @@ public:
 	ELightmapType LightmapType;
 
 	/** Which specific HLOD levels this component should be excluded from */
-	UPROPERTY(EditAnywhere, AdvancedDisplay, BlueprintReadWrite, Category=HLOD, meta=(Bitmask, BitmaskEnum = "EHLODLevelExclusion", DisplayName="Exclude from HLOD Levels", DisplayAfter="bEnableAutoLODGeneration", EditConditionHides, EditCondition="bEnableAutoLODGeneration"))
+	UPROPERTY(EditAnywhere, AdvancedDisplay, BlueprintReadWrite, Category=HLOD, meta=(Bitmask, BitmaskEnum = "/Script/Engine.EHLODLevelExclusion", DisplayName="Exclude from HLOD Levels", DisplayAfter="bEnableAutoLODGeneration", EditConditionHides, EditCondition="bEnableAutoLODGeneration"))
 	uint8 ExcludeFromSpecificHLODLevels;
 
 	/** Determines how the geometry of a component will be incorporated in proxy (simplified) HLODs. */
