Updated RepRootMotion and ReplicatedBasedMovement replication to respect IsReplicatingMovement flag, similar to ReplicatedMovement

diff --git a/Engine/Source/Runtime/Engine/Private/Character.cpp b/Engine/Source/Runtime/Engine/Private/Character.cpp
--- a/Engine/Source/Runtime/Engine/Private/Character.cpp
+++ b/Engine/Source/Runtime/Engine/Private/Character.cpp
@@ -1163,6 +1163,12 @@ void ACharacter::PostNetReceive()
 
 void ACharacter::OnRep_ReplicatedBasedMovement()
 {	
+	// Following the same pattern in AActor::OnRep_ReplicatedMovement() just in case...
+	if (!IsReplicatingMovement())
+	{
+		return;
+	}
+
 	if (GetLocalRole() != ROLE_SimulatedProxy)
 	{
 		return;
@@ -1247,6 +1253,12 @@ FAnimMontageInstance * ACharacter::GetRootMotionAnimMontageInstance() const
 
 void ACharacter::OnRep_RootMotion()
 {
+	// Following the same pattern in AActor::OnRep_ReplicatedMovement() just in case...
+	if (!IsReplicatingMovement())
+	{
+		return;
+	}
+
 	if (GetLocalRole() == ROLE_SimulatedProxy)
 	{
 
@@ -1478,7 +1490,7 @@ void ACharacter::PreReplication( IRepChangedPropertyTracker & ChangedPropertyTra
 {
 	Super::PreReplication( ChangedPropertyTracker );
 
-	if (CharacterMovement->CurrentRootMotion.HasActiveRootMotionSources() || IsPlayingNetworkedRootMotionMontage())
+	if (IsReplicatingMovement() && (CharacterMovement->CurrentRootMotion.HasActiveRootMotionSources() || IsPlayingNetworkedRootMotionMontage()))
 	{
 		const FAnimMontageInstance* RootMotionMontageInstance = GetRootMotionAnimMontageInstance();
 
@@ -1515,21 +1527,27 @@ void ACharacter::PreReplication( IRepChangedPropertyTracker & ChangedPropertyTra
 
 	bProxyIsJumpForceApplied = (JumpForceTimeRemaining > 0.0f);
 	ReplicatedMovementMode = CharacterMovement->PackNetworkMovementMode();	
-	ReplicatedBasedMovement = BasedMovement;
 
-	// Optimization: only update and replicate these values if they are actually going to be used.
-	if (BasedMovement.HasRelativeLocation())
+	if(IsReplicatingMovement())
 	{
-		// When velocity becomes zero, force replication so the position is updated to match the server (it may have moved due to simulation on the client).
-		ReplicatedBasedMovement.bServerHasVelocity = !CharacterMovement->Velocity.IsZero();
+		ReplicatedBasedMovement = BasedMovement;
 
-		// Make sure absolute rotations are updated in case rotation occurred after the base info was saved.
-		if (!BasedMovement.HasRelativeRotation())
+		// Optimization: only update and replicate these values if they are actually going to be used.
+		if (BasedMovement.HasRelativeLocation())
 		{
-			ReplicatedBasedMovement.Rotation = GetActorRotation();
+			// When velocity becomes zero, force replication so the position is updated to match the server (it may have moved due to simulation on the client).
+			ReplicatedBasedMovement.bServerHasVelocity = !CharacterMovement->Velocity.IsZero();
+
+			// Make sure absolute rotations are updated in case rotation occurred after the base info was saved.
+			if (!BasedMovement.HasRelativeRotation())
+			{
+				ReplicatedBasedMovement.Rotation = GetActorRotation();
+			}
 		}
 	}
 
+	DOREPLIFETIME_ACTIVE_OVERRIDE_FAST(ACharacter, ReplicatedBasedMovement, IsReplicatingMovement());
+
 	// Save bandwidth by not replicating this value unless it is necessary, since it changes every update.
 	if ((CharacterMovement->NetworkSmoothingMode == ENetworkSmoothingMode::Linear) || CharacterMovement->bNetworkAlwaysReplicateTransformUpdateTimestamp)
 	{
