Reduce chance of network overflow when sending very large list of new object references

diff --git a/Engine/Source/Runtime/Engine/Classes/Engine/PackageMapClient.h b/Engine/Source/Runtime/Engine/Classes/Engine/PackageMapClient.h
--- a/Engine/Source/Runtime/Engine/Classes/Engine/PackageMapClient.h
+++ b/Engine/Source/Runtime/Engine/Classes/Engine/PackageMapClient.h
@@ -406,6 +406,7 @@ public:
 
 	void ReceiveNetGUIDBunch( FInBunch &InBunch );
 	void AppendExportBunches(TArray<FOutBunch *>& OutgoingBunches);
+	int32 GetNumExportBunches() const;
 
 	void AppendExportData(FArchive& Archive);
 	void ReceiveExportData(FArchive& Archive);
diff --git a/Engine/Source/Runtime/Engine/Private/PackageMapClient.cpp b/Engine/Source/Runtime/Engine/Private/PackageMapClient.cpp
--- a/Engine/Source/Runtime/Engine/Private/PackageMapClient.cpp
+++ b/Engine/Source/Runtime/Engine/Private/PackageMapClient.cpp
@@ -1990,6 +1990,11 @@ void UPackageMapClient::AppendExportBunches(TArray<FOutBunch *>& OutgoingBunches
 	}
 }
 
+int32 UPackageMapClient::GetNumExportBunches() const
+{
+	return ExportBunches.Num();
+}
+
 void UPackageMapClient::SyncPackageMapExportAckStatus( const UPackageMapClient* Source )
 {
 	AckState = Source->AckState;
