[Backout] - CL23465536
Original CL Desc
-----------------------------------------------------------------
Include more HLOD properties from PrimitiveComponents in cooked data
* Compacted ExcludeForSpecificHLODLevels from a TArray<> to a bit mask (16 bytes -> 1 byte)

diff --git a/Engine/Plugins/Experimental/VirtualHeightfieldMesh/Source/VirtualHeightfieldMesh/Private/VirtualHeightfieldMeshComponent.cpp b/Engine/Plugins/Experimental/VirtualHeightfieldMesh/Source/VirtualHeightfieldMesh/Private/VirtualHeightfieldMeshComponent.cpp
--- a/Engine/Plugins/Experimental/VirtualHeightfieldMesh/Source/VirtualHeightfieldMesh/Private/VirtualHeightfieldMeshComponent.cpp
+++ b/Engine/Plugins/Experimental/VirtualHeightfieldMesh/Source/VirtualHeightfieldMesh/Private/VirtualHeightfieldMeshComponent.cpp
@@ -21,7 +21,9 @@ UVirtualHeightfieldMeshComponent::UVirtualHeightfieldMeshComponent(const FObject
 	bAffectDynamicIndirectLighting = false;
 	bAffectDistanceFieldLighting = false;
 	bNeverDistanceCull = true;
+#if WITH_EDITORONLY_DATA
 	bEnableAutoLODGeneration = false;
+#endif // WITH_EDITORONLY_DATA
 	Mobility = EComponentMobility::Static;
 }
 
diff --git a/Engine/Plugins/Experimental/Water/Source/Runtime/Private/WaterBodyComponent.cpp b/Engine/Plugins/Experimental/Water/Source/Runtime/Private/WaterBodyComponent.cpp
--- a/Engine/Plugins/Experimental/Water/Source/Runtime/Private/WaterBodyComponent.cpp
+++ b/Engine/Plugins/Experimental/Water/Source/Runtime/Private/WaterBodyComponent.cpp
@@ -128,7 +128,11 @@ UWaterBodyComponent::UWaterBodyComponent(const FObjectInitializer& ObjectInitial
 
 bool UWaterBodyComponent::IsHLODRelevant() const
 {
+#if WITH_EDITOR
 	return bEnableAutoLODGeneration;
+#else
+	return false;
+#endif
 }
 
 void UWaterBodyComponent::OnVisibilityChanged()
diff --git a/Engine/Plugins/Experimental/Water/Source/Runtime/Private/WaterBodyCustomComponent.cpp b/Engine/Plugins/Experimental/Water/Source/Runtime/Private/WaterBodyCustomComponent.cpp
--- a/Engine/Plugins/Experimental/Water/Source/Runtime/Private/WaterBodyCustomComponent.cpp
+++ b/Engine/Plugins/Experimental/Water/Source/Runtime/Private/WaterBodyCustomComponent.cpp
@@ -72,7 +72,9 @@ void UWaterBodyCustomComponent::OnUpdateBody(bool bWithExclusionVolumes)
 		MeshComp = NewObject<UStaticMeshComponent>(OwnerActor, TEXT("CustomMeshComponent"), RF_Transactional);
 		MeshComp->SetNetAddressable(); // it's deterministically named so it's addressable over network (needed for collision)
 		MeshComp->SetupAttachment(this);
+#if WITH_EDITORONLY_DATA
 		MeshComp->bEnableAutoLODGeneration = false; // The water body component is already responsible for generating HLOD of the WaterBody actor.
+#endif // WITH_EDITORONLY_DATA
 
 		if(IsRegistered())
 		{
diff --git a/Engine/Source/Runtime/Engine/Classes/Components/PrimitiveComponent.h b/Engine/Source/Runtime/Engine/Classes/Components/PrimitiveComponent.h
--- a/Engine/Source/Runtime/Engine/Classes/Components/PrimitiveComponent.h
+++ b/Engine/Source/Runtime/Engine/Classes/Components/PrimitiveComponent.h
@@ -27,7 +27,6 @@
 #include "HitProxies.h"
 #include "Interfaces/Interface_AsyncCompilation.h"
 #include "HLOD/HLODBatchingPolicy.h"
-#include "HLOD/HLODLevelExclusion.h"
 #include "Stats/Stats2.h"
 #if UE_ENABLE_INCLUDE_ORDER_DEPRECATED_IN_5_1
 #include "Engine/OverlapInfo.h"
@@ -306,29 +305,34 @@ public:
 	UPROPERTY(EditAnywhere, AdvancedDisplay, BlueprintReadOnly, Category=Lighting)
 	ELightmapType LightmapType;
 
+#if WITH_EDITORONLY_DATA
+	/** Whether to include this component in HLODs or not. */
+	UPROPERTY(EditAnywhere, BlueprintReadWrite, Category=HLOD, meta=(DisplayName="Include Component in HLOD"))
+	uint8 bEnableAutoLODGeneration : 1;
+
 	/** Which specific HLOD levels this component should be excluded from */
-	UPROPERTY(EditAnywhere, AdvancedDisplay, BlueprintReadWrite, Category=HLOD, meta=(Bitmask, BitmaskEnum = "/Script/Engine.EHLODLevelExclusion", DisplayName="Exclude from HLOD Levels", DisplayAfter="bEnableAutoLODGeneration", EditConditionHides, EditCondition="bEnableAutoLODGeneration"))
-	uint8 ExcludeFromSpecificHLODLevels;
+	UPROPERTY(EditAnywhere, AdvancedDisplay, BlueprintReadWrite, Category=HLOD, meta=(DisplayName="Exclude from HLOD Levels", EditConditionHides, EditCondition="bEnableAutoLODGeneration"))
+	TArray<int32> ExcludeForSpecificHLODLevels;
 
 	/** Determines how the geometry of a component will be incorporated in proxy (simplified) HLODs. */
-	UPROPERTY(EditAnywhere, AdvancedDisplay, BlueprintReadWrite, Category=HLOD, meta=(DisplayName="HLOD Batching Policy", DisplayAfter="bEnableAutoLODGeneration", EditConditionHides, EditCondition="bEnableAutoLODGeneration"))
+	UPROPERTY(EditAnywhere, AdvancedDisplay, BlueprintReadWrite, Category=HLOD, meta=(DisplayName="HLOD Batching Policy", EditConditionHides, EditCondition="bEnableAutoLODGeneration"))
 	EHLODBatchingPolicy HLODBatchingPolicy;
 
-	/** Whether to include this component in HLODs or not. */
-	UPROPERTY(EditAnywhere, BlueprintReadWrite, Category = HLOD, meta=(DisplayName="Include Component in HLOD"))
-	uint8 bEnableAutoLODGeneration : 1;
 
 	/** Indicates that the texture streaming built data is local to the Actor (see UActorTextureStreamingBuildDataComponent). */
 	UPROPERTY()
 	uint8 bIsActorTextureStreamingBuiltData : 1;
+#endif 
 
 	/** Indicates to the texture streaming wether it can use the pre-built texture streaming data (even if empty). */
 	UPROPERTY()
 	uint8 bIsValidTextureStreamingBuiltData : 1;
 
-	/** When enabled this object will not be culled by distance. This is ignored if a child of a HLOD. */
+	/**
+	 * When enabled this object will not be culled by distance. This is ignored if a child of a HLOD.
+	 */
 	UPROPERTY(EditAnywhere, AdvancedDisplay, BlueprintReadOnly, Category=LOD)
-	uint8 bNeverDistanceCull : 1;
+	uint8 bNeverDistanceCull:1;
 
 	/** Whether this primitive is referenced by a FLevelRenderAssetManager  */
 	mutable uint8 bAttachedToStreamingManagerAsStatic : 1;
@@ -683,9 +687,6 @@ private:
 	/** Deprecated - represented by HLODBatchingPolicy == EHLODBatchingPolicy::Instancing */
 	UPROPERTY()
 	uint8 bBatchImpostersAsInstances_DEPRECATED : 1;
-
-	UPROPERTY()
-	TArray<int32> ExcludeForSpecificHLODLevels_DEPRECATED;
 #endif
 
 	FMaskFilter MoveIgnoreMask;
diff --git a/Engine/Source/Runtime/Engine/Private/Components/PrimitiveComponent.cpp b/Engine/Source/Runtime/Engine/Private/Components/PrimitiveComponent.cpp
--- a/Engine/Source/Runtime/Engine/Private/Components/PrimitiveComponent.cpp
+++ b/Engine/Source/Runtime/Engine/Private/Components/PrimitiveComponent.cpp
@@ -323,12 +323,8 @@ UPrimitiveComponent::UPrimitiveComponent(const FObjectInitializer& ObjectInitial
 
 	LDMaxDrawDistance = 0.f;
 	CachedMaxDrawDistance = 0.f;
-
-	bEnableAutoLODGeneration = true;
-	HLODBatchingPolicy = EHLODBatchingPolicy::None;
-	ExcludeFromSpecificHLODLevels = 0;
-
 #if WITH_EDITORONLY_DATA
+	HLODBatchingPolicy = EHLODBatchingPolicy::None;
 	bUseMaxLODAsImposter_DEPRECATED = false;
 	bBatchImpostersAsInstances_DEPRECATED = false;
 #endif
@@ -365,6 +361,7 @@ UPrimitiveComponent::UPrimitiveComponent(const FObjectInitializer& ObjectInitial
 #endif
 
 #if WITH_EDITORONLY_DATA
+	bEnableAutoLODGeneration = true;
 	HitProxyPriority = HPP_World;
 #endif // WITH_EDITORONLY_DATA
 
@@ -1010,18 +1007,6 @@ void UPrimitiveComponent::Serialize(FArchive& Ar)
 			HLODBatchingPolicy = EHLODBatchingPolicy::Instancing;
 		}
 	}
-
-	if (!ExcludeForSpecificHLODLevels_DEPRECATED.IsEmpty())
-	{
-		ExcludeFromSpecificHLODLevels = 0;
-		for (int32 ExcludeFromLevel : ExcludeForSpecificHLODLevels_DEPRECATED)
-		{
-			if (ExcludeFromLevel < CHAR_BIT)
-			{
-				ExcludeFromSpecificHLODLevels |= (1 << ExcludeFromLevel);
-			}
-		}
-	}
 #endif
 }
 
@@ -4179,7 +4164,7 @@ const bool UPrimitiveComponent::ShouldGenerateAutoLOD(const int32 HierarchicalLe
 
 	// bAllowSpecificExclusion
 	bool bExcluded = false;
-	if ((HierarchicalLevelIndex < CHAR_BIT) && (ExcludeFromSpecificHLODLevels & (1 << HierarchicalLevelIndex)))
+	if (ExcludeForSpecificHLODLevels.Contains(HierarchicalLevelIndex))
 	{
 		const TArray<struct FHierarchicalSimplification>& HLODSetup = GetOwner()->GetLevel()->GetWorldSettings()->GetHierarchicalLODSetup();
 		if (HLODSetup.IsValidIndex(HierarchicalLevelIndex))
diff --git a/Engine/Source/Runtime/Engine/Public/HLOD/HLODLevelExclusion.h b/Engine/Source/Runtime/Engine/Public/HLOD/HLODLevelExclusion.h
deleted file mode 100644
diff --git a/Engine/Source/Runtime/Foliage/Private/InstancedFoliage.cpp b/Engine/Source/Runtime/Foliage/Private/InstancedFoliage.cpp
--- a/Engine/Source/Runtime/Foliage/Private/InstancedFoliage.cpp
+++ b/Engine/Source/Runtime/Foliage/Private/InstancedFoliage.cpp
@@ -5871,7 +5871,9 @@ UFoliageInstancedStaticMeshComponent::UFoliageInstancedStaticMeshComponent(const
 	// * warnings and a fixup option for components which have instances with very large coordinates.
 	bUseTranslatedInstanceSpace = true;
 
+#if WITH_EDITORONLY_DATA
 	bEnableAutoLODGeneration = false;
+#endif
 
 	ViewRelevanceType = EHISMViewRelevanceType::Foliage;
 }
