GC reachability optimizations. ~1.5x speedup and ~2.5MB memory saved for an internal title.

diff --git a/Engine/Source/Runtime/Engine/Private/Model.cpp b/Engine/Source/Runtime/Engine/Private/Model.cpp
--- a/Engine/Source/Runtime/Engine/Private/Model.cpp
+++ b/Engine/Source/Runtime/Engine/Private/Model.cpp
@@ -546,11 +546,13 @@ void UModel::GetResourceSizeEx(FResourceSizeEx& CumulativeResourceSize)
 IMPLEMENT_INTRINSIC_CLASS(UModel, ENGINE_API, UObject, CORE_API, "/Script/Engine",
 	{
 		Class->CppClassStaticFunctions = UOBJECT_CPPCLASS_STATICFUNCTIONS_FORCLASS(UModel);
-		Class->EmitObjectReference(STRUCT_OFFSET(UModel, Polys), TEXT("Polys"));
-		const uint32 SkipIndexIndex = Class->EmitStructArrayBegin(STRUCT_OFFSET(UModel, Surfs), TEXT("Surfs"), sizeof(FBspSurf));
-		Class->EmitObjectReference(STRUCT_OFFSET(FBspSurf, Material), TEXT("Material"));
-		Class->EmitObjectReference(STRUCT_OFFSET(FBspSurf, Actor), TEXT("Actor"));
-		Class->EmitStructArrayEnd( SkipIndexIndex );
+
+		UE::GC::FTokenStreamBuilder& Builder = UE::GC::FIntrinsicClassTokens::AllocateBuilder(Class);
+		Builder.EmitObjectReference(STRUCT_OFFSET(UModel, Polys), TEXT("Polys"));
+		const uint32 SkipIndexIndex = Builder.EmitStructArrayBegin(STRUCT_OFFSET(UModel, Surfs), TEXT("Surfs"), sizeof(FBspSurf));
+		Builder.EmitObjectReference(STRUCT_OFFSET(FBspSurf, Material), TEXT("Material"));
+		Builder.EmitObjectReference(STRUCT_OFFSET(FBspSurf, Actor), TEXT("Actor"));
+		Builder.EmitStructArrayEnd( SkipIndexIndex );
 	}
 );
 
@@ -559,10 +561,12 @@ IMPLEMENT_INTRINSIC_CLASS(UModel, ENGINE_API, UObject, CORE_API, "/Script/Engine
 IMPLEMENT_INTRINSIC_CLASS(UModel, ENGINE_API, UObject, CORE_API, "/Script/Engine",
 	{
 		Class->CppClassStaticFunctions = UOBJECT_CPPCLASS_STATICFUNCTIONS_FORCLASS(UModel);
-		const uint32 SkipIndexIndex = Class->EmitStructArrayBegin(STRUCT_OFFSET(UModel, Surfs), TEXT("Surfs"), sizeof(FBspSurf));
-		Class->EmitObjectReference(STRUCT_OFFSET(FBspSurf, Material), TEXT("Material"));
-		Class->EmitObjectReference(STRUCT_OFFSET(FBspSurf, Actor), TEXT("Actor"));
-		Class->EmitStructArrayEnd( SkipIndexIndex );
+
+		UE::GC::FTokenStreamBuilder& Builder = UE::GC::FIntrinsicClassTokens::AllocateBuilder(Class);
+		const uint32 SkipIndexIndex = Builder.EmitStructArrayBegin(STRUCT_OFFSET(UModel, Surfs), TEXT("Surfs"), sizeof(FBspSurf));
+		Builder.EmitObjectReference(STRUCT_OFFSET(FBspSurf, Material), TEXT("Material"));
+		Builder.EmitObjectReference(STRUCT_OFFSET(FBspSurf, Actor), TEXT("Actor"));
+		Builder.EmitStructArrayEnd( SkipIndexIndex );
 	}
 );
 
diff --git a/Engine/Source/Runtime/Engine/Private/NetDriver.cpp b/Engine/Source/Runtime/Engine/Private/NetDriver.cpp
--- a/Engine/Source/Runtime/Engine/Private/NetDriver.cpp
+++ b/Engine/Source/Runtime/Engine/Private/NetDriver.cpp
@@ -3753,24 +3753,24 @@ void UNetDriver::AddReferencedObjects(UObject* InThis, FReferenceCollector& Coll
 	//			E.G., if we detected that FObjectReplicator / FReplicationChangelistMgrs detected
 	//			their associated objects were destroyed, we could destroy the Shadow Buffers
 	//			which should be the last thing referencing the FRepLayout and the Properties.
-	for (auto It = This->RepLayoutMap.CreateIterator(); It; ++It)
+	for (TPair<TWeakObjectPtr<UObject>, TSharedPtr<FRepLayout>>& Pair : This->RepLayoutMap)
 	{
-		It.Value()->AddReferencedObjects(Collector);
+		Pair.Value->AddReferencedObjects(Collector);
 	}
 	
 	for (FObjectReplicator* Replicator : This->AllOwnedReplicators)
 	{
 		if (Replicator->GetWeakObjectPtr().IsValid())
 		{
-			Collector.AddReferencedObject(Replicator->ObjectPtr, This);
+			Collector.AddStableReference(&Replicator->ObjectPtr);
 		}
 
-		Collector.AddReferencedObject(Replicator->ObjectClass, This);
+		Collector.AddStableReference(&Replicator->ObjectClass);
 	}
 
-	for (FConnectionMap::TIterator It(This->MappedClientConnections); It; ++It)
+	for (TPair<TSharedRef<const FInternetAddr>, UNetConnection*>& Pair : This->MappedClientConnections)
 	{
-		Collector.AddReferencedObject(It.Value(), This);
+		Collector.AddStableReference(&Pair.Value);
 	}
 
 	if (This->GuidCache.IsValid())
diff --git a/Engine/Source/Runtime/Engine/Private/StreamableManager.cpp b/Engine/Source/Runtime/Engine/Private/StreamableManager.cpp
--- a/Engine/Source/Runtime/Engine/Private/StreamableManager.cpp
+++ b/Engine/Source/Runtime/Engine/Private/StreamableManager.cpp
@@ -1042,6 +1042,7 @@ struct FStreamable
 };
 
 FStreamableManager::FStreamableManager()
+	: FGCObject(EFlags::AddStableNativeReferencesOnly)
 {
 	FCoreUObjectDelegates::GetPreGarbageCollectDelegate().AddRaw(this, &FStreamableManager::OnPreGarbageCollect);
 	bForceSynchronousLoads = false;
@@ -1109,22 +1110,14 @@ void FStreamableManager::OnPreGarbageCollect()
 void FStreamableManager::AddReferencedObjects(FReferenceCollector& Collector)
 {
 	// If there are active streamable handles in the editor, this will cause the user to Force Delete, which is irritating but necessary because weak pointers cannot be used here
-	for (TStreamableMap::TConstIterator It(StreamableItems); It; ++It)
+	for (TPair<FSoftObjectPath, FStreamable*>& Pair : StreamableItems)
 	{
-		FStreamable* Existing = It.Value();
-		if (Existing->Target)
-		{
-			Collector.AddReferencedObject(Existing->Target);
-		}
+		Collector.AddStableReference(&Pair.Value->Target);
 	}
-
-	for (TStreamableRedirects::TIterator It(StreamableRedirects); It; ++It)
+	
+	for (TPair<FSoftObjectPath, FRedirectedPath>& Pair : StreamableRedirects)
 	{
-		FRedirectedPath& Existing = It.Value();
-		if (Existing.LoadedRedirector)
-		{
-			Collector.AddReferencedObject(Existing.LoadedRedirector);
-		}
+		Collector.AddStableReference(&Pair.Value.LoadedRedirector);
 	}
 }
 
diff --git a/Engine/Source/Runtime/Engine/Private/Subsystems/SubsystemCollection.cpp b/Engine/Source/Runtime/Engine/Private/Subsystems/SubsystemCollection.cpp
--- a/Engine/Source/Runtime/Engine/Private/Subsystems/SubsystemCollection.cpp
+++ b/Engine/Source/Runtime/Engine/Private/Subsystems/SubsystemCollection.cpp
@@ -219,7 +219,7 @@ USubsystem* FSubsystemCollectionBase::InitializeDependency(TSubclassOf<USubsyste
 
 void FSubsystemCollectionBase::AddReferencedObjects(UObject* Referencer, FReferenceCollector& Collector)
 {
-	Collector.AddReferencedObjects(SubsystemMap, Referencer);
+	Collector.AddStableReferenceMap(SubsystemMap);
 }
 
 USubsystem* FSubsystemCollectionBase::AddAndInitializeSubsystem(UClass* SubsystemClass)
diff --git a/Engine/Source/Runtime/Engine/Private/UserDefinedStruct.cpp b/Engine/Source/Runtime/Engine/Private/UserDefinedStruct.cpp
--- a/Engine/Source/Runtime/Engine/Private/UserDefinedStruct.cpp
+++ b/Engine/Source/Runtime/Engine/Private/UserDefinedStruct.cpp
@@ -467,13 +467,11 @@ const uint8* UUserDefinedStruct::GetDefaultInstance() const
 void UUserDefinedStruct::AddReferencedObjects(UObject* InThis, FReferenceCollector& Collector)
 {
 	UUserDefinedStruct* This = CastChecked<UUserDefinedStruct>(InThis);
-
 	ensure(!This->DefaultStructInstance.IsValid() || This->DefaultStructInstance.GetStruct() == This);
-	uint8* StructData = This->DefaultStructInstance.GetStructMemory();
-	if (StructData)
+
+	if (uint8* StructData = This->DefaultStructInstance.GetStructMemory())
 	{
-		FVerySlowReferenceCollectorArchiveScope CollectorScope(Collector.GetVerySlowReferenceCollectorArchive(), This);
-		This->SerializeBin(FStructuredArchiveFromArchive(CollectorScope.GetArchive()).GetSlot(), StructData);
+		Collector.AddPropertyReferences(This, StructData, This);
 	}
 
 	Super::AddReferencedObjects(This, Collector);
diff --git a/Engine/Source/Runtime/UMG/Private/Slate/SObjectWidget.cpp b/Engine/Source/Runtime/UMG/Private/Slate/SObjectWidget.cpp
--- a/Engine/Source/Runtime/UMG/Private/Slate/SObjectWidget.cpp
+++ b/Engine/Source/Runtime/UMG/Private/Slate/SObjectWidget.cpp
@@ -52,6 +52,8 @@ void SObjectWidget::ResetWidget()
 {
 	if ( UObjectInitialized() && WidgetObject )
 	{
+		UnregisterGCObject();
+
 		if ( CanRouteEvent() )
 		{
 			WidgetObject->NativeDestruct();
@@ -97,7 +99,7 @@ FString SObjectWidget::GetReferencerName() const
 
 void SObjectWidget::AddReferencedObjects(FReferenceCollector& Collector)
 {
-	Collector.AddReferencedObject(WidgetObject);
+	Collector.AddStableReference(&WidgetObject);
 }
 
 void SObjectWidget::SetPadding(const TAttribute<FMargin>& InMargin)
