Remove required restart when changing Hardware Ray Tracing - toggling ray tracing at runtime is fully supported.
Allow changing r.SkinCache.Mode when ray tracing is allowed but not enabled.

diff --git a/Engine/Source/Runtime/Engine/Private/GPUSkinCache.cpp b/Engine/Source/Runtime/Engine/Private/GPUSkinCache.cpp
--- a/Engine/Source/Runtime/Engine/Private/GPUSkinCache.cpp
+++ b/Engine/Source/Runtime/Engine/Private/GPUSkinCache.cpp
@@ -2064,7 +2064,7 @@ void FGPUSkinCache::CVarSinkFunction()
 
 	if (GEnableGPUSkinCacheShaders)
 	{
-		if (GIsRHIInitialized && IsGPUSkinCacheRayTracingSupported())
+		if (GIsRHIInitialized && IsGPUSkinCacheRayTracingSupported() && IsRayTracingEnabled())
 		{
 			// Skin cache is *required* for ray tracing.
 			NewGPUSkinCacheValue = 1;
