TopLevelAssetPath Blueprint and Python support: Add FTopLevelAssetPath native Make function to correctly construct TopLevelAssetPath from a single string.
#rnx

diff --git a/Engine/Source/Runtime/CoreUObject/Public/UObject/NoExportTypes.h b/Engine/Source/Runtime/CoreUObject/Public/UObject/NoExportTypes.h
--- a/Engine/Source/Runtime/CoreUObject/Public/UObject/NoExportTypes.h
+++ b/Engine/Source/Runtime/CoreUObject/Public/UObject/NoExportTypes.h
@@ -2095,7 +2095,7 @@ struct FTimespan
  * A struct that can reference a top level asset such as '/Path/To/Package.AssetName'
  * @note The full C++ class is located here: Engine\Source\Runtime\CoreUObject\Public\UObject\TopLevelAssetPath.h
  */
-USTRUCT(noexport, BlueprintType)
+USTRUCT(noexport, BlueprintType, meta = (HasNativeMake = "/Script/Engine.KismetSystemLibrary.MakeTopLevelAssetPath", HasNativeBreak = "/Script/Engine.KismetSystemLibrary.BreakTopLevelAssetPath"))
 struct FTopLevelAssetPath
 {
 private:
diff --git a/Engine/Source/Runtime/Engine/Classes/Kismet/KismetSystemLibrary.h b/Engine/Source/Runtime/Engine/Classes/Kismet/KismetSystemLibrary.h
--- a/Engine/Source/Runtime/Engine/Classes/Kismet/KismetSystemLibrary.h
+++ b/Engine/Source/Runtime/Engine/Classes/Kismet/KismetSystemLibrary.h
@@ -16,6 +16,7 @@
 #include "UObject/TextProperty.h"
 #include "UObject/SoftObjectPtr.h"
 #include "UObject/PropertyAccessUtil.h"
+#include "UObject/TopLevelAssetPath.h"
 #include "Engine/LatentActionManager.h"
 #include "Kismet/BlueprintFunctionLibrary.h"
 #include "Engine/CollisionProfile.h"
@@ -265,7 +266,15 @@ class ENGINE_API UKismetSystemLibrary : public UBlueprintFunctionLibrary
 	UFUNCTION(BlueprintPure, meta = (DisplayName = "To Soft Object Path", CompactNodeTitle = "->", BlueprintAutocast), Category = "Utilities")
 	static FSoftObjectPath Conv_SoftObjRefToSoftObjPath(TSoftObjectPtr<UObject> SoftObjectReference);
 
-	/** 
+	/** Builds a TopLevelAssetPath struct from single Path string or from PackageName and AssetName string. */
+	UFUNCTION(BlueprintPure, Category = "Utilities", meta = (Keywords = "construct build", NativeMakeFunc, BlueprintThreadSafe, BlueprintAutocast))
+	static FORCENOINLINE FTopLevelAssetPath MakeTopLevelAssetPath(const FString& FullPathOrPackageName, const FString& AssetName);
+
+	/** Gets the path string out of a TopLevelAssetPath */
+	UFUNCTION(BlueprintPure, Category = "Utilities", meta = (NativeBreakFunc, BlueprintThreadSafe, BlueprintAutocast))
+	static void BreakTopLevelAssetPath(const FTopLevelAssetPath& TopLevelAssetPath, FString& PathString);
+
+	/**
 	 * Builds a Soft Class Path struct from a string that contains a full /folder/packagename.class path.
 	 * For blueprint classes, this needs to point to the actual class (often with _C) and not the blueprint editor asset
 	 */
diff --git a/Engine/Source/Runtime/Engine/Private/KismetSystemLibrary.cpp b/Engine/Source/Runtime/Engine/Private/KismetSystemLibrary.cpp
--- a/Engine/Source/Runtime/Engine/Private/KismetSystemLibrary.cpp
+++ b/Engine/Source/Runtime/Engine/Private/KismetSystemLibrary.cpp
@@ -1246,6 +1246,50 @@ FSoftObjectPath UKismetSystemLibrary::Conv_SoftObjRefToSoftObjPath(TSoftObjectPt
 	return SoftObjectReference.ToSoftObjectPath();
 }
 
+FTopLevelAssetPath UKismetSystemLibrary::MakeTopLevelAssetPath(const FString& FullPathOrPackageName, const FString& AssetName)
+{
+	if (!FullPathOrPackageName.StartsWith(TEXT("/")))
+	{
+		FFormatNamedArguments Args;
+		Args.Add(TEXT("PathString"), FText::AsCultureInvariant(FullPathOrPackageName));
+		LogRuntimeError(FText::Format(NSLOCTEXT("KismetSystemLibrary", "TopLevelAssetPath_UnrootedPath",
+			"Short path \"{PathString}\" not valid for MakeTopLevelAssetPath. Path must be rooted with /, for example /Game."), Args));
+
+		return FTopLevelAssetPath();
+	}
+	
+	if (AssetName.IsEmpty())
+	{
+		FTopLevelAssetPath Result = FTopLevelAssetPath(FullPathOrPackageName);
+		if (!Result.IsValid() && !FullPathOrPackageName.IsEmpty())
+		{
+			FFormatNamedArguments Args;
+			Args.Add(TEXT("PathString"), FText::AsCultureInvariant(FullPathOrPackageName));
+			LogRuntimeError(FText::Format(NSLOCTEXT("KismetSystemLibrary", "TopLevelAssetPath_PathStringInvalid",
+				"String \"{PathString}\" is invalid for MakeTopLevelAssetPath."), Args));
+		}
+		return Result;
+	}
+	else
+	{
+		FTopLevelAssetPath Result = FTopLevelAssetPath(*FullPathOrPackageName, *AssetName);
+		if (!Result.IsValid())
+		{
+			FFormatNamedArguments Args;
+			Args.Add(TEXT("PackageName"), FText::AsCultureInvariant(FullPathOrPackageName));
+			Args.Add(TEXT("AssetName"), FText::AsCultureInvariant(AssetName));
+			LogRuntimeError(FText::Format(NSLOCTEXT("KismetSystemLibrary", "TopLevelAssetPath_PackageAndAssetStringsInvalid",
+				"Strings (\"{PackageName}\", \"{AssetName}\") are invalid for MakeTopLevelAssetPath."), Args));
+		}
+		return Result;
+	}
+}
+
+void UKismetSystemLibrary::BreakTopLevelAssetPath(const FTopLevelAssetPath& InTopLevelAssetPath, FString& PathString)
+{
+	InTopLevelAssetPath.ToString(PathString);
+}
+
 FSoftClassPath UKismetSystemLibrary::MakeSoftClassPath(const FString& PathString)
 {
 	FSoftClassPath SoftClassPath(PathString);
