Format large numbers and memory when printing shader statistics for readability purposes.

- Cleanup memory formatting code and let `FText` handle that.
- Added small helper function to improve readibility when formatting numbers.

## Before

=== FShaderJobCache stats ===
Total job queries 58857, among them cache hits 0 (0.00%)
Tracking 58857 distinct input hashes that result in 24119 distinct outputs (40.98%)
RAM used: 749.38 MiB (0.74 GB) of 12902.40 MB (12.60 GB) budget. Usage: 5.81

## After

=== FShaderJobCache stats ===
Total job queries 58,857, among them cache hits 0 (0.00%)
Tracking 58,857 distinct input hashes that result in 24,119 distinct outputs (40.98%)
RAM used: 749.38 MiB of 12.60 GiB budget. Usage: 5.81

diff --git a/Engine/Source/Runtime/Engine/Private/ShaderCompiler/ShaderCompiler.cpp b/Engine/Source/Runtime/Engine/Private/ShaderCompiler/ShaderCompiler.cpp
--- a/Engine/Source/Runtime/Engine/Private/ShaderCompiler/ShaderCompiler.cpp
+++ b/Engine/Source/Runtime/Engine/Private/ShaderCompiler/ShaderCompiler.cpp
@@ -3294,16 +3294,26 @@ void FShaderCompilerStats::WriteStats(FOutputDevice* Ar)
 #endif // ALLOW_DEBUG_FILES
 }
 
+template <typename T>
+static FString FormatNumber(T Number)
+{
+	static const FNumberFormattingOptions FormattingOptions = FNumberFormattingOptions().SetUseGrouping(true);
+	return FText::AsNumber(Number, &FormattingOptions).ToString();
+}
+
 void FShaderCompilerStats::WriteStatSummary()
 {
 	const uint32 TotalCompiled = GetTotalShadersCompiled();
 	const double TotalTimeAtLeastOneJobWasInFlight = GetTimeShaderCompilationWasActive();
 
 	UE_LOG(LogShaderCompilers, Display, TEXT("=== Shader Compilation stats ==="));
-	UE_LOG(LogShaderCompilers, Display, TEXT("Shaders Compiled: %u"), TotalCompiled);
+	UE_LOG(LogShaderCompilers, Display, TEXT("Shaders Compiled: %s"), *FormatNumber(TotalCompiled));
 
 	FScopeLock Lock(&CompileStatsLock);	// make a local copy for all the stats?
-	UE_LOG(LogShaderCompilers, Display, TEXT("Jobs assigned %.0f, completed %.0f (%.2f%%)"), JobsAssigned, JobsCompleted, (JobsAssigned > 0.0) ? 100.0 * JobsCompleted / JobsAssigned : 0.0);
+	UE_LOG(LogShaderCompilers, Display, TEXT("Jobs assigned %s, completed %s (%.2f%%)"), 
+		*FormatNumber(JobsAssigned),
+		*FormatNumber(JobsCompleted),
+		(JobsAssigned > 0.0) ? 100.0 * JobsCompleted / JobsAssigned : 0.0);
 
 	if (TimesLocalWorkersWereIdle > 0.0)
 	{
@@ -3329,8 +3339,8 @@ void FShaderCompilerStats::WriteStatSummary()
 		int64 JobBatchesSeen = LocalJobBatchesSeen + DistributedJobBatchesSeen;
 		double TotalJobsReportedInJobBatches = TotalJobsReportedInLocalJobBatches + TotalJobsReportedInDistributedJobBatches;
 
-		UE_LOG(LogShaderCompilers, Display, TEXT("Jobs were issued in %lld batches (%lld local, %lld distributed), average %.2f jobs/batch (%.2f jobs/local batch. %.2f jobs/distributed batch)"),
-			JobBatchesSeen, LocalJobBatchesSeen, DistributedJobBatchesSeen,
+		UE_LOG(LogShaderCompilers, Display, TEXT("Jobs were issued in %s batches (%s local, %s distributed), average %.2f jobs/batch (%.2f jobs/local batch. %.2f jobs/distributed batch)"),
+			*FormatNumber(JobBatchesSeen), *FormatNumber(LocalJobBatchesSeen), *FormatNumber(DistributedJobBatchesSeen),
 			static_cast<double>(TotalJobsReportedInJobBatches) / static_cast<double>(JobBatchesSeen),
 			static_cast<double>(TotalJobsReportedInLocalJobBatches) / static_cast<double>(LocalJobBatchesSeen),
 			static_cast<double>(TotalJobsReportedInDistributedJobBatches) / static_cast<double>(DistributedJobBatchesSeen)
@@ -3338,13 +3348,13 @@ void FShaderCompilerStats::WriteStatSummary()
 	}
 	else if (LocalJobBatchesSeen > 0)
 	{
-		UE_LOG(LogShaderCompilers, Display, TEXT("Jobs were issued in %lld batches (only local compilation was used), average %.2f jobs/batch"), 
-			LocalJobBatchesSeen, static_cast<double>(TotalJobsReportedInLocalJobBatches) / static_cast<double>(LocalJobBatchesSeen));
+		UE_LOG(LogShaderCompilers, Display, TEXT("Jobs were issued in %s batches (only local compilation was used), average %.2f jobs/batch"), 
+			*FormatNumber(LocalJobBatchesSeen), static_cast<double>(TotalJobsReportedInLocalJobBatches) / static_cast<double>(LocalJobBatchesSeen));
 	}
 	else if (DistributedJobBatchesSeen > 0)
 	{
-		UE_LOG(LogShaderCompilers, Display, TEXT("Jobs were issued in %lld batches (only distributed compilation was used), average %.2f jobs/batch"),
-			DistributedJobBatchesSeen, static_cast<double>(TotalJobsReportedInDistributedJobBatches) / static_cast<double>(DistributedJobBatchesSeen));
+		UE_LOG(LogShaderCompilers, Display, TEXT("Jobs were issued in %s batches (only distributed compilation was used), average %.2f jobs/batch"),
+			*FormatNumber(DistributedJobBatchesSeen), static_cast<double>(TotalJobsReportedInDistributedJobBatches) / static_cast<double>(DistributedJobBatchesSeen));
 	}
 
 	if (TotalTimeAtLeastOneJobWasInFlight > 0.0)
@@ -3363,8 +3373,8 @@ void FShaderCompilerStats::WriteStatSummary()
 			TotalThreadPreprocessTimeForAllShaders += Iter.Value().TotalPreprocessTime;
 		}
 
-		UE_LOG(LogShaderCompilers, Display, TEXT("Total thread time: %.2f s"), TotalThreadTimeForAllShaders);
-		UE_LOG(LogShaderCompilers, Display, TEXT("Total thread preprocess time: %.2f s"), TotalThreadPreprocessTimeForAllShaders);
+		UE_LOG(LogShaderCompilers, Display, TEXT("Total thread time: %s s"), *FormatNumber(TotalThreadTimeForAllShaders));
+		UE_LOG(LogShaderCompilers, Display, TEXT("Total thread preprocess time: %s s"), *FormatNumber(TotalThreadPreprocessTimeForAllShaders));
 		UE_LOG(LogShaderCompilers, Display, TEXT("Percentage time preprocessing: %.2f%%"), TotalThreadTimeForAllShaders > 0.0 ? (TotalThreadPreprocessTimeForAllShaders / TotalThreadTimeForAllShaders) * 100.0 : 0.0);
 
 		if (TotalTimeAtLeastOneJobWasInFlight > 0.0)
@@ -8903,26 +8913,30 @@ uint64 FShaderJobCache::GetAllocatedMemory()
 void FShaderJobCache::LogStats()
 {
 	UE_LOG(LogShaderCompilers, Display, TEXT("=== FShaderJobCache stats ==="), this);
-	UE_LOG(LogShaderCompilers, Display, TEXT("Total job queries %lld, among them cache hits %lld (%.2f%%)"),
-		TotalSearchAttempts, TotalCacheHits, (TotalSearchAttempts > 0) ? 100.0 * static_cast<double>(TotalCacheHits) / static_cast<double>(TotalSearchAttempts) : 0.0);
-	UE_LOG(LogShaderCompilers, Display, TEXT("Tracking %d distinct input hashes that result in %d distinct outputs (%.2f%%)"),
-		InputHashToOutput.Num(), Outputs.Num(), (InputHashToOutput.Num() > 0) ? 100.0 * static_cast<double>(Outputs.Num()) / static_cast<double>(InputHashToOutput.Num()) : 0.0);
-
-	uint64 MemUsed = GetAllocatedMemory();
-	double MemUsedMB = FUnitConversion::Convert(static_cast<double>(MemUsed), EUnit::Bytes, EUnit::Megabytes);
-	double MemUsedGB = FUnitConversion::Convert(static_cast<double>(MemUsed), EUnit::Bytes, EUnit::Gigabytes);
-	uint64 MemBudget = GetCurrentMemoryBudget();
+	UE_LOG(LogShaderCompilers, Display, TEXT("Total job queries %s, among them cache hits %s (%.2f%%)"),
+		*FormatNumber(TotalSearchAttempts),
+		*FormatNumber(TotalCacheHits),
+		(TotalSearchAttempts > 0) ? 100.0 * static_cast<double>(TotalCacheHits) / static_cast<double>(TotalSearchAttempts) : 0.0);
+
+	UE_LOG(LogShaderCompilers, Display, TEXT("Tracking %s distinct input hashes that result in %s distinct outputs (%.2f%%)"),
+		*FormatNumber(InputHashToOutput.Num()), 
+		*FormatNumber(Outputs.Num()),
+		(InputHashToOutput.Num() > 0) ? 100.0 * static_cast<double>(Outputs.Num()) / static_cast<double>(InputHashToOutput.Num()) : 0.0);
+
+	static const FNumberFormattingOptions SizeFormattingOptions = FNumberFormattingOptions().SetMinimumFractionalDigits(2).SetMaximumFractionalDigits(2);
+
+	const uint64 MemUsed = GetAllocatedMemory();
+	const uint64 MemBudget = GetCurrentMemoryBudget();
 	if (MemBudget > 0)
 	{
-		double MemBudgetMB = FUnitConversion::Convert(static_cast<double>(MemBudget), EUnit::Bytes, EUnit::Megabytes);
-		double MemBudgetGB = FUnitConversion::Convert(static_cast<double>(MemBudget), EUnit::Bytes, EUnit::Gigabytes);
-
-		UE_LOG(LogShaderCompilers, Display, TEXT("RAM used: %.2f MB (%.2f GB) of %.2f MB (%.2f GB) budget. Usage: %.2f%%"), 
-			MemUsedMB, MemUsedGB, MemBudgetMB, MemBudgetGB, 100.0 * MemUsedMB / MemBudgetMB);
+		UE_LOG(LogShaderCompilers, Display, TEXT("RAM used: %s of %s budget. Usage: %.2f%%"), 
+			*FText::AsMemory(MemUsed, &SizeFormattingOptions, nullptr, EMemoryUnitStandard::IEC).ToString(),
+			*FText::AsMemory(MemBudget, &SizeFormattingOptions, nullptr, EMemoryUnitStandard::IEC).ToString(),
+			100.0 * MemUsed / MemBudget);
 	}
 	else
 	{
-		UE_LOG(LogShaderCompilers, Display, TEXT("RAM used: %.2f MB (%.2f GB), no memory limit set"), MemUsedMB, MemUsedGB);
+		UE_LOG(LogShaderCompilers, Display, TEXT("RAM used: %s, no memory limit set"), *FText::AsMemory(MemUsed, &SizeFormattingOptions, nullptr, EMemoryUnitStandard::IEC).ToString());
 	}
 }
 
