#ue5 - Fix base-relative movement acceleration not being applied correctly on the server when client FPS was high enough that the client sends 2 moves (pending and new). In my testing in QAGame this was a problem when over 60fps and "p.netshowcorrections 1" would show constant corrections when walking on a slowly rotating base (but had no issues at or below 60fps).

diff --git a/Engine/Source/Runtime/Engine/Private/Components/CharacterMovementComponent.cpp b/Engine/Source/Runtime/Engine/Private/Components/CharacterMovementComponent.cpp
--- a/Engine/Source/Runtime/Engine/Private/Components/CharacterMovementComponent.cpp
+++ b/Engine/Source/Runtime/Engine/Private/Components/CharacterMovementComponent.cpp
@@ -9019,9 +9019,6 @@ void FCharacterNetworkMoveData::ClientFillNetworkMoveData(const FSavedMove_Chara
 	CompressedMoveFlags = ClientMove.GetCompressedFlags();
 	MovementMode = ClientMove.EndPackedMovementMode;
 
-	// Location, relative movement base, and ending movement mode is only used for error checking, so only fill in the more complex parts if actually required.
-	if (MoveType == ENetworkMoveType::NewMove)
-	{
 		// Determine if we send absolute or relative location
 		UPrimitiveComponent* ClientMovementBase = ClientMove.EndBase.Get();
 
@@ -9044,14 +9041,6 @@ void FCharacterNetworkMoveData::ClientFillNetworkMoveData(const FSavedMove_Chara
 			MovementBase = nullptr;
 			MovementBaseBoneName = NAME_None;
 		}
-	}
-	else
-	{
-		Location = ClientMove.SavedLocation;
-		Acceleration = ClientMove.Acceleration;
-		MovementBase = nullptr;
-		MovementBaseBoneName = NAME_None;
-	}
 }
 
 
@@ -9074,13 +9063,9 @@ bool FCharacterNetworkMoveData::Serialize(UCharacterMovementComponent& Character
 
 	SerializeOptionalValue<uint8>(bIsSaving, Ar, CompressedMoveFlags, 0);
 
-	if (MoveType == ENetworkMoveType::NewMove)
-	{
-		// Location, relative movement base, and ending movement mode is only used for error checking, so only save for the final move.
 		SerializeOptionalValue<UPrimitiveComponent*>(bIsSaving, Ar, MovementBase, nullptr);
 		SerializeOptionalValue<FName>(bIsSaving, Ar, MovementBaseBoneName, NAME_None);
 		SerializeOptionalValue<uint8>(bIsSaving, Ar, MovementMode, MOVE_Walking);
-	}
 
 	return !Ar.IsError();
 }
