Fix audio test failure in audio bus subsystem, add logging

#rnx

diff --git a/Engine/Source/Runtime/AudioMixer/Private/AudioBusSubsystem.cpp b/Engine/Source/Runtime/AudioMixer/Private/AudioBusSubsystem.cpp
--- a/Engine/Source/Runtime/AudioMixer/Private/AudioBusSubsystem.cpp
+++ b/Engine/Source/Runtime/AudioMixer/Private/AudioBusSubsystem.cpp
@@ -19,11 +19,13 @@ bool UAudioBusSubsystem::ShouldCreateSubsystem(UObject* Outer) const
 
 void UAudioBusSubsystem::Initialize(FSubsystemCollectionBase& Collection)
 {
+	UE_LOG(LogAudioMixer, Log, TEXT("Initializing Audio Bus Subsystem for audio device with ID %d"), GetMixerDevice()->DeviceID);
 	InitDefaultAudioBuses();
 }
 
 void UAudioBusSubsystem::Deinitialize()
 {
+	UE_LOG(LogAudioMixer, Log, TEXT("Deinitializing Audio Bus Subsystem for audio device with ID %d"), GetMixerDevice()->DeviceID);
 	ShutdownDefaultAudioBuses();
 }
 
diff --git a/Engine/Source/Runtime/Engine/Private/AudioDevice.cpp b/Engine/Source/Runtime/Engine/Private/AudioDevice.cpp
--- a/Engine/Source/Runtime/Engine/Private/AudioDevice.cpp
+++ b/Engine/Source/Runtime/Engine/Private/AudioDevice.cpp
@@ -689,7 +689,7 @@ bool FAudioDevice::Init(Audio::FDeviceId InDeviceID, int32 InMaxSources, int32 I
 
 	InitDefaultAudioBuses();
 
-	UE_LOG(LogInit, Log, TEXT("FAudioDevice initialized."));
+	UE_LOG(LogInit, Log, TEXT("FAudioDevice initialized with ID %d."), InDeviceID);
 
 	return true;
 }
@@ -2424,6 +2424,8 @@ void FAudioDevice::InitSoundSources()
 
 void FAudioDevice::InitializeSubsystemCollection()
 {
+	UE_LOG(LogAudio, Log, TEXT("Initializing audio subsystem collection for audio device with id %d"), DeviceID);
+	FModuleManager::Get().LoadModuleChecked("AudioMixer");
 	SubsystemCollectionRoot.Reset(NewObject<UAudioSubsystemCollectionRoot>(GetTransientPackage()));
 	check(SubsystemCollectionRoot.IsValid());
 
