World Partition HLODs: convert HLODs cells mapping to use a cell guid instead of the cell name, so we can rename cells as we want without breaking existing HLODs.

diff --git a/Engine/Plugins/Editor/WorldPartitionHLODUtilities/Source/Private/WorldPartitionHLODUtilities.cpp b/Engine/Plugins/Editor/WorldPartitionHLODUtilities/Source/Private/WorldPartitionHLODUtilities.cpp
--- a/Engine/Plugins/Editor/WorldPartitionHLODUtilities/Source/Private/WorldPartitionHLODUtilities.cpp
+++ b/Engine/Plugins/Editor/WorldPartitionHLODUtilities/Source/Private/WorldPartitionHLODUtilities.cpp
@@ -190,14 +190,14 @@ TArray<AWorldPartitionHLOD*> FWorldPartitionHLODUtilities::CreateHLODActors(FHLO
 		const TSet<FHLODSubActor>& SubActors = Pair.Value;
 		check(!SubActors.IsEmpty());
 
-		auto ComputeCellHash = [](const FString& HLODLayerName, const FName CellName)
+		auto ComputeCellHash = [](const FString& HLODLayerName, const FGuid CellGuid)
 		{
 			uint64 HLODLayerNameHash = FCrc::StrCrc32(*HLODLayerName);
-			uint32 CellNameHash = FCrc::StrCrc32(*CellName.ToString());
-			return HashCombine(HLODLayerNameHash, CellNameHash);
+			uint32 CellGuidHash = GetTypeHash(CellGuid);
+			return HashCombine(HLODLayerNameHash, CellGuidHash);
 		};
 
-		uint64 CellHash = ComputeCellHash(HLODLayer->GetName(), InCreationParams.CellName);
+		uint64 CellHash = ComputeCellHash(HLODLayer->GetName(), InCreationParams.CellGuid);
 		FName HLODActorName = *FString::Printf(TEXT("%s_%016llx"), *HLODLayer->GetName(), CellHash);		
 
 		AWorldPartitionHLOD* HLODActor = nullptr;
@@ -216,7 +216,7 @@ TArray<AWorldPartitionHLOD*> FWorldPartitionHLODUtilities::CreateHLODActors(FHLO
 			SpawnParams.NameMode = FActorSpawnParameters::ESpawnActorNameMode::Required_Fatal;
 			HLODActor = InCreationParams.WorldPartition->GetWorld()->SpawnActor<AWorldPartitionHLOD>(SpawnParams);
 
-			HLODActor->SetSourceCellName(InCreationParams.CellName);
+			HLODActor->SetSourceCellGuid(InCreationParams.CellGuid);
 			HLODActor->SetSubActorsHLODLayer(HLODLayer);
 
 			// Make sure the generated HLOD actor has the same data layers as the source actors
@@ -228,7 +228,7 @@ TArray<AWorldPartitionHLOD*> FWorldPartitionHLODUtilities::CreateHLODActors(FHLO
 		else
 		{
 #if DO_CHECK
-			check(HLODActor->GetSourceCellName() == InCreationParams.CellName);
+			check(HLODActor->GetSourceCellGuid() == InCreationParams.CellGuid);
 			check(HLODActor->GetSubActorsHLODLayer() == HLODLayer);
 #endif
 		}
@@ -292,7 +292,7 @@ TArray<AWorldPartitionHLOD*> FWorldPartitionHLODUtilities::CreateHLODActors(FHLO
 		}
 
 		// Actor label
-		const FString ActorLabel = FString::Printf(TEXT("%s/%s"), *HLODLayer->GetName(), *InCreationParams.CellName.ToString());
+		const FString ActorLabel = FString::Printf(TEXT("%s/%s"), *HLODLayer->GetName(), *InCreationParams.CellGuid.ToString());
 		if (HLODActor->GetActorLabel() != ActorLabel)
 		{
 			HLODActor->SetActorLabel(ActorLabel);
diff --git a/Engine/Source/Runtime/Core/Public/UObject/FortniteMainBranchObjectVersion.h b/Engine/Source/Runtime/Core/Public/UObject/FortniteMainBranchObjectVersion.h
--- a/Engine/Source/Runtime/Core/Public/UObject/FortniteMainBranchObjectVersion.h
+++ b/Engine/Source/Runtime/Core/Public/UObject/FortniteMainBranchObjectVersion.h
@@ -262,6 +262,9 @@ struct CORE_API FFortniteMainBranchObjectVersion
 		// Fixed localization gathering for external actor packages
 		FixedLocalizationGatherForExternalActorPackage,
 
+		// Change HLODActors to RuntimeCells mapping to use a GUID instead of the cell name
+		WorldPartitionHLODActorUseSourceCellGuid,
+
 		// -----<new versions can be added above this line>-------------------------------------------------
 		VersionPlusOne,
 		LatestVersion = VersionPlusOne - 1
diff --git a/Engine/Source/Runtime/Engine/Private/WorldPartition/HLOD/HLODActor.cpp b/Engine/Source/Runtime/Engine/Private/WorldPartition/HLOD/HLODActor.cpp
--- a/Engine/Source/Runtime/Engine/Private/WorldPartition/HLOD/HLODActor.cpp
+++ b/Engine/Source/Runtime/Engine/Private/WorldPartition/HLOD/HLODActor.cpp
@@ -11,6 +11,7 @@
 #include UE_INLINE_GENERATED_CPP_BY_NAME(HLODActor)
 
 #if WITH_EDITOR
+#include "Misc/ArchiveMD5.h"
 #include "WorldPartition/WorldPartition.h"
 #include "WorldPartition/WorldPartitionActorDesc.h"
 #include "WorldPartition/HLOD/HLODActorDesc.h"
@@ -103,7 +104,7 @@ void AWorldPartitionHLOD::Serialize(FArchive& Ar)
 		{
 			CellName = CellPath;
 		}
-		SourceCellName = *CellName;
+		SourceCellName_DEPRECATED = *CellName;
 	}
 #endif
 }
@@ -137,15 +138,67 @@ void AWorldPartitionHLOD::PostLoad()
 			FString WorldName = FPackageName::GetShortName(ExternalActorsPath);
 
 			// Strip "WorldName_" from the cell name
-			FString CellName = SourceCellName.ToString();
+			FString CellName = SourceCellName_DEPRECATED.ToString();
 			bool bRemoved = CellName.RemoveFromStart(WorldName + TEXT("_"), ESearchCase::CaseSensitive);
 			if (bRemoved)
 			{
-				SourceCellName = *CellName;
+				SourceCellName_DEPRECATED = *CellName;
 			}
 		}
 	}
 
+	if (GetLinkerCustomVersion(FFortniteMainBranchObjectVersion::GUID) < FFortniteMainBranchObjectVersion::WorldPartitionHLODActorUseSourceCellGuid)
+	{
+		check(!SourceCellName_DEPRECATED.IsNone());
+		check(!SourceCellGuid.IsValid());
+
+		FString GridName;
+		int64 CellGlobalCoord[3];
+		uint32 DataLayerID;
+		uint32 ContentBundleID;
+
+		// Input format should be GridName_Lx_Xx_Yx_DLx[_CBx]
+		TArray<FString> Tokens;
+		if (SourceCellName_DEPRECATED.ToString().ParseIntoArray(Tokens, TEXT("_")) >= 4)
+		{
+			int32 CurrentIndex = 0;
+			GridName = Tokens[CurrentIndex++];
+
+			// Since GridName can contain underscores, we do our best to extract it
+			while(Tokens.IsValidIndex(CurrentIndex))
+			{
+				if ((Tokens[CurrentIndex][0] == TEXT('L')) && (Tokens[CurrentIndex].Len() > 1))
+				{
+					if (FCString::IsNumeric(*Tokens[CurrentIndex] + 1))
+					{
+						break;
+					}
+				}
+
+				GridName += TEXT("_");
+				GridName += Tokens[CurrentIndex++];
+			}
+
+			GridName = GridName.ToLower();
+
+			CellGlobalCoord[2] = Tokens.IsValidIndex(CurrentIndex) ? FCString::Strtoui64(*Tokens[CurrentIndex++] + 1, nullptr, 10) : 0;
+			CellGlobalCoord[0] = Tokens.IsValidIndex(CurrentIndex) ? FCString::Strtoui64(*Tokens[CurrentIndex++] + 1, nullptr, 10) : 0;
+			CellGlobalCoord[1] = Tokens.IsValidIndex(CurrentIndex) ? FCString::Strtoui64(*Tokens[CurrentIndex++] + 1, nullptr, 10) : 0;
+			DataLayerID = Tokens.IsValidIndex(CurrentIndex) ? FCString::Strtoui64(*Tokens[CurrentIndex++] + 2, nullptr, 16) : 0;
+			ContentBundleID = Tokens.IsValidIndex(CurrentIndex) ? FCString::Strtoui64(*Tokens[CurrentIndex++] + 2, nullptr, 16) : 0;
+		}
+
+		FArchiveMD5 ArMD5;
+		ArMD5 << GridName << CellGlobalCoord[0] << CellGlobalCoord[1] << CellGlobalCoord[2] << DataLayerID << ContentBundleID;
+
+		FMD5Hash MD5Hash;
+		ArMD5.GetHash(MD5Hash);
+
+		check(MD5Hash.GetSize() == sizeof(FGuid));
+		FMemory::Memcpy(&SourceCellGuid, MD5Hash.GetBytes(), sizeof(FGuid));
+		check(SourceCellGuid.IsValid());
+	}
+
 	// Update the disk size stat on load, as we can't really know it when saving
 	HLODStats.Add(FWorldPartitionHLODStats::MemoryDiskSizeBytes, FHLODActorDesc::GetPackageSize(this));
 #endif
@@ -215,9 +268,9 @@ void AWorldPartitionHLOD::SetSubActorsHLODLayer(const UHLODLayer* InSubActorsHLO
 	bRequireWarmup = SubActorsHLODLayer->DoesRequireWarmup();
 }
 
-void AWorldPartitionHLOD::SetSourceCellName(FName InSourceCellName)
+void AWorldPartitionHLOD::SetSourceCellGuid(const FGuid& InSourceCellGuid)
 {
-	SourceCellName = InSourceCellName;
+	SourceCellGuid = InSourceCellGuid;
 }
 
 const FBox& AWorldPartitionHLOD::GetHLODBounds() const
diff --git a/Engine/Source/Runtime/Engine/Private/WorldPartition/HLOD/HLODActorDesc.cpp b/Engine/Source/Runtime/Engine/Private/WorldPartition/HLOD/HLODActorDesc.cpp
--- a/Engine/Source/Runtime/Engine/Private/WorldPartition/HLOD/HLODActorDesc.cpp
+++ b/Engine/Source/Runtime/Engine/Private/WorldPartition/HLOD/HLODActorDesc.cpp
@@ -23,7 +23,6 @@ void FHLODActorDesc::Init(const AActor* InActor)
 	HLODSubActors.Reserve(HLODActor->GetSubActors().Num());
 	Algo::Transform(HLODActor->GetSubActors(), HLODSubActors, [](const FHLODSubActor& SubActor) { return FHLODSubActorDesc(SubActor.ActorGuid, SubActor.ContainerID); });
 
-	SourceCellName = HLODActor->GetSourceCellName();
 	SourceHLODLayerName = HLODActor->GetSubActorsHLODLayer()->GetFName();
 	HLODStats = HLODActor->GetStats();
 }
@@ -66,8 +65,13 @@ void FHLODActorDesc::Serialize(FArchive& Ar)
 	}
 
 	if (Ar.CustomVer(FFortniteMainBranchObjectVersion::GUID) >= FFortniteMainBranchObjectVersion::WorldPartitionHLODActorDescSerializeStats)
-	{		
-		Ar << SourceCellName;
+	{
+		if (Ar.CustomVer(FUE5MainStreamObjectVersion::GUID) < FFortniteMainBranchObjectVersion::WorldPartitionHLODActorUseSourceCellGuid)
+		{
+			FName SourceCellName;
+			Ar << SourceCellName;
+		}
+
 		Ar << SourceHLODLayerName;
 		Ar << HLODStats;
 
@@ -84,8 +88,7 @@ bool FHLODActorDesc::Equals(const FWorldPartitionActorDesc* Other) const
 	if (FWorldPartitionActorDesc::Equals(Other))
 	{
 		const FHLODActorDesc& HLODActorDesc = *(FHLODActorDesc*)Other;
-		return SourceCellName == HLODActorDesc.SourceCellName &&
-			   SourceHLODLayerName == HLODActorDesc.SourceHLODLayerName &&
+		return SourceHLODLayerName == HLODActorDesc.SourceHLODLayerName &&
 			   HLODStats.OrderIndependentCompareEqual(HLODActorDesc.GetStats()) &&
 			   CompareUnsortedArrays(HLODSubActors, HLODActorDesc.HLODSubActors);
 	}
diff --git a/Engine/Source/Runtime/Engine/Private/WorldPartition/HLOD/HLODSubsystem.cpp b/Engine/Source/Runtime/Engine/Private/WorldPartition/HLOD/HLODSubsystem.cpp
--- a/Engine/Source/Runtime/Engine/Private/WorldPartition/HLOD/HLODSubsystem.cpp
+++ b/Engine/Source/Runtime/Engine/Private/WorldPartition/HLOD/HLODSubsystem.cpp
@@ -259,7 +259,7 @@ void UHLODSubsystem::OnWorldPartitionInitialized(UWorldPartition* InWorldPartiti
 		{
 			InWorldPartition->RuntimeHash->ForEachStreamingCells([&WorldPartitionHLODRuntimeData](const UWorldPartitionRuntimeCell* Cell)
 			{
-				WorldPartitionHLODRuntimeData.CellsData.Emplace(Cell->GetFName());
+				WorldPartitionHLODRuntimeData.CellsData.Emplace(Cell->GetGuid());
 				return true;
 			});
 		}
@@ -294,7 +294,7 @@ UHLODSubsystem::FCellData* UHLODSubsystem::GetCellData(const UWorldPartitionRunt
 	FWorldPartitionHLODRuntimeData* WorldPartitionHLODRuntimeData = WorldPartitionsHLODRuntimeData.Find(WorldPartition);
 	if (WorldPartitionHLODRuntimeData)
 	{
-		return WorldPartitionHLODRuntimeData->CellsData.Find(InCell->GetFName());
+		return WorldPartitionHLODRuntimeData->CellsData.Find(InCell->GetGuid());
 	}
 
 	return nullptr;
@@ -308,8 +308,8 @@ UHLODSubsystem::FCellData* UHLODSubsystem::GetCellData(AWorldPartitionHLOD* InWo
 	FWorldPartitionHLODRuntimeData* WorldPartitionHLODRuntimeData = WorldPartitionsHLODRuntimeData.Find(WorldPartition);
 	if (WorldPartitionHLODRuntimeData)
 	{
-		const FName CellName = InWorldPartitionHLOD->GetSourceCellName();
-		return WorldPartitionHLODRuntimeData->CellsData.Find(CellName);
+		const FGuid CellGuid = InWorldPartitionHLOD->GetSourceCellGuid();
+		return WorldPartitionHLODRuntimeData->CellsData.Find(CellGuid);
 	}
 
 	return nullptr;		
@@ -334,14 +334,14 @@ void UHLODSubsystem::RegisterHLODActor(AWorldPartitionHLOD* InWorldPartitionHLOD
 	if (FCellData* CellData = GetCellData(InWorldPartitionHLOD))
 	{
 #if WITH_EDITOR
-		UE_LOG(LogHLODSubsystem, Verbose, TEXT("Registering HLOD %s (%s) for cell %s"), *InWorldPartitionHLOD->GetActorLabel(), *InWorldPartitionHLOD->GetActorGuid().ToString(), *InWorldPartitionHLOD->GetSourceCellName().ToString());
+		UE_LOG(LogHLODSubsystem, Verbose, TEXT("Registering HLOD %s (%s) for cell %s"), *InWorldPartitionHLOD->GetActorLabel(), *InWorldPartitionHLOD->GetActorGuid().ToString(), *InWorldPartitionHLOD->GetSourceCellGuid().ToString());
 #endif
 		CellData->LoadedHLODs.Add(InWorldPartitionHLOD);
 		InWorldPartitionHLOD->SetVisibility(UHLODSubsystem::WorldPartitionHLODEnabled && !CellData->bIsCellVisible);
 	}
 	else
 	{
-		UE_LOG(LogHLODSubsystem, Verbose, TEXT("Found HLOD referencing nonexistent cell '%s'"), *InWorldPartitionHLOD->GetSourceCellName().ToString());
+		UE_LOG(LogHLODSubsystem, Verbose, TEXT("Found HLOD referencing nonexistent cell '%s'"), *InWorldPartitionHLOD->GetSourceCellGuid().ToString());
 		InWorldPartitionHLOD->SetVisibility(false);
 	}
 
@@ -356,7 +356,7 @@ void UHLODSubsystem::UnregisterHLODActor(AWorldPartitionHLOD* InWorldPartitionHL
 	if (FCellData* CellData = GetCellData(InWorldPartitionHLOD))
 	{
 #if WITH_EDITOR
-		UE_LOG(LogHLODSubsystem, Verbose, TEXT("Unregistering HLOD %s (%s) for cell %s"), *InWorldPartitionHLOD->GetActorLabel(), *InWorldPartitionHLOD->GetActorGuid().ToString(), *InWorldPartitionHLOD->GetSourceCellName().ToString());
+		UE_LOG(LogHLODSubsystem, Verbose, TEXT("Unregistering HLOD %s (%s) for cell %s"), *InWorldPartitionHLOD->GetActorLabel(), *InWorldPartitionHLOD->GetActorGuid().ToString(), *InWorldPartitionHLOD->GetSourceCellGuid().ToString());
 #endif
 
 		int32 NumRemoved = CellData->LoadedHLODs.Remove(InWorldPartitionHLOD);
diff --git a/Engine/Source/Runtime/Engine/Private/WorldPartition/RuntimeSpatialHash/RuntimeSpatialHashHLOD.cpp b/Engine/Source/Runtime/Engine/Private/WorldPartition/RuntimeSpatialHash/RuntimeSpatialHashHLOD.cpp
--- a/Engine/Source/Runtime/Engine/Private/WorldPartition/RuntimeSpatialHash/RuntimeSpatialHashHLOD.cpp
+++ b/Engine/Source/Runtime/Engine/Private/WorldPartition/RuntimeSpatialHash/RuntimeSpatialHashHLOD.cpp
@@ -174,11 +174,12 @@ static TArray<FGuid> GenerateHLODsForGrid(UWorldPartition* WorldPartition, const
 				SlowTask.EnterProgressFrame(1);
 				// todo_ow : Generating hlods is not yet supported for content bundles. Pass an invalid UID to specified the cell is not coming from ContentBundles.
 				FGuid ContentBundleInvalidUID;
+				FGuid CellGuid = UWorldPartitionRuntimeSpatialHash::GetCellGuid(RuntimeGrid.GridName, CellGlobalCoord, GridCellDataChunk.GetDataLayersID(), ContentBundleInvalidUID);
 				FString CellName = UWorldPartitionRuntimeSpatialHash::GetCellNameString(RuntimeGrid.GridName, CellGlobalCoord, GridCellDataChunk.GetDataLayersID(), ContentBundleInvalidUID);
 
 				UE_LOG(LogWorldPartitionRuntimeSpatialHashHLOD, Verbose, TEXT("Creating HLOD for cell %s at %s..."), *CellName, *CellCoord.ToString());
 
-				TRACE_CPUPROFILER_EVENT_SCOPE_TEXT(*CellName);
+				TRACE_CPUPROFILER_EVENT_SCOPE_TEXT(*CellCoord.ToString());
 
 				UE_LOG(LogWorldPartitionRuntimeSpatialHashHLOD, Display, TEXT("[%d / %d] Processing cell %s..."), (int32)SlowTask.CompletedWork + 1, (int32)SlowTask.TotalAmountOfWork, *CellName);
 
@@ -206,7 +207,7 @@ static TArray<FGuid> GenerateHLODsForGrid(UWorldPartition* WorldPartition, const
 
 				FHLODCreationParams CreationParams;
 				CreationParams.WorldPartition = WorldPartition;
-				CreationParams.CellName = FName(CellName);
+				CreationParams.CellGuid = CellGuid;
 				CreationParams.CellBounds = CellBounds;
 				CreationParams.HLODLevel = HLODLevel;
 				CreationParams.MinVisibleDistance = RuntimeGrid.LoadingRange;
diff --git a/Engine/Source/Runtime/Engine/Private/WorldPartition/WorldPartitionRuntimeSpatialHash.cpp b/Engine/Source/Runtime/Engine/Private/WorldPartition/WorldPartitionRuntimeSpatialHash.cpp
--- a/Engine/Source/Runtime/Engine/Private/WorldPartition/WorldPartitionRuntimeSpatialHash.cpp
+++ b/Engine/Source/Runtime/Engine/Private/WorldPartition/WorldPartitionRuntimeSpatialHash.cpp
@@ -6,6 +6,7 @@
 #include "WorldPartition/WorldPartitionRuntimeSpatialHash.h"
 
 #include "Engine/Engine.h"
+#include "Misc/ArchiveMD5.h"
 #include "Misc/Paths.h"
 #include "WorldPartition/WorldPartitionStreamingGenerationContext.h"
 #include "WorldPartition/WorldPartitionRuntimeSpatialHashCell.h"
@@ -1308,6 +1309,27 @@ FString UWorldPartitionRuntimeSpatialHash::GetCellNameString(FName InGridName, c
 	return CellName;
 }
 
+FGuid UWorldPartitionRuntimeSpatialHash::GetCellGuid(FName InGridName, const FGridCellCoord& InCellGlobalCoord, const FDataLayersID& InDataLayerID, const FGuid& InContentBundleID)
+{
+	FString GridName = InGridName.ToString().ToLower();
+	FGridCellCoord CellGlobalCoord = InCellGlobalCoord;
+	uint32 DataLayerID = InDataLayerID.GetHash();
+	uint32 ContentBundleID = InContentBundleID.IsValid() ? GetTypeHash(InContentBundleID) : 0;
+
+	FArchiveMD5 ArMD5;
+	ArMD5 << GridName << CellGlobalCoord << DataLayerID << ContentBundleID;
+
+	FMD5Hash MD5Hash;
+	ArMD5.GetHash(MD5Hash);
+
+	FGuid CellGuid;
+	check(MD5Hash.GetSize() == sizeof(FGuid));
+	FMemory::Memcpy(&CellGuid, MD5Hash.GetBytes(), sizeof(FGuid));
+	check(CellGuid.IsValid());
+
+	return CellGuid;
+}
+
 bool UWorldPartitionRuntimeSpatialHash::GetPreviewGrids() const
 {
 	return bPreviewGrids;
@@ -1395,7 +1417,8 @@ bool UWorldPartitionRuntimeSpatialHash::CreateStreamingGrid(const FSpatialHashRu
 				
 				FGridCellCoord CellGlobalCoords;
 				verify(PartionedActors.GetCellGlobalCoords(FGridCellCoord(CellCoordX, CellCoordY, Level), CellGlobalCoords));
-				FString CellName = GetCellNameString(CurrentStreamingGrid.GridName, CellGlobalCoords, GridCellDataChunk.GetDataLayersID(), GridCellDataChunk.GetContentBundleID());
+				const FString CellName = GetCellNameString(CurrentStreamingGrid.GridName, CellGlobalCoords, GridCellDataChunk.GetDataLayersID(), GridCellDataChunk.GetContentBundleID());
+				const FGuid CellGuid = GetCellGuid(CurrentStreamingGrid.GridName, CellGlobalCoords, GridCellDataChunk.GetDataLayersID(), GridCellDataChunk.GetContentBundleID());
 
 				UWorldPartitionRuntimeSpatialHashCell* StreamingCell = NewObject<UWorldPartitionRuntimeSpatialHashCell>(this, StreamingPolicy->GetRuntimeCellClass(), FName(CellName));
 
@@ -1427,6 +1450,7 @@ bool UWorldPartitionRuntimeSpatialHash::CreateStreamingGrid(const FSpatialHashRu
 				StreamingCell->SetClientOnlyVisible(CurrentStreamingGrid.bClientOnlyVisible);
 				StreamingCell->SetBlockOnSlowLoading(CurrentStreamingGrid.bBlockOnSlowStreaming);
 				StreamingCell->SetIsHLOD(RuntimeGrid.HLODLayer ? true : false);
+				StreamingCell->SetGuid(CellGuid);
 
 				check(!StreamingCell->UnsavedActorsContainer);
 				for (const IStreamingGenerationContext::FActorInstance& ActorInstance : FilteredActors)
diff --git a/Engine/Source/Runtime/Engine/Public/WorldPartition/HLOD/HLODActor.h b/Engine/Source/Runtime/Engine/Public/WorldPartition/HLOD/HLODActor.h
--- a/Engine/Source/Runtime/Engine/Public/WorldPartition/HLOD/HLODActor.h
+++ b/Engine/Source/Runtime/Engine/Public/WorldPartition/HLOD/HLODActor.h
@@ -23,7 +23,7 @@ class ENGINE_API AWorldPartitionHLOD : public AActor
 public:
 	void SetVisibility(bool bInVisible);
 
-	inline FName GetSourceCellName() const { return SourceCellName; }
+	inline FGuid GetSourceCellGuid() const { return SourceCellGuid; }
 	inline uint32 GetLODLevel() const { return LODLevel; }
 
 	virtual bool IsHLODRelevant() const override { return true; }
@@ -41,7 +41,7 @@ public:
 
 	void SetRequireWarmup(bool InRequireWarmup) { bRequireWarmup = InRequireWarmup; }
 
-	void SetSourceCellName(FName InSourceCellName);
+	void SetSourceCellGuid(const FGuid& InSourceCellGuid);
 	inline void SetLODLevel(uint32 InLODLevel) { LODLevel = InLODLevel; }
 
 	const FBox& GetHLODBounds() const;
@@ -118,7 +118,10 @@ private:
 	TSoftObjectPtr<UWorldPartitionRuntimeCell> SourceCell_DEPRECATED;
 
 	UPROPERTY()
-	FName SourceCellName;
+	FName SourceCellName_DEPRECATED;
+
+	UPROPERTY()
+	FGuid SourceCellGuid;
 };
 
 DEFINE_ACTORDESC_TYPE(AWorldPartitionHLOD, FHLODActorDesc);
\ No newline at end of file
diff --git a/Engine/Source/Runtime/Engine/Public/WorldPartition/HLOD/HLODActorDesc.h b/Engine/Source/Runtime/Engine/Public/WorldPartition/HLOD/HLODActorDesc.h
--- a/Engine/Source/Runtime/Engine/Public/WorldPartition/HLOD/HLODActorDesc.h
+++ b/Engine/Source/Runtime/Engine/Public/WorldPartition/HLOD/HLODActorDesc.h
@@ -26,7 +26,6 @@ public:
 	typedef TMap<FName, int64>	FStats;
 
 	inline const TArray<FHLODSubActorDesc>& GetSubActors() const { return HLODSubActors; }
-	inline const FName GetSourceCellName() const { return SourceCellName; }
 	inline const FName GetSourceHLODLayerName() const { return SourceHLODLayerName; }
 	inline const FStats& GetStats() const { return HLODStats; }
 	inline int64 GetStat(FName InStatName) const { return HLODStats.FindRef(InStatName); }
@@ -43,7 +42,6 @@ protected:
 	//~ End FWorldPartitionActorDesc Interface.
 
 	TArray<FHLODSubActorDesc> HLODSubActors;
-	FName SourceCellName;
 	FName SourceHLODLayerName;
 	FStats HLODStats;
 };
diff --git a/Engine/Source/Runtime/Engine/Public/WorldPartition/HLOD/HLODSubsystem.h b/Engine/Source/Runtime/Engine/Public/WorldPartition/HLOD/HLODSubsystem.h
--- a/Engine/Source/Runtime/Engine/Public/WorldPartition/HLOD/HLODSubsystem.h
+++ b/Engine/Source/Runtime/Engine/Public/WorldPartition/HLOD/HLODSubsystem.h
@@ -70,7 +70,7 @@ private:
 
 	struct FWorldPartitionHLODRuntimeData
 	{
-		TMap<FName, FCellData> CellsData;
+		TMap<FGuid, FCellData> CellsData;
 	};
 	
 	TMap<TObjectPtr<UWorldPartition>, FWorldPartitionHLODRuntimeData> WorldPartitionsHLODRuntimeData;
diff --git a/Engine/Source/Runtime/Engine/Public/WorldPartition/HLOD/IWorldPartitionHLODUtilities.h b/Engine/Source/Runtime/Engine/Public/WorldPartition/HLOD/IWorldPartitionHLODUtilities.h
--- a/Engine/Source/Runtime/Engine/Public/WorldPartition/HLOD/IWorldPartitionHLODUtilities.h
+++ b/Engine/Source/Runtime/Engine/Public/WorldPartition/HLOD/IWorldPartitionHLODUtilities.h
@@ -27,7 +27,7 @@ struct ENGINE_API FHLODCreationParams
 
 	FName HLODLayerName;
 
-	FName CellName;
+	FGuid CellGuid;
 	FBox  CellBounds;
 	uint32 HLODLevel;
 
diff --git a/Engine/Source/Runtime/Engine/Public/WorldPartition/WorldPartitionRuntimeCell.h b/Engine/Source/Runtime/Engine/Public/WorldPartition/WorldPartitionRuntimeCell.h
--- a/Engine/Source/Runtime/Engine/Public/WorldPartition/WorldPartitionRuntimeCell.h
+++ b/Engine/Source/Runtime/Engine/Public/WorldPartition/WorldPartitionRuntimeCell.h
@@ -233,6 +233,9 @@ class ENGINE_API UWorldPartitionRuntimeCell : public UObject, public IWorldParti
 #endif
 	
 	bool GetIsHLOD() const { return bIsHLOD; }
+	
+	FGuid GetGuid() const { return CellGuid; }
+	void SetGuid(const FGuid& InCellGuid) { CellGuid = InCellGuid; }
 
 #if !UE_BUILD_SHIPPING
 	void SetDebugStreamingPriority(float InDebugStreamingPriority) { DebugStreamingPriority = InDebugStreamingPriority; }
@@ -276,6 +279,9 @@ private:
 	FGuid ContentBundleID;
 
 protected:
+	UPROPERTY()
+	FGuid CellGuid;
+
 	// Source Priority
 	mutable uint8 CachedMinSourcePriority;
 
diff --git a/Engine/Source/Runtime/Engine/Public/WorldPartition/WorldPartitionRuntimeSpatialHash.h b/Engine/Source/Runtime/Engine/Public/WorldPartition/WorldPartitionRuntimeSpatialHash.h
--- a/Engine/Source/Runtime/Engine/Public/WorldPartition/WorldPartitionRuntimeSpatialHash.h
+++ b/Engine/Source/Runtime/Engine/Public/WorldPartition/WorldPartitionRuntimeSpatialHash.h
@@ -281,6 +281,7 @@ public:
 	virtual URuntimeHashExternalStreamingObjectBase* StoreToExternalStreamingObject(UObject* StreamingObjectOuter, FName StreamingObjectName) override;
 
 	static FString GetCellNameString(FName InGridName, const FGridCellCoord& InCellGlobalCoord, const FDataLayersID& InDataLayerID, const FGuid& InContentBundleID);
+	static FGuid GetCellGuid(FName InGridName, const FGridCellCoord& InCellGlobalCoord, const FDataLayersID& InDataLayerID, const FGuid& InContentBundleID);
 
 	bool GetPreviewGrids() const;
 	void SetPreviewGrids(bool bInPreviewGrids);
