Speculative cook crash fix when saving object pointers from GAllUCSModifiedProperties

#rnx

diff --git a/Engine/Source/Runtime/Engine/Private/Components/ActorComponent.cpp b/Engine/Source/Runtime/Engine/Private/Components/ActorComponent.cpp
--- a/Engine/Source/Runtime/Engine/Private/Components/ActorComponent.cpp
+++ b/Engine/Source/Runtime/Engine/Private/Components/ActorComponent.cpp
@@ -2441,7 +2441,7 @@ void UActorComponent::AddReferencedObjects(UObject* InThis, FReferenceCollector&
 		{
 			for (FSimpleMemberReference& MemberReference : *UCSModifiedProperties)
 			{
-				Collector.AddStableReference(&MemberReference.MemberParent);
+				Collector.AddReferencedObject(MemberReference.MemberParent);
 			}
 		}
 	}
