Temporarily undo all ActorComponent optimizations in case 23477355 didn't fix crash

#rnx

diff --git a/Engine/Source/Runtime/Engine/Classes/Components/ActorComponent.h b/Engine/Source/Runtime/Engine/Classes/Components/ActorComponent.h
--- a/Engine/Source/Runtime/Engine/Classes/Components/ActorComponent.h
+++ b/Engine/Source/Runtime/Engine/Classes/Components/ActorComponent.h
@@ -153,8 +153,6 @@ private:
 	int32 UCSSerializationIndex;
 
 protected:
-	/** If this actor currently has any UCS modified properties */
-	uint8 bHasUCSModifiedProperties:1;
 	/** 
 	 *  Indicates if this ActorComponent is currently registered with a scene. 
 	 */
@@ -423,6 +421,9 @@ private:
 	TArray<FSimpleMemberReference> UCSModifiedProperties_DEPRECATED;
 #endif
 
+	static FRWLock AllUCSModifiedPropertiesLock;
+	static TMap<UActorComponent*, TArray<FSimpleMemberReference>> AllUCSModifiedProperties;
+
 public:
 	/**
 	 * Activates the SceneComponent, should be overridden by native child classes.
diff --git a/Engine/Source/Runtime/Engine/Private/Components/ActorComponent.cpp b/Engine/Source/Runtime/Engine/Private/Components/ActorComponent.cpp
--- a/Engine/Source/Runtime/Engine/Private/Components/ActorComponent.cpp
+++ b/Engine/Source/Runtime/Engine/Private/Components/ActorComponent.cpp
@@ -314,9 +314,8 @@ FActorComponentGlobalDestroyPhysicsSignature UActorComponent::GlobalDestroyPhysi
 UActorComponent::FOnMarkRenderStateDirty UActorComponent::MarkRenderStateDirtyEvent;
 
 const FString UActorComponent::ComponentTemplateNameSuffix(TEXT("_GEN_VARIABLE"));
-
-static TMap<UActorComponent*, TArray<FSimpleMemberReference, TInlineAllocator<1>>> GAllUCSModifiedProperties;
-static FRWLock GAllUCSModifiedPropertiesLock;
+TMap<UActorComponent*, TArray<FSimpleMemberReference>> UActorComponent::AllUCSModifiedProperties;
+FRWLock UActorComponent::AllUCSModifiedPropertiesLock;
 
 UActorComponent::UActorComponent(const FObjectInitializer& ObjectInitializer /*= FObjectInitializer::Get()*/)
 	: Super(ObjectInitializer)
@@ -448,9 +447,8 @@ void UActorComponent::PostLoad()
 		{
 			if (UCSModifiedProperties_DEPRECATED.Num())
 			{
-				FWriteScopeLock WriteLock(GAllUCSModifiedPropertiesLock);
-				GAllUCSModifiedProperties.Add(this) = MoveTemp(UCSModifiedProperties_DEPRECATED);
-				bHasUCSModifiedProperties = true;
+				FRWScopeLock Lock(AllUCSModifiedPropertiesLock, SLT_Write);
+				AllUCSModifiedProperties.Add(this, MoveTemp(UCSModifiedProperties_DEPRECATED));
 			}
 		}
 #endif
@@ -2321,7 +2319,7 @@ void UActorComponent::DetermineUCSModifiedProperties()
 {
 	if (CreationMethod == EComponentCreationMethod::SimpleConstructionScript)
 	{
-		TArray<FSimpleMemberReference, TInlineAllocator<1>> UCSModifiedProperties;
+		TArray<FSimpleMemberReference> UCSModifiedProperties;
 
 		class FComponentPropertySkipper : public FArchive
 		{
@@ -2370,79 +2368,58 @@ void UActorComponent::DetermineUCSModifiedProperties()
 			}
 		}
 
-		bHasUCSModifiedProperties = UCSModifiedProperties.Num() > 0;
-		FWriteScopeLock WriteLock(GAllUCSModifiedPropertiesLock);
-		if (bHasUCSModifiedProperties)
+		FRWScopeLock Lock(AllUCSModifiedPropertiesLock, SLT_Write);
+		if (UCSModifiedProperties.Num() > 0)
 		{
-			GAllUCSModifiedProperties.Add(this, MoveTemp(UCSModifiedProperties));
+			AllUCSModifiedProperties.Add(this, MoveTemp(UCSModifiedProperties));
 		}
 		else
 		{
-			GAllUCSModifiedProperties.Remove(this);
+			AllUCSModifiedProperties.Remove(this);
 		}
 	}
 }
 
 void UActorComponent::GetUCSModifiedProperties(TSet<const FProperty*>& ModifiedProperties) const
 {
-	if (bHasUCSModifiedProperties)
+	FRWScopeLock Lock(AllUCSModifiedPropertiesLock, SLT_ReadOnly);
+	if (TArray<FSimpleMemberReference>* UCSModifiedProperties = AllUCSModifiedProperties.Find(this))
 	{
-		FReadScopeLock ReadLock(GAllUCSModifiedPropertiesLock);
-		if (TArray<FSimpleMemberReference, TInlineAllocator<1>>* UCSModifiedProperties = GAllUCSModifiedProperties.Find(this))
+		for (const FSimpleMemberReference& MemberReference : *UCSModifiedProperties)
 		{
-			for (const FSimpleMemberReference& MemberReference : *UCSModifiedProperties)
-			{
-				ModifiedProperties.Add(FMemberReference::ResolveSimpleMemberReference<FProperty>(MemberReference));
-			}
+			ModifiedProperties.Add(FMemberReference::ResolveSimpleMemberReference<FProperty>(MemberReference));
 		}
 	}
 }
 
 void UActorComponent::RemoveUCSModifiedProperties(const TArray<FProperty*>& Properties)
 {
-	if (bHasUCSModifiedProperties && Properties.Num())
+	FRWScopeLock Lock(AllUCSModifiedPropertiesLock, SLT_Write);
+	if (TArray<FSimpleMemberReference>* UCSModifiedProperties = AllUCSModifiedProperties.Find(this))
 	{
-		FWriteScopeLock WriteLock(GAllUCSModifiedPropertiesLock);
-		if (TArray<FSimpleMemberReference, TInlineAllocator<1>>* UCSModifiedProperties = GAllUCSModifiedProperties.Find(this))
+		for (FProperty* Property : Properties)
 		{
-			for (FProperty* Property : Properties)
-			{
-				FSimpleMemberReference MemberReference;
-				FMemberReference::FillSimpleMemberReference<FProperty>(Property, MemberReference);
-				UCSModifiedProperties->RemoveSwap(MemberReference);
-			}
-
-			if (UCSModifiedProperties->IsEmpty())
-			{
-				bHasUCSModifiedProperties = false;
-				GAllUCSModifiedProperties.Remove(this);
-			}
+			FSimpleMemberReference MemberReference;
+			FMemberReference::FillSimpleMemberReference<FProperty>(Property, MemberReference);
+			UCSModifiedProperties->RemoveSwap(MemberReference);
 		}
 	}
 }
 
 void UActorComponent::ClearUCSModifiedProperties()
 {
-	if (bHasUCSModifiedProperties)
-	{
-		bHasUCSModifiedProperties = false;
-		FWriteScopeLock WriteLock(GAllUCSModifiedPropertiesLock);
-		GAllUCSModifiedProperties.Remove(this);
-	}
+	FRWScopeLock Lock(AllUCSModifiedPropertiesLock, SLT_Write);
+	AllUCSModifiedProperties.Remove(this);
 }
 
 void UActorComponent::AddReferencedObjects(UObject* InThis, FReferenceCollector& Collector)
 {
-	UActorComponent* ThisActor = CastChecked<UActorComponent>(InThis);
-	if (ThisActor->bHasUCSModifiedProperties)
+	FRWScopeLock Lock(AllUCSModifiedPropertiesLock, SLT_ReadOnly);
+	if (TArray<FSimpleMemberReference>* UCSModifiedProperties = AllUCSModifiedProperties.Find(CastChecked<UActorComponent>(InThis)))
 	{
-		FReadScopeLock ReadLock(GAllUCSModifiedPropertiesLock);
-		if (TArray<FSimpleMemberReference, TInlineAllocator<1>>* UCSModifiedProperties = GAllUCSModifiedProperties.Find(ThisActor))
+		for (FSimpleMemberReference& MemberReference : *UCSModifiedProperties)
 		{
-			for (FSimpleMemberReference& MemberReference : *UCSModifiedProperties)
-			{
-				Collector.AddReferencedObject(MemberReference.MemberParent);
-			}
+			Collector.AddReferencedObject(MemberReference.MemberParent);
 		}
 	}
 }
@@ -2490,29 +2467,29 @@ void UActorComponent::Serialize(FArchive& Ar)
 	{
 		if (Ar.IsLoading())
 		{
-			TArray<FSimpleMemberReference, TInlineAllocator<1>> UCSModifiedProperties;
+			TArray<FSimpleMemberReference> UCSModifiedProperties;
 			Ar << UCSModifiedProperties;
 
-			FWriteScopeLock WriteLock(GAllUCSModifiedPropertiesLock);
+			FRWScopeLock Lock(AllUCSModifiedPropertiesLock, SLT_Write);
 			if (UCSModifiedProperties.Num() > 0)
 			{
-				GAllUCSModifiedProperties.Add(this, MoveTemp(UCSModifiedProperties));
+				AllUCSModifiedProperties.Add(this, MoveTemp(UCSModifiedProperties));
 			}
 			else
 			{
-				GAllUCSModifiedProperties.Remove(this);
+				AllUCSModifiedProperties.Remove(this);
 			}
 		}
 		else
 		{
-			FReadScopeLock ReadLock(GAllUCSModifiedPropertiesLock);
-			if (TArray<FSimpleMemberReference, TInlineAllocator<1>>* UCSModifiedProperties = GAllUCSModifiedProperties.Find(this))
+			FRWScopeLock Lock(AllUCSModifiedPropertiesLock, SLT_ReadOnly);
+			if (TArray<FSimpleMemberReference>* UCSModifiedProperties = AllUCSModifiedProperties.Find(this))
 			{
 				Ar << *UCSModifiedProperties;
 			}
 			else
 			{
-				TArray<FSimpleMemberReference, TInlineAllocator<1>> EmptyUCSModifiedProperties;
+				TArray<FSimpleMemberReference> EmptyUCSModifiedProperties;
 				Ar << EmptyUCSModifiedProperties;
 			}
 		}
