Adding includes that code depended on implicitly. This is preperation for upcoming header dependency improvements.

diff --git a/Engine/Plugins/Editor/ContentBrowser/ContentBrowserAssetDataSource/Source/ContentBrowserAssetDataSource/Private/AssetFileContextMenu.cpp b/Engine/Plugins/Editor/ContentBrowser/ContentBrowserAssetDataSource/Source/ContentBrowserAssetDataSource/Private/AssetFileContextMenu.cpp
--- a/Engine/Plugins/Editor/ContentBrowser/ContentBrowserAssetDataSource/Source/ContentBrowserAssetDataSource/Private/AssetFileContextMenu.cpp
+++ b/Engine/Plugins/Editor/ContentBrowser/ContentBrowserAssetDataSource/Source/ContentBrowserAssetDataSource/Private/AssetFileContextMenu.cpp
@@ -34,6 +34,7 @@
 #include "Components/ActorComponent.h"
 #include "GameFramework/Actor.h"
 #include "UnrealClient.h"
+#include "Materials/MaterialFunction.h"
 #include "Materials/MaterialFunctionInstance.h"
 #include "Materials/Material.h"
 #include "Settings/EditorExperimentalSettings.h"
diff --git a/Engine/Plugins/Editor/EngineAssetDefinitions/Source/Private/AssetDefinition_Texture.cpp b/Engine/Plugins/Editor/EngineAssetDefinitions/Source/Private/AssetDefinition_Texture.cpp
--- a/Engine/Plugins/Editor/EngineAssetDefinitions/Source/Private/AssetDefinition_Texture.cpp
+++ b/Engine/Plugins/Editor/EngineAssetDefinitions/Source/Private/AssetDefinition_Texture.cpp
@@ -7,6 +7,7 @@
 #include "Misc/PackageName.h"
 #include "Styling/AppStyle.h"
 #include "Materials/Material.h"
+#include "Materials/MaterialFunction.h"
 #include "Materials/MaterialInstance.h"
 #include "Factories/MaterialFactoryNew.h"
 #include "Interfaces/ITextureEditorModule.h"
diff --git a/Engine/Plugins/Editor/GLTFImporter/Source/GLTFImporter/Private/GLTFMaterialElement.cpp b/Engine/Plugins/Editor/GLTFImporter/Source/GLTFImporter/Private/GLTFMaterialElement.cpp
--- a/Engine/Plugins/Editor/GLTFImporter/Source/GLTFImporter/Private/GLTFMaterialElement.cpp
+++ b/Engine/Plugins/Editor/GLTFImporter/Source/GLTFImporter/Private/GLTFMaterialElement.cpp
@@ -16,6 +16,7 @@
 #include "Materials/MaterialExpressionVectorParameter.h"
 #include "Materials/MaterialExpressionThinTranslucentMaterialOutput.h"
 #include "Materials/MaterialExpressionClearCoatNormalCustomOutput.h"
+#include "Materials/MaterialFunction.h"
 #include "UObject/ObjectRedirector.h"
 
 namespace GLTFImporterImpl
diff --git a/Engine/Plugins/Enterprise/AxFImporter/Source/AxFImporter/Private/AxFImporter.cpp b/Engine/Plugins/Enterprise/AxFImporter/Source/AxFImporter/Private/AxFImporter.cpp
--- a/Engine/Plugins/Enterprise/AxFImporter/Source/AxFImporter/Private/AxFImporter.cpp
+++ b/Engine/Plugins/Enterprise/AxFImporter/Source/AxFImporter/Private/AxFImporter.cpp
@@ -51,6 +51,9 @@
 #include "Materials/MaterialExpressionFunctionOutput.h"
 #include "Materials/MaterialExpressionComponentMask.h"
 #include "Materials/MaterialExpressionTransform.h"
+#include "Materials/MaterialExpressionTextureSample.h"
+#include "Materials/MaterialExpressionMaterialFunctionCall.h"
+#include "Materials/MaterialFunction.h"
 
 
 #include "AxF/decoding/AxF_basic_io.h"
diff --git a/Engine/Plugins/Enterprise/DataprepEditor/Source/DataprepLibraries/Private/DataprepOperationsLibrary.cpp b/Engine/Plugins/Enterprise/DataprepEditor/Source/DataprepLibraries/Private/DataprepOperationsLibrary.cpp
--- a/Engine/Plugins/Enterprise/DataprepEditor/Source/DataprepLibraries/Private/DataprepOperationsLibrary.cpp
+++ b/Engine/Plugins/Enterprise/DataprepEditor/Source/DataprepLibraries/Private/DataprepOperationsLibrary.cpp
@@ -27,6 +27,7 @@
 #include "Interfaces/ITargetPlatform.h"
 #include "Interfaces/ITargetPlatformManagerModule.h"
 #include "Materials/Material.h"
+#include "Materials/MaterialFunction.h"
 #include "Materials/MaterialFunctionInstance.h"
 #include "Materials/MaterialInterface.h"
 #include "Math/Vector2D.h"
diff --git a/Engine/Plugins/Enterprise/DatasmithImporter/Source/DatasmithImporter/Private/DatasmithConsumer.cpp b/Engine/Plugins/Enterprise/DatasmithImporter/Source/DatasmithImporter/Private/DatasmithConsumer.cpp
--- a/Engine/Plugins/Enterprise/DatasmithImporter/Source/DatasmithImporter/Private/DatasmithConsumer.cpp
+++ b/Engine/Plugins/Enterprise/DatasmithImporter/Source/DatasmithImporter/Private/DatasmithConsumer.cpp
@@ -61,6 +61,7 @@
 #include "LevelSequence.h"
 #include "LevelUtils.h"
 #include "Materials/Material.h"
+#include "Materials/MaterialFunction.h"
 #include "Materials/MaterialInstanceConstant.h"
 #include "Materials/MaterialInterface.h"
 #include "Misc/Guid.h"
diff --git a/Engine/Plugins/Enterprise/DatasmithImporter/Source/DatasmithImporter/Private/DatasmithFileProducer.cpp b/Engine/Plugins/Enterprise/DatasmithImporter/Source/DatasmithImporter/Private/DatasmithFileProducer.cpp
--- a/Engine/Plugins/Enterprise/DatasmithImporter/Source/DatasmithImporter/Private/DatasmithFileProducer.cpp
+++ b/Engine/Plugins/Enterprise/DatasmithImporter/Source/DatasmithImporter/Private/DatasmithFileProducer.cpp
@@ -46,6 +46,7 @@
 #include "LevelSequence.h"
 #include "LevelVariantSets.h"
 #include "Materials/Material.h"
+#include "Materials/MaterialFunction.h"
 #include "Materials/MaterialInterface.h"
 #include "Materials/MaterialInstance.h"
 #include "MaterialCachedData.h"
diff --git a/Engine/Plugins/Enterprise/GLTFExporter/Source/GLTFExporter/Private/Converters/GLTFMaterialUtility.cpp b/Engine/Plugins/Enterprise/GLTFExporter/Source/GLTFExporter/Private/Converters/GLTFMaterialUtility.cpp
--- a/Engine/Plugins/Enterprise/GLTFExporter/Source/GLTFExporter/Private/Converters/GLTFMaterialUtility.cpp
+++ b/Engine/Plugins/Enterprise/GLTFExporter/Source/GLTFExporter/Private/Converters/GLTFMaterialUtility.cpp
@@ -10,6 +10,7 @@
 #include "GLTFMaterialAnalyzer.h"
 #include "IGLTFMaterialBakingModule.h"
 #include "GLTFMaterialBakingStructures.h"
+#include "Materials/MaterialExpressionTextureSample.h"
 #endif
 #include "Engine/RendererSettings.h"
 #include "Modules/ModuleManager.h"
diff --git a/Engine/Plugins/Enterprise/MDLImporter/Source/MDLImporter/Private/MDLMapHandler.cpp b/Engine/Plugins/Enterprise/MDLImporter/Source/MDLImporter/Private/MDLMapHandler.cpp
--- a/Engine/Plugins/Enterprise/MDLImporter/Source/MDLImporter/Private/MDLMapHandler.cpp
+++ b/Engine/Plugins/Enterprise/MDLImporter/Source/MDLImporter/Private/MDLMapHandler.cpp
@@ -13,6 +13,7 @@
 #include "mdl/Utility.h"
 
 #include "Materials/Material.h"
+#include "Materials/MaterialExpressionMaterialFunctionCall.h"
 
 MDLSDK_INCLUDES_START
 #include "mi/neuraylib/imaterial_definition.h"
diff --git a/Engine/Plugins/Enterprise/MDLImporter/Source/MDLImporter/Private/generator/FunctionGenerator.h b/Engine/Plugins/Enterprise/MDLImporter/Source/MDLImporter/Private/generator/FunctionGenerator.h
--- a/Engine/Plugins/Enterprise/MDLImporter/Source/MDLImporter/Private/generator/FunctionGenerator.h
+++ b/Engine/Plugins/Enterprise/MDLImporter/Source/MDLImporter/Private/generator/FunctionGenerator.h
@@ -6,6 +6,7 @@
 
 #include "Containers/UnrealString.h"
 #include "Engine/Texture2D.h"
+#include "Materials/MaterialExpressionMaterialFunctionCall.h"
 
 namespace Generator
 {
diff --git a/Engine/Plugins/Enterprise/MDLImporter/Source/MDLImporter/Private/generator/GeneratorMaterialExpressions.cpp b/Engine/Plugins/Enterprise/MDLImporter/Source/MDLImporter/Private/generator/GeneratorMaterialExpressions.cpp
--- a/Engine/Plugins/Enterprise/MDLImporter/Source/MDLImporter/Private/generator/GeneratorMaterialExpressions.cpp
+++ b/Engine/Plugins/Enterprise/MDLImporter/Source/MDLImporter/Private/generator/GeneratorMaterialExpressions.cpp
@@ -3,6 +3,7 @@
 #include "MaterialExpressions.h"
 
 #include "common/Logging.h"
+#include "Materials/MaterialExpressionMaterialFunctionCall.h"
 
 namespace Generator
 {
diff --git a/Engine/Plugins/Enterprise/MDLImporter/Source/MDLImporter/Private/generator/MaterialExpressionFactory.cpp b/Engine/Plugins/Enterprise/MDLImporter/Source/MDLImporter/Private/generator/MaterialExpressionFactory.cpp
--- a/Engine/Plugins/Enterprise/MDLImporter/Source/MDLImporter/Private/generator/MaterialExpressionFactory.cpp
+++ b/Engine/Plugins/Enterprise/MDLImporter/Source/MDLImporter/Private/generator/MaterialExpressionFactory.cpp
@@ -3,6 +3,7 @@
 #ifdef USE_MDLSDK
 
 #include "MaterialExpressionFactory.h"
+#include "Materials/MaterialExpressionMaterialFunctionCall.h"
 
 #include "generator/FunctionLoader.h"
 #include "generator/MaterialExpressions.h"
diff --git a/Engine/Plugins/Enterprise/MDLImporter/Source/MDLImporter/Private/material/CarpaintMaterialFactory.cpp b/Engine/Plugins/Enterprise/MDLImporter/Source/MDLImporter/Private/material/CarpaintMaterialFactory.cpp
--- a/Engine/Plugins/Enterprise/MDLImporter/Source/MDLImporter/Private/material/CarpaintMaterialFactory.cpp
+++ b/Engine/Plugins/Enterprise/MDLImporter/Source/MDLImporter/Private/material/CarpaintMaterialFactory.cpp
@@ -8,6 +8,7 @@
 #include "mdl/Material.h"
 
 #include "Materials/MaterialExpressionClearCoatNormalCustomOutput.h"
+#include "Materials/MaterialExpressionMaterialFunctionCall.h"
 
 namespace Mat
 {
diff --git a/Engine/Plugins/Enterprise/MDLImporter/Source/MDLImporter/Private/material/MapConnecter.h b/Engine/Plugins/Enterprise/MDLImporter/Source/MDLImporter/Private/material/MapConnecter.h
--- a/Engine/Plugins/Enterprise/MDLImporter/Source/MDLImporter/Private/material/MapConnecter.h
+++ b/Engine/Plugins/Enterprise/MDLImporter/Source/MDLImporter/Private/material/MapConnecter.h
@@ -5,6 +5,7 @@
 #include "generator/FunctionLoader.h"
 #include "generator/MaterialExpressions.h"
 #include "material/MaterialFactory.h"
+#include "Materials/MaterialExpressionMaterialFunctionCall.h"
 
 namespace Mat
 {
diff --git a/Engine/Plugins/Experimental/Enterprise/DatasmithInterchange/Source/DatasmithInterchange/Private/InterchangeDatasmithMaterialPipeline.cpp b/Engine/Plugins/Experimental/Enterprise/DatasmithInterchange/Source/DatasmithInterchange/Private/InterchangeDatasmithMaterialPipeline.cpp
--- a/Engine/Plugins/Experimental/Enterprise/DatasmithInterchange/Source/DatasmithInterchange/Private/InterchangeDatasmithMaterialPipeline.cpp
+++ b/Engine/Plugins/Experimental/Enterprise/DatasmithInterchange/Source/DatasmithInterchange/Private/InterchangeDatasmithMaterialPipeline.cpp
@@ -19,6 +19,7 @@
 #include "InterchangeTexture2DFactoryNode.h"
 #include "InterchangeMaterialInstanceNode.h"
 
+#include "Materials/MaterialExpressionMaterialFunctionCall.h"
 #include "Materials/MaterialFunction.h"
 #include "Materials/MaterialInstanceConstant.h"
 #include "Materials/MaterialInstanceDynamic.h"
diff --git a/Engine/Plugins/Experimental/MeshLODToolset/Source/MeshLODToolset/Private/Graphs/GenerateStaticMeshLODProcess.cpp b/Engine/Plugins/Experimental/MeshLODToolset/Source/MeshLODToolset/Private/Graphs/GenerateStaticMeshLODProcess.cpp
--- a/Engine/Plugins/Experimental/MeshLODToolset/Source/MeshLODToolset/Private/Graphs/GenerateStaticMeshLODProcess.cpp
+++ b/Engine/Plugins/Experimental/MeshLODToolset/Source/MeshLODToolset/Private/Graphs/GenerateStaticMeshLODProcess.cpp
@@ -30,6 +30,7 @@
 
 #include "Materials/Material.h"
 #include "Materials/MaterialInstanceConstant.h"
+#include "Materials/MaterialExpressionTextureSampleParameter.h"
 #include "Factories/MaterialInstanceConstantFactoryNew.h"
 #include "Materials/MaterialInstanceDynamic.h"
 #include "MaterialGraph/MaterialGraph.h"
diff --git a/Engine/Plugins/Experimental/VirtualHeightfieldMesh/Source/VirtualHeightfieldMesh/Private/HeightfieldMinMaxTextureMaterialExpression.h b/Engine/Plugins/Experimental/VirtualHeightfieldMesh/Source/VirtualHeightfieldMesh/Private/HeightfieldMinMaxTextureMaterialExpression.h
--- a/Engine/Plugins/Experimental/VirtualHeightfieldMesh/Source/VirtualHeightfieldMesh/Private/HeightfieldMinMaxTextureMaterialExpression.h
+++ b/Engine/Plugins/Experimental/VirtualHeightfieldMesh/Source/VirtualHeightfieldMesh/Private/HeightfieldMinMaxTextureMaterialExpression.h
@@ -9,6 +9,7 @@
 
 struct FPropertyChangedEvent;
 class UHeightfieldMinMaxTexture;
+enum EMaterialSamplerType : int;
 
 /** Node which outputs a texture object contained in a UHeightfieldMinMaxTexture. */
 UCLASS(collapsecategories, hidecategories=Object, MinimalAPI)
@@ -20,7 +21,7 @@ class UMaterialExpressionHeightfieldMinMaxTexture : public UMaterialExpression
 	TObjectPtr<UHeightfieldMinMaxTexture> MinMaxTexture;
 
 	UPROPERTY(EditAnywhere, BlueprintReadWrite, Category = MaterialExpressionTexture)
-	TEnumAsByte<enum EMaterialSamplerType> SamplerType;
+	TEnumAsByte<EMaterialSamplerType> SamplerType;
 
 protected:
 #if WITH_EDITOR
diff --git a/Engine/Plugins/FX/Niagara/Source/Niagara/Private/NiagaraSpriteRendererProperties.cpp b/Engine/Plugins/FX/Niagara/Source/Niagara/Private/NiagaraSpriteRendererProperties.cpp
--- a/Engine/Plugins/FX/Niagara/Source/Niagara/Private/NiagaraSpriteRendererProperties.cpp
+++ b/Engine/Plugins/FX/Niagara/Source/Niagara/Private/NiagaraSpriteRendererProperties.cpp
@@ -18,8 +18,10 @@
 
 #include UE_INLINE_GENERATED_CPP_BY_NAME(NiagaraSpriteRendererProperties)
 
-#if WITH_EDITOR
+#if WITH_EDITORONLY_DATA
 #include "DerivedDataCacheInterface.h"
+#include "Materials/MaterialExpressionTextureSample.h"
+#include "Materials/MaterialExpressionTextureSampleParameter.h"
 #include "Widgets/Images/SImage.h"
 #include "Styling/SlateIconFinder.h"
 #include "Widgets/SWidget.h"
diff --git a/Engine/Plugins/Importers/USDImporter/Source/USDUtilities/Private/USDShadeConversion.cpp b/Engine/Plugins/Importers/USDImporter/Source/USDUtilities/Private/USDShadeConversion.cpp
--- a/Engine/Plugins/Importers/USDImporter/Source/USDUtilities/Private/USDShadeConversion.cpp
+++ b/Engine/Plugins/Importers/USDImporter/Source/USDUtilities/Private/USDShadeConversion.cpp
@@ -28,6 +28,7 @@
 #include "Materials/MaterialExpressionConstant.h"
 #include "Materials/MaterialExpressionConstant4Vector.h"
 #include "Materials/MaterialExpressionTextureCoordinate.h"
+#include "Materials/MaterialExpressionTextureSample.h"
 #include "Materials/MaterialExpressionVertexColor.h"
 #include "Materials/MaterialInstanceConstant.h"
 #include "Materials/MaterialInstanceDynamic.h"
diff --git a/Engine/Plugins/Interchange/Runtime/Source/Import/Private/Material/InterchangeMaterialFactory.cpp b/Engine/Plugins/Interchange/Runtime/Source/Import/Private/Material/InterchangeMaterialFactory.cpp
--- a/Engine/Plugins/Interchange/Runtime/Source/Import/Private/Material/InterchangeMaterialFactory.cpp
+++ b/Engine/Plugins/Interchange/Runtime/Source/Import/Private/Material/InterchangeMaterialFactory.cpp
@@ -11,6 +11,7 @@
 #include "Materials/Material.h"
 #include "Materials/MaterialExpression.h"
 #include "Materials/MaterialExpressionClearCoatNormalCustomOutput.h"
+#include "Materials/MaterialExpressionMaterialFunctionCall.h"
 #include "Materials/MaterialExpressionFunctionInput.h"
 #include "Materials/MaterialExpressionFunctionOutput.h"
 #include "Materials/MaterialExpressionMakeMaterialAttributes.h"
diff --git a/Engine/Plugins/Media/MediaPlayerEditor/Source/MediaPlayerEditor/Private/Widgets/SMediaImage.cpp b/Engine/Plugins/Media/MediaPlayerEditor/Source/MediaPlayerEditor/Private/Widgets/SMediaImage.cpp
--- a/Engine/Plugins/Media/MediaPlayerEditor/Source/MediaPlayerEditor/Private/Widgets/SMediaImage.cpp
+++ b/Engine/Plugins/Media/MediaPlayerEditor/Source/MediaPlayerEditor/Private/Widgets/SMediaImage.cpp
@@ -4,6 +4,7 @@
 
 #include "Styling/AppStyle.h"
 #include "Materials/Material.h"
+#include "Materials/MaterialExpressionTextureSample.h"
 #include "Styling/SlateBrush.h"
 #include "UObject/Package.h"
 #include "Widgets/Layout/SScaleBox.h"
diff --git a/Engine/Plugins/Runtime/RenderTrace/Source/Private/Materials/MaterialExpressionPhysicalMaterialOutput.h b/Engine/Plugins/Runtime/RenderTrace/Source/Private/Materials/MaterialExpressionPhysicalMaterialOutput.h
--- a/Engine/Plugins/Runtime/RenderTrace/Source/Private/Materials/MaterialExpressionPhysicalMaterialOutput.h
+++ b/Engine/Plugins/Runtime/RenderTrace/Source/Private/Materials/MaterialExpressionPhysicalMaterialOutput.h
@@ -4,6 +4,7 @@
 
 #include "CoreMinimal.h"
 #include "Materials/MaterialExpressionCustomOutput.h"
+#include "MaterialValueType.h"
 
 #include "MaterialExpressionPhysicalMaterialOutput.generated.h"
 
diff --git a/Engine/Source/Editor/MaterialEditor/Private/MaterialEditingLibrary.cpp b/Engine/Source/Editor/MaterialEditor/Private/MaterialEditingLibrary.cpp
--- a/Engine/Source/Editor/MaterialEditor/Private/MaterialEditingLibrary.cpp
+++ b/Engine/Source/Editor/MaterialEditor/Private/MaterialEditingLibrary.cpp
@@ -8,6 +8,7 @@
 #include "MaterialEditorUtilities.h"
 #include "MaterialShared.h"
 #include "MaterialGraph/MaterialGraphNode.h"
+#include "Materials/MaterialFunction.h"
 #include "Materials/MaterialInstance.h"
 #include "Materials/MaterialInstanceDynamic.h"
 #include "Materials/MaterialInstanceConstant.h"
diff --git a/Engine/Source/Editor/MaterialEditor/Private/MaterialEditorDetailCustomization.cpp b/Engine/Source/Editor/MaterialEditor/Private/MaterialEditorDetailCustomization.cpp
--- a/Engine/Source/Editor/MaterialEditor/Private/MaterialEditorDetailCustomization.cpp
+++ b/Engine/Source/Editor/MaterialEditor/Private/MaterialEditorDetailCustomization.cpp
@@ -10,10 +10,12 @@
 #include "DetailLayoutBuilder.h"
 #include "DetailWidgetRow.h"
 #include "DetailCategoryBuilder.h"
+#include "Materials/MaterialFunction.h"
 #include "Materials/MaterialExpressionScalarParameter.h"
 #include "Materials/MaterialExpressionVectorParameter.h"
 #include "Materials/MaterialExpressionComposite.h"
 #include "Materials/MaterialExpressionPinBase.h"
+#include "Materials/MaterialExpressionTextureSampleParameter.h"
 #include "MaterialLayersFunctionsCustomization.h"
 #include "PropertyEditorModule.h"
 #include "MaterialEditor/MaterialEditorPreviewParameters.h"
diff --git a/Engine/Source/Editor/MaterialEditor/Private/MaterialEditorModule.cpp b/Engine/Source/Editor/MaterialEditor/Private/MaterialEditorModule.cpp
--- a/Engine/Source/Editor/MaterialEditor/Private/MaterialEditorModule.cpp
+++ b/Engine/Source/Editor/MaterialEditor/Private/MaterialEditorModule.cpp
@@ -6,6 +6,7 @@
 #include "MaterialEditor.h"
 #include "MaterialEditorUtilities.h"
 #include "MaterialInstanceEditor.h"
+#include "Materials/MaterialFunction.h"
 #include "Materials/MaterialInstance.h"
 #include "Materials/MaterialFunctionInstance.h"
 #include "Settings/EditorExperimentalSettings.h"
diff --git a/Engine/Source/Editor/MaterialEditor/Private/MaterialInstanceEditor.cpp b/Engine/Source/Editor/MaterialEditor/Private/MaterialInstanceEditor.cpp
--- a/Engine/Source/Editor/MaterialEditor/Private/MaterialInstanceEditor.cpp
+++ b/Engine/Source/Editor/MaterialEditor/Private/MaterialInstanceEditor.cpp
@@ -19,6 +19,7 @@
 #include "Materials/Material.h"
 #include "MaterialEditor/MaterialEditorInstanceConstant.h"
 #include "ThumbnailRendering/SceneThumbnailInfoWithPrimitive.h"
+#include "Materials/MaterialFunction.h"
 #include "Materials/MaterialInstance.h"
 #include "Materials/MaterialInstanceConstant.h"
 #include "Materials/MaterialFunctionInstance.h"
diff --git a/Engine/Source/Editor/MaterialEditor/Private/SMaterialLayersFunctionsTree.cpp b/Engine/Source/Editor/MaterialEditor/Private/SMaterialLayersFunctionsTree.cpp
--- a/Engine/Source/Editor/MaterialEditor/Private/SMaterialLayersFunctionsTree.cpp
+++ b/Engine/Source/Editor/MaterialEditor/Private/SMaterialLayersFunctionsTree.cpp
@@ -33,6 +33,7 @@
 #include "Widgets/Input/SButton.h"
 #include "Widgets/Text/SInlineEditableTextBlock.h"
 #include "Materials/MaterialFunctionInstance.h"
+#include "Materials/MaterialExpressionTextureSampleParameter.h"
 #include "Framework/Application/SlateApplication.h"
 #include "Widgets/Input/SEditableTextBox.h"
 #include "Curves/CurveLinearColor.h"
diff --git a/Engine/Source/Editor/MaterialEditor/Private/SMaterialParametersOverviewWidget.cpp b/Engine/Source/Editor/MaterialEditor/Private/SMaterialParametersOverviewWidget.cpp
--- a/Engine/Source/Editor/MaterialEditor/Private/SMaterialParametersOverviewWidget.cpp
+++ b/Engine/Source/Editor/MaterialEditor/Private/SMaterialParametersOverviewWidget.cpp
@@ -31,7 +31,10 @@
 #include "MaterialEditor/MaterialEditorPreviewParameters.h"
 #include "Widgets/Input/SButton.h"
 #include "Widgets/Text/SInlineEditableTextBlock.h"
+#include "Materials/MaterialFunction.h"
 #include "Materials/MaterialFunctionInstance.h"
+#include "Materials/MaterialFunctionInterface.h"
+#include "Materials/MaterialExpressionTextureSampleParameter.h"
 #include "Framework/Application/SlateApplication.h"
 
 #include "Widgets/Input/SEditableTextBox.h"
diff --git a/Engine/Source/Editor/UnrealEd/Private/MaterialGraph.cpp b/Engine/Source/Editor/UnrealEd/Private/MaterialGraph.cpp
--- a/Engine/Source/Editor/UnrealEd/Private/MaterialGraph.cpp
+++ b/Engine/Source/Editor/UnrealEd/Private/MaterialGraph.cpp
@@ -17,9 +17,11 @@
 #include "Materials/MaterialExpressionFunctionOutput.h"
 #include "Materials/MaterialExpressionCustomOutput.h"
 #include "Materials/MaterialExpressionReroute.h"
+#include "Materials/MaterialExpressionMaterialFunctionCall.h"
 #include "Materials/MaterialExpressionNamedReroute.h"
 #include "Materials/MaterialExpressionExecBegin.h"
 #include "Materials/MaterialExpressionExecEnd.h"
+#include "Materials/MaterialFunction.h"
 
 #include "MaterialCachedHLSLTree.h"
 #include "HLSLTree/HLSLTreeEmit.h"
diff --git a/Engine/Source/Editor/UnrealEd/Private/MaterialGraphNode.cpp b/Engine/Source/Editor/UnrealEd/Private/MaterialGraphNode.cpp
--- a/Engine/Source/Editor/UnrealEd/Private/MaterialGraphNode.cpp
+++ b/Engine/Source/Editor/UnrealEd/Private/MaterialGraphNode.cpp
@@ -38,12 +38,14 @@
 #include "Materials/MaterialExpressionDoubleVectorParameter.h"
 #include "Materials/MaterialExpressionViewProperty.h"
 #include "Materials/MaterialExpressionMaterialLayerOutput.h"
+#include "Materials/MaterialExpressionFontSampleParameter.h"
 #include "Materials/MaterialExpressionTextureObjectParameter.h"
 #include "Materials/MaterialExpressionNamedReroute.h"
 #include "Materials/MaterialExpressionReroute.h"
 #include "Materials/MaterialExpressionCurveAtlasRowParameter.h"
 #include "Materials/MaterialExpressionExecBegin.h"
 #include "Materials/MaterialExpressionStrata.h"
+#include "Materials/MaterialFunction.h"
 #include "MaterialEditorUtilities.h"
 #include "MaterialEditorActions.h"
 #include "GraphEditorActions.h"
diff --git a/Engine/Source/Editor/UnrealEd/Private/MaterialGraphNode_Comment.cpp b/Engine/Source/Editor/UnrealEd/Private/MaterialGraphNode_Comment.cpp
--- a/Engine/Source/Editor/UnrealEd/Private/MaterialGraphNode_Comment.cpp
+++ b/Engine/Source/Editor/UnrealEd/Private/MaterialGraphNode_Comment.cpp
@@ -8,6 +8,7 @@
 #include "ToolMenus.h"
 #include "MaterialGraph/MaterialGraphSchema.h"
 #include "Materials/MaterialExpressionComment.h"
+#include "Materials/MaterialFunction.h"
 #include "Framework/Commands/GenericCommands.h"
 #include "GraphEditorActions.h"
 
diff --git a/Engine/Source/Editor/UnrealEd/Private/PreviewMaterial.cpp b/Engine/Source/Editor/UnrealEd/Private/PreviewMaterial.cpp
--- a/Engine/Source/Editor/UnrealEd/Private/PreviewMaterial.cpp
+++ b/Engine/Source/Editor/UnrealEd/Private/PreviewMaterial.cpp
@@ -28,6 +28,7 @@
 #include "Materials/MaterialExpressionMaterialAttributeLayers.h"
 #include "Materials/MaterialExpressionStaticBoolParameter.h"
 #include "Materials/MaterialExpressionStaticComponentMaskParameter.h"
+#include "Materials/MaterialFunction.h"
 #include "UObject/UObjectIterator.h"
 #include "PropertyEditorDelegates.h"
 #include "IDetailsView.h"
diff --git a/Engine/Source/Editor/VirtualTexturingEditor/Private/RuntimeVirtualTextureAssetTypeActions.cpp b/Engine/Source/Editor/VirtualTexturingEditor/Private/RuntimeVirtualTextureAssetTypeActions.cpp
--- a/Engine/Source/Editor/VirtualTexturingEditor/Private/RuntimeVirtualTextureAssetTypeActions.cpp
+++ b/Engine/Source/Editor/VirtualTexturingEditor/Private/RuntimeVirtualTextureAssetTypeActions.cpp
@@ -14,6 +14,7 @@
 #include "AssetRegistry/AssetRegistryHelpers.h"
 #include "Materials/Material.h"
 #include "Materials/MaterialExpressionRuntimeVirtualTextureSample.h"
+#include "Materials/MaterialExpressionTextureBase.h"
 #include "Materials/MaterialFunction.h"
 #include "Materials/MaterialInstance.h"
 #include "Materials/MaterialInterface.h"
diff --git a/Engine/Source/Editor/VirtualTexturingEditor/Private/SConvertToVirtualTexture.cpp b/Engine/Source/Editor/VirtualTexturingEditor/Private/SConvertToVirtualTexture.cpp
--- a/Engine/Source/Editor/VirtualTexturingEditor/Private/SConvertToVirtualTexture.cpp
+++ b/Engine/Source/Editor/VirtualTexturingEditor/Private/SConvertToVirtualTexture.cpp
@@ -6,6 +6,7 @@
 #include "Engine/Texture2D.h"
 #include "Styling/AppStyle.h"
 #include "Materials/Material.h"
+#include "Materials/MaterialFunction.h"
 #include "Materials/MaterialInstance.h"
 #include "Materials/MaterialFunctionInstance.h"
 #include "IContentBrowserSingleton.h"
diff --git a/Engine/Source/Editor/VirtualTexturingEditor/Private/VirtualTextureConversionWorker.cpp b/Engine/Source/Editor/VirtualTexturingEditor/Private/VirtualTextureConversionWorker.cpp
--- a/Engine/Source/Editor/VirtualTexturingEditor/Private/VirtualTextureConversionWorker.cpp
+++ b/Engine/Source/Editor/VirtualTexturingEditor/Private/VirtualTextureConversionWorker.cpp
@@ -6,7 +6,10 @@
 #include "Misc/PackageName.h"
 #include "Materials/Material.h"
 #include "Materials/MaterialInstance.h"
+#include "Materials/MaterialFunction.h"
 #include "Materials/MaterialFunctionInstance.h"
+#include "Materials/MaterialFunctionInterface.h"
+#include "Materials/MaterialExpressionTextureSampleParameter.h"
 #include "MaterialEditor/PreviewMaterial.h"
 #include "Factories/MaterialFactoryNew.h"
 #include "EditorFramework/AssetImportData.h"
diff --git a/Engine/Source/Runtime/Engine/Classes/Materials/MaterialExpressionBentNormalCustomOutput.h b/Engine/Source/Runtime/Engine/Classes/Materials/MaterialExpressionBentNormalCustomOutput.h
--- a/Engine/Source/Runtime/Engine/Classes/Materials/MaterialExpressionBentNormalCustomOutput.h
+++ b/Engine/Source/Runtime/Engine/Classes/Materials/MaterialExpressionBentNormalCustomOutput.h
@@ -6,6 +6,7 @@
 #include "CoreMinimal.h"
 #include "UObject/ObjectMacros.h"
 #include "MaterialExpressionIO.h"
+#include "MaterialValueType.h"
 #include "Materials/MaterialExpressionCustomOutput.h"
 #include "MaterialExpressionBentNormalCustomOutput.generated.h"
 
diff --git a/Engine/Source/Runtime/Engine/Classes/Materials/MaterialExpressionBlendMaterialAttributes.h b/Engine/Source/Runtime/Engine/Classes/Materials/MaterialExpressionBlendMaterialAttributes.h
--- a/Engine/Source/Runtime/Engine/Classes/Materials/MaterialExpressionBlendMaterialAttributes.h
+++ b/Engine/Source/Runtime/Engine/Classes/Materials/MaterialExpressionBlendMaterialAttributes.h
@@ -5,6 +5,7 @@
 #include "CoreMinimal.h"
 #include "UObject/ObjectMacros.h"
 #include "MaterialExpressionIO.h"
+#include "MaterialValueType.h"
 #include "Materials/MaterialExpression.h"
 #include "MaterialExpressionBlendMaterialAttributes.generated.h"
 
diff --git a/Engine/Source/Runtime/Engine/Classes/Materials/MaterialExpressionBreakMaterialAttributes.h b/Engine/Source/Runtime/Engine/Classes/Materials/MaterialExpressionBreakMaterialAttributes.h
--- a/Engine/Source/Runtime/Engine/Classes/Materials/MaterialExpressionBreakMaterialAttributes.h
+++ b/Engine/Source/Runtime/Engine/Classes/Materials/MaterialExpressionBreakMaterialAttributes.h
@@ -6,6 +6,7 @@
 #include "CoreMinimal.h"
 #include "UObject/ObjectMacros.h"
 #include "MaterialExpressionIO.h"
+#include "MaterialValueType.h"
 #include "Materials/MaterialExpression.h"
 #include "MaterialExpressionBreakMaterialAttributes.generated.h"
 
diff --git a/Engine/Source/Runtime/Engine/Classes/Materials/MaterialExpressionChannelMaskParameter.h b/Engine/Source/Runtime/Engine/Classes/Materials/MaterialExpressionChannelMaskParameter.h
--- a/Engine/Source/Runtime/Engine/Classes/Materials/MaterialExpressionChannelMaskParameter.h
+++ b/Engine/Source/Runtime/Engine/Classes/Materials/MaterialExpressionChannelMaskParameter.h
@@ -4,6 +4,7 @@
 
 #include "CoreMinimal.h"
 #include "UObject/ObjectMacros.h"
+#include "MaterialValueType.h"
 #include "Materials/MaterialExpressionChannelMaskParameterColor.h"
 #include "Materials/MaterialExpressionVectorParameter.h"
 #include "MaterialExpressionChannelMaskParameter.generated.h"
diff --git a/Engine/Source/Runtime/Engine/Classes/Materials/MaterialExpressionClearCoatNormalCustomOutput.h b/Engine/Source/Runtime/Engine/Classes/Materials/MaterialExpressionClearCoatNormalCustomOutput.h
--- a/Engine/Source/Runtime/Engine/Classes/Materials/MaterialExpressionClearCoatNormalCustomOutput.h
+++ b/Engine/Source/Runtime/Engine/Classes/Materials/MaterialExpressionClearCoatNormalCustomOutput.h
@@ -6,7 +6,9 @@
 #include "CoreMinimal.h"
 #include "UObject/ObjectMacros.h"
 #include "MaterialExpressionIO.h"
+#include "MaterialValueType.h"
 #include "Materials/MaterialExpressionCustomOutput.h"
+#include "Shader/ShaderTypes.h"
 #include "MaterialExpressionClearCoatNormalCustomOutput.generated.h"
 
 UCLASS(MinimalAPI)
diff --git a/Engine/Source/Runtime/Engine/Classes/Materials/MaterialExpressionConstant2Vector.h b/Engine/Source/Runtime/Engine/Classes/Materials/MaterialExpressionConstant2Vector.h
--- a/Engine/Source/Runtime/Engine/Classes/Materials/MaterialExpressionConstant2Vector.h
+++ b/Engine/Source/Runtime/Engine/Classes/Materials/MaterialExpressionConstant2Vector.h
@@ -5,6 +5,7 @@
 
 #include "CoreMinimal.h"
 #include "UObject/ObjectMacros.h"
+#include "MaterialValueType.h"
 #include "Materials/MaterialExpression.h"
 #include "MaterialExpressionConstant2Vector.generated.h"
 
diff --git a/Engine/Source/Runtime/Engine/Classes/Materials/MaterialExpressionConstant3Vector.h b/Engine/Source/Runtime/Engine/Classes/Materials/MaterialExpressionConstant3Vector.h
--- a/Engine/Source/Runtime/Engine/Classes/Materials/MaterialExpressionConstant3Vector.h
+++ b/Engine/Source/Runtime/Engine/Classes/Materials/MaterialExpressionConstant3Vector.h
@@ -4,6 +4,7 @@
 
 #include "CoreMinimal.h"
 #include "UObject/ObjectMacros.h"
+#include "MaterialValueType.h"
 #include "Materials/MaterialExpression.h"
 #include "MaterialExpressionConstant3Vector.generated.h"
 
diff --git a/Engine/Source/Runtime/Engine/Classes/Materials/MaterialExpressionConstant4Vector.h b/Engine/Source/Runtime/Engine/Classes/Materials/MaterialExpressionConstant4Vector.h
--- a/Engine/Source/Runtime/Engine/Classes/Materials/MaterialExpressionConstant4Vector.h
+++ b/Engine/Source/Runtime/Engine/Classes/Materials/MaterialExpressionConstant4Vector.h
@@ -6,6 +6,7 @@
 #include "CoreMinimal.h"
 #include "UObject/ObjectMacros.h"
 #include "Materials/MaterialExpression.h"
+#include "MaterialValueType.h"
 #include "MaterialExpressionConstant4Vector.generated.h"
 
 UCLASS(collapsecategories, hidecategories=Object, MinimalAPI)
diff --git a/Engine/Source/Runtime/Engine/Classes/Materials/MaterialExpressionCurveAtlasRowParameter.h b/Engine/Source/Runtime/Engine/Classes/Materials/MaterialExpressionCurveAtlasRowParameter.h
--- a/Engine/Source/Runtime/Engine/Classes/Materials/MaterialExpressionCurveAtlasRowParameter.h
+++ b/Engine/Source/Runtime/Engine/Classes/Materials/MaterialExpressionCurveAtlasRowParameter.h
@@ -4,11 +4,10 @@
 
 #include "CoreMinimal.h"
 #include "UObject/ObjectMacros.h"
+#include "MaterialValueType.h"
 #include "Materials/MaterialExpressionScalarParameter.h"
 #include "MaterialExpressionCurveAtlasRowParameter.generated.h"
 
-
-
 UCLASS(collapsecategories, hidecategories=(Object, MaterialExpressionScalarParameter), MinimalAPI)
 class UMaterialExpressionCurveAtlasRowParameter : public UMaterialExpressionScalarParameter
 {
diff --git a/Engine/Source/Runtime/Engine/Classes/Materials/MaterialExpressionCustom.h b/Engine/Source/Runtime/Engine/Classes/Materials/MaterialExpressionCustom.h
--- a/Engine/Source/Runtime/Engine/Classes/Materials/MaterialExpressionCustom.h
+++ b/Engine/Source/Runtime/Engine/Classes/Materials/MaterialExpressionCustom.h
@@ -6,6 +6,7 @@
 #include "CoreMinimal.h"
 #include "UObject/ObjectMacros.h"
 #include "MaterialExpressionIO.h"
+#include "MaterialValueType.h"
 #include "Materials/MaterialExpression.h"
 #include "MaterialExpressionCustom.generated.h"
 
diff --git a/Engine/Source/Runtime/Engine/Classes/Materials/MaterialExpressionCustomOutput.h b/Engine/Source/Runtime/Engine/Classes/Materials/MaterialExpressionCustomOutput.h
--- a/Engine/Source/Runtime/Engine/Classes/Materials/MaterialExpressionCustomOutput.h
+++ b/Engine/Source/Runtime/Engine/Classes/Materials/MaterialExpressionCustomOutput.h
@@ -6,6 +6,8 @@
 #include "CoreMinimal.h"
 #include "UObject/ObjectMacros.h"
 #include "Materials/MaterialExpression.h"
+#include "RHIDefinitions.h"
+#include "Shader/ShaderTypes.h"
 #include "MaterialExpressionCustomOutput.generated.h"
 
 UCLASS(abstract,collapsecategories, hidecategories = Object, MinimalAPI)
diff --git a/Engine/Source/Runtime/Engine/Classes/Materials/MaterialExpressionDistanceFieldsRenderingSwitch.h b/Engine/Source/Runtime/Engine/Classes/Materials/MaterialExpressionDistanceFieldsRenderingSwitch.h
--- a/Engine/Source/Runtime/Engine/Classes/Materials/MaterialExpressionDistanceFieldsRenderingSwitch.h
+++ b/Engine/Source/Runtime/Engine/Classes/Materials/MaterialExpressionDistanceFieldsRenderingSwitch.h
@@ -3,6 +3,7 @@
 #pragma once
 
 #include "CoreMinimal.h"
+#include "MaterialValueType.h"
 #include "Materials/MaterialExpression.h"
 #include "UObject/ObjectMacros.h"
 #include "MaterialExpressionDistanceFieldsRenderingSwitch.generated.h"
diff --git a/Engine/Source/Runtime/Engine/Classes/Materials/MaterialExpressionFeatureLevelSwitch.h b/Engine/Source/Runtime/Engine/Classes/Materials/MaterialExpressionFeatureLevelSwitch.h
--- a/Engine/Source/Runtime/Engine/Classes/Materials/MaterialExpressionFeatureLevelSwitch.h
+++ b/Engine/Source/Runtime/Engine/Classes/Materials/MaterialExpressionFeatureLevelSwitch.h
@@ -6,7 +6,9 @@
 #include "CoreMinimal.h"
 #include "UObject/ObjectMacros.h"
 #include "MaterialExpressionIO.h"
+#include "MaterialValueType.h"
 #include "Materials/MaterialExpression.h"
+#include "RHIFeatureLevel.h"
 #include "MaterialExpressionFeatureLevelSwitch.generated.h"
 
 UCLASS(collapsecategories, hidecategories=Object, MinimalAPI)
diff --git a/Engine/Source/Runtime/Engine/Classes/Materials/MaterialExpressionFontSampleParameter.h b/Engine/Source/Runtime/Engine/Classes/Materials/MaterialExpressionFontSampleParameter.h
--- a/Engine/Source/Runtime/Engine/Classes/Materials/MaterialExpressionFontSampleParameter.h
+++ b/Engine/Source/Runtime/Engine/Classes/Materials/MaterialExpressionFontSampleParameter.h
@@ -6,6 +6,7 @@
 #include "CoreMinimal.h"
 #include "UObject/ObjectMacros.h"
 #include "Misc/Guid.h"
+#include "MaterialTypes.h"
 #include "Materials/MaterialExpressionFontSample.h"
 #include "MaterialExpressionFontSampleParameter.generated.h"
 
diff --git a/Engine/Source/Runtime/Engine/Classes/Materials/MaterialExpressionGenericConstant.h b/Engine/Source/Runtime/Engine/Classes/Materials/MaterialExpressionGenericConstant.h
--- a/Engine/Source/Runtime/Engine/Classes/Materials/MaterialExpressionGenericConstant.h
+++ b/Engine/Source/Runtime/Engine/Classes/Materials/MaterialExpressionGenericConstant.h
@@ -4,6 +4,7 @@
 #include "CoreMinimal.h"
 #include "UObject/ObjectMacros.h"
 #include "Materials/MaterialExpression.h"
+#include "Shader/ShaderTypes.h"
 #include "MaterialExpressionGenericConstant.generated.h"
 
 UCLASS(abstract, MinimalAPI)
diff --git a/Engine/Source/Runtime/Engine/Classes/Materials/MaterialExpressionGetMaterialAttributes.h b/Engine/Source/Runtime/Engine/Classes/Materials/MaterialExpressionGetMaterialAttributes.h
--- a/Engine/Source/Runtime/Engine/Classes/Materials/MaterialExpressionGetMaterialAttributes.h
+++ b/Engine/Source/Runtime/Engine/Classes/Materials/MaterialExpressionGetMaterialAttributes.h
@@ -6,6 +6,7 @@
 #include "UObject/ObjectMacros.h"
 #include "Misc/Guid.h"
 #include "MaterialExpressionIO.h"
+#include "MaterialValueType.h"
 #include "Materials/MaterialExpression.h"
 #include "MaterialExpressionGetMaterialAttributes.generated.h"
 
diff --git a/Engine/Source/Runtime/Engine/Classes/Materials/MaterialExpressionIf.h b/Engine/Source/Runtime/Engine/Classes/Materials/MaterialExpressionIf.h
--- a/Engine/Source/Runtime/Engine/Classes/Materials/MaterialExpressionIf.h
+++ b/Engine/Source/Runtime/Engine/Classes/Materials/MaterialExpressionIf.h
@@ -6,6 +6,7 @@
 #include "CoreMinimal.h"
 #include "UObject/ObjectMacros.h"
 #include "MaterialExpressionIO.h"
+#include "MaterialValueType.h"
 #include "Materials/MaterialExpression.h"
 #include "MaterialExpressionIf.generated.h"
 
diff --git a/Engine/Source/Runtime/Engine/Classes/Materials/MaterialExpressionMaterialFunctionCall.h b/Engine/Source/Runtime/Engine/Classes/Materials/MaterialExpressionMaterialFunctionCall.h
--- a/Engine/Source/Runtime/Engine/Classes/Materials/MaterialExpressionMaterialFunctionCall.h
+++ b/Engine/Source/Runtime/Engine/Classes/Materials/MaterialExpressionMaterialFunctionCall.h
@@ -10,10 +10,12 @@
 #include "UObject/ObjectMacros.h"
 #include "Misc/Guid.h"
 #include "MaterialExpressionIO.h"
+#include "MaterialTypes.h"
 #include "Materials/MaterialExpression.h"
 #include "MaterialExpressionMaterialFunctionCall.generated.h"
 
 class UMaterialFunction;
+class FMaterialFunctionCompileState;
 struct FPropertyChangedEvent;
 
 /** Struct that stores information about a function input which is needed to maintain connections and implement the function call. */
diff --git a/Engine/Source/Runtime/Engine/Classes/Materials/MaterialExpressionMaterialLayerOutput.h b/Engine/Source/Runtime/Engine/Classes/Materials/MaterialExpressionMaterialLayerOutput.h
--- a/Engine/Source/Runtime/Engine/Classes/Materials/MaterialExpressionMaterialLayerOutput.h
+++ b/Engine/Source/Runtime/Engine/Classes/Materials/MaterialExpressionMaterialLayerOutput.h
@@ -5,6 +5,7 @@
 #include "UObject/ObjectMacros.h"
 #include "Misc/Guid.h"
 #include "MaterialExpressionIO.h"
+#include "MaterialValueType.h"
 #include "Materials/MaterialExpressionFunctionOutput.h"
 #include "MaterialExpressionMaterialLayerOutput.generated.h"
 
diff --git a/Engine/Source/Runtime/Engine/Classes/Materials/MaterialExpressionMaterialProxyReplace.h b/Engine/Source/Runtime/Engine/Classes/Materials/MaterialExpressionMaterialProxyReplace.h
--- a/Engine/Source/Runtime/Engine/Classes/Materials/MaterialExpressionMaterialProxyReplace.h
+++ b/Engine/Source/Runtime/Engine/Classes/Materials/MaterialExpressionMaterialProxyReplace.h
@@ -6,6 +6,7 @@
 #include "CoreMinimal.h"
 #include "UObject/ObjectMacros.h"
 #include "MaterialExpressionIO.h"
+#include "MaterialValueType.h"
 #include "Materials/MaterialExpression.h"
 #include "MaterialExpressionMaterialProxyReplace.generated.h"
 
diff --git a/Engine/Source/Runtime/Engine/Classes/Materials/MaterialExpressionParameter.h b/Engine/Source/Runtime/Engine/Classes/Materials/MaterialExpressionParameter.h
--- a/Engine/Source/Runtime/Engine/Classes/Materials/MaterialExpressionParameter.h
+++ b/Engine/Source/Runtime/Engine/Classes/Materials/MaterialExpressionParameter.h
@@ -6,6 +6,7 @@
 #include "CoreMinimal.h"
 #include "UObject/ObjectMacros.h"
 #include "Misc/Guid.h"
+#include "MaterialTypes.h"
 #include "Materials/MaterialExpression.h"
 #include "MaterialExpressionParameter.generated.h"
 
diff --git a/Engine/Source/Runtime/Engine/Classes/Materials/MaterialExpressionPathTracingQualitySwitch.h b/Engine/Source/Runtime/Engine/Classes/Materials/MaterialExpressionPathTracingQualitySwitch.h
--- a/Engine/Source/Runtime/Engine/Classes/Materials/MaterialExpressionPathTracingQualitySwitch.h
+++ b/Engine/Source/Runtime/Engine/Classes/Materials/MaterialExpressionPathTracingQualitySwitch.h
@@ -5,6 +5,7 @@
 #include "CoreMinimal.h"
 #include "UObject/ObjectMacros.h"
 #include "MaterialExpressionIO.h"
+#include "MaterialValueType.h"
 #include "Materials/MaterialExpression.h"
 #include "MaterialExpressionPathTracingQualitySwitch.generated.h"
 
diff --git a/Engine/Source/Runtime/Engine/Classes/Materials/MaterialExpressionPreviousFrameSwitch.h b/Engine/Source/Runtime/Engine/Classes/Materials/MaterialExpressionPreviousFrameSwitch.h
--- a/Engine/Source/Runtime/Engine/Classes/Materials/MaterialExpressionPreviousFrameSwitch.h
+++ b/Engine/Source/Runtime/Engine/Classes/Materials/MaterialExpressionPreviousFrameSwitch.h
@@ -2,6 +2,7 @@
 
 
 #pragma once
+#include "MaterialValueType.h"
 #include "Materials/MaterialExpression.h"
 #include "MaterialExpressionPreviousFrameSwitch.generated.h"
 
diff --git a/Engine/Source/Runtime/Engine/Classes/Materials/MaterialExpressionQualitySwitch.h b/Engine/Source/Runtime/Engine/Classes/Materials/MaterialExpressionQualitySwitch.h
--- a/Engine/Source/Runtime/Engine/Classes/Materials/MaterialExpressionQualitySwitch.h
+++ b/Engine/Source/Runtime/Engine/Classes/Materials/MaterialExpressionQualitySwitch.h
@@ -6,6 +6,7 @@
 #include "CoreMinimal.h"
 #include "UObject/ObjectMacros.h"
 #include "MaterialExpressionIO.h"
+#include "MaterialValueType.h"
 #include "Materials/MaterialExpression.h"
 #include "MaterialExpressionQualitySwitch.generated.h"
 
diff --git a/Engine/Source/Runtime/Engine/Classes/Materials/MaterialExpressionRayTracingQualitySwitch.h b/Engine/Source/Runtime/Engine/Classes/Materials/MaterialExpressionRayTracingQualitySwitch.h
--- a/Engine/Source/Runtime/Engine/Classes/Materials/MaterialExpressionRayTracingQualitySwitch.h
+++ b/Engine/Source/Runtime/Engine/Classes/Materials/MaterialExpressionRayTracingQualitySwitch.h
@@ -5,6 +5,7 @@
 #include "CoreMinimal.h"
 #include "UObject/ObjectMacros.h"
 #include "MaterialExpressionIO.h"
+#include "MaterialValueType.h"
 #include "Materials/MaterialExpression.h"
 #include "MaterialExpressionRayTracingQualitySwitch.generated.h"
 
diff --git a/Engine/Source/Runtime/Engine/Classes/Materials/MaterialExpressionRuntimeVirtualTextureReplace.h b/Engine/Source/Runtime/Engine/Classes/Materials/MaterialExpressionRuntimeVirtualTextureReplace.h
--- a/Engine/Source/Runtime/Engine/Classes/Materials/MaterialExpressionRuntimeVirtualTextureReplace.h
+++ b/Engine/Source/Runtime/Engine/Classes/Materials/MaterialExpressionRuntimeVirtualTextureReplace.h
@@ -3,6 +3,7 @@
 #pragma once
 
 #include "CoreMinimal.h"
+#include "MaterialValueType.h"
 #include "Materials/MaterialExpression.h"
 #include "UObject/ObjectMacros.h"
 #include "MaterialExpressionRuntimeVirtualTextureReplace.generated.h"
diff --git a/Engine/Source/Runtime/Engine/Classes/Materials/MaterialExpressionRuntimeVirtualTextureSampleParameter.h b/Engine/Source/Runtime/Engine/Classes/Materials/MaterialExpressionRuntimeVirtualTextureSampleParameter.h
--- a/Engine/Source/Runtime/Engine/Classes/Materials/MaterialExpressionRuntimeVirtualTextureSampleParameter.h
+++ b/Engine/Source/Runtime/Engine/Classes/Materials/MaterialExpressionRuntimeVirtualTextureSampleParameter.h
@@ -3,6 +3,7 @@
 #pragma once
 
 #include "CoreMinimal.h"
+#include "MaterialTypes.h"
 #include "Materials/MaterialExpressionRuntimeVirtualTextureSample.h"
 #include "MaterialExpressionRuntimeVirtualTextureSampleParameter.generated.h"
 
diff --git a/Engine/Source/Runtime/Engine/Classes/Materials/MaterialExpressionShaderStageSwitch.h b/Engine/Source/Runtime/Engine/Classes/Materials/MaterialExpressionShaderStageSwitch.h
--- a/Engine/Source/Runtime/Engine/Classes/Materials/MaterialExpressionShaderStageSwitch.h
+++ b/Engine/Source/Runtime/Engine/Classes/Materials/MaterialExpressionShaderStageSwitch.h
@@ -7,6 +7,7 @@
 #include "UObject/ObjectMacros.h"
 #include "MaterialExpressionIO.h"
 #include "Materials/MaterialExpression.h"
+#include "RHIDefinitions.h"
 #include "MaterialExpressionShaderStageSwitch.generated.h"
 
 UCLASS(collapsecategories, hidecategories = Object, MinimalAPI)
diff --git a/Engine/Source/Runtime/Engine/Classes/Materials/MaterialExpressionShadingPathSwitch.h b/Engine/Source/Runtime/Engine/Classes/Materials/MaterialExpressionShadingPathSwitch.h
--- a/Engine/Source/Runtime/Engine/Classes/Materials/MaterialExpressionShadingPathSwitch.h
+++ b/Engine/Source/Runtime/Engine/Classes/Materials/MaterialExpressionShadingPathSwitch.h
@@ -6,7 +6,9 @@
 #include "CoreMinimal.h"
 #include "UObject/ObjectMacros.h"
 #include "MaterialExpressionIO.h"
+#include "MaterialValueType.h"
 #include "Materials/MaterialExpression.h"
+#include "RHIDefinitions.h"
 #include "MaterialExpressionShadingPathSwitch.generated.h"
 
 UCLASS(collapsecategories, hidecategories=Object, MinimalAPI)
diff --git a/Engine/Source/Runtime/Engine/Classes/Materials/MaterialExpressionStaticBool.h b/Engine/Source/Runtime/Engine/Classes/Materials/MaterialExpressionStaticBool.h
--- a/Engine/Source/Runtime/Engine/Classes/Materials/MaterialExpressionStaticBool.h
+++ b/Engine/Source/Runtime/Engine/Classes/Materials/MaterialExpressionStaticBool.h
@@ -5,6 +5,7 @@
 
 #include "CoreMinimal.h"
 #include "UObject/ObjectMacros.h"
+#include "MaterialValueType.h"
 #include "Materials/MaterialExpression.h"
 #include "MaterialExpressionStaticBool.generated.h"
 
diff --git a/Engine/Source/Runtime/Engine/Classes/Materials/MaterialExpressionStaticBoolParameter.h b/Engine/Source/Runtime/Engine/Classes/Materials/MaterialExpressionStaticBoolParameter.h
--- a/Engine/Source/Runtime/Engine/Classes/Materials/MaterialExpressionStaticBoolParameter.h
+++ b/Engine/Source/Runtime/Engine/Classes/Materials/MaterialExpressionStaticBoolParameter.h
@@ -5,6 +5,7 @@
 
 #include "CoreMinimal.h"
 #include "UObject/ObjectMacros.h"
+#include "MaterialValueType.h"
 #include "Materials/MaterialExpressionParameter.h"
 #include "MaterialExpressionStaticBoolParameter.generated.h"
 
diff --git a/Engine/Source/Runtime/Engine/Classes/Materials/MaterialExpressionStaticSwitch.h b/Engine/Source/Runtime/Engine/Classes/Materials/MaterialExpressionStaticSwitch.h
--- a/Engine/Source/Runtime/Engine/Classes/Materials/MaterialExpressionStaticSwitch.h
+++ b/Engine/Source/Runtime/Engine/Classes/Materials/MaterialExpressionStaticSwitch.h
@@ -6,6 +6,7 @@
 #include "CoreMinimal.h"
 #include "UObject/ObjectMacros.h"
 #include "MaterialExpressionIO.h"
+#include "MaterialValueType.h"
 #include "Materials/MaterialExpression.h"
 #include "MaterialExpressionStaticSwitch.generated.h"
 
diff --git a/Engine/Source/Runtime/Engine/Classes/Materials/MaterialExpressionStaticSwitchParameter.h b/Engine/Source/Runtime/Engine/Classes/Materials/MaterialExpressionStaticSwitchParameter.h
--- a/Engine/Source/Runtime/Engine/Classes/Materials/MaterialExpressionStaticSwitchParameter.h
+++ b/Engine/Source/Runtime/Engine/Classes/Materials/MaterialExpressionStaticSwitchParameter.h
@@ -6,6 +6,7 @@
 #include "CoreMinimal.h"
 #include "UObject/ObjectMacros.h"
 #include "MaterialExpressionIO.h"
+#include "MaterialValueType.h"
 #include "Materials/MaterialExpressionStaticBoolParameter.h"
 #include "MaterialExpressionStaticSwitchParameter.generated.h"
 
diff --git a/Engine/Source/Runtime/Engine/Classes/Materials/MaterialExpressionTangentOutput.h b/Engine/Source/Runtime/Engine/Classes/Materials/MaterialExpressionTangentOutput.h
--- a/Engine/Source/Runtime/Engine/Classes/Materials/MaterialExpressionTangentOutput.h
+++ b/Engine/Source/Runtime/Engine/Classes/Materials/MaterialExpressionTangentOutput.h
@@ -5,6 +5,7 @@
 #include "CoreMinimal.h"
 #include "UObject/ObjectMacros.h"
 #include "MaterialExpressionIO.h"
+#include "MaterialValueType.h"
 #include "Materials/MaterialExpressionCustomOutput.h"
 #include "MaterialExpressionTangentOutput.generated.h"
 
diff --git a/Engine/Source/Runtime/Engine/Classes/Materials/MaterialExpressionTextureBase.h b/Engine/Source/Runtime/Engine/Classes/Materials/MaterialExpressionTextureBase.h
--- a/Engine/Source/Runtime/Engine/Classes/Materials/MaterialExpressionTextureBase.h
+++ b/Engine/Source/Runtime/Engine/Classes/Materials/MaterialExpressionTextureBase.h
@@ -14,6 +14,7 @@
 
 class UTexture;
 struct FPropertyChangedEvent;
+enum EMaterialSamplerType : int;
 
 UCLASS(abstract, hidecategories=Object)
 class ENGINE_API UMaterialExpressionTextureBase : public UMaterialExpression 
@@ -21,10 +22,10 @@ class ENGINE_API UMaterialExpressionTextureBase : public UMaterialExpression
 	GENERATED_UCLASS_BODY()
 
 	UPROPERTY(EditAnywhere, BlueprintReadWrite, Category=MaterialExpressionTextureBase)
-	TObjectPtr<class UTexture> Texture;
+	TObjectPtr<UTexture> Texture;
 
 	UPROPERTY(EditAnywhere, BlueprintReadWrite, Category=MaterialExpressionTextureBase, meta = (ShowAsInputPin = "Advanced"))
-	TEnumAsByte<enum EMaterialSamplerType> SamplerType;
+	TEnumAsByte<EMaterialSamplerType> SamplerType;
 	
 	/** Is default selected texture when using mesh paint mode texture painting */
 	UPROPERTY(EditAnywhere, Category=MaterialExpressionTextureBase)
diff --git a/Engine/Source/Runtime/Engine/Classes/Materials/MaterialExpressionTextureObjectParameter.h b/Engine/Source/Runtime/Engine/Classes/Materials/MaterialExpressionTextureObjectParameter.h
--- a/Engine/Source/Runtime/Engine/Classes/Materials/MaterialExpressionTextureObjectParameter.h
+++ b/Engine/Source/Runtime/Engine/Classes/Materials/MaterialExpressionTextureObjectParameter.h
@@ -9,6 +9,7 @@
 
 #include "CoreMinimal.h"
 #include "UObject/ObjectMacros.h"
+#include "MaterialValueType.h"
 #include "Materials/MaterialExpressionTextureSampleParameter.h"
 #include "MaterialExpressionTextureObjectParameter.generated.h"
 
diff --git a/Engine/Source/Runtime/Engine/Classes/Materials/MaterialExpressionTextureSample.h b/Engine/Source/Runtime/Engine/Classes/Materials/MaterialExpressionTextureSample.h
--- a/Engine/Source/Runtime/Engine/Classes/Materials/MaterialExpressionTextureSample.h
+++ b/Engine/Source/Runtime/Engine/Classes/Materials/MaterialExpressionTextureSample.h
@@ -10,6 +10,8 @@
 #include "MaterialExpressionTextureSample.generated.h"
 
 struct FPropertyChangedEvent;
+enum ETextureMipValueMode : int;
+enum ESamplerSourceMode : int;
 
 UCLASS(collapsecategories, hidecategories=Object)
 class ENGINE_API UMaterialExpressionTextureSample : public UMaterialExpressionTextureBase
@@ -43,7 +45,7 @@ class ENGINE_API UMaterialExpressionTextureSample : public UMaterialExpressionTe
 
 	/** Defines how the MipValue property is applied to the texture lookup */
 	UPROPERTY(EditAnywhere, Category=MaterialExpressionTextureSample, meta=(DisplayName = "MipValueMode", ShowAsInputPin = "Advanced"))
-	TEnumAsByte<enum ETextureMipValueMode> MipValueMode;
+	TEnumAsByte<ETextureMipValueMode> MipValueMode;
 
 	/** 
 	 * Controls where the sampler for this texture lookup will come from.  
@@ -52,7 +54,7 @@ class ENGINE_API UMaterialExpressionTextureSample : public UMaterialExpressionTe
 	 * This allows materials to use more than 16 unique textures on SM5 platforms.
 	 */
 	UPROPERTY(EditAnywhere, Category=MaterialExpressionTextureSample, Meta = (ShowAsInputPin = "Advanced"))
-	TEnumAsByte<enum ESamplerSourceMode> SamplerSource;
+	TEnumAsByte<ESamplerSourceMode> SamplerSource;
 
 	/** Whether the texture should be sampled with per view mip biasing for sharper output with Temporal AA. */
 	UPROPERTY(EditAnywhere, Category = MaterialExpressionTextureSample)
diff --git a/Engine/Source/Runtime/Engine/Classes/Materials/MaterialExpressionTextureSampleParameter.h b/Engine/Source/Runtime/Engine/Classes/Materials/MaterialExpressionTextureSampleParameter.h
--- a/Engine/Source/Runtime/Engine/Classes/Materials/MaterialExpressionTextureSampleParameter.h
+++ b/Engine/Source/Runtime/Engine/Classes/Materials/MaterialExpressionTextureSampleParameter.h
@@ -6,6 +6,7 @@
 #include "CoreMinimal.h"
 #include "UObject/ObjectMacros.h"
 #include "Misc/Guid.h"
+#include "MaterialTypes.h"
 #include "Materials/MaterialExpressionTextureSample.h"
 #include "MaterialExpressionTextureSampleParameter.generated.h"
 
diff --git a/Engine/Source/Runtime/Engine/Classes/Materials/MaterialExpressionVertexInterpolator.h b/Engine/Source/Runtime/Engine/Classes/Materials/MaterialExpressionVertexInterpolator.h
--- a/Engine/Source/Runtime/Engine/Classes/Materials/MaterialExpressionVertexInterpolator.h
+++ b/Engine/Source/Runtime/Engine/Classes/Materials/MaterialExpressionVertexInterpolator.h
@@ -2,6 +2,7 @@
 
 #pragma once
 
+#include "MaterialValueType.h"
 #include "Materials/MaterialExpressionCustomOutput.h"
 #include "MaterialExpressionVertexInterpolator.generated.h"
 
diff --git a/Engine/Source/Runtime/Engine/Classes/Materials/MaterialExpressionVirtualTextureFeatureSwitch.h b/Engine/Source/Runtime/Engine/Classes/Materials/MaterialExpressionVirtualTextureFeatureSwitch.h
--- a/Engine/Source/Runtime/Engine/Classes/Materials/MaterialExpressionVirtualTextureFeatureSwitch.h
+++ b/Engine/Source/Runtime/Engine/Classes/Materials/MaterialExpressionVirtualTextureFeatureSwitch.h
@@ -3,6 +3,8 @@
 #pragma once
 
 #include "CoreMinimal.h"
+#include "MaterialValueType.h"
+#include "MaterialValueType.h"
 #include "Materials/MaterialExpression.h"
 #include "UObject/ObjectMacros.h"
 #include "MaterialExpressionVirtualTextureFeatureSwitch.generated.h"
diff --git a/Engine/Source/Runtime/Engine/Classes/Materials/MaterialExpressionVolumetricAdvancedMaterialOutput.h b/Engine/Source/Runtime/Engine/Classes/Materials/MaterialExpressionVolumetricAdvancedMaterialOutput.h
--- a/Engine/Source/Runtime/Engine/Classes/Materials/MaterialExpressionVolumetricAdvancedMaterialOutput.h
+++ b/Engine/Source/Runtime/Engine/Classes/Materials/MaterialExpressionVolumetricAdvancedMaterialOutput.h
@@ -5,6 +5,7 @@
 #include "CoreMinimal.h"
 #include "Materials/MaterialExpressionCustomOutput.h"
 #include "UObject/ObjectMacros.h"
+#include "Shader/ShaderTypes.h"
 #include "MaterialExpressionVolumetricAdvancedMaterialOutput.generated.h"
 
 /** Material output expression for writing advanced volumetric material properties. */
diff --git a/Engine/Source/Runtime/Engine/Classes/Materials/MaterialFunction.h b/Engine/Source/Runtime/Engine/Classes/Materials/MaterialFunction.h
--- a/Engine/Source/Runtime/Engine/Classes/Materials/MaterialFunction.h
+++ b/Engine/Source/Runtime/Engine/Classes/Materials/MaterialFunction.h
@@ -9,6 +9,7 @@
 #include "Templates/Casts.h"
 #include "Materials/MaterialFunctionInterface.h"
 #include "StaticParameterSet.h"
+#include "MaterialExpression.h"
 #include "MaterialFunction.generated.h"
 
 class UMaterial;
diff --git a/Engine/Source/Runtime/Engine/Private/Materials/HLSLMaterialTranslator.cpp b/Engine/Source/Runtime/Engine/Private/Materials/HLSLMaterialTranslator.cpp
--- a/Engine/Source/Runtime/Engine/Private/Materials/HLSLMaterialTranslator.cpp
+++ b/Engine/Source/Runtime/Engine/Private/Materials/HLSLMaterialTranslator.cpp
@@ -7,7 +7,9 @@
 #include "HLSLMaterialTranslator.h"
 #include "VT/VirtualTextureScalability.h"
 #include "Materials/MaterialExpressionCustom.h"
+#include "Materials/MaterialExpressionMaterialAttributeLayers.h"
 #include "Materials/MaterialExpressionStrata.h"
+#include "Materials/MaterialFunction.h"
 #include "MaterialExpressionSettings.h"
 #include "Engine/RendererSettings.h"
 #include "ShaderPlatformCachedIniValue.h"
diff --git a/Engine/Source/Runtime/Engine/Private/Materials/MaterialCachedData.cpp b/Engine/Source/Runtime/Engine/Private/Materials/MaterialCachedData.cpp
--- a/Engine/Source/Runtime/Engine/Private/Materials/MaterialCachedData.cpp
+++ b/Engine/Source/Runtime/Engine/Private/Materials/MaterialCachedData.cpp
@@ -27,6 +27,7 @@
 #include "Materials/MaterialExpressionCurveAtlasRowParameter.h"
 #include "Materials/MaterialExpressionSetMaterialAttributes.h"
 #include "Materials/MaterialExpressionMakeMaterialAttributes.h"
+#include "Materials/MaterialExpressionMaterialAttributeLayers.h"
 #include "Materials/MaterialInstance.h"
 #include "Materials/MaterialFunction.h"
 #include "Materials/MaterialFunctionInstance.h"
diff --git a/Engine/Source/Runtime/Engine/Private/Materials/MaterialExpressionHLSL.cpp b/Engine/Source/Runtime/Engine/Private/Materials/MaterialExpressionHLSL.cpp
--- a/Engine/Source/Runtime/Engine/Private/Materials/MaterialExpressionHLSL.cpp
+++ b/Engine/Source/Runtime/Engine/Private/Materials/MaterialExpressionHLSL.cpp
@@ -141,6 +141,7 @@
 #include "Materials/MaterialExpressionWhileLoop.h"
 #include "Materials/MaterialFunctionInterface.h"
 #include "MaterialHLSLTree.h"
+#include "MaterialShared.h"
 #include "HLSLTree/HLSLTree.h"
 #include "HLSLTree/HLSLTreeCommon.h"
 #include "Engine/Engine.h"
diff --git a/Engine/Source/Runtime/Engine/Private/Materials/MaterialHLSLGenerator.cpp b/Engine/Source/Runtime/Engine/Private/Materials/MaterialHLSLGenerator.cpp
--- a/Engine/Source/Runtime/Engine/Private/Materials/MaterialHLSLGenerator.cpp
+++ b/Engine/Source/Runtime/Engine/Private/Materials/MaterialHLSLGenerator.cpp
@@ -11,6 +11,7 @@
 #include "Materials/MaterialExpressionCustomOutput.h"
 #include "Materials/MaterialExpressionVertexInterpolator.h"
 #include "Materials/MaterialExpressionExecBegin.h"
+#include "Materials/MaterialFunction.h"
 #include "MaterialHLSLTree.h"
 #include "Materials/Material.h"
 #include "MaterialCachedHLSLTree.h"
diff --git a/Engine/Source/Runtime/Landscape/Classes/Materials/MaterialExpressionLandscapeGrassOutput.h b/Engine/Source/Runtime/Landscape/Classes/Materials/MaterialExpressionLandscapeGrassOutput.h
--- a/Engine/Source/Runtime/Landscape/Classes/Materials/MaterialExpressionLandscapeGrassOutput.h
+++ b/Engine/Source/Runtime/Landscape/Classes/Materials/MaterialExpressionLandscapeGrassOutput.h
@@ -6,7 +6,7 @@
 #include "CoreMinimal.h"
 #include "UObject/ObjectMacros.h"
 #include "MaterialExpressionIO.h"
-
+#include "MaterialValueType.h"
 #include "Materials/MaterialExpressionCustomOutput.h"
 #include "MaterialExpressionLandscapeGrassOutput.generated.h"
 
diff --git a/Engine/Source/Runtime/Landscape/Classes/Materials/MaterialExpressionLandscapeLayerBlend.h b/Engine/Source/Runtime/Landscape/Classes/Materials/MaterialExpressionLandscapeLayerBlend.h
--- a/Engine/Source/Runtime/Landscape/Classes/Materials/MaterialExpressionLandscapeLayerBlend.h
+++ b/Engine/Source/Runtime/Landscape/Classes/Materials/MaterialExpressionLandscapeLayerBlend.h
@@ -7,6 +7,7 @@
 #include "UObject/ObjectMacros.h"
 #include "Misc/Guid.h"
 #include "MaterialExpressionIO.h"
+#include "MaterialValueType.h"
 #include "Materials/MaterialExpression.h"
 #include "MaterialExpressionLandscapeLayerBlend.generated.h"
 
diff --git a/Engine/Source/Runtime/Landscape/Classes/Materials/MaterialExpressionLandscapeLayerSwitch.h b/Engine/Source/Runtime/Landscape/Classes/Materials/MaterialExpressionLandscapeLayerSwitch.h
--- a/Engine/Source/Runtime/Landscape/Classes/Materials/MaterialExpressionLandscapeLayerSwitch.h
+++ b/Engine/Source/Runtime/Landscape/Classes/Materials/MaterialExpressionLandscapeLayerSwitch.h
@@ -7,6 +7,7 @@
 #include "UObject/ObjectMacros.h"
 #include "Misc/Guid.h"
 #include "MaterialExpressionIO.h"
+#include "MaterialValueType.h"
 #include "Materials/MaterialExpression.h"
 #include "MaterialExpressionLandscapeLayerSwitch.generated.h"
 
diff --git a/Engine/Source/Runtime/Landscape/Classes/Materials/MaterialExpressionLandscapeLayerWeight.h b/Engine/Source/Runtime/Landscape/Classes/Materials/MaterialExpressionLandscapeLayerWeight.h
--- a/Engine/Source/Runtime/Landscape/Classes/Materials/MaterialExpressionLandscapeLayerWeight.h
+++ b/Engine/Source/Runtime/Landscape/Classes/Materials/MaterialExpressionLandscapeLayerWeight.h
@@ -7,6 +7,7 @@
 #include "UObject/ObjectMacros.h"
 #include "Misc/Guid.h"
 #include "MaterialExpressionIO.h"
+#include "MaterialValueType.h"
 #include "Materials/MaterialExpression.h"
 #include "MaterialExpressionLandscapeLayerWeight.generated.h"
 
diff --git a/Engine/Source/Runtime/Landscape/Classes/Materials/MaterialExpressionLandscapePhysicalMaterialOutput.h b/Engine/Source/Runtime/Landscape/Classes/Materials/MaterialExpressionLandscapePhysicalMaterialOutput.h
--- a/Engine/Source/Runtime/Landscape/Classes/Materials/MaterialExpressionLandscapePhysicalMaterialOutput.h
+++ b/Engine/Source/Runtime/Landscape/Classes/Materials/MaterialExpressionLandscapePhysicalMaterialOutput.h
@@ -4,7 +4,7 @@
 
 #include "CoreMinimal.h"
 #include "Materials/MaterialExpressionCustomOutput.h"
-
+#include "MaterialValueType.h"
 #include "MaterialExpressionLandscapePhysicalMaterialOutput.generated.h"
 
 /** Structure linking a material expression input with a physical material. For use by UMaterialExpressionLandscapePhysicalMaterialOutput. */
