Added world partition option wp.Runtime.EnableServerDataLayerStreamingOut to control whether server can stream-out or not the streaming levels of runtime cells linked to runtime data layers.

diff --git a/Engine/Source/Runtime/Engine/Private/WorldPartition/WorldPartition.cpp b/Engine/Source/Runtime/Engine/Private/WorldPartition/WorldPartition.cpp
--- a/Engine/Source/Runtime/Engine/Private/WorldPartition/WorldPartition.cpp
+++ b/Engine/Source/Runtime/Engine/Private/WorldPartition/WorldPartition.cpp
@@ -122,6 +122,15 @@ FAutoConsoleVariableRef UWorldPartition::CVarEnableServerStreamingOut(
 	TEXT("Changing the value while the game is running won't be considered."),
 	WorldPartition::ECVF_Runtime_ReadOnly);
 
+
+bool UWorldPartition::bEnableServerDataLayerStreamingOut = true;
+FAutoConsoleVariableRef UWorldPartition::CVarEnableServerDataLayerStreamingOut(
+	TEXT("wp.Runtime.EnableServerDataLayerStreamingOut"),
+	UWorldPartition::bEnableServerDataLayerStreamingOut,
+	TEXT("Turn on/off to allow or not the server to stream out levels of runtime cells linked to runtime data layers.\n")
+	TEXT("Changing the value while the game is running won't be considered."),
+	WorldPartition::ECVF_Runtime_ReadOnly);
+
 bool UWorldPartition::bUseMakingVisibleTransactionRequests = false;
 FAutoConsoleVariableRef UWorldPartition::CVarUseMakingVisibleTransactionRequests(
 	TEXT("wp.Runtime.UseMakingVisibleTransactionRequests"),
@@ -1090,6 +1099,18 @@ bool UWorldPartition::IsServerStreamingOutEnabled() const
 	return bCachedIsServerStreamingOutEnabled.Get(false);
 }
 
+bool UWorldPartition::IsServerDataLayerStreamingOutEnabled() const
+{
+	// Resolve once (we don't allow changing the state at runtime)
+	if (!bCachedIsServerDataLayerStreamingOutEnabled.IsSet())
+	{
+		UWorld* OwningWorld = GetWorld();
+		bCachedIsServerDataLayerStreamingOutEnabled = OwningWorld && OwningWorld->IsGameWorld() && IsServer() && UWorldPartition::bEnableServerDataLayerStreamingOut;
+	}
+
+	return bCachedIsServerDataLayerStreamingOutEnabled.Get(false);
+}
+
 bool UWorldPartition::UseMakingVisibleTransactionRequests() const
 {
 	// Resolve once (we don't allow changing the state at runtime)
diff --git a/Engine/Source/Runtime/Engine/Private/WorldPartition/WorldPartitionStreamingPolicy.cpp b/Engine/Source/Runtime/Engine/Private/WorldPartition/WorldPartitionStreamingPolicy.cpp
--- a/Engine/Source/Runtime/Engine/Private/WorldPartition/WorldPartitionStreamingPolicy.cpp
+++ b/Engine/Source/Runtime/Engine/Private/WorldPartition/WorldPartitionStreamingPolicy.cpp
@@ -405,7 +405,6 @@ void UWorldPartitionStreamingPolicy::UpdateStreamingState()
 	const bool bCanStream = WorldPartition->CanStream();
 	const bool bIsServer = WorldPartition->IsServer();
 	const bool bIsServerStreamingEnabled = WorldPartition->IsServerStreamingEnabled();
-	const bool bCanDeactivateOrUnloadCells = !bIsServer || WorldPartition->IsServerStreamingOutEnabled();
 	const int32 NewServerStreamingEnabledEpoch = ComputeServerStreamingEnabledEpoch();
 
 	const bool bIsStreamingInEnabled = WorldPartition->IsStreamingInEnabled();
@@ -442,6 +441,9 @@ void UWorldPartitionStreamingPolicy::UpdateStreamingState()
 
 		if (bIsServer)
 		{
+			const bool bCanServerDeactivateOrUnloadCells = WorldPartition->IsServerStreamingOutEnabled();
+			const bool bCanServerDeactivateOrUnloadDataLayerCells = WorldPartition->IsServerDataLayerStreamingOutEnabled();
+
 			TRACE_CPUPROFILER_EVENT_SCOPE(UWorldPartitionStreamingPolicy::UpdateStreamingState_ServerUpdate);
 
 			// Server will activate all non data layer cells at first and then load/activate/unload data layer cells only when the data layer states change
@@ -471,21 +473,32 @@ void UWorldPartitionStreamingPolicy::UpdateStreamingState()
 			TSet<FName> EffectiveActiveDataLayerNames = DataLayerSubsystem->GetEffectiveActiveDataLayerNames();
 			TSet<FName> EffectiveLoadedDataLayerNames = DataLayerSubsystem->GetEffectiveLoadedDataLayerNames();
 
-			auto AddServerFrameCells = [this, &EffectiveLoadedDataLayerNames, &EffectiveActiveDataLayerNames](const UWorldPartitionRuntimeCell* Cell)
+			auto AddServerFrameCell = [this, bCanServerDeactivateOrUnloadDataLayerCells, &EffectiveLoadedDataLayerNames, &EffectiveActiveDataLayerNames](const UWorldPartitionRuntimeCell* Cell)
 			{
+				// Keep Data Layer cells is their current state if server cannot deactivat/unload data layer cells
+				if (!bCanServerDeactivateOrUnloadDataLayerCells && Cell->HasDataLayers())
+				{
+					if (ActivatedCells.Contains(Cell))
+					{
+						FrameActivateCells.Add(Cell);
+						return;
+					}
+					else if (LoadedCells.Contains(Cell))
+					{
+						FrameLoadCells.Add(Cell);
+						return;
+					}
+				}
+				
 				// Non Data Layer Cells + Active Data Layers
-				if (!Cell->HasDataLayers() || Cell->HasAnyDataLayer(EffectiveActiveDataLayerNames))
+				if (!Cell->HasDataLayers() || (EffectiveActiveDataLayerNames.Num() && Cell->HasAnyDataLayer(EffectiveActiveDataLayerNames)))
 				{
 					FrameActivateCells.Add(Cell);
 				}
-
 				// Loaded Data Layers Cells only
-				if (EffectiveLoadedDataLayerNames.Num())
+				else if (Cell->HasDataLayers() && EffectiveLoadedDataLayerNames.Num() && Cell->HasAnyDataLayer(EffectiveLoadedDataLayerNames))
 				{
-					if (Cell->HasDataLayers() && Cell->HasAnyDataLayer(EffectiveLoadedDataLayerNames))
-					{
-						FrameLoadCells.Add(Cell);
-					}
+					FrameLoadCells.Add(Cell);
 				}
 			};
 
@@ -493,23 +506,23 @@ void UWorldPartitionStreamingPolicy::UpdateStreamingState()
 			{
 				if (WorldPartition->RuntimeHash)
 				{
-					WorldPartition->RuntimeHash->ForEachStreamingCells([this, &AddServerFrameCells](const UWorldPartitionRuntimeCell* Cell)
+					WorldPartition->RuntimeHash->ForEachStreamingCells([this, &AddServerFrameCell](const UWorldPartitionRuntimeCell* Cell)
 					{
-						AddServerFrameCells(Cell);
+						AddServerFrameCell(Cell);
 						return true;
 					});
 				}
 			}
-			else if (!bCanDeactivateOrUnloadCells)
+			else if (!bCanServerDeactivateOrUnloadCells)
 			{
 				// When server streaming-out is disabled, revisit existing loaded/activated cells and add them in the proper FrameLoadCells/FrameActivateCells
 				for (const UWorldPartitionRuntimeCell* Cell : ActivatedCells.GetCells())
 				{
-					AddServerFrameCells(Cell);
+					AddServerFrameCell(Cell);
 				}
 				for (const UWorldPartitionRuntimeCell* Cell : LoadedCells)
 				{
-					AddServerFrameCells(Cell);
+					AddServerFrameCell(Cell);
 				}
 			}
 		}
diff --git a/Engine/Source/Runtime/Engine/Public/WorldPartition/WorldPartition.h b/Engine/Source/Runtime/Engine/Public/WorldPartition/WorldPartition.h
--- a/Engine/Source/Runtime/Engine/Public/WorldPartition/WorldPartition.h
+++ b/Engine/Source/Runtime/Engine/Public/WorldPartition/WorldPartition.h
@@ -250,6 +250,7 @@ public:
 	bool IsServer() const;
 	bool IsServerStreamingEnabled() const;
 	bool IsServerStreamingOutEnabled() const;
+	bool IsServerDataLayerStreamingOutEnabled() const;
 	bool UseMakingVisibleTransactionRequests() const;
 	bool UseMakingInvisibleTransactionRequests() const;
 
@@ -354,6 +355,7 @@ private:
 	mutable TOptional<bool> bCachedUseMakingVisibleTransactionRequests;
 	mutable TOptional<bool> bCachedIsServerStreamingEnabled;
 	mutable TOptional<bool> bCachedIsServerStreamingOutEnabled;
+	mutable TOptional<bool> bCachedIsServerDataLayerStreamingOutEnabled;
 
 	UPROPERTY(Transient)
 	mutable TObjectPtr<UWorldPartitionStreamingPolicy> StreamingPolicy;
@@ -377,10 +379,12 @@ private:
 
 	static int32 EnableServerStreaming;
 	static bool bEnableServerStreamingOut;
+	static bool bEnableServerDataLayerStreamingOut;
 	static bool bUseMakingVisibleTransactionRequests;
 	static bool bUseMakingInvisibleTransactionRequests;
 	static FAutoConsoleVariableRef CVarEnableServerStreaming;
 	static FAutoConsoleVariableRef CVarEnableServerStreamingOut;
+	static FAutoConsoleVariableRef CVarEnableServerDataLayerStreamingOut;
 	static FAutoConsoleVariableRef CVarUseMakingVisibleTransactionRequests;
 	static FAutoConsoleVariableRef CVarUseMakingInvisibleTransactionRequests;
 
