World Partition: components can now force their owning actor to be non-spatially loaded, currently used by directional light components.

diff --git a/Engine/Source/Runtime/Engine/Classes/Components/ActorComponent.h b/Engine/Source/Runtime/Engine/Classes/Components/ActorComponent.h
--- a/Engine/Source/Runtime/Engine/Classes/Components/ActorComponent.h
+++ b/Engine/Source/Runtime/Engine/Classes/Components/ActorComponent.h
@@ -1034,6 +1034,7 @@ public:
 	virtual bool IsSelectedInEditor() const override;
 	virtual void SetPackageExternal(bool bExternal, bool bShouldDirty) {}
 	virtual FBox GetStreamingBounds() const { return FBox(ForceInit); }
+	virtual bool ForceActorNonSpatiallyLoaded() const { return false; }
 #endif // WITH_EDITOR
 	//~ End UObject Interface.
 
diff --git a/Engine/Source/Runtime/Engine/Classes/Components/DirectionalLightComponent.h b/Engine/Source/Runtime/Engine/Classes/Components/DirectionalLightComponent.h
--- a/Engine/Source/Runtime/Engine/Classes/Components/DirectionalLightComponent.h
+++ b/Engine/Source/Runtime/Engine/Classes/Components/DirectionalLightComponent.h
@@ -344,6 +344,7 @@ class ENGINE_API UDirectionalLightComponent : public ULightComponent
 #if WITH_EDITOR
 	virtual void PostEditChangeProperty(FPropertyChangedEvent& PropertyChangedEvent) override;
 	virtual bool CanEditChange(const FProperty* InProperty) const override;
+	virtual bool ForceActorNonSpatiallyLoaded() const override { return true; }
 #endif // WITH_EDITOR
 	virtual void Serialize(FArchive& Ar) override;
 	//~ Begin UObject Interface
diff --git a/Engine/Source/Runtime/Engine/Classes/Engine/PostProcessVolume.h b/Engine/Source/Runtime/Engine/Classes/Engine/PostProcessVolume.h
--- a/Engine/Source/Runtime/Engine/Classes/Engine/PostProcessVolume.h
+++ b/Engine/Source/Runtime/Engine/Classes/Engine/PostProcessVolume.h
@@ -76,7 +76,7 @@ class ENGINE_API APostProcessVolume : public AVolume, public IInterface_PostProc
 	virtual void PostUnregisterAllComponents( void ) override;
 
 #if WITH_EDITOR
-	virtual bool CanChangeIsSpatiallyLoadedFlag() const override { return !bUnbound; }
+	virtual bool CanChangeIsSpatiallyLoadedFlag() const override { return bUnbound ? false : Super::CanChangeIsSpatiallyLoadedFlag(); }
 #endif
 
 protected:
diff --git a/Engine/Source/Runtime/Engine/Classes/GameFramework/Actor.h b/Engine/Source/Runtime/Engine/Classes/GameFramework/Actor.h
--- a/Engine/Source/Runtime/Engine/Classes/GameFramework/Actor.h
+++ b/Engine/Source/Runtime/Engine/Classes/GameFramework/Actor.h
@@ -2315,7 +2315,7 @@ public:
 	}
 
 	/** Returns true if this actor is spatially loaded. */
-	bool GetIsSpatiallyLoaded() const { return bIsSpatiallyLoaded; }
+	bool GetIsSpatiallyLoaded() const;
 	
 	/** Set if this actor should be spatially loaded or not. */
 	void SetIsSpatiallyLoaded(bool bInIsSpatiallyLoaded)
@@ -2328,7 +2328,7 @@ public:
 	}
 
 	/** Returns true if this actor allows changing the spatially loaded flag.  */
-	virtual bool CanChangeIsSpatiallyLoadedFlag() const { return true; }
+	virtual bool CanChangeIsSpatiallyLoadedFlag() const;
 
 	/** Gets the property name for bIsSpatiallyLoaded. */
 	static const FName GetIsSpatiallyLoadedPropertyName() { return GET_MEMBER_NAME_CHECKED(AActor, bIsSpatiallyLoaded); }
diff --git a/Engine/Source/Runtime/Engine/Private/Actor.cpp b/Engine/Source/Runtime/Engine/Private/Actor.cpp
--- a/Engine/Source/Runtime/Engine/Private/Actor.cpp
+++ b/Engine/Source/Runtime/Engine/Private/Actor.cpp
@@ -1076,53 +1076,90 @@ void AActor::ProcessEvent(UFunction* Function, void* Parameters)
 }
 
 #if WITH_EDITOR
-FBox AActor::GetStreamingBounds() const
+static bool IsComponentStreamingRelevant(const AActor* InActor, const UActorComponent* InComponent, FBox& OutStreamingBounds)
 {
-	FBox StreamingBounds(ForceInit);
+	check(InActor);
+	check(InComponent);
 
-	auto HandleComponent = [this, &StreamingBounds](const UActorComponent* Component)
+	if (!InComponent->IsRegistered())
 	{
-		check(Component);
+		return false;
+	}
 
-		if (!Component->IsRegistered())
-		{
-			return;
-		}
+	if (InComponent->HasAnyFlags(RF_Transient))
+	{
+		return false;
+	}
 
-		if (Component->HasAnyFlags(RF_Transient))
-		{
-			return;
-		}
+	if (!InActor->IsEditorOnly() && InComponent->IsEditorOnly())
+	{
+		return false;
+	}
 
-		if (!IsEditorOnly() && Component->IsEditorOnly())
-		{
-			return;
-		}
+	OutStreamingBounds = InComponent->GetStreamingBounds();
+	return !!OutStreamingBounds.IsValid;
+}
+
+template <class F>
+static bool ForEachStreamingRelevantComponent(const AActor* InActor, F Func)
+{
+	bool bHasStreamingRelevantComponents = false;
 
-		const FBox ComponentStreamingBound = Component->GetStreamingBounds();
-		if (!ComponentStreamingBound.IsValid)
+	auto HandleComponent = [InActor, &bHasStreamingRelevantComponents, &Func](const UActorComponent* Component)
+	{
+		FBox ComponentStreamingBound;
+		if (IsComponentStreamingRelevant(InActor, Component, ComponentStreamingBound))
 		{
-			return;
+			Func(Component, ComponentStreamingBound);
+			bHasStreamingRelevantComponents = true;
 		}
-
-		StreamingBounds += ComponentStreamingBound;
 	};
 
-	ForEachComponent<UPrimitiveComponent>(true, [this, &HandleComponent](UActorComponent* Component)
+	InActor->ForEachComponent<UPrimitiveComponent>(true, [&HandleComponent](UActorComponent* Component)
 	{
 		HandleComponent(Component);
 	});
 
-	if (!StreamingBounds.IsValid)
+	if (!bHasStreamingRelevantComponents)
 	{
-		ForEachComponent<UActorComponent>(false, [this, &HandleComponent](UActorComponent* Component)
+		InActor->ForEachComponent<UActorComponent>(false, [&HandleComponent](UActorComponent* Component)
 		{
 			HandleComponent(Component);
 		});
 	}
 
+	return bHasStreamingRelevantComponents;
+}
+
+static bool HasComponentForceActorNonSpatiallyLoaded(const AActor* InActor)
+{
+	bool bHasComponentForceActorNonSpatiallyLoaded = false;
+	ForEachStreamingRelevantComponent(InActor, [&bHasComponentForceActorNonSpatiallyLoaded](const UActorComponent* Component, const FBox& StreamingBound)
+	{
+		bHasComponentForceActorNonSpatiallyLoaded |= Component->ForceActorNonSpatiallyLoaded();
+	});
+	return bHasComponentForceActorNonSpatiallyLoaded;
+}
+
+FBox AActor::GetStreamingBounds() const
+{
+	FBox StreamingBounds(ForceInit);
+	ForEachStreamingRelevantComponent(this, [&StreamingBounds](const UActorComponent* Component, const FBox& StreamingBound)
+	{
+		StreamingBounds += StreamingBound;
+	});
 	return StreamingBounds;
 }
+
+bool AActor::GetIsSpatiallyLoaded() const
+{
+	return bIsSpatiallyLoaded && !HasComponentForceActorNonSpatiallyLoaded(this);
+}
+	
+bool AActor::CanChangeIsSpatiallyLoadedFlag() const
+{
+	return !HasComponentForceActorNonSpatiallyLoaded(this);
+}
 #endif
 
 void AActor::ApplyWorldOffset(const FVector& InOffset, bool bWorldShift)
diff --git a/Engine/Source/Runtime/Landscape/Classes/LandscapeSplineActor.h b/Engine/Source/Runtime/Landscape/Classes/LandscapeSplineActor.h
--- a/Engine/Source/Runtime/Landscape/Classes/LandscapeSplineActor.h
+++ b/Engine/Source/Runtime/Landscape/Classes/LandscapeSplineActor.h
@@ -35,7 +35,6 @@ public:
 #if WITH_EDITOR
 	LANDSCAPE_API bool HasGeneratedLandscapeSplineMeshesActors() const;
 	void SetLandscapeGuid(const FGuid& InGuid);
-	virtual bool CanChangeIsSpatiallyLoadedFlag() const override { return true; }
 	virtual void GetActorDescProperties(FPropertyPairsMap& PropertyPairsMap) const override;
 
 	void GetSharedProperties(ULandscapeInfo* InLandscapeInfo);
diff --git a/Engine/Source/Runtime/Landscape/Classes/LandscapeStreamingProxy.h b/Engine/Source/Runtime/Landscape/Classes/LandscapeStreamingProxy.h
--- a/Engine/Source/Runtime/Landscape/Classes/LandscapeStreamingProxy.h
+++ b/Engine/Source/Runtime/Landscape/Classes/LandscapeStreamingProxy.h
@@ -35,7 +35,7 @@ public:
 	virtual AActor* GetSceneOutlinerParent() const override;
 	virtual bool CanDeleteSelectedActor(FText& OutReason) const override;
 	virtual bool GetReferencedContentObjects(TArray<UObject*>& Objects) const override;
-	virtual bool CanChangeIsSpatiallyLoadedFlag() const override { return true; }
+	virtual bool CanChangeIsSpatiallyLoadedFlag() const override { return AActor::CanChangeIsSpatiallyLoadedFlag(); }
 	virtual bool ShouldIncludeGridSizeInName(UWorld* InWorld, const FActorPartitionIdentifier& InIdentifier) const override;
 #endif
 	//~ End UObject Interface
