Fix bad merge from Main.

#prelight none

diff --git a/Engine/Source/Runtime/Engine/Private/WorldPartition/HLOD/HLODActorDesc.cpp b/Engine/Source/Runtime/Engine/Private/WorldPartition/HLOD/HLODActorDesc.cpp
--- a/Engine/Source/Runtime/Engine/Private/WorldPartition/HLOD/HLODActorDesc.cpp
+++ b/Engine/Source/Runtime/Engine/Private/WorldPartition/HLOD/HLODActorDesc.cpp
@@ -66,7 +66,7 @@ void FHLODActorDesc::Serialize(FArchive& Ar)
 
 	if (Ar.CustomVer(FFortniteMainBranchObjectVersion::GUID) >= FFortniteMainBranchObjectVersion::WorldPartitionHLODActorDescSerializeStats)
 	{
-		if (Ar.CustomVer(FUE5MainStreamObjectVersion::GUID) < FFortniteMainBranchObjectVersion::WorldPartitionHLODActorUseSourceCellGuid)
+		if (Ar.CustomVer(FFortniteMainBranchObjectVersion::GUID) < FFortniteMainBranchObjectVersion::WorldPartitionHLODActorUseSourceCellGuid)
 		{
 			FName SourceCellName;
 			Ar << SourceCellName;
