PR #9902: Making LevelInstanceActorGuid and LevelInstanceActorImpl protected (Contributed by [REDACTED_PERSON])

diff --git a/Engine/Source/Runtime/Engine/Public/LevelInstance/LevelInstanceActor.h b/Engine/Source/Runtime/Engine/Public/LevelInstance/LevelInstanceActor.h
--- a/Engine/Source/Runtime/Engine/Public/LevelInstance/LevelInstanceActor.h
+++ b/Engine/Source/Runtime/Engine/Public/LevelInstance/LevelInstanceActor.h
@@ -34,11 +34,10 @@ protected:
 	UPROPERTY(Transient, ReplicatedUsing=OnRep_LevelInstanceSpawnGuid)
 	FGuid LevelInstanceSpawnGuid;
 
-private:
 	FLevelInstanceActorGuid LevelInstanceActorGuid;
 	FLevelInstanceActorImpl LevelInstanceActorImpl;
-public:
 
+public:
 #if WITH_EDITORONLY_DATA
 	UPROPERTY()
 	ELevelInstanceRuntimeBehavior DesiredRuntimeBehavior;
