Fix Win64 Clang editor compile error

#rnx

diff --git a/Engine/Source/Runtime/Engine/Private/StaticMesh.cpp b/Engine/Source/Runtime/Engine/Private/StaticMesh.cpp
--- a/Engine/Source/Runtime/Engine/Private/StaticMesh.cpp
+++ b/Engine/Source/Runtime/Engine/Private/StaticMesh.cpp
@@ -8233,7 +8233,7 @@ ENGINE_API int32 UStaticMesh::GetMinLODIdx(bool bForceLowestLODIdx) const
 	{
 		int32 CurrentMinLodQualityLevel = GMinLodQualityLevel;
 #if PLATFORM_DESKTOP
-		extern int32 GUseMobileLODBiasOnDesktopES31;
+		extern ENGINE_API int32 GUseMobileLODBiasOnDesktopES31;
 		if (GUseMobileLODBiasOnDesktopES31 != 0 && GMaxRHIFeatureLevel == ERHIFeatureLevel::ES3_1)
 		{
 			CurrentMinLodQualityLevel = (int32)EPerQualityLevels::Low;
diff --git a/Engine/Source/Runtime/Engine/Private/StaticMeshRender.cpp b/Engine/Source/Runtime/Engine/Private/StaticMeshRender.cpp
--- a/Engine/Source/Runtime/Engine/Private/StaticMeshRender.cpp
+++ b/Engine/Source/Runtime/Engine/Private/StaticMeshRender.cpp
@@ -267,7 +267,7 @@ FStaticMeshSceneProxy::FStaticMeshSceneProxy(UStaticMeshComponent* InComponent,
 #endif
 
 #if PLATFORM_DESKTOP
-	extern int32 GUseMobileLODBiasOnDesktopES31;
+	extern ENGINE_API int32 GUseMobileLODBiasOnDesktopES31;
 	if (GUseMobileLODBiasOnDesktopES31 != 0 && FeatureLevel == ERHIFeatureLevel::ES3_1)
 	{
 		EffectiveMinLOD += InComponent->GetStaticMesh()->GetRenderData()->LODBiasModifier;
