Dependency cleanup in Materials headers.
- Removing header dependencies on heavy files by moving some code from the headers to their respective source files
- Removing a couple deprecated functions also helped.
- Removed implicit includes will be added to files that depend on them before this is checked in.

diff --git a/Engine/Source/Runtime/Engine/Classes/Materials/Material.h b/Engine/Source/Runtime/Engine/Classes/Materials/Material.h
--- a/Engine/Source/Runtime/Engine/Classes/Materials/Material.h
+++ b/Engine/Source/Runtime/Engine/Classes/Materials/Material.h
@@ -11,9 +11,13 @@
 #include "Materials/MaterialInterface.h"
 #include "MaterialShared.h"
 #include "MaterialExpressionIO.h"
+#if UE_ENABLE_INCLUDE_ORDER_DEPRECATED_IN_5_2
 #include "Materials/MaterialExpressionMaterialFunctionCall.h"
 #include "Materials/MaterialExpressionMaterialAttributeLayers.h"
 #include "Materials/MaterialFunction.h"
+#endif
+#include "Materials/MaterialExpression.h"
+#include "Materials/MaterialFunctionInterface.h"
 #include "Materials/MaterialLayersFunctions.h"
 #include "Materials/MaterialOverrideNanite.h"
 #include "Templates/UniquePtr.h"
@@ -1282,83 +1286,11 @@ public:
 #endif
 	ENGINE_API virtual bool IsComplete() const override;
 
-#if WITH_EDITORONLY_DATA
-	/**
-	 * @param	OutParameterNames		Storage array for the parameter names.
-	 * @param	OutParameterIds			Storage array for the parameter id's.
-	 *
-	 * @return	Returns a array of parameter names used in this material for the specified expression type.
-	 */
-	template<typename ExpressionType>
-	UE_DEPRECATED(5.0, "Use GetAllParameterInfoOfType or GetAllParametersOfType")
-	void GetAllParameterInfo(TArray<FMaterialParameterInfo>& OutParameterInfo, TArray<FGuid>& OutParameterIds) const
-	{
-		for (const TObjectPtr<UMaterialExpression>& Expression : GetExpressions())
-		{
-			FMaterialParameterInfo BaseParameterInfo;
-			BaseParameterInfo.Association = EMaterialParameterAssociation::GlobalParameter;
-			BaseParameterInfo.Index = INDEX_NONE;
-
-			// Note: Intentionally checking the requested type first as this catches MaterialLayers
-			// which are a top-level only parameter without having to deal with the below recursion
-			if (const ExpressionType* ParameterExpression = Cast<const ExpressionType>(Expression))
-			{
-				PRAGMA_DISABLE_DEPRECATION_WARNINGS
-				ParameterExpression->GetAllParameterInfo(OutParameterInfo, OutParameterIds, BaseParameterInfo);
-				PRAGMA_ENABLE_DEPRECATION_WARNINGS
-			}
-			else if (const UMaterialExpressionMaterialFunctionCall* FunctionExpression = Cast<const UMaterialExpressionMaterialFunctionCall>(Expression))
-			{
-				if (FunctionExpression->MaterialFunction)
-				{
-					PRAGMA_DISABLE_DEPRECATION_WARNINGS
-					FunctionExpression->MaterialFunction->GetAllParameterInfo<ExpressionType>(OutParameterInfo, OutParameterIds, BaseParameterInfo);
-					PRAGMA_ENABLE_DEPRECATION_WARNINGS
-				}
-			}
-			else if (const UMaterialExpressionMaterialAttributeLayers* LayersExpression = Cast<const UMaterialExpressionMaterialAttributeLayers>(Expression))
-			{
-				const TArray<UMaterialFunctionInterface*>* Layers = &LayersExpression->GetLayers();
-				const TArray<UMaterialFunctionInterface*>* Blends = &LayersExpression->GetBlends();
-
-				for (int32 LayerIndex = 0; LayerIndex < Layers->Num(); ++LayerIndex)
-				{
-					if (const UMaterialFunctionInterface* Layer = (*Layers)[LayerIndex])
-					{
-						BaseParameterInfo.Association = EMaterialParameterAssociation::LayerParameter;
-						BaseParameterInfo.Index = LayerIndex;
-						PRAGMA_DISABLE_DEPRECATION_WARNINGS
-						Layer->GetAllParameterInfo<ExpressionType>(OutParameterInfo, OutParameterIds, BaseParameterInfo);
-						PRAGMA_ENABLE_DEPRECATION_WARNINGS
-					}
-				}
-
-				for (int32 BlendIndex = 0; BlendIndex < Blends->Num(); ++BlendIndex)
-				{
-					if (const UMaterialFunctionInterface* Blend = (*Blends)[BlendIndex])
-					{
-						BaseParameterInfo.Association = EMaterialParameterAssociation::BlendParameter;
-						BaseParameterInfo.Index = BlendIndex;
-						PRAGMA_DISABLE_DEPRECATION_WARNINGS
-						Blend->GetAllParameterInfo<ExpressionType>(OutParameterInfo, OutParameterIds, BaseParameterInfo);
-						PRAGMA_ENABLE_DEPRECATION_WARNINGS
-					}
-				}
-			}
-
-			check(OutParameterInfo.Num() == OutParameterIds.Num());
-		}
-	}
-#endif // WITH_EDITORONLY_DATA
-
 #if WITH_EDITORONLY_DATA
 	ENGINE_API virtual bool IterateDependentFunctions(TFunctionRef<bool(UMaterialFunctionInterface*)> Predicate) const override;
 	ENGINE_API virtual void GetDependentFunctions(TArray<class UMaterialFunctionInterface*>& DependentFunctions) const override;
 #endif // WITH_EDITORONLY_DATA
 
-	UE_DEPRECATED(5.0, "No longer used.")
-	uint32 GetDecalBlendMode() const { return DecalBlendMode; }
-
 	/** Returns the material's decal response mode */
 	uint32 GetMaterialDecalResponse() const { return MaterialDecalResponse; }
 
@@ -1390,6 +1322,15 @@ public:
 		}
 	}
 
+	static ENGINE_API UMaterialFunctionInterface* GetExpressionFunctionPointer(const UMaterialExpression* Expression);
+
+	struct FLayersInterfaces
+	{
+		TConstArrayView<UMaterialFunctionInterface*> Layers;
+		TConstArrayView<UMaterialFunctionInterface*> Blends;
+	};
+	static ENGINE_API TOptional<FLayersInterfaces> GetExpressionLayers(const UMaterialExpression* Expression);
+
 	/** Get all expressions of the requested type, recursing through any function expressions in the material */
 	template<typename ExpressionType>
 	void GetAllExpressionsInMaterialAndFunctionsOfType(TArray<ExpressionType*>& OutExpressions) const
@@ -1402,19 +1343,13 @@ public:
 				OutExpressions.Add(ExpressionOfType);
 			}
 
-			UMaterialExpressionMaterialFunctionCall* ExpressionFunctionCall = Cast<UMaterialExpressionMaterialFunctionCall>(Expression);
-			UMaterialExpressionMaterialAttributeLayers* LayersExpression = Cast<UMaterialExpressionMaterialAttributeLayers>(Expression);
-
-			if (ExpressionFunctionCall && ExpressionFunctionCall->MaterialFunction)
+			if (UMaterialFunctionInterface* MaterialFunction = GetExpressionFunctionPointer(Expression))
 			{
-				ExpressionFunctionCall->MaterialFunction->GetAllExpressionsOfType<ExpressionType>(OutExpressions);
+				MaterialFunction->GetAllExpressionsOfType<ExpressionType>(OutExpressions);
 			}
-			else if (LayersExpression)
+			else if (TOptional<FLayersInterfaces> LayersInterfaces = GetExpressionLayers(Expression))
 			{
-				const TArray<UMaterialFunctionInterface*>& Layers = LayersExpression->GetLayers();
-				const TArray<UMaterialFunctionInterface*>& Blends = LayersExpression->GetBlends();
-
-				for (auto* Layer : Layers)
+				for (UMaterialFunctionInterface* Layer : LayersInterfaces->Layers)
 				{
 					if (Layer)
 					{
@@ -1422,7 +1357,7 @@ public:
 					}
 				}
 
-				for (auto* Blend : Blends)
+				for (UMaterialFunctionInterface* Blend : LayersInterfaces->Blends)
 				{
 					if (Blend)
 					{
@@ -1445,22 +1380,16 @@ public:
 				return true;
 			}
 
-			UMaterialExpressionMaterialFunctionCall* ExpressionFunctionCall = Cast<UMaterialExpressionMaterialFunctionCall>(Expression);
-			UMaterialExpressionMaterialAttributeLayers* LayersExpression = Cast<UMaterialExpressionMaterialAttributeLayers>(Expression);
-
-			if (ExpressionFunctionCall)
+			if (UMaterialFunctionInterface* MaterialFunction = GetExpressionFunctionPointer(Expression))
 			{
-				if (ExpressionFunctionCall->MaterialFunction && ExpressionFunctionCall->MaterialFunction->HasAnyExpressionsOfType<ExpressionType>())
+				if (MaterialFunction->HasAnyExpressionsOfType<ExpressionType>())
 				{
 					return true;
 				}
 			}
-			else if (LayersExpression)
+			else if (TOptional<FLayersInterfaces> LayersInterfaces = GetExpressionLayers(Expression))
 			{
-				const TArray<UMaterialFunctionInterface*>& Layers = LayersExpression->GetLayers();
-				const TArray<UMaterialFunctionInterface*>& Blends = LayersExpression->GetBlends();
-
-				for (auto* Layer : Layers)
+				for (auto* Layer : LayersInterfaces->Layers)
 				{
 					if (Layer && Layer->HasAnyExpressionsOfType<ExpressionType>())
 					{
@@ -1468,7 +1397,7 @@ public:
 					}
 				}
 
-				for (auto* Blend : Blends)
+				for (auto* Blend : LayersInterfaces->Blends)
 				{
 					if (Blend && Blend->HasAnyExpressionsOfType<ExpressionType>())
 					{
@@ -1858,27 +1787,22 @@ private:
 		for (int32 ExpressionIndex = 0; ExpressionIndex < InMaterialExpression.Num(); ExpressionIndex++)
 		{
 			UMaterialExpression* ExpressionPtr = InMaterialExpression[ExpressionIndex];
-			UMaterialExpressionMaterialFunctionCall* MaterialFunctionCall = Cast<UMaterialExpressionMaterialFunctionCall>(ExpressionPtr);
-			UMaterialExpressionMaterialAttributeLayers* MaterialLayers = Cast<UMaterialExpressionMaterialAttributeLayers>(ExpressionPtr);
 
 			if (ExpressionPtr && ExpressionPtr->GetParameterExpressionId() == InGUID)
 			{
 				check(ExpressionPtr->bIsParameterExpression);
 				return Cast<ExpressionType>(ExpressionPtr);
 			}
-			else if (MaterialFunctionCall && MaterialFunctionCall->MaterialFunction)
+			else if (UMaterialFunctionInterface* MaterialFunction = GetExpressionFunctionPointer(ExpressionPtr))
 			{
-				if (ExpressionType* Expression = FindExpressionByGUIDRecursive<ExpressionType>(InGUID, MaterialFunctionCall->MaterialFunction->GetExpressions()))
+				if (ExpressionType* Expression = FindExpressionByGUIDRecursive<ExpressionType>(InGUID, MaterialFunction->GetExpressions()))
 				{
 					return Expression;
 				}
 			}
-			else if (MaterialLayers)
+			else if (TOptional<FLayersInterfaces> LayersInterfaces = UMaterial::GetExpressionLayers(ExpressionPtr))
 			{
-				const TArray<UMaterialFunctionInterface*>& Layers = MaterialLayers->GetLayers();
-				const TArray<UMaterialFunctionInterface*>& Blends = MaterialLayers->GetBlends();
-
-				for (const auto* Layer : Layers)
+				for (const auto* Layer : LayersInterfaces->Layers)
 				{
 					if (Layer)
 					{
@@ -1889,7 +1813,7 @@ private:
 					}
 				}
 
-				for (const auto* Blend : Blends)
+				for (const auto* Blend : LayersInterfaces->Blends)
 				{
 					if (Blend)
 					{
diff --git a/Engine/Source/Runtime/Engine/Classes/Materials/MaterialExpression.h b/Engine/Source/Runtime/Engine/Classes/Materials/MaterialExpression.h
--- a/Engine/Source/Runtime/Engine/Classes/Materials/MaterialExpression.h
+++ b/Engine/Source/Runtime/Engine/Classes/Materials/MaterialExpression.h
@@ -7,7 +7,9 @@
 #include "UObject/Object.h"
 #include "UObject/UnrealType.h"
 #include "Misc/Guid.h"
+#if UE_ENABLE_INCLUDE_ORDER_DEPRECATED_IN_5_2
 #include "MaterialShared.h"
+#endif
 #include "MaterialExpressionIO.h"
 
 #include "MaterialExpression.generated.h"
@@ -16,7 +18,9 @@ class UEdGraphNode;
 class UMaterial;
 class UTexture;
 struct FPropertyChangedEvent;
+struct FMaterialParameterMetadata;
 struct FMaterialShadingModelField;
+struct FStrataMaterialInfo;
 struct FStrataOperator;
 
 class UMaterialExpression;
@@ -26,6 +30,7 @@ class UMaterialExpressionExecEnd;
 
 class FMaterialHLSLGenerator;
 enum class EMaterialNewScopeFlag : uint8;
+enum class EMaterialParameterType : uint8;
 
 namespace UE
 {
diff --git a/Engine/Source/Runtime/Engine/Classes/Materials/MaterialFunctionInstance.h b/Engine/Source/Runtime/Engine/Classes/Materials/MaterialFunctionInstance.h
--- a/Engine/Source/Runtime/Engine/Classes/Materials/MaterialFunctionInstance.h
+++ b/Engine/Source/Runtime/Engine/Classes/Materials/MaterialFunctionInstance.h
@@ -17,18 +17,9 @@ class UMaterialFunctionInstance : public UMaterialFunctionInterface
 {
 	GENERATED_UCLASS_BODY()
 
-	void SetParent(UMaterialFunctionInterface* NewParent)
-	{
-		Parent = NewParent;
-		MaterialFunctionUsage = NewParent->GetMaterialFunctionUsage();
-		Base = GetBaseFunction();
-	}
-
-	virtual EMaterialFunctionUsage GetMaterialFunctionUsage() override
-	{
-		UMaterialFunctionInterface* BaseFunction = GetBaseFunction();
-		return BaseFunction ? BaseFunction->GetMaterialFunctionUsage() : EMaterialFunctionUsage::Default;
-	}
+	ENGINE_API void SetParent(UMaterialFunctionInterface* NewParent);
+
+	virtual EMaterialFunctionUsage GetMaterialFunctionUsage() override;
 
 	/** Parent function. */
 	UPROPERTY(EditAnywhere, BlueprintReadOnly, Category=MaterialFunctionInstance, AssetRegistrySearchable)
diff --git a/Engine/Source/Runtime/Engine/Classes/Materials/MaterialFunctionInterface.h b/Engine/Source/Runtime/Engine/Classes/Materials/MaterialFunctionInterface.h
--- a/Engine/Source/Runtime/Engine/Classes/Materials/MaterialFunctionInterface.h
+++ b/Engine/Source/Runtime/Engine/Classes/Materials/MaterialFunctionInterface.h
@@ -7,10 +7,13 @@
 #include "UObject/Object.h"
 #include "Misc/Guid.h"
 #include "Templates/Casts.h"
+#include "MaterialRecursionGuard.h"
+#if UE_ENABLE_INCLUDE_ORDER_DEPRECATED_IN_5_2
 #include "Materials/MaterialExpressionMaterialFunctionCall.h"
 #include "Materials/MaterialExpressionFontSampleParameter.h"
 #include "Materials/MaterialExpressionParameter.h"
 #include "Materials/MaterialExpressionTextureSampleParameter.h"
+#endif
 #include "StaticParameterSet.h"
 #include "MaterialFunctionInterface.generated.h"
 
@@ -19,6 +22,10 @@ class UMaterialFunction;
 class UTexture;
 struct FPropertyChangedEvent;
 class FMaterialHLSLGenerator;
+class FMaterialUpdateContext;
+class UMaterialInterface;
+class UMaterialExpression;
+struct FFunctionExpressionOutput;
 
 /** Usage set on a material function determines feature compatibility and validation. */
 UENUM()
@@ -182,38 +189,6 @@ public:
 #endif // WITH_EDITOR
 
 public:
-#if WITH_EDITORONLY_DATA
-	/** Finds the names of all matching type parameters */
-	template<typename ExpressionType>
-	UE_DEPRECATED(5.0, "Use GetAllParameterInfoOfType or GetAllParametersOfType")
-	void GetAllParameterInfo(TArray<FMaterialParameterInfo>& OutParameterInfo, TArray<FGuid>& OutParameterIds, const FMaterialParameterInfo& InBaseParameterInfo) const
-	{
-		if (const UMaterialFunctionInterface* ParameterFunction = GetBaseFunctionInterface())
-		{
-			const UClass* TargetClass = UMaterialExpressionMaterialFunctionCall::StaticClass();
-			for (const TObjectPtr<UMaterialExpression>& Expression : ParameterFunction->GetExpressions())
-			{
-				if (const UMaterialExpressionMaterialFunctionCall* FunctionExpression = (Expression && Expression.IsA(TargetClass)) ? (const UMaterialExpressionMaterialFunctionCall*)Expression.Get() : nullptr)
-				{
-					if (FunctionExpression->MaterialFunction)
-					{
-						PRAGMA_DISABLE_DEPRECATION_WARNINGS
-						FunctionExpression->MaterialFunction->GetAllParameterInfo<const ExpressionType>(OutParameterInfo, OutParameterIds, InBaseParameterInfo);
-						PRAGMA_ENABLE_DEPRECATION_WARNINGS
-					}
-				}
-				else if (const ExpressionType* ParameterExpression = Cast<const ExpressionType>(Expression))
-				{
-					PRAGMA_DISABLE_DEPRECATION_WARNINGS
-					ParameterExpression->GetAllParameterInfo(OutParameterInfo, OutParameterIds, InBaseParameterInfo);
-					PRAGMA_ENABLE_DEPRECATION_WARNINGS
-				}
-			}
-
-			check(OutParameterInfo.Num() == OutParameterIds.Num());
-		}
-	}
-#endif // WITH_EDITORONLY_DATA
 
 #if WITH_EDITOR
 	/** Finds the first matching parameter by name and type */
diff --git a/Engine/Source/Runtime/Engine/Private/Materials/Material.cpp b/Engine/Source/Runtime/Engine/Private/Materials/Material.cpp
--- a/Engine/Source/Runtime/Engine/Private/Materials/Material.cpp
+++ b/Engine/Source/Runtime/Engine/Private/Materials/Material.cpp
@@ -20,6 +20,7 @@
 #include "Materials/MaterialInstance.h"
 #include "Materials/MaterialInstanceDynamic.h"
 #include "UnrealEngine.h"
+
 #include "Materials/MaterialExpressionCollectionParameter.h"
 #include "Materials/MaterialExpressionCustomOutput.h"
 #include "Materials/MaterialExpressionFunctionOutput.h"
@@ -29,11 +30,14 @@
 #include "Materials/MaterialExpressionShadingPathSwitch.h"
 #include "Materials/MaterialExpressionShaderStageSwitch.h"
 #include "Materials/MaterialExpressionMakeMaterialAttributes.h"
+#include "Materials/MaterialExpressionMaterialAttributeLayers.h"
+#include "Materials/MaterialExpressionMaterialFunctionCall.h"
 #include "Materials/MaterialExpressionSetMaterialAttributes.h"
 #include "Materials/MaterialExpressionRuntimeVirtualTextureOutput.h"
 #include "Materials/MaterialExpressionStaticSwitchParameter.h"
 #include "Materials/MaterialExpressionTextureBase.h"
 #include "Materials/MaterialExpressionTextureSample.h"
+#include "Materials/MaterialExpressionTextureSampleParameter.h"
 #include "Materials/MaterialExpressionVertexInterpolator.h"
 #include "Materials/MaterialExpressionSceneColor.h"
 #include "Materials/MaterialExpressionShadingModel.h"
@@ -50,6 +54,7 @@
 #include "Materials/MaterialExpressionSaturate.h"
 #include "Materials/MaterialExpressionConstant3Vector.h"
 #include "Materials/MaterialExpressionPower.h"
+
 #include "Engine/Font.h"
 #include "SceneManagement.h"
 #include "Materials/MaterialUniformExpressions.h"
@@ -1972,6 +1977,29 @@ void UMaterial::GetMaterialInheritanceChain(FMaterialInheritanceChain& OutChain)
 
 #if WITH_EDITORONLY_DATA
 
+UMaterialFunctionInterface* UMaterial::GetExpressionFunctionPointer(const UMaterialExpression* Expression)
+{
+	const UMaterialExpressionMaterialFunctionCall* ExpressionFunctionCall = Cast<UMaterialExpressionMaterialFunctionCall>(Expression);
+	if (ExpressionFunctionCall && ExpressionFunctionCall->MaterialFunction)
+	{
+		return ExpressionFunctionCall->MaterialFunction;
+	}
+	return nullptr;
+}
+
+TOptional<UMaterial::FLayersInterfaces> UMaterial::GetExpressionLayers(const UMaterialExpression* Expression)
+{
+	const UMaterialExpressionMaterialAttributeLayers* LayersExpression = Cast<UMaterialExpressionMaterialAttributeLayers>(Expression);
+	if (LayersExpression)
+	{
+		FLayersInterfaces Interfaces;
+		Interfaces.Layers = LayersExpression->GetLayers();
+		Interfaces.Blends = LayersExpression->GetBlends();
+		return Interfaces;
+	}
+	return {};
+}
+
 void UMaterial::UpdateTransientExpressionData()
 {
 	for (UMaterialExpression* Expression : GetExpressions())
diff --git a/Engine/Source/Runtime/Engine/Private/Materials/MaterialExpressions.cpp b/Engine/Source/Runtime/Engine/Private/Materials/MaterialExpressions.cpp
--- a/Engine/Source/Runtime/Engine/Private/Materials/MaterialExpressions.cpp
+++ b/Engine/Source/Runtime/Engine/Private/Materials/MaterialExpressions.cpp
@@ -14717,6 +14717,19 @@ UMaterialFunctionInstance::UMaterialFunctionInstance(const FObjectInitializer& O
 #endif
 }
 
+void UMaterialFunctionInstance::SetParent(UMaterialFunctionInterface* NewParent)
+{
+	Parent = NewParent;
+	MaterialFunctionUsage = NewParent->GetMaterialFunctionUsage();
+	Base = GetBaseFunction();
+}
+
+EMaterialFunctionUsage UMaterialFunctionInstance::GetMaterialFunctionUsage()
+{
+	UMaterialFunctionInterface* BaseFunction = GetBaseFunction();
+	return BaseFunction ? BaseFunction->GetMaterialFunctionUsage() : EMaterialFunctionUsage::Default;
+}
+
 UMaterialFunction* UMaterialFunctionInstance::GetBaseFunction(FMFRecursionGuard RecursionGuard)
 {
 	if (!Parent || RecursionGuard.Contains(this))
diff --git a/Engine/Source/Runtime/Engine/Private/Materials/MaterialShared.cpp b/Engine/Source/Runtime/Engine/Private/Materials/MaterialShared.cpp
--- a/Engine/Source/Runtime/Engine/Private/Materials/MaterialShared.cpp
+++ b/Engine/Source/Runtime/Engine/Private/Materials/MaterialShared.cpp
@@ -31,6 +31,7 @@
 #include "MaterialCompiler.h"
 #include "MeshMaterialShaderType.h"
 #include "RendererInterface.h"
+#include "RenderingThread.h"
 #include "Materials/HLSLMaterialTranslator.h"
 #include "ComponentRecreateRenderStateContext.h"
 #include "EngineModule.h"
@@ -2206,6 +2207,23 @@ void FMaterial::DeferredDelete(FMaterial* InMaterial)
 	}
 }
 
+void FMaterial::DeleteMaterialsOnRenderThread(TArray<TRefCountPtr<FMaterial>>& MaterialsRenderThread)
+{
+	if (MaterialsRenderThread.Num() > 0)
+	{
+		ENQUEUE_RENDER_COMMAND(DeferredDestroyMaterialArray)([MaterialsRenderThread = MoveTemp(MaterialsRenderThread)](FRHICommandListImmediate& RHICmdList) mutable
+		{
+			for (TRefCountPtr<FMaterial>& Material : MaterialsRenderThread)
+			{
+				FMaterial* MaterialToDestroy = Material.GetReference();
+				MaterialToDestroy->PrepareDestroy_RenderThread();
+				Material.SafeRelease();
+				delete MaterialToDestroy;
+			}
+		});
+	}
+}
+
 /**
  * Destructor
  */
diff --git a/Engine/Source/Runtime/Engine/Public/MaterialShared.h b/Engine/Source/Runtime/Engine/Public/MaterialShared.h
--- a/Engine/Source/Runtime/Engine/Public/MaterialShared.h
+++ b/Engine/Source/Runtime/Engine/Public/MaterialShared.h
@@ -1605,35 +1605,19 @@ public:
 
 	ENGINE_API static void DeferredDelete(FMaterial* Material);
 
-	template<typename TMaterial>
-	static void DeleteMaterialsOnRenderThread(TArray<TRefCountPtr<TMaterial>>& MaterialsRenderThread)
-	{
-		if (MaterialsRenderThread.Num() > 0)
-		{
-			ENQUEUE_RENDER_COMMAND(DeferredDestroyMaterialArray)([MaterialsRenderThread = MoveTemp(MaterialsRenderThread)](FRHICommandListImmediate& RHICmdList) mutable
-			{
-				for (auto& Material : MaterialsRenderThread)
-				{
-					TMaterial* MaterialToDestroy = Material.GetReference();
-					MaterialToDestroy->PrepareDestroy_RenderThread();
-					Material.SafeRelease();
-					delete MaterialToDestroy;
-				}
-			});
-		}
-	}
+	ENGINE_API static void DeleteMaterialsOnRenderThread(TArray<TRefCountPtr<FMaterial>>& MaterialsRenderThread);
 
 	template<typename TMaterial>
 	static void DeferredDeleteArray(TArray<TRefCountPtr<TMaterial>>& Materials)
 	{
 		if (Materials.Num() > 0)
 		{
-			TArray<TRefCountPtr<TMaterial>> MaterialsRenderThread;
+			TArray<TRefCountPtr<FMaterial>> MaterialsRenderThread;
 			for (TRefCountPtr<TMaterial>& Material : Materials)
 			{
-				TMaterial* MaterialToDestroy = Material.GetReference();
+				FMaterial* MaterialToDestroy = Material.GetReference();
 				MaterialToDestroy->PrepareDestroy_GameThread();
-				MaterialsRenderThread.Add(MoveTemp(Material));
+				MaterialsRenderThread.Emplace(MoveTemp(Material));
 			}
 
 			Materials.Empty();
@@ -1646,11 +1630,11 @@ public:
 	{
 		if (Materials.Num() > 0)
 		{
-			TArray<TRefCountPtr<TMaterial>> MaterialsRenderThread;
-			for (TMaterial* Material : Materials)
+			TArray<TRefCountPtr<FMaterial>> MaterialsRenderThread;
+			for (FMaterial* Material : Materials)
 			{
 				Material->PrepareDestroy_GameThread();
-				MaterialsRenderThread.Add(Material);
+				MaterialsRenderThread.Emplace(Material);
 			}
 
 			Materials.Empty();
