Removing the deprecated GetAllParameterInfo to fix compile errors with the old include order.

diff --git a/Engine/Source/Runtime/Engine/Classes/Materials/MaterialInstance.h b/Engine/Source/Runtime/Engine/Classes/Materials/MaterialInstance.h
--- a/Engine/Source/Runtime/Engine/Classes/Materials/MaterialInstance.h
+++ b/Engine/Source/Runtime/Engine/Classes/Materials/MaterialInstance.h
@@ -890,24 +890,6 @@ public:
 	 * Builds a composited set of static parameters, including inherited and overridden values
 	 */
 	ENGINE_API void GetStaticParameterValues(FStaticParameterSet& OutStaticParameters);
-
-	/**
-	 * Builds a composited set of parameter names, including inherited and overridden values
-	 */
-#if UE_ENABLE_INCLUDE_ORDER_DEPRECATED_IN_5_2
-	template<typename ExpressionType>
-	UE_DEPRECATED(5.0, "Use GetAllParameterInfoOfType or GetAllParametersOfType")
-	void GetAllParameterInfo(TArray<FMaterialParameterInfo>& OutParameterInfo, TArray<FGuid>& OutParameterIds) const
-	{
-		if (const UMaterial* Material = GetMaterial())
-		{
-			PRAGMA_DISABLE_DEPRECATION_WARNINGS
-			Material->GetAllParameterInfo<ExpressionType>(OutParameterInfo, OutParameterIds);
-			PRAGMA_ENABLE_DEPRECATION_WARNINGS
-		}
-	}
-#endif // UE_ENABLE_INCLUDE_ORDER_DEPRECATED_IN_5_2
-
 #endif // WITH_EDITORONLY_DATA
 
 	ENGINE_API virtual void GetAllParametersOfType(EMaterialParameterType Type, TMap<FMaterialParameterInfo, FMaterialParameterMetadata>& OutParameters) const override;
