[Quartz] Converts subscription queue to TDepletableMpscQueue.

diff --git a/Engine/Source/Runtime/Engine/Classes/Sound/QuartzSubscription.h b/Engine/Source/Runtime/Engine/Classes/Sound/QuartzSubscription.h
--- a/Engine/Source/Runtime/Engine/Classes/Sound/QuartzSubscription.h
+++ b/Engine/Source/Runtime/Engine/Classes/Sound/QuartzSubscription.h
@@ -3,7 +3,7 @@
 #pragma once
 
 #include "Sound/QuartzQuantizationUtilities.h"
-#include "Containers/MpscQueue.h"
+#include "Containers/DepletableMpscQueue.h"
 
 // forwards
 class UQuartzSubsystem;
@@ -98,26 +98,14 @@ namespace Audio
 
 		void PumpCommandQueue(ListenerType* InListener)
 		{
-			// gather move all the current commands into our temp container
-			// (in case the commands themselves alter the container)
-			TOptional<TFunction<void(ListenerType*)>> Function = CommandQueue.Dequeue();
-			while (Function.IsSet())
-			{
-				TempCommandQueue.Emplace(Function.GetValue());
-				Function = CommandQueue.Dequeue();
-			}
-
-			for (auto& Command : TempCommandQueue)
+			CommandQueue.Deplete([InListener](TFunction<void(ListenerType*)> Command)
 			{
 				Command(InListener);
-			}
-
-			TempCommandQueue.Reset();
+			});
 		}
 
 	private:
-		TMpscQueue<TFunction<void(ListenerType*)>> CommandQueue;
-		TArray<TFunction<void(ListenerType*)>> TempCommandQueue;
+		UE::TDepletableMpscQueue<TFunction<void(ListenerType*)>> CommandQueue;
 	};
 
 } // namespace Audio
