Make it possible to mark a bodyinstance such that when marked it will generate an underlying rigidkinematic instead of a static body. Practically the difference will be that the underlying body will have mass properties. Useful if you want mass properties on statics to be available on the physics thread, or if you plan to convert the static to a dynamic from the physics thread.

diff --git a/Engine/Source/Runtime/Engine/Classes/Components/PrimitiveComponent.h b/Engine/Source/Runtime/Engine/Classes/Components/PrimitiveComponent.h
--- a/Engine/Source/Runtime/Engine/Classes/Components/PrimitiveComponent.h
+++ b/Engine/Source/Runtime/Engine/Classes/Components/PrimitiveComponent.h
@@ -657,6 +657,17 @@ protected:
 	UPROPERTY()
 	uint8 bHasNoStreamableTextures : 1;
 
+	/** When mobility is stationary, use a static underlying physics body. Static bodies do not have
+		physical data like mass. If false, even stationary bodies will be generated with all data
+		necessary for simulating.
+		
+		If you need this body's physical parameters on the physics thread (eg, in a sim callback)
+		then set this to false. */
+	UPROPERTY(EditAnywhere, BlueprintReadOnly, Category=Lighting, AdvancedDisplay, meta=(
+		DisplayName = "Static When Not Moveable",
+		ToolTip = "When false, the underlying physics body will contain all sim data (mass, inertia tensor, etc) even if mobility is not set to Moveable"))
+	uint8 bStaticWhenNotMoveable:1;
+
 #if WITH_EDITOR
 	uint8 bIgnoreBoundsForEditorFocus : 1;
 
@@ -1401,6 +1412,15 @@ public:
 	UFUNCTION(BlueprintCallable, Category="Physics")
 	virtual void SetSimulatePhysics(bool bSimulate);
 
+	/*
+	 *	
+	 */
+	UFUNCTION(BlueprintCallable, Category="Physics")
+	void SetStaticWhenNotMoveable(bool bInStaticWhenNotMoveable);
+
+	UFUNCTION(BlueprintCallable, Category="Physics")
+	bool GetStaticWhenNotMoveable() const { return bStaticWhenNotMoveable; }
+
 	/**
 	 * Determines whether or not the simulate physics setting can be edited interactively on this component
 	 */
diff --git a/Engine/Source/Runtime/Engine/Private/Components/PrimitiveComponent.cpp b/Engine/Source/Runtime/Engine/Private/Components/PrimitiveComponent.cpp
--- a/Engine/Source/Runtime/Engine/Private/Components/PrimitiveComponent.cpp
+++ b/Engine/Source/Runtime/Engine/Private/Components/PrimitiveComponent.cpp
@@ -295,6 +295,7 @@ UPrimitiveComponent::UPrimitiveComponent(const FObjectInitializer& ObjectInitial
 	bCastVolumetricTranslucentShadow = false;
 	bCastContactShadow = true;
 	IndirectLightingCacheQuality = ILCQ_Point;
+	bStaticWhenNotMoveable = true;
 	bSelectable = true;
 #if WITH_EDITORONLY_DATA
 	bConsiderForActorPlacementWhenHidden = false;
diff --git a/Engine/Source/Runtime/Engine/Private/PhysicsEngine/BodyInstance.cpp b/Engine/Source/Runtime/Engine/Private/PhysicsEngine/BodyInstance.cpp
--- a/Engine/Source/Runtime/Engine/Private/PhysicsEngine/BodyInstance.cpp
+++ b/Engine/Source/Runtime/Engine/Private/PhysicsEngine/BodyInstance.cpp
@@ -1451,7 +1451,9 @@ void FInitBodiesHelperBase::InitBodies()
 
 FInitBodySpawnParams::FInitBodySpawnParams(const UPrimitiveComponent* PrimComp)
 {
-	bStaticPhysics = PrimComp == nullptr || PrimComp->Mobility != EComponentMobility::Movable;
+	bStaticPhysics = PrimComp == nullptr || (
+		PrimComp->Mobility != EComponentMobility::Movable &&
+		PrimComp->GetStaticWhenNotMoveable());
 
 	if(const USkeletalMeshComponent* SKOwner = Cast<USkeletalMeshComponent>(PrimComp))
 	{
diff --git a/Engine/Source/Runtime/Engine/Private/PrimitiveComponentPhysics.cpp b/Engine/Source/Runtime/Engine/Private/PrimitiveComponentPhysics.cpp
--- a/Engine/Source/Runtime/Engine/Private/PrimitiveComponentPhysics.cpp
+++ b/Engine/Source/Runtime/Engine/Private/PrimitiveComponentPhysics.cpp
@@ -109,6 +109,15 @@ void UPrimitiveComponent::SetSimulatePhysics(bool bSimulate)
 	BodyInstance.SetInstanceSimulatePhysics(bSimulate);
 }
 
+void UPrimitiveComponent::SetStaticWhenNotMoveable(bool bInStaticWhenNotMoveable)
+{
+	if (bStaticWhenNotMoveable != bInStaticWhenNotMoveable)
+	{
+		bStaticWhenNotMoveable = bInStaticWhenNotMoveable;
+		RecreatePhysicsState();
+	}
+}
+
 void UPrimitiveComponent::SetConstraintMode(EDOFMode::Type ConstraintMode)
 {
 	FBodyInstance * RootBI = GetBodyInstance(NAME_None, false);
