[Quartz] (fix deprecation warning) switching from DepletableMpscQueue to DepletableMpmcQueue

diff --git a/Engine/Source/Runtime/Engine/Classes/Sound/QuartzSubscription.h b/Engine/Source/Runtime/Engine/Classes/Sound/QuartzSubscription.h
--- a/Engine/Source/Runtime/Engine/Classes/Sound/QuartzSubscription.h
+++ b/Engine/Source/Runtime/Engine/Classes/Sound/QuartzSubscription.h
@@ -3,7 +3,7 @@
 #pragma once
 
 #include "Sound/QuartzQuantizationUtilities.h"
-#include "Containers/DepletableMpscQueue.h"
+#include "Containers/DepletableMpmcQueue.h"
 
 // forwards
 class UQuartzSubsystem;
@@ -105,7 +105,7 @@ namespace Audio
 		}
 
 	private:
-		UE::TDepletableMpscQueue<TFunction<void(ListenerType*)>> CommandQueue;
+		UE::TDepletableMpmcQueue<TFunction<void(ListenerType*)>> CommandQueue;
 	};
 
 } // namespace Audio
