Removing deprecated Material methods

diff --git a/Engine/Source/Runtime/Engine/Classes/Materials/MaterialExpressionFontSampleParameter.h b/Engine/Source/Runtime/Engine/Classes/Materials/MaterialExpressionFontSampleParameter.h
--- a/Engine/Source/Runtime/Engine/Classes/Materials/MaterialExpressionFontSampleParameter.h
+++ b/Engine/Source/Runtime/Engine/Classes/Materials/MaterialExpressionFontSampleParameter.h
@@ -76,9 +76,6 @@ class UMaterialExpressionFontSampleParameter : public UMaterialExpressionFontSam
 #endif
 	//~ End UMaterialExpression Interface
 	
-	UE_DEPRECATED(5.0, "Use GetParameterValue and/or GetParameterName")
-	bool IsNamedParameter(const FHashedMaterialParameterInfo& ParameterInfo, UFont*& OutFontValue, int32& OutFontPage) const;
-
 #if WITH_EDITOR
 	bool SetParameterValue(FName InParameterName, UFont* InFontValue, int32 InFontPage, EMaterialExpressionSetParameterValueFlags Flags = EMaterialExpressionSetParameterValueFlags::None);
 #endif
@@ -92,9 +89,6 @@ class UMaterialExpressionFontSampleParameter : public UMaterialExpressionFontSam
 	{
 		return ExpressionGUID;
 	}
-
-	UE_DEPRECATED(5.0, "Use GetAllParameterInfoOfType or GetAllParametersOfType")
-	void GetAllParameterInfo(TArray<FMaterialParameterInfo> &OutParameterInfo, TArray<FGuid> &OutParameterIds, const FMaterialParameterInfo& InBaseParameterInfo) const;
 };
 
 
diff --git a/Engine/Source/Runtime/Engine/Classes/Materials/MaterialExpressionParameter.h b/Engine/Source/Runtime/Engine/Classes/Materials/MaterialExpressionParameter.h
--- a/Engine/Source/Runtime/Engine/Classes/Materials/MaterialExpressionParameter.h
+++ b/Engine/Source/Runtime/Engine/Classes/Materials/MaterialExpressionParameter.h
@@ -67,8 +67,4 @@ class UMaterialExpressionParameter : public UMaterialExpression
 	{
 		return ExpressionGUID;
 	}
-
-	UE_DEPRECATED(5.0, "Use GetAllParameterInfoOfType or GetAllParametersOfType")
-	virtual void GetAllParameterInfo(TArray<FMaterialParameterInfo> &OutParameterInfo, TArray<FGuid> &OutParameterIds, const FMaterialParameterInfo& InBaseParameterInfo) const;
-
 };
diff --git a/Engine/Source/Runtime/Engine/Classes/Materials/MaterialExpressionRuntimeVirtualTextureSampleParameter.h b/Engine/Source/Runtime/Engine/Classes/Materials/MaterialExpressionRuntimeVirtualTextureSampleParameter.h
--- a/Engine/Source/Runtime/Engine/Classes/Materials/MaterialExpressionRuntimeVirtualTextureSampleParameter.h
+++ b/Engine/Source/Runtime/Engine/Classes/Materials/MaterialExpressionRuntimeVirtualTextureSampleParameter.h
@@ -33,12 +33,6 @@ class UMaterialExpressionRuntimeVirtualTextureSampleParameter : public UMaterial
 	bool SetParameterValue(FName InParameterName, URuntimeVirtualTexture* InValue, EMaterialExpressionSetParameterValueFlags Flags = EMaterialExpressionSetParameterValueFlags::None);
 #endif
 
-	UE_DEPRECATED(5.0, "Use GetParameterValue and/or GetParameterName")
-	bool IsNamedParameter(const FHashedMaterialParameterInfo& ParameterInfo, URuntimeVirtualTexture*& OutValue) const;
-	
-	UE_DEPRECATED(5.0, "Use GetAllParameterInfoOfType or GetAllParametersOfType")
-	void GetAllParameterInfo(TArray<FMaterialParameterInfo> &OutParameterInfo, TArray<FGuid> &OutParameterIds, const FMaterialParameterInfo& InBaseParameterInfo) const;
-
 	//~ Begin UMaterialExpression Interface
 #if WITH_EDITOR
 	virtual bool CanRenameNode() const override { return true; }
diff --git a/Engine/Source/Runtime/Engine/Classes/Materials/MaterialExpressionScalarParameter.h b/Engine/Source/Runtime/Engine/Classes/Materials/MaterialExpressionScalarParameter.h
--- a/Engine/Source/Runtime/Engine/Classes/Materials/MaterialExpressionScalarParameter.h
+++ b/Engine/Source/Runtime/Engine/Classes/Materials/MaterialExpressionScalarParameter.h
@@ -72,9 +72,6 @@ class UMaterialExpressionScalarParameter : public UMaterialExpressionParameter
 #endif
 	//~ End UMaterialExpression Interface
 
-	UE_DEPRECATED(5.0, "Use GetParameterValue and/or GetParameterName")
-	bool IsNamedParameter(const FHashedMaterialParameterInfo& ParameterInfo, float& OutValue) const;
-
 #if WITH_EDITOR
 	bool SetParameterValue(FName InParameterName, float InValue, EMaterialExpressionSetParameterValueFlags Flags = EMaterialExpressionSetParameterValueFlags::None);
 
@@ -84,9 +81,6 @@ class UMaterialExpressionScalarParameter : public UMaterialExpressionParameter
 #endif
 
 	virtual bool IsUsedAsAtlasPosition() const { return false; }
-
-	UE_DEPRECATED(5.0, "Use GetAllParameterInfoOfType or GetAllParametersOfType")
-	virtual void GetAllParameterInfo(TArray<FMaterialParameterInfo> &OutParameterInfo, TArray<FGuid> &OutParameterIds, const FMaterialParameterInfo& InBaseParameterInfo) const override;
 };
 
 
diff --git a/Engine/Source/Runtime/Engine/Classes/Materials/MaterialExpressionStaticBoolParameter.h b/Engine/Source/Runtime/Engine/Classes/Materials/MaterialExpressionStaticBoolParameter.h
--- a/Engine/Source/Runtime/Engine/Classes/Materials/MaterialExpressionStaticBoolParameter.h
+++ b/Engine/Source/Runtime/Engine/Classes/Materials/MaterialExpressionStaticBoolParameter.h
@@ -48,9 +48,6 @@ public:
 #endif
 	//~ End UMaterialExpression Interface
 
-	UE_DEPRECATED(5.0, "Use GetParameterValue and/or GetParameterName")
-	bool IsNamedParameter(const FHashedMaterialParameterInfo& ParameterInfo, bool& OutValue, FGuid& OutExpressionGuid) const;
-
 #if WITH_EDITOR
 	bool SetParameterValue(FName InParameterName, bool OutValue, FGuid InExpressionGuid, EMaterialExpressionSetParameterValueFlags Flags = EMaterialExpressionSetParameterValueFlags::None);
 #endif
diff --git a/Engine/Source/Runtime/Engine/Classes/Materials/MaterialExpressionStaticComponentMaskParameter.h b/Engine/Source/Runtime/Engine/Classes/Materials/MaterialExpressionStaticComponentMaskParameter.h
--- a/Engine/Source/Runtime/Engine/Classes/Materials/MaterialExpressionStaticComponentMaskParameter.h
+++ b/Engine/Source/Runtime/Engine/Classes/Materials/MaterialExpressionStaticComponentMaskParameter.h
@@ -67,9 +67,6 @@ public:
 #endif
 	//~ End UMaterialExpression Interface
 
-	UE_DEPRECATED(5.0, "Use GetParameterValue and/or GetParameterName")
-	bool IsNamedParameter(const FHashedMaterialParameterInfo& ParameterInfo, bool& OutR, bool& OutG, bool& OutB, bool& OutA, FGuid&OutExpressionGuid) const;
-
 #if WITH_EDITOR
 	bool SetParameterValue(FName InParameterName, bool InR, bool InG, bool InB, bool InA, FGuid InExpressionGuid, EMaterialExpressionSetParameterValueFlags Flags = EMaterialExpressionSetParameterValueFlags::None);
 #endif
diff --git a/Engine/Source/Runtime/Engine/Classes/Materials/MaterialExpressionTextureSampleParameter.h b/Engine/Source/Runtime/Engine/Classes/Materials/MaterialExpressionTextureSampleParameter.h
--- a/Engine/Source/Runtime/Engine/Classes/Materials/MaterialExpressionTextureSampleParameter.h
+++ b/Engine/Source/Runtime/Engine/Classes/Materials/MaterialExpressionTextureSampleParameter.h
@@ -80,9 +80,6 @@ class ENGINE_API UMaterialExpressionTextureSampleParameter : public UMaterialExp
 #endif
 	//~ End UMaterialExpression Interface
 
-	UE_DEPRECATED(5.0, "Use GetParameterValue and/or GetParameterName")
-	bool IsNamedParameter(const FHashedMaterialParameterInfo& ParameterInfo, UTexture*& OutValue) const;
-
 #if WITH_EDITOR
 	bool SetParameterValue(FName InParameterName, UTexture* InValue, EMaterialExpressionSetParameterValueFlags Flags = EMaterialExpressionSetParameterValueFlags::None);
 	virtual void PostEditChangeProperty(FPropertyChangedEvent& PropertyChangedEvent) override;
@@ -114,7 +111,4 @@ class ENGINE_API UMaterialExpressionTextureSampleParameter : public UMaterialExp
 		return ChannelNames;
 	}
 #endif
-
-	UE_DEPRECATED(5.0, "Use GetAllParameterInfoOfType or GetAllParametersOfType")
-	void GetAllParameterInfo(TArray<FMaterialParameterInfo> &OutParameterInfo, TArray<FGuid> &OutParameterIds, const FMaterialParameterInfo& InBaseParameterInfo) const;
 };
diff --git a/Engine/Source/Runtime/Engine/Classes/Materials/MaterialExpressionVectorParameter.h b/Engine/Source/Runtime/Engine/Classes/Materials/MaterialExpressionVectorParameter.h
--- a/Engine/Source/Runtime/Engine/Classes/Materials/MaterialExpressionVectorParameter.h
+++ b/Engine/Source/Runtime/Engine/Classes/Materials/MaterialExpressionVectorParameter.h
@@ -61,16 +61,12 @@ class UMaterialExpressionVectorParameter : public UMaterialExpressionParameter
 #endif
 	//~ End UMaterialExpression Interface
 
-	UE_DEPRECATED(5.0, "Use GetParameterValue and/or GetParameterName")
-	bool IsNamedParameter(const FHashedMaterialParameterInfo& ParameterInfo, FLinearColor& OutValue) const;
-
 #if WITH_EDITOR
 	virtual bool SetParameterValue(FName InParameterName, FLinearColor InValue, EMaterialExpressionSetParameterValueFlags Flags = EMaterialExpressionSetParameterValueFlags::None);
 
 	virtual void PostEditChangeProperty(FPropertyChangedEvent& PropertyChangedEvent) override;	
 
 	void ApplyChannelNames();
-	
 
 	virtual void ValidateParameterName(const bool bAllowDuplicateName) override;
 	virtual bool HasClassAndNameCollision(UMaterialExpression* OtherExpression) const override;
@@ -84,9 +80,6 @@ class UMaterialExpressionVectorParameter : public UMaterialExpressionParameter
 		return ChannelNames;
 	}
 #endif
-
-	UE_DEPRECATED(5.0, "Use GetAllParameterInfoOfType or GetAllParametersOfType")
-	virtual void GetAllParameterInfo(TArray<FMaterialParameterInfo> &OutParameterInfo, TArray<FGuid> &OutParameterIds, const FMaterialParameterInfo& InBaseParameterInfo) const override;
 };
 
 
diff --git a/Engine/Source/Runtime/Engine/Classes/Materials/MaterialInstance.h b/Engine/Source/Runtime/Engine/Classes/Materials/MaterialInstance.h
--- a/Engine/Source/Runtime/Engine/Classes/Materials/MaterialInstance.h
+++ b/Engine/Source/Runtime/Engine/Classes/Materials/MaterialInstance.h
@@ -550,33 +550,7 @@ class UMaterialInstance : public UMaterialInterface
 	UPROPERTY(EditAnywhere, Category = MaterialInstance)
 	FMaterialOverrideNanite NaniteOverrideMaterial;
 
-	/**
-	 * Delegate for custom static parameters getter.
-	 *
-	 * @param OutStaticParameterSet Parameter set to append.
-	 * @param Material Material instance to collect parameters.
-	 */
-	DECLARE_MULTICAST_DELEGATE_TwoParams(FCustomStaticParametersGetterDelegate, FStaticParameterSet&, UMaterialInstance*);
-
-	/**
-	 * Delegate for custom static parameters updater.
-	 *
-	 * @param StaticParameterSet Parameter set to update.
-	 * @param Material Material to update.
-	 *
-	 * @returns True if any parameter been updated. False otherwise.
-	 */
-	DECLARE_DELEGATE_RetVal_TwoParams(bool, FCustomParameterSetUpdaterDelegate, FStaticParameterSet&, UMaterial*);
-
 #if WITH_EDITORONLY_DATA
-	// Custom static parameters getter delegate.
-	UE_DEPRECATED(5.0, "Custom static parameter delegates no longer supported.")
-	ENGINE_API static FCustomStaticParametersGetterDelegate CustomStaticParametersGetters;
-
-	// An array of custom parameter set updaters.
-	UE_DEPRECATED(5.0, "Custom static parameter delegates no longer supported.")
-	ENGINE_API static TArray<FCustomParameterSetUpdaterDelegate> CustomParameterSetUpdaters;
-
 	ENGINE_API const FStaticParameterSetEditorOnlyData& GetEditorOnlyStaticParameters() const;
 #endif // WITH_EDITORONLY_DATA
 
diff --git a/Engine/Source/Runtime/Engine/Private/Materials/MaterialExpressions.cpp b/Engine/Source/Runtime/Engine/Private/Materials/MaterialExpressions.cpp
--- a/Engine/Source/Runtime/Engine/Private/Materials/MaterialExpressions.cpp
+++ b/Engine/Source/Runtime/Engine/Private/Materials/MaterialExpressions.cpp
@@ -3394,33 +3394,6 @@ UMaterialExpressionRuntimeVirtualTextureSampleParameter::UMaterialExpressionRunt
 #endif
 }
 
-bool UMaterialExpressionRuntimeVirtualTextureSampleParameter::IsNamedParameter(const FHashedMaterialParameterInfo& ParameterInfo, URuntimeVirtualTexture*& OutValue) const
-{
-	if (ParameterInfo.Name == ParameterName)
-	{
-		OutValue = VirtualTexture;
-		return true;
-	}
-
-	return false;
-}
-
-void UMaterialExpressionRuntimeVirtualTextureSampleParameter::GetAllParameterInfo(TArray<FMaterialParameterInfo> &OutParameterInfo, TArray<FGuid> &OutParameterIds, const FMaterialParameterInfo& InBaseParameterInfo) const
-{
-	int32 CurrentSize = OutParameterInfo.Num();
-	FMaterialParameterInfo NewParameter(ParameterName, InBaseParameterInfo.Association, InBaseParameterInfo.Index);
-#if WITH_EDITOR
-	if (HasConnectedOutputs())
-#endif
-	{
-		OutParameterInfo.AddUnique(NewParameter);
-		if (CurrentSize != OutParameterInfo.Num())
-		{
-			OutParameterIds.Add(ExpressionGUID);
-		}
-	}
-}
-
 #if WITH_EDITOR
 
 template<typename T>
@@ -3686,20 +3659,7 @@ void UMaterialExpressionTextureSampleParameter::ValidateParameterName(const bool
 {
 	ValidateParameterNameInternal(this, Material, bAllowDuplicateName);
 }
-#endif // WITH_EDITOR
-
-bool UMaterialExpressionTextureSampleParameter::IsNamedParameter(const FHashedMaterialParameterInfo& ParameterInfo, UTexture*& OutValue) const
-{
-	if (ParameterInfo.Name == ParameterName)
-	{
-		OutValue = Texture;
-		return true;
-	}
 
-	return false;
-}
-
-#if WITH_EDITOR
 bool UMaterialExpressionTextureSampleParameter::SetParameterValue(FName InParameterName, UTexture* InValue, EMaterialExpressionSetParameterValueFlags Flags)
 {
 	if (InParameterName == ParameterName)
@@ -3759,24 +3719,6 @@ void UMaterialExpressionTextureSampleParameter::SetDefaultTexture()
 }
 #endif // WITH_EDITOR
 
-void UMaterialExpressionTextureSampleParameter::GetAllParameterInfo(TArray<FMaterialParameterInfo> &OutParameterInfo, TArray<FGuid> &OutParameterIds, const FMaterialParameterInfo& InBaseParameterInfo) const
-{
-	int32 CurrentSize = OutParameterInfo.Num();
-	FMaterialParameterInfo NewParameter(ParameterName, InBaseParameterInfo.Association, InBaseParameterInfo.Index);
-
-#if WITH_EDITOR
-	if (HasConnectedOutputs())
-#endif
-	{
-		OutParameterInfo.AddUnique(NewParameter);
-
-		if (CurrentSize != OutParameterInfo.Num())
-		{
-			OutParameterIds.Add(ExpressionGUID);
-		}
-	}
-}
-
 //
 //  UMaterialExpressionTextureObjectParameter
 //
@@ -5685,24 +5627,7 @@ void UMaterialExpressionStaticComponentMaskParameter::GetCaption(TArray<FString>
 	OutCaptions.Add(TEXT("Mask Param")); 
 	OutCaptions.Add(FString::Printf(TEXT("'%s'"), *ParameterName.ToString()));
 }
-#endif // WITH_EDITOR
-
-bool UMaterialExpressionStaticComponentMaskParameter::IsNamedParameter(const FHashedMaterialParameterInfo& ParameterInfo, bool& OutR, bool& OutG, bool& OutB, bool& OutA, FGuid&OutExpressionGuid) const
-{
-	if (ParameterInfo.Name == ParameterName)
-	{
-		OutR = DefaultR;
-		OutG = DefaultG;
-		OutB = DefaultB;
-		OutA = DefaultA;
-		OutExpressionGuid = ExpressionGUID;
-		return true;
-	}
-
-	return false;
-}
 
-#if WITH_EDITOR
 bool  UMaterialExpressionStaticComponentMaskParameter::SetParameterValue(FName InParameterName, bool InR, bool InG, bool InB, bool InA, FGuid InExpressionGuid, EMaterialExpressionSetParameterValueFlags Flags)
 {
 	if (InParameterName == ParameterName)
@@ -8518,26 +8443,6 @@ void UMaterialExpressionParameter::SetEditableName(const FString& NewName)
 	ParameterName = *NewName;
 }
 
-#endif
-
-void UMaterialExpressionParameter::GetAllParameterInfo(TArray<FMaterialParameterInfo> &OutParameterInfo, TArray<FGuid> &OutParameterIds, const FMaterialParameterInfo& InBaseParameterInfo) const
-{
-	int32 CurrentSize = OutParameterInfo.Num();
-	FMaterialParameterInfo NewParameter(ParameterName, InBaseParameterInfo.Association, InBaseParameterInfo.Index);
-
-#if WITH_EDITOR
-	if (HasConnectedOutputs())
-#endif
-	{
-		OutParameterInfo.AddUnique(MoveTemp(NewParameter));
-		if (CurrentSize != OutParameterInfo.Num())
-		{
-			OutParameterIds.Add(ExpressionGUID);
-		}
-	}
-}
-
-#if WITH_EDITOR
 void UMaterialExpressionParameter::ValidateParameterName(const bool bAllowDuplicateName)
 {
 	ValidateParameterNameInternal(this, Material, bAllowDuplicateName);
@@ -8609,30 +8514,7 @@ void UMaterialExpressionVectorParameter::GetCaption(TArray<FString>& OutCaptions
 
 	OutCaptions.Add(FString::Printf(TEXT("'%s'"), *ParameterName.ToString())); 
 }
-#endif // WITH_EDITOR
-
-bool UMaterialExpressionVectorParameter::IsNamedParameter(const FHashedMaterialParameterInfo& ParameterInfo, FLinearColor& OutValue) const
-{
-	if (ParameterInfo.Name == ParameterName)
-	{
-		OutValue = DefaultValue;
-		return true;
-	}
 
-	return false;
-}
-
-void UMaterialExpressionVectorParameter::GetAllParameterInfo(TArray<FMaterialParameterInfo> &OutParameterInfo, TArray<FGuid> &OutParameterIds, const FMaterialParameterInfo& InBaseParameterInfo) const
-{
-	if (!bUseCustomPrimitiveData)
-	{
-		PRAGMA_DISABLE_DEPRECATION_WARNINGS
-		Super::GetAllParameterInfo(OutParameterInfo, OutParameterIds, InBaseParameterInfo);
-		PRAGMA_ENABLE_DEPRECATION_WARNINGS
-	}
-}
-
-#if WITH_EDITOR
 bool UMaterialExpressionVectorParameter::SetParameterValue(FName InParameterName, FLinearColor InValue, EMaterialExpressionSetParameterValueFlags Flags)
 {
 	if (InParameterName == ParameterName)
@@ -8990,30 +8872,7 @@ void UMaterialExpressionScalarParameter::GetCaption(TArray<FString>& OutCaptions
 	}
 	 OutCaptions.Add(FString::Printf(TEXT("'%s'"), *ParameterName.ToString())); 
 }
-#endif // WITH_EDITOR
-
-bool UMaterialExpressionScalarParameter::IsNamedParameter(const FHashedMaterialParameterInfo& ParameterInfo, float& OutValue) const
-{
-	if (ParameterInfo.Name == ParameterName)
-	{
-		OutValue = DefaultValue;
-		return true;
-	}
 
-	return false;
-}
-
-void UMaterialExpressionScalarParameter::GetAllParameterInfo(TArray<FMaterialParameterInfo> &OutParameterInfo, TArray<FGuid> &OutParameterIds, const FMaterialParameterInfo& InBaseParameterInfo) const
-{
-	if (!bUseCustomPrimitiveData)
-	{
-		PRAGMA_DISABLE_DEPRECATION_WARNINGS
-		Super::GetAllParameterInfo(OutParameterInfo, OutParameterIds, InBaseParameterInfo);
-		PRAGMA_ENABLE_DEPRECATION_WARNINGS
-	}
-}
-
-#if WITH_EDITOR
 bool UMaterialExpressionScalarParameter::SetParameterValue(FName InParameterName, float InValue, EMaterialExpressionSetParameterValueFlags Flags)
 {
 	if (InParameterName == ParameterName)
@@ -9229,21 +9088,7 @@ void UMaterialExpressionStaticBoolParameter::GetCaption(TArray<FString>& OutCapt
 	OutCaptions.Add(FString::Printf(TEXT("Static Bool Param (%s)"), (DefaultValue ? TEXT("True") : TEXT("False")))); 
 	OutCaptions.Add(FString::Printf(TEXT("'%s'"), *ParameterName.ToString())); 
 }
-#endif // WITH_EDITOR
-
-bool UMaterialExpressionStaticBoolParameter::IsNamedParameter(const FHashedMaterialParameterInfo& ParameterInfo, bool& OutValue, FGuid& OutExpressionGuid) const
-{
-	if (ParameterInfo.Name == ParameterName)
-	{
-		OutValue = DefaultValue;
-		OutExpressionGuid = ExpressionGUID;
-		return true;
-	}
 
-	return false;
-}
-
-#if WITH_EDITOR
 bool UMaterialExpressionStaticBoolParameter::SetParameterValue(FName InParameterName, bool InValue, FGuid InExpressionGuid, EMaterialExpressionSetParameterValueFlags Flags)
 {
 	if (InParameterName == ParameterName)
@@ -12467,21 +12312,7 @@ void UMaterialExpressionFontSampleParameter::ValidateParameterName(const bool bA
 {
 	ValidateParameterNameInternal(this, Material, bAllowDuplicateName);
 }
-#endif // WITH_EDITOR
-
-bool UMaterialExpressionFontSampleParameter::IsNamedParameter(const FHashedMaterialParameterInfo& ParameterInfo, UFont*& OutFontValue, int32& OutFontPage) const
-{
-	if (ParameterInfo.Name == ParameterName)
-	{
-		OutFontValue = Font;
-		OutFontPage = FontTexturePage;
-		return true;
-	}
-
-	return false;
-}
 
-#if WITH_EDITOR
 bool UMaterialExpressionFontSampleParameter::SetParameterValue(FName InParameterName, UFont* InFontValue, int32 InFontPage, EMaterialExpressionSetParameterValueFlags Flags)
 {
 	if (InParameterName == ParameterName)
@@ -12528,23 +12359,6 @@ void UMaterialExpressionFontSampleParameter::SetEditableName(const FString& NewN
 }
 #endif
 
-void UMaterialExpressionFontSampleParameter::GetAllParameterInfo(TArray<FMaterialParameterInfo> &OutParameterInfo, TArray<FGuid> &OutParameterIds, const FMaterialParameterInfo& InBaseParameterInfo) const
-{
-	int32 CurrentSize = OutParameterInfo.Num();
-	FMaterialParameterInfo NewParameter(ParameterName, InBaseParameterInfo.Association, InBaseParameterInfo.Index);
-#if WITH_EDITOR
-	if (HasConnectedOutputs())
-#endif
-	{
-		OutParameterInfo.AddUnique(NewParameter);
-		if (CurrentSize != OutParameterInfo.Num())
-		{
-			OutParameterIds.Add(ExpressionGUID);
-		}
-	}
-}
-
-
 ///////////////////////////////////////////////////////////////////////////////
 // UMaterialExpressionWorldPosition
 ///////////////////////////////////////////////////////////////////////////////
diff --git a/Engine/Source/Runtime/Engine/Private/Materials/MaterialInstance.cpp b/Engine/Source/Runtime/Engine/Private/Materials/MaterialInstance.cpp
--- a/Engine/Source/Runtime/Engine/Private/Materials/MaterialInstance.cpp
+++ b/Engine/Source/Runtime/Engine/Private/Materials/MaterialInstance.cpp
@@ -622,20 +622,6 @@ bool UMaterialInstance::UpdateParameters()
 
 			// Static component mask parameters
 			bDirty = UpdateParameterSet<FStaticComponentMaskParameter, UMaterialExpressionStaticComponentMaskParameter>(EditorOnly->StaticParameters.StaticComponentMaskParameters, ParentMaterial) || bDirty;
-
-			// Custom parameters
-			PRAGMA_DISABLE_DEPRECATION_WARNINGS
-			if (CustomParameterSetUpdaters.Num() > 0)
-			{
-				FStaticParameterSet LocalStaticParameters = GetStaticParameters();
-				for (const auto& CustomParameterSetUpdater : CustomParameterSetUpdaters)
-				{
-					bDirty |= CustomParameterSetUpdater.Execute(LocalStaticParameters, ParentMaterial);
-				}
-				StaticParametersRuntime = LocalStaticParameters.GetRuntime();
-				EditorOnly->StaticParameters = LocalStaticParameters.EditorOnly;
-			}
-			PRAGMA_ENABLE_DEPRECATION_WARNINGS
 		}
 
 		if (StaticParametersRuntime.bHasMaterialLayers && Parent)
@@ -1867,11 +1853,6 @@ void UMaterialInstance::GetStaticParameterValues(FStaticParameterSet& OutStaticP
 
 	OutStaticParameters.Validate();
 
-	// Custom parameters.
-	PRAGMA_DISABLE_DEPRECATION_WARNINGS
-	CustomStaticParametersGetters.Broadcast(OutStaticParameters, this);
-	PRAGMA_ENABLE_DEPRECATION_WARNINGS
-
 	if (AllowCachingStaticParameterValuesCounter > 0)
 	{
 		CachedStaticParameterValues = OutStaticParameters;
@@ -4748,11 +4729,3 @@ void UMaterialInstance::CopyMaterialUniformParametersInternal(UMaterialInterface
 	FObjectCacheEventSink::NotifyReferencedTextureChanged_Concurrent(this);
 #endif
 }
-
-#if WITH_EDITORONLY_DATA
-PRAGMA_DISABLE_DEPRECATION_WARNINGS
-UMaterialInstance::FCustomStaticParametersGetterDelegate UMaterialInstance::CustomStaticParametersGetters;
-TArray<UMaterialInstance::FCustomParameterSetUpdaterDelegate> UMaterialInstance::CustomParameterSetUpdaters;
-PRAGMA_ENABLE_DEPRECATION_WARNINGS
-#endif // WITH_EDITORONLY_DATA
-
