fix ICE in GLSL based shader compiler toolchains due to bitwise operator used for bool expression

diff --git a/Engine/Source/Runtime/Engine/Private/Materials/HLSLMaterialDerivativeAutogen.cpp b/Engine/Source/Runtime/Engine/Private/Materials/HLSLMaterialDerivativeAutogen.cpp
--- a/Engine/Source/Runtime/Engine/Private/Materials/HLSLMaterialDerivativeAutogen.cpp
+++ b/Engine/Source/Runtime/Engine/Private/Materials/HLSLMaterialDerivativeAutogen.cpp
@@ -1973,7 +1973,7 @@ FString FMaterialDerivativeAutogen::GenerateUsedFunctions(FHLSLMaterialTranslato
 					else
 					{
 						Ret += TEXT("\tRet.Value = saturate(A.Value);") LINE_TERMINATOR;
-						Ret += TEXT("\t") + BoolName + TEXT(" InRange = ((0.0 < A.Value) & (A.Value < 1.0));") LINE_TERMINATOR;
+						Ret += TEXT("\t") + BoolName + TEXT(" InRange = ((0.0 < A.Value) && (A.Value < 1.0));") LINE_TERMINATOR;
 					}
 					Ret += TEXT("\tRet.Ddx = select(InRange, A.Ddx, Zero);") LINE_TERMINATOR;
 					Ret += TEXT("\tRet.Ddy = select(InRange, A.Ddy, Zero);") LINE_TERMINATOR;
