GitHub #9689 : Remove CastChecked on leader pose component use at SkeletalMeshComponent, as it crashes when the leader pose component is a UPoseableMeshComponent

Contributed by [REDACTED_PERSON]

diff --git a/Engine/Source/Runtime/Engine/Private/Components/SkeletalMeshComponent.cpp b/Engine/Source/Runtime/Engine/Private/Components/SkeletalMeshComponent.cpp
--- a/Engine/Source/Runtime/Engine/Private/Components/SkeletalMeshComponent.cpp
+++ b/Engine/Source/Runtime/Engine/Private/Components/SkeletalMeshComponent.cpp
@@ -1356,8 +1356,8 @@ bool USkeletalMeshComponent::ShouldUpdateTransform(bool bLODHasChanged) const
 		// if leader pose is ticking, follower also has to update it
 		if (LeaderPoseComponent.IsValid())
 		{
-			const USkeletalMeshComponent* Leader = CastChecked<USkeletalMeshComponent>(LeaderPoseComponent.Get());
-			if (Leader->GetUpdateAnimationInEditor())
+			const USkeletalMeshComponent* Leader = Cast<USkeletalMeshComponent>(LeaderPoseComponent.Get());
+			if (Leader && Leader->GetUpdateAnimationInEditor())
 			{
 				return true;
 			}
