[shaders]
- Prevent creation of duplicate shader maps for a ShaderMapId due to multiple calls to BeginCacheShaders for materials/material instances pointing to the same shader map in a single tick (in the WITH_EDITOR case, each would not find a shadermap via FMaterialShaderMap::FindId in the first phase and so each would end up calling BeginCompileShaderMap in the second phase; this in contrast to the non-editor case where both phases are synchronously/sequentially executed and so this case cannot occur).
- Modify FMaterial::GetUniqueAssetName to use a combination of the base material name, and a hash of the ShaderMapId portion of the DDC key string (excluding the source hashes) and base property overrides. The previous approach used material/material instance name+path (and was at least sufficiently unique to disambiguate debug info) but was not stable/consistent across cooks (the first material or instance to encounter a particular shader map would "win" and name the debug info folder with its own name). With this change the exact same set of debug info folders will be generated on two cooks which modify both cvars affecting shader compilation (i.e. enabling STB or minifier), and ShaderVersion.ush - this will make it easier to do A/B comparisons. Also fixed the hashing done in this function to correctly account for the size of the string data.

diff --git a/Engine/Source/Runtime/Engine/Private/Materials/MaterialShader.cpp b/Engine/Source/Runtime/Engine/Private/Materials/MaterialShader.cpp
--- a/Engine/Source/Runtime/Engine/Private/Materials/MaterialShader.cpp
+++ b/Engine/Source/Runtime/Engine/Private/Materials/MaterialShader.cpp
@@ -1027,7 +1027,7 @@ void FMaterialShaderMapId::AppendStaticParametersString(FString& ParamsString) c
 	}
 }
 
-void FMaterialShaderMapId::AppendKeyString(FString& KeyString) const
+void FMaterialShaderMapId::AppendKeyString(FString& KeyString, bool bIncludeSourceHashes) const
 {
 	check(IsContentValid());
 	BaseMaterialId.AppendString(KeyString);
@@ -1086,7 +1086,8 @@ void FMaterialShaderMapId::AppendKeyString(FString& KeyString) const
 		MakeArrayView(ShaderPipelineTypeDependencies),
 		MakeArrayView(VertexFactoryTypeDependencies),
 		LayoutParams, 
-		KeyString);
+		KeyString,
+		bIncludeSourceHashes);
 
 	BytesToHex(&TextureReferencesHash.Hash[0], sizeof(TextureReferencesHash.Hash), KeyString);
 
diff --git a/Engine/Source/Runtime/Engine/Private/Materials/MaterialShared.cpp b/Engine/Source/Runtime/Engine/Private/Materials/MaterialShared.cpp
--- a/Engine/Source/Runtime/Engine/Private/Materials/MaterialShared.cpp
+++ b/Engine/Source/Runtime/Engine/Private/Materials/MaterialShared.cpp
@@ -2491,6 +2491,8 @@ bool FMaterial::CacheShaders(const FMaterialShaderMapId& ShaderMapId, EShaderPla
 
 	// Just make sure that we don't already have a pending cache going on.
 	FinishCacheShaders();
+	
+	bool bBeginFoundCompiling = false;
 #endif // WITH_EDITOR
 
 	// If we loaded this material with inline shaders, use what was loaded (GameThreadShaderMap) instead of looking in the DDC
@@ -2525,10 +2527,13 @@ bool FMaterial::CacheShaders(const FMaterialShaderMapId& ShaderMapId, EShaderPla
 			TRefCountPtr<FMaterialShaderMap> ShaderMap = FMaterialShaderMap::FindId(ShaderMapId, Platform);
 			if (ShaderMap)
 			{
+				// another material has registered this shader map and its compilation is in-progress
+				// we only need to ensure it contains all shaders required for this material
 				if (ShaderMap->GetCompilingId() != 0u)
 				{
 					SetCompilingShaderMap(ShaderMap);
 					ShaderMap = ShaderMap->GetFinalizedClone();
+					bBeginFoundCompiling = true;
 				}
 			}
 
@@ -2556,17 +2561,42 @@ bool FMaterial::CacheShaders(const FMaterialShaderMapId& ShaderMapId, EShaderPla
 	// FinishCacheShaders once the DDC call initiated in BeginLoadFromDerivedDataCache above has finished.
 	// For client builds, this is executed in place without any lambda.
 #if WITH_EDITOR
-	CacheShadersCompletion = [this, Platform, ShaderMapId, DDCKeyHash, PrecompileMode, TargetPlatform, CompletionCallback = MoveTemp(CompletionCallback)]() {
+	CacheShadersCompletion = [this, Platform, ShaderMapId, DDCKeyHash, bBeginFoundCompiling, PrecompileMode, TargetPlatform, CompletionCallback = MoveTemp(CompletionCallback)]() {
 
 	ON_SCOPE_EXIT{ CacheShadersCompletion.Reset(); };
-
-	if (CacheShadersPending.IsValid())
+	
+	bool bFoundCompiling = bBeginFoundCompiling;
+	if (!GameThreadShaderMap)
 	{
-		TRefCountPtr<FMaterialShaderMap> LoadedShaderMap = CacheShadersPending->Get();
-		CacheShadersPending.Reset();
-
-		check(!LoadedShaderMap || LoadedShaderMap->GetFrozenContentSize() > 0u);
-		SetGameThreadShaderMap(LoadedShaderMap);
+		TRefCountPtr<FMaterialShaderMap> ShaderMap;
+		check(!bFoundCompiling); // sanity check
+		if (CacheShadersPending.IsValid()) // we started a load above, check the result
+		{
+			ShaderMap = CacheShadersPending->Get();
+			CacheShadersPending.Reset();
+		}
+		
+		if (!ShaderMap) // if we still don't have a shader map it wasn't initially in the inprocess cache and also was not in the DDC
+		{
+			// we need to check again if another material has created, registered and began compilation on the shader map we need
+			// since our previous call to FMaterialShaderMap::FindId in BeginCacheShaders
+			// this can occur if multiple materials referencing the same shadermap get a BeginCacheShaders call in the same tick
+			ShaderMap = FMaterialShaderMap::FindId(ShaderMapId, Platform);
+			if (ShaderMap)
+			{
+				// as above, it's possible (and in this case likely) that if we found a shader map that its compilation is already
+				// in progress, triggered by another material being processed in this tick. similarly we need to check that it contains
+				// all shaders required for this material (and queue compilation for any that are missing).
+				if (ShaderMap->GetCompilingId() != 0u)
+				{
+					SetCompilingShaderMap(ShaderMap);
+					ShaderMap = ShaderMap->GetFinalizedClone();
+					bFoundCompiling = true;
+				}
+			}
+		}
+		check(!ShaderMap || ShaderMap->GetFrozenContentSize() > 0u);
+		SetGameThreadShaderMap(ShaderMap);
 	}
 
 	// some of the above paths did not mark the shader map as associated with an asset, do so
@@ -2579,11 +2609,19 @@ bool FMaterial::CacheShaders(const FMaterialShaderMapId& ShaderMapId, EShaderPla
 	UMaterialInterface* MaterialInterface = GetMaterialInterface();
 	const bool bMaterialInstance = MaterialInterface && MaterialInterface->IsA(UMaterialInstance::StaticClass());
 	const bool bRequiredComplete = !bMaterialInstance && IsRequiredComplete();
+	
 
 	bool bShaderMapValid = (bool)GameThreadShaderMap;
+#if WITH_EDITOR
+	if (bShaderMapValid && bRequiredComplete && !bFoundCompiling)
+#else
 	if (bShaderMapValid && bRequiredComplete)
+#endif
 	{
 		// Special engine materials (default materials) are required to be complete
+		// We can bypass this check in the case where we found a "required complete" shader map whose compilation is in progress;
+		// we will check if it's complete below and queue any jobs necessary if not (we only need to log warnings if we found a map without
+		// compilation in progress that is incomplete).
 		bShaderMapValid = GameThreadShaderMap->IsComplete(this, false);
 	}
 
@@ -2717,10 +2755,12 @@ bool FMaterial::CacheShaders(const FMaterialShaderMapId& ShaderMapId, EShaderPla
 #if WITH_EDITOR
 FString FMaterial::GetUniqueAssetName(EShaderPlatform Platform, const FMaterialShaderMapId& ShaderMapId) const
 {
-	FString ToHash = GetAssetPath().GetPlainNameString();
-	ShaderMapId.AppendStaticParametersString(ToHash);
-	uint64 Hash = CityHash64((const char*)*ToHash, ToHash.Len());
-	return FString::Printf(TEXT("%s_%llx"), *GetAssetName(), Hash);
+	FString IdKeyStr;
+	// append the portion of the DDC key string associated with the shadermapid, but exclude the source code
+	// hashes, such that this name remains stable when source code edits are applied (including version bumps)
+	ShaderMapId.AppendKeyString(IdKeyStr, /* bIncludeSourceHashes */ false);
+	uint64 Hash = CityHash64((const char*)*IdKeyStr, IdKeyStr.Len() * sizeof(TCHAR));
+	return FString::Printf(TEXT("%s_%llx"), *GetFriendlyName(), Hash);
 }
 #endif // WITH_EDITOR
 
diff --git a/Engine/Source/Runtime/Engine/Public/MaterialShared.h b/Engine/Source/Runtime/Engine/Public/MaterialShared.h
--- a/Engine/Source/Runtime/Engine/Public/MaterialShared.h
+++ b/Engine/Source/Runtime/Engine/Public/MaterialShared.h
@@ -986,7 +986,7 @@ public:
 	void UpdateFromParameterSet(const FStaticParameterSet& StaticParameters);
 
 	/** Appends string representations of this Id to a key string. */
-	void AppendKeyString(FString& KeyString) const;
+	void AppendKeyString(FString& KeyString, bool bIncludeSourceHashes = true) const;
 	void AppendStaticParametersString(FString& ParamsString) const;
 
 	const TArray<FStaticSwitchParameter> &GetStaticSwitchParameters() const 					{ return StaticSwitchParameters; }
@@ -1861,8 +1861,8 @@ public:
 	virtual FString GetAssetName() const { return GetFriendlyName(); }
 
 #if WITH_EDITOR
-	/** Similar to GetAssetName, but appends a hash of the asset path and static switch params (if editoronly data is available)
-	to differentiate materials that would collide when considering name alone. */
+	/** Similar to GetFriendlyName, but appends a hash of shader map ID portion of the DDC key (excluding source file hashes) to differentiate materials/material 
+		instances that would collide when considering the base material name alone (this name is used when creating debug output paths for disambiguation). */
 	ENGINE_API FString GetUniqueAssetName(EShaderPlatform Platform, const FMaterialShaderMapId& ShaderMapId) const;
 #else
 	/** GetUniqueAssetName just returns GetAssetName in non-editor builds */
diff --git a/Engine/Source/Runtime/RenderCore/Private/ShaderCore.cpp b/Engine/Source/Runtime/RenderCore/Private/ShaderCore.cpp
--- a/Engine/Source/Runtime/RenderCore/Private/ShaderCore.cpp
+++ b/Engine/Source/Runtime/RenderCore/Private/ShaderCore.cpp
@@ -1781,7 +1781,8 @@ FSHAHash GetShaderTypeLayoutHash(const FTypeLayoutDesc& TypeDesc, FPlatformTypeL
 void AppendKeyStringShaderDependencies(
 	TConstArrayView<FShaderTypeDependency> ShaderTypeDependencies,
 	FPlatformTypeLayoutParameters LayoutParams,
-	FString& OutKeyString)
+	FString& OutKeyString,
+	bool bIncludeSourceHashes)
 {
 	// Simplified interface if we only have ShaderTypeDependencies
 	AppendKeyStringShaderDependencies(
@@ -1789,7 +1790,8 @@ void AppendKeyStringShaderDependencies(
 		TConstArrayView<FShaderPipelineTypeDependency>(),
 		TConstArrayView<FVertexFactoryTypeDependency>(),
 		LayoutParams,
-		OutKeyString);
+		OutKeyString,
+		bIncludeSourceHashes);
 }
 
 void AppendKeyStringShaderDependencies(
@@ -1797,7 +1799,8 @@ void AppendKeyStringShaderDependencies(
 	TConstArrayView<FShaderPipelineTypeDependency> ShaderPipelineTypeDependencies,
 	TConstArrayView<FVertexFactoryTypeDependency> VertexFactoryTypeDependencies,
 	FPlatformTypeLayoutParameters LayoutParams,
-	FString& OutKeyString)
+	FString& OutKeyString,
+	bool bIncludeSourceHashes)
 {
 	TSet<const TCHAR*> ReferencedUniformBufferNames;
 
@@ -1815,8 +1818,11 @@ void AppendKeyStringShaderDependencies(
 		OutKeyString.AppendChar('_');
 		OutKeyString.AppendInt(GetRayTracingPayloadTypeMaxSize(RayTracingPayloadType));
 
-		// Add the type's source hash so that we can invalidate cached shaders when .usf changes are made
-		ShaderTypeDependency.SourceHash.AppendString(OutKeyString);
+		if (bIncludeSourceHashes)
+		{
+			// Add the type's source hash so that we can invalidate cached shaders when .usf changes are made
+			ShaderTypeDependency.SourceHash.AppendString(OutKeyString);
+		}
 
 		if (const FShaderParametersMetadata* ParameterStructMetadata = ShaderType->GetRootParametersMetadata())
 		{
@@ -1840,7 +1846,11 @@ void AppendKeyStringShaderDependencies(
 
 		OutKeyString.AppendChar('_');
 		OutKeyString.Append(ShaderPipelineType->GetName());
-		Dependency.StagesSourceHash.AppendString(OutKeyString);
+
+		if (bIncludeSourceHashes)
+		{
+			Dependency.StagesSourceHash.AppendString(OutKeyString);
+		}
 
 		for (const FShaderType* ShaderType : ShaderPipelineType->GetStages())
 		{
@@ -1863,7 +1873,11 @@ void AppendKeyStringShaderDependencies(
 		const FVertexFactoryType* VertexFactoryType = FVertexFactoryType::GetVFByName(VFDependency.VertexFactoryTypeName);
 
 		OutKeyString.Append(VertexFactoryType->GetName());
-		VFDependency.VFSourceHash.AppendString(OutKeyString);
+
+		if (bIncludeSourceHashes)
+		{
+			VFDependency.VFSourceHash.AppendString(OutKeyString);
+		}
 
 		for (int32 Frequency = 0; Frequency < SF_NumFrequencies; Frequency++)
 		{
diff --git a/Engine/Source/Runtime/RenderCore/Public/ShaderCore.h b/Engine/Source/Runtime/RenderCore/Public/ShaderCore.h
--- a/Engine/Source/Runtime/RenderCore/Public/ShaderCore.h
+++ b/Engine/Source/Runtime/RenderCore/Public/ShaderCore.h
@@ -1147,14 +1147,16 @@ class FVertexFactoryTypeDependency;
 extern RENDERCORE_API void AppendKeyStringShaderDependencies(
 	TConstArrayView<FShaderTypeDependency> ShaderTypeDependencies,
 	FPlatformTypeLayoutParameters LayoutParams,
-	FString& OutKeyString);
+	FString& OutKeyString,
+	bool bIncludeSourceHashes = true);
 
 extern RENDERCORE_API void AppendKeyStringShaderDependencies(
 	TConstArrayView<FShaderTypeDependency> ShaderTypeDependencies,
 	TConstArrayView<FShaderPipelineTypeDependency> ShaderPipelineTypeDependencies,
 	TConstArrayView<FVertexFactoryTypeDependency> VertexFactoryTypeDependencies,
 	FPlatformTypeLayoutParameters LayoutParams,
-	FString& OutKeyString);
+	FString& OutKeyString,
+	bool bIncludeSourceHashes = true);
 #endif // WITH_EDITOR
 
 /** Create a block of source code to be injected in the preprocessed shader code. The Block will be put into a #line directive
