Engine: add IsPluginMounted() function to PluginBlueprintLibrary

diff --git a/Engine/Source/Runtime/Engine/Private/PluginBlueprintLibrary.cpp b/Engine/Source/Runtime/Engine/Private/PluginBlueprintLibrary.cpp
--- a/Engine/Source/Runtime/Engine/Private/PluginBlueprintLibrary.cpp
+++ b/Engine/Source/Runtime/Engine/Private/PluginBlueprintLibrary.cpp
@@ -7,6 +7,7 @@
 #include "Containers/UnrealString.h"
 #include "Interfaces/IPluginManager.h"
 #include "Interfaces/IProjectManager.h"
+#include "Misc/PackageName.h"
 #include "Templates/SharedPointer.h"
 #include "UObject/SoftObjectPath.h"
 
@@ -191,3 +192,14 @@ bool UPluginBlueprintLibrary::GetPluginEditorCustomVirtualPath(
 
 	return true;
 }
+
+bool UPluginBlueprintLibrary::IsPluginMounted(const FString& PluginName)
+{
+	FString PluginMountedAssetPath;
+	if (!GetPluginMountedAssetPath(PluginName, PluginMountedAssetPath))
+	{
+		return false;
+	}
+
+	return FPackageName::MountPointExists(PluginMountedAssetPath);
+}
diff --git a/Engine/Source/Runtime/Engine/Public/PluginBlueprintLibrary.h b/Engine/Source/Runtime/Engine/Public/PluginBlueprintLibrary.h
--- a/Engine/Source/Runtime/Engine/Public/PluginBlueprintLibrary.h
+++ b/Engine/Source/Runtime/Engine/Public/PluginBlueprintLibrary.h
@@ -170,4 +170,14 @@ public:
 	static bool GetPluginEditorCustomVirtualPath(
 			const FString& PluginName,
 			FString& OutVirtualPath);
+
+	/**
+	 * Determine whether a plugin is mounted.
+	 *
+	 * @param PluginName - Name of the plugin
+	 *
+	 * @return true if the named plugin is mounted, or false otherwise
+	 */
+	UFUNCTION(BlueprintCallable, Category="Engine Scripting | Plugin Utilities")
+	static bool IsPluginMounted(const FString& PluginName);
 };
