Adding missing public include module paths to Engine and UnrealEd

diff --git a/Engine/Source/Editor/UnrealEd/UnrealEd.Build.cs b/Engine/Source/Editor/UnrealEd/UnrealEd.Build.cs
--- a/Engine/Source/Editor/UnrealEd/UnrealEd.Build.cs
+++ b/Engine/Source/Editor/UnrealEd/UnrealEd.Build.cs
@@ -339,6 +339,7 @@ public class UnrealEd : ModuleRules
 				"UncontrolledChangelists",
 				"PropertyEditor",
 				"ClassViewer",
+				"EditorSubsystem"
 			}
 		);
 
diff --git a/Engine/Source/Runtime/Engine/Engine.Build.cs b/Engine/Source/Runtime/Engine/Engine.Build.cs
--- a/Engine/Source/Runtime/Engine/Engine.Build.cs
+++ b/Engine/Source/Runtime/Engine/Engine.Build.cs
@@ -119,6 +119,16 @@ public class Engine : ModuleRules
 			new string[] {
 				"TypedElementFramework",
 				"TypedElementRuntime",
+				"NetCore",
+				"RenderCore",
+				"CoreUObject",
+				"CoreOnline",
+				"PhysicsCore",
+				"ChaosCore",
+				"DeveloperSettings",
+				"NetCommon",
+				"Sockets",
+				"MeshDescription"
 			}
 		);
 
