Abstract Modulate and AlphaHoldout blend mode.

diff --git a/Engine/Plugins/FX/Niagara/Source/Niagara/Private/NiagaraRendererSprites.cpp b/Engine/Plugins/FX/Niagara/Source/Niagara/Private/NiagaraRendererSprites.cpp
--- a/Engine/Plugins/FX/Niagara/Source/Niagara/Private/NiagaraRendererSprites.cpp
+++ b/Engine/Plugins/FX/Niagara/Source/Niagara/Private/NiagaraRendererSprites.cpp
@@ -229,7 +229,8 @@ void FNiagaraRendererSprites::PrepareParticleSpriteRenderData(FParticleSpriteRen
 	check(MaterialRenderProxy);
 
 	// Do we have anything to render?
-	const EBlendMode BlendMode = MaterialRenderProxy->GetIncompleteMaterialWithFallback(FeatureLevel).GetBlendMode();
+	const FMaterial& Material = MaterialRenderProxy->GetIncompleteMaterialWithFallback(FeatureLevel);
+	const EBlendMode BlendMode = Material.GetBlendMode();
 	ParticleSpriteRenderData.BlendMode = BlendMode; // STRATA_TODO_BLENDMODE
 	ParticleSpriteRenderData.bHasTranslucentMaterials = IsTranslucentBlendMode(BlendMode);
 
@@ -264,7 +265,7 @@ void FNiagaraRendererSprites::PrepareParticleSpriteRenderData(FParticleSpriteRen
 		const EShaderPlatform ShaderPlatform = SceneProxy->GetComputeDispatchInterface()->GetShaderPlatform();
 
 		// Determine if we need sorting
-		ParticleSpriteRenderData.bNeedsSort = SortMode != ENiagaraSortMode::None && (BlendMode == BLEND_AlphaComposite || BlendMode == BLEND_AlphaHoldout || BlendMode == BLEND_Translucent || !bSortOnlyWhenTranslucent);
+		ParticleSpriteRenderData.bNeedsSort = SortMode != ENiagaraSortMode::None && (BlendMode == BLEND_AlphaComposite || IsAlphaHoldoutBlendMode(Material) || IsTranslucentOnlyBlendMode(Material) || !bSortOnlyWhenTranslucent);
 		const bool bNeedCustomSort = ParticleSpriteRenderData.bNeedsSort && (SortMode == ENiagaraSortMode::CustomAscending || SortMode == ENiagaraSortMode::CustomDecending);
 		ParticleSpriteRenderData.RendererLayout = bNeedCustomSort ? RendererLayoutWithCustomSort : RendererLayoutWithoutCustomSort;
 		ParticleSpriteRenderData.SortVariable = bNeedCustomSort ? ENiagaraSpriteVFLayout::CustomSorting : ENiagaraSpriteVFLayout::Position;
diff --git a/Engine/Source/Editor/MaterialEditor/Private/SMaterialEditorViewport.cpp b/Engine/Source/Editor/MaterialEditor/Private/SMaterialEditorViewport.cpp
--- a/Engine/Source/Editor/MaterialEditor/Private/SMaterialEditorViewport.cpp
+++ b/Engine/Source/Editor/MaterialEditor/Private/SMaterialEditorViewport.cpp
@@ -177,11 +177,11 @@ FLinearColor FMaterialEditorViewportClient::GetBackgroundColor() const
 			if (MaterialInterface)
 			{
 				const EBlendMode PreviewBlendMode = (EBlendMode)MaterialInterface->GetBlendMode();
-				if (PreviewBlendMode == BLEND_Modulate)
+				if (IsModulateBlendMode(*MaterialInterface))
 				{
 					BackgroundColor = FLinearColor::White;
 				}
-				else if (PreviewBlendMode == BLEND_Translucent || PreviewBlendMode == BLEND_AlphaComposite || PreviewBlendMode == BLEND_AlphaHoldout)
+				else if (IsTranslucentOnlyBlendMode(*MaterialInterface) || PreviewBlendMode == BLEND_AlphaComposite || IsAlphaHoldoutBlendMode(*MaterialInterface))
 				{
 					BackgroundColor = FColor(64, 64, 64);
 				}
diff --git a/Engine/Source/Editor/UnrealEd/Private/Lightmass/LightmassRender.cpp b/Engine/Source/Editor/UnrealEd/Private/Lightmass/LightmassRender.cpp
--- a/Engine/Source/Editor/UnrealEd/Private/Lightmass/LightmassRender.cpp
+++ b/Engine/Source/Editor/UnrealEd/Private/Lightmass/LightmassRender.cpp
@@ -400,6 +400,9 @@ public:
 			bool bIsMaterialUnlit = MaterialInterface->GetShadingModels().IsUnlit();
 			const bool bIsOpaque = IsOpaqueBlendMode(*MaterialInterface);
 			const bool bIsMasked = IsMaskedBlendMode(*MaterialInterface);
+			const bool bIsModulate = IsModulateBlendMode(*MaterialInterface);
+			const bool bIsTranslucentOnly = IsTranslucentOnlyBlendMode(*MaterialInterface);
+			const bool bIsAlphaHoldout = IsAlphaHoldoutBlendMode(*MaterialInterface);
 			check(ProxyMaterial);
 			FLightmassMaterialCompiler ProxyCompiler(Compiler);
 
@@ -442,7 +445,7 @@ public:
 					return MaterialInterface->CompileProperty(&ProxyCompiler, MP_Opacity);
 				}
 				// STRATA_TODO_BLENDMODE
-				else if (BlendMode == BLEND_Modulate)
+				else if (bIsModulate)
 				{
 					if (bIsMaterialUnlit)
 					{
@@ -454,7 +457,7 @@ public:
 					}
 				}
 				// STRATA_TODO_BLENDMODE
-				else if ((BlendMode == BLEND_Translucent) || (BlendMode == BLEND_Additive) || (BlendMode == BLEND_AlphaComposite) || (BlendMode == BLEND_AlphaHoldout))
+				else if (bIsTranslucentOnly || (BlendMode == BLEND_Additive) || (BlendMode == BLEND_AlphaComposite) || bIsAlphaHoldout)
 				{
 					int32 ColoredOpacity = INDEX_NONE;
 					if (bIsMaterialUnlit)
diff --git a/Engine/Source/Runtime/Engine/Classes/Engine/EngineTypes.h b/Engine/Source/Runtime/Engine/Classes/Engine/EngineTypes.h
--- a/Engine/Source/Runtime/Engine/Classes/Engine/EngineTypes.h
+++ b/Engine/Source/Runtime/Engine/Classes/Engine/EngineTypes.h
@@ -306,6 +306,18 @@ ENGINE_API bool IsTranslucentBlendMode(EBlendMode LegacyBlendMode, EStrataBlendM
 ENGINE_API bool IsTranslucentBlendMode(const FMaterial& In);
 ENGINE_API bool IsTranslucentBlendMode(const UMaterialInterface& In);
 
+ENGINE_API bool IsAlphaHoldoutBlendMode(EStrataBlendMode BlendMode);
+ENGINE_API bool IsAlphaHoldoutBlendMode(EBlendMode BlendMode);
+ENGINE_API bool IsAlphaHoldoutBlendMode(EBlendMode LegacyBlendMode, EStrataBlendMode StrataBlendMode);
+ENGINE_API bool IsAlphaHoldoutBlendMode(const FMaterial& In);
+ENGINE_API bool IsAlphaHoldoutBlendMode(const UMaterialInterface& In);
+
+ENGINE_API bool IsModulateBlendMode(EStrataBlendMode BlendMode);
+ENGINE_API bool IsModulateBlendMode(EBlendMode BlendMode);
+ENGINE_API bool IsModulateBlendMode(EBlendMode LegacyBlendMode, EStrataBlendMode StrataBlendMode);
+ENGINE_API bool IsModulateBlendMode(const FMaterial& In);
+ENGINE_API bool IsModulateBlendMode(const UMaterialInterface& In);
+
 /** The default float precision for material's pixel shaders on mobile devices*/
 UENUM()
 enum EMaterialFloatPrecisionMode : int
diff --git a/Engine/Source/Runtime/Engine/Private/Materials/HLSLMaterialTranslator.cpp b/Engine/Source/Runtime/Engine/Private/Materials/HLSLMaterialTranslator.cpp
--- a/Engine/Source/Runtime/Engine/Private/Materials/HLSLMaterialTranslator.cpp
+++ b/Engine/Source/Runtime/Engine/Private/Materials/HLSLMaterialTranslator.cpp
@@ -1081,14 +1081,14 @@ bool FHLSLMaterialTranslator::Translate()
 						((MaterialShadingModels.HasShadingModel(MSM_SubsurfaceProfile) && IsMaterialPropertyUsed(MP_CustomData0, Chunk[MP_CustomData0], FLinearColor(1, 0, 0, 0), 1))
 						|| (MaterialShadingModels.HasShadingModel(MSM_Eye) && IsMaterialPropertyUsed(MP_Opacity, Chunk[MP_Opacity], FLinearColor(1, 0, 0, 0), 1)));
 
-		if (BlendMode == BLEND_Modulate && MaterialShadingModels.IsLit() && !Material->IsDeferredDecal())
+		if (!bStrataEnabled && IsModulateBlendMode(BlendMode, StrataBlendMode) && MaterialShadingModels.IsLit() && !Material->IsDeferredDecal())
 		{
 			Errorf(TEXT("Dynamically lit translucency is not supported for BLEND_Modulate materials."));
 		}
 
 		if (Domain == MD_Surface)
 		{
-			if (BlendMode == BLEND_Modulate && Material->IsTranslucencyAfterDOFEnabled() && !RHISupportsDualSourceBlending(Platform))
+			if (IsModulateBlendMode(BlendMode, StrataBlendMode) && Material->IsTranslucencyAfterDOFEnabled() && !RHISupportsDualSourceBlending(Platform))
 			{
 				Errorf(TEXT("Translucency after DOF with BLEND_Modulate is only allowed on platforms that support dual-blending. Consider using BLEND_Translucent with black emissive"));
 			}
@@ -1119,7 +1119,7 @@ bool FHLSLMaterialTranslator::Translate()
 			}
 		}
 
-		if (BlendMode == BLEND_AlphaHoldout && !MaterialShadingModels.IsUnlit())
+		if (IsAlphaHoldoutBlendMode(BlendMode, StrataBlendMode) && !MaterialShadingModels.IsUnlit())
 		{
 			Errorf(TEXT("Alpha Holdout blend mode must use unlit shading model."));
 		}
@@ -1218,7 +1218,7 @@ bool FHLSLMaterialTranslator::Translate()
 			}
 		}
 
-		if (Domain == MD_DeferredDecal && !(BlendMode == BLEND_Translucent || BlendMode == BLEND_AlphaComposite || BlendMode == BLEND_Modulate))
+		if (Domain == MD_DeferredDecal && !(IsTranslucentOnlyBlendMode(BlendMode, StrataBlendMode) || BlendMode == BLEND_AlphaComposite || IsModulateBlendMode(BlendMode, StrataBlendMode)))
 		{
 			// We could make the change for the user but it would be confusing when going to DeferredDecal and back
 			// or we would have to pay a performance cost to make the change more transparently.
@@ -1242,7 +1242,7 @@ bool FHLSLMaterialTranslator::Translate()
 			}
 		}
 
-		if (BlendMode == BLEND_Modulate && Material->IsTranslucencyAfterMotionBlurEnabled())
+		if (IsModulateBlendMode(BlendMode, StrataBlendMode) && Material->IsTranslucencyAfterMotionBlurEnabled())
 		{
 			// We don't currently have a separate translucency modulation pass for After Motion Blur
 			Errorf(TEXT("Blend Mode \"Modulate\" materials are not currently supported in the \"After Motion Blur\" translucency pass."));
diff --git a/Engine/Source/Runtime/Engine/Private/Materials/Material.cpp b/Engine/Source/Runtime/Engine/Private/Materials/Material.cpp
--- a/Engine/Source/Runtime/Engine/Private/Materials/Material.cpp
+++ b/Engine/Source/Runtime/Engine/Private/Materials/Material.cpp
@@ -6643,10 +6643,10 @@ static bool IsPropertyActive_Internal(EMaterialProperty InProperty,
 			switch (InProperty)
 			{
 			case MP_Refraction:
-				Active = ((bIsTranslucentBlendMode && BlendMode != BLEND_AlphaHoldout && BlendMode != BLEND_Modulate) || ShadingModels.HasShadingModel(MSM_SingleLayerWater)) && bUsesDistortion;
+				Active = ((bIsTranslucentBlendMode && !IsAlphaHoldoutBlendMode(BlendMode, StrataBlendMode) && !IsModulateBlendMode(BlendMode, StrataBlendMode)) || ShadingModels.HasShadingModel(MSM_SingleLayerWater)) && bUsesDistortion;
 				break;
 			case MP_Opacity:
-				Active = (bIsTranslucentBlendMode && BlendMode != BLEND_Modulate) || ShadingModels.HasShadingModel(MSM_SingleLayerWater);
+				Active = (bIsTranslucentBlendMode && !IsModulateBlendMode(BlendMode, StrataBlendMode)) || ShadingModels.HasShadingModel(MSM_SingleLayerWater);
 				break;
 			case MP_OpacityMask:
 				Active = IsMaskedBlendMode(BlendMode, StrataBlendMode);
@@ -6678,10 +6678,10 @@ static bool IsPropertyActive_Internal(EMaterialProperty InProperty,
 			Active = false;
 			break;
 		case MP_Refraction:
-			Active = ((bIsTranslucentBlendMode && BlendMode != BLEND_AlphaHoldout && BlendMode != BLEND_Modulate) || ShadingModels.HasShadingModel(MSM_SingleLayerWater)) && bUsesDistortion;
+			Active = ((bIsTranslucentBlendMode && !IsAlphaHoldoutBlendMode(BlendMode, StrataBlendMode) && !IsModulateBlendMode(BlendMode, StrataBlendMode)) || ShadingModels.HasShadingModel(MSM_SingleLayerWater)) && bUsesDistortion;
 			break;
 		case MP_Opacity:
-			Active = (bIsTranslucentBlendMode && BlendMode != BLEND_Modulate) || ShadingModels.HasShadingModel(MSM_SingleLayerWater);
+			Active = (bIsTranslucentBlendMode && !IsModulateBlendMode(BlendMode, StrataBlendMode)) || ShadingModels.HasShadingModel(MSM_SingleLayerWater);
 			if (IsSubsurfaceShadingModel(ShadingModels))
 			{
 				Active = true;
@@ -6723,7 +6723,7 @@ static bool IsPropertyActive_Internal(EMaterialProperty InProperty,
 		case MP_EmissiveColor:
 			// Emissive is always active, even for light functions and post process materials, 
 			// but not for AlphaHoldout
-			Active = BlendMode != BLEND_AlphaHoldout;
+			Active = !IsAlphaHoldoutBlendMode(BlendMode, StrataBlendMode);
 			break;
 		case MP_WorldPositionOffset:
 			Active = true;
diff --git a/Engine/Source/Runtime/Engine/Private/Materials/StrataMaterial.cpp b/Engine/Source/Runtime/Engine/Private/Materials/StrataMaterial.cpp
--- a/Engine/Source/Runtime/Engine/Private/Materials/StrataMaterial.cpp
+++ b/Engine/Source/Runtime/Engine/Private/Materials/StrataMaterial.cpp
@@ -86,4 +86,18 @@ bool IsTranslucentOnlyBlendMode(EBlendMode BlendMode)											{ return BlendMo
 bool IsTranslucentOnlyBlendMode(EStrataBlendMode BlendMode)										{ return BlendMode == SBM_TranslucentColoredTransmittance || BlendMode == SBM_TranslucentGreyTransmittance; }
 bool IsTranslucentOnlyBlendMode(EBlendMode LegacyBlendMode, EStrataBlendMode StrataBlendMode)	{ return IsStrataEnabled() ? IsTranslucentOnlyBlendMode(StrataBlendMode) : IsTranslucentOnlyBlendMode(LegacyBlendMode); }
 bool IsTranslucentOnlyBlendMode(const FMaterial& In)											{ return IsStrataEnabled() ? IsTranslucentOnlyBlendMode(In.GetStrataBlendMode()) : IsTranslucentOnlyBlendMode(In.GetBlendMode()); }
-bool IsTranslucentOnlyBlendMode(const UMaterialInterface& In)									{ return IsStrataEnabled() ? IsTranslucentOnlyBlendMode(In.GetStrataBlendMode()) : IsTranslucentOnlyBlendMode(In.GetBlendMode()); }
\ No newline at end of file
+bool IsTranslucentOnlyBlendMode(const UMaterialInterface& In)									{ return IsStrataEnabled() ? IsTranslucentOnlyBlendMode(In.GetStrataBlendMode()) : IsTranslucentOnlyBlendMode(In.GetBlendMode()); }
+
+// AlphaHoldout blend mode
+bool IsAlphaHoldoutBlendMode(EBlendMode BlendMode)												{ return BlendMode == BLEND_AlphaHoldout; }
+bool IsAlphaHoldoutBlendMode(EStrataBlendMode BlendMode)										{ return BlendMode == SBM_AlphaHoldout; }
+bool IsAlphaHoldoutBlendMode(EBlendMode LegacyBlendMode, EStrataBlendMode StrataBlendMode)		{ return IsStrataEnabled() ? IsAlphaHoldoutBlendMode(StrataBlendMode) : IsAlphaHoldoutBlendMode(LegacyBlendMode); }
+bool IsAlphaHoldoutBlendMode(const FMaterial& In)												{ return IsStrataEnabled() ? IsAlphaHoldoutBlendMode(In.GetStrataBlendMode()) : IsAlphaHoldoutBlendMode(In.GetBlendMode()); }
+bool IsAlphaHoldoutBlendMode(const UMaterialInterface& In)										{ return IsStrataEnabled() ? IsAlphaHoldoutBlendMode(In.GetStrataBlendMode()) : IsAlphaHoldoutBlendMode(In.GetBlendMode()); }
+
+// Modulate blend mode
+bool IsModulateBlendMode(EBlendMode BlendMode)													{ return BlendMode == BLEND_Modulate; }
+bool IsModulateBlendMode(EStrataBlendMode BlendMode)											{ return BlendMode == SBM_ColoredTransmittanceOnly; }
+bool IsModulateBlendMode(EBlendMode LegacyBlendMode, EStrataBlendMode StrataBlendMode)			{ return IsStrataEnabled() ? IsModulateBlendMode(StrataBlendMode) : IsModulateBlendMode(LegacyBlendMode); }
+bool IsModulateBlendMode(const FMaterial& In)													{ return IsStrataEnabled() ? IsModulateBlendMode(In.GetStrataBlendMode()) : IsModulateBlendMode(In.GetBlendMode()); }
+bool IsModulateBlendMode(const UMaterialInterface& In)											{ return IsStrataEnabled() ? IsModulateBlendMode(In.GetStrataBlendMode()) : IsModulateBlendMode(In.GetBlendMode()); }
diff --git a/Engine/Source/Runtime/Engine/Private/Particles/ParticleSystemRender.cpp b/Engine/Source/Runtime/Engine/Private/Particles/ParticleSystemRender.cpp
--- a/Engine/Source/Runtime/Engine/Private/Particles/ParticleSystemRender.cpp
+++ b/Engine/Source/Runtime/Engine/Private/Particles/ParticleSystemRender.cpp
@@ -896,7 +896,7 @@ void FDynamicSpriteEmitterData::GetDynamicMeshElementsEmitter(const FParticleSys
 					const FMaterial* Material = MaterialResource ? &MaterialResource->GetIncompleteMaterialWithFallback(FeatureLevel) : nullptr;
 
 					if (Material && 
-						(Material->GetBlendMode() == BLEND_Translucent || Material->GetBlendMode() == BLEND_AlphaComposite || Material->GetBlendMode() == BLEND_AlphaHoldout ||
+						(IsTranslucentOnlyBlendMode(*Material) || Material->GetBlendMode() == BLEND_AlphaComposite || IsAlphaHoldoutBlendMode(*Material) ||
 						((SourceData->SortMode == PSORTMODE_Age_OldestFirst) || (SourceData->SortMode == PSORTMODE_Age_NewestFirst)))
 						)
 					{
diff --git a/Engine/Source/Runtime/Renderer/Private/BasePassRendering.cpp b/Engine/Source/Runtime/Renderer/Private/BasePassRendering.cpp
--- a/Engine/Source/Runtime/Renderer/Private/BasePassRendering.cpp
+++ b/Engine/Source/Runtime/Renderer/Private/BasePassRendering.cpp
@@ -1859,14 +1859,13 @@ bool FBasePassMeshProcessor::ShouldDraw(const FMaterial& Material)
 				if (AutoBeforeDOFTranslucencyBoundary > 0.0f)
 				{
 					bShouldDraw = bShouldDraw || (Material.IsTranslucencyAfterDOFEnabled() && StrataBlendMode != SBM_ColoredTransmittanceOnly);
-					bShouldDraw = bShouldDraw || (Material.IsTranslucencyAfterDOFEnabled() && (Material.IsDualBlendingEnabled(GetFeatureLevelShaderPlatform(FeatureLevel)) ||
-						BlendMode == BLEND_Modulate || StrataBlendMode == SBM_ColoredTransmittanceOnly));
+					bShouldDraw = bShouldDraw || (Material.IsTranslucencyAfterDOFEnabled() && (Material.IsDualBlendingEnabled(GetFeatureLevelShaderPlatform(FeatureLevel)) || IsModulateBlendMode(BlendMode, StrataBlendMode)));
 				}
 				break;
 
 			case ETranslucencyPass::TPT_TranslucencyStandardModulate:
 				bShouldDraw = !Material.IsTranslucencyAfterDOFEnabled() && !Material.IsTranslucencyAfterMotionBlurEnabled() 
-					&& (Material.IsDualBlendingEnabled(GetFeatureLevelShaderPlatform(FeatureLevel)) || BlendMode == BLEND_Modulate || StrataBlendMode == SBM_ColoredTransmittanceOnly);
+					&& (Material.IsDualBlendingEnabled(GetFeatureLevelShaderPlatform(FeatureLevel)) || IsModulateBlendMode(BlendMode, StrataBlendMode));
 				break;
 
 			case ETranslucencyPass::TPT_TranslucencyAfterDOF:
@@ -1926,8 +1925,7 @@ bool FBasePassMeshProcessor::TryAddMeshBatch(const FMeshBatch& RESTRICT MeshBatc
 
 		bool bIsStandardTranslucency = !Material.IsTranslucencyAfterDOFEnabled() && !Material.IsTranslucencyAfterMotionBlurEnabled();
 		bool bIsAfterDOF = Material.IsTranslucencyAfterDOFEnabled() && !Material.IsTranslucencyAfterMotionBlurEnabled() && StrataBlendMode != SBM_ColoredTransmittanceOnly;
-		bool bIsAfterDOFModulate = Material.IsTranslucencyAfterDOFEnabled() && (Material.IsDualBlendingEnabled(GetFeatureLevelShaderPlatform(FeatureLevel)) ||
-			BlendMode == BLEND_Modulate || StrataBlendMode == SBM_ColoredTransmittanceOnly);
+		bool bIsAfterDOFModulate = Material.IsTranslucencyAfterDOFEnabled() && (Material.IsDualBlendingEnabled(GetFeatureLevelShaderPlatform(FeatureLevel)) || IsModulateBlendMode(BlendMode, StrataBlendMode));
 
 		// When AutoBeforeDOFTranslucencyBoundary is valid, we automatically move After DOF translucent meshes (never blurred by DOF)
 		// before DOF if those elements are behind the focus distance.
diff --git a/Engine/Source/Runtime/Renderer/Private/PathTracing.cpp b/Engine/Source/Runtime/Renderer/Private/PathTracing.cpp
--- a/Engine/Source/Runtime/Renderer/Private/PathTracing.cpp
+++ b/Engine/Source/Runtime/Renderer/Private/PathTracing.cpp
@@ -1085,12 +1085,12 @@ static bool NeedsAnyHitShader(bool bIsMasked, EBlendMode BlendMode) // STRATA_TO
 	{
 		return true;
 	}
-	switch (BlendMode)
+	else if (IsOpaqueBlendMode(BlendMode) || IsAlphaHoldoutBlendMode(BlendMode))
 	{
-	case BLEND_Opaque:
-	case BLEND_AlphaHoldout:
 		return false;
-	default:
+	}
+	else
+	{
 		return true;
 	}
 }
