#ue5 - Fix character movement server->client corrections causing continuous corrections when landing on an object with mismatched server/client orientations until the character velocity is zero. Added option to convert server velocity to a base-relative velocity for client corrections. Added "p.NetUseBaseRelativeVelocity" to control this feature, and renamed similar "p.NetUseBaseRelativeAcceleration" to prepend "Net" to the name for clarity.

Test map info:
- Press Q/E to translate the platform
- Press R/T to rotate the platform
"p.netshowcorrections 1" to show errors with/without new cvar toggle.

diff --git a/Engine/Source/Runtime/Engine/Classes/GameFramework/Character.h b/Engine/Source/Runtime/Engine/Classes/GameFramework/Character.h
--- a/Engine/Source/Runtime/Engine/Classes/GameFramework/Character.h
+++ b/Engine/Source/Runtime/Engine/Classes/GameFramework/Character.h
@@ -146,20 +146,20 @@ namespace MovementBaseUtility
 	/** Get the tangential velocity at WorldLocation for the given component. */
 	ENGINE_API FVector GetMovementBaseTangentialVelocity(const UPrimitiveComponent* MovementBase, const FName BoneName, const FVector& WorldLocation);
 
-	/** Get the transforms for the given MovementBase, optionally at the location of a bone. Returns false if MovementBase is nullptr, or if BoneName is not a valid bone. */
+	/** Get the transform (local-to-world) for the given MovementBase, optionally at the location of a bone. Returns false if MovementBase is nullptr, or if BoneName is not a valid bone. */
 	ENGINE_API bool GetMovementBaseTransform(const UPrimitiveComponent* MovementBase, const FName BoneName, FVector& OutLocation, FQuat& OutQuat);
 
-	/** Get the world space location from a local space location for a given MovementBase. Returns false if MovementBase is nullptr, or if BoneName is not a valid bone. Scaling is ignored. */
-	ENGINE_API bool GetLocalMovementBaseLocationInWorldSpace(const UPrimitiveComponent* MovementBase, const FName BoneName, const FVector& LocalLocation, FVector& OutLocationWorldSpace);
+	/** Convert a local location to a world location for a given MovementBase. Returns false if MovementBase is nullptr, or if BoneName is not a valid bone. Scaling is ignored. */
+	ENGINE_API bool TransformLocationToWorld(const UPrimitiveComponent* MovementBase, const FName BoneName, const FVector& LocalLocation, FVector& OutLocationWorldSpace);
 
-	/** Get the local space location for a given MovementBase, optionally at the location of a bone. Returns false if MovementBase is nullptr, or if BoneName is not a valid bone. Scaling is ignored. */
-	ENGINE_API bool GetLocalMovementBaseLocation(const UPrimitiveComponent* MovementBase, const FName BoneName, const FVector& WorldSpaceLocation, FVector& OutLocalLocation);
+	/** Convert a world location to a local location for a given MovementBase, optionally at the location of a bone. Returns false if MovementBase is nullptr, or if BoneName is not a valid bone. Scaling is ignored. */
+	ENGINE_API bool TransformLocationToLocal(const UPrimitiveComponent* MovementBase, const FName BoneName, const FVector& WorldSpaceLocation, FVector& OutLocalLocation);
 
-	/** Get the world space acceleration from a local space acceleration for a given MovementBase. Returns false if MovementBase is nullptr, or if BoneName is not a valid bone. Scaling is ignored. */
-	ENGINE_API bool GetLocalMovementBaseAccelerationInWorldSpace(const UPrimitiveComponent* MovementBase, const FName BoneName, const FVector& LocalAccel, FVector& OutAccelWorldSpace);
+	/** Convert a local direction to a world direction for a given MovementBase. Returns false if MovementBase is nullptr, or if BoneName is not a valid bone. Scaling is ignored. */
+	ENGINE_API bool TransformDirectionToWorld(const UPrimitiveComponent* MovementBase, const FName BoneName, const FVector& LocalDirection, FVector& OutDirectionWorldSpace);
 
-	/** Get the local space acceleration for a given MovementBase, optionally relative to the orientation of a bone. Returns false if MovementBase is nullptr, or if BoneName is not a valid bone. Scaling is ignored. */
-	ENGINE_API bool GetLocalMovementBaseAcceleration(const UPrimitiveComponent* MovementBase, const FName BoneName, const FVector& WorldSpaceAcceleration, FVector& OutLocalAcceleration);
+	/** Convert a world direction to a local direction for a given MovementBase, optionally relative to the orientation of a bone. Returns false if MovementBase is nullptr, or if BoneName is not a valid bone. Scaling is ignored. */
+	ENGINE_API bool TransformDirectionToLocal(const UPrimitiveComponent* MovementBase, const FName BoneName, const FVector& WorldSpaceDirection, FVector& OutLocalDirection);
 }
 
 /** Struct to hold information about the "base" object the character is standing on. */
diff --git a/Engine/Source/Runtime/Engine/Classes/GameFramework/CharacterMovementReplication.h b/Engine/Source/Runtime/Engine/Classes/GameFramework/CharacterMovementReplication.h
--- a/Engine/Source/Runtime/Engine/Classes/GameFramework/CharacterMovementReplication.h
+++ b/Engine/Source/Runtime/Engine/Classes/GameFramework/CharacterMovementReplication.h
@@ -262,19 +262,21 @@ public:
 		, NewBaseBoneName(NAME_None)
 		, bAckGoodMove(false)
 		, bBaseRelativePosition(false)
+		, bBaseRelativeVelocity(false)
 		, MovementMode(0)
 	{
 	}
 
 	float TimeStamp;
 	float DeltaTime;
-	FVector NewLoc;
-	FVector NewVel;
+	FVector NewLoc; // Note: if bBaseRelativePosition is set, this is a relative location to the Movement base.
+	FVector NewVel; // Note: if bBaseRelativeVelocity is set, this is a relative velocity to the Movement base.
 	FRotator NewRot;
 	UPrimitiveComponent* NewBase;
 	FName NewBaseBoneName;
 	bool bAckGoodMove;
 	bool bBaseRelativePosition;
+	bool bBaseRelativeVelocity;
 	uint8 MovementMode;
 };
 
diff --git a/Engine/Source/Runtime/Engine/Private/Character.cpp b/Engine/Source/Runtime/Engine/Private/Character.cpp
--- a/Engine/Source/Runtime/Engine/Private/Character.cpp
+++ b/Engine/Source/Runtime/Engine/Private/Character.cpp
@@ -693,7 +693,7 @@ namespace MovementBaseUtility
 		return false;
 	}
 
-	bool GetLocalMovementBaseLocationInWorldSpace(const UPrimitiveComponent* MovementBase, const FName BoneName, const FVector& LocalLocation, FVector& OutLocationWorldSpace)
+	bool TransformLocationToWorld(const UPrimitiveComponent* MovementBase, const FName BoneName, const FVector& LocalLocation, FVector& OutLocationWorldSpace)
 	{
 		FVector OutLocation;
 		FQuat OutQuat;
@@ -702,7 +702,7 @@ namespace MovementBaseUtility
 		return bResult;
 	}
 
-	bool GetLocalMovementBaseLocation(const UPrimitiveComponent* MovementBase, const FName BoneName, const FVector& WorldSpaceLocation, FVector& OutLocalLocation)
+	bool TransformLocationToLocal(const UPrimitiveComponent* MovementBase, const FName BoneName, const FVector& WorldSpaceLocation, FVector& OutLocalLocation)
 	{
 		FVector OutLocation;
 		FQuat OutQuat;
@@ -711,21 +711,21 @@ namespace MovementBaseUtility
 		return bResult;
 	}
 
-	bool GetLocalMovementBaseAccelerationInWorldSpace(const UPrimitiveComponent* MovementBase, const FName BoneName, const FVector& LocalAccel, FVector& OutAccelWorldSpace)
+	bool TransformDirectionToWorld(const UPrimitiveComponent* MovementBase, const FName BoneName, const FVector& LocalDirection, FVector& OutDirectionWorldSpace)
 	{
 		FVector IgnoredLocation;
 		FQuat OutQuat;
 		const bool bResult = GetMovementBaseTransform(MovementBase, BoneName, IgnoredLocation, OutQuat);
-		OutAccelWorldSpace = OutQuat.RotateVector(LocalAccel);
+		OutDirectionWorldSpace = OutQuat.RotateVector(LocalDirection);
 		return bResult;
 	}
 
-	bool GetLocalMovementBaseAcceleration(const UPrimitiveComponent* MovementBase, const FName BoneName, const FVector& WorldSpaceAcceleration, FVector& OutLocalAcceleration)
+	bool TransformDirectionToLocal(const UPrimitiveComponent* MovementBase, const FName BoneName, const FVector& WorldSpaceDirection, FVector& OutLocalDirection)
 	{
 		FVector IgnoredLocation;
 		FQuat OutQuat;
 		const bool bResult = GetMovementBaseTransform(MovementBase, BoneName, IgnoredLocation, OutQuat);
-		OutLocalAcceleration = OutQuat.Inverse().RotateVector(WorldSpaceAcceleration);
+		OutLocalDirection = OutQuat.UnrotateVector(WorldSpaceDirection);
 		return bResult;
 	}
 
diff --git a/Engine/Source/Runtime/Engine/Private/Components/CharacterMovementComponent.cpp b/Engine/Source/Runtime/Engine/Private/Components/CharacterMovementComponent.cpp
--- a/Engine/Source/Runtime/Engine/Private/Components/CharacterMovementComponent.cpp
+++ b/Engine/Source/Runtime/Engine/Private/Components/CharacterMovementComponent.cpp
@@ -256,12 +256,18 @@ namespace CharacterMovementCVars
 	 * This also may be a good option if using dynamic movement bases that become the player's primary frame of visual reference, such as when walking on a large boat or airship.
 	 * Not compatible with deprecated move RPCs. @see NetUsePackedMovementRPCs
 	 */
-	static int32 UseBaseRelativeAcceleration = 1;
-	FAutoConsoleVariableRef CVarUseBaseRelativeAcceleration(
-		TEXT("p.UseBaseRelativeAcceleration"),
-		UseBaseRelativeAcceleration,
+	static int32 NetUseBaseRelativeAcceleration = 1;
+	FAutoConsoleVariableRef CVarNetUseBaseRelativeAcceleration(
+		TEXT("p.NetUseBaseRelativeAcceleration"),
+		NetUseBaseRelativeAcceleration,
 		TEXT("If enabled, character acceleration will be treated as relative to dynamic movement bases."));
 
+	static int32 NetUseBaseRelativeVelocity = 1;
+	FAutoConsoleVariableRef CVarNetUseBaseRelativeVelocity(
+		TEXT("p.NetUseBaseRelativeVelocity"),
+		NetUseBaseRelativeVelocity,
+		TEXT("If enabled, character velocity corrections will be treated as relative to dynamic movement bases."));
+
 	static int32 UseTargetVelocityOnImpact = 1;
 	FAutoConsoleVariableRef CVarUseTargetVelocityOnImpact(
 		TEXT("p.UseTargetVelocityOnImpact"),
@@ -2787,7 +2793,7 @@ void UCharacterMovementComponent::SaveBaseLocation()
 		{
 			// Relative Location
 			FVector RelativeLocation;
-			MovementBaseUtility::GetLocalMovementBaseLocation(MovementBase, CharacterOwner->GetBasedMovement().BoneName, UpdatedComponent->GetComponentLocation(), RelativeLocation);
+			MovementBaseUtility::TransformLocationToLocal(MovementBase, CharacterOwner->GetBasedMovement().BoneName, UpdatedComponent->GetComponentLocation(), RelativeLocation);
 
 			// Rotation
 			if (bIgnoreBaseRotation)
@@ -9023,7 +9029,7 @@ void FCharacterNetworkMoveData::ClientFillNetworkMoveData(const FSavedMove_Chara
 		UPrimitiveComponent* ClientMovementBase = ClientMove.EndBase.Get();
 
 		const bool bSendBaseRelativeLocation     = MovementBaseUtility::UseRelativeLocation(ClientMovementBase);
-		const bool bSendBaseRelativeAcceleration = CharacterMovementCVars::UseBaseRelativeAcceleration && bSendBaseRelativeLocation;
+		const bool bSendBaseRelativeAcceleration = CharacterMovementCVars::NetUseBaseRelativeAcceleration && bSendBaseRelativeLocation;
 
 		const FVector SendLocation     = bSendBaseRelativeLocation ? ClientMove.SavedRelativeLocation : FRepMovement::RebaseOntoZeroOrigin(ClientMove.SavedLocation, ClientMove.CharacterOwner->GetCharacterMovement());
 		const FVector SendAcceleration = bSendBaseRelativeAcceleration ? ClientMove.SavedRelativeAcceleration : ClientMove.Acceleration;
@@ -9172,9 +9178,9 @@ void UCharacterMovementComponent::ServerMove_PerformMovement(const FCharacterNet
 	FVector ClientAccel = MoveData.Acceleration;
 
 	// Convert the move's acceleration to worldspace if necessary
-	if (CharacterMovementCVars::UseBaseRelativeAcceleration && MovementBaseUtility::IsDynamicBase(MoveData.MovementBase))
+	if (CharacterMovementCVars::NetUseBaseRelativeAcceleration && MovementBaseUtility::IsDynamicBase(MoveData.MovementBase))
 	{
-		MovementBaseUtility::GetLocalMovementBaseAccelerationInWorldSpace(MoveData.MovementBase, MoveData.MovementBaseBoneName, MoveData.Acceleration, ClientAccel);
+		MovementBaseUtility::TransformDirectionToWorld(MoveData.MovementBase, MoveData.MovementBaseBoneName, MoveData.Acceleration, ClientAccel);
 	}
 
 	const uint8 ClientMoveFlags = MoveData.CompressedMoveFlags;
@@ -9391,7 +9397,7 @@ void UCharacterMovementComponent::ServerMoveHandleClientError(float ClientTimeSt
 	FVector ClientLoc = RelativeClientLoc;
 	if (MovementBaseUtility::UseRelativeLocation(ClientMovementBase))
 	{
-		MovementBaseUtility::GetLocalMovementBaseLocationInWorldSpace(ClientMovementBase, ClientBaseBoneName, RelativeClientLoc, ClientLoc);
+		MovementBaseUtility::TransformLocationToWorld(ClientMovementBase, ClientBaseBoneName, RelativeClientLoc, ClientLoc);
 	}
 	else
 	{
@@ -9498,7 +9504,7 @@ void UCharacterMovementComponent::ServerMoveHandleClientError(float ClientTimeSt
 				if (MovementBaseUtility::UseRelativeLocation(LastServerMovementBasePtr))
 				{
 					// Relative Location
-					MovementBaseUtility::GetLocalMovementBaseLocation(LastServerMovementBasePtr, LastServerMovementBaseBoneName, UpdatedComponent->GetComponentLocation(), RelativeLocation);
+					MovementBaseUtility::TransformLocationToLocal(LastServerMovementBasePtr, LastServerMovementBaseBoneName, UpdatedComponent->GetComponentLocation(), RelativeLocation);
 					bUseLastBase = true;
 				}
 			}
@@ -9575,9 +9581,11 @@ void UCharacterMovementComponent::ServerMoveHandleClientError(float ClientTimeSt
 		ServerData->PendingAdjustment.NewRot = UpdatedComponent->GetComponentRotation();
 
 		ServerData->PendingAdjustment.bBaseRelativePosition = (bDeferServerCorrectionsWhenFalling && bUseLastBase) || MovementBaseUtility::UseRelativeLocation(MovementBase);
+		ServerData->PendingAdjustment.bBaseRelativeVelocity = false;
+		
+		// Relative location?
 		if (ServerData->PendingAdjustment.bBaseRelativePosition)
 		{
-			// Relative location
 			if (bDeferServerCorrectionsWhenFalling && bUseLastBase)
 			{
 				ServerData->PendingAdjustment.NewVel = RelativeVelocity;
@@ -9588,6 +9596,13 @@ void UCharacterMovementComponent::ServerMoveHandleClientError(float ClientTimeSt
 			else
 			{
 				ServerData->PendingAdjustment.NewLoc = CharacterOwner->GetBasedMovement().Location;
+				if (CharacterMovementCVars::NetUseBaseRelativeVelocity)
+				{
+					// Store world velocity converted to local space of movement base
+					ServerData->PendingAdjustment.bBaseRelativeVelocity = true;
+					const FVector CurrentVelocity = ServerData->PendingAdjustment.NewVel;
+					MovementBaseUtility::TransformDirectionToLocal(MovementBase, MovementBaseBoneName, CurrentVelocity, ServerData->PendingAdjustment.NewVel);
+				}
 			}
 			
 			// TODO: this could be a relative rotation, but all client corrections ignore rotation right now except the root motion one, which would need to be updated.
@@ -10034,6 +10049,7 @@ bool FCharacterMoveResponseDataContainer::Serialize(UCharacterMovementComponent&
 		SerializeOptionalValue<FName>(bIsSaving, Ar, ClientAdjustment.NewBaseBoneName, NAME_None);
 		SerializeOptionalValue<uint8>(bIsSaving, Ar, ClientAdjustment.MovementMode, MOVE_Walking);
 		Ar.SerializeBits(&ClientAdjustment.bBaseRelativePosition, 1);
+		Ar.SerializeBits(&ClientAdjustment.bBaseRelativeVelocity, 1);
 
 		if (bRootMotionMontageCorrection)
 		{
@@ -10352,13 +10368,21 @@ void UCharacterMovementComponent::ClientAdjustPosition_Implementation
 	//  Received Location is relative to dynamic base
 	if (bBaseRelativePosition)
 	{
-		MovementBaseUtility::GetLocalMovementBaseLocationInWorldSpace(NewBase, NewBaseBoneName, NewLocation, WorldShiftedNewLocation); // TODO: error handling if returns false	
+		MovementBaseUtility::TransformLocationToWorld(NewBase, NewBaseBoneName, NewLocation, WorldShiftedNewLocation); // TODO: error handling if returns false	
 	}
 	else
 	{
 		WorldShiftedNewLocation = FRepMovement::RebaseOntoLocalOrigin(NewLocation, this);
 	}
 
+	// Server's world velocity may need to be converted to velocity relative to the movement base orientation, if the base orientations don't match.
+	const FCharacterMoveResponseDataContainer& ResponseDataContainer = GetMoveResponseDataContainer();
+	if (ResponseDataContainer.ClientAdjustment.bBaseRelativeVelocity)
+	{
+		// Convert Relative Velocity -> World Velocity
+		const FVector CurrentVelocity = NewVelocity;
+		MovementBaseUtility::TransformDirectionToWorld(NewBase, NewBaseBoneName, CurrentVelocity, NewVelocity);
+	}
 
 	// Trigger event
 	OnClientCorrectionReceived(*ClientData, TimeStamp, WorldShiftedNewLocation, NewVelocity, NewBase, NewBaseBoneName, bHasBase, bBaseRelativePosition, ServerMovementMode);
@@ -11803,9 +11827,9 @@ void FSavedMove_Character::PostUpdate(ACharacter* Character, FSavedMove_Characte
 		}
 
 		// Save off movement base-relative acceleration if needed
-		if (CharacterMovementCVars::UseBaseRelativeAcceleration && MovementBaseUtility::IsDynamicBase(MovementBase))
+		if (CharacterMovementCVars::NetUseBaseRelativeAcceleration && MovementBaseUtility::IsDynamicBase(MovementBase))
 		{
-			MovementBaseUtility::GetLocalMovementBaseAcceleration(EndBase.Get(), EndBoneName, Acceleration, SavedRelativeAcceleration);
+			MovementBaseUtility::TransformDirectionToLocal(EndBase.Get(), EndBoneName, Acceleration, SavedRelativeAcceleration);
 		}
 
 		// Attachment state
@@ -11933,7 +11957,7 @@ FVector FSavedMove_Character::GetRevertedLocation() const
 	if (MovementBaseUtility::UseRelativeLocation(MovementBase))
 	{
 		FVector WorldSpacePosition;
-		MovementBaseUtility::GetLocalMovementBaseLocationInWorldSpace(MovementBase, StartBoneName, StartRelativeLocation, WorldSpacePosition);
+		MovementBaseUtility::TransformLocationToWorld(MovementBase, StartBoneName, StartRelativeLocation, WorldSpacePosition);
 		return WorldSpacePosition;
 	}
 
