Fixes for Strata to compile in some streams

diff --git a/Engine/Source/Runtime/Engine/Private/Materials/Material.cpp b/Engine/Source/Runtime/Engine/Private/Materials/Material.cpp
--- a/Engine/Source/Runtime/Engine/Private/Materials/Material.cpp
+++ b/Engine/Source/Runtime/Engine/Private/Materials/Material.cpp
@@ -3442,7 +3442,7 @@ void UMaterial::ConvertMaterialToStrataMaterial()
 					{
 						// Reconnect the shading model expression. 
 						// Note: assign the expression directly, as using ConvertNode->GetInput(19)->Connect(..) causes the expression to not be assigned
-						ConvertNode->ShadingModel.Expression = EditorOnly->ShadingModelFromMaterialExpression.Expression;
+						ConvertNode->ShadingModel.Connect(EditorOnly->ShadingModelFromMaterialExpression.OutputIndex, EditorOnly->ShadingModelFromMaterialExpression.Expression);
 					}
 
 					// Store strata shading model of the converted material. 
