Fixed a bug where a cvar was being read with GetValueOnGameThread from a function that only runs on the Render Thread

diff --git a/Engine/Source/Runtime/Engine/Private/DynamicResolution.cpp b/Engine/Source/Runtime/Engine/Private/DynamicResolution.cpp
--- a/Engine/Source/Runtime/Engine/Private/DynamicResolution.cpp
+++ b/Engine/Source/Runtime/Engine/Private/DynamicResolution.cpp
@@ -315,7 +315,7 @@ void FDynamicResolutionHeuristicProxy::RefreshCurrentFrameResolutionFraction_Ren
 			float MedianFrameTime = SortedFrameTime[SortedFrameTime.Num() / 2];
 
 			// When CPU bound and vsync is enabled, the GPU may end up underused due to vsync wait, so we can crank up the available GPU time a bit more.
-			if (float RoundUpToVSyncError = CVarDynamicFrameTimeRoundUpToVsync.GetValueOnGameThread() > 0.0f && VSyncCVar->GetInt())
+			if (float RoundUpToVSyncError = CVarDynamicFrameTimeRoundUpToVsync.GetValueOnRenderThread() > 0.0f && VSyncCVar->GetInt())
 			{
 				const float VSyncFrameTime = 1000.0f / float(FPlatformMisc::GetMaxRefreshRate());
 				float VSyncRateChangeMultiplier = 1.0f + RoundUpToVSyncError / 100.f;
