Do not cache the MeshDescription triangles/vertices count when we load the mesh description. Setting the cache when building the renderdata is good enough and it can use the already computed "FOverlappingCorners".

#rnx

diff --git a/Engine/Source/Developer/MeshBuilder/Private/StaticMeshBuilder.cpp b/Engine/Source/Developer/MeshBuilder/Private/StaticMeshBuilder.cpp
--- a/Engine/Source/Developer/MeshBuilder/Private/StaticMeshBuilder.cpp
+++ b/Engine/Source/Developer/MeshBuilder/Private/StaticMeshBuilder.cpp
@@ -407,8 +407,17 @@ bool FStaticMeshBuilder::Build(FStaticMeshRenderData& StaticMeshRenderData, USta
 		if (bIsMeshDescriptionValid)
 		{
 			MeshDescriptionHelper.SetupRenderMeshDescription(StaticMesh, MeshDescriptions[LodIndex], false, true);
+			//Make sure the cache is good before looking for the active reduction
+			if (SrcModel.CacheMeshDescriptionTrianglesCount == MAX_uint32)
+			{
+				SrcModel.CacheMeshDescriptionTrianglesCount = static_cast<uint32>(MeshDescriptions[LodIndex].Triangles().Num());
+			}
+			if (SrcModel.CacheMeshDescriptionVerticesCount == MAX_uint32)
+			{
+				SrcModel.CacheMeshDescriptionVerticesCount = static_cast<uint32>(FStaticMeshOperations::GetUniqueVertexCount(MeshDescriptions[LodIndex], MeshDescriptionHelper.GetOverlappingCorners()));
+			}
 			//Get back the reduction status once we apply all build settings, vertex count can change depending on the build settings
-			bUseReduction = StaticMesh->IsReductionActive(LodIndex, MeshDescriptions[LodIndex]);
+			bUseReduction = StaticMesh->IsReductionActive(LodIndex);
 		}
 		else
 		{
diff --git a/Engine/Source/Runtime/Engine/Classes/Engine/StaticMesh.h b/Engine/Source/Runtime/Engine/Classes/Engine/StaticMesh.h
--- a/Engine/Source/Runtime/Engine/Classes/Engine/StaticMesh.h
+++ b/Engine/Source/Runtime/Engine/Classes/Engine/StaticMesh.h
@@ -1948,8 +1948,6 @@ public:
 	/* Return true if the reduction settings are setup to reduce a LOD*/
 	ENGINE_API bool IsReductionActive(int32 LODIndex) const;
 
-	ENGINE_API bool IsReductionActive(int32 LodIndex, const FMeshDescription& MeshDescription) const;
-
 	/* Get a copy of the reduction settings for a specified LOD index. */
 	ENGINE_API struct FMeshReductionSettings GetReductionSettings(int32 LODIndex) const;
 
diff --git a/Engine/Source/Runtime/Engine/Private/StaticMesh.cpp b/Engine/Source/Runtime/Engine/Private/StaticMesh.cpp
--- a/Engine/Source/Runtime/Engine/Private/StaticMesh.cpp
+++ b/Engine/Source/Runtime/Engine/Private/StaticMesh.cpp
@@ -2476,25 +2476,6 @@ bool UStaticMesh::IsReductionActive(int32 LodIndex) const
 	return bReductionActive;
 }
 
-bool UStaticMesh::IsReductionActive(int32 LodIndex, const FMeshDescription& MeshDescription) const
-{
-	//Invalid LOD are not reduced
-	if (!IsSourceModelValid(LodIndex))
-	{
-		return false;
-	}
-	bool bReductionActive = false;
-	if (IMeshReduction* ReductionModule = FModuleManager::Get().LoadModuleChecked<IMeshReductionManagerModule>("MeshReductionInterface").GetStaticMeshReductionInterface())
-	{
-		FStaticMeshSourceModel* MutableSourceModel = const_cast<FStaticMeshSourceModel*>(&GetSourceModel(LodIndex));
-		MutableSourceModel->CacheMeshDescriptionVerticesCount = FStaticMeshOperations::GetUniqueVertexCount(MeshDescription);
-		MutableSourceModel->CacheMeshDescriptionTrianglesCount = MeshDescription.Triangles().Num();
-		FMeshReductionSettings ReductionSettings = GetReductionSettings(LodIndex);
-		bReductionActive = ReductionModule->IsReductionActive(ReductionSettings, MutableSourceModel->CacheMeshDescriptionVerticesCount, MutableSourceModel->CacheMeshDescriptionTrianglesCount);
-	}
-	return bReductionActive;
-}
-
 FMeshReductionSettings UStaticMesh::GetReductionSettings(int32 LODIndex) const
 {
 	check(IsSourceModelValid(LODIndex));
diff --git a/Engine/Source/Runtime/Engine/Private/StaticMeshSourceData.cpp b/Engine/Source/Runtime/Engine/Private/StaticMeshSourceData.cpp
--- a/Engine/Source/Runtime/Engine/Private/StaticMeshSourceData.cpp
+++ b/Engine/Source/Runtime/Engine/Private/StaticMeshSourceData.cpp
@@ -209,11 +209,6 @@ bool FStaticMeshSourceModel::LoadMeshDescription(FMeshDescription& OutMeshDescri
 	{
 		// Unpack MeshDescription from the bulk data which was deserialized
 		StaticMeshDescriptionBulkData->GetBulkData().LoadMeshDescription(OutMeshDescription);
-		
-		//Set the tris/verts count cache in case we load an asset that was save before the cache exist
-		FStaticMeshSourceModel* MutableThis = const_cast<FStaticMeshSourceModel*>(this);
-		MutableThis->CacheMeshDescriptionTrianglesCount = static_cast<uint32>(OutMeshDescription.Triangles().Num());
-		MutableThis->CacheMeshDescriptionVerticesCount = static_cast<uint32>(FStaticMeshOperations::GetUniqueVertexCount(OutMeshDescription));
 		return true;
 	}
 
diff --git a/Engine/Source/Runtime/StaticMeshDescription/Private/StaticMeshOperations.cpp b/Engine/Source/Runtime/StaticMeshDescription/Private/StaticMeshOperations.cpp
--- a/Engine/Source/Runtime/StaticMeshDescription/Private/StaticMeshOperations.cpp
+++ b/Engine/Source/Runtime/StaticMeshDescription/Private/StaticMeshOperations.cpp
@@ -2609,9 +2609,14 @@ namespace UE::Private
 
 int32 FStaticMeshOperations::GetUniqueVertexCount(const FMeshDescription& MeshDescription)
 {
-	constexpr uint32 NumTexCoords = MAX_MESH_TEXTURE_COORDS_MD;
 	FOverlappingCorners OverlappingCorners;
 	FStaticMeshOperations::FindOverlappingCorners(OverlappingCorners, MeshDescription, UE_THRESH_POINTS_ARE_SAME);
+	return GetUniqueVertexCount(MeshDescription, OverlappingCorners);
+}
+
+int32 FStaticMeshOperations::GetUniqueVertexCount(const FMeshDescription& MeshDescription, const FOverlappingCorners& OverlappingCorners)
+{
+	constexpr uint32 NumTexCoords = MAX_MESH_TEXTURE_COORDS_MD;
 	TArray< UE::Private::FPrivateVertexInfo > Verts;
 	Verts.Reserve(MeshDescription.Vertices().Num());
 	TMap< int32, int32 > VertsMap;
diff --git a/Engine/Source/Runtime/StaticMeshDescription/Public/StaticMeshOperations.h b/Engine/Source/Runtime/StaticMeshDescription/Public/StaticMeshOperations.h
--- a/Engine/Source/Runtime/StaticMeshDescription/Public/StaticMeshOperations.h
+++ b/Engine/Source/Runtime/StaticMeshDescription/Public/StaticMeshOperations.h
@@ -188,4 +188,5 @@ public:
 	 * 
 	 */
 	static int32 GetUniqueVertexCount(const FMeshDescription& MeshDescription);
+	static int32 GetUniqueVertexCount(const FMeshDescription& MeshDescription, const FOverlappingCorners& OverlappingCorners);
 };
