Avoid potential deref of temporary string address

diff --git a/Engine/Source/Runtime/Engine/Private/UnrealNetwork.cpp b/Engine/Source/Runtime/Engine/Private/UnrealNetwork.cpp
--- a/Engine/Source/Runtime/Engine/Private/UnrealNetwork.cpp
+++ b/Engine/Source/Runtime/Engine/Private/UnrealNetwork.cpp
@@ -87,7 +87,9 @@ void RegisterReplicatedLifetimeProperty(
 		return;
 	}
 
-	RegisterReplicatedLifetimeProperty(NetworkingPrivate::FRepPropertyDescriptor(ReplicatedProperty), OutLifetimeProps, Params);
+	const FString ReplicatedPropertyName = ReplicatedProperty->GetName();
+	NetworkingPrivate::FRepPropertyDescriptor PropDesc(*ReplicatedPropertyName, ReplicatedProperty->RepIndex, ReplicatedProperty->ArrayDim);
+	RegisterReplicatedLifetimeProperty(PropDesc, OutLifetimeProps, Params);
 }
 
 void SetReplicatedPropertyToDisabled(const NetworkingPrivate::FRepPropertyDescriptor& PropertyDescriptor, TArray<FLifetimeProperty>& OutLifetimeProps)
@@ -110,7 +112,9 @@ void SetReplicatedPropertyToDisabled(const NetworkingPrivate::FRepPropertyDescri
 
 void SetReplicatedPropertyToDisabled(const FProperty* ReplicatedProperty, TArray<FLifetimeProperty>& OutLifetimeProps)
 {
-	SetReplicatedPropertyToDisabled(NetworkingPrivate::FRepPropertyDescriptor(ReplicatedProperty), OutLifetimeProps);
+	const FString ReplicatedPropertyName = ReplicatedProperty->GetName();
+	NetworkingPrivate::FRepPropertyDescriptor PropDesc(*ReplicatedPropertyName, ReplicatedProperty->RepIndex, ReplicatedProperty->ArrayDim);
+	SetReplicatedPropertyToDisabled(PropDesc, OutLifetimeProps);
 }
 
 void DisableReplicatedLifetimeProperty(const NetworkingPrivate::FRepPropertyDescriptor& PropertyDescriptor, TArray<FLifetimeProperty>& OutLifetimeProps)
@@ -126,7 +130,9 @@ void DisableReplicatedLifetimeProperty(const UClass* ThisClass, const UClass* Pr
 		return;
 	}
 
-	SetReplicatedPropertyToDisabled(NetworkingPrivate::FRepPropertyDescriptor(ReplicatedProperty), OutLifetimeProps);
+	const FString ReplicatedPropertyName = ReplicatedProperty->GetName();
+	NetworkingPrivate::FRepPropertyDescriptor PropDesc(*ReplicatedPropertyName, ReplicatedProperty->RepIndex, ReplicatedProperty->ArrayDim);
+	SetReplicatedPropertyToDisabled(PropDesc, OutLifetimeProps);
 }
 
 void ResetReplicatedLifetimeProperty(
@@ -159,7 +165,9 @@ void ResetReplicatedLifetimeProperty(const UClass* ThisClass, const UClass* Prop
 		return;
 	}
 
-	ResetReplicatedLifetimeProperty(NetworkingPrivate::FRepPropertyDescriptor(ReplicatedProperty), LifetimeCondition, OutLifetimeProps);
+	const FString ReplicatedPropertyName = ReplicatedProperty->GetName();
+	NetworkingPrivate::FRepPropertyDescriptor PropDesc(*ReplicatedPropertyName, ReplicatedProperty->RepIndex, ReplicatedProperty->ArrayDim);
+	ResetReplicatedLifetimeProperty(PropDesc, LifetimeCondition, OutLifetimeProps);
 }
 
 void ResetReplicatedLifetimeProperty(const UClass* ThisClass, const UClass* PropertyClass, FName PropertyName, const FDoRepLifetimeParams& Params, TArray< FLifetimeProperty >& OutLifetimeProps)
@@ -170,7 +178,9 @@ void ResetReplicatedLifetimeProperty(const UClass* ThisClass, const UClass* Prop
 		return;
 	}
 
-	ResetReplicatedLifetimeProperty(NetworkingPrivate::FRepPropertyDescriptor(ReplicatedProperty), Params, OutLifetimeProps);
+	const FString ReplicatedPropertyName = ReplicatedProperty->GetName();
+	NetworkingPrivate::FRepPropertyDescriptor PropDesc(*ReplicatedPropertyName, ReplicatedProperty->RepIndex, ReplicatedProperty->ArrayDim);
+	ResetReplicatedLifetimeProperty(PropDesc, Params, OutLifetimeProps);
 }
 
 void ResetReplicatedLifetimeProperty(const NetworkingPrivate::FRepPropertyDescriptor& PropertyDescriptor, const FDoRepLifetimeParams& Params, TArray< FLifetimeProperty >& OutLifetimeProps)
diff --git a/Engine/Source/Runtime/Engine/Public/Net/UnrealNetwork.h b/Engine/Source/Runtime/Engine/Public/Net/UnrealNetwork.h
--- a/Engine/Source/Runtime/Engine/Public/Net/UnrealNetwork.h
+++ b/Engine/Source/Runtime/Engine/Public/Net/UnrealNetwork.h
@@ -155,8 +155,9 @@ namespace NetworkingPrivate
 {
 	struct ENGINE_API FRepPropertyDescriptor
 	{
+		UE_DEPRECATED(5.2, "No longer used")
 		FRepPropertyDescriptor(const FProperty* Property)
-			: PropertyName(VerifyPropertyAndGetName(Property))
+			: PropertyName(TEXT(""))
 			, RepIndex(Property->RepIndex)
 			, ArrayDim(Property->ArrayDim)
 		{
@@ -174,13 +175,6 @@ namespace NetworkingPrivate
 		const int32 ArrayDim;
 
 	private:
-
-		static const TCHAR* VerifyPropertyAndGetName(const FProperty* Property)
-		{
-			check(Property);
-			return *(Property->GetName());
-		}
-
 		UE_NONCOPYABLE(FRepPropertyDescriptor);
 
 		void* operator new(size_t) = delete;
@@ -426,8 +420,8 @@ DisableAllReplicatedPropertiesOfClass(StaticClass(), c::StaticClass(), SuperClas
 #define DISABLE_ALL_CLASS_REPLICATED_PROPERTIES_FAST(c, SuperClassBehavior) \
 { \
 	static const bool bIsValid_##c_##v = ValidateReplicatedClassInheritance(StaticClass(), c::StaticClass(), TEXT("DISABLE_ALL_CLASS_REPLICATED_PROPERTIES")); \
-	const TCHAR* DoRepPropertyName_##c(TEXT(#c)); \
-	const NetworkingPrivate::FRepClassDescriptor ClassDescriptor_##c(DoRepPropertyName_##c, (int32)c::ENetFields_Private::NETFIELD_REP_START, (int32)c::ENetFields_Private::NETFIELD_REP_END); \
+	const TCHAR* DoRepClassName_##c(TEXT(#c)); \
+	const NetworkingPrivate::FRepClassDescriptor ClassDescriptor_##c(DoRepClassName_##c, (int32)c::ENetFields_Private::NETFIELD_REP_START, (int32)c::ENetFields_Private::NETFIELD_REP_END); \
 	DisableAllReplicatedPropertiesOfClass(ClassDescriptor_##c, SuperClassBehavior, OutLifetimeProps); \
 }	
 
