Adding missing include.

diff --git a/Engine/Source/Runtime/Engine/Private/RayTracingInstance.cpp b/Engine/Source/Runtime/Engine/Private/RayTracingInstance.cpp
--- a/Engine/Source/Runtime/Engine/Private/RayTracingInstance.cpp
+++ b/Engine/Source/Runtime/Engine/Private/RayTracingInstance.cpp
@@ -9,6 +9,7 @@
 #if RHI_RAYTRACING
 
 #include "RayTracingDefinitions.h"
+#include "MaterialShared.h"
 
 void FRayTracingInstance::BuildInstanceMaskAndFlags(ERHIFeatureLevel::Type FeatureLevel, ERayTracingInstanceLayer InstanceLayer, uint8 ExtraMask)
 {
