Reissue shader compile jobs on OOM error.

#rnx

diff --git a/Engine/Source/Programs/ShaderCompileWorker/Private/ShaderCompileWorker.cpp b/Engine/Source/Programs/ShaderCompileWorker/Private/ShaderCompileWorker.cpp
--- a/Engine/Source/Programs/ShaderCompileWorker/Private/ShaderCompileWorker.cpp
+++ b/Engine/Source/Programs/ShaderCompileWorker/Private/ShaderCompileWorker.cpp
@@ -333,6 +333,9 @@ static int64 WriteOutputFileHeader(FArchive& OutputFile, int32 ErrorCode, int32
 		}
 	}
 
+	// Reset error code as it can be receive a new value now
+	FSCWErrorCode::Reset();
+
 	UpdateFileSize(OutputFile, FileSizePosition);
 	return FileSizePosition;
 }
@@ -1263,6 +1266,14 @@ static int32 GuardedMain(int32 argc, TCHAR* argv[], bool bDirectMode)
 	FString IniBootstrapFilename;
 	FString ModulesBootstrapFilename;
 
+	// Register out-of-memory delegate to report error code on exit
+	FCoreDelegates::GetOutOfMemoryDelegate().AddLambda(
+		[]()
+		{
+			FSCWErrorCode::Report(FSCWErrorCode::OutOfMemory);
+		}
+	);
+
 	if (IsUsingXGE())
 	{
 		// Tie the bootstrap filenames to the xge job id to refresh bootstraps state every time a new build starts
diff --git a/Engine/Source/Runtime/Engine/Private/ShaderCompiler/ShaderCompiler.cpp b/Engine/Source/Runtime/Engine/Private/ShaderCompiler/ShaderCompiler.cpp
--- a/Engine/Source/Runtime/Engine/Private/ShaderCompiler/ShaderCompiler.cpp
+++ b/Engine/Source/Runtime/Engine/Private/ShaderCompiler/ShaderCompiler.cpp
@@ -944,9 +944,11 @@ static constexpr int32 GSingleThreadedRunsMaxCount = (1 << 24);
 
 static void ModalErrorOrLog(const FString& Text, int64 CurrentFilePos = 0, int64 ExpectedFileSize = 0)
 {
+	static FThreadSafeBool bModalReported;
+
 	FString BadFile;
 	if (CurrentFilePos > ExpectedFileSize)
-{
+	{
 		// Corrupt file
 		BadFile = FString::Printf(TEXT("(Truncated or corrupt output file! Current file pos %lld, file size %lld)"), CurrentFilePos, ExpectedFileSize);
 	}
@@ -954,9 +956,17 @@ static void ModalErrorOrLog(const FString& Text, int64 CurrentFilePos = 0, int64
 	if (FPlatformProperties::SupportsWindowedMode())
 	{
 		UE_LOG(LogShaderCompilers, Error, TEXT("%s%s"), *Text, *BadFile);
-		FMessageDialog::Open(EAppMsgType::Ok, FText::FromString(Text));
-		FPlatformMisc::RequestExit(false);
-		return;
+		if (!bModalReported.AtomicSet(true))
+		{
+			// Show dialog box with error message and request exit
+			FMessageDialog::Open(EAppMsgType::Ok, FText::FromString(Text));
+			FPlatformMisc::RequestExit(false);
+		}
+		else
+		{
+			// Another thread already opened a dialog box and requests exit
+			FPlatformProcess::SleepInfinite();
+		}
 	}
 	else
 	{
@@ -1247,9 +1257,66 @@ namespace ShaderCompilerCookStats
 }
 #endif
 
+static void ReissueShaderCompileJobs(const TArray<FShaderCommonCompileJob*>& SourceJobs)
+{
+	if (SourceJobs.Num())
+	{
+		TRACE_CPUPROFILER_EVENT_SCOPE(ReissueShaderCompileJobs);
+
+		TArray<FShaderCommonCompileJobPtr> ReissueJobs;
+		ReissueJobs.Reserve(SourceJobs.Num());
+		const uint32 JobId = FShaderCommonCompileJob::GetNextJobId();
+		for (const FShaderCommonCompileJob* SourceJob : SourceJobs)
+		{
+			if (const FShaderCompileJob* SingleSourceJob = SourceJob->GetSingleShaderJob())
+			{
+				if (FShaderCompileJob* ReissueJob = GShaderCompilingManager->PrepareShaderCompileJob(JobId, SingleSourceJob->Key, SingleSourceJob->Priority))
+				{
+					ReissueJob->Input = SingleSourceJob->Input;
+					ReissueJobs.Add(FShaderCommonCompileJobPtr(ReissueJob));
+				}
+			}
+			else if (const FShaderPipelineCompileJob* PipelineSourceJob = SourceJob->GetShaderPipelineJob())
+			{
+				if (FShaderPipelineCompileJob* ReissueJob = GShaderCompilingManager->PreparePipelineCompileJob(JobId, PipelineSourceJob->Key, PipelineSourceJob->Priority))
+				{
+					ReissueJob->StageJobs = PipelineSourceJob->StageJobs;
+					ReissueJobs.Add(FShaderCommonCompileJobPtr(ReissueJob));
+				}
+			}
+			else
+			{
+				checkf(0, TEXT("Reissued shader compile job is neither a single nor a pipeline job"));
+			}
+		}
+
+		GShaderCompilingManager->SubmitJobs(ReissueJobs, FString(""), FString(""));
+	}
+}
+
 // Make functions so the crash reporter can disambiguate the actual error because of the different callstacks
 namespace ShaderCompileWorkerError
 {
+	static bool TryReissueShaderCompileJobs(const TArray<FShaderCommonCompileJobPtr>& SourceJobs)
+	{
+		// Reschedule all queued shader compile jobs when a worker crashed due to running out of memory.
+		TArray<FShaderCommonCompileJob*> ReissueSourceJobs;
+		ReissueSourceJobs.Reserve(SourceJobs.Num());
+		for (const FShaderCommonCompileJobPtr& Job : SourceJobs)
+		{
+			// Switch to local thread. If the worker already is a local thread, we cannot recover from this error.
+			FShaderCommonCompileJob* CurrentJob = Job.GetReference();
+			if (CurrentJob->CurrentWorker == EShaderCompilerWorkerType::LocalThread)
+			{
+				return false;
+			}
+			CurrentJob->CurrentWorker = EShaderCompilerWorkerType::LocalThread;
+			ReissueSourceJobs.Add(CurrentJob);
+		}
+		ReissueShaderCompileJobs(ReissueSourceJobs);
+		return true;
+	}
+
 	void HandleGeneralCrash(const TCHAR* ExceptionInfo, const TCHAR* Callstack)
 	{
 		GLog->Panic();
@@ -1316,11 +1383,23 @@ namespace ShaderCompileWorkerError
 		ModalErrorOrLog(FString::Printf(TEXT("ShaderCompileWorker failed with bad-input-file exception:\n%s\n"), Data));
 	}
 
-	void HandleOutOfMemory(const TCHAR* ExceptionInfo, const TCHAR* Hostname, uint64 AvailableMemory, uint64 UsedMemory)
+	bool HandleOutOfMemory(const TCHAR* ExceptionInfo, const TCHAR* Hostname, uint64 AvailableMemory, uint64 UsedMemory, const TArray<FShaderCommonCompileJobPtr>& QueuedJobs)
 	{
 		const FString AvailableMemoryStr = FString::Printf(TEXT("%d MB"), FUnitConversion::Convert(AvailableMemory, EUnit::Bytes, EUnit::Megabytes));
 		const FString UsedMemoryStr = FString::Printf(TEXT("%d MB"), FUnitConversion::Convert(UsedMemory, EUnit::Bytes, EUnit::Megabytes));
-		ModalErrorOrLog(FString::Printf(TEXT("ShaderCompileWorker failed with out-of-memory exception on machine \"%s\"\n%s; Used physical memory %s of %s\n"), Hostname, ExceptionInfo, *UsedMemoryStr, *AvailableMemoryStr));
+		const FString ErrorReport = FString::Printf(TEXT("ShaderCompileWorker failed with out-of-memory exception on machine \"%s\"\n%s; Used physical memory %s of %s\n"), Hostname, ExceptionInfo, *UsedMemoryStr, *AvailableMemoryStr);
+
+		if (TryReissueShaderCompileJobs(QueuedJobs))
+		{
+			// We recovered from this error
+			UE_LOG(LogShaderCompilers, Warning, TEXT("%sReissue %d shader compile %s to remaining workers"), *ErrorReport, QueuedJobs.Num(), (QueuedJobs.Num() == 1 ? TEXT("job") : TEXT("jobs")));
+			return true;
+		}
+		else
+		{
+			ModalErrorOrLog(ErrorReport);
+			return false;
+		}
 	}
 }
 
@@ -1866,7 +1945,7 @@ static void LogQueuedCompileJobs(const TArray<FShaderCommonCompileJobPtr>& Queue
 
 // Disable optimization for this crash handler to get full access to the entire stack frame when debugging a crash dump
 UE_DISABLE_OPTIMIZATION_SHIP
-static void HandleWorkerCrash(const TArray<FShaderCommonCompileJobPtr>& QueuedJobs, FArchive& OutputFile, int32 OutputVersion, int64 FileSize, FSCWErrorCode::ECode ErrorCode, int32 NumProcessedJobs, int32 CallstackLength, int32 ExceptionInfoLength, int32 HostnameLength)
+static bool HandleWorkerCrash(const TArray<FShaderCommonCompileJobPtr>& QueuedJobs, FArchive& OutputFile, int32 OutputVersion, int64 FileSize, FSCWErrorCode::ECode ErrorCode, int32 NumProcessedJobs, int32 CallstackLength, int32 ExceptionInfoLength, int32 HostnameLength)
 {
 	TArray<TCHAR> Callstack;
 	Callstack.AddUninitialized(CallstackLength + 1);
@@ -1976,12 +2055,12 @@ static void HandleWorkerCrash(const TArray<FShaderCommonCompileJobPtr>& QueuedJo
 		ShaderCompileWorkerError::HandleBadInputFile(ExceptionInfo.GetData());
 		break;
 	case FSCWErrorCode::OutOfMemory:
-		ShaderCompileWorkerError::HandleOutOfMemory(ExceptionInfo.GetData(), Hostname.GetData(), AvailableMemory, UsedMemory);
-		break;
+		return ShaderCompileWorkerError::HandleOutOfMemory(ExceptionInfo.GetData(), Hostname.GetData(), AvailableMemory, UsedMemory, QueuedJobs);
 	case FSCWErrorCode::Success:
 		// Can't get here...
-		break;
+		return true;
 	}
+	return false;
 }
 UE_ENABLE_OPTIMIZATION_SHIP
 
@@ -2026,16 +2105,20 @@ void FShaderCompileUtilities::DoReadTaskResults(const TArray<FShaderCommonCompil
 	int32 HostnameLength = 0;
 	OutputFile << HostnameLength;
 
-	// Worker crashed
 	if (ErrorCode != FSCWErrorCode::Success)
 	{
-		HandleWorkerCrash(QueuedJobs, OutputFile, OutputVersion, FileSize, (FSCWErrorCode::ECode)ErrorCode, NumProcessedJobs, CallstackLength, ExceptionInfoLength, HostnameLength);
+		// If worker crashed in a way we were able to recover from, return and expect the compile jobs to be reissued already
+		if (HandleWorkerCrash(QueuedJobs, OutputFile, OutputVersion, FileSize, (FSCWErrorCode::ECode)ErrorCode, NumProcessedJobs, CallstackLength, ExceptionInfoLength, HostnameLength))
+		{
+			FSCWErrorCode::Reset();
+			return;
+		}
 	}
 
 	TArray<FShaderCompileJob*> QueuedSingleJobs;
 	TArray<FShaderPipelineCompileJob*> QueuedPipelineJobs;
 	SplitJobsByType(QueuedJobs, QueuedSingleJobs, QueuedPipelineJobs);
-	TArray<FShaderCompileJob*> ReissueSourceJobs;
+	TArray<FShaderCommonCompileJob*> ReissueSourceJobs;
 
 	// Read single jobs
 	{
@@ -2127,25 +2210,7 @@ void FShaderCompileUtilities::DoReadTaskResults(const TArray<FShaderCommonCompil
 	}
 	
 	// Requeue any jobs we wish to run again
-	if (ReissueSourceJobs.Num())
-	{
-		TRACE_CPUPROFILER_EVENT_SCOPE(ReissueShaderJobs);
-
-		TArray<FShaderCommonCompileJobPtr> ReissueJobs;
-		ReissueJobs.Reserve(ReissueSourceJobs.Num());
-		const uint32 JobId = FShaderCommonCompileJob::GetNextJobId();
-		for (const FShaderCompileJob* ReissueSourceJob : ReissueSourceJobs)
-		{
-			FShaderCompileJob* ReissueJob = GShaderCompilingManager->PrepareShaderCompileJob(JobId, ReissueSourceJob->Key, ReissueSourceJob->Priority);
-			if (ReissueJob)
-			{
-				ReissueJob->Input = ReissueSourceJob->Input;
-				ReissueJobs.Add(FShaderCommonCompileJobPtr(ReissueJob));
-			}
-		}
-
-		GShaderCompilingManager->SubmitJobs(ReissueJobs, FString(""), FString(""));
-	}
+	ReissueShaderCompileJobs(ReissueSourceJobs);
 }
 
 #if WITH_EDITOR
diff --git a/Engine/Source/Runtime/RenderCore/Private/ShaderCore.cpp b/Engine/Source/Runtime/RenderCore/Private/ShaderCore.cpp
--- a/Engine/Source/Runtime/RenderCore/Private/ShaderCore.cpp
+++ b/Engine/Source/Runtime/RenderCore/Private/ShaderCore.cpp
@@ -241,7 +241,6 @@ static FString GSCWErrorCodeInfo;
 
 void FSCWErrorCode::Report(ECode Code, const FStringView& Info)
 {
-	checkf(!FSCWErrorCode::IsSet(), TEXT("ShaderCompileWorker error code already set; this is only supposed to be set once to handle a worker crash:\n%s"), Info);
 	GSCWErrorCode = Code;
 	GSCWErrorCodeInfo = Info;
 }
