#sh-12-13-stw-servers-stumped-by-quotes-ext
Rolling back OnlineReplStructs.cpp to 23.00 state for 23.10 server

diff --git a/Engine/Source/Runtime/Engine/Private/OnlineReplStructs.cpp b/Engine/Source/Runtime/Engine/Private/OnlineReplStructs.cpp
--- a/Engine/Source/Runtime/Engine/Private/OnlineReplStructs.cpp
+++ b/Engine/Source/Runtime/Engine/Private/OnlineReplStructs.cpp
@@ -470,11 +470,11 @@ bool FUniqueNetIdRepl::ExportTextItem(FString& ValueStr, FUniqueNetIdRepl const&
 			FName Type = GetType();
 			if (Type == UOnlineEngineInterface::Get()->GetDefaultOnlineSubsystemName())
 			{
-				ValueStr += FString::Printf(TEXT("\"%s\""), *ToString());
+				ValueStr += ToString();
 			}
 			else
 			{
-				ValueStr += FString::Printf(TEXT("%s:\"%s\""), *Type.ToString(), *ToString());
+				ValueStr += FString::Printf(TEXT("%s:%s"), *Type.ToString(), *ToString());
 			}
 		}
 		else
