[Audio] Fix for changing project default compression type changing the USoundWave compression type to the new default (instead of being Project Default).

diff --git a/Engine/Source/Runtime/Engine/Private/SoundWave.cpp b/Engine/Source/Runtime/Engine/Private/SoundWave.cpp
--- a/Engine/Source/Runtime/Engine/Private/SoundWave.cpp
+++ b/Engine/Source/Runtime/Engine/Private/SoundWave.cpp
@@ -1168,7 +1168,17 @@ ESoundAssetCompressionType USoundWave::GetSoundAssetCompressionTypeEnum() const
 void USoundWave::SetSoundAssetCompressionType(ESoundAssetCompressionType InSoundAssetCompressionType, bool bMarkDirty)
 {
 #if WITH_EDITOR
-	SoundAssetCompressionType = InSoundAssetCompressionType;
+	/*
+	 * This is a bit of a hack to make viable for a hot-fixable.
+	 * If bMarkDirty is False and the current compression type is project defined,
+	 * it means the project default codec was changed.  In this scenario we don't want
+	 * to change the SoundAssetCompressionType field, as this will change the property
+	 * on the USoundWave to whatever the new project default is.
+	 */
+	if(bMarkDirty || SoundAssetCompressionType != ESoundAssetCompressionType::ProjectDefined)
+	{
+		SoundAssetCompressionType = InSoundAssetCompressionType;
+	}
 	SoundWaveDataPtr->bIsSeekable = IsSeekable();
 	SoundWaveDataPtr->RuntimeFormat = SoundWaveDataPtr->FindRuntimeFormat(*this);
 	UpdateAsset(bMarkDirty);
