Adapted from PR #9221: Fix CharacterMovementComp server to client corrections causing RotationRate offsets to be added again (Contributed by [REDACTED_PERSON])

diff --git a/Engine/Source/Runtime/Engine/Private/Components/CharacterMovementComponent.cpp b/Engine/Source/Runtime/Engine/Private/Components/CharacterMovementComponent.cpp
--- a/Engine/Source/Runtime/Engine/Private/Components/CharacterMovementComponent.cpp
+++ b/Engine/Source/Runtime/Engine/Private/Components/CharacterMovementComponent.cpp
@@ -308,6 +308,13 @@ namespace CharacterMovementCVars
 		TEXT("This enabled a workaround to allow impulses to be applied to geometry collection.\n"),
 		ECVF_Default);
 
+	static int32 bUseLastGoodRotationDuringCorrection = 1;
+	FAutoConsoleVariableRef CVarUseLastGoodRotationDuringCorrection(
+		TEXT("p.UseLastGoodRotationDuringCorrection"),
+		bUseLastGoodRotationDuringCorrection,
+		TEXT("When enabled, during a correction, restore the last good rotation before re-simulating saved moves if the server didn't specify one. This improves visual quality with options like bOrientToMovement or bUseControllerDesiredRotation that rotate over time."),
+		ECVF_Default);
+
 #if !UE_BUILD_SHIPPING
 
 	int32 NetShowCorrections = 0;
@@ -10384,6 +10391,14 @@ void UCharacterMovementComponent::ClientAdjustPosition_Implementation
 		MovementBaseUtility::TransformDirectionToWorld(NewBase, NewBaseBoneName, CurrentVelocity, NewVelocity);
 	}
 
+	// Fall back to the last-known good rotation if the server didn't send one
+	if (CharacterMovementCVars::bUseLastGoodRotationDuringCorrection
+		&& (bOrientRotationToMovement || bUseControllerDesiredRotation)
+		&& (!OptionalRotation.IsSet() && ClientData->LastAckedMove.IsValid()))
+	{
+		OptionalRotation = ClientData->LastAckedMove->SavedRotation;
+	}
+
 	// Trigger event
 	OnClientCorrectionReceived(*ClientData, TimeStamp, WorldShiftedNewLocation, NewVelocity, NewBase, NewBaseBoneName, bHasBase, bBaseRelativePosition, ServerMovementMode);
 
