Fixed Vislog not being able to record world-less UObjects

diff --git a/Engine/Source/Runtime/Engine/Private/VisualLogger/VisualLogger.cpp b/Engine/Source/Runtime/Engine/Private/VisualLogger/VisualLogger.cpp
--- a/Engine/Source/Runtime/Engine/Private/VisualLogger/VisualLogger.cpp
+++ b/Engine/Source/Runtime/Engine/Private/VisualLogger/VisualLogger.cpp
@@ -257,7 +257,9 @@ namespace
 		if (GIsEditor && EEngine != nullptr && World == nullptr)
 		{
 			// lets use PlayWorld during PIE/Simulate and regular world from editor otherwise, to draw debug information
-			World = EEngine->PlayWorld != nullptr ? ToRawPtr(EEngine->PlayWorld) : EEngine->GetEditorWorldContext().World();
+			World = EEngine->PlayWorld != nullptr 
+				? ToRawPtr(EEngine->PlayWorld) 
+				: (EEngine->GetWorldContexts().Num() > 0 ? EEngine->GetEditorWorldContext().World() : nullptr);
 		}
 
 #endif
