Animation Compression

diff --git a/Engine/Source/Runtime/Engine/Classes/Animation/AnimData/IAnimationDataModel.h b/Engine/Source/Runtime/Engine/Classes/Animation/AnimData/IAnimationDataModel.h
--- a/Engine/Source/Runtime/Engine/Classes/Animation/AnimData/IAnimationDataModel.h
+++ b/Engine/Source/Runtime/Engine/Classes/Animation/AnimData/IAnimationDataModel.h
@@ -155,49 +155,66 @@ public:
 	/**
 	* @return	Array containing all bone animation tracks 
 	*/
+	UE_DEPRECATED(5.2, "GetBoneAnimationTracks has been deprecated")
 	UFUNCTION(BlueprintCallable, Category = AnimationDataModel)
 	virtual const TArray<FBoneAnimationTrack>& GetBoneAnimationTracks() const = 0;
+
+	virtual FTransform EvaluateBoneTrackTransform(FName TrackName, const FFrameTime& FrameTime, const EAnimInterpolationType& Interpolation) const = 0;
+	virtual FTransform GetBoneTrackTransform(FName TrackName, const FFrameNumber& FrameNumber) const = 0;
+	virtual void GetBoneTrackTransforms(FName TrackName, const TArray<FFrameNumber>& FrameNumbers, TArray<FTransform>& OutTransforms) const = 0;
+	virtual void GetBoneTrackTransforms(FName TrackName, TArray<FTransform>& OutTransforms) const = 0;
+	virtual void GetBoneTracksTransform(const TArray<FName>& TrackNames, const FFrameNumber& FrameNumber, TArray<FTransform>& OutTransforms) const = 0;
 	
 	/**
 	* @return	Bone animation track for the provided index
 	*/
+	UE_DEPRECATED(5.2, "GetBoneTrackByIndex has been deprecated")
 	UFUNCTION(BlueprintCallable, Category = AnimationDataModel)
 	virtual const FBoneAnimationTrack& GetBoneTrackByIndex(int32 TrackIndex) const = 0;
 
 	/**
 	* @return	Bone animation track for the provided (bone) name
 	*/
+	UE_DEPRECATED(5.2, "GetBoneTrackByName has been deprecated")
 	UFUNCTION(BlueprintCallable, Category = AnimationDataModel)
 	virtual const FBoneAnimationTrack& GetBoneTrackByName(FName TrackName) const = 0;
 
 	/**
 	* @return	Bone animation track for the provided (bone) name if found, otherwise returns a nullptr 
 	*/
+	UE_DEPRECATED(5.2, "FindBoneTrackByName has been deprecated")
 	virtual const FBoneAnimationTrack* FindBoneTrackByName(FName Name) const = 0;
 
 	/**
 	* @return	Bone animation track for the provided index if valid, otherwise returns a nullptr 
 	*/
+	UE_DEPRECATED(5.2, "FindBoneTrackByIndex has been deprecated")
 	virtual const FBoneAnimationTrack* FindBoneTrackByIndex(int32 BoneIndex) const = 0;
 
 	/**
 	* @return	Internal track index for the provided bone animation track if found, otherwise returns INDEX_NONE 
 	*/
+	UE_DEPRECATED(5.2, "FindBoneTrackByIndex has been deprecated")
 	UFUNCTION(BlueprintCallable, Category = AnimationDataModel)
 	virtual int32 GetBoneTrackIndex(const FBoneAnimationTrack& Track) const = 0;
 
 	/**
 	* @return	Internal track index for the provided (bone) name if found, otherwise returns INDEX_NONE 
 	*/
+	UE_DEPRECATED(5.2, "GetBoneTrackIndexByName has been deprecated")
 	UFUNCTION(BlueprintCallable, Category = AnimationDataModel)
 	virtual int32 GetBoneTrackIndexByName(FName TrackName) const = 0;
 
 	/**
 	* @return	Whether or not the provided track index is valid 
 	*/
+	UE_DEPRECATED(5.2, "IsValidBoneTrackIndex has been deprecated")
 	UFUNCTION(BlueprintCallable, Category = AnimationDataModel)
 	virtual bool IsValidBoneTrackIndex(int32 TrackIndex) const = 0;
 
+	UFUNCTION(BlueprintCallable, Category = AnimationDataModel)
+	virtual bool IsValidBoneTrackName(const FName& TrackName) const = 0;
+
 	/**
 	* @return	Total number of bone animation tracks
 	*/
@@ -339,9 +356,21 @@ public:
 #endif // WITH_EDITOR
 
 	virtual bool HasBeenPopulated() const = 0;
+	virtual void IterateBoneKeys(const FName& BoneName, TFunction<bool(const FVector3f& Pos, const FQuat4f&, const FVector3f, const FFrameNumber&)> IterationFunction) const = 0;
+
+	struct ENGINE_API FEvaluationAndModificationLock
+	{
+		FEvaluationAndModificationLock(IAnimationDataModel& InModel) : Model(InModel) { InModel.LockEvaluationAndModification(); }
+		~FEvaluationAndModificationLock()  { Model.UnlockEvaluationAndModification(); }
+	private:
+		IAnimationDataModel& Model;
+	};
 protected:
 	virtual FAnimDataModelModifiedDynamicEvent& GetModifiedDynamicEvent() = 0;
 	virtual void OnNotify(const EAnimDataModelNotifyType& NotifyType, const FAnimDataModelNotifPayload& Payload) = 0;
+
+	virtual void LockEvaluationAndModification() const = 0;
+	virtual void UnlockEvaluationAndModification() const = 0;	
 	
 	struct FModelNotifier;
 	virtual IAnimationDataModel::FModelNotifier& GetNotifier() = 0;
diff --git a/Engine/Source/Runtime/Engine/Classes/Animation/AnimSequence.h b/Engine/Source/Runtime/Engine/Classes/Animation/AnimSequence.h
--- a/Engine/Source/Runtime/Engine/Classes/Animation/AnimSequence.h
+++ b/Engine/Source/Runtime/Engine/Classes/Animation/AnimSequence.h
@@ -20,6 +20,8 @@
 #include "Animation/CustomAttributes.h"
 #include "Animation/AnimData/AnimDataNotifications.h"
 #include "Animation/AttributeCurve.h"
+#include "PerPlatformProperties.h"
+#include "IO/IoHash.h"
 
 #if WITH_EDITOR
 #include "AnimData/IAnimationDataModel.h"
@@ -31,12 +33,15 @@
 typedef TArray<FTransform> FTransformArrayA2;
 
 class USkeletalMesh;
+class FQueuedThreadPool;
+enum class EQueuedWorkPriority : uint8;
 struct FAnimCompressContext;
 struct FAnimSequenceDecompressionContext;
 struct FCompactPose;
-class ITargetPlatform;
 
-namespace UE { namespace Anim { namespace Compression { struct FScopedCompressionGuard; } } }
+namespace UE { namespace Anim { class FAnimSequenceCompilingManager; namespace Compression { struct FScopedCompressionGuard; } class FAnimationSequenceAsyncCacheTask; } }
+
+extern ENGINE_API int32 GPerformFrameStripping;
 
 // These two always should go together, but it is not right now. 
 // I wonder in the future, we change all compressed to be inside as well, so they all stay together
@@ -180,6 +185,7 @@ struct ENGINE_API FCompressedTrack
 };
 
 // Param structure for UAnimSequence::RequestAnimCompressionParams
+struct UE_DEPRECATED(5.2, "FRequestAnimCompressionParams has been deprecated") FRequestAnimCompressionParams;
 struct ENGINE_API FRequestAnimCompressionParams
 {
 	// Is the compression to be performed Async
@@ -195,12 +201,12 @@ struct ENGINE_API FRequestAnimCompressionParams
 	TSharedPtr<FAnimCompressContext> CompressContext;
 
 	// Constructors
-	FRequestAnimCompressionParams(bool bInAsyncCompression, bool bInAllowAlternateCompressor = false, bool bInOutput = false, const ITargetPlatform* = nullptr);
-	FRequestAnimCompressionParams(bool bInAsyncCompression, TSharedPtr<FAnimCompressContext> InCompressContext);
+	FRequestAnimCompressionParams(bool bInAsyncCompression, bool bInAllowAlternateCompressor = false, bool bInOutput = false, const ITargetPlatform* = nullptr) {}
+	FRequestAnimCompressionParams(bool bInAsyncCompression, TSharedPtr<FAnimCompressContext> InCompressContext) {}
 
 	// Frame stripping initialization funcs (allow stripping per platform)
-	void InitFrameStrippingFromCVar();
-	void InitFrameStrippingFromPlatform();
+	void InitFrameStrippingFromCVar() {}
+	void InitFrameStrippingFromPlatform() {}
 
 	const ITargetPlatform* TargetPlatform;
 };
@@ -380,12 +386,14 @@ public:
 	virtual void GetPreloadDependencies(TArray<UObject*>& OutDeps) override;
 #if WITH_EDITOR
 	virtual void PostEditChangeProperty(FPropertyChangedEvent& PropertyChangedEvent) override;
+	virtual void BeginCacheForCookedPlatformData(const ITargetPlatform* TargetPlatform) override;
 	virtual bool IsCachedCookedPlatformDataLoaded(const ITargetPlatform* TargetPlatform) override;
 	virtual void WillNeverCacheCookedPlatformDataAgain() override;
 	virtual void ClearCachedCookedPlatformData(const ITargetPlatform* TargetPlatform) override;
 	virtual void ClearAllCachedCookedPlatformData() override;
 #endif // WITH_EDITOR
 	virtual void BeginDestroy() override;
+	virtual bool IsReadyForFinishDestroy() override;
 	virtual void GetAssetRegistryTags(TArray<FAssetRegistryTag>& OutTags) const override;
 	static void AddReferencedObjects(UObject* This, FReferenceCollector& Collector);
 	//~ End UObject Interface
@@ -405,7 +413,7 @@ public:
 	virtual void RefreshCacheData() override;
 	virtual EAdditiveAnimationType GetAdditiveAnimType() const override { return AdditiveAnimType; }
 	virtual int32 GetNumberOfSampledKeys() const override;
-	virtual FFrameRate GetSamplingFrameRate() const override { return TargetFrameRate; }
+	virtual FFrameRate GetSamplingFrameRate() const override { return PlatformTargetFrameRate.Default; }
 	virtual void EvaluateCurveData(FBlendedCurve& OutCurve, float CurrentTime, bool bForceUseRawData = false) const override;
 	virtual float EvaluateCurveData(SmartName::UID_Type CurveUID, float CurrentTime, bool bForceUseRawData = false) const override;
 	virtual bool HasCurveData(SmartName::UID_Type CurveUID, bool bForceUseRawData) const override;
@@ -457,6 +465,7 @@ private:
 	void UpdateFrameRate();
 #endif
 	
+public:
 	const TArray<FTransform>& GetRetargetTransforms() const;
 	FName GetRetargetTransformsSourceName() const;
 
@@ -470,7 +479,6 @@ private:
 	*/
 	void RetargetBoneTransform(FTransform& BoneTransform, const int32 SkeletonBoneIndex, const FCompactPoseBoneIndex& BoneIndex, const FBoneContainer& RequiredBones, const bool bIsBakedAdditive) const;
 
-public:
 	/**
 	* Get Bone Transform of the additive animation for the Time given, relative to Parent for all RequiredBones
 	*
@@ -501,9 +509,9 @@ public:
 	 * @param	Time			Time on track to interpolate to.
 	 * @param	bUseRawData		If true, use raw animation data instead of compressed data.
 	 */
-	UE_DEPRECATED(5.1, "Use other GetBoneTransform signature")
-	void GetBoneTransform(FTransform& OutAtom, int32 TrackIndex, float Time, bool bUseRawData) const;
-	void GetBoneTransform(FTransform& OutAtom, int32 TrackIndex, double Time, bool bUseRawData) const;
+	UE_DEPRECATED(5.1, "Use other GetBoneTransform signature using double and skeleton index")
+	void GetBoneTransform(FTransform& OutAtom, int32 TrackIndex, float Time, bool bUseRawData) const {}
+	void GetBoneTransform(FTransform& OutAtom, FSkeletonPoseBoneIndex BoneIndex, double Time, bool bUseRawData) const;
 
 	/**
 	 * Get Bone Transform of the Time given, relative to Parent for the Track Given
@@ -513,7 +521,7 @@ public:
 	 * @param	DecompContext	Decompression context to use.
 	 * @param	bUseRawData		If true, use raw animation data instead of compressed data.
 	 */
-	void GetBoneTransform(FTransform& OutAtom, int32 TrackIndex, FAnimSequenceDecompressionContext& DecompContext, bool bUseRawData) const;
+	void GetBoneTransform(FTransform& OutAtom, FSkeletonPoseBoneIndex BoneIndex, FAnimSequenceDecompressionContext& DecompContext, bool bUseRawData) const;
 	// End Transform related functions 
 
 	// Begin Memory related functions
@@ -526,7 +534,7 @@ public:
 	/**
 	 * @return		The approximate size of raw animation data.
 	 */
-	int32 GetApproxRawSize() const;
+	int64 GetApproxRawSize() const;
 #endif // WITH_EDITOR
 
 	/**
@@ -541,35 +549,36 @@ public:
 
 	// Get compressed data for this UAnimSequence. May be built directly or pulled from DDC
 #if WITH_EDITOR
+	UE_DEPRECATED(5.2, "Public access to ShouldPerformStripping has been deprecated")
 	bool ShouldPerformStripping(const bool bPerformFrameStripping, const bool bPerformStrippingOnOddFramedAnims) const;
+
+	UE_DEPRECATED(5.2, "GetDDCCacheKeySuffix has been deprecated use CreateDerivedDataKeyHash instead")
 	FString GetDDCCacheKeySuffix(const bool bPerformStripping, const ITargetPlatform* TargetPlatform) const;
-	void ApplyCompressedData(const FString& DataCacheKeySuffix, const bool bPerformFrameStripping, const TArray<uint8>& Data);
+
+	UE_DEPRECATED(5.2, "ApplyCompressedData has been deprecated")
+	void ApplyCompressedData(const FString& DataCacheKeySuffix, const bool bPerformFrameStripping, const TArray<uint8>& Data) {}
 
 	void WaitOnExistingCompression(const bool bWantResults=true);
-	void RequestAnimCompression(FRequestAnimCompressionParams Params);
-	void RequestSyncAnimRecompression(bool bOutput = false)
-	{
-		PRAGMA_DISABLE_DEPRECATION_WARNINGS
-		RequestAnimCompression(FRequestAnimCompressionParams(false, false, bOutput));
-		PRAGMA_ENABLE_DEPRECATION_WARNINGS
-	}
+	
+	UE_DEPRECATED(5.2, "RequestAnimCompression has been deprecated use BeginCacheDerivedData or CacheDerivedData instead")
+	void RequestAnimCompression(
+	PRAGMA_DISABLE_DEPRECATION_WARNINGS
+		FRequestAnimCompressionParams Params) {}
+	PRAGMA_ENABLE_DEPRECATION_WARNINGS
 
-	void RequestAsyncAnimRecompression(bool bOutput = false)
-	{ 
-		PRAGMA_DISABLE_DEPRECATION_WARNINGS
-		RequestAnimCompression(FRequestAnimCompressionParams(true, false, bOutput));
-		PRAGMA_ENABLE_DEPRECATION_WARNINGS
-	}
+	UE_DEPRECATED(5.2, "RequestSyncAnimRecompression has been deprecated use CacheDerivedData instead")
+	void RequestSyncAnimRecompression(bool bOutput = false);
+
+	UE_DEPRECATED(5.2, "RequestSyncAnimRecompression has been deprecated use BeginCacheDerivedData instead")
+	void RequestAsyncAnimRecompression(bool bOutput = false);
 #endif
 
 protected:
-	void ApplyCompressedData(const TArray<uint8>& Data);
     void ClearCompressedBoneData();
     void ClearCompressedCurveData();
     // Write the compressed data to the supplied FArchive
     void SerializeCompressedData(FArchive& Ar, bool bDDCData);
 #if WITH_EDITOR
-	void SetCompressionComplete();
 	virtual void OnAnimModelLoaded() override;
 #endif
 public:
@@ -584,14 +593,18 @@ public:
 	*
 	* @param	TrackIndex		Track Index
 	*/
+	UE_DEPRECATED(5.2, "GetSkeletonIndexFromCompressedDataTrackIndex has been deprecated")
 	int32 GetSkeletonIndexFromCompressedDataTrackIndex(const int32 TrackIndex) const
 	{
+		PRAGMA_DISABLE_DEPRECATION_WARNINGS
 		return GetCompressedTrackToSkeletonMapTable()[TrackIndex].BoneTreeIndex;
+    	PRAGMA_ENABLE_DEPRECATION_WARNINGS
 	}
 	// End Utility functions
 	
 #if WITH_EDITOR
-	void BakeTrackCurvesToRawAnimationTracks(TArray<FRawAnimSequenceTrack>& NewRawTracks, TArray<FName>& NewTrackNames, TArray<FTrackToSkeletonMap>& NewTrackToSkeletonMapTable);
+	UE_DEPRECATED(5.2, "BakeTrackCurvesToRawAnimationTracks has been deprecated")
+	void BakeTrackCurvesToRawAnimationTracks(TArray<FRawAnimSequenceTrack>& NewRawTracks, TArray<FName>& NewTrackNames, TArray<FTrackToSkeletonMap>& NewTrackToSkeletonMapTable) {}
 
 	/**
 	 * Add Key to Transform Curves
@@ -663,13 +676,16 @@ public:
 	bool CanBakeAdditive() const;
 
 	// Bakes out track data for the skeletons virtual bones into the raw data
-	void BakeOutVirtualBoneTracks(TArray<FRawAnimSequenceTrack>& NewRawTracks, TArray<FName>& NewAnimationTrackNames, TArray<FTrackToSkeletonMap>& NewTrackToSkeletonMapTable);
+	UE_DEPRECATED(5.2, "BakeOutVirtualBoneTracks has been deprecated")
+	void BakeOutVirtualBoneTracks(TArray<FRawAnimSequenceTrack>& NewRawTracks, TArray<FName>& NewAnimationTrackNames, TArray<FTrackToSkeletonMap>& NewTrackToSkeletonMapTable) {} 
 
 	// Performs multiple evaluations of the animation as a test of compressed data validatity
-	void TestEvalauteAnimation() const;
+	UE_DEPRECATED(5.2, "TestEvalauteAnimation has been deprecated")
+	void TestEvalauteAnimation() const {}
 
 	// Bakes out the additive version of this animation into the raw data.
-	void BakeOutAdditiveIntoRawData(TArray<FRawAnimSequenceTrack>& NewRawTracks, TArray<FName>& NewAnimationTrackNames, TArray<FTrackToSkeletonMap>& NewTrackToSkeletonMapTable, TArray<FFloatCurve>& NewCurveTracks, TArray<FRawAnimSequenceTrack>& AdditiveBaseAnimationData);
+	UE_DEPRECATED(5.2, "BakeOutAdditiveIntoRawData has been deprecated, this behaviour has moved to FCompressibleAnimData")
+	void BakeOutAdditiveIntoRawData(TArray<FRawAnimSequenceTrack>& NewRawTracks, TArray<FName>& NewAnimationTrackNames, TArray<FTrackToSkeletonMap>& NewTrackToSkeletonMapTable, TArray<FFloatCurve>& NewCurveTracks, TArray<FRawAnimSequenceTrack>& AdditiveBaseAnimationData) const {}
 
 	// Test whether at any point we will scale a bone to 0 (needed for validating additive anims)
 	bool DoesSequenceContainZeroScale() const;
@@ -748,12 +764,6 @@ public:
 	/** List of Unique marker names in this animation sequence */
 	TArray<FName>				UniqueMarkerNames;
 
-private:
-#if WITH_EDITOR
-	// Are we currently compressing this animation
-	bool bCompressionInProgress;
-#endif
-
 public:
 #if WITH_EDITOR
 	UE_DEPRECATED(5.0, "AddBoneCustomAttribute has been deprecated see UAnimDataController::AddAttribute")
@@ -796,15 +806,17 @@ protected:
 	// End UAnimSequenceBase virtual overrides
 
 	void EnsureValidRawDataGuid();
+	UE_DEPRECATED(5.2, "RecompressAnimationData has been deprecated")
 	void RecompressAnimationData();
-	void ResampleAnimationTrackData();
+	UE_DEPRECATED(5.2, "ResampleAnimationTrackData has been deprecated")
+	void ResampleAnimationTrackData() const {}
 	void CalculateNumberOfSampledKeys();
-	void ClearResampledAnimationTrackData();
+	UE_DEPRECATED(5.2, "ClearResampledAnimationTrackData has been deprecated")
+	void ClearResampledAnimationTrackData() {}
 
 	void DeleteBoneAnimationData();
 	void DeleteDeprecatedRawAnimationData();
 public:
-	const TArray<FBoneAnimationTrack>& GetResampledTrackData() const { return ResampledAnimationTrackData; }
 	void DeleteNotifyTrackData();
 #endif // WITH_EDITOR
 
@@ -812,6 +824,9 @@ protected:
 	UPROPERTY(VisibleAnywhere, AssetRegistrySearchable, Category = "Animation")
 	FFrameRate TargetFrameRate;
 
+	UPROPERTY(VisibleAnywhere, AssetRegistrySearchable, Category = "Animation")
+	FPerPlatformFrameRate PlatformTargetFrameRate;
+
 #if WITH_EDITORONLY_DATA
 	UPROPERTY(VisibleAnywhere, AssetRegistrySearchable, Category = "Animation", Transient, DuplicateTransient)
 	int32 NumberOfSampledKeys;
@@ -819,13 +834,7 @@ protected:
 	UPROPERTY(VisibleAnywhere, Category = "Animation", Transient, DuplicateTransient)
 	int32 NumberOfSampledFrames;
 
-	UPROPERTY(VisibleAnywhere, Category = "Animation", Transient, DuplicateTransient)
-	TArray<FBoneAnimationTrack> ResampledAnimationTrackData;
-
 	bool bBlockCompressionRequests;
-
-	FString RequestedCompression_DDCKey;
-
 private:
 	PRAGMA_DISABLE_DEPRECATION_WARNINGS
 	UE_DEPRECATED(5.0, "PerBoneCustomAttributeData has been deprecated see UAnimDataModel::AnimatedBoneAttributes")
@@ -838,6 +847,32 @@ protected:
 	UPROPERTY()
 	TMap<FAnimationAttributeIdentifier, FAttributeCurve> AttributeCurves;
 
+#if WITH_EDITOR
+	FIoHash CreateDerivedDataKeyHash(const ITargetPlatform* TargetPlatform);
+	FIoHash BeginCacheDerivedData(const ITargetPlatform* TargetPlatform);
+	bool PollCacheDerivedData(const FIoHash& KeyHash) const;
+	void EndCacheDerivedData(const FIoHash& KeyHash);
+
+	FIoHash DataKeyHash;
+	TMap<FIoHash, TUniquePtr<FCompressedAnimSequence>> DataByPlatformKeyHash;
+	TMap<FIoHash, TPimplPtr<UE::Anim::FAnimationSequenceAsyncCacheTask>> CacheTasksByKeyHash;
+
+protected:
+	bool TryCancelAsyncTasks();
+	void FinishAsyncTasks();
+	void Reschedule(FQueuedThreadPool* InThreadPool, EQueuedWorkPriority InPriority);
+	bool IsAsyncTaskComplete() const;
+	bool IsCompiling() const;
+public:	
+	void BeginCacheDerivedDataForCurrentPlatform();
+	void CacheDerivedDataForCurrentPlatform();
+	
+	// Synchronous caching of compressed animation data for provided target platform
+	FCompressedAnimSequence& CacheDerivedData(const ITargetPlatform* TargetPlatform);
+
+	FFrameRate GetTargetSamplingFrameRate(const ITargetPlatform* InPlatform) const;
+#endif
+
 public:
 	friend class UAnimationAsset;
 	friend struct FScopedAnimSequenceRawDataCache;
@@ -847,4 +882,5 @@ public:
 	friend class FAnimSequenceTestBase;
 	friend struct UE::Anim::Compression::FScopedCompressionGuard;
 	friend class FAnimDataControllerTestBase;
+	friend class UE::Anim::FAnimSequenceCompilingManager;
 };
diff --git a/Engine/Source/Runtime/Engine/Classes/Animation/AnimSequenceHelpers.h b/Engine/Source/Runtime/Engine/Classes/Animation/AnimSequenceHelpers.h
--- a/Engine/Source/Runtime/Engine/Classes/Animation/AnimSequenceHelpers.h
+++ b/Engine/Source/Runtime/Engine/Classes/Animation/AnimSequenceHelpers.h
@@ -80,6 +80,7 @@ namespace Anim {
 	* @param	Time				Time value to sample at
 	* @param	InterpolationType	Method to be used to interpolate between different sample points
 	*/
+	UE_DEPRECATED(5.2, "GetBoneTransformFromModel with TrackIndex has been deprecated, use IAnimationDataModel::GetBoneTrackTransform instead")
 	ENGINE_API void GetBoneTransformFromModel(const IAnimationDataModel* Model, FTransform& OutTransform, int32 TrackIndex, double Time, const EAnimInterpolationType& Interpolation);
 	
 	/**
@@ -90,6 +91,7 @@ namespace Anim {
 	* @param	TrackIndex			Track / bone index which should be sampld
 	* @param	KeyIndex			Specific key index to sample
 	*/
+	UE_DEPRECATED(5.2, "GetBoneTransformFromModel with TrackIndex has been deprecated, use IAnimationDataModel::GetBoneTrackTransform instead")
 	ENGINE_API void GetBoneTransformFromModel(const IAnimationDataModel* Model, FTransform& OutTransform, int32 TrackIndex, int32 KeyIndex);
 
 	/**
@@ -156,15 +158,18 @@ namespace Anim {
 		ENGINE_API void RemoveKeys(UAnimSequence* InSequence, int32 StartKeyIndex, int32 NumKeysToRemove);
 
 		/**
-		* Finds the index of the first child track (bone) for the provided bone name
+		* Finds the name of the first child track (bone) for the provided bone name
 		*
 		* @param	InSequence			Animation Sequence to remove the animated frames from
 		* @param	Skeleton			Skeleton to use for retrieving child bone names
 		* @param	BoneName			Name of the bone to find the child track index for
 		*
-		* @return The first child track index for the provided bone name, if none found will return the end of the tracks array
+		* @return The first child track name for the provided bone name, if none found will return NAME_None
 		*/
-		ENGINE_API int32 FindFirstChildTrackIndex(const UAnimSequence* InSequence, const USkeleton* Skeleton, const FName& BoneName);
+		ENGINE_API FName FindFirstChildTrackName(const UAnimSequence* InSequence, const USkeleton* Skeleton, const FName& BoneName);
+
+		UE_DEPRECATED(5.2, "FindFirstChildTrackIndex has been deprecated, use FindFirstChildTrackName instead")
+		static int32 FindFirstChildTrackIndex(const UAnimSequence* InSequence, const USkeleton* Skeleton, const FName& BoneName) { return INDEX_NONE; }
 	}
 #endif // WITH_EDITOR
 
diff --git a/Engine/Source/Runtime/Engine/Classes/Animation/Skeleton.h b/Engine/Source/Runtime/Engine/Classes/Animation/Skeleton.h
--- a/Engine/Source/Runtime/Engine/Classes/Animation/Skeleton.h
+++ b/Engine/Source/Runtime/Engine/Classes/Animation/Skeleton.h
@@ -814,6 +814,7 @@ public:
 	 *
 	 * @return	Index of Track of Animation Sequence
 	 */
+	UE_DEPRECATED(5.2, "GetRawAnimationTrackIndex has been deprecated, use tracks are referenced by name instead")
 	ENGINE_API int32 GetRawAnimationTrackIndex(const int32 InSkeletonBoneIndex, const UAnimSequence* InAnimSeq);
 
 	/** 
diff --git a/Engine/Source/Runtime/Engine/Private/Animation/AnimCompress.cpp b/Engine/Source/Runtime/Engine/Private/Animation/AnimCompress.cpp
--- a/Engine/Source/Runtime/Engine/Private/Animation/AnimCompress.cpp
+++ b/Engine/Source/Runtime/Engine/Private/Animation/AnimCompress.cpp
@@ -16,6 +16,20 @@
 
 DEFINE_LOG_CATEGORY(LogAnimationCompression);
 
+#if WITH_EDITOR
+namespace UE
+{
+	namespace Anim
+	{
+		namespace Compression
+		{
+			std::atomic<bool> FAnimationCompressionMemorySummaryScope::ScopeExists = false;
+			TUniquePtr<FCompressionMemorySummary> FAnimationCompressionMemorySummaryScope::CompressionSummary = nullptr;
+		}
+	}
+}
+#endif // WITH_EDITOR
+
 UAnimCompress::UAnimCompress(const FObjectInitializer& ObjectInitializer)
 	: Super(ObjectInitializer)
 {
@@ -46,10 +60,6 @@ void FCompressionMemorySummary::GatherPostCompressionStats(const FCompressedAnim
 #endif
 }
 
-void FAnimCompressContext::GatherPostCompressionStats(const FCompressedAnimSequence& CompressedData, const TArray<FBoneData>& BoneData, const FName AnimFName, double CompressionTime, bool bInPerformedCompression)
-{
-	CompressionSummary.GatherPostCompressionStats(CompressedData, BoneData, AnimFName, CompressionTime, bInPerformedCompression);
-}
 
 FCompressionMemorySummary::FCompressionMemorySummary(bool bInEnabled)
 	: bEnabled(bInEnabled)
@@ -202,36 +212,14 @@ uint8 MakeBitForFlag(uint32 Item, uint32 Position)
 
 //////////////////////////////////////////////////////////////////////////////////////
 // FCompressionMemorySummary
-
-void FCompressionMemorySummary::GatherPreCompressionStats(const FString& Name, int32 RawSize, int32 PreCompressedSize, int32 ProgressNumerator, int32 ProgressDenominator)
+void FCompressionMemorySummary::GatherPreCompressionStats(int32 RawSize, int32 PreCompressedSize)
 {
-	if (bEnabled)
-	{
-		bUsed = true;
-		FFormatNamedArguments Args;
-		Args.Add(TEXT("AnimSequenceName"), FText::FromString(Name));
-		Args.Add(TEXT("ProgressNumerator"), ProgressNumerator);
-		Args.Add(TEXT("ProgressDenominator"), ProgressDenominator);
-
-		GWarn->StatusUpdate(ProgressNumerator,
-			ProgressDenominator,
-			FText::Format(NSLOCTEXT("CompressionMemorySummary", "CompressingTaskStatusMessageFormat", "Compressing {AnimSequenceName} ({ProgressNumerator}/{ProgressDenominator})"), Args));
-
-		TotalRaw += RawSize;
-		TotalBeforeCompressed += PreCompressedSize;
-		++NumberOfAnimations;
-	}
-}
-
-
-//////////////////////////////////////////////////////////////////////////////////////
-
-void FAnimCompressContext::GatherPreCompressionStats(const FString& Name, int32 RawSize, int32 PreviousCompressionSize)
-{
-	CompressionSummary.GatherPreCompressionStats(Name, RawSize, PreviousCompressionSize, AnimIndex, MaxAnimations);
+	bUsed = true;
+	TotalRaw += RawSize;
+	TotalBeforeCompressed += PreCompressedSize;
+	++NumberOfAnimations;
 }
 
-
 //////////////////////////////////////////////////////////////////////////////////////
 // UAnimCompress
 
@@ -1113,4 +1101,4 @@ void UAnimCompress::SeparateRawDataIntoTracks(
 		OutScaleData.Empty();
 	}
 }
-#endif
+#endif
\ No newline at end of file
diff --git a/Engine/Source/Runtime/Engine/Private/Animation/AnimCompress_RemoveLinearKeys.cpp b/Engine/Source/Runtime/Engine/Private/Animation/AnimCompress_RemoveLinearKeys.cpp
--- a/Engine/Source/Runtime/Engine/Private/Animation/AnimCompress_RemoveLinearKeys.cpp
+++ b/Engine/Source/Runtime/Engine/Private/Animation/AnimCompress_RemoveLinearKeys.cpp
@@ -1042,15 +1042,15 @@ bool UAnimCompress_RemoveLinearKeys::DoReduction(const FCompressibleAnimData& Co
 	// If the processor is to be run, then additive animations need to be converted from relative to absolute
 	const bool bNeedToConvertBackToAdditive = bRunningProcessor ? CompressibleAnimData.bIsValidAdditive : false;
 
-	FCompressibleAnimData TempConvertedAnimData;
+	TUniquePtr<FCompressibleAnimData> TempConvertedAnimData;
 
 	if (bNeedToConvertBackToAdditive)
 	{
-		TempConvertedAnimData = CompressibleAnimData; // duplicate so we can safely convert to additive
-		ConvertFromRelativeSpace(TempConvertedAnimData);
+		TempConvertedAnimData = MakeUnique<FCompressibleAnimData>(CompressibleAnimData); // duplicate so we can safely convert to additive
+		ConvertFromRelativeSpace(*TempConvertedAnimData);
 	}
 
-	const FCompressibleAnimData& CompressibleDataToOperateOn = bNeedToConvertBackToAdditive ? TempConvertedAnimData : CompressibleAnimData;
+	const FCompressibleAnimData& CompressibleDataToOperateOn = bNeedToConvertBackToAdditive ? *TempConvertedAnimData : CompressibleAnimData;
 
 	// Separate the raw data into tracks and remove trivial tracks (all the same value)
 	TArray<FTranslationTrack> TranslationData;
diff --git a/Engine/Source/Runtime/Engine/Private/Animation/AnimCompressionDerivedData.cpp b/Engine/Source/Runtime/Engine/Private/Animation/AnimCompressionDerivedData.cpp
--- a/Engine/Source/Runtime/Engine/Private/Animation/AnimCompressionDerivedData.cpp
+++ b/Engine/Source/Runtime/Engine/Private/Animation/AnimCompressionDerivedData.cpp
@@ -16,12 +16,16 @@
 #include "Misc/CoreDelegates.h"
 #include "ProfilingDebugging/CookStats.h"
 #include "ProfilingDebugging/CountersTrace.h"
+#include "Animation/AnimationCompressionDerivedData.h"
 
 #if WITH_EDITOR
 
 DECLARE_CYCLE_STAT(TEXT("Anim Compression (Derived Data)"), STAT_AnimCompressionDerivedData, STATGROUP_Anim);
 
-FDerivedDataAnimationCompression::FDerivedDataAnimationCompression(const TCHAR* InTypeName, const FString& InAssetDDCKey, TSharedPtr<FAnimCompressContext> InCompressContext)
+FDerivedDataAnimationCompression::FDerivedDataAnimationCompression(const TCHAR* InTypeName, const FString& InAssetDDCKey,
+PRAGMA_DISABLE_DEPRECATION_WARNINGS
+	TSharedPtr<FAnimCompressContext> InCompressContext)
+PRAGMA_ENABLE_DEPRECATION_WARNINGS
 	: TypeName(InTypeName)
 	, AssetDDCKey(InAssetDDCKey)
 	, CompressContext(InCompressContext)
@@ -35,10 +39,7 @@ FDerivedDataAnimationCompression::~FDerivedDataAnimationCompression()
 
 const TCHAR* FDerivedDataAnimationCompression::GetVersionString() const
 {
-	// This is a version string that mimics the old versioning scheme. If you
-	// want to bump this version, generate a new guid using VS->Tools->Create GUID and
-	// return it here. Ex.
-	return TEXT("FB8D0FF1ED0C4848951601EB6BB028F6");
+	return *UE::Anim::Compression::AnimationCompressionVersionString;
 }
 
 FString FDerivedDataAnimationCompression::GetDebugContextString() const
@@ -63,7 +64,7 @@ bool FDerivedDataAnimationCompression::Build( TArray<uint8>& OutDataArray )
 	}
 
 	// Update UsageStats only when not running on the game thread since that thread times this at a higher level.
-	COOK_STAT(auto Timer = AnimSequenceCookStats::UsageStats.TimeSyncWork());
+	COOK_STAT(auto Timer = UE::Anim::AnimSequenceCookStats::UsageStats.TimeSyncWork());
 	COOK_STAT(IsInGameThread() ? Timer.Cancel() : Timer.TrackCyclesOnly());
 
 	SCOPE_CYCLE_COUNTER(STAT_AnimCompressionDerivedData);
@@ -73,6 +74,7 @@ bool FDerivedDataAnimationCompression::Build( TArray<uint8>& OutDataArray )
 
 	bool bCompressionSuccessful = false;
 	{
+		DataToCompress.FetchData(DataToCompress.TargetPlatform);
 		DataToCompress.Update(OutData);
 
 		const bool bBoneCompressionOk = FAnimationUtils::CompressAnimBones(DataToCompress, CompressionResult);
@@ -84,19 +86,18 @@ bool FDerivedDataAnimationCompression::Build( TArray<uint8>& OutDataArray )
 
 #if DO_CHECK
 		FString CompressionName = DataToCompress.BoneCompressionSettings->GetFullName();
-		const TCHAR* AAC = CompressContext.Get()->bAllowAlternateCompressor ? TEXT("true") : TEXT("false");
-		const TCHAR* OutputStr = CompressContext.Get()->bOutput ? TEXT("true") : TEXT("false");
 #endif
 
 		bCompressionSuccessful = bBoneCompressionOk && bCurveCompressionOk;
 
-		ensureMsgf(bCompressionSuccessful, TEXT("Anim Compression failed for Sequence '%s' with compression scheme '%s': compressed data empty\n\tAnimIndex: %i\n\tMaxAnim:%i\n\tAllowAltCompressor:%s\n\tOutput:%s"), 
+		ensureMsgf(bCompressionSuccessful, TEXT("Anim Compression failed for Sequence '%s' with compression scheme '%s': compressed data empty\n\tAnimIndex: %i\n\tMaxAnim:%i"), 
 											*DataToCompress.FullName,
 											*CompressionName,
+											PRAGMA_DISABLE_DEPRECATION_WARNINGS
 											CompressContext.Get()->AnimIndex,
-											CompressContext.Get()->MaxAnimations,
-											AAC,
-											OutputStr);
+											CompressContext.Get()->MaxAnimations
+											PRAGMA_ENABLE_DEPRECATION_WARNINGS
+											);
 	}
 
 	bCompressionSuccessful = bCompressionSuccessful && !DataToCompress.IsCancelled();
@@ -106,7 +107,9 @@ bool FDerivedDataAnimationCompression::Build( TArray<uint8>& OutDataArray )
 		const double CompressionEndTime = FPlatformTime::Seconds();
 		const double CompressionTime = CompressionEndTime - CompressionStartTime;
 
+		PRAGMA_DISABLE_DEPRECATION_WARNINGS
 		CompressContext->GatherPostCompressionStats(OutData, DataToCompress.BoneData, DataToCompress.AnimFName, CompressionTime, true);
+		PRAGMA_ENABLE_DEPRECATION_WARNINGS
 
 		OutData.CompressedByteStream = MoveTemp(CompressionResult.CompressedByteStream);
 		OutData.CompressedDataStructure = MoveTemp(CompressionResult.AnimData);
@@ -121,329 +124,4 @@ bool FDerivedDataAnimationCompression::Build( TArray<uint8>& OutDataArray )
 	return bCompressionSuccessful;
 }
 
-const uint64 GigaBytes = 1024 * 1024 * 1024;
-const uint64 MAX_ASYNC_COMPRESSION_MEM_USAGE = 3 * GigaBytes;
-const int32 MAX_ACTIVE_COMPRESSIONS = 64;
-
-FAsyncCompressedAnimationsManagement* GAsyncCompressedAnimationsTracker = nullptr;
-
-FAsyncCompressedAnimationsManagement& FAsyncCompressedAnimationsManagement::Get()
-{
-	static FAsyncCompressedAnimationsManagement SingletonInstance;
-	GAsyncCompressedAnimationsTracker = &SingletonInstance;
-	return SingletonInstance;
-}
-
-FAsyncCompressedAnimationsManagement::FAsyncCompressedAnimationsManagement() : ActiveMemoryUsage(0), LastTickTime(0.0)
-{
-	FCoreDelegates::OnPreExit.AddRaw(this, &FAsyncCompressedAnimationsManagement::Shutdown);
-}
-
-void FAsyncCompressedAnimationsManagement::Shutdown()
-{
-	TArray<uint8> OutData; // For canceling
-
-	for (FActiveAsyncCompressionTask& ActiveTask : ActiveAsyncCompressionTasks)
-	{
-		ActiveTask.DataToCompress->IsCancelledSignal.Cancel();
-		if (ActiveTask.Sequence)
-		{
-			ActiveTask.Sequence->ApplyCompressedData(FString(), false, OutData); // Clear active compression on Sequence
-			ActiveTask.Sequence = nullptr;
-		}
-	}
-
-	for (const FQueuedAsyncCompressionWork& QueuedTask : QueuedAsyncCompressionWork)
-	{
-		if (QueuedTask.Anim)
-		{
-			QueuedTask.Anim->ApplyCompressedData(FString(), false, OutData); // Clear active compression on Sequence
-		}
-	}
-
-	QueuedAsyncCompressionWork.Reset();
-
-	const double PreExitStartTimer = FPlatformTime::Seconds();
-	while(ActiveAsyncCompressionTasks.Num() > 0)
-	{
-		Tick(0.f); //Give active tasks that have already started a chance to finish
-		
-		const double PreExitDuration = FPlatformTime::Seconds() - PreExitStartTimer;
-
-		if (PreExitDuration > 20.0)
-		{
-			UE_LOG(LogAnimationCompression, Warning, TEXT("Async Compression Pre Init Timer hit, Active Compressions:%i"), ActiveAsyncCompressionTasks.Num());
-			return;
-		}
-		FPlatformProcess::Sleep(1.f);
-	} 
-}
-
-void FAsyncCompressedAnimationsManagement::OnActiveCompressionFinished(int32 ActiveAnimIndex)
-{
-	TRACE_CPUPROFILER_EVENT_SCOPE(FAsyncCompressedAnimationsManagement::OnActiveCompressionFinished);
-	COOK_STAT(auto Timer = AnimSequenceCookStats::UsageStats.TimeSyncWork());
-
-	FDerivedDataCacheInterface& DerivedDataCache = GetDerivedDataCacheRef();
-
-	FActiveAsyncCompressionTask& Task = ActiveAsyncCompressionTasks[ActiveAnimIndex];
-
-	TArray<uint8> OutData;
-	bool bBuiltLocally = false;
-	const bool bOk = DerivedDataCache.GetAsynchronousResults(Task.AsyncHandle, OutData, &bBuiltLocally);
-
-	COOK_STAT(using EHitOrMiss = FCookStats::CallStats::EHitOrMiss);
-	COOK_STAT(Timer.AddHitOrMiss(bBuiltLocally ? EHitOrMiss::Miss : EHitOrMiss::Hit, OutData.Num()));
-
-	if (Task.Sequence)
-	{
-		if (bOk)
-		{
-			Task.Sequence->ApplyCompressedData(Task.CacheKey, Task.bPerformFrameStripping, OutData);
-		}
-		else
-		{
-			UE_LOG(LogAnimationCompression, Fatal, TEXT("Failed to get async compressed animation data for anim '%s'"), *Task.Sequence->GetName());
-			Task.Sequence->ApplyCompressedData(FString(), false, OutData); // Clear active compression on Sequence
-		}
-	}
-
-	ActiveMemoryUsage -= Task.TaskSize;
-	ActiveAsyncCompressionTasks.RemoveAtSwap(ActiveAnimIndex, 1, false);
-}
-
-float FormatBytes(uint64 Bytes, const TCHAR*& OutPostFix)
-{
-	uint64 Size = GigaBytes;
-
-	if (Bytes > Size)
-	{
-		OutPostFix = TEXT("GB");
-		return (float)Bytes / (float)Size;
-	}
-
-	Size = Size / 1024;
-	if (Bytes > Size)
-	{
-		OutPostFix = TEXT("MB");
-		return (float)Bytes / (float)Size;
-	}
-
-	Size = Size / 1024;
-	if (Bytes > Size)
-	{
-		OutPostFix = TEXT("KB");
-		return (float)Bytes / (float)Size;
-	}
-
-	OutPostFix = TEXT("B");
-	return (float)Bytes;
-}
-
-#define ASYNC_MEM_LOG 0
-
-TRACE_DECLARE_INT_COUNTER(QueuedAnimationCompilation, TEXT("AsyncCompilation/QueuedAnimations"));
-TRACE_DECLARE_INT_COUNTER(ActiveAnimationCompilation, TEXT("AsyncCompilation/ActiveAnimations"));
-TRACE_DECLARE_INT_COUNTER(ProcessedAnimationCompilation, TEXT("AsyncCompilation/ProcessedAnimations"));
-void FAsyncCompressedAnimationsManagement::Tick(float DeltaTime)
-{
-	TRACE_CPUPROFILER_EVENT_SCOPE(FAsyncCompressedAnimationsManagement::Tick);
-
-	const double MaxProcessingTime = 0.1; // try not to hang the editor too much
-	const double EndTime = FPlatformTime::Seconds() + MaxProcessingTime;
-	const double StartTime = FPlatformTime::Seconds();
-
-	int32 ProcessedAnimation = 0;
-
-	ON_SCOPE_EXIT
-	{
-		TRACE_COUNTER_SET(ActiveAnimationCompilation, ActiveAsyncCompressionTasks.Num());
-		TRACE_COUNTER_SET(QueuedAnimationCompilation, QueuedAsyncCompressionWork.Num());
-		TRACE_COUNTER_SET(ProcessedAnimationCompilation, ProcessedAnimation);
-	};
-
-#if ASYNC_MEM_LOG
-	const TCHAR* BytesPostFix = TEXT("ERROR");
-	const float FormatSize = FormatBytes(ActiveMemoryUsage, BytesPostFix);
-	FPlatformMisc::LowLevelOutputDebugStringf(TEXT("FAsyncCompressedAnimationsManagement::Tick\n\tDelta:%.2f TimeSinceLastRun:%.2f\n\tActive:%i Queued:%i\n\tMemUsage:%.2f%s\n"), DeltaTime, StartTime - LastTickTime, ActiveAsyncCompressionTasks.Num(), QueuedAsyncCompressionWork.Num(), FormatSize, BytesPostFix);
-#endif
-
-	LastTickTime = StartTime;
-
-	FDerivedDataCacheInterface& DerivedDataCache = GetDerivedDataCacheRef();
-
-	for (int32 ActiveAnim = ActiveAsyncCompressionTasks.Num() - 1; ActiveAnim >= 0; --ActiveAnim)
-	{
-		const FActiveAsyncCompressionTask& Task = ActiveAsyncCompressionTasks[ActiveAnim];
-
-		if (DerivedDataCache.PollAsynchronousCompletion(Task.AsyncHandle))
-		{
-			OnActiveCompressionFinished(ActiveAnim);
-			ProcessedAnimation++;
-		}
-
-		if (FPlatformTime::Seconds() > EndTime)
-		{
-			return; // Finish for this tick
-		}
-	}
-
-	const bool bHasQueuedTasks = QueuedAsyncCompressionWork.Num() > 0;
-
-	StartQueuedTasks(MAX_ACTIVE_COMPRESSIONS);
-
-	if (bHasQueuedTasks && QueuedAsyncCompressionWork.Num() == 0)
-	{
-		QueuedAsyncCompressionWork.Empty(); //free memory
-	}
-}
-
-void FAsyncCompressedAnimationsManagement::TickCook(float DeltaTime, bool bCookCompete)
-{
-	Tick(DeltaTime);
-}
-
-void FAsyncCompressedAnimationsManagement::StartQueuedTasks(int32 MaxActiveTasks)
-{
-	while (ActiveAsyncCompressionTasks.Num() < MaxActiveTasks)
-	{
-		if (QueuedAsyncCompressionWork.Num() == 0)
-		{
-			break;
-		}
-		FQueuedAsyncCompressionWork NewTask = QueuedAsyncCompressionWork.Pop(false);
-		StartAsyncWork(NewTask.Compressor, NewTask.Anim, NewTask.Compressor.GetMemoryUsage(), NewTask.bPerformFrameStripping);
-	}
-}
-
-void FAsyncCompressedAnimationsManagement::AddReferencedObjects(FReferenceCollector& Collector)
-{
-	for (const FActiveAsyncCompressionTask& Task : ActiveAsyncCompressionTasks)
-	{
-		Task.DataToCompress->AddReferencedObjects(Collector);
-	}
-
-	for (const FQueuedAsyncCompressionWork& QueuedTask : QueuedAsyncCompressionWork)
-	{
-		QueuedTask.Compressor.GetCompressibleData()->AddReferencedObjects(Collector);
-	}
-}
-
-TStatId FAsyncCompressedAnimationsManagement::GetStatId() const
-{
-	RETURN_QUICK_DECLARE_CYCLE_STAT(FAsyncCompressedAnimationsTracker, STATGROUP_Tickables);
-}
-
-bool FAsyncCompressedAnimationsManagement::RequestAsyncCompression(FDerivedDataAnimationCompression& Compressor, UAnimSequence* Anim, const bool bPerformFrameStripping, TArray<uint8>& OutData)
-{
-	const uint64 NewTaskSize = Compressor.GetMemoryUsage();
-
-	bool bWasAsync = true;
-
-	const bool bOutOfMemForQueue = (ActiveMemoryUsage + NewTaskSize) >= MAX_ASYNC_COMPRESSION_MEM_USAGE;
-	const bool bGameThreadStarved = FPlatformTime::Seconds() - LastTickTime > 2.f;
-
-	//Boost Max active compressions in the case of game thread not running for extended periods (assumes heavy loading situation)
-	// so that memory usage of async compressions does not grow too large
-	const int32 MaxActiveCompressions = bGameThreadStarved ? (MAX_ACTIVE_COMPRESSIONS * 6) : MAX_ACTIVE_COMPRESSIONS;
-
-	if (bOutOfMemForQueue || bGameThreadStarved)
-	{
-		Tick(0.f); // Try to free up some memory
-	}
-
-	const bool bCanRunASync = (ActiveMemoryUsage + NewTaskSize < MAX_ASYNC_COMPRESSION_MEM_USAGE);
-	const bool bForceSync = false; //debugging override
-
-
-	if (bCanRunASync && !bForceSync)
-	{
-		//Queue Async
-		ActiveMemoryUsage += NewTaskSize;
-
-		if(ActiveAsyncCompressionTasks.Num() < MaxActiveCompressions)
-		{
-			StartAsyncWork(Compressor, Anim, NewTaskSize, bPerformFrameStripping);
-		}
-		else
-		{
-			QueuedAsyncCompressionWork.Emplace(Compressor, Anim, bPerformFrameStripping);
-		}
-	}
-	else
-	{
-		//Do in place
-		GetDerivedDataCacheRef().GetSynchronous(&Compressor, OutData);
-		bWasAsync = false;
-	}
-
-	StartQueuedTasks(MaxActiveCompressions);
-
-	return bWasAsync;
-}
-
-void FAsyncCompressedAnimationsManagement::StartAsyncWork(FDerivedDataAnimationCompression& Compressor, UAnimSequence* Anim, const uint64 NewTaskSize, const bool bPerformFrameStripping)
-{
-	TRACE_CPUPROFILER_EVENT_SCOPE(FAsyncCompressedAnimationsManagement::StartAsyncWork);
-
-	const FString CacheKey = Compressor.GetPluginSpecificCacheKeySuffix();
-	FCompressibleAnimPtr SourceData = Compressor.GetCompressibleData();
-	uint32 AsyncHandle = GetDerivedDataCacheRef().GetAsynchronous(&Compressor);
-	ActiveAsyncCompressionTasks.Emplace(Anim, SourceData, CacheKey, NewTaskSize, AsyncHandle, bPerformFrameStripping);
-}
-
-bool FAsyncCompressedAnimationsManagement::WaitOnActiveCompression(UAnimSequence* Anim, bool bWantResults)
-{
-	if (ActiveAsyncCompressionTasks.Num())
-	{
-		TRACE_CPUPROFILER_EVENT_SCOPE(FAsyncCompressedAnimationsManagement::WaitOnActiveCompression);
-
-		for (int32 ActiveIndex = 0; ActiveIndex < ActiveAsyncCompressionTasks.Num(); ++ActiveIndex)
-		{
-			FActiveAsyncCompressionTask& Task = ActiveAsyncCompressionTasks[ActiveIndex];
-			if (Task.Sequence == Anim)
-			{
-				if(bWantResults)
-				{
-					GetDerivedDataCacheRef().WaitAsynchronousCompletion(Task.AsyncHandle);
-					OnActiveCompressionFinished(ActiveIndex);
-				}
-				else
-				{
-					Task.DataToCompress->IsCancelledSignal.Cancel();
-					Task.Sequence = nullptr;
-				}
-				return true; // Was active
-			}
-		}
-	}
-	return false;
-}
-
-bool FAsyncCompressedAnimationsManagement::WaitOnExistingCompression(UAnimSequence* Anim, const bool bWantResults)
-{
-	if(!WaitOnActiveCompression(Anim, bWantResults))
-	{
-		//Check if we have a queued task
-		for (int32 QueuedIndex = 0; QueuedIndex < QueuedAsyncCompressionWork.Num(); ++QueuedIndex)
-		{
-			FQueuedAsyncCompressionWork& Task = QueuedAsyncCompressionWork[QueuedIndex];
-			if (Task.Anim == Anim)
-			{
-				if (bWantResults)
-				{
-					StartAsyncWork(Task.Compressor, Task.Anim, Task.Compressor.GetMemoryUsage(), Task.bPerformFrameStripping);
-				}
-				else
-				{
-					delete &Task.Compressor;
-				}
-				QueuedAsyncCompressionWork.RemoveAtSwap(QueuedIndex, 1, false);
-			}
-		}
-		return bWantResults ? WaitOnActiveCompression(Anim, bWantResults) : false;
-	}
-	return true;
-}
-
 #endif	//WITH_EDITOR
diff --git a/Engine/Source/Runtime/Engine/Private/Animation/AnimCompressionDerivedData.h b/Engine/Source/Runtime/Engine/Private/Animation/AnimCompressionDerivedData.h
--- a/Engine/Source/Runtime/Engine/Private/Animation/AnimCompressionDerivedData.h
+++ b/Engine/Source/Runtime/Engine/Private/Animation/AnimCompressionDerivedData.h
@@ -30,10 +30,11 @@ private:
 	const FString AssetDDCKey;
 
 	// FAnimCompressContext to use during compression if we don't pull from the DDC
+	PRAGMA_DISABLE_DEPRECATION_WARNINGS
 	TSharedPtr<FAnimCompressContext> CompressContext;
-
 public:
-	FDerivedDataAnimationCompression(const TCHAR* InTypeName, const FString& InAssetDDCKey, TSharedPtr<FAnimCompressContext> InCompressContext);
+	FDerivedDataAnimationCompression(const TCHAR* InTypeName, const FString& InAssetDDCKey, TSharedPtr<FAnimCompressContext> InCompressContext);	
+	PRAGMA_ENABLE_DEPRECATION_WARNINGS
 	virtual ~FDerivedDataAnimationCompression();
 
 	void SetCompressibleData(FCompressibleAnimRef InCompressibleAnimData)
diff --git a/Engine/Source/Runtime/Engine/Private/Animation/AnimCompressionTypes.cpp b/Engine/Source/Runtime/Engine/Private/Animation/AnimCompressionTypes.cpp
--- a/Engine/Source/Runtime/Engine/Private/Animation/AnimCompressionTypes.cpp
+++ b/Engine/Source/Runtime/Engine/Private/Animation/AnimCompressionTypes.cpp
@@ -18,6 +18,7 @@
 #include "UObject/LinkerLoad.h"
 #include "Animation/SkeletonRemapping.h"
 #include "Animation/SkeletonRemappingRegistry.h"
+#include "Animation/AttributesRuntime.h"
 
 static FString GCompressionJsonOutput;
 static FAutoConsoleVariableRef CVarCompressionJsonOutput(
@@ -31,6 +32,18 @@ static FAutoConsoleVariableRef CVarCompressionJsonOutput(
 	ECVF_Cheat
 	);
 
+PRAGMA_DISABLE_DEPRECATION_WARNINGS
+ICompressedAnimData::ICompressedAnimData(const ICompressedAnimData&) = default;
+ICompressedAnimData& ICompressedAnimData::operator=(const ICompressedAnimData&) = default;
+PRAGMA_ENABLE_DEPRECATION_WARNINGS
+
+#if WITH_EDITOR
+
+PRAGMA_DISABLE_DEPRECATION_WARNINGS
+FCompressibleAnimData::FCompressibleAnimData(const FCompressibleAnimData&) = default;
+FCompressibleAnimData& FCompressibleAnimData::operator=(const FCompressibleAnimData&) = default;
+PRAGMA_ENABLE_DEPRECATION_WARNINGS
+
 template<typename ArrayValue>
 void StripFramesEven(TArray<ArrayValue>& Keys, const int32 NumFrames)
 {
@@ -89,12 +102,394 @@ void StripFramesOdd(TArray<ArrayValue>& Keys, const int32 NumFrames)
 	}
 }
 
-PRAGMA_DISABLE_DEPRECATION_WARNINGS
-FCompressibleAnimData::FCompressibleAnimData(const FCompressibleAnimData&) = default;
-FCompressibleAnimData& FCompressibleAnimData::operator=(const FCompressibleAnimData&) = default;
-ICompressedAnimData::ICompressedAnimData(const ICompressedAnimData&) = default;
-ICompressedAnimData& ICompressedAnimData::operator=(const ICompressedAnimData&) = default;
-PRAGMA_ENABLE_DEPRECATION_WARNINGS
+struct FByFramePoseEvalContext
+{
+public:
+	FBoneContainer RequiredBones;
+
+	TArray<FBoneIndexType> RequiredBoneIndexArray;
+
+	FByFramePoseEvalContext(const UAnimSequence* InAnimToEval)
+		: FByFramePoseEvalContext(InAnimToEval->GetSkeleton())
+	{}
+		
+	FByFramePoseEvalContext(USkeleton* InSkeleton)
+	{
+		// Initialize RequiredBones for pose evaluation
+		RequiredBones.SetUseRAWData(true);
+
+		check(InSkeleton);
+
+		RequiredBoneIndexArray.AddUninitialized(InSkeleton->GetReferenceSkeleton().GetNum());
+		for (int32 BoneIndex = 0; BoneIndex < RequiredBoneIndexArray.Num(); ++BoneIndex)
+		{
+			RequiredBoneIndexArray[BoneIndex] = BoneIndex;
+		}
+
+		RequiredBones.InitializeTo(RequiredBoneIndexArray, FCurveEvaluationOption(true), *InSkeleton);
+	}
+
+};
+
+static void CopyTransformToRawAnimationData(const FTransform& BoneTransform, FRawAnimSequenceTrack& Track, int32 Frame)
+{
+	Track.PosKeys[Frame] = FVector3f(BoneTransform.GetTranslation());
+	Track.RotKeys[Frame] = FQuat4f(BoneTransform.GetRotation());
+	Track.RotKeys[Frame].Normalize();
+	Track.ScaleKeys[Frame] = FVector3f(BoneTransform.GetScale3D());
+}
+
+static FFloatCurve* GetFloatCurve(TArray<FFloatCurve>& FloatCurves, const USkeleton::AnimCurveUID& CurveUID)
+{
+	for (FFloatCurve& Curve : FloatCurves)
+	{
+		if (Curve.Name.UID == CurveUID)
+		{
+			return &Curve;
+		}
+	}
+
+	return nullptr;
+}
+
+static FFloatCurve* FindOrAddCurve(TArray<FFloatCurve>& FloatCurves, const FSmartName& CurveName)
+{
+	FFloatCurve* ReturnCurve = GetFloatCurve(FloatCurves, CurveName.UID);
+
+	if (ReturnCurve == nullptr)
+	{
+		FFloatCurve& NewCurve = FloatCurves.Add_GetRef(FFloatCurve(CurveName, 0));
+		ReturnCurve = &NewCurve;
+	}
+
+	return ReturnCurve;
+	
+}
+
+static bool IsNewKeyDifferent(const FRichCurveKey& LastKey, float NewValue)
+{
+	return LastKey.Value != NewValue;
+}
+
+void FCompressibleAnimData::BakeOutAdditiveIntoRawData(const FFrameRate& SampleRate, TArray<FBoneAnimationTrack>& ResampledTrackData, TArray<FFloatCurve>& FloatCurves)
+{
+	TRACE_CPUPROFILER_EVENT_SCOPE(FCompressibleAnimData::BakeOutAdditiveIntoRawData);
+	const UAnimSequence* AnimSequence = WeakSequence.Get();
+	checkf(AnimSequence, TEXT("Animation Sequence is invalid"));
+
+	const USkeleton* Skeleton = AnimSequence->GetSkeleton();
+	checkf(Skeleton, TEXT("No valid skeleton found"));
+
+	const bool bSelfAdditiveType = AnimSequence->RefPoseType == EAdditiveBasePoseType::ABPT_RefPose || AnimSequence->RefPoseType == EAdditiveBasePoseType::ABPT_LocalAnimFrame;
+	const UAnimSequence* AdditiveBaseAnimation = AnimSequence->RefPoseSeq.Get();
+	checkf(bSelfAdditiveType || (AdditiveBaseAnimation && !AdditiveBaseAnimation->HasAnyFlags(EObjectFlags::RF_NeedPostLoad)), TEXT("Invalid additive base animation state"));
+
+	// Lock DataModel evaluation
+	IAnimationDataModel::FEvaluationAndModificationLock Lock(*AnimSequence->GetDataModelInterface());
+	
+	TUniquePtr<IAnimationDataModel::FEvaluationAndModificationLock> AdditiveBaseLock = nullptr;
+	if (!bSelfAdditiveType)
+	{
+		// Additive is based off another Animation Asset, so lock its evaluation path as well
+		AdditiveBaseLock = MakeUnique<IAnimationDataModel::FEvaluationAndModificationLock>(*AdditiveBaseAnimation->GetDataModelInterface());
+	}
+	
+	FMemMark Mark(FMemStack::Get());
+	FByFramePoseEvalContext EvalContext(AnimSequence);
+
+	TScriptInterface<IAnimationDataModel> DataModelInterface = AnimSequence->GetDataModelInterface();
+	// We actually need to resample bone transforms
+	const FFrameNumber ModelNumberOfFrames = DataModelInterface->GetNumberOfFrames();
+	const FFrameTime ResampledFrameTime = FFrameRate::TransformTime(ModelNumberOfFrames, DataModelInterface->GetFrameRate(), SampleRate);
+	ensureMsgf(FMath::IsNearlyZero(ResampledFrameTime.GetSubFrame()), TEXT("Incompatible resampling frame rate for animation sequence %s, frame remainder of %1.8f"), *AnimSequence->GetName(), ResampledFrameTime.GetSubFrame());
+
+	const int32 SampledFrames = ResampledFrameTime.FloorToFrame().Value;
+	const int32 SampledKeys = SampledFrames + 1;
+
+	ResampledTrackData.SetNum(EvalContext.RequiredBoneIndexArray.Num());
+	AdditiveBaseAnimationData.SetNum(EvalContext.RequiredBoneIndexArray.Num());
+	
+	const FReferenceSkeleton& ReferenceSkeleton = Skeleton->GetReferenceSkeleton();
+
+	// Populate tracks
+	for (int32 TrackIndex = 0; TrackIndex < EvalContext.RequiredBoneIndexArray.Num(); ++TrackIndex)
+	{
+		FBoneAnimationTrack& Track = ResampledTrackData[TrackIndex];
+		
+		Track.InternalTrackData.PosKeys.SetNumUninitialized(SampledKeys);
+		Track.InternalTrackData.RotKeys.SetNumUninitialized(SampledKeys);
+		Track.InternalTrackData.ScaleKeys.SetNumUninitialized(SampledKeys);
+
+		Track.BoneTreeIndex = TrackIndex;
+		Track.Name = ReferenceSkeleton.GetBoneName(TrackIndex);
+		
+		AdditiveBaseAnimationData[TrackIndex].PosKeys.SetNumUninitialized(SampledKeys);
+		AdditiveBaseAnimationData[TrackIndex].RotKeys.SetNumUninitialized(SampledKeys);
+		AdditiveBaseAnimationData[TrackIndex].ScaleKeys.SetNumUninitialized(SampledKeys);
+	}
+
+	{
+		//Pose evaluation data
+		FCompactPose Pose;
+		Pose.SetBoneContainer(&EvalContext.RequiredBones);
+		FCompactPose BasePose;
+		BasePose.SetBoneContainer(&EvalContext.RequiredBones);
+
+		FAnimExtractContext ExtractContext;
+
+		FBlendedCurve Curve;
+		FBlendedCurve DummyBaseCurve;
+
+		for (int32 KeyIndex = 0; KeyIndex < SampledKeys; ++KeyIndex)
+		{
+			// Initialise curve data from Skeleton
+			Curve.InitFrom(EvalContext.RequiredBones);
+			DummyBaseCurve.InitFrom(EvalContext.RequiredBones);
+
+			//Grab pose for this frame
+			const double PreviousKeyTime = SampleRate.AsSeconds(KeyIndex - 1);
+			const double CurrentKeyTime = SampleRate.AsSeconds(KeyIndex);
+			FFrameTime CurrentFrameTime = KeyIndex;
+
+			ExtractContext.CurrentTime = CurrentKeyTime;
+
+			UE::Anim::FStackAttributeContainer BaseAttributes;
+			FAnimationPoseData AnimPoseData(Pose, Curve, BaseAttributes);
+			AnimSequence->GetAnimationPose(AnimPoseData, ExtractContext);
+
+			UE::Anim::FStackAttributeContainer AdditiveAttributes;
+			FAnimationPoseData AnimBasePoseData(BasePose, DummyBaseCurve, AdditiveAttributes);
+			AnimSequence->GetAdditiveBasePose(AnimBasePoseData, ExtractContext);
+
+			//Write out every track for this frame
+			for (FCompactPoseBoneIndex TrackIndex(0); TrackIndex < ResampledTrackData.Num(); ++TrackIndex)
+			{
+				CopyTransformToRawAnimationData(Pose[TrackIndex], ResampledTrackData[TrackIndex.GetInt()].InternalTrackData, KeyIndex);
+				CopyTransformToRawAnimationData(BasePose[TrackIndex], AdditiveBaseAnimationData[TrackIndex.GetInt()], KeyIndex);
+			}
+
+			//Write out curve data for this frame
+			for (int32 CurveIndex = 0; CurveIndex < Curve.UIDToArrayIndexLUT->Num(); ++CurveIndex)
+			{
+				USkeleton::AnimCurveUID CurveUID = (USkeleton::AnimCurveUID) CurveIndex;
+				int32 ArrayIndex = Curve.GetArrayIndexByUID(CurveUID);
+				if (ArrayIndex != INDEX_NONE)
+				{
+					float CurveWeight = Curve.CurveWeights[ArrayIndex];
+					FFloatCurve* RawCurve = GetFloatCurve(FloatCurves, CurveUID);
+					if (!RawCurve && !FMath::IsNearlyZero(CurveWeight)) //Only make a new curve if we are going to give it data
+					{
+						FSmartName NewCurveName;
+						// if we don't have name, there is something wrong here. 
+						ensureAlways(Skeleton->GetSmartNameByUID(USkeleton::AnimCurveMappingName, CurveUID, NewCurveName));
+						// curve flags don't matter much for compressed curves
+						RawCurve = FindOrAddCurve(FloatCurves, NewCurveName);
+					}
+
+					if (RawCurve)
+					{
+						const bool bHasKeys = RawCurve->FloatCurve.GetNumKeys() > 0;
+						if (!bHasKeys)
+						{
+							//Add pre key of 0
+							if (KeyIndex > 0)
+							{
+								RawCurve->UpdateOrAddKey(0.f, PreviousKeyTime);
+							}
+						}
+
+						if (!bHasKeys || IsNewKeyDifferent(RawCurve->FloatCurve.GetLastKey(), CurveWeight))
+						{
+							RawCurve->UpdateOrAddKey(CurveWeight, CurrentKeyTime);
+							TArray<FRichCurveKey>& CurveKeys = RawCurve->FloatCurve.Keys;
+							if (CurveKeys.Num() > 1)
+							{
+								FRichCurveKey& PrevKey = CurveKeys.Last(1);
+								// Round to frame here as it would have been added at a specific frame boundary (though float->double conversion might mean the value is off)
+								const FFrameTime PrevKeyTime = SampleRate.AsFrameTime(PrevKey.Time).RoundToFrame();
+								if (PrevKeyTime < CurrentFrameTime - 1) // Did we skip a frame, if so need to make previous key const
+								{
+									PrevKey.InterpMode = RCIM_Constant;
+								}
+							}
+						}
+					}
+				}
+			}
+		}
+	}
+}
+
+static void FindAnimatedVirtualBones(const TArray<FBoneAnimationTrack>& AnimatedBoneTracks, const USkeleton* Skeleton, TArray<int32>& OutVirtualBoneIndices)
+{
+	TArray<int32> SourceParents;
+	const TArray<FVirtualBone>& VirtualBones = Skeleton->GetVirtualBones();
+	const FReferenceSkeleton& ReferenceSkeleton = Skeleton->GetReferenceSkeleton();
+	for (int32 VBIndex = 0; VBIndex < VirtualBones.Num(); ++VBIndex)
+	{
+		const FVirtualBone& VirtualBone = VirtualBones[VBIndex];
+		if (Algo::FindBy(AnimatedBoneTracks, VirtualBone.VirtualBoneName, &FBoneAnimationTrack::Name) == nullptr)
+		{
+			//Need to test if we will animation virtual bone. This involves seeing if any bone that can affect the position
+			//of the target relative to the source is animated by this animation. A bone that can affect the relative position
+			//is any both that is a child of the common ancestor of the target and source
+
+			SourceParents.Reset();
+			bool bBuildVirtualBone = false;
+
+			// First get all the bones that form the chain to the source bone. 
+			int32 CurrentBone = ReferenceSkeleton.FindBoneIndex(VirtualBone.SourceBoneName);
+			while (CurrentBone != INDEX_NONE)
+			{
+				SourceParents.Add(CurrentBone);
+				CurrentBone = ReferenceSkeleton.GetParentIndex(CurrentBone);
+			}
+
+			// Now start checking every bone in the target bones hierarchy until a common ancestor is reached. 
+			CurrentBone = ReferenceSkeleton.FindBoneIndex(VirtualBone.TargetBoneName);
+
+			while (!SourceParents.Contains(CurrentBone))
+			{
+				if (Algo::FindBy(AnimatedBoneTracks, CurrentBone, &FBoneAnimationTrack::BoneTreeIndex) != nullptr)
+				{
+					//We animate this bone so the virtual bone is needed
+					bBuildVirtualBone = true;
+					break;
+				}
+
+				CurrentBone = ReferenceSkeleton.GetParentIndex(CurrentBone);
+				check(CurrentBone != INDEX_NONE);
+			}
+
+			// Now we have all the non common bones from the target chain we need the same check from the source chain
+			const int32 FirstCommon = SourceParents.IndexOfByKey(CurrentBone);
+			for (int32 i = FirstCommon - 1; i >= 0; --i)
+			{
+				if (Algo::FindBy(AnimatedBoneTracks, i, &FBoneAnimationTrack::BoneTreeIndex) != nullptr)
+				{
+					//We animate this bone so the virtual bone is needed
+					bBuildVirtualBone = true;
+					break;
+				}
+			}
+
+			if (bBuildVirtualBone)
+			{
+				OutVirtualBoneIndices.Add(VBIndex);
+			}
+		}
+	}
+}
+
+void FCompressibleAnimData::ResampleAnimationTrackData(const FFrameRate& SampleRate, TArray<FBoneAnimationTrack>& ResampledTrackData) const
+{
+	UAnimSequence* AnimSequence = WeakSequence.Get();
+	checkf(AnimSequence, TEXT("Animation Sequence is invalid"));
+	
+	QUICK_SCOPE_CYCLE_COUNTER(STAT_ResampleAnimationTrackData);
+	
+	if (AnimSequence->ShouldDataModelBeValid())
+	{
+		FAnimationUtils::EnsureAnimSequenceLoaded(*AnimSequence);
+
+		TScriptInterface<IAnimationDataModel> DataModelInterface = AnimSequence->GetDataModelInterface();
+
+		// Make a copy, deals with bone name and index
+		TArray<FName> TrackNames;
+		DataModelInterface->GetBoneTrackNames(TrackNames);
+	
+		ensure(DataModelInterface != nullptr);
+		{
+			IAnimationDataModel::FEvaluationAndModificationLock Lock(*DataModelInterface);
+			
+			FMemMark Mark(FMemStack::Get());
+			FByFramePoseEvalContext EvalContext(AnimSequence);
+			EvalContext.RequiredBones.SetDisableRetargeting(true);
+			EvalContext.RequiredBones.SetUseRAWData(true);
+			EvalContext.RequiredBones.SetUseSourceData(false);
+
+			TArray<FName> BoneTrackNames;
+			BoneTrackNames.Reserve(DataModelInterface->GetNumBoneTracks());
+			DataModelInterface->GetBoneTrackNames(BoneTrackNames);
+			const USkeleton* MySkeleton = AnimSequence->GetSkeleton();
+
+
+			// We actually need to resample bone transforms
+			const FFrameNumber ModelNumberOfFrames = DataModelInterface->GetNumberOfFrames();
+			const FFrameTime ResampledFrameTime = FFrameRate::TransformTime(ModelNumberOfFrames, DataModelInterface->GetFrameRate(), SampleRate);
+			ensureMsgf(FMath::IsNearlyZero(ResampledFrameTime.GetSubFrame()), TEXT("Incompatible resampling frame rate for animation sequence %s, frame remainder of %1.8f"), *AnimSequence->GetName(), ResampledFrameTime.GetSubFrame());
+
+			const int32 SampledFrames = ResampledFrameTime.FloorToFrame().Value;
+			const int32 SampledKeys = SampledFrames + 1;
+
+			ResampledTrackData.Reset(BoneTrackNames.Num());
+
+			const FReferenceSkeleton& ReferenceSkeleton = MySkeleton->GetReferenceSkeleton();
+			for (int32 TrackIndex = 0; TrackIndex < BoneTrackNames.Num(); ++TrackIndex)
+			{
+				const int32 BoneIndex = ReferenceSkeleton.FindBoneIndex(BoneTrackNames[TrackIndex]);
+				if (BoneIndex != INDEX_NONE)
+				{
+					FBoneAnimationTrack& TrackData = ResampledTrackData.AddDefaulted_GetRef();
+					TrackData.Name = BoneTrackNames[TrackIndex];
+					TrackData.BoneTreeIndex = BoneIndex;
+
+					TrackData.InternalTrackData.PosKeys.SetNumUninitialized(SampledKeys);
+					TrackData.InternalTrackData.RotKeys.SetNumUninitialized(SampledKeys);
+					TrackData.InternalTrackData.ScaleKeys.SetNumUninitialized(SampledKeys);
+				}
+			}
+
+			const TArray<FVirtualBone>& VirtualBones = MySkeleton->GetVirtualBones();
+			if (VirtualBones.Num())
+			{
+				TArray<int32> VirtualBoneIndices;
+				FindAnimatedVirtualBones(ResampledTrackData, MySkeleton, VirtualBoneIndices);
+
+				for (int32 VBIndex : VirtualBoneIndices)
+				{
+					const FVirtualBone& VirtualBone = VirtualBones[VBIndex];
+					FBoneAnimationTrack& TrackData = ResampledTrackData.AddDefaulted_GetRef();
+					
+					TrackData.Name = VirtualBone.VirtualBoneName;
+					
+					const int32 VirtualBoneSkeletonIndex = ReferenceSkeleton.GetRequiredVirtualBones()[VBIndex];
+					TrackData.BoneTreeIndex = VirtualBoneSkeletonIndex;
+
+					TrackData.InternalTrackData.PosKeys.SetNumUninitialized(SampledKeys);
+					TrackData.InternalTrackData.RotKeys.SetNumUninitialized(SampledKeys);
+					TrackData.InternalTrackData.ScaleKeys.SetNumUninitialized(SampledKeys);
+				}
+			}
+			
+			FBlendedCurve Curve;
+			Curve.InitFrom(EvalContext.RequiredBones);
+			UE::Anim::FStackAttributeContainer AttributeContainer;
+						
+			for (int32 FrameIndex = 0; FrameIndex < SampledKeys; ++FrameIndex)
+			{
+				UE::Anim::DataModel::FEvaluationContext EvaluationContext(FFrameTime(FrameIndex), SampleRate, AnimSequence->GetRetargetTransformsSourceName(), AnimSequence->GetRetargetTransforms());
+				
+				FCompactPose Pose;
+				Pose.SetBoneContainer(&EvalContext.RequiredBones);
+
+				FAnimationPoseData PoseData(Pose, Curve, AttributeContainer);
+				DataModelInterface->Evaluate(PoseData, EvaluationContext);
+
+				for (int32 TrackIndex = 0; TrackIndex < ResampledTrackData.Num(); ++TrackIndex)
+				{
+					FBoneAnimationTrack& TrackData = ResampledTrackData[TrackIndex];
+					const FTransform& Transform = Pose[EvalContext.RequiredBones.GetCompactPoseIndexFromSkeletonIndex(TrackData.BoneTreeIndex)];
+
+					TrackData.InternalTrackData.PosKeys[FrameIndex] = FVector3f(Transform.GetLocation());
+					TrackData.InternalTrackData.RotKeys[FrameIndex] = FQuat4f(Transform.GetRotation());
+					TrackData.InternalTrackData.RotKeys[FrameIndex].Normalize();
+					TrackData.InternalTrackData.ScaleKeys[FrameIndex] = FVector3f(Transform.GetScale3D());
+				}
+			}
+		}
+	}	
+}
 
 void FCompressibleAnimData::WriteCompressionDataToJSON(TArrayView<FName> OriginalTrackNames, TArrayView<FRawAnimSequenceTrack> FinalRawAnimationData, TArrayView<FName> FinalTrackNames) const
 {
@@ -309,8 +704,7 @@ FCompressibleAnimData::FCompressibleAnimData(class UAnimSequence* InSeq, const b
 	: CurveCompressionSettings(InSeq->CurveCompressionSettings)
 	, BoneCompressionSettings(InSeq->BoneCompressionSettings)
 	, Interpolation(InSeq->Interpolation)
-	, SequenceLength(InSeq->GetPlayLength())
-	, NumberOfKeys(InSeq->GetNumberOfSampledKeys())
+	, SequenceLength(InSeq->GetDataModelInterface()->GetPlayLength())
 	, bIsValidAdditive(InSeq->IsValidAdditive())
 #if WITH_EDITORONLY_DATA
 	, ErrorThresholdScale(InSeq->CompressionErrorThresholdScale)
@@ -320,28 +714,60 @@ FCompressibleAnimData::FCompressibleAnimData(class UAnimSequence* InSeq, const b
 	, Name(InSeq->GetName())
 	, FullName(InSeq->GetFullName())
 	, AnimFName(InSeq->GetFName())
+	, bShouldPerformStripping(bPerformStripping)
+	, WeakSequence(InSeq)
 	, TargetPlatform(InTargetPlatform)
 {
+}
+
+void FCompressibleAnimData::FetchData(const ITargetPlatform* InPlatform)
+{
 #if WITH_EDITOR
-	TRACE_CPUPROFILER_EVENT_SCOPE(FCompressibleAnimData::FCompressibleAnimData);
+	if (bDataFetched)
+		return;
+
+	bDataFetched = true;
+	TRACE_CPUPROFILER_EVENT_SCOPE(FCompressibleAnimData::FetchData);
 
-	USkeleton* Skeleton = InSeq->GetSkeleton();
+	const UAnimSequence* AnimSequence = WeakSequence.Get();
+	USkeleton* Skeleton = AnimSequence->GetSkeleton();
 	FAnimationUtils::BuildSkeletonMetaData(Skeleton, BoneData);
 
+	const FFrameRate DefaultSamplingFrameRate = AnimSequence->GetSamplingFrameRate();
+	const FFrameRate PlatformSamplingFrameRate = InPlatform ? AnimSequence->GetTargetSamplingFrameRate(InPlatform) : FFrameRate(0,0);
+	
+	const bool bValidTargetSampleRate = PlatformSamplingFrameRate.IsValid() && (PlatformSamplingFrameRate.IsMultipleOf(DefaultSamplingFrameRate) || PlatformSamplingFrameRate.IsFactorOf(DefaultSamplingFrameRate));
+	
+	const FFrameRate& FrameRateToSampleWith = (DefaultSamplingFrameRate != PlatformSamplingFrameRate && bValidTargetSampleRate) ? PlatformSamplingFrameRate : DefaultSamplingFrameRate;	
+	
+	const FFrameTime SampleFrameTime = FrameRateToSampleWith.AsFrameTime(SequenceLength);
+	check(FMath::IsNearlyZero(SampleFrameTime.GetSubFrame()));	
+	NumberOfKeys = SampleFrameTime.FrameNumber.Value + 1;
+	
 	RefLocalPoses = Skeleton->GetRefLocalPoses();
 	RefSkeleton = Skeleton->GetReferenceSkeleton();
 
-	const bool bHasVirtualBones = InSeq->GetSkeleton()->GetVirtualBones().Num() > 0;
+	const bool bHasVirtualBones = Skeleton->GetVirtualBones().Num() > 0;
 
 	/* Always get the resampled data to start off with */
-	const TArray<FBoneAnimationTrack>& ResampledTrackData = InSeq->GetResampledTrackData();
+	TArray<FBoneAnimationTrack> ResampledTrackData;
 
-	SampledFrameRate = InSeq->GetSamplingFrameRate();
+	SampledFrameRate = FrameRateToSampleWith;
 
 	TArray<FName> OriginalTrackNames;
 	RawAnimationData.Empty(ResampledTrackData.Num());
 	TrackToSkeletonMapTable.Empty(ResampledTrackData.Num());
 	OriginalTrackNames.Empty(ResampledTrackData.Num());
+	const bool bIsAdditiveAnimation = AnimSequence->CanBakeAdditive();
+	if (bIsAdditiveAnimation)
+	{
+		BakeOutAdditiveIntoRawData(SampledFrameRate, ResampledTrackData, RawFloatCurves);
+	}
+	else
+    {
+		ResampleAnimationTrackData(SampledFrameRate, ResampledTrackData);	
+		RawFloatCurves = AnimSequence->GetDataModelInterface()->GetFloatCurves();
+    }
 
 	for (const FBoneAnimationTrack& AnimTrack : ResampledTrackData)
 	{
@@ -349,35 +775,13 @@ FCompressibleAnimData::FCompressibleAnimData(class UAnimSequence* InSeq, const b
 		UE::Anim::Compression::SanitizeRawAnimSequenceTrack(Track);
 		TrackToSkeletonMapTable.Add(AnimTrack.BoneTreeIndex);
 		OriginalTrackNames.Add(AnimTrack.Name);
-	}
-
-	const bool bIsAdditiveAnimation = InSeq->CanBakeAdditive();
-	if (bIsAdditiveAnimation)
-	{
-		InSeq->BakeOutAdditiveIntoRawData(RawAnimationData, OriginalTrackNames, TrackToSkeletonMapTable, RawFloatCurves, AdditiveBaseAnimationData);
-	}
-	else
-    {
-		// In case we require baking down transform curves, do so now meaning Virtual Bone baking will incorporate the correct bone transforms
-	    if (InSeq->GetDataModel()->GetNumberOfTransformCurves() > 0)
-	    {
-		    InSeq->BakeTrackCurvesToRawAnimationTracks(RawAnimationData, OriginalTrackNames, TrackToSkeletonMapTable);
-	    }
-
-		RawFloatCurves = InSeq->GetCurveData().FloatCurves;
-
-		// If we aren't additive we must bake virtual bones
-	    if (bHasVirtualBones)
-	    {
-		    InSeq->BakeOutVirtualBoneTracks(RawAnimationData, OriginalTrackNames, TrackToSkeletonMapTable);
-	    }
     } 
 
 	// Apply any key reduction if possible
 	if (RawAnimationData.Num())
 	{ 
-		UE::Anim::Compression::CompressAnimationDataTracks(RawAnimationData, NumberOfKeys, InSeq->GetFName(), -1.f, -1.f);
-		UE::Anim::Compression::CompressAnimationDataTracks(RawAnimationData, NumberOfKeys, InSeq->GetFName());
+		UE::Anim::Compression::CompressAnimationDataTracks(RawAnimationData, NumberOfKeys, AnimSequence->GetFName(), -1.f, -1.f);
+		UE::Anim::Compression::CompressAnimationDataTracks(RawAnimationData, NumberOfKeys, AnimSequence->GetFName());
 	}
 
 	auto IsKeyArrayValidForRemoval = [](const auto& Keys, const auto& IdentityValue) -> bool
@@ -458,7 +862,7 @@ FCompressibleAnimData::FCompressibleAnimData(class UAnimSequence* InSeq, const b
 		Skeleton->VerifySmartName(USkeleton::AnimCurveMappingName, Curve.Name);
 	}
 	   
-	if (bPerformStripping)
+	if (bShouldPerformStripping)
 	{
 		const int32 NumTracks = RawAnimationData.Num();
 		
@@ -553,6 +957,8 @@ void FCompressibleAnimData::Update(FCompressedAnimSequence& InOutCompressedData)
 	}
 }
 
+#endif // WITH_EDITOR
+
 template<typename T>
 void WriteArray(FMemoryWriter& MemoryWriter, TArray<T>& Array)
 {
@@ -1026,6 +1432,15 @@ SIZE_T FCompressedAnimSequence::GetMemorySize() const
 			+ sizeof(FCompressedAnimSequence);
 }
 
+void FCompressedAnimSequence::Reset()
+{
+	CompressedTrackToSkeletonMapTable.Empty();
+	CompressedCurveNames.Empty();
+	ClearCompressedBoneData();
+	ClearCompressedCurveData();
+	CompressedRawDataSize = 0;
+}
+
 void FCompressedAnimSequence::ClearCompressedBoneData()
 {
 	CompressedByteStream.Empty(0);
@@ -1039,6 +1454,21 @@ void FCompressedAnimSequence::ClearCompressedCurveData()
 	CurveCompressionCodec = nullptr;
 }
 
+bool FCompressedAnimSequence::IsValid(const UAnimSequence* AnimSequence) const
+{
+#if WITH_EDITOR
+	const bool bHasBoneTracks = AnimSequence->GetDataModelInterface()->GetNumBoneTracks() != 0;
+	const bool bHasCurveData = AnimSequence->GetDataModelInterface()->GetNumberOfFloatCurves() != 0;
+
+	const bool bValidCompressedData = (!bHasBoneTracks || CompressedCurveByteStream.Num() != 0) &&
+		(!bHasCurveData || CompressedCurveByteStream.Num() != 0 || AnimSequence->IsValidAdditive());
+
+	return bValidCompressedData;
+#else
+	return true;
+#endif
+}
+
 void DecompressPose(FCompactPose& OutPose, const FCompressedAnimSequence& CompressedData, const FAnimExtractContext& ExtractionContext, USkeleton* SourceSkeleton, float SequenceLength, EAnimInterpolationType Interpolation, bool bIsBakedAdditive, FName RetargetSource, FName SourceName, const FRootMotionReset& RootMotionReset)
 {
 	const TArray<FTransform>& RetargetTransforms = SourceSkeleton->GetRefLocalPoses(RetargetSource);
@@ -1083,5 +1513,81 @@ UE::Anim::Compression::FAnimDDCKeyArgs::FAnimDDCKeyArgs(const UAnimSequenceBase&
 	, TargetPlatform(TargetPlatform)
 {
 }
-#endif
+#endif // WITH_EDITORONLY_DATA
+
+#if WITH_EDITOR
+template <typename ArrayType>
+void UpdateSHAWithArray(FSHA1& Sha, const TArray<ArrayType>& Array)
+{
+	Sha.Update((uint8*)Array.GetData(), Array.Num() * Array.GetTypeSize());
+}
+
+void UpdateSHAWithRawTrack(FSHA1& Sha, const FRawAnimSequenceTrack& RawTrack)
+{
+	UpdateSHAWithArray(Sha, RawTrack.PosKeys);
+	UpdateSHAWithArray(Sha, RawTrack.RotKeys);
+	UpdateSHAWithArray(Sha, RawTrack.ScaleKeys);
+}
+
+template<class DataType>
+void UpdateWithData(FSHA1& Sha, const DataType& Data)
+{
+	Sha.Update((uint8*)(&Data), sizeof(DataType));
+}
+
+void UpdateSHAWithCurves(FSHA1& Sha, const FRawCurveTracks& InRawCurveData) 
+{
+	auto UpdateWithFloatCurve = [&Sha](const FRichCurve& Curve)
+	{
+		UpdateWithData(Sha, Curve.DefaultValue);
+		UpdateSHAWithArray(Sha, Curve.GetConstRefOfKeys());
+		UpdateWithData(Sha, Curve.PreInfinityExtrap);
+		UpdateWithData(Sha, Curve.PostInfinityExtrap);
+	};
+
+	for (const FFloatCurve& Curve : InRawCurveData.FloatCurves)
+	{
+		UpdateWithData(Sha, Curve.Name.UID);
+		UpdateWithFloatCurve(Curve.FloatCurve);
+	}
+
+	for (const FTransformCurve& Curve : InRawCurveData.TransformCurves)
+	{
+		UpdateWithData(Sha, Curve.Name.UID);
+
+		auto UpdateWithComponent = [&Sha, &UpdateWithFloatCurve](const FVectorCurve& VectorCurve)
+		{
+			for (int32 ChannelIndex = 0; ChannelIndex < 3; ++ChannelIndex)
+			{
+				UpdateWithFloatCurve(VectorCurve.FloatCurves[ChannelIndex]);
+			}
+		};
+		
+		UpdateWithComponent(Curve.TranslationCurve);
+		UpdateWithComponent(Curve.RotationCurve);
+		UpdateWithComponent(Curve.ScaleCurve);
+	}
+}
+
+FGuid GenerateGuidFromRawAnimData(const TArray<FRawAnimSequenceTrack>& RawAnimationData, const FRawCurveTracks& RawCurveData)
+{
+	FSHA1 Sha;
+
+	for (const FRawAnimSequenceTrack& Track : RawAnimationData)
+	{
+		UpdateSHAWithRawTrack(Sha, Track);
+	}
+
+	UpdateSHAWithCurves(Sha, RawCurveData);
+
+	Sha.Final();
+
+	uint32 Hash[5];
+	Sha.GetHash((uint8*)Hash);
+	FGuid Guid(Hash[0] ^ Hash[4], Hash[1], Hash[2], Hash[3]);
+	return Guid;
+}
+
+#endif // WITH_EDITOR
+
 
diff --git a/Engine/Source/Runtime/Engine/Private/Animation/AnimCurveCompressionCodec_UniformIndexable.cpp b/Engine/Source/Runtime/Engine/Private/Animation/AnimCurveCompressionCodec_UniformIndexable.cpp
--- a/Engine/Source/Runtime/Engine/Private/Animation/AnimCurveCompressionCodec_UniformIndexable.cpp
+++ b/Engine/Source/Runtime/Engine/Private/Animation/AnimCurveCompressionCodec_UniformIndexable.cpp
@@ -19,9 +19,10 @@ FAnimCurveBufferAccess::FAnimCurveBufferAccess(const UAnimSequenceBase* InSequen
 		{
 			if (InSequence->IsCurveCompressedDataValid())
 			{
-				if (UAnimCurveCompressionCodec_UniformIndexable* IndexableCompressedCurve = Cast<UAnimCurveCompressionCodec_UniformIndexable>(InSequence->CompressedData.CurveCompressionCodec))
+				const FCompressedAnimSequence& CompressedAnimSequence = InSequence->CompressedData;
+				if (UAnimCurveCompressionCodec_UniformIndexable* IndexableCompressedCurve = Cast<UAnimCurveCompressionCodec_UniformIndexable>(CompressedAnimSequence.CurveCompressionCodec))
 				{
-					bUseCompressedData = IndexableCompressedCurve->GetCurveBufferAndSamples(InSequence->CompressedData, InUID, CompressedBuffer, NumSamples, SampleRate);
+					bUseCompressedData = IndexableCompressedCurve->GetCurveBufferAndSamples(CompressedAnimSequence, InUID, CompressedBuffer, NumSamples, SampleRate);
 				}
 				else
 				{
diff --git a/Engine/Source/Runtime/Engine/Private/Animation/AnimDataModel.cpp b/Engine/Source/Runtime/Engine/Private/Animation/AnimDataModel.cpp
--- a/Engine/Source/Runtime/Engine/Private/Animation/AnimDataModel.cpp
+++ b/Engine/Source/Runtime/Engine/Private/Animation/AnimDataModel.cpp
@@ -205,11 +205,123 @@ void UAnimDataModel::GetBoneTrackNames(TArray<FName>& OutNames) const
 	});
 }
 
+FTransform UAnimDataModel::GetBoneTrackTransform(FName TrackName, const FFrameNumber& FrameNumber) const
+{
+	const FBoneAnimationTrack* Track = GetBoneAnimationTracks().FindByPredicate([TrackName](const FBoneAnimationTrack& Track)
+	{
+		return Track.Name == TrackName;
+	});
+
+	if (Track)
+	{
+		const int32 KeyIndex = FrameNumber.Value;
+		if (Track->InternalTrackData.PosKeys.IsValidIndex(KeyIndex) &&
+			Track->InternalTrackData.RotKeys.IsValidIndex(KeyIndex) &&
+			Track->InternalTrackData.ScaleKeys.IsValidIndex(KeyIndex))
+		{
+			return FTransform(FQuat(Track->InternalTrackData.RotKeys[KeyIndex]),
+							FVector(Track->InternalTrackData.PosKeys[KeyIndex]),
+							FVector(Track->InternalTrackData.ScaleKeys[KeyIndex])
+			);
+		}
+	}
+
+	return FTransform::Identity;
+}
+
+void UAnimDataModel::GetBoneTrackTransforms(FName TrackName, const TArray<FFrameNumber>& FrameNumbers, TArray<FTransform>& OutTransforms) const
+{
+	const FBoneAnimationTrack* Track = GetBoneAnimationTracks().FindByPredicate([TrackName](const FBoneAnimationTrack& Track)
+	{
+		return Track.Name == TrackName;
+	});
+
+	OutTransforms.SetNum(FrameNumbers.Num());
+
+	if (Track)
+	{
+		for (int32 EntryIndex = 0; EntryIndex < FrameNumbers.Num(); ++EntryIndex)
+		{
+			OutTransforms[EntryIndex] = GetBoneTrackTransform(TrackName, FrameNumbers[EntryIndex]);
+		}
+	}
+}
+
+void UAnimDataModel::GetBoneTrackTransforms(FName TrackName, TArray<FTransform>& OutTransforms) const
+{
+	const FBoneAnimationTrack* Track = GetBoneAnimationTracks().FindByPredicate([TrackName](const FBoneAnimationTrack& Track)
+	{
+		return Track.Name == TrackName;
+	});
+	
+	OutTransforms.SetNum(NumberOfKeys);
+
+	if (Track)
+	{
+		for (int32 KeyIndex = 0; KeyIndex < NumberOfKeys; ++KeyIndex)
+		{
+			OutTransforms[KeyIndex] = GetBoneTrackTransform(TrackName, FFrameNumber(KeyIndex));
+		}
+	}
+}
+
+void UAnimDataModel::GetBoneTracksTransform(const TArray<FName>& TrackNames, const FFrameNumber& FrameNumber, TArray<FTransform>& OutTransforms) const
+{	
+	OutTransforms.SetNum(TrackNames.Num());
+	for (int32 EntryIndex = 0; EntryIndex < TrackNames.Num(); ++EntryIndex)
+	{
+		OutTransforms[EntryIndex] = GetBoneTrackTransform(TrackNames[EntryIndex], FrameNumber);
+	}
+}
+
+FTransform UAnimDataModel::EvaluateBoneTrackTransform(FName TrackName, const FFrameTime& FrameTime, const EAnimInterpolationType& Interpolation) const
+{
+	const float Alpha = Interpolation == EAnimInterpolationType::Step ? FMath::RoundToFloat(FrameTime.GetSubFrame()) : FrameTime.GetSubFrame();
+
+	if (FMath::IsNearlyEqual(Alpha, 1.0f))
+	{
+		return GetBoneTrackTransform(TrackName, FrameTime.CeilToFrame());
+	}
+	else if (FMath::IsNearlyZero(Alpha))
+	{
+		return GetBoneTrackTransform(TrackName, FrameTime.FloorToFrame());
+	}
+	
+	const FTransform From = GetBoneTrackTransform(TrackName, FrameTime.FloorToFrame());
+	const FTransform To = GetBoneTrackTransform(TrackName, FrameTime.CeilToFrame());
+
+	FTransform Blend;
+	Blend.Blend(From, To, Alpha);
+	return Blend;
+}
+
+bool UAnimDataModel::IsValidBoneTrackName(const FName& TrackName) const
+{
+	return BoneAnimationTracks.ContainsByPredicate([&TrackName](const FBoneAnimationTrack& SearchTrack)
+	{
+		return SearchTrack.Name == TrackName;
+	});
+}
+
 const FAnimationCurveData& UAnimDataModel::GetCurveData() const
 {
 	return CurveData;
 }
 
+void UAnimDataModel::IterateBoneKeys(const FName& BoneName, TFunction<bool(const FVector3f& Pos, const FQuat4f&, const FVector3f, const FFrameNumber&)> IterationFunction) const
+{
+	if (const FBoneAnimationTrack* Track = FindBoneTrackByName(BoneName))
+	{
+		for (int32 KeyIndex = 0; KeyIndex < GetNumberOfKeys(); ++KeyIndex)
+		{
+			if(!IterationFunction(Track->InternalTrackData.PosKeys[KeyIndex], Track->InternalTrackData.RotKeys[KeyIndex], Track->InternalTrackData.ScaleKeys[KeyIndex], KeyIndex))
+			{
+				return;
+			}
+		}
+	}
+}
+
 IAnimationDataModel::FModelNotifier& UAnimDataModel::GetNotifier()
 {
 	if (!Notifier)
@@ -686,11 +798,11 @@ void UAnimDataModel::Evaluate(FAnimationPoseData& InOutPoseData, const UE::Anim:
 		UE::Anim::Retargeting::FRetargetingScope RetargetingScope(OutPose, EvaluationContext);
 		if (bShouldInterpolate)
 		{		
-			ExtractPose<true>(GetBoneAnimationTracks(), ActiveCurves, OutPose, KeyIndex1, KeyIndex2, Alpha, TimePerFrame, RetargetingScope);
+			ExtractPose<true>(BoneAnimationTracks, ActiveCurves, OutPose, KeyIndex1, KeyIndex2, Alpha, TimePerFrame, RetargetingScope);
 		}
 		else
 		{
-			ExtractPose<false>(GetBoneAnimationTracks(), ActiveCurves, OutPose, KeyIndex1, KeyIndex2, Alpha, TimePerFrame, RetargetingScope);
+			ExtractPose<false>(BoneAnimationTracks, ActiveCurves, OutPose, KeyIndex1, KeyIndex2, Alpha, TimePerFrame, RetargetingScope);
 		}
 	}
 	
