Animation Compression

diff --git a/Engine/Source/Runtime/Engine/Private/Animation/AnimSequence.cpp b/Engine/Source/Runtime/Engine/Private/Animation/AnimSequence.cpp
--- a/Engine/Source/Runtime/Engine/Private/Animation/AnimSequence.cpp
+++ b/Engine/Source/Runtime/Engine/Private/Animation/AnimSequence.cpp
@@ -7,6 +7,7 @@
 #include "Animation/AnimSequence.h"
 
 #include "Algo/Find.h"
+#include "AnimationCompressionDerivedData.h"
 #include "Misc/MessageDialog.h"
 #include "Logging/LogScopedVerbosityOverride.h"
 #include "UObject/FrameworkObjectVersion.h"
@@ -24,6 +25,7 @@
 #include "AnimationUtils.h"
 #include "BonePose.h"
 #include "AnimationRuntime.h"
+#include "AnimationSequenceCompiler.h"
 #include "Animation/AnimCompress.h"
 #include "Components/SkeletalMeshComponent.h"
 #include "Animation/AnimRootMotionProvider.h"
@@ -41,8 +43,6 @@
 #include "DerivedDataCacheInterface.h"
 #include "Interfaces/ITargetPlatform.h"
 #include "Interfaces/ITargetPlatformManagerModule.h"
-#include "Animation/AnimCompressionDerivedData.h"
-#include "Animation/AnimCompressionDerivedDataPublic.h"
 #include "UObject/UObjectIterator.h"
 #include "UObject/UObjectThreadContext.h"
 #include "Animation/AnimNotifies/AnimNotifyState.h"
@@ -61,6 +61,8 @@
 #include "ITimeManagementModule.h"
 #include "Algo/Accumulate.h"
 #include "HAL/LowLevelMemTracker.h"
+#include "Interfaces/ITargetPlatformManagerModule.h"
+#include "Serialization/MemoryHasher.h"
 
 LLM_DEFINE_TAG(SequenceData);
 
@@ -86,17 +88,6 @@ LLM_DEFINE_TAG(SequenceData);
 DECLARE_CYCLE_STAT(TEXT("AnimSeq GetBonePose"), STAT_AnimSeq_GetBonePose, STATGROUP_Anim);
 DECLARE_CYCLE_STAT(TEXT("AnimSeq EvalCurveData"), STAT_AnimSeq_EvalCurveData, STATGROUP_Anim);
 
-#if ENABLE_COOK_STATS
-namespace AnimSequenceCookStats
-{
-	FDerivedDataUsageStats UsageStats;
-	static FCookStatsManager::FAutoRegisterCallback RegisterCookStats([](FCookStatsManager::AddStatFuncRef AddStat)
-	{
-		UsageStats.LogStats(AddStat, TEXT("AnimSequence.Usage"), TEXT(""));
-	});
-}
-#endif
-
 CSV_DECLARE_CATEGORY_MODULE_EXTERN(ENGINE_API, Animation);
 
 int32 GPerformFrameStripping = 0;
@@ -237,8 +228,7 @@ void OnCVarsChanged()
 		// Rebake/compress the animations
 		for (UAnimSequence* Seq : SequenceCache)
 		{
-			Seq->RequestSyncAnimRecompression();
-
+			Seq->CacheDerivedDataForCurrentPlatform();
 			Sizes.Emplace(Seq->GetApproxCompressedSize(), Seq);
 		}
 
@@ -319,55 +309,6 @@ FString GetAnimSequenceSpecificCacheKeySuffix(const UAnimSequence& Seq, bool bPe
 }
 #endif
 
-/////////////////////////////////////////////////////
-// FRequestAnimCompressionParams
-FRequestAnimCompressionParams::FRequestAnimCompressionParams(bool bInAsyncCompression, bool bInAllowAlternateCompressor, bool bInOutput, const ITargetPlatform* InTargetPlatform)
-	: bAsyncCompression(bInAsyncCompression)
-	, CompressContext(MakeShared<FAnimCompressContext>(bInAllowAlternateCompressor, bInOutput))
-	, TargetPlatform(InTargetPlatform)
-{
-	InitFrameStrippingFromCVar();
-}
-
-FRequestAnimCompressionParams::FRequestAnimCompressionParams(bool bInAsyncCompression, TSharedPtr<FAnimCompressContext> InCompressContext)
-	: bAsyncCompression(bInAsyncCompression)
-	, CompressContext(InCompressContext)
-	, TargetPlatform(nullptr)
-{
-	InitFrameStrippingFromCVar();
-}
-
-void FRequestAnimCompressionParams::InitFrameStrippingFromCVar()
-{
-	bPerformFrameStripping = (GPerformFrameStripping == 1);
-	bPerformFrameStrippingOnOddNumberedFrames = (GPerformFrameStrippingOddFramedAnimations == 1);
-}
-
-void FRequestAnimCompressionParams::InitFrameStrippingFromPlatform()
-{
-#if WITH_EDITOR
-	bPerformFrameStripping = false;
-
-	if (TargetPlatform)
-	{
-		if (UDeviceProfile* DeviceProfile = UDeviceProfileManager::Get().FindProfile(TargetPlatform->IniPlatformName()))
-		{
-			int32 CVarPlatformFrameStrippingValue = 0;
-			if (DeviceProfile->GetConsolidatedCVarValue(StripFrameCVarName, CVarPlatformFrameStrippingValue))
-			{
-				bPerformFrameStripping = CVarPlatformFrameStrippingValue == 1;
-			}
-
-			int32 CVarPlatformOddAnimFrameStrippingValue = 0;
-			if (DeviceProfile->GetConsolidatedCVarValue(OddFrameStripStrippingCVarName, CVarPlatformOddAnimFrameStrippingValue))
-			{
-				bPerformFrameStrippingOnOddNumberedFrames = CVarPlatformOddAnimFrameStrippingValue == 1;
-			}
-		}
-	}
-#endif
-}
-
 /////////////////////////////////////////////////////
 // FRawAnimSequenceTrackNativeDeprecated
 
@@ -456,7 +397,6 @@ UAnimSequence::UAnimSequence(const FObjectInitializer& ObjectInitializer)
 	, bRootMotionSettingsCopiedFromMontage(false)
 	, bUseRawDataOnly(!FPlatformProperties::RequiresCookedData())
 #if WITH_EDITOR
-	, bCompressionInProgress(false)
 	, bBlockCompressionRequests(false)
 #endif
 {
@@ -511,7 +451,7 @@ void UAnimSequence::GetAssetRegistryTags(TArray<FAssetRegistryTag>& OutTags) con
 #endif
 
 	OutTags.Add(FAssetRegistryTag(TEXT("Compressed Size (KB)"), FString::Printf(TEXT("%.02f"), (float)GetApproxCompressedSize() / 1024.0f), FAssetRegistryTag::TT_Numerical));
-	OutTags.Add(FAssetRegistryTag(TEXT("Target Frame Rate"), FString::Printf(TEXT("%.2f"), TargetFrameRate.AsDecimal()), FAssetRegistryTag::TT_Numerical));
+	OutTags.Add(FAssetRegistryTag(TEXT("Target Frame Rate"), FString::Printf(TEXT("%.2f"), PlatformTargetFrameRate.Default.AsDecimal()), FAssetRegistryTag::TT_Numerical));
 	Super::GetAssetRegistryTags(OutTags);
 }
 
@@ -528,9 +468,14 @@ void UAnimSequence::AddReferencedObjects(UObject* This, FReferenceCollector& Col
 void UAnimSequence::WillNeverCacheCookedPlatformDataAgain()
 {
 	Super::WillNeverCacheCookedPlatformDataAgain();
-	ClearResampledAnimationTrackData();
 	ClearCompressedCurveData();
 	ClearCompressedBoneData();
+
+	// Clear out current platform, and any target platform data
+	CompressedData.Reset();
+		
+	CacheTasksByKeyHash.Empty();
+	DataByPlatformKeyHash.Empty();
 	bUseRawDataOnly = true;
 }
 
@@ -538,16 +483,32 @@ void UAnimSequence::ClearCachedCookedPlatformData(const ITargetPlatform* TargetP
 {
 	Super::ClearCachedCookedPlatformData(TargetPlatform);
 
-	// For now just clear resampled data (as this will always be resampled when not found for compression)
-	ClearResampledAnimationTrackData();
+	// Remove platform data if any
+	const FIoHash PlatformHash = CreateDerivedDataKeyHash(TargetPlatform);
+
+	CacheTasksByKeyHash.Remove(PlatformHash);
+	DataByPlatformKeyHash.Remove(PlatformHash);
+
+	if (TargetPlatform->IsRunningPlatform() && PlatformHash == DataKeyHash)
+	{
+		CompressedData.Reset();
+		DataKeyHash = FIoHash::Zero;
+		bUseRawDataOnly = true;
+	}
 }
 
 void UAnimSequence::ClearAllCachedCookedPlatformData()
 {
 	Super::ClearAllCachedCookedPlatformData();
 	
-	// For now just clear resampled data (as this will always be resampled when not found for compression)
-	ClearResampledAnimationTrackData();
+	// Delete any cache tasks first because the destructor will cancel the cache and build tasks,
+	// and drop their pointers to the data.
+	CacheTasksByKeyHash.Empty();
+	DataByPlatformKeyHash.Empty();
+	CompressedData.Reset();
+	DataKeyHash = FIoHash::Zero;
+	
+	bUseRawDataOnly = true;
 }
 
 int32 UAnimSequence::GetUncompressedRawSize() const
@@ -567,18 +528,18 @@ int32 UAnimSequence::GetUncompressedRawSize() const
 	return 0;
 }
 
-int32 UAnimSequence::GetApproxRawSize() const
+int64 UAnimSequence::GetApproxRawSize() const
 {
 	if (ShouldDataModelBeValid())
 	{
 		ValidateModel();
 		int32 Total = sizeof(FRawAnimSequenceTrack) * DataModelInterface->GetNumBoneTracks();
-		for (const FBoneAnimationTrack& AnimationTrack : DataModelInterface->GetBoneAnimationTracks())
-		{
-			Total += AnimationTrack.InternalTrackData.PosKeys.Num() * AnimationTrack.InternalTrackData.PosKeys.GetTypeSize();
-			Total += AnimationTrack.InternalTrackData.RotKeys.Num() * AnimationTrack.InternalTrackData.RotKeys.GetTypeSize();
-			Total += AnimationTrack.InternalTrackData.ScaleKeys.Num() * AnimationTrack.InternalTrackData.ScaleKeys.GetTypeSize();
-		}
+
+		const int32 NumberOfTracks = DataModelInterface->GetNumBoneTracks();
+
+		Total += DataModelInterface->GetNumberOfKeys() * sizeof(FVector3f) * NumberOfTracks;
+		Total += DataModelInterface->GetNumberOfKeys() * sizeof(FQuat4f) * NumberOfTracks;
+		Total += DataModelInterface->GetNumberOfKeys() * sizeof(FVector3f) * NumberOfTracks;
 
 		for (const FFloatCurve& Curve : DataModelInterface->GetCurveData().FloatCurves)
 		{
@@ -706,6 +667,20 @@ void UAnimSequence::Serialize(FArchive& Ar)
 #endif // WITH_EDITORONLY_DATA
 	}
 
+#if WITH_EDITOR
+	if (Ar.IsCooking())
+	{
+		if (IsCompiling())
+		{
+			UE::Anim::FAnimSequenceCompilingManager::Get().FinishCompilation({this});
+		}
+	}
+#endif // WITH_EDITOR	
+	if ( Ar.IsLoading() && Ar.CustomVer(FUE5MainStreamObjectVersion::GUID) < FUE5MainStreamObjectVersion::PerPlatformAnimSequenceTargetFrameRate)
+	{
+		PlatformTargetFrameRate = TargetFrameRate;
+	}
+
 	if (Ar.CustomVer(FFrameworkObjectVersion::GUID) < FFrameworkObjectVersion::MoveCompressedAnimDataToTheDDC)
 	{
 		// Serialize the compressed byte stream from the archive to the buffer.
@@ -885,31 +860,11 @@ void UAnimSequence::PreSave(const class ITargetPlatform* TargetPlatform)
 void UAnimSequence::PreSave(FObjectPreSaveContext ObjectSaveContext)
 {
 #if WITH_EDITOR
-	// Could already be compressing
-	WaitOnExistingCompression();
-
-	// make sure if it does contain transform curvesm it contains source data
-	// empty track animation still can be made by retargeting to invalid skeleton
-	// make sure to not trigger ensure if RawAnimationData is also null
-	
-	if (DoesNeedRecompress())
-	{
-		RequestSyncAnimRecompression(); // Update Normal data
-
-		ensureAlwaysMsgf(!bUseRawDataOnly,  TEXT("Animation : %s failed to compress"), *GetName());
-	}
-
-	const ITargetPlatform* TargetPlatform = ObjectSaveContext.GetTargetPlatform();
-	if (TargetPlatform)
+	if (ObjectSaveContext.IsCooking())
 	{
-		// Update compressed data for platform
-		FRequestAnimCompressionParams Params(false, false, false, TargetPlatform);
-		Params.InitFrameStrippingFromPlatform();
-		RequestAnimCompression(Params);
+		check(IsCachedCookedPlatformDataLoaded(ObjectSaveContext.GetTargetPlatform()));
 	}
 
-	WaitOnExistingCompression(); // Wait on updated data
-
 	if (!ObjectSaveContext.IsProceduralSave())
 	{
 		UpdateRetargetSourceAsset();
@@ -987,16 +942,19 @@ void UAnimSequence::PostLoad()
 
 	if (USkeleton* CurrentSkeleton = GetSkeleton())
 	{
+	#if WITH_EDITOR
+    		for (const FAnimSyncMarker& SyncMarker : AuthoredSyncMarkers)
+    		{
+    			CurrentSkeleton->RegisterMarkerName(SyncMarker.MarkerName);
+    		}
+    #endif
+    
+    #if !WITH_EDITOR
 		for (FSmartName& CurveName : CompressedData.CompressedCurveNames)
 		{
 			CurrentSkeleton->VerifySmartName(USkeleton::AnimCurveMappingName, CurveName);
 		}
-#if WITH_EDITOR
-		for (const FAnimSyncMarker& SyncMarker : AuthoredSyncMarkers)
-		{
-			CurrentSkeleton->RegisterMarkerName(SyncMarker.MarkerName);
-		}
-#endif
+    #endif
 	}
 }
 
@@ -1023,9 +981,7 @@ void UAnimSequence::BeginDestroy()
 {
 #if WITH_EDITOR
 	// Could already be compressing
-	PRAGMA_DISABLE_DEPRECATION_WARNINGS
-	WaitOnExistingCompression(false);
-	PRAGMA_ENABLE_DEPRECATION_WARNINGS
+	UE::Anim::FAnimSequenceCompilingManager::Get().FinishCompilation({this});
 #endif // WITH_EDITOR
 
 	Super::BeginDestroy();
@@ -1036,6 +992,23 @@ void UAnimSequence::BeginDestroy()
 	PRAGMA_ENABLE_DEPRECATION_WARNINGS
 }
 
+bool UAnimSequence::IsReadyForFinishDestroy()
+{
+	if (!Super::IsReadyForFinishDestroy())
+	{
+		return false;
+	}
+
+#if WITH_EDITOR
+	if (!TryCancelAsyncTasks())
+	{
+		return false;
+	}
+#endif
+
+	return true;
+}
+
 #if WITH_EDITOR
 void UAnimSequence::PostEditChangeProperty(FPropertyChangedEvent& PropertyChangedEvent)
 {
@@ -1077,11 +1050,6 @@ void UAnimSequence::PostEditChangeProperty(FPropertyChangedEvent& PropertyChange
 
 		const bool bChangedRefFrameIndex = PropertyChangedEvent.Property->GetFName() == GET_MEMBER_NAME_CHECKED(UAnimSequence, RefFrameIndex);
 
-		if (bChangedRefFrameIndex)
-		{
-			bUseRawDataOnly = true;
-		}
-
 		if ((bChangedRefFrameIndex && PropertyChangedEvent.ChangeType != EPropertyChangeType::Interactive) ||
 			PropertyChangedEvent.Property->GetFName() == GET_MEMBER_NAME_CHECKED(UAnimSequence, AdditiveAnimType) ||
 			PropertyChangedEvent.Property->GetFName() == GET_MEMBER_NAME_CHECKED(UAnimSequence, RefPoseSeq) ||
@@ -1101,12 +1069,11 @@ void UAnimSequence::PostEditChangeProperty(FPropertyChangedEvent& PropertyChange
 											  || PropertyChangedEvent.Property->GetFName() == GET_MEMBER_NAME_CHECKED(UAnimSequence, BoneCompressionSettings)
 											  || PropertyChangedEvent.Property->GetFName() == GET_MEMBER_NAME_CHECKED(UAnimSequence, BoneCompressionSettings);
 
-		bShouldResample = PropertyChangedEvent.Property->GetFName() == GET_MEMBER_NAME_CHECKED(UAnimSequence, TargetFrameRate);
+		bShouldResample = PropertyChangedEvent.GetMemberPropertyName() == GET_MEMBER_NAME_CHECKED(UAnimSequence, PlatformTargetFrameRate) || bChangedRefFrameIndex;
 	}
 
 	if (bShouldResample)
 	{
-		ClearResampledAnimationTrackData();
 		CalculateNumberOfSampledKeys();
 	}
 
@@ -1115,16 +1082,46 @@ void UAnimSequence::PostEditChangeProperty(FPropertyChangedEvent& PropertyChange
 	// @Todo fix me: This is temporary fix to make sure they always have compressed data
 	ValidateModel();
 	if (IsDataModelValid() && DataModelInterface->GetNumBoneTracks() && bNeedPostProcess)
-	{
-		RecompressAnimationData();
+	{		
+		ClearAllCachedCookedPlatformData();
+		BeginCacheDerivedDataForCurrentPlatform();
 	}
 }
+
+void UAnimSequence::BeginCacheForCookedPlatformData(const ITargetPlatform* TargetPlatform)
+{
+	BeginCacheDerivedData(TargetPlatform);
+}
 #endif // WITH_EDITOR
 
 #if WITH_EDITOR
 bool UAnimSequence::IsCachedCookedPlatformDataLoaded(const ITargetPlatform* TargetPlatform)
 {
-	return !bCompressionInProgress;
+	const FIoHash KeyHash = CreateDerivedDataKeyHash(TargetPlatform);
+	if (KeyHash.IsZero())
+	{
+		return true;
+	}
+
+	if (PollCacheDerivedData(KeyHash))
+	{
+		EndCacheDerivedData(KeyHash);
+		return true;
+	}
+
+	if (KeyHash == DataKeyHash)
+	{
+		return CompressedData.IsValid(this);
+	}
+	else
+	{
+		if(const TUniquePtr<FCompressedAnimSequence>* CompressedDataPtr = DataByPlatformKeyHash.Find(KeyHash))
+		{
+			return (*CompressedDataPtr)->IsValid(this);
+		}
+	}
+
+	return false;
 }
 #endif // WITH_EDITOR
 
@@ -1158,12 +1155,7 @@ static FAutoConsoleVariableRef CVarForceRawData(
 	ECVF_Cheat);
 #endif // WITH_EDITOR
 
-void UAnimSequence::GetBoneTransform(FTransform& OutAtom, int32 TrackIndex, float Time, bool bUseRawData) const
-{
-	GetBoneTransform(OutAtom, TrackIndex, (double)Time, bUseRawData);
-}
-
-void UAnimSequence::GetBoneTransform(FTransform& OutAtom, int32 TrackIndex, double Time, bool bUseRawData) const
+void UAnimSequence::GetBoneTransform(FTransform& OutAtom, FSkeletonPoseBoneIndex BoneIndex, double Time, bool bUseRawData) const
 {
 	// If the caller didn't request that raw animation data be used . . .
 	const bool bEvaluateCompressedData =
@@ -1174,24 +1166,42 @@ void UAnimSequence::GetBoneTransform(FTransform& OutAtom, int32 TrackIndex, doub
 
 	if ( bEvaluateCompressedData && IsCompressedDataValid())
 	{
-		FAnimSequenceDecompressionContext DecompContext(TargetFrameRate, TargetFrameRate.AsFrameNumber(GetPlayLength()).Value, Interpolation, GetRetargetTransformsSourceName(), *CompressedData.CompressedDataStructure, GetSkeleton()->GetRefLocalPoses(), CompressedData.CompressedTrackToSkeletonMapTable, GetSkeleton(), IsValidAdditive());
+		FAnimSequenceDecompressionContext DecompContext(PlatformTargetFrameRate.Default, PlatformTargetFrameRate.Default.AsFrameNumber(GetPlayLength()).Value, Interpolation, GetRetargetTransformsSourceName(), *CompressedData.CompressedDataStructure, GetSkeleton()->GetRefLocalPoses(), CompressedData.CompressedTrackToSkeletonMapTable, GetSkeleton(), IsValidAdditive());
 		DecompContext.Seek(Time);
 		if (CompressedData.BoneCompressionCodec)
 		{
-			CompressedData.BoneCompressionCodec->DecompressBone(DecompContext, TrackIndex, OutAtom);
-			return;
+			const int32 TrackIndex = CompressedData.GetTrackIndexFromSkeletonIndex(BoneIndex.GetInt());
+			if (TrackIndex != INDEX_NONE)
+			{				
+				CompressedData.BoneCompressionCodec->DecompressBone(DecompContext, TrackIndex, OutAtom);
+			}
 		}
 	}
 	else if (CanEvaluateRawAnimationData())
 	{
 #if WITH_EDITOR
 		ValidateModel();
-		UE::Anim::GetBoneTransformFromModel(DataModelInterface.GetInterface(), OutAtom, TrackIndex, Time, Interpolation);
+		const FName BoneName = GetSkeleton()->GetReferenceSkeleton().GetBoneName(BoneIndex.GetInt());
+		OutAtom = DataModelInterface->EvaluateBoneTrackTransform(BoneName, DataModelInterface->GetFrameRate().AsFrameTime(Time), Interpolation);
+		if (const USkeleton* CurrentSkeleton = GetSkeleton())
+		{
+			FSmartName CurveSmartName;
+			CurrentSkeleton->GetSmartNameByName(USkeleton::AnimTrackCurveMappingName, BoneName, CurveSmartName);
+			const FAnimationCurveIdentifier TransformCurveId(CurveSmartName, ERawCurveTrackTypes::RCT_Transform);
+			if (const FTransformCurve* TransformCurvePtr = DataModelInterface->FindTransformCurve(TransformCurveId))
+			{
+				const FTransform AdditiveTransform = TransformCurvePtr->Evaluate(Time, 1.f);
+				const FTransform LocalTransform = OutAtom;
+				OutAtom.SetRotation(LocalTransform.GetRotation() * AdditiveTransform.GetRotation());
+				OutAtom.SetTranslation(LocalTransform.TransformPosition(AdditiveTransform.GetTranslation()));
+				OutAtom.SetScale3D(LocalTransform.GetScale3D() * AdditiveTransform.GetScale3D());	
+			}
+		}	
 #endif
 	}
 }
 
-void UAnimSequence::GetBoneTransform(FTransform& OutAtom, int32 TrackIndex, FAnimSequenceDecompressionContext& DecompContext, bool bUseRawData) const
+void UAnimSequence::GetBoneTransform(FTransform& OutAtom, FSkeletonPoseBoneIndex BoneIndex, FAnimSequenceDecompressionContext& DecompContext, bool bUseRawData) const
 {
 	// If the caller didn't request that raw animation data be used . . .
 	const bool bEvaluateCompressedData = (!bUseRawData || !CanEvaluateRawAnimationData());
@@ -1199,7 +1209,7 @@ void UAnimSequence::GetBoneTransform(FTransform& OutAtom, int32 TrackIndex, FAni
 	{
 		if (CompressedData.BoneCompressionCodec)
 		{
-			CompressedData.BoneCompressionCodec->DecompressBone(DecompContext, TrackIndex, OutAtom);
+			CompressedData.BoneCompressionCodec->DecompressBone(DecompContext, CompressedData.GetTrackIndexFromSkeletonIndex(BoneIndex.GetInt()), OutAtom);
 			return;
 		}
 	}
@@ -1207,7 +1217,23 @@ void UAnimSequence::GetBoneTransform(FTransform& OutAtom, int32 TrackIndex, FAni
 	{
 #if WITH_EDITOR
 		ValidateModel();
-		UE::Anim::GetBoneTransformFromModel(DataModelInterface.GetInterface(), OutAtom, TrackIndex, DecompContext.GetEvaluationTime(), Interpolation);
+		const FName BoneName = GetSkeleton()->GetReferenceSkeleton().GetBoneName(BoneIndex.GetInt());
+		OutAtom = DataModelInterface->EvaluateBoneTrackTransform(BoneName, DataModelInterface->GetFrameRate().AsFrameTime(DecompContext.GetEvaluationTime()), Interpolation);
+
+		if (const USkeleton* CurrentSkeleton = GetSkeleton())
+		{
+			FSmartName CurveSmartName;
+			CurrentSkeleton->GetSmartNameByName(USkeleton::AnimTrackCurveMappingName, BoneName, CurveSmartName);		
+			const FAnimationCurveIdentifier TransformCurveId(CurveSmartName, ERawCurveTrackTypes::RCT_Transform);
+			if (const FTransformCurve* TransformCurvePtr = DataModelInterface->FindTransformCurve(TransformCurveId))
+			{
+				const FTransform AdditiveTransform = TransformCurvePtr->Evaluate(DecompContext.GetEvaluationTime(), 1.f);
+				const FTransform LocalTransform = OutAtom;
+				OutAtom.SetRotation(LocalTransform.GetRotation() * AdditiveTransform.GetRotation());
+				OutAtom.SetTranslation(LocalTransform.TransformPosition(AdditiveTransform.GetTranslation()));
+				OutAtom.SetScale3D(LocalTransform.GetScale3D() * AdditiveTransform.GetScale3D());	
+			}
+		}	
 #endif
 	}
 }
@@ -1225,15 +1251,17 @@ void UAnimSequence::HandleAssetPlayerTickedInternal(FAnimAssetTickContext &Conte
 FTransform UAnimSequence::ExtractRootTrackTransform(float Pos, const FBoneContainer * RequiredBones) const
 {
 	const int32 RootBoneIndex = 0;
-	const int32 TrackIndex = [this, RootBoneIndex]() -> int32
+	const bool bContainsRootBoneTrack = [this, RootBoneIndex]() 
 	{
 #if WITH_EDITOR
 		if (CanEvaluateRawAnimationData() && bUseRawDataOnly)
 		{
 			ValidateModel();
-			if (const FBoneAnimationTrack* RootTrack = DataModelInterface->FindBoneTrackByIndex(RootBoneIndex))
+
+			const FName RootBoneName = GetSkeleton()->GetReferenceSkeleton().GetBoneName(RootBoneIndex);
+			if (DataModelInterface->IsValidBoneTrackName(RootBoneName))
 			{
-				return DataModelInterface->GetBoneTrackIndex(*RootTrack);
+				return true;
 			}
 		}
 		else
@@ -1243,18 +1271,18 @@ FTransform UAnimSequence::ExtractRootTrackTransform(float Pos, const FBoneContai
 			// we assume root is in first data if available = SkeletonIndex == 0 && BoneTreeIndex == 0)
 			if ((TrackToSkeletonMap.Num() > 0) && (TrackToSkeletonMap[0].BoneTreeIndex == RootBoneIndex))
 			{
-				return 0;
+				return true;
 			}
 		}
 
-		return INDEX_NONE;
+		return false;
 	}();
 
-	if(TrackIndex != INDEX_NONE)
+	if(bContainsRootBoneTrack)
 	{
 		// if we do have root data, then return root data
 		FTransform RootTransform;
-		GetBoneTransform(RootTransform, TrackIndex, static_cast<double>(Pos), bUseRawDataOnly);
+		GetBoneTransform(RootTransform, FSkeletonPoseBoneIndex(RootBoneIndex), static_cast<double>(Pos), bUseRawDataOnly);
 		return RootTransform;
 	}
 
@@ -1542,20 +1570,6 @@ void UAnimSequence::GetBonePose(FAnimationPoseData& OutAnimationPoseData, const
 	// Evaluate raw (source) curve and bone data
 	if (bUseRawDataForPoseExtraction)
 	{
-		// Warning if we have invalid data
-		for (int32 TrackIndex = 0; TrackIndex < NumTracks; TrackIndex++)
-		{
-			const FBoneAnimationTrack& BoneAnimationTrack = DataModelInterface->GetBoneTrackByIndex(TrackIndex);
-			const FRawAnimSequenceTrack& TrackToExtract = BoneAnimationTrack.InternalTrackData;
-
-			// Bail out (with rather wacky data) if data is empty for some reason.
-			if (TrackToExtract.PosKeys.Num() == 0 || TrackToExtract.RotKeys.Num() == 0)
-			{
-				UE_LOG(LogAnimation, Warning, TEXT("UAnimSequence::GetBoneTransform : No anim data in AnimSequence '%s' Track '%s'"), *GetPathName(), 
-					*BoneAnimationTrack.Name.ToString());
-			}
-		}
-
 		{
 			FSkeletonRemappingCurve RemappedCurve(OutAnimationPoseData.GetCurve(), RequiredBones, GetSkeleton());
 			FAnimationPoseData RemappedPoseData(OutAnimationPoseData.GetPose(), RemappedCurve.GetCurve(), OutAnimationPoseData.GetAttributes());
@@ -1577,7 +1591,7 @@ void UAnimSequence::GetBonePose(FAnimationPoseData& OutAnimationPoseData, const
 	if (NumTracks != 0)
 	{
 		// Evaluate compressed bone data
-		FAnimSequenceDecompressionContext DecompContext(TargetFrameRate, TargetFrameRate.AsFrameTime(GetPlayLength()).RoundToFrame().Value, Interpolation, GetRetargetTransformsSourceName(), *CompressedData.CompressedDataStructure, GetSkeleton()->GetRefLocalPoses(), CompressedData.CompressedTrackToSkeletonMapTable, GetSkeleton(), IsValidAdditive());
+		FAnimSequenceDecompressionContext DecompContext(PlatformTargetFrameRate.Default, PlatformTargetFrameRate.Default.AsFrameTime(GetPlayLength()).RoundToFrame().Value, Interpolation, GetRetargetTransformsSourceName(), *CompressedData.CompressedDataStructure, GetSkeleton()->GetRefLocalPoses(), CompressedData.CompressedTrackToSkeletonMapTable, GetSkeleton(), IsValidAdditive());
 		UE::Anim::Decompression::DecompressPose(OutPose, CompressedData, ExtractionContext, DecompContext, GetRetargetTransforms(), RootMotionReset);
 	}
 
@@ -1844,192 +1858,46 @@ FString UAnimSequence::GetDDCCacheKeySuffix(const bool bPerformStripping, const
 void UAnimSequence::WaitOnExistingCompression(const bool bWantResults)
 {
 	check(IsInGameThread());
-	if (bCompressionInProgress)
-	{
-		COOK_STAT(auto Timer = AnimSequenceCookStats::UsageStats.TimeAsyncWait());
-		FAsyncCompressedAnimationsManagement::Get().WaitOnExistingCompression(this, bWantResults);
-		SetCompressionComplete();
-		COOK_STAT(Timer.TrackCyclesOnly()); // Need to get hit/miss and size from WaitOnExistingCompression!
-	}
-}
 
-void UAnimSequence::SetCompressionComplete()
-{
-	if (bCompressionInProgress)
+	if(!bWantResults)
 	{
-		bCompressionInProgress = false;
-		FAssetRegistryModule& AssetRegistryModule = FModuleManager::LoadModuleChecked<FAssetRegistryModule>(TEXT("AssetRegistry"));
-		AssetRegistryModule.Get().AssetTagsFinalized(*this);
+		TryCancelAsyncTasks();
 	}
+	UE::Anim::FAnimSequenceCompilingManager::Get().FinishCompilation({this});
 }
 #endif
 
-#if WITH_EDITOR
-int32 GSkipDDC = 0;
-static FAutoConsoleVariableRef CVarSkipDDC(
-	TEXT("a.SkipDDC"),
-	GSkipDDC,
-	TEXT("1 = Skip DDC during compression. 0 = Include DDC results during compresion "));
-
-void UAnimSequence::RequestAnimCompression(FRequestAnimCompressionParams Params)
+void UAnimSequence::SerializeCompressedData(FArchive& Ar, bool bDDCData)
 {
-	check(IsInGameThread());
-	USkeleton* CurrentSkeleton = GetSkeleton();
-	if (CurrentSkeleton == nullptr)
-	{
-		bUseRawDataOnly = true;
-		return;
-	}
-
-	if (FPlatformProperties::RequiresCookedData() || bBlockCompressionRequests)
-	{
-		return;
-	}
-
-	WaitOnExistingCompression(false);
-
-	if (BoneCompressionSettings == nullptr || !BoneCompressionSettings->AreSettingsValid())
-	{
-		BoneCompressionSettings = FAnimationUtils::GetDefaultAnimationBoneCompressionSettings();
-	}
-
-	if (CurveCompressionSettings == nullptr || !CurveCompressionSettings->AreSettingsValid())
-	{
-		CurveCompressionSettings = FAnimationUtils::GetDefaultAnimationCurveCompressionSettings();
-	}
-
-	// Make sure all our required dependencies are loaded
-	FAnimationUtils::EnsureAnimSequenceLoaded(*this);
-
-	// EnsureAnimSequenceLoaded can also start some compression, we'll need to wait here since the code
-	// below assumes everything is finished.
-	WaitOnExistingCompression(false);
-
-	EnsureValidRawDataGuid();
-
-	if(ResampledAnimationTrackData.Num() == 0)
-	{		
-		ResampleAnimationTrackData();
-	}
-
-	bUseRawDataOnly = true;
-
-	check(!bCompressionInProgress);
-	bCompressionInProgress = true;
-
-	Params.CompressContext->GatherPreCompressionStats(GetName(), GetApproxRawSize(), GetApproxCompressedSize());
-
-	const double CompressionStartTime = FPlatformTime::Seconds();
-
-	const bool bPerformStripping = ShouldPerformStripping(Params.bPerformFrameStripping, Params.bPerformFrameStrippingOnOddNumberedFrames);
-	const FString AssetDDCKey = GetDDCCacheKeySuffix(bPerformStripping, Params.TargetPlatform);
-	RequestedCompression_DDCKey = FString(AssetDDCKey);
-
-	bool bCompressedDataFromDDC = false;
-
-	TArray<uint8> OutData;
-	{
-		COOK_STAT(auto Timer = AnimSequenceCookStats::UsageStats.TimeSyncWork());
-
-		FDerivedDataAnimationCompression* AnimCompressor = new FDerivedDataAnimationCompression(TEXT("AnimSeq"), AssetDDCKey, Params.CompressContext);
-
-		// For debugging DDC/Compression issues
-		const bool bSkipDDC = GSkipDDC == 1;
-
-		// Data does not exist, need to build it.
-		FCompressibleAnimRef CompressibleData = MakeShared<FCompressibleAnimData, ESPMode::ThreadSafe>(this, bPerformStripping, Params.TargetPlatform);
-		AnimCompressor->SetCompressibleData(CompressibleData);
-
-		if (bSkipDDC || (CompressCommandletVersion == INDEX_NONE))
-		{
-			AnimCompressor->Build(OutData);
-			COOK_STAT(Timer.AddMiss(OutData.Num()));
-		}
-		else if (AnimCompressor->CanBuild())
-		{
-			if (Params.bAsyncCompression)
-			{
-				FAsyncCompressedAnimationsManagement::Get().RequestAsyncCompression(*AnimCompressor, this, bPerformStripping, OutData);
-				COOK_STAT(Timer.TrackCyclesOnly());
-			}
-			else
-			{
-				bool bBuilt = false;
-				const bool bSuccess = GetDerivedDataCacheRef().GetSynchronous(AnimCompressor, OutData, &bBuilt);
-				COOK_STAT(Timer.AddHitOrMiss(!bSuccess || bBuilt ? FCookStats::CallStats::EHitOrMiss::Miss : FCookStats::CallStats::EHitOrMiss::Hit, OutData.Num()));
-			}
-			AnimCompressor = nullptr;
-		}
-		else
-		{
-			COOK_STAT(Timer.TrackCyclesOnly());
-		}
-
-		if (AnimCompressor)
-		{
-			// Would really like to do auto mem management but GetDerivedDataCacheRef().GetSynchronous expects a pointer it can delete
-			delete AnimCompressor;
-			AnimCompressor = nullptr;
-		}
-	}
-
-	if (OutData.Num() > 0) // Haven't async compressed
+	LLM_SCOPE_BYNAME(TEXT("Animation/SequenceData")); 
+	
+	if (!HasAnyFlags(RF_ClassDefaultObject))
 	{
-		ApplyCompressedData(OutData);
-
-		if(bCompressedDataFromDDC)
+		FCompressedAnimSequence& CompressedDataToSerialize = [this, &Ar]() -> FCompressedAnimSequence&
 		{
-			const double CompressionEndTime = FPlatformTime::Seconds();
-			const double CompressionTime = CompressionEndTime - CompressionStartTime;
-
-			TArray<FBoneData> BoneData;
-			FAnimationUtils::BuildSkeletonMetaData(GetSkeleton(), BoneData);
-			Params.CompressContext->GatherPostCompressionStats(CompressedData, BoneData, GetFName(), CompressionTime, false);
-		}
-	}
-}
-
-void UAnimSequence::ApplyCompressedData(const FString& DataCacheKeySuffix, const bool bPerformFrameStripping, const TArray<uint8>& Data)
-{
-	if (RequestedCompression_DDCKey == DataCacheKeySuffix)
-	{
-		ApplyCompressedData(Data);
-	}
-	else
-	{
-		SetCompressionComplete();
-	}
-}
-#endif
-
-void UAnimSequence::ApplyCompressedData(const TArray<uint8>& Data)
-{
 #if WITH_EDITOR
-	SynchronousAnimatedBoneAttributesCompression();
-#endif
-	if(Data.Num() > 0)
-	{
-		FMemoryReader MemAr(Data);
-		PRAGMA_DISABLE_DEPRECATION_WARNINGS
-		SerializeCompressedData(MemAr, true);
-		PRAGMA_ENABLE_DEPRECATION_WARNINGS
-		//This is only safe during sync anim compression
-		SetSkeletonVirtualBoneGuid(GetSkeleton()->GetVirtualBoneGuid());
-		bUseRawDataOnly = false;
-	}
+			if (Ar.IsSaving() && Ar.IsCooking())
+			{
+				if (const ITargetPlatform* Platform = Ar.CookingTarget())
+				{
+					const FIoHash Hash = CreateDerivedDataKeyHash(Platform);
+					if (!Hash.IsZero())
+					{
+						if (const TUniquePtr<FCompressedAnimSequence>* DataPtr = DataByPlatformKeyHash.Find(Hash))
+						{
+							return **DataPtr;
+						}
 
-#if WITH_EDITOR
-	SetCompressionComplete();
+						// If there is no entry in the map for this platform key, it will be stored in CompressedData - so verify against the stored KeyHash for it
+						ensure(DataKeyHash == Hash);
+					}
+				}
+			}
 #endif
-}
-
-void UAnimSequence::SerializeCompressedData(FArchive& Ar, bool bDDCData)
-{
-	LLM_SCOPE_BYNAME(TEXT("Animation/SequenceData")); 
+			return CompressedData;
+		}();
 	
-	if (!HasAnyFlags(RF_ClassDefaultObject))
-	{
-		CompressedData.SerializeCompressedData(Ar, bDDCData, this, this->GetSkeleton(), BoneCompressionSettings, CurveCompressionSettings);
+		CompressedDataToSerialize.SerializeCompressedData(Ar, bDDCData, this, this->GetSkeleton(), BoneCompressionSettings, CurveCompressionSettings);
 	}
 }
 
@@ -2045,14 +1913,7 @@ void UAnimSequence::OnAnimModelLoaded()
 
 		if (FAnimationUtils::GetForcedRecompressionSetting())
 		{
-			const ITargetPlatform* TargetPlatform = nullptr;
-			ITargetPlatformManagerModule* TargetPlatformManager = GetTargetPlatformManager();
-			if (TargetPlatformManager)
-			{
-				TargetPlatform = TargetPlatformManager->GetRunningTargetPlatform();
-			}
-
-			RequestAnimCompression(FRequestAnimCompressionParams(true, false, false, TargetPlatform));
+			BeginCacheDerivedDataForCurrentPlatform();
 		}
 
 		if ( DataModelInterface->GetNumberOfKeys() == 0 && DataModelInterface->GetNumberOfFloatCurves() == 0 )
@@ -2077,6 +1938,12 @@ void UAnimSequence::OnAnimModelLoaded()
 			ValidateModel();
 			MoveAttributesToModel();
 		}
+
+		// Do not try to cache compressed animation data when cooking, as cook target platform might not match current platform
+		if (!IsRunningCookCommandlet())
+		{
+			BeginCacheDerivedDataForCurrentPlatform();			
+		}
 	}
 }
 
@@ -2091,432 +1958,31 @@ bool UAnimSequence::DoesSequenceContainZeroScale() const
 {
 	ValidateModel();
 
-	const TArray<FBoneAnimationTrack>& BoneAnimationTracks = DataModelInterface->GetBoneAnimationTracks();
-	for (const FBoneAnimationTrack& Track : BoneAnimationTracks)
-	{
-		for (const FVector3f& ScaleKey : Track.InternalTrackData.ScaleKeys)
-		{
-			if (ScaleKey.IsZero())
-			{
-				return true;
-			}
-		}
-	}
-
-	return false;
-}
-
-FGuid UAnimSequence::GenerateGuidFromRawData() const
-{
-	ValidateModel();
-	return DataModelInterface->GenerateGuid();
-}
-
-void CopyTransformToRawAnimationData(const FTransform& BoneTransform, FRawAnimSequenceTrack& Track, int32 Frame)
-{
-	Track.PosKeys[Frame] = FVector3f(BoneTransform.GetTranslation());
-	Track.RotKeys[Frame] = FQuat4f(BoneTransform.GetRotation());
-	Track.RotKeys[Frame].Normalize();
-	Track.ScaleKeys[Frame] = FVector3f(BoneTransform.GetScale3D());
-}
-
-struct FByFramePoseEvalContext
-{
-public:
-	FBoneContainer RequiredBones;
-
-	TArray<FBoneIndexType> RequiredBoneIndexArray;
-
-	FByFramePoseEvalContext(const UAnimSequence* InAnimToEval)
-		: FByFramePoseEvalContext(InAnimToEval->GetSkeleton())
-	{}
-		
-	FByFramePoseEvalContext(USkeleton* InSkeleton)
-	{
-		// Initialize RequiredBones for pose evaluation
-		RequiredBones.SetUseRAWData(true);
-
-		check(InSkeleton);
-
-		RequiredBoneIndexArray.AddUninitialized(InSkeleton->GetReferenceSkeleton().GetNum());
-		for (int32 BoneIndex = 0; BoneIndex < RequiredBoneIndexArray.Num(); ++BoneIndex)
-		{
-			RequiredBoneIndexArray[BoneIndex] = BoneIndex;
-		}
-
-		RequiredBones.InitializeTo(RequiredBoneIndexArray, FCurveEvaluationOption(true), *InSkeleton);
-	}
-
-};
-
-void UAnimSequence::BakeOutVirtualBoneTracks(TArray<FRawAnimSequenceTrack>& NewRawTracks, TArray<FName>& NewAnimationTrackNames, TArray<FTrackToSkeletonMap>& NewTrackToSkeletonMapTable)
-{
-	ValidateModel();
-
-	const int32 NumVirtualBonesOnSkeleton = GetSkeleton()->GetVirtualBones().Num();
-
-	TArray<int> VirtualBonesToAdd;
-	VirtualBonesToAdd.Reserve(NumVirtualBonesOnSkeleton);
-
-	TArray<int32> SourceParents;
-
-	const FReferenceSkeleton& RefSkeleton = GetSkeleton()->GetReferenceSkeleton();
-	
-	const TArray<FBoneAnimationTrack>& BoneAnimationTracks = ResampledAnimationTrackData;
-
-	for (FRawAnimSequenceTrack& RawTrack : NewRawTracks)
+	TArray<FName> BoneNames;
+	DataModelInterface->GetBoneTrackNames(BoneNames);
+	for (const FName& BoneName : BoneNames)
 	{
-		UE::Anim::Compression::SanitizeRawAnimSequenceTrack(RawTrack);
-	}
-
-	const TArray<FVirtualBone>& VirtualBones = GetSkeleton()->GetVirtualBones();
-	for (int32 VBIndex = 0; VBIndex < NumVirtualBonesOnSkeleton; ++VBIndex)
-	{
-		const FVirtualBone& VirtualBone = VirtualBones[VBIndex];
-		if (!NewAnimationTrackNames.Contains(VirtualBone.VirtualBoneName))
+		bool bZeroScale = false;
+		DataModelInterface->IterateBoneKeys(BoneName, [&bZeroScale](const FVector3f& Position, const FQuat4f& Rotation, const FVector3f& Scale, const FFrameNumber& FrameNumber)
 		{
-			//Need to test if we will animation virtual bone. This involves seeing if any bone that can affect the position
-			//of the target relative to the source is animated by this animation. A bone that can affect the relative position
-			//is any both that is a child of the common ancestor of the target and source
-
-			SourceParents.Reset();
-			bool bBuildVirtualBone = false;
-
-			// First get all the bones that form the chain to the source bone. 
-			int32 CurrentBone = RefSkeleton.FindBoneIndex(VirtualBone.SourceBoneName);
-			while (CurrentBone != INDEX_NONE)
-			{
-				SourceParents.Add(CurrentBone);
-				CurrentBone = RefSkeleton.GetParentIndex(CurrentBone);
-			}
-
-			// Now start checking every bone in the target bones hierarchy until a common ancestor is reached. 
-			CurrentBone = RefSkeleton.FindBoneIndex(VirtualBone.TargetBoneName);
-
-			while (!SourceParents.Contains(CurrentBone))
-			{
-				if (Algo::FindBy(NewTrackToSkeletonMapTable, CurrentBone, &FTrackToSkeletonMap::BoneTreeIndex) != nullptr)
-				{
-					//We animate this bone so the virtual bone is needed
-					bBuildVirtualBone = true;
-					break;
-				}
-
-				CurrentBone = RefSkeleton.GetParentIndex(CurrentBone);
-				check(CurrentBone != INDEX_NONE);
-			}
-
-			// Now we have all the non common bones from the target chain we need the same check from the source chain
-			const int32 FirstCommon = SourceParents.IndexOfByKey(CurrentBone);
-			for (int32 i = FirstCommon - 1; i >= 0; --i)
-			{
-				if (Algo::FindBy(NewTrackToSkeletonMapTable, i, &FTrackToSkeletonMap::BoneTreeIndex) != nullptr)
-				{
-					//We animate this bone so the virtual bone is needed
-					bBuildVirtualBone = true;
-					break;
-				}
-			}
-
-			if (bBuildVirtualBone)
-			{
-				VirtualBonesToAdd.Add(VBIndex);
-			}
-		}
-	}
-
-	const int32 NumVirtualBonesToAdd = VirtualBonesToAdd.Num();
-
-	const USkeleton* TargetSkeleton = GetSkeleton();
-
-	const int32 VirtualBoneStartIndex = NewRawTracks.Num();
-
-	const int32 NumTrackKeys = NumberOfSampledKeys;
-
-	// Generate track data for virtual bones
-	for (int32 VBIndex : VirtualBonesToAdd)
-	{
-		const FVirtualBone& VirtualBone = VirtualBones[VBIndex];
-
-		FRawAnimSequenceTrack& RawTrack = NewRawTracks.AddDefaulted_GetRef();
-
-		//Init new tracks
-		RawTrack.PosKeys.SetNumUninitialized(NumTrackKeys);
-		RawTrack.RotKeys.SetNumUninitialized(NumTrackKeys);
-		RawTrack.ScaleKeys.SetNumUninitialized(NumTrackKeys);
-
-		NewAnimationTrackNames.Add(VirtualBone.VirtualBoneName);
-		const int32 VirtualBoneSkeletonIndex = TargetSkeleton->GetReferenceSkeleton().GetRequiredVirtualBones()[VBIndex];
-		NewTrackToSkeletonMapTable.Add(VirtualBoneSkeletonIndex);
-	}
-
-	// Setup data for evaluating bone transforms
-	FMemMark Mark(FMemStack::Get());
-	FByFramePoseEvalContext EvalContext(this);
-
-	//Pose evaluation data
-	FCompactPose Pose;
-	Pose.SetBoneContainer(&EvalContext.RequiredBones);
-
-	FAnimExtractContext ExtractContext;
-	const TArray<FVirtualBoneRefData>& VBRefData = GetSkeleton()->GetReferenceSkeleton().GetVirtualBoneRefData();
-	for (int32 KeyIndex = 0; KeyIndex < NumTrackKeys; ++KeyIndex)
-	{
-		// Initialise curve data from Skeleton
-		FBlendedCurve Curve;
-		Curve.InitFrom(EvalContext.RequiredBones);
-
-		//Grab pose for this frame
-		ExtractContext.CurrentTime = TargetFrameRate.AsSeconds(KeyIndex);
-
-		UE::Anim::FStackAttributeContainer TempAttributes;
-		FAnimationPoseData AnimPoseData(Pose, Curve, TempAttributes);
-		GetAnimationPose(AnimPoseData, ExtractContext);
-
-		// Copy the retrieved frame data into the virtual bone track keys
-		for (int32 VBToAddIndex = 0; VBToAddIndex < VirtualBonesToAdd.Num(); ++VBToAddIndex)
-		{
-			const int32 VBIndex = VirtualBonesToAdd[VBToAddIndex];
-			const FVirtualBoneRefData& VB = VBRefData[VBIndex];
-			CopyTransformToRawAnimationData(Pose[FCompactPoseBoneIndex(VB.VBRefSkelIndex)], NewRawTracks[VBToAddIndex + VirtualBoneStartIndex], KeyIndex);
-			UE::Anim::Compression::SanitizeRawAnimSequenceTrack(NewRawTracks[VBToAddIndex + VirtualBoneStartIndex]);
-		}
-	}
-}
-
-bool IsIdentity(const FVector3d& Pos)
-{
-	return Pos.Equals(FVector3d::ZeroVector);
-}
-
-bool IsIdentity(const FVector3f& Pos)
-{
-	return Pos.Equals(FVector3f::ZeroVector);
-}
-
-bool IsIdentity(const FQuat4d& Rot)
-{
-	return Rot.Equals(FQuat4d::Identity);
-}
-
-bool IsIdentity(const FQuat4f& Rot)
-{
-	return Rot.Equals(FQuat4f::Identity);
-}
-
-template<class KeyType>
-bool IsKeyArrayValidForRemoval(const TArray<KeyType>& Keys)
-{
-	return Keys.Num() == 0 || (Keys.Num() == 1 && IsIdentity(Keys[0]));
-}
-
-bool IsRawTrackValidForRemoval(const FRawAnimSequenceTrack& Track)
-{
-	return	IsKeyArrayValidForRemoval(Track.PosKeys) &&
-			IsKeyArrayValidForRemoval(Track.RotKeys) &&
-			IsKeyArrayValidForRemoval(Track.ScaleKeys);
-}
-
-void UAnimSequence::TestEvalauteAnimation() const
-{
-	FMemMark Mark(FMemStack::Get());
-	FByFramePoseEvalContext EvalContext(this);
-	EvalContext.RequiredBones.SetUseRAWData(false);
-
-	FCompactPose Pose;
-	Pose.SetBoneContainer(&EvalContext.RequiredBones);
-
-	FAnimExtractContext ExtractContext;
-
-	for (int32 KeyIndex = 0; KeyIndex < NumberOfSampledKeys; ++KeyIndex)
-	{
-		// Initialise curve data from Skeleton
-		FBlendedCurve Curve;
-		Curve.InitFrom(EvalContext.RequiredBones);
-
-		//Grab pose for this frame
-		ExtractContext.CurrentTime = TargetFrameRate.AsSeconds(KeyIndex);
-
-		UE::Anim::FStackAttributeContainer TempAttributes;
-		FAnimationPoseData AnimPoseData(Pose, Curve, TempAttributes);
-		GetAnimationPose(AnimPoseData, ExtractContext);
-	}
-}
+			bZeroScale = Scale.IsZero();
+			return !bZeroScale;
+		});
 
-FFloatCurve* GetFloatCurve(TArray<FFloatCurve>& FloatCurves, const USkeleton::AnimCurveUID& CurveUID)
-{
-	for (FFloatCurve& Curve : FloatCurves)
-	{
-		if (Curve.Name.UID == CurveUID)
+		if (bZeroScale)
 		{
-			return &Curve;
+			return true;
 		}
 	}
 
-	return nullptr;
-}
-
-FFloatCurve* FindOrAddCurve(TArray<FFloatCurve>& FloatCurves, const FSmartName& CurveName)
-{
-	FFloatCurve* ReturnCurve = GetFloatCurve(FloatCurves, CurveName.UID);
-
-	if (ReturnCurve == nullptr)
-	{
-		FFloatCurve& NewCurve = FloatCurves.Add_GetRef(FFloatCurve(CurveName, 0));
-		ReturnCurve = &NewCurve;
-	}
-
-	return ReturnCurve;
-	
-}
-
-bool IsNewKeyDifferent(const FRichCurveKey& LastKey, float NewValue)
-{
-	return LastKey.Value != NewValue;
+	return false;
 }
 
-void UAnimSequence::BakeOutAdditiveIntoRawData(TArray<FRawAnimSequenceTrack>& NewRawTracks, TArray<FName>& NewAnimationTrackNames, TArray<FTrackToSkeletonMap>& NewTrackToSkeletonMapTable, TArray<FFloatCurve>& NewFloatCurves, TArray<FRawAnimSequenceTrack>& AdditiveBaseAnimationData)
+FGuid UAnimSequence::GenerateGuidFromRawData() const
 {
-	if (!CanBakeAdditive())
-	{
-		return; // Nothing to do
-	}
-
-	TRACE_CPUPROFILER_EVENT_SCOPE(UAnimSequence::BakeOutAdditiveIntoRawData);
-
-	ValidateModel();
-
-	USkeleton* MySkeleton = GetSkeleton();
-	check(MySkeleton);
-
-	const bool bRequiresValidRefPoseAsset = RefPoseType == EAdditiveBasePoseType::ABPT_AnimScaled || RefPoseType == EAdditiveBasePoseType::ABPT_AnimFrame;
-	ensure(!bRequiresValidRefPoseAsset || (RefPoseSeq && !RefPoseSeq->HasAnyFlags(EObjectFlags::RF_NeedPostLoad)));
-	
-	const FSmartNameMapping* CurveMapping = MySkeleton->GetSmartNameContainer(USkeleton::AnimCurveMappingName);
-	TArray<FName> SkeletonCurveNames;
-	CurveMapping->FillNameArray(SkeletonCurveNames);
-	for (const FFloatCurve& FloatCurve : DataModelInterface->GetFloatCurves())
-	{
-		ensure(SkeletonCurveNames.Contains(FloatCurve.Name.DisplayName));
-	}
-
-	FMemMark Mark(FMemStack::Get());
-
-	FByFramePoseEvalContext EvalContext(this);
-
-	NewRawTracks.Reset(EvalContext.RequiredBoneIndexArray.Num());
-	NewRawTracks.SetNum(EvalContext.RequiredBoneIndexArray.Num());
-
-	for (FRawAnimSequenceTrack& RawTrack : NewRawTracks)
-	{
-		RawTrack.PosKeys.SetNumUninitialized(NumberOfSampledKeys);
-		RawTrack.RotKeys.SetNumUninitialized(NumberOfSampledKeys);
-		RawTrack.ScaleKeys.SetNumUninitialized(NumberOfSampledKeys);
-	}
-
-	// keep the same buffer size
-	AdditiveBaseAnimationData = NewRawTracks;
-
-	NewTrackToSkeletonMapTable.Reset(EvalContext.RequiredBoneIndexArray.Num());
-	NewTrackToSkeletonMapTable.SetNumUninitialized(EvalContext.RequiredBoneIndexArray.Num());
-
-	NewAnimationTrackNames.SetNumUninitialized(EvalContext.RequiredBoneIndexArray.Num());
-	NewAnimationTrackNames.SetNumUninitialized(EvalContext.RequiredBoneIndexArray.Num());
-
-	for (int32 TrackIndex = 0; TrackIndex < EvalContext.RequiredBoneIndexArray.Num(); ++TrackIndex)
-	{
-		NewTrackToSkeletonMapTable[TrackIndex].BoneTreeIndex = TrackIndex;
-		NewAnimationTrackNames[TrackIndex] = GetSkeleton()->GetReferenceSkeleton().GetBoneName(TrackIndex);
-	}
-
-	//Pose evaluation data
-	FCompactPose Pose;
-	Pose.SetBoneContainer(&EvalContext.RequiredBones);
-	FCompactPose BasePose;
-	BasePose.SetBoneContainer(&EvalContext.RequiredBones);
-
-	FAnimExtractContext ExtractContext;
-
-	FBlendedCurve Curve;
-	FBlendedCurve DummyBaseCurve;
-	
-	for (int32 KeyIndex = 0; KeyIndex < NumberOfSampledKeys; ++KeyIndex)
-	{
-		// Initialise curve data from Skeleton
-		Curve.InitFrom(EvalContext.RequiredBones);
-		DummyBaseCurve.InitFrom(EvalContext.RequiredBones);
-
-		//Grab pose for this frame
-		const double PreviousKeyTime = TargetFrameRate.AsSeconds(KeyIndex - 1);
-		const double CurrentKeyTime = TargetFrameRate.AsSeconds(KeyIndex);
-		FFrameTime CurrentFrameTime = KeyIndex;
-
-		ExtractContext.CurrentTime = CurrentKeyTime;
-
-		UE::Anim::FStackAttributeContainer BaseAttributes;
-		FAnimationPoseData AnimPoseData(Pose, Curve, BaseAttributes);
-		GetAnimationPose(AnimPoseData, ExtractContext);
-
-		UE::Anim::FStackAttributeContainer AdditiveAttributes;
-		FAnimationPoseData AnimBasePoseData(BasePose, DummyBaseCurve, AdditiveAttributes);
-		GetAdditiveBasePose(AnimBasePoseData, ExtractContext);
-
-		//Write out every track for this frame
-		for (FCompactPoseBoneIndex TrackIndex(0); TrackIndex < NewRawTracks.Num(); ++TrackIndex)
-		{
-			CopyTransformToRawAnimationData(Pose[TrackIndex], NewRawTracks[TrackIndex.GetInt()], KeyIndex);
-			CopyTransformToRawAnimationData(BasePose[TrackIndex], AdditiveBaseAnimationData[TrackIndex.GetInt()], KeyIndex);
-		}
-
-		//Write out curve data for this frame
-		for (int32 CurveIndex = 0; CurveIndex < Curve.UIDToArrayIndexLUT->Num(); ++CurveIndex)
-		{
-			USkeleton::AnimCurveUID CurveUID = (USkeleton::AnimCurveUID) CurveIndex;
-			int32 ArrayIndex = Curve.GetArrayIndexByUID(CurveUID);
-			if (ArrayIndex != INDEX_NONE)
-			{
-				float CurveWeight = Curve.CurveWeights[ArrayIndex];
-				FFloatCurve* RawCurve = GetFloatCurve(NewFloatCurves, CurveUID);
-				if (!RawCurve && !FMath::IsNearlyZero(CurveWeight)) //Only make a new curve if we are going to give it data
-				{
-					FSmartName NewCurveName;
-					// if we don't have name, there is something wrong here. 
-					ensureAlways(MySkeleton->GetSmartNameByUID(USkeleton::AnimCurveMappingName, CurveUID, NewCurveName));
-					// curve flags don't matter much for compressed curves
-					RawCurve = FindOrAddCurve(NewFloatCurves, NewCurveName);
-				}
-
-				if (RawCurve)
-				{
-					const bool bHasKeys = RawCurve->FloatCurve.GetNumKeys() > 0;
-					if (!bHasKeys)
-					{
-						//Add pre key of 0
-						if (KeyIndex > 0)
-						{
-							RawCurve->UpdateOrAddKey(0.f, PreviousKeyTime);
-						}
-					}
-
-					if (!bHasKeys || IsNewKeyDifferent(RawCurve->FloatCurve.GetLastKey(), CurveWeight))
-					{
-						RawCurve->UpdateOrAddKey(CurveWeight, CurrentKeyTime);
-						TArray<FRichCurveKey>& CurveKeys = RawCurve->FloatCurve.Keys;
-						if (CurveKeys.Num() > 1)
-						{
-							FRichCurveKey& PrevKey = CurveKeys.Last(1);
-							// Round to frame here as it would have been added at a specific frame boundary (though float->double conversion might mean the value is off)
-							const FFrameTime PrevKeyTime = TargetFrameRate.AsFrameTime(PrevKey.Time).RoundToFrame();
-							if (PrevKeyTime < CurrentFrameTime - 1) // Did we skip a frame, if so need to make previous key const
-							{
-								PrevKey.InterpMode = RCIM_Constant;
-							}
-						}
-					}
-				}
-			}
-		}
-	}
+PRAGMA_DISABLE_DEPRECATION_WARNINGS
+	return GenerateGuidFromRawAnimData(RawAnimationData, RawCurveData);
+PRAGMA_ENABLE_DEPRECATION_WARNINGS	
 }
 
 void UAnimSequence::FlagDependentAnimationsAsRawDataOnly() const
@@ -2551,7 +2017,7 @@ void UAnimSequence::ResetAnimation()
 	// clear everything. Making new animation, so need to reset all the things that belong here
 	ValidateModel();
 
-	TargetFrameRate = UAnimationSettings::Get()->GetDefaultFrameRate();
+	PlatformTargetFrameRate = UAnimationSettings::Get()->GetDefaultFrameRate();
 
 	Controller->OpenBracket(LOCTEXT("ResetAnimation_Bracket", "Resetting Animation Sequence"));
 	{
@@ -3039,21 +2505,22 @@ void UAnimSequence::RemapTracksToNewSkeleton( USkeleton* NewSkeleton, bool bConv
 
 			// now calculating old animated space bases
 			// this one calculates aniamted space per bones and per key
+			TArray<FTransform> BoneTransforms;
 			for(int32 BoneIndex=0; BoneIndex<NumBones; ++BoneIndex)
 			{
 				FName BoneName = NewSkeleton->GetReferenceSkeleton().GetBoneName(BoneIndex);
 				const int32 OldBoneIndex = OldSkeleton->GetReferenceSkeleton().FindBoneIndex(BoneName);
-				const int32 TrackIndex = DataModelInterface->GetBoneTrackIndexByName(BoneName);
+				const bool bContainsTrack = DataModelInterface->IsValidBoneTrackName(BoneName);
 				const int32 ParentBoneIndex = NewSkeleton->GetReferenceSkeleton().GetParentIndex(BoneIndex);
 
-				if(TrackIndex != INDEX_NONE)
+				if(bContainsTrack)
 				{
 					// fill up keys - calculate PK1 * K1
+					BoneTransforms.Reset();
+					DataModelInterface->GetBoneTrackTransforms(BoneName, BoneTransforms);
 					for(int32 Key=0; Key<NumKeys; ++Key)
 					{
-						FTransform AnimatedLocalKey;
-						UE::Anim::GetBoneTransformFromModel(DataModelInterface.GetInterface(), AnimatedLocalKey, TrackIndex, Interval * Key, Interpolation);
-
+						FTransform AnimatedLocalKey = BoneTransforms[Key];
 						// note that we apply scale in the animated space
 						// at this point, you should have scaled version of animated skeleton
 						AnimatedLocalKey.ScaleTranslation(OldToNewTranslationRatio[BoneIndex]);
@@ -3095,7 +2562,6 @@ void UAnimSequence::RemapTracksToNewSkeleton( USkeleton* NewSkeleton, bool bConv
 			for(int32 BoneIndex=0; BoneIndex<NumBones; ++BoneIndex)
 			{
 				const FName BoneName = NewSkeleton->GetReferenceSkeleton().GetBoneName(BoneIndex);
-				const int32 TrackIndex = DataModelInterface->GetBoneTrackIndexByName(BoneName);
 				int32 ParentBoneIndex = NewSkeleton->GetReferenceSkeleton().GetParentIndex(BoneIndex);
 
 				for(int32 Key=0; Key<NumKeys; ++Key)
@@ -3128,7 +2594,7 @@ void UAnimSequence::RemapTracksToNewSkeleton( USkeleton* NewSkeleton, bool bConv
 				}
 
 				// now save back to animation data
-				if(TrackIndex != INDEX_NONE)
+				if(DataModelInterface->IsValidBoneTrackName(BoneName))
 				{
 					PosKeys.SetNumUninitialized(NumKeys);
 					RotKeys.SetNumUninitialized(NumKeys);
@@ -3428,20 +2894,20 @@ int32 UAnimSequence::GetSpaceBasedAnimationData(TArray< TArray<FTransform> >& An
 	{
 		// now calculating old animated space bases
 		// this one calculates aniamted space per bones and per key
+		TArray<FTransform> BoneTransforms;
 		for (int32 BoneIndex = 0; BoneIndex < NumBones; ++BoneIndex)
 		{
 			const FName BoneName = MySkeleton->GetReferenceSkeleton().GetBoneName(BoneIndex);
-			const int32 TrackIndex = DataModelInterface->GetBoneTrackIndexByName(BoneName);
 			const int32 ParentBoneIndex = MySkeleton->GetReferenceSkeleton().GetParentIndex(BoneIndex);
 
-			if (TrackIndex != INDEX_NONE)
+			if (DataModelInterface->IsValidBoneTrackName(BoneName))
 			{
+				BoneTransforms.Reset();
+				DataModelInterface->GetBoneTrackTransforms(BoneName, BoneTransforms);
 				// fill up keys - calculate PK1 * K1
 				for (int32 Key = 0; Key < NumKeys; ++Key)
 				{
-					FTransform AnimatedLocalKey;
-					UE::Anim::GetBoneTransformFromModel(DataModelInterface.GetInterface(), AnimatedLocalKey, TrackIndex, Interval * Key, Interpolation);
-
+					FTransform AnimatedLocalKey = BoneTransforms[Key];
 					if (ParentBoneIndex != INDEX_NONE)
 					{
 						AnimationDataInComponentSpace[BoneIndex][Key] = AnimatedLocalKey * AnimationDataInComponentSpace[ParentBoneIndex][Key];
@@ -3645,193 +3111,57 @@ bool UAnimSequence::ConvertRiggingDataToAnimationData(FAnimSequenceTrackContaine
 		TArray<FQuat4f> RotKeys;
 		TArray<FVector3f> ScaleKeys;
 
-		const int32 NumModelKeys = DataModelInterface->GetNumberOfKeys();
-		const FReferenceSkeleton& RefSkeleton = MySkeleton->GetReferenceSkeleton();
-		const URig* Rig = MySkeleton->GetRig();
-		for (int32 NodeIndex = 0; NodeIndex < ValidNumNodes; ++NodeIndex)
-		{
-			FName BoneName = MySkeleton->GetRigBoneMapping(ValidNodeNames[NodeIndex]);
-			int32 BoneIndex = RefSkeleton.FindBoneIndex(BoneName);
-
-			if (BoneIndex != INDEX_NONE)
-			{
-				// Add a new track by name
-				Controller->AddBoneTrack(BoneName);
-
-				PosKeys.SetNumUninitialized(NumModelKeys);
-				RotKeys.SetNumUninitialized(NumModelKeys);
-				ScaleKeys.SetNumUninitialized(NumModelKeys);
-
-				const int32& ParentBoneIndex = RefSkeleton.GetParentIndex(BoneIndex);
-
-				if(ParentBoneIndex != INDEX_NONE)
-				{
-					for(int32 KeyIndex = 0; KeyIndex < NumModelKeys; ++KeyIndex)
-					{
-						const FTransform LocalTransform = AnimationDataInComponentSpace[BoneIndex][KeyIndex].GetRelativeTransform(AnimationDataInComponentSpace[ParentBoneIndex][KeyIndex]);
-
-						PosKeys[KeyIndex] = FVector3f(LocalTransform.GetTranslation());
-						RotKeys[KeyIndex] = FQuat4f(LocalTransform.GetRotation());
-						ScaleKeys[KeyIndex] = FVector3f(LocalTransform.GetScale3D());
-					}
-				}
-				else
-				{
-					for(int32 KeyIndex = 0; KeyIndex < NumModelKeys; ++KeyIndex)
-					{
-						const FTransform LocalTransform = AnimationDataInComponentSpace[BoneIndex][KeyIndex];
-
-						PosKeys[KeyIndex] = FVector3f(LocalTransform.GetTranslation());
-						RotKeys[KeyIndex] = FQuat4f(LocalTransform.GetRotation());
-						ScaleKeys[KeyIndex] = FVector3f(LocalTransform.GetScale3D());
-					}
-				}
-
-				// Set the track's key data
-				Controller->SetBoneTrackKeys(BoneName, PosKeys, RotKeys, ScaleKeys);
-			}
-		}
-
-		return true;
-	}
-
-	return false;
-}
-
-void UAnimSequence::BakeTrackCurvesToRawAnimationTracks(TArray<FRawAnimSequenceTrack>& InOutNewRawTracks, TArray<FName>& InOutNewTrackNames, TArray<FTrackToSkeletonMap>& InOutNewTrackToSkeletonMapTable)
-{
-	ValidateModel();
-
-	if (DataModelInterface->GetNumberOfTransformCurves() > 0)
-	{
-		USkeleton* CurSkeleton = GetSkeleton();
-		check(CurSkeleton);
-
-		const TArray<FBoneAnimationTrack>& BoneAnimationTracks = ResampledAnimationTrackData;
-			
-		InOutNewRawTracks.Reserve(BoneAnimationTracks.Num());
-		InOutNewTrackToSkeletonMapTable.Reserve(BoneAnimationTracks.Num());
-		InOutNewTrackNames.Reserve(BoneAnimationTracks.Num());
-
-		for (FRawAnimSequenceTrack& Track : InOutNewRawTracks)
-		{
-		// since now I'm about to modify Scale Keys. I should add all of them here at least one key. 
-		// if all turns out to be same, it will clear it up. 
-			if (Track.ScaleKeys.Num() == 0)
-			{
-				// at least add one
-				static FVector3f ScaleConstantKey(1.f);
-				Track.ScaleKeys.Add(ScaleConstantKey);
-			}
-		}
-
-		auto InsertNewBoneTrack = [&InOutNewRawTracks, &InOutNewTrackToSkeletonMapTable, &InOutNewTrackNames, CurSkeleton](const FName& BoneName) -> int32
-		{
-			const FReferenceSkeleton& RefSkeleton = CurSkeleton->GetReferenceSkeleton();
-			const int32 NewTrackIndex = FindFirstChildTrack(CurSkeleton, RefSkeleton, InOutNewTrackNames, BoneName);
-			const int32 BoneIndex = RefSkeleton.FindBoneIndex(BoneName);
-			if (NewTrackIndex != INDEX_NONE)
-			{
-				const TArray<FTransform>& RefPose = RefSkeleton.GetRefBonePose();
-
-				FRawAnimSequenceTrack& RawTrack = InOutNewRawTracks.InsertDefaulted_GetRef(NewTrackIndex);
-				RawTrack.PosKeys.Add(FVector3f(RefPose[BoneIndex].GetTranslation()));
-				RawTrack.RotKeys.Add(FQuat4f(RefPose[BoneIndex].GetRotation()));
-				RawTrack.ScaleKeys.Add(FVector3f(RefPose[BoneIndex].GetScale3D()));
-
-				InOutNewTrackToSkeletonMapTable.Insert(BoneIndex, NewTrackIndex);
-				InOutNewTrackNames.Insert(BoneName, NewTrackIndex);
-
-				check(InOutNewRawTracks.Num() == InOutNewTrackToSkeletonMapTable.Num());
-			}
-
-			return NewTrackIndex;
-		};
-
-
-		// NumFrames can't be zero (filtered earlier)
-
-		const FSmartNameMapping* NameMapping = CurSkeleton->GetSmartNameContainer(USkeleton::AnimTrackCurveMappingName);
-		const FAnimationCurveData& AnimationCurveData = DataModelInterface->GetCurveData();
-		for (const FTransformCurve& Curve : AnimationCurveData.TransformCurves)
+		const int32 NumModelKeys = DataModelInterface->GetNumberOfKeys();
+		const FReferenceSkeleton& RefSkeleton = MySkeleton->GetReferenceSkeleton();
+		const URig* Rig = MySkeleton->GetRig();
+		for (int32 NodeIndex = 0; NodeIndex < ValidNumNodes; ++NodeIndex)
 		{
-			// find curves first, and then see what is index of this curve
-			FName BoneName;
-			NameMapping->GetName(Curve.Name.UID, BoneName);
-			if (Curve.GetCurveTypeFlag(AACF_Disabled) == false &&
-				ensureAlways(NameMapping->GetName(Curve.Name.UID, BoneName)))
-			{
-				int32 TrackIndex = InOutNewTrackNames.Find(BoneName);
+			FName BoneName = MySkeleton->GetRigBoneMapping(ValidNodeNames[NodeIndex]);
+			int32 BoneIndex = RefSkeleton.FindBoneIndex(BoneName);
 
-				// the animation data doesn't have this track, so insert it
-				if (TrackIndex == INDEX_NONE)
-				{
-					TrackIndex = InsertNewBoneTrack(BoneName);
-					// if it still didn't find, something went horribly wrong
-					if(TrackIndex == INDEX_NONE)
-					{
-						UE_LOG(LogAnimation, Log, TEXT("Animation Track Curve Baking : Error adding %s track to %s. Bone does not exist. Skipping"), *BoneName.ToString(), *GetName());
-						// I can't do anything about it
-						continue;
-					}
-				}
+			if (BoneIndex != INDEX_NONE)
+			{
+				// Add a new track by name
+				Controller->AddBoneCurve(BoneName);
 
-				// now modify data
-				FRawAnimSequenceTrack& RawTrack = InOutNewRawTracks[TrackIndex];
+				PosKeys.SetNumUninitialized(NumModelKeys);
+				RotKeys.SetNumUninitialized(NumModelKeys);
+				ScaleKeys.SetNumUninitialized(NumModelKeys);
 
-				// since now we're editing keys, 
-				// if 1 (which meant constant), just expands to # of frames
-				if (RawTrack.PosKeys.Num() == 1)
-				{
-					FVector3f OneKey = RawTrack.PosKeys[0];
-					RawTrack.PosKeys.Init(OneKey, NumberOfSampledKeys);
-				}
-				else
-				{
-					ensure(RawTrack.PosKeys.Num() == NumberOfSampledKeys);
-				}
+				const int32& ParentBoneIndex = RefSkeleton.GetParentIndex(BoneIndex);
 
-				if (RawTrack.RotKeys.Num() == 1)
-				{
-					FQuat4f OneKey = RawTrack.RotKeys[0];
-					RawTrack.RotKeys.Init(OneKey, NumberOfSampledKeys);
-				}
-				else
+				if(ParentBoneIndex != INDEX_NONE)
 				{
-					ensure(RawTrack.RotKeys.Num() == NumberOfSampledKeys);
-				}
+					for(int32 KeyIndex = 0; KeyIndex < NumModelKeys; ++KeyIndex)
+					{
+						const FTransform LocalTransform = AnimationDataInComponentSpace[BoneIndex][KeyIndex].GetRelativeTransform(AnimationDataInComponentSpace[ParentBoneIndex][KeyIndex]);
 
-				// although we don't allow edit of scale
-				// it is important to consider scale when apply transform
-				// so make sure this also is included
-				if (RawTrack.ScaleKeys.Num() == 1)
-				{
-					FVector3f OneKey = RawTrack.ScaleKeys[0];
-					RawTrack.ScaleKeys.Init(OneKey, NumberOfSampledKeys);
+						PosKeys[KeyIndex] = FVector3f(LocalTransform.GetTranslation());
+						RotKeys[KeyIndex] = FQuat4f(LocalTransform.GetRotation());
+						ScaleKeys[KeyIndex] = FVector3f(LocalTransform.GetScale3D());
+					}
 				}
 				else
 				{
-					ensure(RawTrack.ScaleKeys.Num() == NumberOfSampledKeys);
-				}
+					for(int32 KeyIndex = 0; KeyIndex < NumModelKeys; ++KeyIndex)
+					{
+						const FTransform LocalTransform = AnimationDataInComponentSpace[BoneIndex][KeyIndex];
 
-				// now we have all data ready to apply
-				for (int32 KeyIndex = 0; KeyIndex < NumberOfSampledKeys; ++KeyIndex)
-				{
-					// now evaluate
-					const double EvaluationTime = DataModelInterface->GetFrameRate().AsSeconds(KeyIndex);
-					FTransform AdditiveTransform = Curve.Evaluate(EvaluationTime, 1.0);
-					FTransform LocalTransform(FQuat(RawTrack.RotKeys[KeyIndex]), FVector(RawTrack.PosKeys[KeyIndex]), FVector(RawTrack.ScaleKeys[KeyIndex]));
-
-					RawTrack.RotKeys[KeyIndex] = FQuat4f(LocalTransform.GetRotation() * AdditiveTransform.GetRotation());
-					RawTrack.PosKeys[KeyIndex] = FVector3f(LocalTransform.TransformPosition(AdditiveTransform.GetTranslation()));
-					RawTrack.ScaleKeys[KeyIndex] = FVector3f(LocalTransform.GetScale3D() * AdditiveTransform.GetScale3D());
+						PosKeys[KeyIndex] = FVector3f(LocalTransform.GetTranslation());
+						RotKeys[KeyIndex] = FQuat4f(LocalTransform.GetRotation());
+						ScaleKeys[KeyIndex] = FVector3f(LocalTransform.GetScale3D());
+					}
 				}
 
-				// Apply PostProcess behaviour in-place
-				UE::Anim::Compression::SanitizeRawAnimSequenceTrack(RawTrack);
+				// Set the track's key data
+				Controller->SetBoneTrackKeys(BoneName, PosKeys, RotKeys, ScaleKeys);
 			}
 		}
+
+		return true;
 	}
+
+	return false;
 }
 
 void UAnimSequence::AddKeyToSequence(float Time, const FName& BoneName, const FTransform& AdditiveTransform)
@@ -3895,7 +3225,7 @@ bool UAnimSequence::CreateAnimation(USkeletalMesh* Mesh)
 		for (int32 BoneIndex = 0; BoneIndex < NumBones; ++BoneIndex)
 		{
 			const FName& BoneName = RefSkeleton.GetBoneName(BoneIndex);
-			Controller->AddBoneTrack(BoneName);
+			Controller->AddBoneCurve(BoneName);
 			Controller->SetBoneTrackKeys(BoneName, { FVector3f(RefBonePose[BoneIndex].GetTranslation()), FVector3f(RefBonePose[BoneIndex].GetTranslation()) }, { FQuat4f(RefBonePose[BoneIndex].GetRotation()), FQuat4f(RefBonePose[BoneIndex].GetRotation()) }, { FVector3f(RefBonePose[BoneIndex].GetScale3D()), FVector3f(RefBonePose[BoneIndex].GetScale3D()) });
 		}
 
@@ -3932,7 +3262,7 @@ bool UAnimSequence::CreateAnimation(USkeletalMeshComponent* MeshComponent)
 		for (int32 BoneIndex = 0; BoneIndex < NumBones; ++BoneIndex)
 		{
 			const FName& BoneName = RefSkeleton.GetBoneName(BoneIndex);
-			Controller->AddBoneTrack(BoneName);
+			Controller->AddBoneCurve(BoneName);
 			Controller->SetBoneTrackKeys(BoneName, TArray<FVector>{ BoneSpaceTransforms[BoneIndex].GetTranslation(), BoneSpaceTransforms[BoneIndex].GetTranslation() }, { BoneSpaceTransforms[BoneIndex].GetRotation(), BoneSpaceTransforms[BoneIndex].GetRotation() }, TArray<FVector>{ BoneSpaceTransforms[BoneIndex].GetScale3D(), BoneSpaceTransforms[BoneIndex].GetScale3D() });
 		}
 
@@ -3996,7 +3326,7 @@ int32 UAnimSequence::GetNumberOfSampledKeys() const
 #if WITH_EDITOR
 	return NumberOfSampledKeys;
 #else
-	return CompressedData.CompressedDataStructure ? CompressedData.CompressedDataStructure->CompressedNumberOfKeys : TargetFrameRate.AsFrameNumber(GetPlayLength()).Value + 1;
+	return CompressedData.CompressedDataStructure ? CompressedData.CompressedDataStructure->CompressedNumberOfKeys :  PlatformTargetFrameRate.Default.AsFrameNumber(GetPlayLength()).Value + 1;
 #endif
 }
 
@@ -4476,6 +3806,34 @@ void UAnimSequence::RemoveAllCustomAttributes()
 
 void UAnimSequence::SynchronousAnimatedBoneAttributesCompression()
 {
+	struct FByFramePoseEvalContext
+	{
+	public:
+		FBoneContainer RequiredBones;
+
+		TArray<FBoneIndexType> RequiredBoneIndexArray;
+
+		FByFramePoseEvalContext(const UAnimSequence* InAnimToEval)
+			: FByFramePoseEvalContext(InAnimToEval->GetSkeleton())
+		{}
+		
+		FByFramePoseEvalContext(USkeleton* InSkeleton)
+		{
+			// Initialize RequiredBones for pose evaluation
+			RequiredBones.SetUseRAWData(true);
+
+			check(InSkeleton);
+
+			RequiredBoneIndexArray.AddUninitialized(InSkeleton->GetReferenceSkeleton().GetNum());
+			for (int32 BoneIndex = 0; BoneIndex < RequiredBoneIndexArray.Num(); ++BoneIndex)
+			{
+				RequiredBoneIndexArray[BoneIndex] = BoneIndex;
+			}
+
+			RequiredBones.InitializeTo(RequiredBoneIndexArray, FCurveEvaluationOption(true), *InSkeleton);
+		}
+	};
+	
 	ValidateModel();
 
 	AttributeCurves.Empty();
@@ -4525,7 +3883,7 @@ void UAnimSequence::SynchronousAnimatedBoneAttributesCompression()
 				const FAnimatedBoneAttribute& RefAttribute = *RefPoseAttributePtr;
 				for (int32 KeyIndex = 0; KeyIndex < NumberOfSampledKeys; ++KeyIndex)
 				{
-					const double CurrentFrameTime = TargetFrameRate.AsSeconds(KeyIndex);
+					const double CurrentFrameTime = PlatformTargetFrameRate.Default.AsSeconds(KeyIndex);
 					UE::Anim::FStackAttributeContainer AdditiveAttributes;
 					uint8* AdditivePtr = AdditiveAttributes.FindOrAdd(AdditiveAttribute.Identifier.GetType(), UE::Anim::FAttributeId(NAME_None, 0,NAME_None));
 					AdditiveAttribute.Curve.EvaluateToPtr(AdditiveAttribute.Identifier.GetType(), CurrentFrameTime, AdditivePtr);
@@ -5206,6 +4564,7 @@ void UAnimSequence::PopulateModel()
 	const float PlayLength = SequenceLength;
 	// Reset target framerate to current sampling frame rate
 	TargetFrameRate = SamplingFrameRate;
+	PlatformTargetFrameRate = SamplingFrameRate;
 	const FFrameRate FrameRate = SamplingFrameRate;
 	const FRawCurveTracks CurveData = RawCurveData;
 	const TArray<FRawAnimSequenceTrack> SequenceTracks = SourceRawAnimationData_DEPRECATED.Num() && SourceRawAnimationData_DEPRECATED.Num() == RawAnimationData.Num() ?
@@ -5253,8 +4612,7 @@ void UAnimSequence::PopulateModel()
 			continue;
 		}
 		
-		const int32 BoneTrackIndex = Controller->AddBoneTrack(TempAnimationTrackNames[TrackIndex]);
-		if (BoneTrackIndex != INDEX_NONE)
+		if (Controller->AddBoneCurve(TempAnimationTrackNames[TrackIndex]))
 		{
 			auto GenerateUniformKeys = [NumKeys](auto& Keys, auto IdentityKey)
 			{
@@ -5322,10 +4680,10 @@ void UAnimSequence::OnModelModified(const EAnimDataModelNotifyType& NotifyType,
 	{
 		if (NotifyCollector.IsNotWithinBracket())
 		{
-			ClearResampledAnimationTrackData();
 			CalculateNumberOfSampledKeys();
 			UpdateRawDataGuid(RegenerateGUID);
-			RecompressAnimationData();
+			ClearAllCachedCookedPlatformData();
+			BeginCacheDerivedDataForCurrentPlatform();
 		}
 	};
 
@@ -5420,11 +4778,15 @@ void UAnimSequence::OnModelModified(const EAnimDataModelNotifyType& NotifyType,
 
 		case EAnimDataModelNotifyType::Populated:
 		{
-			TargetFrameRate = Model->GetFrameRate();
+			PlatformTargetFrameRate = Model->GetFrameRate();
+				
+			PRAGMA_DISABLE_DEPRECATION_WARNINGS
+			RawCurveData.Empty();
+			PRAGMA_ENABLE_DEPRECATION_WARNINGS
 
 			if (NotifyCollector.IsNotWithinBracket())
 			{
-				HandleTrackDataChanged();
+				HandleTrackDataChanged();				
 			}	
 			break;
 		}
@@ -5442,7 +4804,6 @@ void UAnimSequence::OnModelModified(const EAnimDataModelNotifyType& NotifyType,
 				
 				if (NotifyCollector.Contains(LengthChangingNotifies) || NotifyCollector.Contains(ResamplingNotifies))
 				{
-					ClearResampledAnimationTrackData();
 					CalculateNumberOfSampledKeys();
 				}
 
@@ -5451,7 +4812,7 @@ void UAnimSequence::OnModelModified(const EAnimDataModelNotifyType& NotifyType,
 					const bool bWasModelReset = NotifyCollector.Contains(EAnimDataModelNotifyType::Reset);
 					UpdateRawDataGuid(bWasModelReset ? GenerateNewGUID : RegenerateGUID);
 					ClearCompressedCurveData();
-					RecompressAnimationData();
+					HandleTrackDataChanged();
 				}
 			}
 
@@ -5467,10 +4828,10 @@ void UAnimSequence::OnModelModified(const EAnimDataModelNotifyType& NotifyType,
 		{
 			if (NotifyCollector.IsNotWithinBracket())
 			{
-				ClearResampledAnimationTrackData();
 				CalculateNumberOfSampledKeys();
 				UpdateRawDataGuid(GenerateNewGUID);
-				RecompressAnimationData();
+				ClearAllCachedCookedPlatformData();
+				BeginCacheDerivedDataForCurrentPlatform();
 			}
 
 			break;
@@ -5490,12 +4851,12 @@ void UAnimSequence::OnModelModified(const EAnimDataModelNotifyType& NotifyType,
 		case EAnimDataModelNotifyType::CurveFlagsChanged:
 		case EAnimDataModelNotifyType::CurveScaled:
 		{
-			ClearCompressedCurveData();
+			ClearAllCachedCookedPlatformData();
 				
 			if (NotifyCollector.IsNotWithinBracket())
 			{						
 				UpdateRawDataGuid(RegenerateGUID);
-				RecompressAnimationData();
+				BeginCacheDerivedDataForCurrentPlatform();
 			}
 
 			break;
@@ -5508,7 +4869,7 @@ void UAnimSequence::OnModelModified(const EAnimDataModelNotifyType& NotifyType,
 			if (NotifyCollector.IsNotWithinBracket())
 			{
 				UpdateRawDataGuid(RegenerateGUID);
-				RecompressAnimationData();
+				BeginCacheDerivedDataForCurrentPlatform();
 			}
 			break;
 		}
@@ -5553,7 +4914,7 @@ void UAnimSequence::OnModelModified(const EAnimDataModelNotifyType& NotifyType,
 
 void UAnimSequence::CalculateNumberOfSampledKeys()
 {
-	if (TargetFrameRate == DataModelInterface->GetFrameRate())
+	if (PlatformTargetFrameRate.Default == DataModelInterface->GetFrameRate())
 	{
 		NumberOfSampledKeys = DataModelInterface->GetNumberOfKeys();
 		NumberOfSampledFrames = DataModelInterface->GetNumberOfFrames();
@@ -5561,7 +4922,7 @@ void UAnimSequence::CalculateNumberOfSampledKeys()
 	else
 	{
 		const FFrameNumber ModelNumberOfFrames = DataModelInterface->GetNumberOfFrames();
-		const FFrameTime ResampledFrameTime = FFrameRate::TransformTime(ModelNumberOfFrames, DataModelInterface->GetFrameRate(), TargetFrameRate);
+		const FFrameTime ResampledFrameTime = FFrameRate::TransformTime(ModelNumberOfFrames, DataModelInterface->GetFrameRate(), PlatformTargetFrameRate.Default);
 		ensureMsgf(FMath::IsNearlyZero(ResampledFrameTime.GetSubFrame()), TEXT("Incompatible resampling frame rate for animation sequence %s, frame remainder of %1.8f"), *GetName(), ResampledFrameTime.GetSubFrame());
 
 		NumberOfSampledFrames = ResampledFrameTime.FloorToFrame().Value;
@@ -5569,82 +4930,323 @@ void UAnimSequence::CalculateNumberOfSampledKeys()
 	}
 }
 
-
-void UAnimSequence::ClearResampledAnimationTrackData()
+FIoHash UAnimSequence::CreateDerivedDataKeyHash(const ITargetPlatform* TargetPlatform)
 {
-	ResampledAnimationTrackData.Empty();
+	const USkeleton* CurrentSkeleton = GetSkeleton();
+	if (CurrentSkeleton == nullptr)
+	{
+		UE_LOG(LogAnimation, Warning, TEXT("Animation Compression request for %s failed, Skeleton == nullptr."), *GetName());
+		return FIoHash::Zero;
+	}
+
+	// New animation DDC key format
+	FMemoryHasherBlake3 Writer;
+
+	const bool bIsValidAdditive = IsValidAdditive();
+	const char AdditiveType = bIsValidAdditive ? NibbleToTChar(AdditiveAnimType) : '0';
+	const char RefType = bIsValidAdditive ? NibbleToTChar(RefPoseType) : '0';
+
+	bool bPerformFrameStripping  = false;
+	bool bPerformFrameStrippingOnOddNumberedFrames = false;
+	{
+		if (UDeviceProfile* DeviceProfile = UDeviceProfileManager::Get().FindProfile(TargetPlatform->IniPlatformName()))
+		{
+			int32 CVarPlatformFrameStrippingValue = 0;
+			if (DeviceProfile->GetConsolidatedCVarValue(StripFrameCVarName, CVarPlatformFrameStrippingValue))
+			{
+				bPerformFrameStripping = CVarPlatformFrameStrippingValue == 1;
+			}
+			
+			int32 CVarPlatformOddAnimFrameStrippingValue = 0;
+			if (DeviceProfile->GetConsolidatedCVarValue(OddFrameStripStrippingCVarName, CVarPlatformOddAnimFrameStrippingValue))
+			{
+				bPerformFrameStrippingOnOddNumberedFrames = CVarPlatformOddAnimFrameStrippingValue == 1;
+			}
+		}
+
+		PRAGMA_DISABLE_DEPRECATION_WARNINGS
+		bPerformFrameStripping = ShouldPerformStripping(bPerformFrameStripping, bPerformFrameStrippingOnOddNumberedFrames);
+		PRAGMA_ENABLE_DEPRECATION_WARNINGS
+	}
+		
+	FArcToHexString ArcToHexString;
+
+	ArcToHexString.Ar << CompressionErrorThresholdScale;
+	ArcToHexString.Ar << bPerformFrameStripping;
+	BoneCompressionSettings->PopulateDDCKey(UE::Anim::Compression::FAnimDDCKeyArgs(*this, TargetPlatform), ArcToHexString.Ar);
+	CurveCompressionSettings->PopulateDDCKey(ArcToHexString.Ar);
+	
+	const FFrameRate FrameRate = PlatformTargetFrameRate.GetValueForPlatform(TargetPlatform->GetPlatformInfo().IniPlatformName);
+
+	FString Ret = FString::Printf(TEXT("%i_%s%s%s_%c%c%i_%s_%s_%i_%i_%s"),
+		CompressCommandletVersion,
+		PRAGMA_DISABLE_DEPRECATION_WARNINGS
+		*GetDataModel()->GenerateGuid().ToString(),
+		PRAGMA_ENABLE_DEPRECATION_WARNINGS
+		*GetSkeleton()->GetGuid().ToString(),
+		*GetSkeleton()->GetVirtualBoneGuid().ToString(),
+		AdditiveType,
+		RefType,
+		RefFrameIndex,
+		PRAGMA_DISABLE_DEPRECATION_WARNINGS
+		(bIsValidAdditive && RefPoseSeq) ? *RefPoseSeq->GetDataModel()->GenerateGuid().ToString() : TEXT("NoAdditiveGuid"),
+		PRAGMA_ENABLE_DEPRECATION_WARNINGS
+		*ArcToHexString.MakeString(),
+		FrameRate.Numerator,
+		FrameRate.Denominator,
+		*UE::Anim::Compression::AnimationCompressionVersionString
+	);
+
+	Writer << Ret;
+
+	return Writer.Finalize();
 }
 
-void UAnimSequence::ResampleAnimationTrackData()
+FIoHash UAnimSequence::BeginCacheDerivedData(const ITargetPlatform* TargetPlatform)
 {
-	QUICK_SCOPE_CYCLE_COUNTER(STAT_ResampleAnimationTrackData);
+	QUICK_SCOPE_CYCLE_COUNTER(STAT_AnimSequence_BeginCacheDerivedData);
 	
-	if (ShouldDataModelBeValid())
+	check(TargetPlatform != nullptr);
+	check(IsInGameThread());
+	check(!FPlatformProperties::RequiresCookedData());
+	
+	if(bBlockCompressionRequests)
 	{
-		ValidateModel();
+		UE_LOG(LogAnimation, Warning, TEXT("Animation Compression request for %s was blocked, bBlockCompressionRequests == true."), *GetName());
+		return FIoHash::Zero;
+	}
 
-		FAnimationUtils::EnsureAnimSequenceLoaded(*this);
+	// Ensure that there are valid compression settings
+	if (BoneCompressionSettings == nullptr || !BoneCompressionSettings->AreSettingsValid())
+	{
+		BoneCompressionSettings = FAnimationUtils::GetDefaultAnimationBoneCompressionSettings();
+	}
+	
+	if (CurveCompressionSettings == nullptr || !CurveCompressionSettings->AreSettingsValid())
+	{
+		CurveCompressionSettings = FAnimationUtils::GetDefaultAnimationCurveCompressionSettings();
+	}
 
-		// Make a copy, deals with bone name and index
-		ResampledAnimationTrackData = DataModelInterface->GetBoneAnimationTracks();
-		CalculateNumberOfSampledKeys();
+	// Make sure all our required dependencies are loaded
+	FAnimationUtils::EnsureAnimSequenceLoaded(*this);
+		
+	const FIoHash KeyHash = CreateDerivedDataKeyHash(TargetPlatform);
+	// Early out if not valid, equal to running platform hash or any previously cached hashes
+	if (KeyHash.IsZero() || DataKeyHash == KeyHash || DataByPlatformKeyHash.Contains(KeyHash))
+	{
+		return KeyHash;
+	}
+
+	// Wait for any in-flight requests
+	UE::Anim::FAnimSequenceCompilingManager::Get().FinishCompilation({this});
 	
-		ensure(DataModelInterface != nullptr);
+	bUseRawDataOnly = true;
+
+	// We should wait here for any previous compilations?
+	FCompressedAnimSequence* TargetData = nullptr;
+	if (TargetPlatform->IsRunningPlatform())
+	{
+		DataKeyHash = KeyHash;
+		TargetData = &CompressedData;
+	}
+	else
+    {
+    	TargetData = DataByPlatformKeyHash.Emplace(KeyHash, MakeUnique<FCompressedAnimSequence>()).Get();
+    }
+	check(TargetData);
+	TargetData->Reset();
+
+	bool bPerformFrameStripping  = false;
+	{
+		bool bPerformFrameStrippingOnOddNumberedFrames = false;
+		if (UDeviceProfile* DeviceProfile = UDeviceProfileManager::Get().FindProfile(TargetPlatform->IniPlatformName()))
 		{
-			// Can directly copy (for now, as we still maintain GetBoneAnimationTracks)
-			if (TargetFrameRate != DataModelInterface->GetFrameRate())
+			int32 CVarPlatformFrameStrippingValue = 0;
+			if (DeviceProfile->GetConsolidatedCVarValue(StripFrameCVarName, CVarPlatformFrameStrippingValue))
 			{
-				FMemMark Mark(FMemStack::Get());
-				FByFramePoseEvalContext EvalContext(this);
-				EvalContext.RequiredBones.SetDisableRetargeting(true);
-				EvalContext.RequiredBones.SetUseRAWData(true);
-				EvalContext.RequiredBones.SetUseSourceData(true);
-
-				// We actually need to resample bone transforms
-				for (int32 TrackIndex = 0; TrackIndex < ResampledAnimationTrackData.Num(); ++TrackIndex)
-				{
-					FBoneAnimationTrack& TrackData = ResampledAnimationTrackData[TrackIndex];
-					TrackData.InternalTrackData.PosKeys.SetNumUninitialized(NumberOfSampledKeys);
-					TrackData.InternalTrackData.RotKeys.SetNumUninitialized(NumberOfSampledKeys);
-					TrackData.InternalTrackData.ScaleKeys.SetNumUninitialized(NumberOfSampledKeys);
-				}
-				
-				FBlendedCurve Curve;
-				Curve.InitFrom(EvalContext.RequiredBones);
-				UE::Anim::FStackAttributeContainer AttributeContainer;
+				bPerformFrameStripping = CVarPlatformFrameStrippingValue == 1;
+			}
+			
+			int32 CVarPlatformOddAnimFrameStrippingValue = 0;
+			if (DeviceProfile->GetConsolidatedCVarValue(OddFrameStripStrippingCVarName, CVarPlatformOddAnimFrameStrippingValue))
+			{
+				bPerformFrameStrippingOnOddNumberedFrames = CVarPlatformOddAnimFrameStrippingValue == 1;
+			}
+		}
+
+		PRAGMA_DISABLE_DEPRECATION_WARNINGS
+		bPerformFrameStripping = ShouldPerformStripping(bPerformFrameStripping, bPerformFrameStrippingOnOddNumberedFrames);
+		PRAGMA_ENABLE_DEPRECATION_WARNINGS
+	}
+	
+	{		
+		// Data does not exist, need to build it.
+		const FCompressibleAnimRef CompressibleData = MakeShared<FCompressibleAnimData, ESPMode::ThreadSafe>(this, bPerformFrameStripping, TargetPlatform);
+		
+		COOK_STAT(auto Timer = UE::Anim::AnimSequenceCookStats::UsageStats.TimeSyncWork());
+		check(CacheTasksByKeyHash.Find(KeyHash) == nullptr);
+		CacheTasksByKeyHash.Emplace(KeyHash, MakePimpl<UE::Anim::FAnimationSequenceAsyncCacheTask>(KeyHash, CompressibleData, TargetData, *this, TargetPlatform));		
+		COOK_STAT(Timer.TrackCyclesOnly());
+	}
+
+	// The compiling manager provides throttling, notification manager, etc... for the asset being built.
+	UE::Anim::FAnimSequenceCompilingManager::Get().AddAnimSequences({this});
+	
+	return KeyHash;
+}
+
+bool UAnimSequence::PollCacheDerivedData(const FIoHash& KeyHash) const
+{
+	if (KeyHash.IsZero())
+	{
+		return true;
+	}
+
+	if (const TPimplPtr<UE::Anim::FAnimationSequenceAsyncCacheTask>* Task = CacheTasksByKeyHash.Find(KeyHash))
+	{
+		return (*Task)->Poll();
+	}
+
+	return true;
+}
+
+void UAnimSequence::EndCacheDerivedData(const FIoHash& KeyHash)
+{
+	//UE_LOG(LogAnimation, Warning, TEXT("EndCacheDerivedData %s"), *GetName());
+	if (KeyHash.IsZero())
+	{
+		return;
+	}
+
+	UE::Anim::FAnimSequenceCompilingManager::Get().FinishCompilation({this});
+
+	// QQ calling FinishCompilation after this, makes it so that in ::FinishAsyncTasks we do not apply the skeleton GUID etc
+	// TPimplPtr<UE::Anim::FAnimationSequenceAsyncCacheTask> Task;
+	// if (CacheTasksByKeyHash.RemoveAndCopyValue(KeyHash, Task))
+	// {
+	// 	Task->Wait();	// 	
+	// }
+}
+	
+FCompressedAnimSequence& UAnimSequence::CacheDerivedData(const ITargetPlatform* TargetPlatform)
+{
+	const FIoHash KeyHash = BeginCacheDerivedData(TargetPlatform);
+	EndCacheDerivedData(KeyHash);
+	return DataKeyHash == KeyHash ? CompressedData : *DataByPlatformKeyHash[KeyHash];
+}
+
+FFrameRate UAnimSequence::GetTargetSamplingFrameRate(const ITargetPlatform* InPlatform) const
+{
+	return PlatformTargetFrameRate.GetValueForPlatform(*InPlatform->IniPlatformName());
+}
+
+void UAnimSequence::BeginCacheDerivedDataForCurrentPlatform()
+{
+	if (const ITargetPlatform* RunningPlatform = GetTargetPlatformManagerRef().GetRunningTargetPlatform())
+	{
+		BeginCacheDerivedData(RunningPlatform);
+	}
+}
+
+void UAnimSequence::CacheDerivedDataForCurrentPlatform()
+{
+	if (const ITargetPlatform* RunningPlatform = GetTargetPlatformManagerRef().GetRunningTargetPlatform())
+	{
+		CacheDerivedData(RunningPlatform);
+	}
+}
 				
-				for (int32 FrameIndex = 0; FrameIndex < NumberOfSampledKeys; ++FrameIndex)
-				{
-					UE::Anim::DataModel::FEvaluationContext EvaluationContext(FFrameTime(FrameIndex), TargetFrameRate, GetRetargetTransformsSourceName(), GetRetargetTransforms());
+bool UAnimSequence::IsCompiling() const
+{
+	return CacheTasksByKeyHash.Num() > 0;
+}
 					
-					FCompactPose Pose;
-					Pose.SetBoneContainer(&EvalContext.RequiredBones);
+bool UAnimSequence::TryCancelAsyncTasks()
+{
+	for (auto It = CacheTasksByKeyHash.CreateIterator(); It; ++It)
+	{
+		if (It->Value->Poll())
+		{
+			It.RemoveCurrent();
+		}
+		else
+		{
+			It->Value->Cancel();
+		}
+	}
 
-					FAnimationPoseData PoseData(Pose, Curve, AttributeContainer);
-					DataModelInterface->Evaluate(PoseData, EvaluationContext);
+	return CacheTasksByKeyHash.IsEmpty();
+}
 
-					for (int32 TrackIndex = 0; TrackIndex < ResampledAnimationTrackData.Num(); ++TrackIndex)
-					{
-						FBoneAnimationTrack& TrackData = ResampledAnimationTrackData[TrackIndex];
-						const FTransform& Transform = Pose[EvalContext.RequiredBones.GetCompactPoseIndexFromSkeletonIndex(TrackData.BoneTreeIndex)];
+bool UAnimSequence::IsAsyncTaskComplete() const
+{
+	for (auto& Pair : CacheTasksByKeyHash)
+	{
+		if (!Pair.Value->Poll())
+		{
+			return false;
+		}
+	}
 
-						TrackData.InternalTrackData.PosKeys[FrameIndex] = FVector3f(Transform.GetLocation());
-						TrackData.InternalTrackData.RotKeys[FrameIndex] = FQuat4f(Transform.GetRotation());
-						TrackData.InternalTrackData.ScaleKeys[FrameIndex] = FVector3f(Transform.GetScale3D());
-					}
-				}
-			}
+	return true;
+}
+
+void UAnimSequence::FinishAsyncTasks()
+{	
+	//UE_LOG(LogAnimation, Warning, TEXT("FinishAsyncTasks %s [%i]"), *GetName(), CacheTasksByKeyHash.Num());
+	const bool bHasInflightTasks = CacheTasksByKeyHash.Num() > 0;
+	if (bHasInflightTasks)
+	{
+		COOK_STAT(auto Timer = UE::Anim::AnimSequenceCookStats::UsageStats.TimeAsyncWait());
+		for (auto It = CacheTasksByKeyHash.CreateIterator(); It; ++It)
+		{
+			It->Value->Wait();
+			It.RemoveCurrent();
 		}
+		COOK_STAT(Timer.TrackCyclesOnly());
+	
+#if WITH_EDITOR
+		SynchronousAnimatedBoneAttributesCompression();
+		//This is only safe during sync anim compression
+		SetSkeletonVirtualBoneGuid(GetSkeleton()->GetVirtualBoneGuid());
+#endif
+		
+		FAssetRegistryModule& AssetRegistryModule = FModuleManager::LoadModuleChecked<FAssetRegistryModule>(TEXT("AssetRegistry"));
+		AssetRegistryModule.Get().AssetTagsFinalized(*this);
+		
+		check(bUseRawDataOnly);
+		bUseRawDataOnly = false;
+	}
+}
+
+void UAnimSequence::Reschedule(FQueuedThreadPool* InThreadPool, EQueuedWorkPriority InPriority)
+{
+	for (auto& Pair : CacheTasksByKeyHash)
+	{
+		Pair.Value->Reschedule(InThreadPool, InPriority);
 	}
 }
 
 void UAnimSequence::RecompressAnimationData()
+{	
+	ClearAllCachedCookedPlatformData();
+	BeginCacheDerivedDataForCurrentPlatform();
+}
+
+void UAnimSequence::RequestAsyncAnimRecompression(bool bOutput)
+{ 
+	if (ITargetPlatform* RunningPlatform = GetTargetPlatformManagerRef().GetRunningTargetPlatform())
+	{
+		BeginCacheDerivedData(RunningPlatform);
+	}
+}
+
+void UAnimSequence::RequestSyncAnimRecompression(bool bOutput)
 {
-	PRAGMA_DISABLE_DEPRECATION_WARNINGS
-	ClearCompressedBoneData();
-	PRAGMA_ENABLE_DEPRECATION_WARNINGS
-	bUseRawDataOnly = true;
-	RequestAsyncAnimRecompression(false);
+	if (ITargetPlatform* RunningPlatform = GetTargetPlatformManagerRef().GetRunningTargetPlatform())
+	{
+		CacheDerivedData(RunningPlatform);
+	}
 }
 #endif // WITH_EDITOR
 
diff --git a/Engine/Source/Runtime/Engine/Private/Animation/AnimSequenceHelpers.cpp b/Engine/Source/Runtime/Engine/Private/Animation/AnimSequenceHelpers.cpp
--- a/Engine/Source/Runtime/Engine/Private/Animation/AnimSequenceHelpers.cpp
+++ b/Engine/Source/Runtime/Engine/Private/Animation/AnimSequenceHelpers.cpp
@@ -94,44 +94,20 @@ void EvaluateTransformCurvesFromModel(const IAnimationDataModel* Model, TMap<FNa
 
 void GetBoneTransformFromModel(const IAnimationDataModel* Model, FTransform& OutTransform, int32 TrackIndex, double Time, const EAnimInterpolationType& Interpolation)
 {
-	const FBoneAnimationTrack& TrackData = Model->GetBoneTrackByIndex(TrackIndex);
-	const FRawAnimSequenceTrack& RawTrack = TrackData.InternalTrackData;
-	
-	FAnimationUtils::ExtractTransformFromTrack(RawTrack, Time, Model->GetNumberOfKeys(), Model->GetPlayLength(), Interpolation, OutTransform);
+	TArray<FName> TrackNames;
+	Model->GetBoneTrackNames(TrackNames);
+	OutTransform = Model->EvaluateBoneTrackTransform(TrackNames[TrackIndex], Model->GetFrameRate().AsFrameTime(Time), Interpolation);
 }
 
 void GetBoneTransformFromModel(const IAnimationDataModel* Model, FTransform& OutTransform, int32 TrackIndex, int32 KeyIndex)
 {
-	const FBoneAnimationTrack& TrackData = Model->GetBoneTrackByIndex(TrackIndex);
-	const FRawAnimSequenceTrack& RawTrack = TrackData.InternalTrackData;
-
-	// Bail out (with rather wacky data) if data is empty for some reason.
-	if (RawTrack.PosKeys.Num() == 0 || RawTrack.RotKeys.Num() == 0)
-	{
-		UE_LOG(LogAnimation, Log, TEXT("UAnimSequence::GetBoneTransform : No anim data in AnimSequence!"));
-		OutTransform.SetIdentity();
-		return;
-	}
-
-	const int32 PosKeyIndex = FMath::Min(KeyIndex, RawTrack.PosKeys.Num() - 1);
-	const int32 RotKeyIndex = FMath::Min(KeyIndex, RawTrack.RotKeys.Num() - 1);
-	static const FVector DefaultScale3D = FVector(1.f);
-
-	OutTransform.SetTranslation(FVector(RawTrack.PosKeys[PosKeyIndex]));
-	OutTransform.SetRotation(FQuat(RawTrack.RotKeys[RotKeyIndex]));
-	if (RawTrack.ScaleKeys.Num() > 0)
-	{
-		const int32 ScaleKeyIndex = FMath::Min(KeyIndex, RawTrack.ScaleKeys.Num() - 1);
-		OutTransform.SetScale3D(FVector(RawTrack.ScaleKeys[ScaleKeyIndex]));
-	}
-	else
-	{
-		OutTransform.SetScale3D(DefaultScale3D);
-	}
+	TArray<FName> TrackNames;
+	Model->GetBoneTrackNames(TrackNames);
+	OutTransform = Model->GetBoneTrackTransform(TrackNames[TrackIndex], KeyIndex);
 
 	for (const FTransformCurve& AdditiveTransformCurve : Model->GetTransformCurves())
 	{
-		if (AdditiveTransformCurve.Name.DisplayName == TrackData.Name)
+		if (AdditiveTransformCurve.Name.DisplayName == TrackNames[TrackIndex])
 		{
 			const float TimeInterval = Model->GetFrameRate().AsSeconds(KeyIndex);
 			const FTransform AdditiveTransform = AdditiveTransformCurve.Evaluate(TimeInterval, 1.f);
@@ -644,19 +620,32 @@ bool AnimationData::AddLoopingInterpolation(UAnimSequence* InSequence)
 		
 		IAnimationDataController::FScopedBracket ScopedBracket(Controller, LOCTEXT("AddLoopingInterpolation_Bracket", "Adding looping interpolation"));
 
-		const TArray<FBoneAnimationTrack>& BoneAnimationTracks = DataModel->GetBoneAnimationTracks();
-		for (const FBoneAnimationTrack& AnimationTrack : BoneAnimationTracks)
+		TArray<FTransform> BoneTransforms;
+		TArray<FName> TrackNames;
+		DataModel->GetBoneTrackNames(TrackNames);
+ 
+		TArray<FVector3f> PositionalKeys, ScaleKeys;
+		TArray<FQuat4f> RotationalKeys;
+		for (const FName& TrackName : TrackNames)
 		{
-			auto PositionalKeys = AnimationTrack.InternalTrackData.PosKeys;
-			LoopKeyData(PositionalKeys);
+			BoneTransforms.Reset();
+			PositionalKeys.Reset();
+			RotationalKeys.Reset();
+			ScaleKeys.Reset();
+			DataModel->GetBoneTrackTransforms(TrackName, BoneTransforms);
 
-			auto RotationalKeys = AnimationTrack.InternalTrackData.RotKeys;
+			for (const FTransform& Transform : BoneTransforms)
+			{
+				PositionalKeys.Add(FVector3f(Transform.GetLocation()));
+				RotationalKeys.Add(FQuat4f(Transform.GetRotation()));
+				ScaleKeys.Add(FVector3f(Transform.GetScale3D()));
+			}
+		
+			LoopKeyData(PositionalKeys);
 			LoopKeyData(RotationalKeys);
-
-			auto ScaleKeys = AnimationTrack.InternalTrackData.ScaleKeys;
 			LoopKeyData(ScaleKeys);
 
-			Controller.SetBoneTrackKeys(AnimationTrack.Name, PositionalKeys, RotationalKeys, ScaleKeys);
+			Controller.SetBoneTrackKeys(TrackName, PositionalKeys, RotationalKeys, ScaleKeys);
 		}
 
 		// New number of frames is equal to current number of keys, as we'll be adding one frame (Number of Frames + 1 == Number of Keys)
@@ -731,7 +720,6 @@ void AnimationData::DuplicateKeys(UAnimSequence* InSequence, int32 StartKeyIndex
 			const int32 NumFramesToInsert = NumDuplicates;
 			const int32 EndFrameIndex = StartKeyIndex + NumDuplicates;
 
-			TArray<FBoneAnimationTrack> BoneAnimationTracks = Model->GetBoneAnimationTracks();
 			auto InsertFrames = [&](auto& KeyData)
 			{
 				if (KeyData.Num() >= 1 && KeyData.IsValidIndex(CopyKeyIndex))
@@ -746,15 +734,50 @@ void AnimationData::DuplicateKeys(UAnimSequence* InSequence, int32 StartKeyIndex
 				}
 			};
 
-			for (FBoneAnimationTrack& AnimationTrack : BoneAnimationTracks)
+			TArray<FName> TrackNames;
+			Model->GetBoneTrackNames(TrackNames);
+			
+			TArray<FTransform> BoneTransforms;
+			BoneTransforms.Reserve(Model->GetNumberOfKeys());
+			
+			for (const FName& TrackName : TrackNames)
 			{
-				FRawAnimSequenceTrack& TrackData = AnimationTrack.InternalTrackData;
-
-				InsertFrames(TrackData.PosKeys);
-				InsertFrames(TrackData.RotKeys);
-				InsertFrames(TrackData.ScaleKeys);
+				BoneTransforms.Reset();
+				Model->GetBoneTrackTransforms(TrackName, BoneTransforms);
+
+				TArray<FVector3f> PosKeys;
+				PosKeys.SetNum(BoneTransforms.Num() + NumDuplicates);
+				TArray<FQuat4f> RotKeys;
+				RotKeys.SetNum(BoneTransforms.Num() + NumDuplicates);
+				TArray<FVector3f> ScaleKeys;
+				ScaleKeys.SetNum(BoneTransforms.Num() + NumDuplicates);
+
+				ensure(BoneTransforms.IsValidIndex(SourceKeyIndex));
+				for (int32 TransformIndex = 0; TransformIndex < BoneTransforms.Num() + NumFramesToInsert; ++TransformIndex)
+				{
+					if (TransformIndex < StartKeyIndex)
+					{
+						PosKeys[TransformIndex] = FVector3f(BoneTransforms[TransformIndex].GetLocation());
+						RotKeys[TransformIndex] = FQuat4f(BoneTransforms[TransformIndex].GetRotation());
+						ScaleKeys[TransformIndex] = FVector3f(BoneTransforms[TransformIndex].GetScale3D());
+					}
+					if (TransformIndex >= StartKeyIndex && TransformIndex < StartKeyIndex + NumDuplicates)
+					{
+						PosKeys[TransformIndex] = FVector3f(BoneTransforms[SourceKeyIndex].GetLocation());
+						RotKeys[TransformIndex] = FQuat4f(BoneTransforms[SourceKeyIndex].GetRotation());
+						ScaleKeys[TransformIndex] = FVector3f(BoneTransforms[SourceKeyIndex].GetScale3D());
+					}
+					else if (TransformIndex >= StartKeyIndex + NumDuplicates)
+					{
+						PosKeys[TransformIndex] = FVector3f(BoneTransforms[TransformIndex - NumDuplicates].GetLocation());
+						RotKeys[TransformIndex] = FQuat4f(BoneTransforms[TransformIndex - NumDuplicates].GetRotation());
+						ScaleKeys[TransformIndex] = FVector3f(BoneTransforms[TransformIndex - NumDuplicates].GetScale3D());
+					}
+				}
 
-				Controller.SetBoneTrackKeys(AnimationTrack.Name, TrackData.PosKeys, TrackData.RotKeys, TrackData.ScaleKeys);
+				Controller.SetBoneTrackKeys(TrackName, PosKeys, RotKeys, ScaleKeys);
+				
+				BoneTransforms.Reset();
 			}
 
 			// The number of keys has changed, which means that the sequence length and number of frames should be updated as well
@@ -803,16 +826,42 @@ void AnimationData::RemoveKeys(UAnimSequence* InSequence, int32 StartKeyIndex, i
 
 		const int32 NewNumberOfKeys = NumberOfKeys - NumKeysToRemove;
 
-		TArray<FBoneAnimationTrack> BoneAnimationTracks = Model->GetBoneAnimationTracks();
-		for (FBoneAnimationTrack& AnimationTrack : BoneAnimationTracks)
+		TArray<FName> TrackNames;
+		Model->GetBoneTrackNames(TrackNames);
+			
+		TArray<FTransform> BoneTransforms;
+		BoneTransforms.Reserve(Model->GetNumberOfKeys());
+			
+		for (const FName& TrackName : TrackNames)
 		{
-			FRawAnimSequenceTrack& TrackData = AnimationTrack.InternalTrackData;
+			BoneTransforms.Reset();
 
-			ShrinkKeys(TrackData.PosKeys);
-			ShrinkKeys(TrackData.RotKeys);
-			ShrinkKeys(TrackData.ScaleKeys);
+			Model->GetBoneTrackTransforms(TrackName, BoneTransforms);
+
+			TArray<FVector3f> PosKeys;
+			PosKeys.SetNum(BoneTransforms.Num() - NumKeysToRemove);
+			TArray<FQuat4f> RotKeys;
+			RotKeys.SetNum(BoneTransforms.Num() - NumKeysToRemove);
+			TArray<FVector3f> ScaleKeys;
+			ScaleKeys.SetNum(BoneTransforms.Num() - NumKeysToRemove);
+
+			for (int32 TransformIndex = 0; TransformIndex < BoneTransforms.Num() - NumKeysToRemove; ++TransformIndex)
+			{
+				if (TransformIndex < StartKeyIndex)
+				{
+					PosKeys[TransformIndex] = FVector3f(BoneTransforms[TransformIndex].GetLocation());
+					RotKeys[TransformIndex] = FQuat4f(BoneTransforms[TransformIndex].GetRotation());
+					ScaleKeys[TransformIndex] = FVector3f(BoneTransforms[TransformIndex].GetScale3D());
+				}
+				if (TransformIndex >= StartKeyIndex)
+				{
+					PosKeys[TransformIndex] = FVector3f(BoneTransforms[TransformIndex + NumKeysToRemove].GetLocation());
+					RotKeys[TransformIndex] = FQuat4f(BoneTransforms[TransformIndex + NumKeysToRemove].GetRotation());
+					ScaleKeys[TransformIndex] = FVector3f(BoneTransforms[TransformIndex + NumKeysToRemove].GetScale3D());		
+				}
+			}
 			
-			Controller.SetBoneTrackKeys(AnimationTrack.Name, TrackData.PosKeys, TrackData.RotKeys, TrackData.ScaleKeys);
+			Controller.SetBoneTrackKeys(TrackName, PosKeys, RotKeys, ScaleKeys);
 		}
 
 		const int32 NewNumberOfFrames = FMath::Max(NewNumberOfKeys - 1, 1);
@@ -824,7 +873,7 @@ void AnimationData::RemoveKeys(UAnimSequence* InSequence, int32 StartKeyIndex, i
 	}
 }
 
-int32 AnimationData::FindFirstChildTrackIndex(const UAnimSequence* InSequence, const USkeleton* Skeleton, const FName& BoneName)
+FName AnimationData::FindFirstChildTrackName(const UAnimSequence* InSequence, const USkeleton* Skeleton, const FName& BoneName)
 {
 	const IAnimationDataModel* DataModel = InSequence->GetDataModel();
 	const FReferenceSkeleton& RefSkeleton = Skeleton->GetReferenceSkeleton();
@@ -833,52 +882,53 @@ int32 AnimationData::FindFirstChildTrackIndex(const UAnimSequence* InSequence, c
 	if (BoneIndex == INDEX_NONE)
 	{
 		// get out, nothing to do
-		return INDEX_NONE;
+		return NAME_None;
 	}
 
 	// find children
-	TArray<int32> Childs;
-	if (Skeleton->GetChildBones(BoneIndex, Childs) > 0)
+	TArray<int32> ChildBoneIndices;
+	if (Skeleton->GetChildBones(BoneIndex, ChildBoneIndices) > 0)
 	{
 		// first look for direct children
-		for (auto ChildIndex : Childs)
+		for (const int32 ChildBoneIndex : ChildBoneIndices)
 		{
-			FName ChildBoneName = RefSkeleton.GetBoneName(ChildIndex);
-
-			int32 ChildTrackIndex = DataModel->GetBoneTrackIndexByName(ChildBoneName);
-			if (ChildTrackIndex != INDEX_NONE)
+			const FName ChildBoneName = RefSkeleton.GetBoneName(ChildBoneIndex);
+			if (DataModel->IsValidBoneTrackName(ChildBoneName))
 			{
 				// found the new track
-				return ChildTrackIndex;
+				return ChildBoneName;
 			}
 		}
 
 		int32 BestGrandChildIndex = INDEX_NONE;
+		FName BestGrandChildName = NAME_None;
 		// if you didn't find yet, now you have to go through all children
-		for (auto ChildIndex : Childs)
+		for (const int32 ChildBoneIndex : ChildBoneIndices)
 		{
-			FName ChildBoneName = RefSkeleton.GetBoneName(ChildIndex);
+			const FName ChildBoneName = RefSkeleton.GetBoneName(ChildBoneIndex);
 			// now I have to go through all childrewn and find who is earliest since I don't know which one might be the closest one
-			int32 GrandChildIndex = FindFirstChildTrackIndex(InSequence, Skeleton, ChildBoneName);
-			if (GrandChildIndex != INDEX_NONE)
+			const FName GrandChildName = FindFirstChildTrackName(InSequence, Skeleton, ChildBoneName);
+			if (GrandChildName != NAME_None)
 			{
+				const int32 GrandChildIndex = Skeleton->GetReferenceSkeleton().FindBoneIndex(GrandChildName);
 				if (BestGrandChildIndex == INDEX_NONE)
 				{
 					BestGrandChildIndex = GrandChildIndex;
+					BestGrandChildName = GrandChildName;
 				}
 				else if (BestGrandChildIndex > GrandChildIndex)
 				{
 					// best should be earlier track index
 					BestGrandChildIndex = GrandChildIndex;
+					BestGrandChildName = GrandChildName;
 				}
 			}
 		}
 
-		return BestGrandChildIndex;
+		return BestGrandChildName;
 	}
 
-	// there is no child, just add at the end
-	return DataModel->GetNumBoneTracks();
+	return NAME_None;
 }
 
 #endif // WITH_EDITOR
