Animation Compression

diff --git a/Engine/Source/Runtime/Engine/Private/Animation/AnimStreamable.cpp b/Engine/Source/Runtime/Engine/Private/Animation/AnimStreamable.cpp
--- a/Engine/Source/Runtime/Engine/Private/Animation/AnimStreamable.cpp
+++ b/Engine/Source/Runtime/Engine/Private/Animation/AnimStreamable.cpp
@@ -6,13 +6,13 @@
 
 #include "Animation/AnimStreamable.h"
 
+#include "AnimCompressionDerivedData.h"
 #include "Interfaces/ITargetPlatform.h"
 #include "Interfaces/ITargetPlatformManagerModule.h"
 #include "DeviceProfiles/DeviceProfileManager.h"
 #include "DeviceProfiles/DeviceProfile.h"
 #include "UObject/LinkerLoad.h"
 #include "UObject/ObjectSaveContext.h"
-#include "Animation/AnimCompressionDerivedData.h"
 #include "DerivedDataCacheInterface.h"
 #include "Animation/AnimBoneCompressionSettings.h"
 #include "Animation/AnimCurveCompressionSettings.h"
@@ -317,17 +317,6 @@ void UAnimStreamable::GetAnimationPose(FAnimationPoseData& OutAnimationPoseData,
 		//Need to evaluate raw data
 		ValidateModel();
 
-		// Warning if we have invalid data
-		for (const FBoneAnimationTrack& AnimationTrack : DataModelInterface->GetBoneAnimationTracks())
-		{
-			const FRawAnimSequenceTrack& TrackToExtract = AnimationTrack.InternalTrackData;
-			// Bail out (with rather wacky data) if data is empty for some reason.
-			if (TrackToExtract.PosKeys.Num() == 0 || TrackToExtract.RotKeys.Num() == 0)
-			{
-				UE_LOG(LogAnimation, Warning, TEXT("UAnimSequence::GetBoneTransform : No anim data in AnimSequence '%s' Track '%s'"), *GetPathName(), *AnimationTrack.Name.ToString());
-			}
-		}
-
 		const UE::Anim::DataModel::FEvaluationContext EvaluationContext(ExtractionContext.CurrentTime, DataModelInterface->GetFrameRate(), RetargetSource, MySkeleton->GetRefLocalPoses(RetargetSource), Interpolation);
 		DataModelInterface->Evaluate(OutAnimationPoseData, EvaluationContext);
 
@@ -606,7 +595,9 @@ void UAnimStreamable::RequestCompressedData(const ITargetPlatform* Platform)
 	const bool bInAllowAlternateCompressor = false;
 	const bool bInOutput				   = false;
 
+	PRAGMA_DISABLE_DEPRECATION_WARNINGS
 	TSharedRef<FAnimCompressContext> CompressContext = MakeShared<FAnimCompressContext>(bInAllowAlternateCompressor, bInOutput);
+	PRAGMA_ENABLE_DEPRECATION_WARNINGS
 
 	for (uint32 ChunkIndex = 0; ChunkIndex < NumChunks; ++ChunkIndex)
 	{
@@ -624,7 +615,7 @@ void UAnimStreamable::RequestCompressedData(const ITargetPlatform* Platform)
 		IStreamingManager::Get().GetAnimationStreamingManager().AddStreamingAnim(this);
 	}
 	//PlatformData.SetSkeletonVirtualBoneGuid(GetSkeleton()->GetVirtualBoneGuid()); //MDW DO THIS
-	//PlatformData.bUseRawDataOnly = false; //MDW Need to do something with this? 
+	//PlatformData.bUseRawDataOnly = false; //MDW Need to do something with this?
 }
 
 float UAnimStreamable::GetChunkSizeSeconds(const ITargetPlatform* Platform) const
@@ -642,12 +633,12 @@ float UAnimStreamable::GetChunkSizeSeconds(const ITargetPlatform* Platform) cons
 	return CVarPlatformChunkSizeSeconds;
 }
 
-template<typename KeyType>
-void MakeKeyChunk(const TArray<KeyType>& SrcKeys, TArray<KeyType>& DestKeys, int32 NumberOfKeys, const uint32 FrameStart, const uint32 FrameEnd)
+template<typename KeyInType, typename KeyOutType, typename Predicate>
+void MakeKeyChunk(const TArray<KeyInType>& SrcKeys, TArray<KeyOutType>& DestKeys, int32 NumberOfKeys, const uint32 FrameStart, const uint32 FrameEnd, Predicate InPredicate)
 {
 	if (SrcKeys.Num() == 1)
 	{
-		DestKeys.Add(SrcKeys[0]);
+		DestKeys.Add(::Invoke(InPredicate, SrcKeys[0]));
 	}
 	else
 	{
@@ -656,12 +647,16 @@ void MakeKeyChunk(const TArray<KeyType>& SrcKeys, TArray<KeyType>& DestKeys, int
 		DestKeys.Reset((FrameEnd - FrameStart) + 1);
 		for (uint32 FrameIndex = FrameStart; FrameIndex <= FrameEnd; ++FrameIndex)
 		{
-			DestKeys.Add(SrcKeys[FrameIndex]);
+			DestKeys.Add(::Invoke(InPredicate, SrcKeys[FrameIndex]));
 		}
 	}
 }
 
-void UAnimStreamable::RequestCompressedDataForChunk(const FString& ChunkDDCKey, FAnimStreamableChunk& Chunk, const int32 ChunkIndex, const uint32 FrameStart, const uint32 FrameEnd, TSharedRef<FAnimCompressContext> CompressContext, const ITargetPlatform* Platform)
+void UAnimStreamable::RequestCompressedDataForChunk(const FString& ChunkDDCKey, FAnimStreamableChunk& Chunk, const int32 ChunkIndex, const uint32 FrameStart, const uint32 FrameEnd,
+	PRAGMA_DISABLE_DEPRECATION_WARNINGS
+	TSharedRef<FAnimCompressContext> CompressContext,
+	PRAGMA_ENABLE_DEPRECATION_WARNINGS
+	const ITargetPlatform* Platform)
 {
 	// Need to unify with Anim Sequence!
 
@@ -690,22 +685,40 @@ void UAnimStreamable::RequestCompressedDataForChunk(const FString& ChunkDDCKey,
 		{
 			FCompressibleAnimRef CompressibleData = MakeShared<FCompressibleAnimData, ESPMode::ThreadSafe>(BoneCompressionSettings, CurveCompressionSettings, GetSkeleton(), Interpolation, Chunk.SequenceLength, ChunkNumFrames+1, Platform);
 
-			const TArray<FBoneAnimationTrack>& BoneAnimationTracks = DataModelInterface->GetBoneAnimationTracks();
+			TArray<FName> TrackNames;
+			DataModelInterface->GetBoneTrackNames(TrackNames);
+			
 			CompressibleData->TrackToSkeletonMapTable.Empty();
 
-			for (const FBoneAnimationTrack& AnimTrack : BoneAnimationTracks)
+			TArray<FTransform> BoneTransforms;
+			for (const FName& TrackName : TrackNames)
 			{
-				const FRawAnimSequenceTrack& SrcTrack = AnimTrack.InternalTrackData;
-				FRawAnimSequenceTrack& DestTrack = CompressibleData->RawAnimationData.AddDefaulted_GetRef();
-
-				MakeKeyChunk(SrcTrack.PosKeys, DestTrack.PosKeys, NumberOfKeys, FrameStart, FrameEnd);
-				MakeKeyChunk(SrcTrack.RotKeys, DestTrack.RotKeys, NumberOfKeys, FrameStart, FrameEnd);
-				if (SrcTrack.ScaleKeys.Num() > 0)
+				const int32 BoneIndex = GetSkeleton()->GetReferenceSkeleton().FindBoneIndex(TrackName);
+				if (BoneIndex != INDEX_NONE)
 				{
-					MakeKeyChunk(SrcTrack.ScaleKeys, DestTrack.ScaleKeys, NumberOfKeys, FrameStart, FrameEnd);
+					FRawAnimSequenceTrack& DestTrack = CompressibleData->RawAnimationData.AddDefaulted_GetRef();
+					
+					BoneTransforms.Reset();
+					DataModelInterface->GetBoneTrackTransforms(TrackName, BoneTransforms);
+					
+
+					MakeKeyChunk(BoneTransforms, DestTrack.PosKeys, NumberOfKeys, FrameStart, FrameEnd, [](const FTransform& BoneTransform) -> FVector3f
+					{
+						return FVector3f(BoneTransform.GetLocation());
+					});
+					
+					MakeKeyChunk(BoneTransforms, DestTrack.RotKeys, NumberOfKeys, FrameStart, FrameEnd, [](const FTransform& BoneTransform) -> FQuat4f
+					{
+						return FQuat4f(BoneTransform.GetRotation());
+					});
+					
+					MakeKeyChunk(BoneTransforms, DestTrack.ScaleKeys, NumberOfKeys, FrameStart, FrameEnd, [](const FTransform& BoneTransform) -> FVector3f
+					{
+						return FVector3f(BoneTransform.GetScale3D());
+					});
+
+					CompressibleData->TrackToSkeletonMapTable.Add(BoneIndex);
 				}
-
-				CompressibleData->TrackToSkeletonMapTable.Add(AnimTrack.BoneTreeIndex);
 			}
 
 			if (FrameStart == 0)
diff --git a/Engine/Source/Runtime/Engine/Private/Animation/PoseAsset.cpp b/Engine/Source/Runtime/Engine/Private/Animation/PoseAsset.cpp
--- a/Engine/Source/Runtime/Engine/Private/Animation/PoseAsset.cpp
+++ b/Engine/Source/Runtime/Engine/Private/Animation/PoseAsset.cpp
@@ -1455,7 +1455,7 @@ void UPoseAsset::CreatePoseFromAnimation(class UAnimSequence* AnimSequence, cons
 					// now get rawanimationdata, and each key is converted to new pose
 					for (int32 TrackIndex = 0; TrackIndex < NumTracks; ++TrackIndex)
 					{
-						UE::Anim::GetBoneTransformFromModel(AnimSequence->GetDataModel(), NewPose[TrackIndex], TrackIndex, PoseIndex);
+						NewPose[TrackIndex] = AnimSequence->GetDataModel()->GetBoneTrackTransform(TrackNames[TrackIndex], FFrameNumber(PoseIndex));
 					}
 
 					if (TotalFloatCurveCount > 0)
diff --git a/Engine/Source/Runtime/Engine/Private/Animation/Skeleton.cpp b/Engine/Source/Runtime/Engine/Private/Animation/Skeleton.cpp
--- a/Engine/Source/Runtime/Engine/Private/Animation/Skeleton.cpp
+++ b/Engine/Source/Runtime/Engine/Private/Animation/Skeleton.cpp
@@ -1045,10 +1045,11 @@ int32 USkeleton::GetRawAnimationTrackIndex(const int32 InSkeletonBoneIndex, cons
 	if (InSkeletonBoneIndex != INDEX_NONE)
 	{
 #if WITH_EDITOR
-		return InAnimSeq->GetDataModel()->GetBoneAnimationTracks().IndexOfByPredicate([InSkeletonBoneIndex](const FBoneAnimationTrack& AnimationTrack)
-			{
-				return AnimationTrack.BoneTreeIndex == InSkeletonBoneIndex;
-			});
+		TArray<FName> BoneTrackNames;
+		InAnimSeq->GetDataModel()->GetBoneTrackNames(BoneTrackNames);
+		const FName BoneName = ReferenceSkeleton.GetBoneName(InSkeletonBoneIndex);
+		
+		return BoneTrackNames.IndexOfByKey(BoneName);
 #else
 		return InAnimSeq->GetCompressedTrackToSkeletonMapTable().IndexOfByPredicate([InSkeletonBoneIndex](const FTrackToSkeletonMap& Mapping)
 			{
diff --git a/Engine/Source/Runtime/Engine/Private/AssetCompilingManager.cpp b/Engine/Source/Runtime/Engine/Private/AssetCompilingManager.cpp
--- a/Engine/Source/Runtime/Engine/Private/AssetCompilingManager.cpp
+++ b/Engine/Source/Runtime/Engine/Private/AssetCompilingManager.cpp
@@ -26,6 +26,7 @@ LLM_DEFINE_TAG(AssetCompilation, NAME_None, NAME_None, GET_STATFNAME(STAT_AssetC
 #include "Algo/TopologicalSort.h"
 #include "Algo/Find.h"
 #include "ProfilingDebugging/CountersTrace.h"
+#include "Animation/AnimationSequenceCompiler.h"
 
 #if WITH_EDITOR
 #include "DerivedDataBuildSchedulerQueue.h"
@@ -313,6 +314,7 @@ FAssetCompilingManager::FAssetCompilingManager()
 	RegisterManager(&FTextureCompilingManager::Get());
 	RegisterManager(&FActorDeferredScriptManager::Get());
 	RegisterManager(&FSoundWaveCompilingManager::Get());
+	RegisterManager(&UE::Anim::FAnimSequenceCompilingManager::Get());
 
 	// Ensure that the thread pool is constructed before the memory queue is registered because
 	// the memory queue can call GetThreadPool() from a worker thread, and could lead to a race
diff --git a/Engine/Source/Runtime/Engine/Private/PerPlatformProperties.cpp b/Engine/Source/Runtime/Engine/Private/PerPlatformProperties.cpp
--- a/Engine/Source/Runtime/Engine/Private/PerPlatformProperties.cpp
+++ b/Engine/Source/Runtime/Engine/Private/PerPlatformProperties.cpp
@@ -89,6 +89,9 @@ template ENGINE_API void operator<<(FStructuredArchive::FSlot Slot, TPerPlatform
 template ENGINE_API void operator<<(FStructuredArchive::FSlot Slot, TPerPlatformProperty<FPerPlatformBool, bool, NAME_BoolProperty>&);
 template ENGINE_API void operator<<(FStructuredArchive::FSlot Slot, TPerPlatformProperty<FFreezablePerPlatformFloat, float, NAME_FloatProperty>&);
 
+template ENGINE_API FArchive& operator<<(FArchive&, TPerPlatformProperty<FPerPlatformFrameRate, FFrameRate, NAME_FrameRate>&);
+template ENGINE_API void operator<<(FStructuredArchive::FSlot Slot, TPerPlatformProperty<FPerPlatformFrameRate, FFrameRate, NAME_FrameRate>&);
+
 FString FPerPlatformInt::ToString() const
 {
 	FString Result = FString::FromInt(Default);
diff --git a/Engine/Source/Runtime/Engine/Public/Animation/AnimCompressionDerivedDataPublic.h b/Engine/Source/Runtime/Engine/Public/Animation/AnimCompressionDerivedDataPublic.h
--- a/Engine/Source/Runtime/Engine/Public/Animation/AnimCompressionDerivedDataPublic.h
+++ b/Engine/Source/Runtime/Engine/Public/Animation/AnimCompressionDerivedDataPublic.h
@@ -12,7 +12,7 @@ class UAnimSequence;
 class FDerivedDataAnimationCompression;
 struct FCompressibleAnimData;
 
-
+UE_DEPRECATED(5.2, "GAsyncCompressedAnimationsTracker has been deprecated")
 extern ENGINE_API class FAsyncCompressedAnimationsManagement* GAsyncCompressedAnimationsTracker;
 
 // Animation data that is currently being compressed
@@ -48,64 +48,32 @@ struct FQueuedAsyncCompressionWork
 	const bool bPerformFrameStripping;
 };
 
-
 // Manager for Async anim compression.
 //   Maintains active compressions
 //   tracks memory usage of async compression
 //   Gives API for blocking on compression
-class ENGINE_API FAsyncCompressedAnimationsManagement : public FTickableEditorObject, public FTickableCookObject, public FGCObject
+class UE_DEPRECATED(5.2, "FAsyncCompressedAnimationsManagement has been deprecated") FAsyncCompressedAnimationsManagement;
+class ENGINE_API FAsyncCompressedAnimationsManagement
 {
 public:
-	static FAsyncCompressedAnimationsManagement& Get();
+	static FAsyncCompressedAnimationsManagement& Get()
+	{
+		static FAsyncCompressedAnimationsManagement Management;
+		return Management;
+	}
 
 	// Request an async compression of an animation, may not actually run async if memory usage is already high
 	// Returns true if an async compression task was allowed.
-	bool RequestAsyncCompression(FDerivedDataAnimationCompression& Compressor, UAnimSequence* Anim, const bool bPerformFrameStripping, TArray<uint8>& OutData);
+	UE_DEPRECATED(5.2, "RequestAsyncCompression has been deprecated")
+	bool RequestAsyncCompression(FDerivedDataAnimationCompression& Compressor, UAnimSequence* Anim, const bool bPerformFrameStripping, TArray<uint8>& OutData) { return false; }
 
 	// Returns the number of remaining compression jobs (used for UI)
-	int32 GetNumRemainingJobs() const { return ActiveAsyncCompressionTasks.Num() + QueuedAsyncCompressionWork.Num(); }
+	UE_DEPRECATED(5.2, "GetNumRemainingJobs has been deprecated")
+	int32 GetNumRemainingJobs() const { return 0; }
 	
-	// Blocks on compression for the supplied animation. Returns false if there was no compression job to wait on. 
-	bool WaitOnExistingCompression(UAnimSequence* Anim, const bool bWantResults);
-private:
-
-	void StartAsyncWork(FDerivedDataAnimationCompression& Compressor, UAnimSequence* Anim, const uint64 NewTaskSize, const bool bPerformFrameStripping);
-
-	void StartQueuedTasks(int32 MaxActiveTasks);
-
-	bool WaitOnActiveCompression(UAnimSequence* Anim, bool bWantResults);
-
-	void OnActiveCompressionFinished(int32 ActiveAnimIndex);
-
-	FAsyncCompressedAnimationsManagement();
-
-	void Shutdown();
-
-	//Array of active and queued compression jobs
-	TArray<FActiveAsyncCompressionTask> ActiveAsyncCompressionTasks;
-	TArray<FQueuedAsyncCompressionWork> QueuedAsyncCompressionWork;
-	uint64 ActiveMemoryUsage;
-
-	//Cache the last tick time to keep queue pumping during large loading segments
-	double LastTickTime;
-
-	/* Begin FTickableEditorObject */
-	virtual void Tick(float DeltaTime) override;
-	virtual ETickableTickType GetTickableTickType() const override { return ETickableTickType::Always; }
-	virtual TStatId GetStatId() const override;
-	/* End FTickableEditorObject */
-
-	/* Begin FTickableCookObject */
-	virtual void TickCook(float DeltaTime, bool bCookCompete) override;
-	/* End FTickableCookObject */
-
-	// FGCObject interface
-	void AddReferencedObjects(FReferenceCollector& Collector) override;
-	virtual FString GetReferencerName() const override
-	{
-		return TEXT("FAsyncCompressedAnimationsManagement");
-	}
-	// End of FGCObject interface
+	// Blocks on compression for the supplied animation. Returns false if there was no compression job to wait on.
+	UE_DEPRECATED(5.2, "WaitOnExistingCompression has been deprecated") 
+	bool WaitOnExistingCompression(UAnimSequence* Anim, const bool bWantResults) { return false; }
 };
 
 #endif
diff --git a/Engine/Source/Runtime/Engine/Public/Animation/AnimCompressionTypes.h b/Engine/Source/Runtime/Engine/Public/Animation/AnimCompressionTypes.h
--- a/Engine/Source/Runtime/Engine/Public/Animation/AnimCompressionTypes.h
+++ b/Engine/Source/Runtime/Engine/Public/Animation/AnimCompressionTypes.h
@@ -48,6 +48,7 @@ class UAnimBoneCompressionSettings;
 class UAnimBoneCompressionCodec;
 class USkeleton;
 struct FCompactPose;
+struct FBoneAnimationTrack;
 
 template<typename ArrayClass>
 struct ENGINE_API FCompressedOffsetDataBase
@@ -190,6 +191,7 @@ public:
 	}
 };
 
+#if WITH_EDITOR
 struct ENGINE_API FCompressibleAnimData
 {
 public:
@@ -222,7 +224,7 @@ public:
 
 	TArray<FFloatCurve> RawFloatCurves;
 
-	float SequenceLength;
+	double SequenceLength;
 
 	/** Number of keys within the (non-uniform) RawAnimationData tracks */
 	
@@ -245,7 +247,11 @@ public:
 
 	FCancelCompressionSignal IsCancelledSignal;
 
-	const ITargetPlatform* TargetPlatform;
+	bool bShouldPerformStripping = false;
+	TWeakObjectPtr<UAnimSequence> WeakSequence;
+	bool bDataFetched = false;
+
+	const ITargetPlatform* TargetPlatform = nullptr;
 
 	static int32 GetApproxRawDataArraySize(const TArray<FRawAnimSequenceTrack>& AnimData)
 	{
@@ -292,6 +298,8 @@ public:
 		return MemUsage;
 	}
 
+	void FetchData(const ITargetPlatform* InPlatform);
+	
 	void Update(struct FCompressedAnimSequence& CompressedData) const;
 
 	void AddReferencedObjects(FReferenceCollector& Collector)
@@ -305,6 +313,10 @@ public:
 		return IsCancelledSignal.IsCancelled();
 	}
 
+protected:
+	void BakeOutAdditiveIntoRawData(const FFrameRate& SampleRate, TArray<FBoneAnimationTrack>& ResampledTrackData, TArray<FFloatCurve>& FloatCurves);
+	void ResampleAnimationTrackData(const FFrameRate& SampleRate, TArray<FBoneAnimationTrack>& ResampledTrackData) const;
+
 private:
 	void WriteCompressionDataToJSON(TArrayView<FName> OriginalTrackNames, TArrayView<FRawAnimSequenceTrack> FinalRawAnimationData, TArrayView<FName> FinalTrackNames) const;
 };
@@ -312,6 +324,8 @@ private:
 typedef TSharedPtr<FCompressibleAnimData, ESPMode::ThreadSafe> FCompressibleAnimPtr;
 typedef TSharedRef<FCompressibleAnimData, ESPMode::ThreadSafe> FCompressibleAnimRef;
 
+#endif // WITH_EDITOR
+
 // Wrapper Code
 template <typename T>
 struct TArrayMaker
@@ -763,11 +777,19 @@ public:
 		return CompressedTrackToSkeletonMapTable[TrackIndex].BoneTreeIndex;
 	}
 
+	int32 GetTrackIndexFromSkeletonIndex(const int32 BoneIndex) const
+	{
+		return CompressedTrackToSkeletonMapTable.IndexOfByPredicate([BoneIndex](const FTrackToSkeletonMap& Entry) { return Entry.BoneTreeIndex == BoneIndex; });
+	}
+
 	// Return the number of bytes used
 	SIZE_T GetMemorySize() const;
 
+	void Reset();
 	void ClearCompressedBoneData();
 	void ClearCompressedCurveData();
+
+	bool IsValid(const UAnimSequence* AnimSequence) const;
 };
 
 struct FRootMotionReset
@@ -846,3 +868,8 @@ extern void DecompressPose(	FCompactPose& OutPose,
 							FName RetargetSource,
 							FName SourceName,
 							const FRootMotionReset& RootMotionReset);
+
+
+#if WITH_EDITOR
+extern FGuid GenerateGuidFromRawAnimData(const TArray<FRawAnimSequenceTrack>& RawAnimationData, const FRawCurveTracks& RawCurveData);
+#endif 
\ No newline at end of file
diff --git a/Engine/Source/Runtime/Engine/Public/AnimationCompression.h b/Engine/Source/Runtime/Engine/Public/AnimationCompression.h
--- a/Engine/Source/Runtime/Engine/Public/AnimationCompression.h
+++ b/Engine/Source/Runtime/Engine/Public/AnimationCompression.h
@@ -11,7 +11,7 @@
 #include "Animation/AnimEnums.h"
 #include "EngineLogs.h"
 
-ENGINE_API DECLARE_LOG_CATEGORY_EXTERN(LogAnimationCompression, Log, All);
+ENGINE_API DECLARE_LOG_CATEGORY_EXTERN(LogAnimationCompression, Display, All);
 
 // Thresholds
 #define TRANSLATION_ZEROING_THRESHOLD (0.0001f)
diff --git a/Engine/Source/Runtime/Engine/Public/AnimationUtils.h b/Engine/Source/Runtime/Engine/Public/AnimationUtils.h
--- a/Engine/Source/Runtime/Engine/Public/AnimationUtils.h
+++ b/Engine/Source/Runtime/Engine/Public/AnimationUtils.h
@@ -120,7 +120,9 @@ public:
 	 * @param TrackHeights [OUT]	The computed track heights
 	 *
 	 */
+#if WITH_EDITOR
 	static void CalculateTrackHeights(const FCompressibleAnimData& CompressibleAnimData, int NumTracks, TArray<int32>& TrackHeights);
+#endif // WITH_EDITOR
 
 	/**
 	 * Checks a set of key times to see if the spacing is uniform or non-uniform.
@@ -137,6 +139,8 @@ public:
 	/**
 	 * Perturbs the bone(s) associated with each track in turn, measuring the maximum error introduced in end effectors as a result
 	 */
+	
+#if WITH_EDITOR
 	static void TallyErrorsFromPerturbation(
 		const FCompressibleAnimData& CompressibleAnimData,
 		int32 NumTracks,
@@ -145,7 +149,6 @@ public:
 		const FVector& ScaleNudge,
 		TArray<FAnimPerturbationError>& InducedErrors);
 
-#if WITH_EDITOR
 	/** Preload the bone, curve and recorder compression settings */
 	ENGINE_API static void PreloadCompressionSettings();
 #endif
diff --git a/Engine/Source/Runtime/MovieSceneTracks/Private/Sections/MovieSceneSkeletalAnimationSection.cpp b/Engine/Source/Runtime/MovieSceneTracks/Private/Sections/MovieSceneSkeletalAnimationSection.cpp
--- a/Engine/Source/Runtime/MovieSceneTracks/Private/Sections/MovieSceneSkeletalAnimationSection.cpp
+++ b/Engine/Source/Runtime/MovieSceneTracks/Private/Sections/MovieSceneSkeletalAnimationSection.cpp
@@ -11,6 +11,7 @@
 #include "Tracks/MovieSceneSkeletalAnimationTrack.h"
 #include "BoneContainer.h"
 #include "Animation/AnimationPoseData.h"
+#include "Animation/AnimSequenceDecompressionContext.h"
 #include "Animation/AttributesRuntime.h"
 #include "Misc/FrameRate.h"
 #include "EntitySystem/BuiltInComponentTypes.h"
@@ -533,35 +534,35 @@ int32 UMovieSceneSkeletalAnimationSection::SetBoneIndexForRootMotionCalculations
 			//but if not first find first
 			int32 RootIndex = INDEX_NONE;
 #if WITH_EDITOR
+			TArray<FName> TrackNames;
+			AnimSequence->GetDataModelInterface()->GetBoneTrackNames(TrackNames);
+			const FReferenceSkeleton& RefSkeleton = AnimSequence->GetSkeleton()->GetReferenceSkeleton();
 			const IAnimationDataModel* DataModel = AnimSequence->GetDataModel();
-			const TArray<FBoneAnimationTrack>& BoneAnimationTracks = DataModel->GetBoneAnimationTracks();
-			const int32 NumTracks = BoneAnimationTracks.Num();
-			for (int32 TrackIndex = 0; TrackIndex < NumTracks; ++TrackIndex)
+			TArray<FTransform> OutTransforms;
+			for (int32 BoneIndex = 0; BoneIndex < RefSkeleton.GetNum(); ++BoneIndex)
 			{
-				const FBoneAnimationTrack& AnimationTrack = BoneAnimationTracks[TrackIndex];
-				// verify if this bone exists in skeleton
-				const int32 BoneTreeIndex = AnimationTrack.BoneTreeIndex;
-				if (BoneTreeIndex != INDEX_NONE)
+				const FName BoneName = RefSkeleton.GetBoneName(BoneIndex);
+				if (DataModel->IsValidBoneTrackName(BoneName))
 				{
-	
-						const FRawAnimSequenceTrack& InternalTrackData = AnimationTrack.InternalTrackData;
-						const TArray<FVector3f>& TrackPosKeys = InternalTrackData.PosKeys;
-						for (const FVector3f& Vector : TrackPosKeys)
-						{
-							if (Vector.IsNearlyZero() == false)
-							{
-								TempRootBoneIndex = BoneTreeIndex;
-								break;
-							}
-						}
-						if (TempRootBoneIndex.IsSet())
+					DataModel->GetBoneTrackTransforms(BoneName, OutTransforms);
+					
+					for (const FTransform& Transform : OutTransforms)
+					{
+						if (Transform.GetScale3D().IsNearlyZero() == false)
 						{
+							TempRootBoneIndex = BoneIndex;
 							break;
-						}				
-
+						}
+					}
+					
+					OutTransforms.Reset();
+					if (TempRootBoneIndex.IsSet())
+					{
+						break;
+					}
+					
 				}
 			}
-
 #else
 			const TArray<FTrackToSkeletonMap>& BoneMappings = AnimSequence->GetCompressedTrackToSkeletonMapTable();
 			const int32 NumTracks = BoneMappings.Num();
@@ -583,12 +584,8 @@ int32 UMovieSceneSkeletalAnimationSection::SetBoneIndexForRootMotionCalculations
 						const int32 NumFrames = AnimSequence->GetNumberOfSampledKeys();
 						for (int32 Index = 0; Index < NumFrames; ++Index)
 						{
-							float Pos = FMath::Clamp((float)AnimSequence->GetSamplingFrameRate().AsSeconds(Index), 0.f, AnimSequence->GetPlayLength());
+							const double Pos = FMath::Clamp(AnimSequence->GetSamplingFrameRate().AsSeconds(Index), 0.0, AnimSequence->GetPlayLength());
 PRAGMA_DISABLE_DEPRECATION_WARNINGS
-							//jurre this is what I am talking  about
-							// can't create tkhis since GetRetargetTransformsSourceName isn't  public.
-							//FAnimSequenceDecompressionContext DecompContext(AnimSequence->GetSamplingFrameRate(), AnimSequence->GetSamplingFrameRate().AsFrameNumber(AnimSequence->GetPlayLength()).Value, AnimSequence->Interpolation, AnimSequence->GetRetargetTransformsSourceName(), *AnimSequence->CompressedData.CompressedDataStructure, AnimSequence->GetSkeleton()->GetRefLocalPoses(), AnimSequence->CompressedData.CompressedTrackToSkeletonMapTable, AnimSequence->GetSkeleton(), AnimSequence->IsValidAdditive());
-
 							AnimSequence->GetBoneTransform(Transform, TrackIndex, Pos, false);
 PRAGMA_ENABLE_DEPRECATION_WARNINGS
 							if (Transform.Equals(FTransform::Identity) == false)
diff --git a/Engine/Source/Runtime/MovieSceneTracks/Private/Tracks/MovieSceneSkeletalAnimationTrack.cpp b/Engine/Source/Runtime/MovieSceneTracks/Private/Tracks/MovieSceneSkeletalAnimationTrack.cpp
--- a/Engine/Source/Runtime/MovieSceneTracks/Private/Tracks/MovieSceneSkeletalAnimationTrack.cpp
+++ b/Engine/Source/Runtime/MovieSceneTracks/Private/Tracks/MovieSceneSkeletalAnimationTrack.cpp
@@ -641,31 +641,9 @@ void UMovieSceneSkeletalAnimationTrack::SetUpRootMotions(bool bForce)
 	}
 }
 
-
-static int32 GetAnimationTrackIndex(const int32 BoneIndex, const UAnimSequence* AnimSeq)
-{
-	if (BoneIndex == INDEX_NONE)
-	{
-		return INDEX_NONE;
-	}
-
-	const TArray<FTrackToSkeletonMap>& TrackToSkelMap = AnimSeq->CompressedData.CompressedTrackToSkeletonMapTable;
-	for (int32 TrackIndex = 0; TrackIndex < TrackToSkelMap.Num(); ++TrackIndex)
-	{
-		const FTrackToSkeletonMap& TrackToSkeleton = TrackToSkelMap[TrackIndex];
-		if (TrackToSkeleton.BoneTreeIndex == BoneIndex)
-		{
-			return TrackIndex;
-		}
-	}
-
-	return INDEX_NONE;
-}
-
 static FTransform GetTransformForBoneRelativeToIndex(UAnimSequence* AnimSequence, USkeletalMeshComponent* MeshComponent, const FName& InBoneName,
 	const FCompactPoseBoneIndex& ParentCPIndex, double Seconds)
 {
-
 	FTransform  WorldTransform = FTransform::Identity;
 	//AnimSequence->GetBoneTransform doesn't seem to be as accurate as GetAnimationPose
 	FMemMark Mark(FMemStack::Get());
