Remove explicit blend mode usage from FDebugViewModeInterface.

diff --git a/Engine/Source/Runtime/Engine/Private/DebugViewModeInterface.cpp b/Engine/Source/Runtime/Engine/Private/DebugViewModeInterface.cpp
--- a/Engine/Source/Runtime/Engine/Private/DebugViewModeInterface.cpp
+++ b/Engine/Source/Runtime/Engine/Private/DebugViewModeInterface.cpp
@@ -13,15 +13,15 @@ DebugViewModeInterface.cpp: Contains definitions for rendering debug viewmodes.
 FDebugViewModeInterface* FDebugViewModeInterface::Singleton = nullptr;
 
 // STRATA_TODO_BLENDMODE
-void FDebugViewModeInterface::SetDrawRenderState(EDebugViewShaderMode DebugViewMode, EBlendMode BlendMode, FRenderState& DrawRenderState, bool bHasDepthPrepassForMaskedMaterial) const
+void FDebugViewModeInterface::SetDrawRenderState(EDebugViewShaderMode DebugViewMode, const FMaterial& InMaterial, FRenderState& DrawRenderState, bool bHasDepthPrepassForMaskedMaterial) const
 {
 	if (DebugViewMode == DVSM_QuadComplexity || DebugViewMode == DVSM_ShaderComplexityBleedingQuadOverhead || DebugViewMode == DVSM_ShaderComplexityContainedQuadOverhead || DebugViewMode == DVSM_ShaderComplexity)
 	{
-		if (BlendMode == BLEND_Opaque)
+		if (IsOpaqueBlendMode(InMaterial))
 		{
 			DrawRenderState.DepthStencilState = TStaticDepthStencilState<true, CF_DepthNearOrEqual>::GetRHI();
 		}
-		else if (BlendMode == BLEND_Masked)
+		else if (IsMaskedBlendMode(InMaterial))
 		{
 			if (bHasDepthPrepassForMaskedMaterial)
 			{
@@ -45,7 +45,7 @@ void FDebugViewModeInterface::SetDrawRenderState(EDebugViewShaderMode DebugViewM
 	}
 	else
 	{
-		if (IsTranslucentBlendMode(BlendMode)) // STRATA_TODO_BLENDMODE
+		if (IsTranslucentBlendMode(InMaterial))
 		{
 			// Otherwise, force translucent blend mode (shaders will use an hardcoded alpha).
 			DrawRenderState.BlendState = TStaticBlendState<CW_RGBA, BO_Add, BF_SourceAlpha, BF_InverseSourceAlpha, BO_Add, BF_Zero, BF_InverseSourceAlpha>::GetRHI();
@@ -56,7 +56,7 @@ void FDebugViewModeInterface::SetDrawRenderState(EDebugViewShaderMode DebugViewM
 			DrawRenderState.BlendState = TStaticBlendState<>::GetRHI();
 
 			// If not selected, use depth equal to make alpha test stand out (goes with EarlyZPassMode = DDM_AllOpaque) 
-			if (BlendMode == BLEND_Masked && bHasDepthPrepassForMaskedMaterial) // STRATA_TODO_BLENDMODE
+			if (IsMaskedBlendMode(InMaterial) && bHasDepthPrepassForMaskedMaterial) // STRATA_TODO_BLENDMODE
 			{
 				DrawRenderState.DepthStencilState = TStaticDepthStencilState<false, CF_Equal>::GetRHI();
 			}
diff --git a/Engine/Source/Runtime/Engine/Public/DebugViewModeInterface.h b/Engine/Source/Runtime/Engine/Public/DebugViewModeInterface.h
--- a/Engine/Source/Runtime/Engine/Public/DebugViewModeInterface.h
+++ b/Engine/Source/Runtime/Engine/Public/DebugViewModeInterface.h
@@ -48,7 +48,7 @@ public:
 		const FVertexFactoryType* InVertexFactoryType,
 		FMaterialShaderTypes& OutShaderTypes) const = 0;
 
-	virtual void SetDrawRenderState(EDebugViewShaderMode DebugViewMode, EBlendMode BlendMode, FRenderState& DrawRenderState, bool bHasDepthPrepassForMaskedMaterial) const;
+	virtual void SetDrawRenderState(EDebugViewShaderMode DebugViewMode, const FMaterial& InMaterial, FRenderState& DrawRenderState, bool bHasDepthPrepassForMaskedMaterial) const;
 
 	virtual void GetDebugViewModeShaderBindings(
 		const FDebugViewModePS& Shader,
diff --git a/Engine/Source/Runtime/Renderer/Private/DebugViewModeRendering.cpp b/Engine/Source/Runtime/Renderer/Private/DebugViewModeRendering.cpp
--- a/Engine/Source/Runtime/Renderer/Private/DebugViewModeRendering.cpp
+++ b/Engine/Source/Runtime/Renderer/Private/DebugViewModeRendering.cpp
@@ -389,7 +389,7 @@ void FDebugViewModeMeshProcessor::AddMeshBatch(const FMeshBatch& RESTRICT MeshBa
 	FMeshPassProcessorRenderState DrawRenderState;
 
 	FDebugViewModeInterface::FRenderState InterfaceRenderState;
-	DebugViewModeInterface->SetDrawRenderState(DebugViewMode, BatchMaterial->GetBlendMode(), InterfaceRenderState, Scene ? (Scene->GetShadingPath() == EShadingPath::Deferred && Scene->EarlyZPassMode != DDM_NonMaskedOnly) : false);
+	DebugViewModeInterface->SetDrawRenderState(DebugViewMode, *BatchMaterial, InterfaceRenderState, Scene ? (Scene->GetShadingPath() == EShadingPath::Deferred && Scene->EarlyZPassMode != DDM_NonMaskedOnly) : false);
 	DrawRenderState.SetBlendState(InterfaceRenderState.BlendState);
 	DrawRenderState.SetDepthStencilState(InterfaceRenderState.DepthStencilState);
 
