Silence PVS warning

#rnx

diff --git a/Engine/Source/Runtime/Engine/Private/BlueprintGeneratedClass.cpp b/Engine/Source/Runtime/Engine/Private/BlueprintGeneratedClass.cpp
--- a/Engine/Source/Runtime/Engine/Private/BlueprintGeneratedClass.cpp
+++ b/Engine/Source/Runtime/Engine/Private/BlueprintGeneratedClass.cpp
@@ -1838,7 +1838,7 @@ void UBlueprintGeneratedClass::Bind()
 
 #define UE_CHECK_BLUEPRINT_REFERENCES !(UE_BUILD_TEST || UE_BUILD_SHIPPING)
 
-class FPersistentFrameCollector final : public FReferenceCollector
+class FPersistentFrameCollector final : public FReferenceCollector //-V1052
 {
 	FReferenceCollector& InnerCollector;
 #if UE_CHECK_BLUEPRINT_REFERENCES
@@ -1939,7 +1939,7 @@ public:
 		}
 	}
 
-	virtual void HandleObjectReferences(UObject** Objects, int32 Num, const UObject* ReferencingObject, const FProperty* ReferencingProperty)
+	virtual void HandleObjectReferences(UObject** Objects, int32 Num, const UObject* ReferencingObject, const FProperty* ReferencingProperty) override
 	{
 		for (int32 Idx = 0; Idx < Num; ++Idx)
 		{
