DataLayer:
- Only fixup Actor DataLayer assets on cook so that data isn't invalidated on map load in editor
- Re-resolve containers when WorldDataLayers as DataLayerInstances added/removed/modified

diff --git a/Engine/Source/Runtime/Engine/Private/ActorEditor.cpp b/Engine/Source/Runtime/Engine/Private/ActorEditor.cpp
--- a/Engine/Source/Runtime/Engine/Private/ActorEditor.cpp
+++ b/Engine/Source/Runtime/Engine/Private/ActorEditor.cpp
@@ -1720,7 +1720,11 @@ void AActor::FixupDataLayers(bool bRevertChangesOnLockedDataLayer /*= false*/)
 					}
 				};
 
-				CleanupDataLayers(DataLayerAssets);
+				// Only invalidate DataLayerAssets on cook. In Editor we want to be able to re-resolve if the asset gets readded to the WorldDataLayers actor
+				if (IsRunningCookCommandlet())
+				{
+					CleanupDataLayers(DataLayerAssets);
+				}
 				CleanupDataLayers(DataLayers);
 				PRAGMA_ENABLE_DEPRECATION_WARNINGS
 			}
diff --git a/Engine/Source/Runtime/Engine/Private/WorldPartition/DataLayer/DataLayerInstanceWithAsset.cpp b/Engine/Source/Runtime/Engine/Private/WorldPartition/DataLayer/DataLayerInstanceWithAsset.cpp
--- a/Engine/Source/Runtime/Engine/Private/WorldPartition/DataLayer/DataLayerInstanceWithAsset.cpp
+++ b/Engine/Source/Runtime/Engine/Private/WorldPartition/DataLayer/DataLayerInstanceWithAsset.cpp
@@ -91,4 +91,32 @@ bool UDataLayerInstanceWithAsset::Validate(IStreamingGenerationErrorHandler* Err
 
 	return bIsValid;
 }
+
+void UDataLayerInstanceWithAsset::PostEditChangeProperty(FPropertyChangedEvent& PropertyChangedEvent)
+{
+	Super::PostEditChangeProperty(PropertyChangedEvent);
+
+	const FName PropertyName = (PropertyChangedEvent.Property != nullptr) ? PropertyChangedEvent.Property->GetFName() : NAME_None;
+
+	if (PropertyName == GET_MEMBER_NAME_CHECKED(UDataLayerInstanceWithAsset, DataLayerAsset))
+	{
+		GetOuterAWorldDataLayers()->ResolveActorDescContainers();
+	}
+}
+
+void UDataLayerInstanceWithAsset::PreEditUndo()
+{
+	Super::PreEditUndo();
+	CachedDataLayerAsset = DataLayerAsset;
+}
+
+void UDataLayerInstanceWithAsset::PostEditUndo()
+{
+	Super::PostEditUndo();
+	if (CachedDataLayerAsset != DataLayerAsset)
+	{
+		GetOuterAWorldDataLayers()->ResolveActorDescContainers();
+	}
+	CachedDataLayerAsset = nullptr;
+}
 #endif
diff --git a/Engine/Source/Runtime/Engine/Private/WorldPartition/DataLayer/DataLayerSubsystem.cpp b/Engine/Source/Runtime/Engine/Private/WorldPartition/DataLayer/DataLayerSubsystem.cpp
--- a/Engine/Source/Runtime/Engine/Private/WorldPartition/DataLayer/DataLayerSubsystem.cpp
+++ b/Engine/Source/Runtime/Engine/Private/WorldPartition/DataLayer/DataLayerSubsystem.cpp
@@ -763,6 +763,11 @@ uint32 UDataLayerSubsystem::GetDataLayerEditorContextHash() const
 }
 
 void UDataLayerSubsystem::OnActorDescContainerInitialized(UActorDescContainer* InActorDescContainer)
+{
+	ResolveActorDescContainer(InActorDescContainer);
+}
+
+void UDataLayerSubsystem::ResolveActorDescContainer(UActorDescContainer* InActorDescContainer)
 {
 	check(InActorDescContainer);
 
@@ -793,6 +798,17 @@ void UDataLayerSubsystem::OnActorDescContainerInitialized(UActorDescContainer* I
 	}
 }
 
+void UDataLayerSubsystem::ResolveActorDescContainers()
+{
+	for (TObjectIterator<UActorDescContainer> ContainerIt; ContainerIt; ++ContainerIt)
+	{
+		if (UActorDescContainer* ActorDescContainer = *ContainerIt; ActorDescContainer && GetWorld() == ActorDescContainer->GetWorld())
+		{
+			ResolveActorDescContainer(ActorDescContainer);
+		}
+	}
+}
+
 void UDataLayerSubsystem::RegisterEarlyLoadedWorldDataLayers()
 {
 	// ULevel::WorldDatalayers is loaded before the subsystem is created and miss their registration.
diff --git a/Engine/Source/Runtime/Engine/Private/WorldPartition/DataLayer/WorldDataLayers.cpp b/Engine/Source/Runtime/Engine/Private/WorldPartition/DataLayer/WorldDataLayers.cpp
--- a/Engine/Source/Runtime/Engine/Private/WorldPartition/DataLayer/WorldDataLayers.cpp
+++ b/Engine/Source/Runtime/Engine/Private/WorldPartition/DataLayer/WorldDataLayers.cpp
@@ -528,6 +528,8 @@ int32 AWorldDataLayers::RemoveDataLayers(const TArray<UDataLayerInstance*>& InDa
 	if (RemovedCount > 0)
 	{
 		UpdateContainsDeprecatedDataLayers();
+
+		ResolveActorDescContainers();
 	}
 
 	return RemovedCount;
@@ -546,6 +548,8 @@ bool AWorldDataLayers::RemoveDataLayer(const UDataLayerInstance* InDataLayerInst
 			UpdateContainsDeprecatedDataLayers();
 		}
 
+		ResolveActorDescContainers();
+		
 		return true;
 	}
 	return false;
@@ -935,6 +939,43 @@ void AWorldDataLayers::UpdateContainsDeprecatedDataLayers()
 	}
 }
 
+void AWorldDataLayers::ResolveActorDescContainers()
+{
+	if (UDataLayerSubsystem* DataLayerSubsystem = GetWorld()->GetSubsystem<UDataLayerSubsystem>())
+	{
+		DataLayerSubsystem->ResolveActorDescContainers();
+	}
+}
+
+void AWorldDataLayers::PreEditUndo()
+{
+	Super::PreEditUndo();
+	CachedDataLayerInstances = DataLayerInstances;
+}
+
+void AWorldDataLayers::PostEditUndo()
+{
+	Super::PostEditUndo();
+
+	bool bNeedResolve = CachedDataLayerInstances.Num() != DataLayerInstances.Num();
+	if (!bNeedResolve)
+	{
+		for (UDataLayerInstance* DataLayerInstance : DataLayerInstances)
+		{
+			if (!CachedDataLayerInstances.Contains(DataLayerInstance))
+			{
+				bNeedResolve = true;
+				break;
+			}
+		}
+	}
+	if (bNeedResolve)
+	{
+		ResolveActorDescContainers();
+	}
+	CachedDataLayerInstances.Empty();
+}
+
 #endif
 
 //~ Begin Deprecated
diff --git a/Engine/Source/Runtime/Engine/Public/WorldPartition/DataLayer/DataLayerInstanceWithAsset.h b/Engine/Source/Runtime/Engine/Public/WorldPartition/DataLayer/DataLayerInstanceWithAsset.h
--- a/Engine/Source/Runtime/Engine/Public/WorldPartition/DataLayer/DataLayerInstanceWithAsset.h
+++ b/Engine/Source/Runtime/Engine/Public/WorldPartition/DataLayer/DataLayerInstanceWithAsset.h
@@ -18,12 +18,15 @@ public:
 	static FName MakeName(const UDataLayerAsset* DeprecatedDataLayer);
 	void OnCreated(const UDataLayerAsset* Asset);
 
+	virtual void PreEditUndo() override;
+	virtual void PostEditUndo() override;
 	virtual bool IsLocked() const override;
 	virtual bool IsReadOnly() const override;
 	virtual bool AddActor(AActor* Actor) const override;
 	virtual bool RemoveActor(AActor* Actor) const override;
 
 	virtual bool Validate(IStreamingGenerationErrorHandler* ErrorHandler) const override;
+	virtual void PostEditChangeProperty(FPropertyChangedEvent& PropertyChangedEvent) override;
 #endif
 
 	const UDataLayerAsset* GetAsset() const { return DataLayerAsset; }
@@ -40,4 +43,9 @@ public:
 private:
 	UPROPERTY(Category = "Data Layer", EditAnywhere)
 	TObjectPtr<const UDataLayerAsset> DataLayerAsset;
+
+#if WITH_EDITOR
+	// Used to compare state pre/post undo
+	TObjectPtr<const UDataLayerAsset> CachedDataLayerAsset;
+#endif
 };
\ No newline at end of file
diff --git a/Engine/Source/Runtime/Engine/Public/WorldPartition/DataLayer/DataLayerSubsystem.h b/Engine/Source/Runtime/Engine/Public/WorldPartition/DataLayer/DataLayerSubsystem.h
--- a/Engine/Source/Runtime/Engine/Public/WorldPartition/DataLayer/DataLayerSubsystem.h
+++ b/Engine/Source/Runtime/Engine/Public/WorldPartition/DataLayer/DataLayerSubsystem.h
@@ -230,6 +230,8 @@ private:
 #if WITH_EDITOR
 	void RegisterEarlyLoadedWorldDataLayers();
 	void OnActorDescContainerInitialized(UActorDescContainer* InActorDescContainer);
+	void ResolveActorDescContainer(UActorDescContainer* InActorDescContainer);
+	void ResolveActorDescContainers();
 
 	mutable int32 DataLayerActorEditorContextID;
 #endif
diff --git a/Engine/Source/Runtime/Engine/Public/WorldPartition/DataLayer/WorldDataLayers.h b/Engine/Source/Runtime/Engine/Public/WorldPartition/DataLayer/WorldDataLayers.h
--- a/Engine/Source/Runtime/Engine/Public/WorldPartition/DataLayer/WorldDataLayers.h
+++ b/Engine/Source/Runtime/Engine/Public/WorldPartition/DataLayer/WorldDataLayers.h
@@ -51,6 +51,8 @@ public:
 	virtual void BeginPlay() override;
 
 #if WITH_EDITOR
+	virtual void PreEditUndo() override;
+	virtual void PostEditUndo() override;
 	virtual bool ShouldLevelKeepRefIfExternal() const override { return true; }
 	virtual bool ShouldImport(FStringView ActorPropString, bool IsMovingLevel) override { return false; }
 	virtual bool IsLockLocation() const { return true; }
@@ -185,6 +187,11 @@ private:
 #if WITH_EDITOR
 	void ConvertDataLayerToInstancces();
 	void UpdateContainsDeprecatedDataLayers();
+	void ResolveActorDescContainers();
+	friend class UDataLayerInstanceWithAsset;
+
+	// Used to compare state pre/post undo
+	TSet<TObjectPtr<UDataLayerInstance>> CachedDataLayerInstances;
 #endif
 
 #if !WITH_EDITOR
@@ -262,7 +269,7 @@ DataLayerInstanceType* AWorldDataLayers::CreateDataLayer(CreationsArgs... InCrea
 	DataLayerInstances.Add(NewDataLayer);
 
 	UpdateContainsDeprecatedDataLayers();
-
+	ResolveActorDescContainers();
 	return NewDataLayer;
 }
 
