Deprecate use of TUniquePtr in Audio::IProxyData
Introduce TSharedPtr<Audio::IProxyData> CreateProxyData()
Fix up deprecated callsites.

diff --git a/Engine/Plugins/Runtime/AudioModulation/Source/AudioModulation/Private/SoundControlBus.cpp b/Engine/Plugins/Runtime/AudioModulation/Source/AudioModulation/Private/SoundControlBus.cpp
--- a/Engine/Plugins/Runtime/AudioModulation/Source/AudioModulation/Private/SoundControlBus.cpp
+++ b/Engine/Plugins/Runtime/AudioModulation/Source/AudioModulation/Private/SoundControlBus.cpp
@@ -138,10 +138,10 @@ const Audio::FModulationMixFunction USoundControlBus::GetMixFunction() const
 	return Audio::FModulationParameter::GetDefaultMixFunction();
 }
 
-TUniquePtr<Audio::IProxyData> USoundControlBus::CreateNewProxyData(const Audio::FProxyDataInitParams& InitParams)
+TSharedPtr<Audio::IProxyData> USoundControlBus::CreateProxyData(const Audio::FProxyDataInitParams& InitParams)
 {
 	using namespace AudioModulation;
-	return MakeUnique<FSoundModulatorAssetProxy>(*this);
+	return MakeShared<FSoundModulatorAssetProxy>(*this);
 }
 
 const Audio::FModulationParameter& USoundControlBus::GetOutputParameter() const
diff --git a/Engine/Plugins/Runtime/AudioModulation/Source/AudioModulation/Private/SoundModulationParameter.cpp b/Engine/Plugins/Runtime/AudioModulation/Source/AudioModulation/Private/SoundModulationParameter.cpp
--- a/Engine/Plugins/Runtime/AudioModulation/Source/AudioModulation/Private/SoundModulationParameter.cpp
+++ b/Engine/Plugins/Runtime/AudioModulation/Source/AudioModulation/Private/SoundModulationParameter.cpp
@@ -10,10 +10,10 @@
 #include UE_INLINE_GENERATED_CPP_BY_NAME(SoundModulationParameter)
 
 
-TUniquePtr<Audio::IProxyData> USoundModulationParameter::CreateNewProxyData(const Audio::FProxyDataInitParams& InitParams)
+TSharedPtr<Audio::IProxyData> USoundModulationParameter::CreateProxyData(const Audio::FProxyDataInitParams& InitParams)
 {
 	using namespace AudioModulation;
-	return MakeUnique<FSoundModulationPluginParameterAssetProxy>(this);
+	return MakeShared<FSoundModulationPluginParameterAssetProxy>(this);
 }
 
 #if WITH_EDITOR
@@ -266,9 +266,4 @@ namespace AudioModulation
 	{
 		Parameter = GetOrRegisterParameter(InParameter, TEXT("FSoundModulationPluginParameterAssetProxy construction"));
 	}
-
-	Audio::IProxyDataPtr FSoundModulationPluginParameterAssetProxy::Clone() const
-	{
-		return TUniquePtr<FSoundModulationPluginParameterAssetProxy>(new FSoundModulationPluginParameterAssetProxy(*this));
-	}
 }
diff --git a/Engine/Plugins/Runtime/AudioModulation/Source/AudioModulation/Private/SoundModulationPatch.cpp b/Engine/Plugins/Runtime/AudioModulation/Source/AudioModulation/Private/SoundModulationPatch.cpp
--- a/Engine/Plugins/Runtime/AudioModulation/Source/AudioModulation/Private/SoundModulationPatch.cpp
+++ b/Engine/Plugins/Runtime/AudioModulation/Source/AudioModulation/Private/SoundModulationPatch.cpp
@@ -21,10 +21,10 @@ USoundModulationPatch::USoundModulationPatch(const FObjectInitializer& ObjectIni
 {
 }
 
-TUniquePtr<Audio::IProxyData> USoundModulationPatch::CreateNewProxyData(const Audio::FProxyDataInitParams& InitParams)
+TSharedPtr<Audio::IProxyData> USoundModulationPatch::CreateProxyData(const Audio::FProxyDataInitParams& InitParams)
 {
 	using namespace AudioModulation;
-	return MakeUnique<FSoundModulatorAssetProxy>(*this);
+	return MakeShared<FSoundModulatorAssetProxy>(*this);
 }
 
 TUniquePtr<Audio::IModulatorSettings> USoundModulationPatch::CreateProxySettings() const
diff --git a/Engine/Plugins/Runtime/AudioModulation/Source/AudioModulation/Private/SoundModulatorAsset.cpp b/Engine/Plugins/Runtime/AudioModulation/Source/AudioModulation/Private/SoundModulatorAsset.cpp
--- a/Engine/Plugins/Runtime/AudioModulation/Source/AudioModulation/Private/SoundModulatorAsset.cpp
+++ b/Engine/Plugins/Runtime/AudioModulation/Source/AudioModulation/Private/SoundModulatorAsset.cpp
@@ -9,7 +9,7 @@ namespace AudioModulation
 	const FString PluginAuthor = TEXT("Epic Games, Inc.");
 	const FText PluginNodeMissingPrompt = NSLOCTEXT("AudioModulation", "DefaultMissingNodePrompt", "The node was likely removed, renamed, or the AudioModulation plugin is not loaded.");
 
-	FSoundModulatorAsset::FSoundModulatorAsset(const Audio::IProxyDataPtr& InProxyPtr)
+	FSoundModulatorAsset::FSoundModulatorAsset(const TSharedPtr<Audio::IProxyData>& InProxyPtr)
 	{
 		if (!InProxyPtr.IsValid())
 		{
@@ -23,7 +23,7 @@ namespace AudioModulation
 		}
 	}
 
-	FSoundModulationParameterAsset::FSoundModulationParameterAsset(const Audio::IProxyDataPtr& InProxyPtr)
+	FSoundModulationParameterAsset::FSoundModulationParameterAsset(const TSharedPtr<Audio::IProxyData>& InProxyPtr)
 	{
 		if (!InProxyPtr.IsValid())
 		{
diff --git a/Engine/Plugins/Runtime/AudioModulation/Source/AudioModulation/Public/SoundControlBus.h b/Engine/Plugins/Runtime/AudioModulation/Source/AudioModulation/Public/SoundControlBus.h
--- a/Engine/Plugins/Runtime/AudioModulation/Source/AudioModulation/Public/SoundControlBus.h
+++ b/Engine/Plugins/Runtime/AudioModulation/Source/AudioModulation/Public/SoundControlBus.h
@@ -58,7 +58,7 @@ public:
 
 	/* USoundModulatorBase Implementation */
 	virtual TUniquePtr<Audio::IModulatorSettings> CreateProxySettings() const override;
-	virtual TUniquePtr<Audio::IProxyData> CreateNewProxyData(const Audio::FProxyDataInitParams& InitParams) override;
+	virtual TSharedPtr<Audio::IProxyData> CreateProxyData(const Audio::FProxyDataInitParams& InitParams) override;
 	virtual const Audio::FModulationParameter& GetOutputParameter() const override;
 
 	const Audio::FModulationMixFunction GetMixFunction() const;
diff --git a/Engine/Plugins/Runtime/AudioModulation/Source/AudioModulation/Public/SoundModulationParameter.h b/Engine/Plugins/Runtime/AudioModulation/Source/AudioModulation/Public/SoundModulationParameter.h
--- a/Engine/Plugins/Runtime/AudioModulation/Source/AudioModulation/Public/SoundModulationParameter.h
+++ b/Engine/Plugins/Runtime/AudioModulation/Source/AudioModulation/Public/SoundModulationParameter.h
@@ -104,7 +104,7 @@ public:
 	}
 
 	//~Begin IAudioProxyDataFactory Interface.
-	virtual TUniquePtr<Audio::IProxyData> CreateNewProxyData(const Audio::FProxyDataInitParams& InitParams) override;
+	virtual TSharedPtr<Audio::IProxyData> CreateProxyData(const Audio::FProxyDataInitParams& InitParams) override;
 	//~ End IAudioProxyDataFactory Interface.
 
 #if WITH_EDITOR
@@ -260,8 +260,6 @@ namespace AudioModulation
 	public:
 		explicit FSoundModulationPluginParameterAssetProxy(USoundModulationParameter* InParameter);
 		FSoundModulationPluginParameterAssetProxy(const FSoundModulationPluginParameterAssetProxy& InProxy) = default;
-
-		virtual Audio::IProxyDataPtr Clone() const override;
 	};
 
 	/*
@@ -270,4 +268,4 @@ namespace AudioModulation
 	 * @param InBreadcrumb - String identifying get or register request initiator.
 	 */
 	AUDIOMODULATION_API const Audio::FModulationParameter& GetOrRegisterParameter(const USoundModulationParameter* InParameter, const FString& InBreadcrumb);
-} // namespace AudioModulation
\ No newline at end of file
+} // namespace AudioModulation
diff --git a/Engine/Plugins/Runtime/AudioModulation/Source/AudioModulation/Public/SoundModulationPatch.h b/Engine/Plugins/Runtime/AudioModulation/Source/AudioModulation/Public/SoundModulationPatch.h
--- a/Engine/Plugins/Runtime/AudioModulation/Source/AudioModulation/Public/SoundModulationPatch.h
+++ b/Engine/Plugins/Runtime/AudioModulation/Source/AudioModulation/Public/SoundModulationPatch.h
@@ -71,7 +71,7 @@ public:
 	FSoundControlModulationPatch PatchSettings;
 
 	/* USoundModulatorBase Implementation */
-	virtual TUniquePtr<Audio::IProxyData> CreateNewProxyData(const Audio::FProxyDataInitParams& InitParams) override;
+	virtual TSharedPtr<Audio::IProxyData> CreateProxyData(const Audio::FProxyDataInitParams& InitParams) override;
 	virtual const Audio::FModulationParameter& GetOutputParameter() const override;
 
 	virtual TUniquePtr<Audio::IModulatorSettings> CreateProxySettings() const override;
diff --git a/Engine/Plugins/Runtime/AudioModulation/Source/AudioModulation/Public/SoundModulatorAsset.h b/Engine/Plugins/Runtime/AudioModulation/Source/AudioModulation/Public/SoundModulatorAsset.h
--- a/Engine/Plugins/Runtime/AudioModulation/Source/AudioModulation/Public/SoundModulatorAsset.h
+++ b/Engine/Plugins/Runtime/AudioModulation/Source/AudioModulation/Public/SoundModulatorAsset.h
@@ -37,7 +37,7 @@ namespace AudioModulation
 			return *this;
 		}
 
-		FSoundModulatorAsset(const Audio::IProxyDataPtr& InInitData);
+		FSoundModulatorAsset(const TSharedPtr<Audio::IProxyData>& InInitData);
 
 		Audio::FModulatorId GetModulatorId() const
 		{
@@ -78,7 +78,7 @@ namespace AudioModulation
 		FSoundModulationParameterAsset() = default;
 		FSoundModulationParameterAsset(const FSoundModulationParameterAsset& InOther) = default;
 		FSoundModulationParameterAsset(FSoundModulationParameterAsset&& InOther) = default;
-		FSoundModulationParameterAsset(const Audio::IProxyDataPtr& InInitData);
+		FSoundModulationParameterAsset(const TSharedPtr<Audio::IProxyData>& InInitData);
 
 		FSoundModulationParameterAsset& operator=(const FSoundModulationParameterAsset& InOther) = default;
 		FSoundModulationParameterAsset& operator=(FSoundModulationParameterAsset && InOther) = default;
diff --git a/Engine/Plugins/Runtime/Metasound/Source/MetasoundEngine/Private/MetasoundAudioBus.cpp b/Engine/Plugins/Runtime/Metasound/Source/MetasoundEngine/Private/MetasoundAudioBus.cpp
--- a/Engine/Plugins/Runtime/Metasound/Source/MetasoundEngine/Private/MetasoundAudioBus.cpp
+++ b/Engine/Plugins/Runtime/Metasound/Source/MetasoundEngine/Private/MetasoundAudioBus.cpp
@@ -7,7 +7,7 @@
 
 namespace Metasound
 {
-	FAudioBusAsset::FAudioBusAsset(const TUniquePtr<Audio::IProxyData>& InInitData)
+	FAudioBusAsset::FAudioBusAsset(const TSharedPtr<Audio::IProxyData>& InInitData)
 	{
 		if (InInitData.IsValid())
 		{
diff --git a/Engine/Plugins/Runtime/Metasound/Source/MetasoundEngine/Private/MetasoundSource.cpp b/Engine/Plugins/Runtime/Metasound/Source/MetasoundEngine/Private/MetasoundSource.cpp
--- a/Engine/Plugins/Runtime/Metasound/Source/MetasoundEngine/Private/MetasoundSource.cpp
+++ b/Engine/Plugins/Runtime/Metasound/Source/MetasoundEngine/Private/MetasoundSource.cpp
@@ -495,7 +495,7 @@ void UMetaSoundSource::InitParameters(TArray<FAudioParameter>& ParametersToInit,
 				FDataTypeRegistryInfo DataTypeInfo;
 				if (IDataTypeRegistry::Get().GetDataTypeInfo(OutParamToInit.ObjectParam, DataTypeInfo))
 				{
-					Audio::IProxyDataPtr ProxyPtr = IDataTypeRegistry::Get().CreateProxyFromUObject(DataTypeInfo.DataTypeName, OutParamToInit.ObjectParam);
+					TSharedPtr<Audio::IProxyData> ProxyPtr = IDataTypeRegistry::Get().CreateProxyFromUObject(DataTypeInfo.DataTypeName, OutParamToInit.ObjectParam);
 					OutParamToInit.ObjectProxies.Emplace(MoveTemp(ProxyPtr));
 
 					// Null out param as it is no longer needed (nor desired to be accessed once passed to the Audio Thread)
@@ -511,7 +511,7 @@ void UMetaSoundSource::InitParameters(TArray<FAudioParameter>& ParametersToInit,
 					FDataTypeRegistryInfo DataTypeInfo;
 					if (IDataTypeRegistry::Get().GetDataTypeInfo(Object, DataTypeInfo))
 					{
-						Audio::IProxyDataPtr ProxyPtr = IDataTypeRegistry::Get().CreateProxyFromUObject(DataTypeInfo.DataTypeName, Object);
+						TSharedPtr<Audio::IProxyData> ProxyPtr = IDataTypeRegistry::Get().CreateProxyFromUObject(DataTypeInfo.DataTypeName, Object);
 						OutParamToInit.ObjectProxies.Emplace(MoveTemp(ProxyPtr));
 					}
 				}
@@ -842,7 +842,7 @@ bool UMetaSoundSource::IsParameterValid(const FAudioParameter& InParameter, cons
 		{
 			FDataTypeRegistryInfo DataTypeInfo;
 			bIsValid = IDataTypeRegistry::Get().GetDataTypeInfo(InParameter.ObjectParam, DataTypeInfo);
-			bIsValid &= DataTypeInfo.bIsProxyParsable;
+			bIsValid &= DataTypeInfo.bIsProxyParsable || DataTypeInfo.bIsUniquePtrProxyParsable_DEPRECATED;
 			bIsValid &= DataTypeInfo.DataTypeName == TypeName;
 		}
 		break;
@@ -856,7 +856,7 @@ bool UMetaSoundSource::IsParameterValid(const FAudioParameter& InParameter, cons
 			{
 				FDataTypeRegistryInfo DataTypeInfo;
 				bIsValid = IDataTypeRegistry::Get().GetDataTypeInfo(Object, DataTypeInfo);
-				bIsValid &= DataTypeInfo.bIsProxyParsable;
+				bIsValid &= DataTypeInfo.bIsProxyParsable || DataTypeInfo.bIsUniquePtrProxyParsable_DEPRECATED;
 				bIsValid &= DataTypeInfo.DataTypeName == ElementTypeName;
 
 				if (!bIsValid)
diff --git a/Engine/Plugins/Runtime/Metasound/Source/MetasoundEngine/Private/MetasoundWave.cpp b/Engine/Plugins/Runtime/Metasound/Source/MetasoundEngine/Private/MetasoundWave.cpp
--- a/Engine/Plugins/Runtime/Metasound/Source/MetasoundEngine/Private/MetasoundWave.cpp
+++ b/Engine/Plugins/Runtime/Metasound/Source/MetasoundEngine/Private/MetasoundWave.cpp
@@ -24,7 +24,7 @@ FAutoConsoleVariableRef CVarDisableMetasoundWaveAssetCachePriming(
 
 namespace Metasound
 {
-	FWaveAsset::FWaveAsset(const TUniquePtr<Audio::IProxyData>& InInitData)
+	FWaveAsset::FWaveAsset(const TSharedPtr<Audio::IProxyData>& InInitData)
 	{
 		if (InInitData.IsValid())
 		{
diff --git a/Engine/Plugins/Runtime/Metasound/Source/MetasoundEngine/Public/MetasoundAudioBus.h b/Engine/Plugins/Runtime/Metasound/Source/MetasoundEngine/Public/MetasoundAudioBus.h
--- a/Engine/Plugins/Runtime/Metasound/Source/MetasoundEngine/Public/MetasoundAudioBus.h
+++ b/Engine/Plugins/Runtime/Metasound/Source/MetasoundEngine/Public/MetasoundAudioBus.h
@@ -23,7 +23,7 @@ namespace Metasound
 		FAudioBusAsset(const FAudioBusAsset&) = default;
 		FAudioBusAsset& operator=(const FAudioBusAsset& Other) = default;
 
-		FAudioBusAsset(const TUniquePtr<Audio::IProxyData>& InInitData);
+		FAudioBusAsset(const TSharedPtr<Audio::IProxyData>& InInitData);
 
 		const FAudioBusProxyPtr& GetAudioBusProxy() const
 		{
diff --git a/Engine/Plugins/Runtime/Metasound/Source/MetasoundEngine/Public/MetasoundWave.h b/Engine/Plugins/Runtime/Metasound/Source/MetasoundEngine/Public/MetasoundWave.h
--- a/Engine/Plugins/Runtime/Metasound/Source/MetasoundEngine/Public/MetasoundWave.h
+++ b/Engine/Plugins/Runtime/Metasound/Source/MetasoundEngine/Public/MetasoundWave.h
@@ -28,7 +28,7 @@ namespace Metasound
 		FWaveAsset(const FWaveAsset&) = default;
 		FWaveAsset& operator=(const FWaveAsset& Other) = default;
 
-		FWaveAsset(const TUniquePtr<Audio::IProxyData>& InInitData);
+		FWaveAsset(const TSharedPtr<Audio::IProxyData>& InInitData);
 
 		bool IsSoundWaveValid() const;
 
diff --git a/Engine/Plugins/Runtime/Metasound/Source/MetasoundEngine/Public/MetasoundWaveTable.h b/Engine/Plugins/Runtime/Metasound/Source/MetasoundEngine/Public/MetasoundWaveTable.h
--- a/Engine/Plugins/Runtime/Metasound/Source/MetasoundEngine/Public/MetasoundWaveTable.h
+++ b/Engine/Plugins/Runtime/Metasound/Source/MetasoundEngine/Public/MetasoundWaveTable.h
@@ -45,7 +45,7 @@ namespace Metasound
 			return *this;
 		}
 
-		FWaveTableBankAsset(const Audio::IProxyDataPtr& InInitData)
+		FWaveTableBankAsset(const TSharedPtr<Audio::IProxyData>& InInitData)
 		{
 			if (InInitData.IsValid())
 			{
diff --git a/Engine/Plugins/Runtime/Metasound/Source/MetasoundFrontend/Private/MetasoundFrontendDataTypeRegistry.cpp b/Engine/Plugins/Runtime/Metasound/Source/MetasoundFrontend/Private/MetasoundFrontendDataTypeRegistry.cpp
--- a/Engine/Plugins/Runtime/Metasound/Source/MetasoundFrontend/Private/MetasoundFrontendDataTypeRegistry.cpp
+++ b/Engine/Plugins/Runtime/Metasound/Source/MetasoundFrontend/Private/MetasoundFrontendDataTypeRegistry.cpp
@@ -513,7 +513,7 @@ namespace Metasound
 				virtual UClass* GetUClassForDataType(const FName& InDataType) const override;
 
 				bool IsUObjectProxyFactory(UObject* InObject) const override;
-				Audio::IProxyDataPtr CreateProxyFromUObject(const FName& InDataType, UObject* InObject) const override;
+				TSharedPtr<Audio::IProxyData> CreateProxyFromUObject(const FName& InDataType, UObject* InObject) const override;
 
 				virtual FLiteral CreateDefaultLiteral(const FName& InDataType) const override;
 				virtual FLiteral CreateLiteralFromUObject(const FName& InDataType, UObject* InObject) const override;
@@ -731,11 +731,11 @@ namespace Metasound
 
 						case Metasound::ELiteralType::UObjectProxy:
 						{
-							return Info.bIsProxyParsable;
+							return Info.bIsProxyParsable || Info.bIsUniquePtrProxyParsable_DEPRECATED;
 						}
 						case Metasound::ELiteralType::UObjectProxyArray:
 						{
-							return Info.bIsProxyArrayParsable;
+							return Info.bIsProxyArrayParsable || Info.bIsUniquePtrProxyArrayParsable_DEPRECATED;
 						}
 
 						case Metasound::ELiteralType::None:
@@ -812,9 +812,9 @@ namespace Metasound
 				return false;
 			}
 
-			Audio::IProxyDataPtr FDataTypeRegistry::CreateProxyFromUObject(const FName& InDataType, UObject* InObject) const
+			TSharedPtr<Audio::IProxyData> FDataTypeRegistry::CreateProxyFromUObject(const FName& InDataType, UObject* InObject) const
 			{
-				Audio::IProxyDataPtr ProxyPtr;
+				TSharedPtr<Audio::IProxyData> ProxyPtr;
 
 				if (const IDataTypeRegistryEntry* Entry = FindDataTypeEntry(InDataType))
 				{
@@ -830,13 +830,13 @@ namespace Metasound
 
 			FLiteral FDataTypeRegistry::CreateLiteralFromUObject(const FName& InDataType, UObject* InObject) const
 			{
-				Audio::IProxyDataPtr ProxyPtr = CreateProxyFromUObject(InDataType, InObject);
+				TSharedPtr<Audio::IProxyData> ProxyPtr = CreateProxyFromUObject(InDataType, InObject);
 				return Metasound::FLiteral(MoveTemp(ProxyPtr));
 			}
 
 			FLiteral FDataTypeRegistry::CreateLiteralFromUObjectArray(const FName& InDataType, const TArray<UObject*>& InObjectArray) const
 			{
-				TArray<Audio::IProxyDataPtr> ProxyArray;
+				TArray<TSharedPtr<Audio::IProxyData>> ProxyArray;
 				const IDataTypeRegistryEntry* DataTypeEntry = FindDataTypeEntry(InDataType);
 				if (!DataTypeEntry)
 				{
@@ -846,14 +846,8 @@ namespace Metasound
 
 				const FDataTypeRegistryInfo& DataTypeInfo = DataTypeEntry->GetDataTypeInfo();
 
-				const bool bIsArrayType = DataTypeInfo.bIsProxyArrayParsable;
-				if (!bIsArrayType)
-				{
-					UE_LOG(LogMetaSound, Error, TEXT("Failed to create a valid proxy from DataType '%s': Type is not 'ArrayType'."), *InDataType.ToString());
-					return Metasound::FLiteral(MoveTemp(ProxyArray));
-				}
+				const bool bIsProxyArrayParseable = DataTypeInfo.bIsProxyArrayParsable || DataTypeInfo.bIsUniquePtrProxyArrayParsable_DEPRECATED;
 
-				const bool bIsProxyArrayParseable = DataTypeInfo.bIsProxyArrayParsable;
 				if (!bIsProxyArrayParseable)
 				{
 					UE_LOG(LogMetaSound, Error, TEXT("Failed to create a valid proxy from DataType '%s': Type is not proxy parseable."), *InDataType.ToString());
@@ -870,7 +864,7 @@ namespace Metasound
 
 				for (UObject* InObject : InObjectArray)
 				{
-					Audio::IProxyDataPtr ProxyPtr = CreateProxyFromUObject(ElementDataType, InObject);
+					TSharedPtr<Audio::IProxyData> ProxyPtr = CreateProxyFromUObject(ElementDataType, InObject);
 					ProxyPtr = ElementEntry->CreateProxy(InObject);
 					if (!ProxyPtr && InObject)
 					{
diff --git a/Engine/Plugins/Runtime/Metasound/Source/MetasoundFrontend/Private/MetasoundFrontendLiteral.cpp b/Engine/Plugins/Runtime/Metasound/Source/MetasoundFrontend/Private/MetasoundFrontendLiteral.cpp
--- a/Engine/Plugins/Runtime/Metasound/Source/MetasoundFrontend/Private/MetasoundFrontendLiteral.cpp
+++ b/Engine/Plugins/Runtime/Metasound/Source/MetasoundFrontend/Private/MetasoundFrontendLiteral.cpp
@@ -388,7 +388,7 @@ void FMetasoundFrontendLiteral::SetFromLiteral(const Metasound::FLiteral& InLite
 			// Only error if attempting to retrieve valid UObject from ProxyDataPtr
 			// as this function can safely is used to initialize from defaults (which
 			// is valid as a null proxy can safely correspond to a null UObject ptr).
-			if (InLiteral.Value.Get<Audio::IProxyDataPtr>().IsValid())
+			if (InLiteral.Value.Get<TSharedPtr<Audio::IProxyData>>().IsValid())
 			{
 				UE_LOG(LogMetaSound, Error, TEXT("Cannot set UObjectProxy from Metasound::FLiteral"));
 			}
@@ -432,7 +432,7 @@ void FMetasoundFrontendLiteral::SetFromLiteral(const Metasound::FLiteral& InLite
 			// Only error if attempting to retrieve valid UObject from ProxyDataPtr
 			// as this function can safely is used to initialize from defaults (which
 			// is valid as a null proxy can safely correspond to a null UObject ptr).
-			if (!InLiteral.Value.Get<TArray<Audio::IProxyDataPtr>>().IsEmpty())
+			if (!InLiteral.Value.Get<TArray<TSharedPtr<Audio::IProxyData>>>().IsEmpty())
 			{
 				UE_LOG(LogMetaSound, Error, TEXT("Cannot set UObjectProxy from Metasound::FLiteral"));
 			}
diff --git a/Engine/Plugins/Runtime/Metasound/Source/MetasoundFrontend/Private/MetasoundParameterTransmitter.cpp b/Engine/Plugins/Runtime/Metasound/Source/MetasoundFrontend/Private/MetasoundParameterTransmitter.cpp
--- a/Engine/Plugins/Runtime/Metasound/Source/MetasoundFrontend/Private/MetasoundParameterTransmitter.cpp
+++ b/Engine/Plugins/Runtime/Metasound/Source/MetasoundFrontend/Private/MetasoundParameterTransmitter.cpp
@@ -155,18 +155,12 @@ namespace Metasound
 						return FLiteral();
 					}
 
-					Audio::IProxyDataPtr ObjectProxyClone = InValue.ObjectProxies.Last()->Clone();
-					return FLiteral(MoveTemp(ObjectProxyClone));
+					return FLiteral(InValue.ObjectProxies.Last());
 				}
 
 				case EAudioParameterType::ObjectArray:
 				{
-					TArray<Audio::IProxyDataPtr> ObjectProxiesClone;
-					Algo::Transform(InValue.ObjectProxies, ObjectProxiesClone, [](const Audio::IProxyDataPtr& DataPtr)
-					{
-						return DataPtr->Clone();
-					});
-					return FLiteral(MoveTemp(ObjectProxiesClone));
+					return FLiteral(InValue.ObjectProxies);
 				}
 
 				case EAudioParameterType::String:
diff --git a/Engine/Plugins/Runtime/Metasound/Source/MetasoundFrontend/Public/MetasoundDataTypeRegistrationMacro.h b/Engine/Plugins/Runtime/Metasound/Source/MetasoundFrontend/Public/MetasoundDataTypeRegistrationMacro.h
--- a/Engine/Plugins/Runtime/Metasound/Source/MetasoundFrontend/Public/MetasoundDataTypeRegistrationMacro.h
+++ b/Engine/Plugins/Runtime/Metasound/Source/MetasoundFrontend/Public/MetasoundDataTypeRegistrationMacro.h
@@ -32,6 +32,16 @@ namespace Metasound
 {
 	namespace MetasoundDataTypeRegistrationPrivate
 	{
+		template<typename DataType>
+		struct TDataTypeProxyConstructorDeprecation
+		{
+			private:
+				static constexpr bool bIsParsableWithDeprecatedPtr = TIsParsable<DataType, DataFactoryPrivate::ProxyDataPtrType_DEPRECATED>::Value;
+				static constexpr bool bIsParsableWithSharedProxyPtr = TIsParsable<DataType, TSharedPtr<Audio::IProxyData>>::Value;
+			public:
+				static constexpr bool bOnlySupportsDeprecatedProxyPtr = bIsParsableWithDeprecatedPtr && !bIsParsableWithSharedProxyPtr;
+		};
+
 		// Returns the Array version of a literal type if it exists.
 		template<ELiteralType LiteralType>
 		struct TLiteralArrayEnum 
@@ -171,14 +181,18 @@ namespace Metasound
 			RegistryInfo.bIsIntParsable = TIsParsable<TDataType, int32>::Value;
 			RegistryInfo.bIsFloatParsable = TIsParsable<TDataType, float>::Value;
 			RegistryInfo.bIsStringParsable = TIsParsable<TDataType, FString>::Value;
-			RegistryInfo.bIsProxyParsable = TIsParsable<TDataType, const Audio::IProxyDataPtr&>::Value;
+			RegistryInfo.bIsProxyParsable = TIsParsable<TDataType, const TSharedPtr<Audio::IProxyData>&>::Value;
+
+			RegistryInfo.bIsUniquePtrProxyParsable_DEPRECATED = TIsParsable<TDataType, const TUniquePtr<Audio::IProxyData>&>::Value;
 
 			RegistryInfo.bIsDefaultArrayParsable = TIsParsable<TDataType, TArray<FLiteral::FNone>>::Value;
 			RegistryInfo.bIsBoolArrayParsable = TIsParsable<TDataType, TArray<bool>>::Value;
 			RegistryInfo.bIsIntArrayParsable = TIsParsable<TDataType, TArray<int32>>::Value;
 			RegistryInfo.bIsFloatArrayParsable = TIsParsable<TDataType, TArray<float>>::Value;
 			RegistryInfo.bIsStringArrayParsable = TIsParsable<TDataType, TArray<FString>>::Value;
-			RegistryInfo.bIsProxyArrayParsable = TIsParsable<TDataType, const TArray<Audio::IProxyDataPtr>& >::Value;
+			RegistryInfo.bIsProxyArrayParsable = TIsParsable<TDataType, const TArray<TSharedPtr<Audio::IProxyData>>& >::Value;
+
+			RegistryInfo.bIsUniquePtrProxyArrayParsable_DEPRECATED = TIsParsable<TDataType, const TArray<TUniquePtr<Audio::IProxyData>>& >::Value;
 
 			RegistryInfo.bIsEnum = TEnumTraits<TDataType>::bIsEnum;
 			RegistryInfo.bIsVariable = TIsVariable<TDataType>::Value;
@@ -280,6 +294,15 @@ namespace Metasound
 				return false;
 			}
 
+	
+			if constexpr (MetasoundDataTypeRegistrationPrivate::TDataTypeProxyConstructorDeprecation<TDataType>::bOnlySupportsDeprecatedProxyPtr)
+			{
+				// TUniquePtr<Audio::IProxyData> deprecated in 5.2. Log a warning
+				// during data type registration to warn users to update their MetaSound
+				// data type constructor. 
+				UE_LOG(LogMetaSound, Warning, TEXT("MetaSound data type \"%s\" supports construction from deprecated TUniquePtr<Audio::IProxyData>. Please update the constructor to accept a \"const TSharedPtr<Audio::IProxyData>& \""), TDataReferenceTypeInfo<TDataType>::TypeName);
+			}
+
 			bAlreadyRegisteredThisDataType = true;
 
 			// Define registry entry class for this data type.
@@ -505,7 +528,7 @@ namespace Metasound
 					}
 				}
 
-				virtual Audio::IProxyDataPtr CreateProxy(UObject* InObject) const override
+				virtual TSharedPtr<Audio::IProxyData> CreateProxy(UObject* InObject) const override
 				{
 					// Only attempt to create proxy if the `UClassToUse` is not void.
 					if constexpr (!std::is_same<UClassToUse, void>::value)
@@ -519,12 +542,12 @@ namespace Metasound
 								Audio::FProxyDataInitParams ProxyInitParams;
 								ProxyInitParams.NameOfFeatureRequestingProxy = "MetaSound";
 
-								return ObjectAsFactory->CreateNewProxyData(ProxyInitParams);
+								return ObjectAsFactory->CreateProxyData(ProxyInitParams);
 							}
 						}
 					}
 
-					return Audio::IProxyDataPtr(nullptr);
+					return TSharedPtr<Audio::IProxyData>(nullptr);
 				}
 
 				virtual TOptional<FAnyDataReference> CreateDataReference(EDataReferenceAccessType InAccessType, const FLiteral& InLiteral, const FOperatorSettings& InOperatorSettings) const override
diff --git a/Engine/Plugins/Runtime/Metasound/Source/MetasoundFrontend/Public/MetasoundFrontendDataTypeRegistry.h b/Engine/Plugins/Runtime/Metasound/Source/MetasoundFrontend/Public/MetasoundFrontendDataTypeRegistry.h
--- a/Engine/Plugins/Runtime/Metasound/Source/MetasoundFrontend/Public/MetasoundFrontendDataTypeRegistry.h
+++ b/Engine/Plugins/Runtime/Metasound/Source/MetasoundFrontend/Public/MetasoundFrontendDataTypeRegistry.h
@@ -39,12 +39,14 @@ namespace Metasound
 			bool bIsFloatParsable = false;
 			bool bIsStringParsable = false;
 			bool bIsProxyParsable = false;
+			bool bIsUniquePtrProxyParsable_DEPRECATED = false;
 			bool bIsDefaultArrayParsable = false;
 			bool bIsBoolArrayParsable = false;
 			bool bIsIntArrayParsable = false;
 			bool bIsFloatArrayParsable = false;
 			bool bIsStringArrayParsable = false;
 			bool bIsProxyArrayParsable = false;
+			bool bIsUniquePtrProxyArrayParsable_DEPRECATED = false;
 
 			// Is a TEnum wrapped enum
 			bool bIsEnum = false;
@@ -202,7 +204,7 @@ namespace Metasound
 
 			/** Create a proxy from a UObject. If this data type does not support
 			 * UObject proxies, return a nullptr. */
-			virtual Audio::IProxyDataPtr CreateProxy(UObject* InObject) const = 0;
+			virtual TSharedPtr<Audio::IProxyData> CreateProxy(UObject* InObject) const = 0;
 
 			/** Create a data channel for transmission. If this data type does not
 			 * support transmission, return a nullptr. */
@@ -253,7 +255,7 @@ namespace Metasound
 
 			virtual UClass* GetUClassForDataType(const FName& InDataType) const = 0;
 
-			virtual Audio::IProxyDataPtr CreateProxyFromUObject(const FName& InDataType, UObject* InObject) const = 0;
+			virtual TSharedPtr<Audio::IProxyData> CreateProxyFromUObject(const FName& InDataType, UObject* InObject) const = 0;
 
 			virtual FLiteral CreateDefaultLiteral(const FName& InDataType) const = 0;
 			virtual FLiteral CreateLiteralFromUObject(const FName& InDataType, UObject* InObject) const = 0;
diff --git a/Engine/Plugins/Runtime/Metasound/Source/MetasoundGraphCore/Private/MetasoundLiteral.cpp b/Engine/Plugins/Runtime/Metasound/Source/MetasoundGraphCore/Private/MetasoundLiteral.cpp
--- a/Engine/Plugins/Runtime/Metasound/Source/MetasoundGraphCore/Private/MetasoundLiteral.cpp
+++ b/Engine/Plugins/Runtime/Metasound/Source/MetasoundGraphCore/Private/MetasoundLiteral.cpp
@@ -43,7 +43,7 @@ namespace Metasound
 				return FLiteral(TLiteralTypeInfo<FString>::GetDefaultValue());
 
 			case ELiteralType::UObjectProxy:
-				return FLiteral(TLiteralTypeInfo<Audio::IProxyDataPtr>::GetDefaultValue());
+				return FLiteral(TLiteralTypeInfo<TSharedPtr<Audio::IProxyData>>::GetDefaultValue());
 			
 			case ELiteralType::NoneArray:
 				return FLiteral(TLiteralTypeInfo<TArray<FLiteral::FNone>>::GetDefaultValue());
@@ -61,7 +61,7 @@ namespace Metasound
 				return FLiteral(TLiteralTypeInfo<TArray<FString>>::GetDefaultValue());
 
 			case ELiteralType::UObjectProxyArray:
-				return FLiteral(TLiteralTypeInfo<TArray<Audio::IProxyDataPtr>>::GetDefaultValue());
+				return FLiteral(TLiteralTypeInfo<TArray<TSharedPtr<Audio::IProxyData>>>::GetDefaultValue());
 
 			case ELiteralType::Invalid:
 			default:
@@ -148,11 +148,11 @@ FString LexToString(const Metasound::FLiteral& InLiteral)
 		case ELiteralType::UObjectProxy:
 		{
 			FString ProxyType = TEXT("nullptr");
-			if (InLiteral.Value.Get<Audio::IProxyDataPtr>().IsValid())
+			if (InLiteral.Value.Get<TSharedPtr<Audio::IProxyData>>().IsValid())
 			{
-				ProxyType = InLiteral.Value.Get<Audio::IProxyDataPtr>()->GetProxyTypeName().ToString();
+				ProxyType = InLiteral.Value.Get<TSharedPtr<Audio::IProxyData>>()->GetProxyTypeName().ToString();
 			}
-			return FString::Printf(TEXT("Audio::IProxyDataPtr: %s"), *ProxyType);
+			return FString::Printf(TEXT("TSharedPtr<Audio::IProxyData>: %s"), *ProxyType);
 		}
 		break;
 
@@ -172,7 +172,7 @@ FString LexToString(const Metasound::FLiteral& InLiteral)
 			return TEXT("TArray<FString>");
 
 		case ELiteralType::UObjectProxyArray:
-			return TEXT("TArray<Audio::IProxyDataPtr>");
+			return TEXT("TArray<TSharedPtr<Audio::IProxyData>>");
 
 		case ELiteralType::Invalid:
 			return TEXT("INVALID");
diff --git a/Engine/Plugins/Runtime/Metasound/Source/MetasoundGraphCore/Public/MetasoundDataFactory.h b/Engine/Plugins/Runtime/Metasound/Source/MetasoundGraphCore/Public/MetasoundDataFactory.h
--- a/Engine/Plugins/Runtime/Metasound/Source/MetasoundGraphCore/Public/MetasoundDataFactory.h
+++ b/Engine/Plugins/Runtime/Metasound/Source/MetasoundGraphCore/Public/MetasoundDataFactory.h
@@ -12,6 +12,9 @@ namespace Metasound
 {
 	namespace DataFactoryPrivate
 	{
+		// Deprecated in 5.2. Use TSharedPtr<Audio::IProxyData> instead of TUniquePtr<Audio::IProxyData>.
+		using ProxyDataPtrType_DEPRECATED = TUniquePtr<Audio::IProxyData>;
+
 		/** Description of available constructors for a registered Metasound Data Type. 
 		 *
 		 * @tparam DataType - The registered Metasound Data Type.
@@ -825,8 +828,10 @@ namespace Metasound
 					return TIsParsable<DataType, FString>::Value;
 
 				case ELiteralType::UObjectProxy:
-
-					return TIsParsable<DataType, Audio::IProxyDataPtr>::Value;
+					// Support for current proxy pointer and deprected proxy pointer. Deprecation of prior proxy pointer occured in release 5.2
+					PRAGMA_DISABLE_DEPRECATION_WARNINGS
+					return TIsParsable<DataType, TSharedPtr<Audio::IProxyData>>::Value || TIsParsable<DataType, DataFactoryPrivate::ProxyDataPtrType_DEPRECATED>::Value;
+					PRAGMA_ENABLE_DEPRECATION_WARNINGS
 
 				case ELiteralType::NoneArray:
 
@@ -850,7 +855,9 @@ namespace Metasound
 
 				case ELiteralType::UObjectProxyArray:
 
-					return TIsParsable<DataType, TArray<Audio::IProxyDataPtr>>::Value;
+					PRAGMA_DISABLE_DEPRECATION_WARNINGS
+					return TIsParsable<DataType, TArray<TSharedPtr<Audio::IProxyData>>>::Value || TIsParsable<DataType, TArray<DataFactoryPrivate::ProxyDataPtrType_DEPRECATED>>::Value;
+					PRAGMA_ENABLE_DEPRECATION_WARNINGS
 
 				default:
 
@@ -1055,8 +1062,33 @@ namespace Metasound
 					return CreateExplicitArgs<FString>(InSettings, InLiteral.Value);
 
 				case ELiteralType::UObjectProxy:
-
-					return CreateExplicitArgs<Audio::IProxyDataPtr>(InSettings, InLiteral.Value);
+					{
+						constexpr bool bIsParsable = TIsParsable<DataType, TSharedPtr<Audio::IProxyData>>::Value;
+
+						PRAGMA_DISABLE_DEPRECATION_WARNINGS
+						constexpr bool bIsParsableWithDeprecatedPtr = TIsParsable<DataType, DataFactoryPrivate::ProxyDataPtrType_DEPRECATED>::Value;
+						PRAGMA_ENABLE_DEPRECATION_WARNINGS
+
+						// Temporarily convert TSharedPtr to TUniquePtr in scenario where only the deprecated proxy factory exists. 
+						if constexpr (!bIsParsable && bIsParsableWithDeprecatedPtr)
+						{
+							using FFactoryForDeprecatedPtr = TDataFactory<DataType, CreatorType>;
+
+							// The proxy is stored in a TUniquePtr<>, but it is released from the TUniquePtr<> before the TUniquePtr<> destructor is called. 
+							// This is done to match the interface for MetaSound data types which utilized a deprecated interface. It is not intended
+							// to pass ownership of the proxy to the TUniquePtr<>.
+							TUniquePtr<Audio::IProxyData> DeprecatedProxyDataPtr(InLiteral.Value.Get<TSharedPtr<Audio::IProxyData>>().Get());
+
+							auto Result = FFactoryForDeprecatedPtr::template CreateExplicitArgs(InSettings, (const TUniquePtr<Audio::IProxyData>&)DeprecatedProxyDataPtr);
+
+							DeprecatedProxyDataPtr.Release();
+							return Result;
+						}
+						else
+						{
+							return CreateExplicitArgs<TSharedPtr<Audio::IProxyData>>(InSettings, InLiteral.Value);
+						}
+					}
 
 				case ELiteralType::NoneArray:
 
@@ -1079,8 +1111,40 @@ namespace Metasound
 					return CreateExplicitArgs<TArray<FString>>(InSettings, InLiteral.Value);
 
 				case ELiteralType::UObjectProxyArray:
-
-					return CreateExplicitArgs<TArray<Audio::IProxyDataPtr>>(InSettings, InLiteral.Value);
+					{
+						constexpr bool bIsParsable = TIsParsable<DataType, TArray<TSharedPtr<Audio::IProxyData>>>::Value;
+						PRAGMA_DISABLE_DEPRECATION_WARNINGS
+						constexpr bool bIsParsableWithDeprecatedPtr = TIsParsable<DataType, TArray<DataFactoryPrivate::ProxyDataPtrType_DEPRECATED>>::Value;
+						PRAGMA_ENABLE_DEPRECATION_WARNINGS
+
+						// Temporarily convert TSharedPtr to TUniquePtr in scenario where only the deprecated proxy factory exists. 
+						if constexpr (!bIsParsable && bIsParsableWithDeprecatedPtr)
+						{
+							using FFactoryForDeprecatedPtr = TDataFactory<DataType, CreatorType>;
+
+							// The proxies are stored in a TArray<TUniquePtr<>>, but it is released before the TArray<TUniquePtr<>> destructor is called. 
+							// This is done to match the interface for MetaSound data types which utilized a deprecated interface. It is not intended
+							// to pass ownership of the proxies to the TArray<TUniquePtr<>>.
+							const TArray<TSharedPtr<Audio::IProxyData>>& ProxyDataPtrArray = InLiteral.Value.Get<TArray<TSharedPtr<Audio::IProxyData>>>();
+							TArray<TUniquePtr<Audio::IProxyData>> DeprecatedArray;
+
+							for (const TSharedPtr<Audio::IProxyData>& ProxyDataPtr : ProxyDataPtrArray)
+							{
+								DeprecatedArray.Emplace(ProxyDataPtr.Get());
+							}
+
+							auto Result = FFactoryForDeprecatedPtr::template CreateExplicitArgs(InSettings, (const TArray<TUniquePtr<Audio::IProxyData>>&)DeprecatedArray);
+
+							for (TUniquePtr<Audio::IProxyData>& DeprecatedPtr : DeprecatedArray)
+							{
+								DeprecatedPtr.Release();
+							}
+
+							return Result;
+						}
+					}
+
+					return CreateExplicitArgs<TArray<TSharedPtr<Audio::IProxyData>>>(InSettings, InLiteral.Value);
 
 				case ELiteralType::None:
 
diff --git a/Engine/Plugins/Runtime/Metasound/Source/MetasoundGraphCore/Public/MetasoundLiteral.h b/Engine/Plugins/Runtime/Metasound/Source/MetasoundGraphCore/Public/MetasoundLiteral.h
--- a/Engine/Plugins/Runtime/Metasound/Source/MetasoundGraphCore/Public/MetasoundLiteral.h
+++ b/Engine/Plugins/Runtime/Metasound/Source/MetasoundGraphCore/Public/MetasoundLiteral.h
@@ -81,8 +81,8 @@ namespace Metasound
 		struct FNone {};
 
 		using FVariantType = TVariant<
-			FNone, bool, int32, float, FString, Audio::IProxyDataPtr, // Single value types
-			TArray<FNone>, TArray<bool>, TArray<int32>, TArray<float>, TArray<FString>, TArray<Audio::IProxyDataPtr>, // Array of values types
+			FNone, bool, int32, float, FString, TSharedPtr<Audio::IProxyData>, // Single value types
+			TArray<FNone>, TArray<bool>, TArray<int32>, TArray<float>, TArray<FString>, TArray<TSharedPtr<Audio::IProxyData>>, // Array of values types
 			FInvalid
 		>;
 
@@ -203,7 +203,7 @@ namespace Metasound
 		}
 
 		template<>
-		FORCEINLINE ELiteralType GetLiteralArgTypeFromDecayed<Audio::IProxyDataPtr>()
+		FORCEINLINE ELiteralType GetLiteralArgTypeFromDecayed<TSharedPtr<Audio::IProxyData>>()
 		{
 			return ELiteralType::UObjectProxy;
 		}
@@ -239,7 +239,7 @@ namespace Metasound
 		}
 
 		template<>
-		FORCEINLINE ELiteralType GetLiteralArgTypeFromDecayed<TArray<Audio::IProxyDataPtr>>()
+		FORCEINLINE ELiteralType GetLiteralArgTypeFromDecayed<TArray<TSharedPtr<Audio::IProxyData>>>()
 		{
 			return ELiteralType::UObjectProxyArray;
 		}
@@ -287,20 +287,20 @@ namespace Metasound
 		};
 
 		template<>
-		struct TLiteralDefaultValueFromDecayed<Audio::IProxyDataPtr>
+		struct TLiteralDefaultValueFromDecayed<TSharedPtr<Audio::IProxyData>>
 		{
-			static Audio::IProxyDataPtr GetValue()
+			static TSharedPtr<Audio::IProxyData> GetValue()
 			{
-				return Audio::IProxyDataPtr(nullptr);
+				return TSharedPtr<Audio::IProxyData>(nullptr);
 			}
 		};
 
 		template<>
-		struct TLiteralDefaultValueFromDecayed<TArray<Audio::IProxyDataPtr>>
+		struct TLiteralDefaultValueFromDecayed<TArray<TSharedPtr<Audio::IProxyData>>>
 		{
-			static TArray<Audio::IProxyDataPtr> GetValue()
+			static TArray<TSharedPtr<Audio::IProxyData>> GetValue()
 			{
-				return TArray<Audio::IProxyDataPtr>();
+				return TArray<TSharedPtr<Audio::IProxyData>>();
 			}
 		};
 	}
diff --git a/Engine/Plugins/Runtime/SoundMod/Source/SoundMod/SoundMod.Build.cs b/Engine/Plugins/Runtime/SoundMod/Source/SoundMod/SoundMod.Build.cs
--- a/Engine/Plugins/Runtime/SoundMod/Source/SoundMod/SoundMod.Build.cs
+++ b/Engine/Plugins/Runtime/SoundMod/Source/SoundMod/SoundMod.Build.cs
@@ -9,6 +9,7 @@ namespace UnrealBuildTool.Rules
             PublicDependencyModuleNames.AddRange(
                 new string[]
 				{
+					"AudioExtensions",
 					"Core",
 					"CoreUObject",
                     "Engine",
diff --git a/Engine/Plugins/Runtime/WaveTable/Source/WaveTable/Private/WaveTableBank.cpp b/Engine/Plugins/Runtime/WaveTable/Source/WaveTable/Private/WaveTableBank.cpp
--- a/Engine/Plugins/Runtime/WaveTable/Source/WaveTable/Private/WaveTableBank.cpp
+++ b/Engine/Plugins/Runtime/WaveTable/Source/WaveTable/Private/WaveTableBank.cpp
@@ -7,9 +7,9 @@
 #define LOCTEXT_NAMESPACE "WaveTable"
 
 
-TUniquePtr<Audio::IProxyData> UWaveTableBank::CreateNewProxyData(const Audio::FProxyDataInitParams& InitParams)
+TSharedPtr<Audio::IProxyData> UWaveTableBank::CreateProxyData(const Audio::FProxyDataInitParams& InitParams)
 {
-	return MakeUnique<FWaveTableBankAssetProxy>(*this);
+	return MakeShared<FWaveTableBankAssetProxy>(*this);
 }
 
 #if WITH_EDITOR
diff --git a/Engine/Plugins/Runtime/WaveTable/Source/WaveTable/Public/WaveTableBank.h b/Engine/Plugins/Runtime/WaveTable/Source/WaveTable/Public/WaveTableBank.h
--- a/Engine/Plugins/Runtime/WaveTable/Source/WaveTable/Public/WaveTableBank.h
+++ b/Engine/Plugins/Runtime/WaveTable/Source/WaveTable/Public/WaveTableBank.h
@@ -52,7 +52,7 @@ public:
 	TArray<FWaveTableBankEntry> Entries;
 
 	/* IAudioProxyDataFactory Implementation */
-	virtual TUniquePtr<Audio::IProxyData> CreateNewProxyData(const Audio::FProxyDataInitParams& InitParams) override;
+	virtual TSharedPtr<Audio::IProxyData> CreateProxyData(const Audio::FProxyDataInitParams& InitParams) override;
 
 #if WITH_EDITOR
 	void RefreshWaveTables();
@@ -81,11 +81,6 @@ public:
 		});
 	}
 
-	virtual Audio::IProxyDataPtr Clone() const
-	{
-		return MakeUnique<FWaveTableBankAssetProxy>(*this);
-	}
-
 	virtual const TArray<WaveTable::FWaveTable>& GetWaveTables() const
 	{
 		return WaveTables;
diff --git a/Engine/Source/Runtime/AudioExtensions/Private/AudioExtensionsModule.cpp b/Engine/Source/Runtime/AudioExtensions/Private/AudioExtensionsModule.cpp
--- a/Engine/Source/Runtime/AudioExtensions/Private/AudioExtensionsModule.cpp
+++ b/Engine/Source/Runtime/AudioExtensions/Private/AudioExtensionsModule.cpp
@@ -2,6 +2,7 @@
 
 #include "Modules/ModuleManager.h"
 
+#include "AudioExtensionsLog.h"
 #include "HAL/LowLevelMemStats.h"
 #include "IAudioCodecRegistry.h"
 #include "IAudioExtensionPlugin.h"
@@ -10,6 +11,8 @@
 DECLARE_LLM_MEMORY_STAT(TEXT("AudioSpatializationPlugins"), STAT_AudioSpatializationPluginsLLM, STATGROUP_LLMFULL);
 LLM_DEFINE_TAG(Audio_SpatializationPlugins, NAME_None, TEXT("Audio"), GET_STATFNAME(STAT_AudioSpatializationPluginsLLM), GET_STATFNAME(STAT_AudioSummaryLLM));
 
+DEFINE_LOG_CATEGORY(LogAudioExtensions)
+
 class FAudioExtensionsModule : public IModuleInterface
 {
 public:
diff --git a/Engine/Source/Runtime/AudioExtensions/Private/AudioParameter.cpp b/Engine/Source/Runtime/AudioExtensions/Private/AudioParameter.cpp
--- a/Engine/Source/Runtime/AudioExtensions/Private/AudioParameter.cpp
+++ b/Engine/Source/Runtime/AudioExtensions/Private/AudioParameter.cpp
@@ -129,11 +129,11 @@ void FAudioParameter::Merge(const FAudioParameter& InParameter, bool bInTakeName
 				ObjectProxies.Reset();
 			}
 
-			for (const Audio::IProxyDataPtr& ProxyPtr : InParameter.ObjectProxies)
+			for (const TSharedPtr<Audio::IProxyData>& ProxyPtr : InParameter.ObjectProxies)
 			{
 				if (ensure(ProxyPtr.IsValid()))
 				{
-					ObjectProxies.Emplace(ProxyPtr->Clone());
+					ObjectProxies.Emplace(ProxyPtr);
 				}
 			}
 		}
@@ -144,11 +144,11 @@ void FAudioParameter::Merge(const FAudioParameter& InParameter, bool bInTakeName
 			ObjectParam = InParameter.ObjectParam;
 
 			ObjectProxies.Reset();
-			for (const Audio::IProxyDataPtr& ProxyPtr : InParameter.ObjectProxies)
+			for (const TSharedPtr<Audio::IProxyData>& ProxyPtr : InParameter.ObjectProxies)
 			{
 				if (ensure(ProxyPtr.IsValid()))
 				{
-					ObjectProxies.Emplace(ProxyPtr->Clone());
+					ObjectProxies.Emplace(ProxyPtr);
 				}
 			}
 		}
@@ -163,11 +163,11 @@ void FAudioParameter::Merge(const FAudioParameter& InParameter, bool bInTakeName
 				ObjectProxies.Reset();
 			}
 
-			for (const Audio::IProxyDataPtr& ProxyPtr : InParameter.ObjectProxies)
+			for (const TSharedPtr<Audio::IProxyData>& ProxyPtr : InParameter.ObjectProxies)
 			{
 				if (ensure(ProxyPtr.IsValid()))
 				{
-					ObjectProxies.Emplace(ProxyPtr->Clone());
+					ObjectProxies.Emplace(ProxyPtr);
 				}
 			}
 		}
diff --git a/Engine/Source/Runtime/AudioExtensions/Private/IAudioModulation.cpp b/Engine/Source/Runtime/AudioExtensions/Private/IAudioModulation.cpp
--- a/Engine/Source/Runtime/AudioExtensions/Private/IAudioModulation.cpp
+++ b/Engine/Source/Runtime/AudioExtensions/Private/IAudioModulation.cpp
@@ -379,11 +379,11 @@ const Audio::FModulationParameter& USoundModulatorBase::GetOutputParameter() con
 	return Audio::GetModulationParameter({ });
 }
 
-TUniquePtr<Audio::IProxyData> USoundModulatorBase::CreateNewProxyData(const Audio::FProxyDataInitParams& InitParams)
+TSharedPtr<Audio::IProxyData> USoundModulatorBase::CreateProxyData(const Audio::FProxyDataInitParams& InitParams)
 {
 	// This should never be hit as all instances of modulators should implement their own version of the proxy data interface.
 	checkNoEntry();
-	return TUniquePtr<Audio::IProxyData>();
+	return TSharedPtr<Audio::IProxyData>();
 }
 
 TUniquePtr<Audio::IModulatorSettings> USoundModulatorBase::CreateProxySettings() const
diff --git a/Engine/Source/Runtime/AudioExtensions/Private/IAudioProxyInitializer.cpp b/Engine/Source/Runtime/AudioExtensions/Private/IAudioProxyInitializer.cpp
new file mode 100644
--- /dev/null
+++ b/Engine/Source/Runtime/AudioExtensions/Private/IAudioProxyInitializer.cpp
@@ -0,0 +1,38 @@
+// Copyright Epic Games, Inc. All Rights Reserved.
+
+#include "IAudioProxyInitializer.h"
+
+#include "AudioExtensionsLog.h"
+#include "Logging/LogMacros.h"
+#include "Containers/Set.h"
+#include "Templates/SharedPointer.h"
+#include "Templates/UniquePtr.h"
+#include "UObject/NameTypes.h"
+
+TUniquePtr<Audio::IProxyData> IAudioProxyDataFactory::CreateNewProxyData(const Audio::FProxyDataInitParams& InitParams) 
+{
+	return nullptr;
+}
+
+TSharedPtr<Audio::IProxyData> IAudioProxyDataFactory::CreateProxyData(const Audio::FProxyDataInitParams& InitParams) 
+{
+	// If the new interface is not overridden, use deprecated interface
+	PRAGMA_DISABLE_DEPRECATION_WARNINGS
+	TUniquePtr<Audio::IProxyData> Proxy = CreateNewProxyData(InitParams);
+	PRAGMA_ENABLE_DEPRECATION_WARNINGS
+
+	// Only log an error once per a proxy data type to avoid logspam.
+	static TSet<FName> LoggedProxyWarnings;
+	if (Proxy.IsValid())
+	{
+		const FName ProxyTypeName = Proxy->GetProxyTypeName();
+		if (!LoggedProxyWarnings.Contains(ProxyTypeName))
+		{
+			LoggedProxyWarnings.Add(ProxyTypeName);
+			UE_LOG(LogAudioExtensions, Warning, TEXT("Use of deprecated 'TUniquePtr<Audio::IProxyData> CreateNewProxyData(...)' for proxy type \"%s\". Please override 'TSharedPtr<Audio::IProxyData> CreateProxyData(...)'"), *ProxyTypeName.ToString());
+		}
+	}
+
+	// Pass ownership of proxy from unique ptr to shared ptr. 
+	return TSharedPtr<Audio::IProxyData>(Proxy.Release());
+}
diff --git a/Engine/Source/Runtime/AudioExtensions/Public/AudioExtensionsLog.h b/Engine/Source/Runtime/AudioExtensions/Public/AudioExtensionsLog.h
new file mode 100644
--- /dev/null
+++ b/Engine/Source/Runtime/AudioExtensions/Public/AudioExtensionsLog.h
@@ -0,0 +1,7 @@
+// Copyright Epic Games, Inc. All Rights Reserved.
+
+#pragma once
+
+#include "Logging/LogMacros.h"
+
+AUDIOEXTENSIONS_API DECLARE_LOG_CATEGORY_EXTERN(LogAudioExtensions, Display, All);
diff --git a/Engine/Source/Runtime/AudioExtensions/Public/AudioParameter.h b/Engine/Source/Runtime/AudioExtensions/Public/AudioParameter.h
--- a/Engine/Source/Runtime/AudioExtensions/Public/AudioParameter.h
+++ b/Engine/Source/Runtime/AudioExtensions/Public/AudioParameter.h
@@ -222,91 +222,11 @@ struct AUDIOEXTENSIONS_API FAudioParameter
 	// by InParam.
 	static void Merge(TArray<FAudioParameter>&& InParams, TArray<FAudioParameter>& OutParams);
 
-	FAudioParameter(FAudioParameter&& InParameter)
-		: ParamName(MoveTemp(InParameter.ParamName))
-		, FloatParam(MoveTemp(InParameter.FloatParam))
-		, BoolParam(MoveTemp(InParameter.BoolParam))
-		, IntParam(MoveTemp(InParameter.IntParam))
-		, ObjectParam(MoveTemp(InParameter.ObjectParam))
-		, StringParam(MoveTemp(InParameter.StringParam))
-		, ArrayFloatParam(MoveTemp(InParameter.ArrayFloatParam))
-		, ArrayBoolParam(MoveTemp(InParameter.ArrayBoolParam))
-		, ArrayIntParam(MoveTemp(InParameter.ArrayIntParam))
-		, ArrayObjectParam(MoveTemp(InParameter.ArrayObjectParam))
-		, ArrayStringParam(MoveTemp(InParameter.ArrayStringParam))
-		, ParamType(MoveTemp(InParameter.ParamType))
-		, TypeName(MoveTemp(InParameter.TypeName))
-		, ObjectProxies(MoveTemp(InParameter.ObjectProxies))		
-	{
-	}
-
-	FAudioParameter(const FAudioParameter& InParameter)
-		: ParamName(InParameter.ParamName)
-		, FloatParam(InParameter.FloatParam)
-		, BoolParam(InParameter.BoolParam)
-		, IntParam(InParameter.IntParam)
-		, ObjectParam(InParameter.ObjectParam)
-		, StringParam(InParameter.StringParam)
-		, ArrayFloatParam(InParameter.ArrayFloatParam)
-		, ArrayBoolParam(InParameter.ArrayBoolParam)
-		, ArrayIntParam(InParameter.ArrayIntParam)
-		, ArrayObjectParam(InParameter.ArrayObjectParam)
-		, ArrayStringParam(InParameter.ArrayStringParam)
-		, ParamType(InParameter.ParamType)
-		, TypeName(InParameter.TypeName)
-	{
-		ObjectProxies.Reset();
-		for (const Audio::IProxyDataPtr& DataPtr : InParameter.ObjectProxies)
-		{
-			ObjectProxies.Emplace(DataPtr->Clone());
-		}
-	}
-
-	FAudioParameter& operator=(const FAudioParameter& InParameter)
-	{
-		ParamName = InParameter.ParamName;
-		FloatParam = InParameter.FloatParam;
-		BoolParam = InParameter.BoolParam;
-		IntParam = InParameter.IntParam;
-		ObjectParam = InParameter.ObjectParam;
-		StringParam = InParameter.StringParam;
-		ArrayFloatParam = InParameter.ArrayFloatParam;
-		ArrayBoolParam = InParameter.ArrayBoolParam;
-		ArrayIntParam = InParameter.ArrayIntParam;
-		ArrayObjectParam = InParameter.ArrayObjectParam;
-		ArrayStringParam = InParameter.ArrayStringParam;
-		ParamType = InParameter.ParamType;
-		TypeName = InParameter.TypeName;
-
-		ObjectProxies.Reset();
-		for (const Audio::IProxyDataPtr& DataPtr : InParameter.ObjectProxies)
-		{
-			ObjectProxies.Emplace(DataPtr->Clone());
-		}
-
-		return *this;
-	}
+	FAudioParameter(FAudioParameter&& InParameter) = default;
+	FAudioParameter(const FAudioParameter& InParameter) = default;
 
-	FAudioParameter& operator=(FAudioParameter&& InParameter)
-	{
-		ParamName = MoveTemp(InParameter.ParamName);
-		FloatParam = MoveTemp(InParameter.FloatParam);
-		BoolParam = MoveTemp(InParameter.BoolParam);
-		IntParam = MoveTemp(InParameter.IntParam);
-		ObjectParam = MoveTemp(InParameter.ObjectParam);
-		StringParam = MoveTemp(InParameter.StringParam);
-		ArrayFloatParam = MoveTemp(InParameter.ArrayFloatParam);
-		ArrayBoolParam = MoveTemp(InParameter.ArrayBoolParam);
-		ArrayIntParam = MoveTemp(InParameter.ArrayIntParam);
-		ArrayObjectParam = MoveTemp(InParameter.ArrayObjectParam);
-		ArrayStringParam = MoveTemp(InParameter.ArrayStringParam);
-		ParamType = MoveTemp(InParameter.ParamType);
-		TypeName = MoveTemp(InParameter.TypeName);
-
-		ObjectProxies = MoveTemp(InParameter.ObjectProxies);
-
-		return *this;
-	}
+	FAudioParameter& operator=(const FAudioParameter& InParameter) = default;
+	FAudioParameter& operator=(FAudioParameter&& InParameter) = default;
 
 	// Name of the parameter
 	UPROPERTY(EditAnywhere, BlueprintReadWrite, meta = (DisplayName="Name"), Category = AudioParameter)
@@ -360,7 +280,7 @@ struct AUDIOEXTENSIONS_API FAudioParameter
 	FName TypeName;
 
 	// Object proxies to be generated when parameter is passed to the AudioThread to represent ObjectParam/ArrayObjectParam safely
-	TArray<Audio::IProxyDataPtr> ObjectProxies;
+	TArray<TSharedPtr<Audio::IProxyData>> ObjectProxies;
 
 	// Common find algorithm for default/legacy parameter system
 	static const FAudioParameter* FindParam(const TArray<FAudioParameter>& InParams, FName InParamName)
@@ -405,4 +325,4 @@ struct AUDIOEXTENSIONS_API FAudioParameter
 
 		return Param;
 	}
-};
\ No newline at end of file
+};
diff --git a/Engine/Source/Runtime/AudioExtensions/Public/IAudioModulation.h b/Engine/Source/Runtime/AudioExtensions/Public/IAudioModulation.h
--- a/Engine/Source/Runtime/AudioExtensions/Public/IAudioModulation.h
+++ b/Engine/Source/Runtime/AudioExtensions/Public/IAudioModulation.h
@@ -215,7 +215,7 @@ class AUDIOEXTENSIONS_API USoundModulatorBase : public UObject, public IAudioPro
 public:
 	virtual const Audio::FModulationParameter& GetOutputParameter() const;
 
-	virtual TUniquePtr<Audio::IProxyData> CreateNewProxyData(const Audio::FProxyDataInitParams& InitParams) override;
+	virtual TSharedPtr<Audio::IProxyData> CreateProxyData(const Audio::FProxyDataInitParams& InitParams) override;
 
 	virtual TUniquePtr<Audio::IModulatorSettings> CreateProxySettings() const;
 };
@@ -240,11 +240,6 @@ public:
 	{
 	}
 
-	virtual Audio::IProxyDataPtr Clone() const
-	{
-		return MakeUnique<FSoundModulatorAssetProxy>(*this);
-	}
-
 	virtual Audio::FModulatorHandle CreateModulatorHandle(IAudioModulationManager& InModulation) const
 	{
 		check(ModulatorSettings.IsValid());
@@ -273,11 +268,6 @@ class AUDIOEXTENSIONS_API FSoundModulationParameterAssetProxy : public Audio::TP
 public:
 	IMPL_AUDIOPROXY_CLASS(FSoundModulationParameterAssetProxy);
 
-	virtual Audio::IProxyDataPtr Clone() const override
-	{
-		return Audio::IProxyDataPtr();
-	}
-
 	virtual const Audio::FModulationParameter& GetParameter() const
 	{
 		return Parameter;
diff --git a/Engine/Source/Runtime/AudioExtensions/Public/IAudioProxyInitializer.h b/Engine/Source/Runtime/AudioExtensions/Public/IAudioProxyInitializer.h
--- a/Engine/Source/Runtime/AudioExtensions/Public/IAudioProxyInitializer.h
+++ b/Engine/Source/Runtime/AudioExtensions/Public/IAudioProxyInitializer.h
@@ -2,9 +2,14 @@
 
 #pragma once
 
-#include "CoreMinimal.h"
+#include "Misc/AssertionMacros.h"
+#include "Misc/CoreMiscDefines.h"
 #include "Templates/Casts.h"
-#include "Templates/UnrealTypeTraits.h"
+#include "Templates/SharedPointer.h"
+#include "Templates/UniquePtr.h"
+#include "UObject/NameTypes.h"
+
+#include <type_traits>
 
 /**
  * Interfaces for Audio Proxy Objects 
@@ -28,7 +33,7 @@ namespace Audio
 {
 	// Forward Declarations
 	class IProxyData;
-	using IProxyDataPtr = TUniquePtr<IProxyData>;
+	using IProxyDataPtr UE_DEPRECATED(5.2, "Replace IProxyDataPtr with TSharedPtr<Audio::IProxyData>") = TUniquePtr<Audio::IProxyData>;
 
 	/*
 	 * Base class that allows us to typecheck proxy data before downcasting it in debug builds.
@@ -55,7 +60,7 @@ namespace Audio
 		template<typename ProxyType>
 		ProxyType& GetAs()
 		{
-			static_assert(TIsDerivedFrom<ProxyType, IProxyData>::Value, "Tried to downcast IProxyInitData to an unrelated type!");
+			static_assert(std::is_base_of_v<IProxyData, ProxyType>, "Tried to downcast IProxyInitData to an unrelated type!");
 			if (CheckTypeCast<ProxyType>())
 			{
 				return static_cast<ProxyType&>(*this);
@@ -71,7 +76,7 @@ namespace Audio
 		template<typename ProxyType>
 		const ProxyType& GetAs() const
 		{
-			static_assert(TIsDerivedFrom<ProxyType, IProxyData>::Value, "Tried to downcast IProxyInitData to an unrelated type!");
+			static_assert(std::is_base_of_v<IProxyData, ProxyType>, "Tried to downcast IProxyInitData to an unrelated type!");
 			if (CheckTypeCast<ProxyType>())
 			{
 				return static_cast<const ProxyType&>(*this);
@@ -88,7 +93,8 @@ namespace Audio
 			: ProxyTypeName(InProxyTypeName)
 		{}
 
-		virtual IProxyDataPtr Clone() const = 0;
+		UE_DEPRECATED(5.2, "Proxy data is stored in a TSharedPtr<> and no longer requires cloning")
+		virtual TUniquePtr<IProxyData> Clone() const { return nullptr; }
 	};
 
 	/**
@@ -109,7 +115,6 @@ namespace Audio
 		}
 	};
 
-
 	struct FProxyDataInitParams
 	{
 		FName NameOfFeatureRequestingProxy;
@@ -122,13 +127,16 @@ namespace Audio
 class AUDIOEXTENSIONS_API IAudioProxyDataFactory
 {
 public:
-	virtual TUniquePtr<Audio::IProxyData> CreateNewProxyData(const Audio::FProxyDataInitParams& InitParams) = 0;
+	UE_DEPRECATED(5.2, "Call TSharedPtr<Audio::IProxyData> CreateProxyData(...) instead of a TUniquePtr<Audio::IProxyData> CreateNewProxyData(...).")
+	virtual TUniquePtr<Audio::IProxyData> CreateNewProxyData(const Audio::FProxyDataInitParams& InitParams);
+
+	virtual TSharedPtr<Audio::IProxyData> CreateProxyData(const Audio::FProxyDataInitParams& InitParams);
 };
 
 namespace Audio
 {
 	// SFINAE used to optionally invoke subclasses of IAudioProxyDataFactory when we can.
-	template<typename UClassToUse, typename TEnableIf<TIsDerivedFrom<UClassToUse, IAudioProxyDataFactory>::Value, bool>::Type = true>
+	template<typename UClassToUse, typename std::enable_if_t<std::is_base_of_v<IAudioProxyDataFactory, UClassToUse>, bool> = true>
 	IAudioProxyDataFactory* CastToProxyDataFactory(UObject* InObject)
 	{
 		if (InObject)
@@ -143,7 +151,7 @@ namespace Audio
 		return nullptr;
 	}
 
-	template<typename UClassToUse, typename TEnableIf<!TIsDerivedFrom<UClassToUse, IAudioProxyDataFactory>::Value, bool>::Type = true>
+	template<typename UClassToUse, typename std::enable_if_t<!std::is_base_of_v<IAudioProxyDataFactory, UClassToUse>, bool>::Type = true>
 	IAudioProxyDataFactory* CastToProxyDataFactory(UObject* InObject)
 	{
 		return nullptr;
diff --git a/Engine/Source/Runtime/Engine/Classes/Sound/AudioBus.h b/Engine/Source/Runtime/Engine/Classes/Sound/AudioBus.h
--- a/Engine/Source/Runtime/Engine/Classes/Sound/AudioBus.h
+++ b/Engine/Source/Runtime/Engine/Classes/Sound/AudioBus.h
@@ -33,11 +33,6 @@ public:
 
 	virtual ~FAudioBusProxy() override {}
 
-	virtual Audio::IProxyDataPtr Clone() const override
-	{
-		return MakeUnique<FAudioBusProxy>(*this);
-	}
-
 	uint32 AudioBusId = INDEX_NONE;
 	int32 NumChannels = INDEX_NONE;
 };
@@ -72,6 +67,6 @@ public:
 	//~ End UObject Interface.
 
 	//~ Begin IAudioProxy Interface
-	virtual TUniquePtr<Audio::IProxyData> CreateNewProxyData(const Audio::FProxyDataInitParams& InitParams) override;
+	virtual TSharedPtr<Audio::IProxyData> CreateProxyData(const Audio::FProxyDataInitParams& InitParams) override;
 	//~ End IAudioProxy Interface
 };
diff --git a/Engine/Source/Runtime/Engine/Classes/Sound/SoundWave.h b/Engine/Source/Runtime/Engine/Classes/Sound/SoundWave.h
--- a/Engine/Source/Runtime/Engine/Classes/Sound/SoundWave.h
+++ b/Engine/Source/Runtime/Engine/Classes/Sound/SoundWave.h
@@ -472,7 +472,7 @@ public:
 
 private:
 	// cached proxy
-	FSoundWaveProxyPtr InternalProxy{ nullptr };
+	FSoundWaveProxyPtr Proxy{ nullptr };
 
 public:
 
@@ -906,7 +906,7 @@ public:
 	FSoundWaveProxyPtr CreateSoundWaveProxy();
 
 	//~Begin IAudioProxyDataFactory Interface.
-	virtual TUniquePtr<Audio::IProxyData> CreateNewProxyData(const Audio::FProxyDataInitParams& InitParams) override;
+	virtual TSharedPtr<Audio::IProxyData> CreateProxyData(const Audio::FProxyDataInitParams& InitParams) override;
 	//~ End IAudioProxyDataFactory Interface.
 
 	// Called  when the procedural sound wave begins on the render thread. Only used in the audio mixer and when bProcedural is true.
@@ -1386,12 +1386,6 @@ public:
 
 	~FSoundWaveProxy();
 
-	Audio::IProxyDataPtr Clone() const override
-	{
-		LLM_SCOPE(ELLMTag::AudioSoundWaveProxies);
-		return MakeUnique<FSoundWaveProxy>(*this);
-	}
-
 	// USoundWave Interface
 	void ReleaseCompressedAudio();
 
diff --git a/Engine/Source/Runtime/Engine/Private/AudioBus.cpp b/Engine/Source/Runtime/Engine/Private/AudioBus.cpp
--- a/Engine/Source/Runtime/Engine/Private/AudioBus.cpp
+++ b/Engine/Source/Runtime/Engine/Private/AudioBus.cpp
@@ -78,8 +78,8 @@ void UAudioBus::PostEditChangeProperty(FPropertyChangedEvent& PropertyChangedEve
 }
 #endif // WITH_EDITOR
 
- TUniquePtr<Audio::IProxyData> UAudioBus::CreateNewProxyData(const Audio::FProxyDataInitParams& InitParams)
+TSharedPtr<Audio::IProxyData> UAudioBus::CreateProxyData(const Audio::FProxyDataInitParams& InitParams)
 {
-	return MakeUnique<FAudioBusProxy>(this);
+	return MakeShared<FAudioBusProxy>(this);
 }
 
diff --git a/Engine/Source/Runtime/Engine/Private/SoundWave.cpp b/Engine/Source/Runtime/Engine/Private/SoundWave.cpp
--- a/Engine/Source/Runtime/Engine/Private/SoundWave.cpp
+++ b/Engine/Source/Runtime/Engine/Private/SoundWave.cpp
@@ -1121,11 +1121,11 @@ void USoundWave::Serialize( FArchive& Ar )
 				SoundWaveDataPtr->LoadingBehavior = ESoundWaveLoadingBehavior::PrimeOnLoad;
 				SoundWaveDataPtr->bLoadingBehaviorOverridden = true;
 
-				if (!InternalProxy.IsValid())
+				if (!Proxy.IsValid())
 				{
-					InternalProxy = CreateSoundWaveProxy();
+					Proxy = CreateSoundWaveProxy();
 				}
-				IStreamingManager::Get().GetAudioStreamingManager().RequestChunk(InternalProxy, 1, [](EAudioChunkLoadResult) {});
+				IStreamingManager::Get().GetAudioStreamingManager().RequestChunk(Proxy, 1, [](EAudioChunkLoadResult) {});
 			}
 		}
 	}
@@ -1712,9 +1712,9 @@ void USoundWave::PostLoad()
 	ESoundWaveLoadingBehavior ActualLoadingBehavior = GetLoadingBehavior();
 
 	// no need for a proxy if we can't render audio!
-	if (!InternalProxy.IsValid() && ActualLoadingBehavior != ESoundWaveLoadingBehavior::ForceInline && FApp::CanEverRenderAudio())
+	if (!Proxy.IsValid() && ActualLoadingBehavior != ESoundWaveLoadingBehavior::ForceInline && FApp::CanEverRenderAudio())
 	{
-		InternalProxy = CreateSoundWaveProxy();
+		Proxy = CreateSoundWaveProxy();
 	}
 
 	if (FApp::CanEverRenderAudio() && ShouldUseStreamCaching() && ActualLoadingBehavior != GetLoadingBehavior(false))
@@ -1735,7 +1735,7 @@ void USoundWave::PostLoad()
 			
 			if (bShouldPrime && bHasMultipleChunks)
 			{
-				IStreamingManager::Get().GetAudioStreamingManager().RequestChunk(InternalProxy, 1, [](EAudioChunkLoadResult) {});
+				IStreamingManager::Get().GetAudioStreamingManager().RequestChunk(Proxy, 1, [](EAudioChunkLoadResult) {});
 			}
 		}
 
@@ -1802,7 +1802,7 @@ void USoundWave::PostLoad()
 #endif // #if WITH_EDITORONLY_DATA
 		if (!ShouldUseStreamCaching())
 		{
-			IStreamingManager::Get().GetAudioStreamingManager().AddStreamingSoundWave(InternalProxy);
+			IStreamingManager::Get().GetAudioStreamingManager().AddStreamingSoundWave(Proxy);
 		}
 
 		// Only request loading the zeroth chunk when streaming is supported and we can render audio.
@@ -2703,9 +2703,9 @@ bool USoundWave::IsReadyForFinishDestroy()
 	bool bIsStreamingInProgress = false;
 	if (CurrentLoadingBehavior != ESoundWaveLoadingBehavior::ForceInline)
 	{
-		if (InternalProxy.IsValid())
+		if (Proxy.IsValid())
 		{
-			bIsStreamingInProgress = IStreamingManager::Get().GetAudioStreamingManager().IsStreamingInProgress(InternalProxy);
+			bIsStreamingInProgress = IStreamingManager::Get().GetAudioStreamingManager().IsStreamingInProgress(Proxy);
 		}
 	}
 
@@ -3221,11 +3221,11 @@ TArrayView<const uint8> USoundWave::GetZerothChunk(bool bForImmediatePlayback)
 		if (GetNumChunks() > 1)
 		{
 			// Prime first chunk for playback.
-			if (!InternalProxy.IsValid())
+			if (!Proxy.IsValid())
 			{
-				InternalProxy = CreateSoundWaveProxy();
+				Proxy = CreateSoundWaveProxy();
 			}
-			IStreamingManager::Get().GetAudioStreamingManager().RequestChunk(InternalProxy, 1, [](EAudioChunkLoadResult InResult) {}, ENamedThreads::AnyThread, bForImmediatePlayback);
+			IStreamingManager::Get().GetAudioStreamingManager().RequestChunk(Proxy, 1, [](EAudioChunkLoadResult InResult) {}, ENamedThreads::AnyThread, bForImmediatePlayback);
 		}
 
 		FBulkDataBuffer<uint8>::ViewType View = SoundWaveDataPtr->GetZerothChunkData().GetView();
@@ -3234,11 +3234,11 @@ TArrayView<const uint8> USoundWave::GetZerothChunk(bool bForImmediatePlayback)
 	}
 	else
 	{
-		if (!InternalProxy.IsValid())
+		if (!Proxy.IsValid())
 		{
-			InternalProxy = CreateSoundWaveProxy();
+			Proxy = CreateSoundWaveProxy();
 		}
-		FAudioChunkHandle ChunkHandle = IStreamingManager::Get().GetAudioStreamingManager().GetLoadedChunk(InternalProxy, 0);
+		FAudioChunkHandle ChunkHandle = IStreamingManager::Get().GetAudioStreamingManager().GetLoadedChunk(Proxy, 0);
 		return TArrayView<const uint8>(ChunkHandle.GetData(), ChunkHandle.Num());
 	}
 }
@@ -3305,22 +3305,16 @@ FSoundWaveProxyPtr USoundWave::CreateSoundWaveProxy()
 	return MakeShared<FSoundWaveProxy, ESPMode::ThreadSafe>(this);
 }
 
-TUniquePtr<Audio::IProxyData> USoundWave::CreateNewProxyData(const Audio::FProxyDataInitParams& InitParams)
+TSharedPtr<Audio::IProxyData> USoundWave::CreateProxyData(const Audio::FProxyDataInitParams& InitParams)
 {
 	check(SoundWaveDataPtr);
 	LLM_SCOPE(ELLMTag::AudioSoundWaveProxies);
 
-#if WITH_EDITORONLY_DATA
-	if (!LoadZerothChunk())
+	if (!Proxy.IsValid())
 	{
-		return nullptr;
+		Proxy = CreateSoundWaveProxy();
 	}
-#endif // #if WITH_EDITORONLY_DATA
-
-	check(SoundWaveDataPtr);
-	SoundWaveDataPtr->InitializeDataFromSoundWave(*this);
-
-	return MakeUnique<FSoundWaveProxy>(this);
+	return Proxy;
 }
 
 void USoundWave::AddPlayingSource(const FSoundWaveClientPtr& Source)
@@ -3349,13 +3343,13 @@ void USoundWave::UpdatePlatformData()
 {
 	if (IsStreaming(nullptr) && FApp::CanEverRenderAudio())
 	{
-		if (!InternalProxy.IsValid())
+		if (!Proxy.IsValid())
 		{
-			InternalProxy = CreateSoundWaveProxy();
+			Proxy = CreateSoundWaveProxy();
 		}
 
 		// Make sure there are no pending requests in flight.
-		while (IStreamingManager::Get().GetAudioStreamingManager().IsStreamingInProgress(InternalProxy))
+		while (IStreamingManager::Get().GetAudioStreamingManager().IsStreamingInProgress(Proxy))
 		{
 			// Give up timeslice.
 			FPlatformProcess::Sleep(0);
@@ -3363,18 +3357,18 @@ void USoundWave::UpdatePlatformData()
 
 #if WITH_EDITORONLY_DATA
 		// Temporarily remove from streaming manager to release currently used data chunks
-		IStreamingManager::Get().GetAudioStreamingManager().RemoveStreamingSoundWave(InternalProxy);
+		IStreamingManager::Get().GetAudioStreamingManager().RemoveStreamingSoundWave(Proxy);
 
 		// Recache platform data if the source has changed.
 		CachePlatformData(true /* bAsyncCache */);
 
 		// Add back to the streaming manager to reload first chunk
-		IStreamingManager::Get().GetAudioStreamingManager().AddStreamingSoundWave(InternalProxy);
+		IStreamingManager::Get().GetAudioStreamingManager().AddStreamingSoundWave(Proxy);
 #endif
 	}
-	else if (InternalProxy.IsValid())
+	else if (Proxy.IsValid())
 	{
-		IStreamingManager::Get().GetAudioStreamingManager().RemoveStreamingSoundWave(InternalProxy);
+		IStreamingManager::Get().GetAudioStreamingManager().RemoveStreamingSoundWave(Proxy);
 	}
 }
 
@@ -3705,9 +3699,9 @@ void USoundWave::GetHandleForChunkOfAudio(TFunction<void(FAudioChunkHandle&&)> O
 {
 	ENamedThreads::Type ThreadToDispatchCallbackOn = (DispatchToGameThreadOnChunkRequestCVar != 0) ? ENamedThreads::GameThread : ENamedThreads::AnyThread;
 
-	if (!InternalProxy.IsValid() && FApp::CanEverRenderAudio())
+	if (!Proxy.IsValid() && FApp::CanEverRenderAudio())
 	{
-		InternalProxy = CreateSoundWaveProxy();
+		Proxy = CreateSoundWaveProxy();
 	}
 
 	// if we are requesting a chunk that is out of bounds,
@@ -3720,7 +3714,7 @@ void USoundWave::GetHandleForChunkOfAudio(TFunction<void(FAudioChunkHandle&&)> O
 	else if (bForceSync)
 	{
 		// For sync cases, we call GetLoadedChunk with bBlockForLoad = true, then execute the callback immediately.
-		FAudioChunkHandle ChunkHandle = IStreamingManager::Get().GetAudioStreamingManager().GetLoadedChunk(InternalProxy, ChunkIndex, true);
+		FAudioChunkHandle ChunkHandle = IStreamingManager::Get().GetAudioStreamingManager().GetLoadedChunk(Proxy, ChunkIndex, true);
 		OnLoadCompleted(MoveTemp(ChunkHandle));
 	}
 	else
@@ -3728,7 +3722,7 @@ void USoundWave::GetHandleForChunkOfAudio(TFunction<void(FAudioChunkHandle&&)> O
 		TWeakObjectPtr<USoundWave> WeakThis = MakeWeakObjectPtr(this);
 
 		// For async cases, we call RequestChunk and request the loaded chunk in the completion callback.
-		IStreamingManager::Get().GetAudioStreamingManager().RequestChunk(InternalProxy, ChunkIndex, [ThreadToDispatchCallbackOn, WeakThis, OnLoadCompleted, ChunkIndex, CallbackThread](EAudioChunkLoadResult LoadResult)
+		IStreamingManager::Get().GetAudioStreamingManager().RequestChunk(Proxy, ChunkIndex, [ThreadToDispatchCallbackOn, WeakThis, OnLoadCompleted, ChunkIndex, CallbackThread](EAudioChunkLoadResult LoadResult)
 		{
 			auto DispatchOnLoadCompletedCallback = [ThreadToDispatchCallbackOn, OnLoadCompleted, CallbackThread](FAudioChunkHandle&& InHandle)
 			{
@@ -3750,7 +3744,7 @@ void USoundWave::GetHandleForChunkOfAudio(TFunction<void(FAudioChunkHandle&&)> O
 			if (WeakThis.IsValid() && (LoadResult == EAudioChunkLoadResult::Completed || LoadResult == EAudioChunkLoadResult::AlreadyLoaded))
 			{
 				USoundWave* ThisSoundWave = WeakThis.Get();
-				FAudioChunkHandle ChunkHandle = IStreamingManager::Get().GetAudioStreamingManager().GetLoadedChunk(ThisSoundWave->InternalProxy, ChunkIndex, (BlockOnChunkLoadCompletionCVar != 0));
+				FAudioChunkHandle ChunkHandle = IStreamingManager::Get().GetAudioStreamingManager().GetLoadedChunk(ThisSoundWave->Proxy, ChunkIndex, (BlockOnChunkLoadCompletionCVar != 0));
 
 				// If we hit this, something went wrong in GetLoadedChunk.
 				if (!ChunkHandle.IsValid())
@@ -3788,12 +3782,12 @@ void USoundWave::RetainCompressedAudio(bool bForceSync /*= false*/)
 	}
 	else if (bForceSync)
 	{
-		if (!InternalProxy.IsValid())
+		if (!Proxy.IsValid())
 		{
-			InternalProxy = CreateSoundWaveProxy();
+			Proxy = CreateSoundWaveProxy();
 		}
 
-		SoundWaveDataPtr->FirstChunk = IStreamingManager::Get().GetAudioStreamingManager().GetLoadedChunk(InternalProxy, 1, true);
+		SoundWaveDataPtr->FirstChunk = IStreamingManager::Get().GetAudioStreamingManager().GetLoadedChunk(Proxy, 1, true);
 		UE_CLOG(!(SoundWaveDataPtr->FirstChunk.IsValid()), LogAudio, Display, TEXT("First chunk was invalid after synchronous load in RetainCompressedAudio(). This was likely because the cache was blown. Sound: %s"), *GetFullName());
 	}
 	else
@@ -3877,19 +3871,19 @@ void USoundWave::OverrideLoadingBehavior(ESoundWaveLoadingBehavior InLoadingBeha
 				// In editor, just make sure that data is available in the local DDC cache for quick access.
 				CachePlatformData(true /* bAsyncCache */);
 
-				if (!InternalProxy.IsValid())
+				if (!Proxy.IsValid())
 				{
-					InternalProxy = CreateSoundWaveProxy();
+					Proxy = CreateSoundWaveProxy();
 				}
 #else
 				if (GetNumChunks() > 1)
 				{
-					if (!InternalProxy.IsValid())
+					if (!Proxy.IsValid())
 					{
-						InternalProxy = CreateSoundWaveProxy();
+						Proxy = CreateSoundWaveProxy();
 					}
 
-					IStreamingManager::Get().GetAudioStreamingManager().RequestChunk(InternalProxy, 1, [](EAudioChunkLoadResult) {});
+					IStreamingManager::Get().GetAudioStreamingManager().RequestChunk(Proxy, 1, [](EAudioChunkLoadResult) {});
 				}
 #endif
 			}
