Engine dependency cleanup around MeshBatch includes. A few things were moved to cpp files to reduce header dependencies. Implicit dependencies on these includes will be taken care of before this CL.

diff --git a/Engine/Source/Runtime/Engine/Classes/Engine/Scene.h b/Engine/Source/Runtime/Engine/Classes/Engine/Scene.h
--- a/Engine/Source/Runtime/Engine/Classes/Engine/Scene.h
+++ b/Engine/Source/Runtime/Engine/Classes/Engine/Scene.h
@@ -13,6 +13,7 @@
 #include "Engine/BlendableInterface.h"
 #if UE_ENABLE_INCLUDE_ORDER_DEPRECATED_IN_5_2
 #include "RHIDefinitions.h"
+#include "EngineDefines.h" // SDPG_NumBits moved to here
 #endif
 #include "SceneUtils.h"
 #include "Engine/EngineTypes.h"
@@ -2324,8 +2325,4 @@ UCLASS()
 class UScene : public UObject
 {
 	GENERATED_UCLASS_BODY()
-
-
-	/** bits needed to store DPG value */
-	#define SDPG_NumBits	3
 };
diff --git a/Engine/Source/Runtime/Engine/Private/MeshBatch.cpp b/Engine/Source/Runtime/Engine/Private/MeshBatch.cpp
new file mode 100644
--- /dev/null
+++ b/Engine/Source/Runtime/Engine/Private/MeshBatch.cpp
@@ -0,0 +1,40 @@
+// Copyright Epic Games, Inc. All Rights Reserved.
+
+#include "MeshBatch.h"
+#include "MaterialShared.h"
+
+bool FMeshBatch::IsTranslucent(ERHIFeatureLevel::Type InFeatureLevel) const
+{
+	// Note: blend mode does not depend on the feature level we are actually rendering in.
+	return IsTranslucentBlendMode(MaterialRenderProxy->GetIncompleteMaterialWithFallback(InFeatureLevel).GetBlendMode());
+}
+
+bool FMeshBatch::IsDecal(ERHIFeatureLevel::Type InFeatureLevel) const
+{
+	// Note: does not depend on the feature level we are actually rendering in.
+	const FMaterial& Mat = MaterialRenderProxy->GetIncompleteMaterialWithFallback(InFeatureLevel);
+	return Mat.IsDeferredDecal();
+}
+
+bool FMeshBatch::IsDualBlend(ERHIFeatureLevel::Type InFeatureLevel) const
+{
+	const FMaterial& Mat = MaterialRenderProxy->GetIncompleteMaterialWithFallback(InFeatureLevel);
+	return Mat.IsDualBlendingEnabled(GShaderPlatformForFeatureLevel[InFeatureLevel]);
+}
+
+bool FMeshBatch::UseForHairStrands(ERHIFeatureLevel::Type InFeatureLevel) const
+{
+	if (InFeatureLevel != ERHIFeatureLevel::SM5)
+	{
+		return false;
+	}
+
+	const FMaterial& Mat = MaterialRenderProxy->GetIncompleteMaterialWithFallback(InFeatureLevel);
+	return IsCompatibleWithHairStrands(&Mat, InFeatureLevel);
+}
+
+bool FMeshBatch::IsMasked(ERHIFeatureLevel::Type InFeatureLevel) const
+{
+	// Note: blend mode does not depend on the feature level we are actually rendering in.
+	return MaterialRenderProxy->GetIncompleteMaterialWithFallback(InFeatureLevel).IsMasked();
+}
diff --git a/Engine/Source/Runtime/Engine/Private/Particles/ParticleSystemRender.cpp b/Engine/Source/Runtime/Engine/Private/Particles/ParticleSystemRender.cpp
--- a/Engine/Source/Runtime/Engine/Private/Particles/ParticleSystemRender.cpp
+++ b/Engine/Source/Runtime/Engine/Private/Particles/ParticleSystemRender.cpp
@@ -6513,10 +6513,10 @@ FParticleSystemSceneProxy::FParticleSystemSceneProxy(UParticleSystemComponent* C
 		)
 	, DynamicData(InDynamicData)
 	, LastDynamicData(NULL)
-	, DeselectedWireframeMaterialInstance(
+	, DeselectedWireframeMaterialInstance(new FColoredMaterialRenderProxy(
 		GEngine->WireframeMaterial ? GEngine->WireframeMaterial->GetRenderProxy() : NULL,
 		GetSelectionColor(FLinearColor(1.0f, 0.0f, 0.0f, 1.0f),false,false)
-		)
+		))
 	, PendingLODDistance(0.0f)
 	, VisualizeLODIndex(Component->GetCurrentLODIndex())
 	, LastFramePreRendered(-1)
@@ -6550,6 +6550,9 @@ FParticleSystemSceneProxy::~FParticleSystemSceneProxy()
 
 	delete DynamicData;
 	DynamicData = NULL;
+
+	delete DeselectedWireframeMaterialInstance;
+	DeselectedWireframeMaterialInstance = nullptr;
 }
 
 FMeshBatch* FParticleSystemSceneProxy::GetPooledMeshBatch()
diff --git a/Engine/Source/Runtime/Engine/Private/SceneManagement.cpp b/Engine/Source/Runtime/Engine/Private/SceneManagement.cpp
--- a/Engine/Source/Runtime/Engine/Private/SceneManagement.cpp
+++ b/Engine/Source/Runtime/Engine/Private/SceneManagement.cpp
@@ -314,6 +314,74 @@ FMeshElementCollector::FMeshElementCollector(ERHIFeatureLevel::Type InFeatureLev
 {	
 }
 
+FMeshElementCollector::~FMeshElementCollector()
+{
+	DeleteTemporaryProxies();
+}
+
+void FMeshElementCollector::DeleteTemporaryProxies()
+{
+	check(!ParallelTasks.Num()); // We should have blocked on this already
+	for (int32 ProxyIndex = 0; ProxyIndex < TemporaryProxies.Num(); ProxyIndex++)
+	{
+		delete TemporaryProxies[ProxyIndex];
+	}
+
+	TemporaryProxies.Empty();
+}
+
+void FMeshElementCollector::SetPrimitive(const FPrimitiveSceneProxy* InPrimitiveSceneProxy, FHitProxyId DefaultHitProxyId)
+{
+	check(InPrimitiveSceneProxy);
+	PrimitiveSceneProxy = InPrimitiveSceneProxy;
+
+	for (int32 ViewIndex = 0; ViewIndex < SimpleElementCollectors.Num(); ViewIndex++)
+	{
+		SimpleElementCollectors[ViewIndex]->HitProxyId = DefaultHitProxyId;
+		SimpleElementCollectors[ViewIndex]->PrimitiveMeshId = 0;
+	}
+
+	for (int32 ViewIndex = 0; ViewIndex < MeshIdInPrimitivePerView.Num(); ++ViewIndex)
+	{
+		MeshIdInPrimitivePerView[ViewIndex] = 0;
+	}
+}
+
+void FMeshElementCollector::ClearViewMeshArrays()
+{
+	Views.Empty();
+	MeshBatches.Empty();
+	SimpleElementCollectors.Empty();
+	MeshIdInPrimitivePerView.Empty();
+	DynamicPrimitiveCollectorPerView.Empty();
+	NumMeshBatchElementsPerView.Empty();
+	DynamicIndexBuffer = nullptr;
+	DynamicVertexBuffer = nullptr;
+	DynamicReadBuffer = nullptr;
+}
+
+void FMeshElementCollector::AddViewMeshArrays(
+	FSceneView* InView,
+	TArray<FMeshBatchAndRelevance, SceneRenderingAllocator>* ViewMeshes,
+	FSimpleElementCollector* ViewSimpleElementCollector,
+	FGPUScenePrimitiveCollector* InDynamicPrimitiveCollector,
+	ERHIFeatureLevel::Type InFeatureLevel,
+	FGlobalDynamicIndexBuffer* InDynamicIndexBuffer,
+	FGlobalDynamicVertexBuffer* InDynamicVertexBuffer,
+	FGlobalDynamicReadBuffer* InDynamicReadBuffer)
+{
+	Views.Add(InView);
+	MeshIdInPrimitivePerView.Add(0);
+	MeshBatches.Add(ViewMeshes);
+	NumMeshBatchElementsPerView.Add(0);
+	SimpleElementCollectors.Add(ViewSimpleElementCollector);
+	DynamicPrimitiveCollectorPerView.Add(InDynamicPrimitiveCollector);
+
+	check(InDynamicIndexBuffer && InDynamicVertexBuffer && InDynamicReadBuffer);
+	DynamicIndexBuffer = InDynamicIndexBuffer;
+	DynamicVertexBuffer = InDynamicVertexBuffer;
+	DynamicReadBuffer = InDynamicReadBuffer;
+}
 
 void FMeshElementCollector::ProcessTasks()
 {
@@ -391,6 +459,12 @@ void FMeshElementCollector::AddMesh(int32 ViewIndex, FMeshBatch& MeshBatch)
 	new (ViewMeshBatches) FMeshBatchAndRelevance(MeshBatch, PrimitiveSceneProxy, FeatureLevel);	
 }
 
+FDynamicPrimitiveUniformBuffer::FDynamicPrimitiveUniformBuffer() = default;
+FDynamicPrimitiveUniformBuffer::~FDynamicPrimitiveUniformBuffer()
+{
+	UniformBuffer.ReleaseResource();
+}
+
 void FDynamicPrimitiveUniformBuffer::Set(
 	const FMatrix& LocalToWorld,
 	const FMatrix& PreviousLocalToWorld,
diff --git a/Engine/Source/Runtime/Engine/Public/EngineDefines.h b/Engine/Source/Runtime/Engine/Public/EngineDefines.h
--- a/Engine/Source/Runtime/Engine/Public/EngineDefines.h
+++ b/Engine/Source/Runtime/Engine/Public/EngineDefines.h
@@ -54,3 +54,6 @@
 #define UE_DOUBLE_HUGE_DISTANCE		562949953421312.0		/* Maximum distance representable by a double whilst maintaining precision of at least 0.0625 units (1/16th of a cm) - Precision issues may occur for positions/distances represented by double types that exceed this value */
 
 #define DEFAULT_ORTHOZOOM			10000.0					/* Default 2D viewport zoom */
+
+/** bits needed to store DPG value */
+#define SDPG_NumBits 3
diff --git a/Engine/Source/Runtime/Engine/Public/MeshBatch.h b/Engine/Source/Runtime/Engine/Public/MeshBatch.h
--- a/Engine/Source/Runtime/Engine/Public/MeshBatch.h
+++ b/Engine/Source/Runtime/Engine/Public/MeshBatch.h
@@ -2,20 +2,36 @@
 
 #pragma once
 
+#include "CoreTypes.h"
+#include "Containers/ArrayView.h"
+#include "EngineDefines.h"
+#include "GPUSceneWriter.h"
+#include "HitProxies.h"
+#include "InstanceUniformShaderParameters.h"
+#include "RHIDefinitions.h"
+#include "VT/RuntimeVirtualTextureEnum.h"
+
+#if UE_ENABLE_INCLUDE_ORDER_DEPRECATED_IN_5_2
 #include "CoreMinimal.h"
 #include "UniformBuffer.h"
-#include "HitProxies.h"
 #include "MaterialShared.h"
 #include "Engine/Scene.h"
 #include "PrimitiveUniformShaderParameters.h"
-#include "InstanceUniformShaderParameters.h"
-#include "VT/RuntimeVirtualTextureEnum.h"
-#include "EngineDefines.h"
-#include "GPUSceneWriter.h"
+#endif
 
 #define USE_MESH_BATCH_VALIDATION !UE_BUILD_SHIPPING
 
+class FIndexBuffer;
 class FLightCacheInterface;
+class FMaterial;
+class FMaterialRenderProxy;
+class FPrimitiveSceneProxy;
+class FPrimitiveUniformShaderParameters;
+class FVertexFactory;
+struct FInstanceDynamicData;
+struct FInstanceSceneData;
+struct FMaterialShaderParameters;
+struct FRenderBounds;
 
 enum EPrimitiveIdMode
 {
@@ -389,47 +405,16 @@ struct FMeshBatch
 	int8 VisualizeLODIndex;
 #endif
 
-	FORCEINLINE bool IsTranslucent(ERHIFeatureLevel::Type InFeatureLevel) const
-	{
-		// Note: blend mode does not depend on the feature level we are actually rendering in.
-		return IsTranslucentBlendMode(MaterialRenderProxy->GetIncompleteMaterialWithFallback(InFeatureLevel));
-	}
+	ENGINE_API bool IsTranslucent(ERHIFeatureLevel::Type InFeatureLevel) const;
 
 	// todo: can be optimized with a single function that returns multiple states (Translucent, Decal, Masked) 
-	FORCEINLINE bool IsDecal(ERHIFeatureLevel::Type InFeatureLevel) const
-	{
-		// Note: does not depend on the feature level we are actually rendering in.
-		const FMaterial& Mat = MaterialRenderProxy->GetIncompleteMaterialWithFallback(InFeatureLevel);
-		return Mat.IsDeferredDecal();
-	}
+	ENGINE_API bool IsDecal(ERHIFeatureLevel::Type InFeatureLevel) const;
 
-	FORCEINLINE bool IsDualBlend(ERHIFeatureLevel::Type InFeatureLevel) const
-	{
-		const FMaterial& Mat = MaterialRenderProxy->GetIncompleteMaterialWithFallback(InFeatureLevel);
-		return Mat.IsDualBlendingEnabled(GShaderPlatformForFeatureLevel[InFeatureLevel]);
-	}
+	ENGINE_API bool IsDualBlend(ERHIFeatureLevel::Type InFeatureLevel) const;
 
-	FORCEINLINE bool UseForHairStrands(ERHIFeatureLevel::Type InFeatureLevel) const
-	{
-		if (ERHIFeatureLevel::SM5 != InFeatureLevel)
-			return false;
+	ENGINE_API bool UseForHairStrands(ERHIFeatureLevel::Type InFeatureLevel) const;
 
-		const FMaterial& Mat = MaterialRenderProxy->GetIncompleteMaterialWithFallback(InFeatureLevel);
-		return IsCompatibleWithHairStrands(&Mat, InFeatureLevel);
-	}
-
-	FORCEINLINE bool IsMasked(ERHIFeatureLevel::Type InFeatureLevel) const
-	{
-		// Note: blend mode does not depend on the feature level we are actually rendering in.
-		return MaterialRenderProxy->GetIncompleteMaterialWithFallback(InFeatureLevel).IsMasked();
-	}
-
-	/** Converts from an int32 index into a int8 */
-	static int8 QuantizeLODIndex(int32 NewLODIndex)
-	{
-		checkSlow(NewLODIndex >= SCHAR_MIN && NewLODIndex <= SCHAR_MAX);
-		return (int8)NewLODIndex;
-	}
+	ENGINE_API bool IsMasked(ERHIFeatureLevel::Type InFeatureLevel) const;
 
 	FORCEINLINE int32 GetNumPrimitives() const
 	{
diff --git a/Engine/Source/Runtime/Engine/Public/ParticleHelper.h b/Engine/Source/Runtime/Engine/Public/ParticleHelper.h
--- a/Engine/Source/Runtime/Engine/Public/ParticleHelper.h
+++ b/Engine/Source/Runtime/Engine/Public/ParticleHelper.h
@@ -17,7 +17,9 @@
 #include "PrimitiveViewRelevance.h"
 #include "PrimitiveSceneProxy.h"
 #include "Materials/MaterialInterface.h"
+#if UE_ENABLE_INCLUDE_ORDER_DEPRECATED_IN_5_2
 #include "MaterialShared.h"
+#endif
 #include "MeshBatch.h"
 #include "MeshParticleVertexFactory.h"
 #include "PrimitiveSceneProxy.h"
@@ -25,6 +27,7 @@
 
 #define _ENABLE_PARTICLE_LOD_INGAME_
 
+class FColoredMaterialRenderProxy;
 class FParticleSystemSceneProxy;
 class UParticleModuleRequired;
 class UParticleSystemComponent;
@@ -2471,7 +2474,7 @@ public:
 
 	inline FRHIUniformBuffer* GetWorldSpacePrimitiveUniformBuffer() const { return WorldSpacePrimitiveUniformBuffer.GetUniformBufferRHI(); }
 
-	const FColoredMaterialRenderProxy* GetDeselectedWireframeMatInst() const	{	return &DeselectedWireframeMaterialInstance;	}
+	const FColoredMaterialRenderProxy* GetDeselectedWireframeMatInst() const { return DeselectedWireframeMaterialInstance; }
 
 	/** Gets a mesh batch from the pool. */
 	FMeshBatch* GetPooledMeshBatch();
@@ -2514,7 +2517,7 @@ protected:
 	FParticleDynamicData* DynamicData;			// RENDER THREAD USAGE ONLY
 	FParticleDynamicData* LastDynamicData;		// RENDER THREAD USAGE ONLY
 
-	FColoredMaterialRenderProxy DeselectedWireframeMaterialInstance;
+	FColoredMaterialRenderProxy* DeselectedWireframeMaterialInstance;
 
 	int32 LODMethod;
 	float PendingLODDistance;
diff --git a/Engine/Source/Runtime/Engine/Public/SceneManagement.h b/Engine/Source/Runtime/Engine/Public/SceneManagement.h
--- a/Engine/Source/Runtime/Engine/Public/SceneManagement.h
+++ b/Engine/Source/Runtime/Engine/Public/SceneManagement.h
@@ -21,6 +21,7 @@
 #include "HitProxies.h"
 #include "SceneTypes.h"
 #include "ConvexVolume.h"
+#include "PrimitiveUniformShaderParameters.h"
 #include "RendererInterface.h"
 #if UE_ENABLE_INCLUDE_ORDER_DEPRECATED_IN_5_2
 #include "Engine/TextureLightProfile.h"
@@ -2335,7 +2336,7 @@ public:
 /** 
  * Encapsulates the gathering of meshes from the various FPrimitiveSceneProxy classes. 
  */
-class FMeshElementCollector
+class ENGINE_API FMeshElementCollector
 {
 public:
 
@@ -2391,7 +2392,7 @@ public:
 	/** 
 	 * Adds a mesh batch to the collector for the specified view so that it can be rendered.
 	 */
-	ENGINE_API void AddMesh(int32 ViewIndex, FMeshBatch& MeshBatch);
+	void AddMesh(int32 ViewIndex, FMeshBatch& MeshBatch);
 
 	/** Add a material render proxy that will be cleaned up automatically */
 	void RegisterOneFrameMaterialProxy(FMaterialRenderProxy* Proxy)
@@ -2421,85 +2422,34 @@ public:
 		ParallelTasks.Emplace(Task);
 	}
 
-	ENGINE_API void ProcessTasks();
+	void ProcessTasks();
 
-	ENGINE_API ERHIFeatureLevel::Type GetFeatureLevel() const
+	FORCEINLINE ERHIFeatureLevel::Type GetFeatureLevel() const
 	{
 		return FeatureLevel;
 	}
 
 protected:
 
-	ENGINE_API FMeshElementCollector(ERHIFeatureLevel::Type InFeatureLevel, FSceneRenderingBulkObjectAllocator& InBulkAllocator);
+	FMeshElementCollector(ERHIFeatureLevel::Type InFeatureLevel, FSceneRenderingBulkObjectAllocator& InBulkAllocator);
 
-	~FMeshElementCollector()
-	{
-		DeleteTemporaryProxies();
-	}
+	~FMeshElementCollector();
 
-	void DeleteTemporaryProxies()
-	{
-		check(!ParallelTasks.Num()); // We should have blocked on this already
-		for (int32 ProxyIndex = 0; ProxyIndex < TemporaryProxies.Num(); ProxyIndex++)
-		{
-			delete TemporaryProxies[ProxyIndex];
-		}
+	void DeleteTemporaryProxies();
 
-		TemporaryProxies.Empty();
-	}
-
-	void SetPrimitive(const FPrimitiveSceneProxy* InPrimitiveSceneProxy, FHitProxyId DefaultHitProxyId)
-	{
-		check(InPrimitiveSceneProxy);
-		PrimitiveSceneProxy = InPrimitiveSceneProxy;
+	void SetPrimitive(const FPrimitiveSceneProxy* InPrimitiveSceneProxy, FHitProxyId DefaultHitProxyId);
 
-		for (int32 ViewIndex = 0; ViewIndex < SimpleElementCollectors.Num(); ViewIndex++)
-		{
-			SimpleElementCollectors[ViewIndex]->HitProxyId = DefaultHitProxyId;
-			SimpleElementCollectors[ViewIndex]->PrimitiveMeshId = 0;
-		}
-
-		for (int32 ViewIndex = 0; ViewIndex < MeshIdInPrimitivePerView.Num(); ++ViewIndex)
-		{
-			MeshIdInPrimitivePerView[ViewIndex] = 0;
-		}
-	}
-
-	void ClearViewMeshArrays()
-	{
-		Views.Empty();
-		MeshBatches.Empty();
-		SimpleElementCollectors.Empty();
-		MeshIdInPrimitivePerView.Empty();
-		DynamicPrimitiveCollectorPerView.Empty();
-		NumMeshBatchElementsPerView.Empty();
-		DynamicIndexBuffer = nullptr;
-		DynamicVertexBuffer = nullptr;
-		DynamicReadBuffer = nullptr;
-	}
+	void ClearViewMeshArrays();
 
 	void AddViewMeshArrays(
-		FSceneView* InView, 
-		TArray<FMeshBatchAndRelevance,SceneRenderingAllocator>* ViewMeshes,
-		FSimpleElementCollector* ViewSimpleElementCollector, 
+		FSceneView* InView,
+		TArray<FMeshBatchAndRelevance, SceneRenderingAllocator>* ViewMeshes,
+		FSimpleElementCollector* ViewSimpleElementCollector,
 		FGPUScenePrimitiveCollector* InDynamicPrimitiveCollector,
 		ERHIFeatureLevel::Type InFeatureLevel,
 		FGlobalDynamicIndexBuffer* InDynamicIndexBuffer,
 		FGlobalDynamicVertexBuffer* InDynamicVertexBuffer,
-		FGlobalDynamicReadBuffer* InDynamicReadBuffer)
-	{
-		Views.Add(InView);
-		MeshIdInPrimitivePerView.Add(0);
-		MeshBatches.Add(ViewMeshes);
-		NumMeshBatchElementsPerView.Add(0);
-		SimpleElementCollectors.Add(ViewSimpleElementCollector);
-		DynamicPrimitiveCollectorPerView.Add(InDynamicPrimitiveCollector);
-
-		check(InDynamicIndexBuffer && InDynamicVertexBuffer && InDynamicReadBuffer);
-		DynamicIndexBuffer = InDynamicIndexBuffer;
-		DynamicVertexBuffer = InDynamicVertexBuffer;
-		DynamicReadBuffer = InDynamicReadBuffer;
-	}
+		FGlobalDynamicReadBuffer* InDynamicReadBuffer);
 
 	/** 
 	 * Using TChunkedArray which will never realloc as new elements are added
@@ -2627,21 +2577,17 @@ struct FRayTracingMaterialGatheringContext
 };
 #endif
 
-class FDynamicPrimitiveUniformBuffer : public FOneFrameResource
+class ENGINE_API FDynamicPrimitiveUniformBuffer : public FOneFrameResource
 {
 public:
-	FDynamicPrimitiveUniformBuffer() = default;
+	FDynamicPrimitiveUniformBuffer();
 	// FDynamicPrimitiveUniformBuffer is non-copyable
 	FDynamicPrimitiveUniformBuffer(const FDynamicPrimitiveUniformBuffer&) = delete;
-
-	virtual ~FDynamicPrimitiveUniformBuffer()
-	{
-		UniformBuffer.ReleaseResource();
-	}
+	virtual ~FDynamicPrimitiveUniformBuffer();
 
 	TUniformBuffer<FPrimitiveUniformShaderParameters> UniformBuffer;
 
-	ENGINE_API void Set(
+	void Set(
 		const FMatrix& LocalToWorld,
 		const FMatrix& PreviousLocalToWorld,
 		const FVector& ActorPositionWS, 
@@ -2664,7 +2610,7 @@ public:
 		bool bOutputVelocity,
 		const FCustomPrimitiveData* CustomPrimitiveData);
 
-	ENGINE_API void Set(
+	void Set(
 		const FMatrix& LocalToWorld,
 		const FMatrix& PreviousLocalToWorld,
 		const FBoxSphereBounds& WorldBounds,
@@ -2675,7 +2621,7 @@ public:
 		bool bOutputVelocity);
 
 	/** Pass-through implementation which calls the overloaded Set function with LocalBounds for PreSkinnedLocalBounds. */
-	ENGINE_API void Set(
+	void Set(
 		const FMatrix& LocalToWorld,
 		const FMatrix& PreviousLocalToWorld,
 		const FBoxSphereBounds& WorldBounds,
