Applying resolved changes to newly added file.

diff --git a/Engine/Source/Runtime/Engine/Private/MeshBatch.cpp b/Engine/Source/Runtime/Engine/Private/MeshBatch.cpp
--- a/Engine/Source/Runtime/Engine/Private/MeshBatch.cpp
+++ b/Engine/Source/Runtime/Engine/Private/MeshBatch.cpp
@@ -6,7 +6,7 @@
 bool FMeshBatch::IsTranslucent(ERHIFeatureLevel::Type InFeatureLevel) const
 {
 	// Note: blend mode does not depend on the feature level we are actually rendering in.
-	return IsTranslucentBlendMode(MaterialRenderProxy->GetIncompleteMaterialWithFallback(InFeatureLevel).GetBlendMode());
+	return IsTranslucentBlendMode(MaterialRenderProxy->GetIncompleteMaterialWithFallback(InFeatureLevel));
 }
 
 bool FMeshBatch::IsDecal(ERHIFeatureLevel::Type InFeatureLevel) const
