Removing ENGINE_API that's not needed anymore.

diff --git a/Engine/Source/Runtime/Engine/Public/SceneManagement.h b/Engine/Source/Runtime/Engine/Public/SceneManagement.h
--- a/Engine/Source/Runtime/Engine/Public/SceneManagement.h
+++ b/Engine/Source/Runtime/Engine/Public/SceneManagement.h
@@ -2599,7 +2599,7 @@ public:
 		bool bOutputVelocity,
 		const FCustomPrimitiveData* CustomPrimitiveData);
 
-	ENGINE_API void Set(
+	void Set(
 		const FMatrix& LocalToWorld,
 		const FMatrix& PreviousLocalToWorld,
 		const FBoxSphereBounds& WorldBounds,
