Use TArrayView instead of TArray with inlined allocator

diff --git a/Engine/Source/Runtime/Engine/Classes/Engine/NetworkObjectList.h b/Engine/Source/Runtime/Engine/Classes/Engine/NetworkObjectList.h
--- a/Engine/Source/Runtime/Engine/Classes/Engine/NetworkObjectList.h
+++ b/Engine/Source/Runtime/Engine/Classes/Engine/NetworkObjectList.h
@@ -307,13 +307,13 @@ public:
 	void RemoveSubObjectChannelReference(AActor* OwnerActor, const TWeakObjectPtr<UObject>& ReplicatedSubObject, UObject* ReferenceOwner);
 
 	/** Called when multiple subobjects need to remove their reference from either the active or inactive list */
-	void RemoveMultipleSubObjectChannelReference(AActor* OwnerActor, const TArray< TWeakObjectPtr<UObject>, TInlineAllocator<16> >& SubObjectsToRemove, UObject* ReferenceOwner);
+	void RemoveMultipleSubObjectChannelReference(AActor* OwnerActor, const TArrayView<TWeakObjectPtr<UObject>>& SubObjectsToRemove, UObject* ReferenceOwner);
 
 	/** Called when multiple subobjects that were flagged torn off or delete have removed their channel reference */
-	void RemoveMultipleInvalidSubObjectChannelReference(FNetworkObjectInfo* ActorNetInfo, const TArray< TWeakObjectPtr<UObject>, TInlineAllocator<16> >& SubObjectsToRemove, UObject* ReferenceOwner);
+	void RemoveMultipleInvalidSubObjectChannelReference(FNetworkObjectInfo* ActorNetInfo, const TArrayView<TWeakObjectPtr<UObject>>& SubObjectsToRemove, UObject* ReferenceOwner);
 
 	/** Called when multiple subobjects that were still considered active have removed their channel reference */
-	void RemoveMultipleActiveSubObjectChannelReference(FNetworkObjectInfo* ActorNetInfo, const TArray< TWeakObjectPtr<UObject>, TInlineAllocator<16> >& SubObjectsToRemove, UObject* ReferenceOwner);
+	void RemoveMultipleActiveSubObjectChannelReference(FNetworkObjectInfo* ActorNetInfo, const TArrayView<TWeakObjectPtr<UObject>>& SubObjectsToRemove, UObject* ReferenceOwner);
 
 	/** 
 	* Keep track of the transfer of ownership from the channel to the connection when the actor becomes dormant.
@@ -321,7 +321,7 @@ public:
 	* No actual logic is modified here.
 	*/
 #if DO_REPLICATED_OBJECT_CHANNELREF_CHECKS
-	void SwapMultipleReferencesForDormancy(AActor* OwnerActor, const TArray< TWeakObjectPtr<UObject>, TInlineAllocator<16> >& SubObjectsToSwap, UActorChannel* PreviousChannelRefOwner, UNetConnection* NewConnectionRefOwner);
+	void SwapMultipleReferencesForDormancy(AActor* OwnerActor, const TArrayView<TWeakObjectPtr<UObject>>& SubObjectsToSwap, UActorChannel* PreviousChannelRefOwner, UNetConnection* NewConnectionRefOwner);
 	void SwapReferenceForDormancy(AActor* OwnerActor, UObject* ReplicatedSubObject, UNetConnection* PreviousConnectionRefOwner, UActorChannel* NewChannelRefOwner);
 #endif 
 
diff --git a/Engine/Source/Runtime/Engine/Private/NetworkObjectList.cpp b/Engine/Source/Runtime/Engine/Private/NetworkObjectList.cpp
--- a/Engine/Source/Runtime/Engine/Private/NetworkObjectList.cpp
+++ b/Engine/Source/Runtime/Engine/Private/NetworkObjectList.cpp
@@ -634,7 +634,7 @@ void FNetworkObjectList::AddSubObjectChannelReference(AActor* OwnerActor, UObjec
 	}
 }
 
-void FNetworkObjectList::RemoveMultipleSubObjectChannelReference(AActor* OwnerActor, const TArray< TWeakObjectPtr<UObject>, TInlineAllocator<16> >& SubObjectsToRemove, UObject* ReferenceOwner)
+void FNetworkObjectList::RemoveMultipleSubObjectChannelReference(AActor* OwnerActor, const TArrayView<TWeakObjectPtr<UObject>>& SubObjectsToRemove, UObject* ReferenceOwner)
 {
 	if (TSharedPtr<FNetworkObjectInfo>* InfoPtr = AllNetworkObjects.Find(OwnerActor))
 	{
@@ -649,7 +649,7 @@ void FNetworkObjectList::RemoveMultipleSubObjectChannelReference(AActor* OwnerAc
 	}
 }
 
-void FNetworkObjectList::RemoveMultipleInvalidSubObjectChannelReference(FNetworkObjectInfo* ActorNetInfo, const TArray< TWeakObjectPtr<UObject>, TInlineAllocator<16> >& SubObjectsToRemove, UObject* ReferenceOwner)
+void FNetworkObjectList::RemoveMultipleInvalidSubObjectChannelReference(FNetworkObjectInfo* ActorNetInfo, const TArrayView<TWeakObjectPtr<UObject>>& SubObjectsToRemove, UObject* ReferenceOwner)
 {
 	check(ActorNetInfo);
 
@@ -659,7 +659,7 @@ void FNetworkObjectList::RemoveMultipleInvalidSubObjectChannelReference(FNetwork
 	}
 }
 
-void FNetworkObjectList::RemoveMultipleActiveSubObjectChannelReference(FNetworkObjectInfo* ActorNetInfo, const TArray< TWeakObjectPtr<UObject>, TInlineAllocator<16> >& SubObjectsToRemove, UObject* ReferenceOwner)
+void FNetworkObjectList::RemoveMultipleActiveSubObjectChannelReference(FNetworkObjectInfo* ActorNetInfo, const TArrayView<TWeakObjectPtr<UObject>>& SubObjectsToRemove, UObject* ReferenceOwner)
 {
 	check(ActorNetInfo);
 
@@ -762,7 +762,7 @@ bool FNetworkObjectList::HandleRemoveInvalidSubObjectRef(FNetworkObjectInfo* Act
 }
 
 #if DO_REPLICATED_OBJECT_CHANNELREF_CHECKS
-void FNetworkObjectList::SwapMultipleReferencesForDormancy(AActor* OwnerActor, const TArray< TWeakObjectPtr<UObject>, TInlineAllocator<16> >& SubObjectsToSwap, UActorChannel* PreviousChannelRefOwner, UNetConnection* NewConnectionRefOwner)
+void FNetworkObjectList::SwapMultipleReferencesForDormancy(AActor* OwnerActor, const TArrayView<TWeakObjectPtr<UObject>>& SubObjectsToSwap, UActorChannel* PreviousChannelRefOwner, UNetConnection* NewConnectionRefOwner)
 {
 	if (TSharedPtr<FNetworkObjectInfo>* InfoPtr = AllNetworkObjects.Find(OwnerActor))
 	{
