World Partition:
- Added helper functions in FWorldPartitionHelpers to convert editor to runtime actor paths (and the other way around).
- Added streaming generation unit test that loads a map, generate streaming, and perform conversion tests.

diff --git a/Engine/Source/Runtime/Engine/Private/Tests/WorldPartition/WorldPartitionHandleTests.cpp b/Engine/Source/Runtime/Engine/Private/Tests/WorldPartition/WorldPartitionHandleTests.cpp
--- a/Engine/Source/Runtime/Engine/Private/Tests/WorldPartition/WorldPartitionHandleTests.cpp
+++ b/Engine/Source/Runtime/Engine/Private/Tests/WorldPartition/WorldPartitionHandleTests.cpp
@@ -16,8 +16,7 @@
 
 namespace WorldPartitionTests
 {
-	constexpr const uint32 TestFlags = EAutomationTestFlags::ApplicationContextMask | EAutomationTestFlags::EngineFilter;
-	IMPLEMENT_SIMPLE_AUTOMATION_TEST(FWorldPartitionSoftRefTest, TEST_NAME_ROOT ".Handle", TestFlags)
+	IMPLEMENT_SIMPLE_AUTOMATION_TEST(FWorldPartitionSoftRefTest, TEST_NAME_ROOT ".Handle", EAutomationTestFlags::ApplicationContextMask | EAutomationTestFlags::EngineFilter)
 
 #if WITH_EDITOR
 	template <class LoadingContextType>
diff --git a/Engine/Source/Runtime/Engine/Private/Tests/WorldPartition/WorldPartitionStreamingGenerationTests.cpp b/Engine/Source/Runtime/Engine/Private/Tests/WorldPartition/WorldPartitionStreamingGenerationTests.cpp
new file mode 100644
--- /dev/null
+++ b/Engine/Source/Runtime/Engine/Private/Tests/WorldPartition/WorldPartitionStreamingGenerationTests.cpp
@@ -0,0 +1,102 @@
+// Copyright Epic Games, Inc. All Rights Reserved.
+
+#include "Misc/AutomationTest.h"
+
+#if WITH_EDITOR
+#include "WorldPartition/WorldPartition.h"
+#include "WorldPartition/WorldPartitionHandle.h"
+#include "WorldPartition/WorldPartitionHelpers.h"
+#include "WorldPartition/ActorDescContainer.h"
+#include "Engine/World.h"
+#include "Engine/Level.h"
+#include "PackageTools.h"
+#include "EditorWorldUtils.h"
+#endif
+
+#if WITH_DEV_AUTOMATION_TESTS
+
+#define TEST_NAME_ROOT "System.Engine.WorldPartition"
+
+namespace WorldPartitionTests
+{
+	IMPLEMENT_SIMPLE_AUTOMATION_TEST(FWorldPartitionStreamingGenerationTest, TEST_NAME_ROOT ".StreamingGeneration", EAutomationTestFlags::ApplicationContextMask | EAutomationTestFlags::EngineFilter)
+
+	bool FWorldPartitionStreamingGenerationTest::RunTest(const FString& Parameters)
+	{
+#if WITH_EDITOR
+		ON_SCOPE_EXIT
+		{
+			CollectGarbage(GARBAGE_COLLECTION_KEEPFLAGS);
+		};
+
+		FScopedEditorWorld ScopedEditorWorld(
+			TEXTVIEW("/Engine/WorldPartition/WorldPartitionUnitTest"),
+			UWorld::InitializationValues()
+				.RequiresHitProxies(false)
+				.ShouldSimulatePhysics(false)
+				.EnableTraceCollision(false)
+				.CreateNavigation(false)
+				.CreateAISystem(false)
+				.AllowAudioPlayback(false)
+				.CreatePhysicsScene(true)
+		);
+		
+		UWorld* World = ScopedEditorWorld.GetWorld();
+		if (!TestTrue(TEXT("Missing World Object"), !!World))
+		{
+			return false;
+		}
+
+		UWorldPartition* WorldPartition = World->GetWorldPartition();
+		if (!TestTrue(TEXT("Missing World Partition Object"), !!WorldPartition))
+		{
+			return false;
+		}
+
+		UActorDescContainer* ActorDescMainContainer = WorldPartition->GetActorDescContainer();
+		if (!TestTrue(TEXT("Missing World Partition Container"), !!ActorDescMainContainer))
+		{
+			return false;
+		}
+
+		if (!TestTrue(TEXT("World Partition Generate Streaming"), WorldPartition->GenerateContainerStreaming(ActorDescMainContainer)))
+		{
+			return false;
+		}
+
+		FWorldPartitionReference ActorRef(ActorDescMainContainer, FGuid(TEXT("5D9F93BA407A811AFDDDAAB4F1CECC6A")));
+		if (!TestTrue(TEXT("Invalid Actor Reference"), ActorRef.IsValid()))
+		{
+			return false;
+		}
+
+		AActor* Actor = ActorRef->GetActor();
+		if (!TestTrue(TEXT("Missing Actor"), !!Actor))
+		{
+			return false;
+		}
+
+		FSoftObjectPath ActorRuntimePath;
+		if (!TestTrue(TEXT("Actor Path Editor to Runtime Conversion Failed"), FWorldPartitionHelpers::ConvertEditorPathToRuntimePath(ActorRef->GetActorSoftPath(), ActorRuntimePath)))
+		{
+			return false;
+		}
+
+		FSoftObjectPath ActorEditorPath;
+		if (!TestTrue(TEXT("Actor Path Runtime to Editor Conversion Failed"), FWorldPartitionHelpers::ConvertRuntimePathToEditorPath(ActorRuntimePath, ActorEditorPath)))
+		{
+			return false;
+		}
+
+		if (!TestTrue(TEXT("Actor Path Editor to Runtime to Editor Conversion Failed"), ActorEditorPath == ActorRef->GetActorSoftPath()))
+		{
+			return false;
+		}
+#endif
+		return true;
+	}
+}
+
+#undef TEST_NAME_ROOT
+
+#endif 
\ No newline at end of file
diff --git a/Engine/Source/Runtime/Engine/Private/WorldPartition/WorldPartition.cpp b/Engine/Source/Runtime/Engine/Private/WorldPartition/WorldPartition.cpp
--- a/Engine/Source/Runtime/Engine/Private/WorldPartition/WorldPartition.cpp
+++ b/Engine/Source/Runtime/Engine/Private/WorldPartition/WorldPartition.cpp
@@ -398,6 +398,11 @@ bool UWorldPartition::IsValidPackageName(const FString& InPackageName)
 	return false;
 }
 
+bool UWorldPartition::ConvertEditorPathToRuntimePath(const FSoftObjectPath& InPath, FSoftObjectPath& OutPath) const
+{
+	return StreamingPolicy ? StreamingPolicy->ConvertEditorPathToRuntimePath(InPath, OutPath) : false;
+}
+
 void UWorldPartition::OnPreBeginPIE(bool bStartSimulate)
 {
 	check(!bIsPIE);
diff --git a/Engine/Source/Runtime/Engine/Private/WorldPartition/WorldPartitionHelpers.cpp b/Engine/Source/Runtime/Engine/Private/WorldPartition/WorldPartitionHelpers.cpp
--- a/Engine/Source/Runtime/Engine/Private/WorldPartition/WorldPartitionHelpers.cpp
+++ b/Engine/Source/Runtime/Engine/Private/WorldPartition/WorldPartitionHelpers.cpp
@@ -225,4 +225,49 @@ void FWorldPartitionHelpers::FakeEngineTick(UWorld* InWorld)
 	CommandletHelpers::TickEngine(InWorld);
 }
 
-#endif // #if WITH_EDITOR
+bool FWorldPartitionHelpers::ConvertRuntimePathToEditorPath(const FSoftObjectPath& InPath, FSoftObjectPath& OutPath)
+{
+	//
+	// Try to convert from /.../WorldName/_Generated_/MainGrid_L0_X0_Y0_DL0.WorldName:PersistentLevel.ActorName
+	//                  to /.../WorldName.WorldName:PersistentLevel.ActorName
+	//
+	FString OutPathString = InPath.ToString();
+	const FStringView InPathView(OutPathString);
+
+	if (int32 GeneratedPos = InPathView.Find(TEXTVIEW("/_Generated_/"), 0); GeneratedPos != INDEX_NONE)
+	{
+		if (int32 NextDotPos = InPathView.Find(TEXTVIEW("."), GeneratedPos); NextDotPos != INDEX_NONE)
+		{
+			if (int32 NextColonPos = InPathView.Find(TEXTVIEW(":"), NextDotPos); NextColonPos != INDEX_NONE)
+			{
+				OutPathString.RemoveAt(GeneratedPos, NextDotPos - GeneratedPos);
+				OutPath = OutPathString;
+				return true;
+			}
+		}
+	}
+
+	return false;
+}
+
+bool FWorldPartitionHelpers::ConvertEditorPathToRuntimePath(const FSoftObjectPath& InPath, FSoftObjectPath& OutPath)
+{
+	//
+	// Try to convert from /.../WorldName.WorldName:PersistentLevel.ActorName
+	//                  to /.../WorldName/_Generated_/MainGrid_L0_X0_Y0_DL0.WorldName:PersistentLevel.ActorName
+	//
+	FSoftObjectPath Path(InPath);
+	FSoftObjectPath WorldPath(Path.GetAssetPath(), FString());
+
+	if (UWorld* World = Cast<UWorld>(WorldPath.ResolveObject()))
+	{
+		if (UWorldPartition* WorldPartition = World->GetWorldPartition())
+		{
+			return WorldPartition->ConvertEditorPathToRuntimePath(InPath, OutPath);
+		}
+	}
+
+	return false;
+}
+
+#endif // #if WITH_EDITOR
\ No newline at end of file
diff --git a/Engine/Source/Runtime/Engine/Private/WorldPartition/WorldPartitionLevelStreamingPolicy.cpp b/Engine/Source/Runtime/Engine/Private/WorldPartition/WorldPartitionLevelStreamingPolicy.cpp
--- a/Engine/Source/Runtime/Engine/Private/WorldPartition/WorldPartitionLevelStreamingPolicy.cpp
+++ b/Engine/Source/Runtime/Engine/Private/WorldPartition/WorldPartitionLevelStreamingPolicy.cpp
@@ -181,6 +181,20 @@ void UWorldPartitionLevelStreamingPolicy::RemapSoftObjectPath(FSoftObjectPath& O
 		}
 	}
 }
+
+bool UWorldPartitionLevelStreamingPolicy::ConvertEditorPathToRuntimePath(const FSoftObjectPath& InPath, FSoftObjectPath& OutPath) const
+{
+	if (const FName* CellName = ActorToCellRemapping.Find(*InPath.ToString()))
+	{
+		const UWorld* OuterWorld = WorldPartition->GetTypedOuter<UWorld>();
+		const UPackage* OuterWorldPackage = OuterWorld->GetPackage();
+		const FString OuterWorldPackageName = FPackageName::GetShortName(OuterWorldPackage->GetName());
+		const FString OuterWorldPackagePath = FPackageName::GetLongPackagePath(OuterWorldPackage->GetPathName());
+		OutPath = FString::Printf(TEXT("%s/%s/_Generated_/%s.%s:%s"), *OuterWorldPackagePath, *OuterWorldPackageName, *(CellName->ToString()), *OuterWorldPackageName, *InPath.GetSubPathString());
+		return true;
+	}
+	return false;
+}
 #endif
 
 UObject* UWorldPartitionLevelStreamingPolicy::GetSubObject(const TCHAR* SubObjectPath)
diff --git a/Engine/Source/Runtime/Engine/Public/WorldPartition/WorldPartition.h b/Engine/Source/Runtime/Engine/Public/WorldPartition/WorldPartition.h
--- a/Engine/Source/Runtime/Engine/Public/WorldPartition/WorldPartition.h
+++ b/Engine/Source/Runtime/Engine/Public/WorldPartition/WorldPartition.h
@@ -164,6 +164,9 @@ public:
 	void RemapSoftObjectPath(FSoftObjectPath& ObjectPath);
 	bool IsValidPackageName(const FString& InPackageName);
 
+	// Editor/Runtime conversions
+	bool ConvertEditorPathToRuntimePath(const FSoftObjectPath& InPath, FSoftObjectPath& OutPath) const;
+
 	// Begin Cooking
 	void BeginCook(IWorldPartitionCookPackageContext& CookContext);
 	DECLARE_MULTICAST_DELEGATE_OneParam(FWorldPartitionBeginCookDelegate, IWorldPartitionCookPackageContext&);
diff --git a/Engine/Source/Runtime/Engine/Public/WorldPartition/WorldPartitionHelpers.h b/Engine/Source/Runtime/Engine/Public/WorldPartition/WorldPartitionHelpers.h
--- a/Engine/Source/Runtime/Engine/Public/WorldPartition/WorldPartitionHelpers.h
+++ b/Engine/Source/Runtime/Engine/Public/WorldPartition/WorldPartitionHelpers.h
@@ -76,6 +76,10 @@ public:
 
 	// Simulate an engine frame tick
 	static void FakeEngineTick(UWorld* World);
+
+	// Editor/Runtime conversions
+	static bool ConvertRuntimePathToEditorPath(const FSoftObjectPath& InPath, FSoftObjectPath& OutPath);
+	static bool ConvertEditorPathToRuntimePath(const FSoftObjectPath& InPath, FSoftObjectPath& OutPath);
 };
 
 #endif // WITH_EDITOR
diff --git a/Engine/Source/Runtime/Engine/Public/WorldPartition/WorldPartitionLevelStreamingPolicy.h b/Engine/Source/Runtime/Engine/Public/WorldPartition/WorldPartitionLevelStreamingPolicy.h
--- a/Engine/Source/Runtime/Engine/Public/WorldPartition/WorldPartitionLevelStreamingPolicy.h
+++ b/Engine/Source/Runtime/Engine/Public/WorldPartition/WorldPartitionLevelStreamingPolicy.h
@@ -30,6 +30,7 @@ public:
 	virtual TSubclassOf<class UWorldPartitionRuntimeCell> GetRuntimeCellClass() const override;
 	virtual void PrepareActorToCellRemapping() override;
 	virtual void RemapSoftObjectPath(FSoftObjectPath& ObjectPath) override;
+	virtual bool ConvertEditorPathToRuntimePath(const FSoftObjectPath& InPath, FSoftObjectPath& OutPath) const override;
 	static FString GetCellPackagePath(const FName& InCellName, const UWorld* InWorld);
 #endif
 
diff --git a/Engine/Source/Runtime/Engine/Public/WorldPartition/WorldPartitionStreamingPolicy.h b/Engine/Source/Runtime/Engine/Public/WorldPartition/WorldPartitionStreamingPolicy.h
--- a/Engine/Source/Runtime/Engine/Public/WorldPartition/WorldPartitionStreamingPolicy.h
+++ b/Engine/Source/Runtime/Engine/Public/WorldPartition/WorldPartitionStreamingPolicy.h
@@ -65,6 +65,9 @@ public:
 	// PIE/Game methods
 	virtual void PrepareActorToCellRemapping() {}
 	virtual void RemapSoftObjectPath(FSoftObjectPath& ObjectPath) {}
+
+	// Editor/Runtime conversions
+	virtual bool ConvertEditorPathToRuntimePath(const FSoftObjectPath& InPath, FSoftObjectPath& OutPath) const { return false; }
 #endif
 
 #if !UE_BUILD_SHIPPING
