Mark a primitive scene proxy as excluded from ray tracing if it has no static meshes

Previously it was only excluded if it has no ray tarcing geometries, however it is possible for a primitive to have a RT representation but no statcic meshes. For example, this may happen in the editor when trying to isolate a mesh section with no triangles.

diff --git a/Engine/Source/Runtime/Engine/Private/PrimitiveSceneProxy.cpp b/Engine/Source/Runtime/Engine/Private/PrimitiveSceneProxy.cpp
--- a/Engine/Source/Runtime/Engine/Private/PrimitiveSceneProxy.cpp
+++ b/Engine/Source/Runtime/Engine/Private/PrimitiveSceneProxy.cpp
@@ -1537,13 +1537,12 @@ ERayTracingPrimitiveFlags FPrimitiveSceneProxy::GetCachedRayTracingInstance(FRay
 
 	if (IsRayTracingStaticRelevant())
 	{
-		if (PrimitiveSceneInfo->GetRayTracingGeometryNum() == 0)
+		if (PrimitiveSceneInfo->GetRayTracingGeometryNum() == 0
+			|| PrimitiveSceneInfo->StaticMeshes.IsEmpty())
 		{
 			return ERayTracingPrimitiveFlags::Excluded;
 		}
 
-		check(PrimitiveSceneInfo->StaticMeshes.Num() > 0);
-
 		// overwrite flag if static
 		ResultFlags = ERayTracingPrimitiveFlags::StaticMesh | ERayTracingPrimitiveFlags::ComputeLOD | ERayTracingPrimitiveFlags::CacheMeshCommands;
 	}
