Fix non-unity build.

diff --git a/Engine/Source/Runtime/Engine/Private/WorldPartition/WorldPartitionHelpers.cpp b/Engine/Source/Runtime/Engine/Private/WorldPartition/WorldPartitionHelpers.cpp
--- a/Engine/Source/Runtime/Engine/Private/WorldPartition/WorldPartitionHelpers.cpp
+++ b/Engine/Source/Runtime/Engine/Private/WorldPartition/WorldPartitionHelpers.cpp
@@ -4,6 +4,7 @@
 
 #include "WorldPartition/WorldPartition.h"
 #include "WorldPartition/WorldPartitionEditorHash.h"
+#include "Engine/World.h"
 #include "Algo/AnyOf.h"
 
 #include "Commandlets/Commandlet.h"
