World Partition Streaming Generation Unit Test: add editor soft object path resolving runtime actor test.

diff --git a/Engine/Source/Runtime/Engine/Private/Tests/WorldPartition/WorldPartitionStreamingGenerationTests.cpp b/Engine/Source/Runtime/Engine/Private/Tests/WorldPartition/WorldPartitionStreamingGenerationTests.cpp
--- a/Engine/Source/Runtime/Engine/Private/Tests/WorldPartition/WorldPartitionStreamingGenerationTests.cpp
+++ b/Engine/Source/Runtime/Engine/Private/Tests/WorldPartition/WorldPartitionStreamingGenerationTests.cpp
@@ -92,6 +92,28 @@ namespace WorldPartitionTests
 		{
 			return false;
 		}
+
+		FWorldPartitionHandle ActorHandle(ActorDescMainContainer, FGuid(TEXT("0D2B04D240BE5DE58FE437A8D2DBF5C9")));
+		if (!TestTrue(TEXT("Invalid Actor Handle"), ActorHandle.IsValid()))
+		{
+			return false;
+		}
+
+		if (!TestTrue(TEXT("Actor Handle Not Loaded"), !ActorHandle->GetActor()))
+		{
+			return false;
+		}
+
+		UObject* ResolvedObject = ActorHandle->GetActorSoftPath().TryLoad();
+		if (!TestTrue(TEXT("Actor Handle Loaded"), !!ActorHandle->GetActor()))
+		{
+			return false;
+		}
+
+		if (!TestTrue(TEXT("Resolving Runtime Actor From Editor Path Failed"), ResolvedObject == ActorHandle->GetActor()))
+		{
+			return false;
+		}
 #endif
 		return true;
 	}
