Load AudioMixer module in AudioDeviceManager so audio engine subsystems are added

#rnx

diff --git a/Engine/Source/Runtime/Engine/Private/AudioDevice.cpp b/Engine/Source/Runtime/Engine/Private/AudioDevice.cpp
--- a/Engine/Source/Runtime/Engine/Private/AudioDevice.cpp
+++ b/Engine/Source/Runtime/Engine/Private/AudioDevice.cpp
@@ -2425,7 +2425,7 @@ void FAudioDevice::InitSoundSources()
 void FAudioDevice::InitializeSubsystemCollection()
 {
 	UE_LOG(LogAudio, Log, TEXT("Initializing audio subsystem collection for audio device with id %d"), DeviceID);
-	FModuleManager::Get().LoadModuleChecked("AudioMixer");
+
 	SubsystemCollectionRoot.Reset(NewObject<UAudioSubsystemCollectionRoot>(GetTransientPackage()));
 	check(SubsystemCollectionRoot.IsValid());
 
diff --git a/Engine/Source/Runtime/Engine/Private/AudioDeviceManager.cpp b/Engine/Source/Runtime/Engine/Private/AudioDeviceManager.cpp
--- a/Engine/Source/Runtime/Engine/Private/AudioDeviceManager.cpp
+++ b/Engine/Source/Runtime/Engine/Private/AudioDeviceManager.cpp
@@ -591,6 +591,10 @@ bool FAudioDeviceManager::InitializeManager()
 		AudioSettings->RegisterParameterInterfaces();
 
 		const bool bIsAudioMixerEnabled = AudioDeviceModule->IsAudioMixerModule();
+		if (bIsAudioMixerEnabled)
+		{
+			FModuleManager::Get().LoadModuleChecked(TEXT("AudioMixer"));
+		}
 		AudioSettings->SetAudioMixerEnabled(bIsAudioMixerEnabled);
 
 		// Register all formats
