World Partition: don't load last loaded regions when running automation tests.

diff --git a/Engine/Source/Runtime/Engine/Private/WorldPartition/WorldPartition.cpp b/Engine/Source/Runtime/Engine/Private/WorldPartition/WorldPartition.cpp
--- a/Engine/Source/Runtime/Engine/Private/WorldPartition/WorldPartition.cpp
+++ b/Engine/Source/Runtime/Engine/Private/WorldPartition/WorldPartition.cpp
@@ -660,7 +660,7 @@ void UWorldPartition::Initialize(UWorld* InWorld, const FTransform& InTransform)
 
 		// Load more cells depending on the user's settings
 		// Skipped when running from a commandlet and for subpartitions
-		if (IsMainWorldPartition() && IsStreamingEnabled() && !IsRunningCommandlet())
+		if (IsMainWorldPartition() && IsStreamingEnabled() && !IsRunningCommandlet() && !GIsAutomationTesting)
 		{
 			// Load last loaded regions
 			if (GetMutableDefault<UWorldPartitionEditorPerProjectUserSettings>()->GetEnableLoadingOfLastLoadedRegions())
