World Partition: report invalid actor runtime grids.

diff --git a/Engine/Plugins/Editor/DataValidation/Source/DataValidation/Private/WorldPartitionChangelistValidator.cpp b/Engine/Plugins/Editor/DataValidation/Source/DataValidation/Private/WorldPartitionChangelistValidator.cpp
--- a/Engine/Plugins/Editor/DataValidation/Source/DataValidation/Private/WorldPartitionChangelistValidator.cpp
+++ b/Engine/Plugins/Editor/DataValidation/Source/DataValidation/Private/WorldPartitionChangelistValidator.cpp
@@ -268,6 +268,18 @@ bool UWorldPartitionChangelistValidator::Filter(const UDataLayerInstance* InData
 	return DataLayerWithAsset != nullptr && DataLayerWithAsset->GetAsset() != nullptr && RelevantDataLayerAssets.Contains(DataLayerWithAsset->GetAsset()->GetPathName());
 }
 
+void UWorldPartitionChangelistValidator::OnInvalidRuntimeGrid(const FWorldPartitionActorDescView& ActorDescView, FName GridName)
+{
+	if (Filter(ActorDescView))
+	{
+		FText CurrentError = FText::Format(LOCTEXT("DataValidation.Changelist.WorldPartition.InvalidRuntimeGrid", "Actor {0} has an invalid runtime grid {1}"),
+											FText::FromString(GetFullActorName(ActorDescView)), 
+											FText::FromName(GridName));
+
+		Errors->Add(CurrentError);
+	}
+}
+
 void UWorldPartitionChangelistValidator::OnInvalidReference(const FWorldPartitionActorDescView& ActorDescView, const FGuid& ReferenceGuid, FWorldPartitionActorDescView* ReferenceActorDescView)
 {
 	if (Filter(ActorDescView))
@@ -278,7 +290,6 @@ void UWorldPartitionChangelistValidator::OnInvalidReference(const FWorldPartitio
 
 		Errors->Add(CurrentError);
 	}
-
 }
 
 void UWorldPartitionChangelistValidator::OnInvalidReferenceGridPlacement(const FWorldPartitionActorDescView& ActorDescView, const FWorldPartitionActorDescView& ReferenceActorDescView)
diff --git a/Engine/Plugins/Editor/DataValidation/Source/DataValidation/Private/WorldPartitionChangelistValidator.h b/Engine/Plugins/Editor/DataValidation/Source/DataValidation/Private/WorldPartitionChangelistValidator.h
--- a/Engine/Plugins/Editor/DataValidation/Source/DataValidation/Private/WorldPartitionChangelistValidator.h
+++ b/Engine/Plugins/Editor/DataValidation/Source/DataValidation/Private/WorldPartitionChangelistValidator.h
@@ -34,6 +34,7 @@ protected:
 	bool Filter(const UDataLayerInstance* DataLayerInstance);
 
 	// IStreamingGenerationErrorHandler Interface methods
+	virtual void OnInvalidRuntimeGrid(const FWorldPartitionActorDescView& ActorDescView, FName GridName) override;
 	virtual void OnInvalidReference(const FWorldPartitionActorDescView& ActorDescView, const FGuid& ReferenceGuid, FWorldPartitionActorDescView* ReferenceActorDescView) override;
 	virtual void OnInvalidReferenceGridPlacement(const FWorldPartitionActorDescView& ActorDescView, const FWorldPartitionActorDescView& ReferenceActorDescView) override;
 	virtual void OnInvalidReferenceDataLayers(const FWorldPartitionActorDescView& ActorDescView, const FWorldPartitionActorDescView& ReferenceActorDescView) override;
diff --git a/Engine/Source/Runtime/Engine/Private/WorldPartition/ErrorHandling/WorldPartitionStreamingGenerationLogErrorHandler.cpp b/Engine/Source/Runtime/Engine/Private/WorldPartition/ErrorHandling/WorldPartitionStreamingGenerationLogErrorHandler.cpp
--- a/Engine/Source/Runtime/Engine/Private/WorldPartition/ErrorHandling/WorldPartitionStreamingGenerationLogErrorHandler.cpp
+++ b/Engine/Source/Runtime/Engine/Private/WorldPartition/ErrorHandling/WorldPartitionStreamingGenerationLogErrorHandler.cpp
@@ -6,6 +6,11 @@
 #include "WorldPartition/DataLayer/DataLayerInstanceWithAsset.h"
 #include "GameFramework/Actor.h"
 
+void FStreamingGenerationLogErrorHandler::OnInvalidRuntimeGrid(const FWorldPartitionActorDescView& ActorDescView, FName GridName)
+{
+	UE_LOG(LogWorldPartition, Log, TEXT("Actor %s has an invalid runtime grid %s"), *GetFullActorName(ActorDescView), *GridName.ToString());
+}
+
 void FStreamingGenerationLogErrorHandler::OnInvalidReference(const FWorldPartitionActorDescView& ActorDescView, const FGuid& ReferenceGuid, FWorldPartitionActorDescView* ReferenceActorDescView)
 {
 	UE_LOG(LogWorldPartition, Log, TEXT("Actor %s has an invalid reference to %s"), *GetFullActorName(ActorDescView), ReferenceActorDescView ? *GetFullActorName(*ReferenceActorDescView) : *ReferenceGuid.ToString());
diff --git a/Engine/Source/Runtime/Engine/Private/WorldPartition/ErrorHandling/WorldPartitionStreamingGenerationTokenizedMessageErrorHandler.cpp b/Engine/Source/Runtime/Engine/Private/WorldPartition/ErrorHandling/WorldPartitionStreamingGenerationTokenizedMessageErrorHandler.cpp
--- a/Engine/Source/Runtime/Engine/Private/WorldPartition/ErrorHandling/WorldPartitionStreamingGenerationTokenizedMessageErrorHandler.cpp
+++ b/Engine/Source/Runtime/Engine/Private/WorldPartition/ErrorHandling/WorldPartitionStreamingGenerationTokenizedMessageErrorHandler.cpp
@@ -11,6 +11,19 @@
 
 #define LOCTEXT_NAMESPACE "WorldPartition"
 
+void ITokenizedMessageErrorHandler::OnInvalidRuntimeGrid(const FWorldPartitionActorDescView& ActorDescView, FName GridName)
+{
+	TSharedRef<FTokenizedMessage> Message = FTokenizedMessage::Create(EMessageSeverity::Error);
+	Message->AddToken(FTextToken::Create(LOCTEXT("TokenMessage_WorldPartition_Actor", "Actor")))
+		->AddToken(FActorToken::Create(ActorDescView.GetActorSoftPath().ToString(), ActorDescView.GetGuid(), FText::FromString(GetFullActorName(ActorDescView))))
+		->AddToken(FTextToken::Create(LOCTEXT("TokenMessage_WorldPartition_HaveInvalidRuntimeGrid", "has an invalid runtime grid")))
+		->AddToken(FTextToken::Create(FText::FromName(GridName)));
+
+	AddAdditionalNameToken(Message, FName(TEXT("WorldPartition_InvalidRuntimeGrid_CheckForErrors")));
+
+	HandleTokenizedMessage(MoveTemp(Message));
+}
+
 void ITokenizedMessageErrorHandler::OnInvalidReference(const FWorldPartitionActorDescView& ActorDescView, const FGuid& ReferenceGuid, FWorldPartitionActorDescView* ReferenceActorDescView)
 {
 	TSharedRef<FTokenizedMessage> Message = FTokenizedMessage::Create(EMessageSeverity::Error);
diff --git a/Engine/Source/Runtime/Engine/Private/WorldPartition/WorldPartitionRuntimeSpatialHash.cpp b/Engine/Source/Runtime/Engine/Private/WorldPartition/WorldPartitionRuntimeSpatialHash.cpp
--- a/Engine/Source/Runtime/Engine/Private/WorldPartition/WorldPartitionRuntimeSpatialHash.cpp
+++ b/Engine/Source/Runtime/Engine/Private/WorldPartition/WorldPartitionRuntimeSpatialHash.cpp
@@ -1592,6 +1592,24 @@ void UWorldPartitionRuntimeSpatialHash::FlushStreaming()
 	PackagesToGenerateForCook.Empty();
 }
 
+bool UWorldPartitionRuntimeSpatialHash::IsValidGrid(FName GridName) const
+{
+	if (GridName.IsNone())
+	{
+		return true;
+	}
+
+	for (const FSpatialHashRuntimeGrid& Grid : Grids)
+	{
+		if (Grid.GridName == GridName)
+		{
+			return true;
+		}
+	}
+
+	return false;
+}
+
 #endif //WITH_EDITOR
 
 FAutoConsoleCommand UWorldPartitionRuntimeSpatialHash::OverrideLoadingRangeCommand(
diff --git a/Engine/Source/Runtime/Engine/Private/WorldPartition/WorldPartitionStreamingGeneration.cpp b/Engine/Source/Runtime/Engine/Private/WorldPartition/WorldPartitionStreamingGeneration.cpp
--- a/Engine/Source/Runtime/Engine/Private/WorldPartition/WorldPartitionStreamingGeneration.cpp
+++ b/Engine/Source/Runtime/Engine/Private/WorldPartition/WorldPartitionStreamingGeneration.cpp
@@ -595,11 +595,30 @@ class FWorldPartitionStreamingGenerator
 				};
 
 				// Validate runtime grid
+				auto IsRuntimeGridValid = [this](const FWorldPartitionActorDescView& ActorDescView)
+				{
+					return IsValidGrid(ActorDescView.GetRuntimeGrid());
+				};
+
+				// Validate runtime grid references
 				auto IsReferenceRuntimeGridValid = [](const FWorldPartitionActorDescView& RefererActorDescView, const FWorldPartitionActorDescView& ReferenceActorDescView)
 				{
 					return RefererActorDescView.GetRuntimeGrid() == ReferenceActorDescView.GetRuntimeGrid();
 				};
 
+				if (!IsRuntimeGridValid(ActorDescView))
+				{
+					if (!NbValidationPasses)
+					{
+						ErrorHandler->OnInvalidRuntimeGrid(ActorDescView, ActorDescView.GetRuntimeGrid());
+					}
+					else
+					{
+						ActorDescView.SetInvalidRuntimeGrid();
+					}
+				}
+
+				// Build references List
 				struct FActorReferenceInfo
 				{
 					FGuid ActorGuid;
@@ -608,7 +627,6 @@ class FWorldPartitionStreamingGenerator
 					FWorldPartitionActorDescView* ReferenceActorDesc;
 				};
 
-				// Build references List
 				TArray<FActorReferenceInfo> References;
 
 				// Add normal actor references
@@ -783,6 +801,7 @@ public:
 		bool bEnableStreaming;
 		TMap<FGuid, const UActorDescContainer*> ActorGuidsToContainerMap;
 		TArray<TSubclassOf<AActor>> FilteredClasses;
+		TFunction<bool(FName)> IsValidGrid;
 
 		inline static FStreamingGenerationNullErrorHandler NullErrorHandler;
 	};
@@ -791,6 +810,7 @@ public:
 		: bEnableStreaming(Params.bEnableStreaming)
 		, ModifiedActorsDescList(Params.ModifiedActorsDescList)
 		, FilteredClasses(Params.FilteredClasses)
+		, IsValidGrid(Params.IsValidGrid)
 		, ErrorHandler(Params.ErrorHandler)
 		, ActorGuidsToContainerMap(Params.ActorGuidsToContainerMap)
 	{}
@@ -907,6 +927,7 @@ private:
 	bool bEnableStreaming;
 	FActorDescList* ModifiedActorsDescList;
 	TArray<TSubclassOf<AActor>> FilteredClasses;
+	TFunction<bool(FName)> IsValidGrid;
 	IStreamingGenerationErrorHandler* ErrorHandler;
 
 	/** Maps containers to their container descriptor */
@@ -955,6 +976,7 @@ bool UWorldPartition::GenerateContainerStreaming(const UActorDescContainer* InAc
 
 	FWorldPartitionStreamingGenerator::FWorldPartitionStreamingGeneratorParams StreamingGeneratorParams;
 	StreamingGeneratorParams.ModifiedActorsDescList = ModifiedActorsDescList;
+	StreamingGeneratorParams.IsValidGrid = [this](FName GridName) { return RuntimeHash->IsValidGrid(GridName); };
 	StreamingGeneratorParams.ErrorHandler = ErrorHandler;
 	StreamingGeneratorParams.bEnableStreaming = IsStreamingEnabled();
 
@@ -999,6 +1021,7 @@ void UWorldPartition::GenerateHLOD(ISourceControlHelper* SourceControlHelper, bo
 	StreamingGeneratorParams.ErrorHandler = &LogErrorHandler;
 	StreamingGeneratorParams.bEnableStreaming = IsStreamingEnabled();
 	StreamingGeneratorParams.FilteredClasses.Add(AWorldPartitionHLOD::StaticClass());
+	StreamingGeneratorParams.IsValidGrid = [this](FName GridName) { return RuntimeHash->IsValidGrid(GridName); };
 
 	FWorldPartitionStreamingGenerator StreamingGenerator(StreamingGeneratorParams);
 	StreamingGenerator.PreparationPhase(ActorDescContainer);
@@ -1012,27 +1035,31 @@ void UWorldPartition::GenerateHLOD(ISourceControlHelper* SourceControlHelper, bo
 
 void UWorldPartition::CheckForErrors(IStreamingGenerationErrorHandler* ErrorHandler) const
 {
-	FCheckForErrorsParams Params;
-	Params.ErrorHandler = ErrorHandler;
-	Params.ActorDescContainer = ActorDescContainer;
-	Params.bEnableStreaming = IsStreamingEnabled();
+	FActorDescList ModifiedActorDescList;
+
+	FWorldPartitionStreamingGenerator::FWorldPartitionStreamingGeneratorParams StreamingGeneratorParams;
+	StreamingGeneratorParams.ModifiedActorsDescList = &ModifiedActorDescList;
+	StreamingGeneratorParams.ErrorHandler = ErrorHandler;
+	StreamingGeneratorParams.IsValidGrid = [this](FName GridName) { return RuntimeHash->IsValidGrid(GridName); };
+	StreamingGeneratorParams.bEnableStreaming = IsStreamingEnabled();
 
-	ForEachActorDescContainer([&Params](const UActorDescContainer* InActorDescContainer)
+	ForEachActorDescContainer([&StreamingGeneratorParams](const UActorDescContainer* InActorDescContainer)
 	{
 		for (FActorDescList::TConstIterator<> ActorDescIt(InActorDescContainer); ActorDescIt; ++ActorDescIt)
 		{
-			check(!Params.ActorGuidsToContainerMap.Contains(ActorDescIt->GetGuid()));
-			Params.ActorGuidsToContainerMap.Add(ActorDescIt->GetGuid(), InActorDescContainer);
+			check(!StreamingGeneratorParams.ActorGuidsToContainerMap.Contains(ActorDescIt->GetGuid()));
+			StreamingGeneratorParams.ActorGuidsToContainerMap.Add(ActorDescIt->GetGuid(), InActorDescContainer);
 		}
 	});
 
-	ForEachActorDescContainer([this, &Params](const UActorDescContainer* InActorDescContainer)
+	ForEachActorDescContainer([this, &StreamingGeneratorParams](const UActorDescContainer* InActorDescContainer)
 	{
-		Params.ActorDescContainer = InActorDescContainer;
-		CheckForErrors(Params);
+		FWorldPartitionStreamingGenerator StreamingGenerator(StreamingGeneratorParams);
+		StreamingGenerator.PreparationPhase(InActorDescContainer);
 	});
 }
 
+/* Deprecated */
 void UWorldPartition::CheckForErrors(IStreamingGenerationErrorHandler* ErrorHandler, const UActorDescContainer* ActorDescContainer, bool bEnableStreaming, bool)
 {
 	FCheckForErrorsParams Params;
@@ -1050,6 +1077,7 @@ void UWorldPartition::CheckForErrors(const FCheckForErrorsParams& Params)
 	FWorldPartitionStreamingGenerator::FWorldPartitionStreamingGeneratorParams StreamingGeneratorParams;
 	StreamingGeneratorParams.ModifiedActorsDescList = Params.ActorDescContainer->GetWorld() ? &ModifiedActorDescList : nullptr;
 	StreamingGeneratorParams.ErrorHandler = Params.ErrorHandler;
+	StreamingGeneratorParams.IsValidGrid = [](FName GridName) { return true; };
 	StreamingGeneratorParams.bEnableStreaming = Params.bEnableStreaming;
 	StreamingGeneratorParams.ActorGuidsToContainerMap = Params.ActorGuidsToContainerMap;
 
diff --git a/Engine/Source/Runtime/Engine/Public/WorldPartition/ErrorHandling/WorldPartitionStreamingGenerationErrorHandler.h b/Engine/Source/Runtime/Engine/Public/WorldPartition/ErrorHandling/WorldPartitionStreamingGenerationErrorHandler.h
--- a/Engine/Source/Runtime/Engine/Public/WorldPartition/ErrorHandling/WorldPartitionStreamingGenerationErrorHandler.h
+++ b/Engine/Source/Runtime/Engine/Public/WorldPartition/ErrorHandling/WorldPartitionStreamingGenerationErrorHandler.h
@@ -16,6 +16,11 @@ public:
 	UE_DEPRECATED(5.2, "OnInvalidReference is deprecated, use the version which takes an optional actor descriptor view.")
 	virtual void OnInvalidReference(const FWorldPartitionActorDescView& ActorDescView, const FGuid& ReferenceGuid) {}
 
+	/** 
+	 * Called when an actor has an invalid runtime grid.
+	 */
+	virtual void OnInvalidRuntimeGrid(const FWorldPartitionActorDescView& ActorDescView, FName GridName) = 0;
+
 	/** 
 	 * Called when an actor references an invalid actor.
 	 */
diff --git a/Engine/Source/Runtime/Engine/Public/WorldPartition/ErrorHandling/WorldPartitionStreamingGenerationLogErrorHandler.h b/Engine/Source/Runtime/Engine/Public/WorldPartition/ErrorHandling/WorldPartitionStreamingGenerationLogErrorHandler.h
--- a/Engine/Source/Runtime/Engine/Public/WorldPartition/ErrorHandling/WorldPartitionStreamingGenerationLogErrorHandler.h
+++ b/Engine/Source/Runtime/Engine/Public/WorldPartition/ErrorHandling/WorldPartitionStreamingGenerationLogErrorHandler.h
@@ -7,6 +7,7 @@
 class ENGINE_API FStreamingGenerationLogErrorHandler : public IStreamingGenerationErrorHandler
 {
 public:
+	virtual void OnInvalidRuntimeGrid(const FWorldPartitionActorDescView& ActorDescView, FName GridName) override;
 	virtual void OnInvalidReference(const FWorldPartitionActorDescView& ActorDescView, const FGuid& ReferenceGuid, FWorldPartitionActorDescView* ReferenceActorDescView) override;
 	virtual void OnInvalidReferenceGridPlacement(const FWorldPartitionActorDescView& ActorDescView, const FWorldPartitionActorDescView& ReferenceActorDescView) override;
 	virtual void OnInvalidReferenceDataLayers(const FWorldPartitionActorDescView& ActorDescView, const FWorldPartitionActorDescView& ReferenceActorDescView) override;
diff --git a/Engine/Source/Runtime/Engine/Public/WorldPartition/ErrorHandling/WorldPartitionStreamingGenerationNullErrorHandler.h b/Engine/Source/Runtime/Engine/Public/WorldPartition/ErrorHandling/WorldPartitionStreamingGenerationNullErrorHandler.h
--- a/Engine/Source/Runtime/Engine/Public/WorldPartition/ErrorHandling/WorldPartitionStreamingGenerationNullErrorHandler.h
+++ b/Engine/Source/Runtime/Engine/Public/WorldPartition/ErrorHandling/WorldPartitionStreamingGenerationNullErrorHandler.h
@@ -7,6 +7,7 @@
 class ENGINE_API FStreamingGenerationNullErrorHandler : public IStreamingGenerationErrorHandler
 {
 public:
+	virtual void OnInvalidRuntimeGrid(const FWorldPartitionActorDescView& ActorDescView, FName GridName) override {}
 	virtual void OnInvalidReference(const FWorldPartitionActorDescView& ActorDescView, const FGuid& ReferenceGuid, FWorldPartitionActorDescView* ReferenceActorDescView) override {}
 	virtual void OnInvalidReferenceGridPlacement(const FWorldPartitionActorDescView& ActorDescView, const FWorldPartitionActorDescView& ReferenceActorDescView) override {}
 	virtual void OnInvalidReferenceDataLayers(const FWorldPartitionActorDescView& ActorDescView, const FWorldPartitionActorDescView& ReferenceActorDescView) override {}
diff --git a/Engine/Source/Runtime/Engine/Public/WorldPartition/ErrorHandling/WorldPartitionStreamingGenerationTokenizedMessageErrorHandler.h b/Engine/Source/Runtime/Engine/Public/WorldPartition/ErrorHandling/WorldPartitionStreamingGenerationTokenizedMessageErrorHandler.h
--- a/Engine/Source/Runtime/Engine/Public/WorldPartition/ErrorHandling/WorldPartitionStreamingGenerationTokenizedMessageErrorHandler.h
+++ b/Engine/Source/Runtime/Engine/Public/WorldPartition/ErrorHandling/WorldPartitionStreamingGenerationTokenizedMessageErrorHandler.h
@@ -8,6 +8,7 @@ class FTokenizedMessage;
 
 class ENGINE_API ITokenizedMessageErrorHandler : public IStreamingGenerationErrorHandler
 {
+	virtual void OnInvalidRuntimeGrid(const FWorldPartitionActorDescView& ActorDescView, FName GridName) override;
 	virtual void OnInvalidReference(const FWorldPartitionActorDescView& ActorDescView, const FGuid& ReferenceGuid, FWorldPartitionActorDescView* ReferenceActorDescView) override;
 	virtual void OnInvalidReferenceGridPlacement(const FWorldPartitionActorDescView& ActorDescView, const FWorldPartitionActorDescView& ReferenceActorDescView) override;
 	virtual void OnInvalidReferenceDataLayers(const FWorldPartitionActorDescView& ActorDescView, const FWorldPartitionActorDescView& ReferenceActorDescView) override;
diff --git a/Engine/Source/Runtime/Engine/Public/WorldPartition/WorldPartitionRuntimeHash.h b/Engine/Source/Runtime/Engine/Public/WorldPartition/WorldPartitionRuntimeHash.h
--- a/Engine/Source/Runtime/Engine/Public/WorldPartition/WorldPartitionRuntimeHash.h
+++ b/Engine/Source/Runtime/Engine/Public/WorldPartition/WorldPartitionRuntimeHash.h
@@ -90,6 +90,7 @@ class ENGINE_API UWorldPartitionRuntimeHash : public UObject, public IWorldParti
 	virtual bool GenerateStreaming(class UWorldPartitionStreamingPolicy* StreamingPolicy, const IStreamingGenerationContext* StreamingGenerationContext, TArray<FString>* OutPackagesToGenerate) { return false; }
 	virtual void FlushStreaming() {}
 	virtual bool GenerateHLOD(ISourceControlHelper* SourceControlHelper, const IStreamingGenerationContext* StreamingGenerationContext, bool bCreateActorsOnly) { return false; }
+	virtual bool IsValidGrid(FName GridName) const { return false; }
 	virtual void DrawPreview() const {}
 
 	virtual URuntimeHashExternalStreamingObjectBase* StoreToExternalStreamingObject(UObject* StreamingObjectOuter, FName StreamingObjectName) { return nullptr; }
diff --git a/Engine/Source/Runtime/Engine/Public/WorldPartition/WorldPartitionRuntimeSpatialHash.h b/Engine/Source/Runtime/Engine/Public/WorldPartition/WorldPartitionRuntimeSpatialHash.h
--- a/Engine/Source/Runtime/Engine/Public/WorldPartition/WorldPartitionRuntimeSpatialHash.h
+++ b/Engine/Source/Runtime/Engine/Public/WorldPartition/WorldPartitionRuntimeSpatialHash.h
@@ -276,6 +276,7 @@ public:
 	virtual UWorldPartitionRuntimeCell* GetCellForPackage(const FWorldPartitionCookPackage& PackageToCook) const override;
 	virtual void FlushStreaming() override;
 	virtual bool GenerateHLOD(ISourceControlHelper* SourceControlHelper, const IStreamingGenerationContext* StreamingGenerationContext, bool bCreateActorsOnly) override;
+	virtual bool IsValidGrid(FName GridName) const override;
 	virtual void DrawPreview() const override;
 
 	virtual URuntimeHashExternalStreamingObjectBase* StoreToExternalStreamingObject(UObject* StreamingObjectOuter, FName StreamingObjectName) override;
