Materials: Adds support for new material graph node, Switch, which implements the functionality of a switch, or discrete selector / multiplexer, for the material graph. This allows more compact graphs where multiple paths or inputs need to be decided on either via constant or uniform parameters. It also allows for several editor mode materials, used for tool implementations, to avoid needing custom nodes, unusable in Restricted mode, or complex If/Else graphs.

diff --git a/Engine/Source/Runtime/Engine/Classes/Materials/MaterialExpressionSwitch.h b/Engine/Source/Runtime/Engine/Classes/Materials/MaterialExpressionSwitch.h
new file mode 100644
--- /dev/null
+++ b/Engine/Source/Runtime/Engine/Classes/Materials/MaterialExpressionSwitch.h
@@ -0,0 +1,83 @@
+// Copyright Epic Games, Inc. All Rights Reserved.
+
+
+#pragma once
+
+#include "CoreMinimal.h"
+#include "UObject/ObjectMacros.h"
+#include "MaterialExpressionIO.h"
+#include "Materials/MaterialExpression.h"
+#include "MaterialExpressionSwitch.generated.h"
+
+struct FPropertyChangedEvent;
+
+UENUM()
+enum ESwitchMaterialOutputType
+{
+	TMMOT_Float1,
+	TMMOT_Float2,
+	TMMOT_Float3,
+	TMMOT_Float4,
+	TMMOT_MAX,
+};
+
+USTRUCT()
+struct FSwitchCustomInput
+{
+	GENERATED_USTRUCT_BODY()
+
+	UPROPERTY(EditAnywhere, Category=CustomInput)
+	FName InputName;
+
+	UPROPERTY()
+	FExpressionInput Input;
+};
+
+UCLASS(collapsecategories, hidecategories=Object, MinimalAPI)
+class UMaterialExpressionSwitch : public UMaterialExpression
+{
+	GENERATED_UCLASS_BODY()
+
+	UPROPERTY(EditAnywhere, Category= UMaterialExpressionSwitch)
+	FString Description;
+
+	UPROPERTY(meta = (RequiredInput = "false", ToolTip = "Defaults to 'ConstSwitchValue' if not specified"))
+	FExpressionInput SwitchValue;
+
+	/** only used if Selector is not hooked up */
+	UPROPERTY(EditAnywhere, Category= UMaterialExpressionSwitch, meta=(OverridingInputProperty = "SwitchValue"))
+	float ConstSwitchValue;
+
+	UPROPERTY(meta = (RequiredInput = "false", ToolTip = "Defaults to 'ConstDefault' if not specified"))
+	FExpressionInput Default;
+
+	/** only used if Selector is not hooked up */
+	UPROPERTY(EditAnywhere, Category= UMaterialExpressionSwitch, meta=(OverridingInputProperty = "Default"))
+	float ConstDefault;
+
+	UPROPERTY(EditAnywhere, Category= UMaterialExpressionSwitch)
+	TArray<struct FSwitchCustomInput> Inputs;
+
+	//~ Begin UObject Interface.
+#if WITH_EDITOR
+	virtual void PostEditChangeProperty(FPropertyChangedEvent& PropertyChangedEvent) override;
+	void RebuildOutputs();
+#endif // WITH_EDITOR
+	virtual void Serialize(FStructuredArchive::FRecord Record) override;
+	//~ End UObject Interface.
+	
+	//~ Begin UMaterialExpression Interface
+#if WITH_EDITOR
+	virtual int32 Compile(class FMaterialCompiler* Compiler, int32 OutputIndex) override;
+	virtual void GetCaption(TArray<FString>& OutCaptions) const override;
+	virtual FText GetCreationName() const override { return FText::FromString(TEXT("Switch")); }
+	virtual const TArray<FExpressionInput*> GetInputs() override;
+	virtual FExpressionInput* GetInput(int32 InputIndex) override;
+	virtual FName GetInputName(int32 InputIndex) const override;
+	virtual uint32 GetInputType(int32 InputIndex) override {return MCT_Unknown;}
+#endif // WITH_EDITOR
+	//~ End UMaterialExpression Interface
+};
+
+
+
diff --git a/Engine/Source/Runtime/Engine/Private/Materials/HLSLMaterialTranslator.cpp b/Engine/Source/Runtime/Engine/Private/Materials/HLSLMaterialTranslator.cpp
--- a/Engine/Source/Runtime/Engine/Private/Materials/HLSLMaterialTranslator.cpp
+++ b/Engine/Source/Runtime/Engine/Private/Materials/HLSLMaterialTranslator.cpp
@@ -6842,6 +6842,109 @@ int32 FHLSLMaterialTranslator::DBufferTextureLookup(int32 ViewportUV, uint32 DBu
 	return AddCodeChunk(MCT_Float4, TEXT("MaterialExpressionDBufferTextureLookup(Parameters, %s, %d)"), *CoerceParameter(BufferUV, MCT_Float2), (int)DBufferTextureIndex);
 }
 
+int32 FHLSLMaterialTranslator::Switch(int32 SwitchValueInput, int32 DefaultInput, TArray<int32>& CompiledInputs)
+{
+	if (SwitchValueInput == INDEX_NONE)
+	{
+		return INDEX_NONE;
+	}
+
+	if (DefaultInput == INDEX_NONE)
+	{
+		return INDEX_NONE;
+	}
+
+	if (CompiledInputs.Num() == 0)
+	{
+		return INDEX_NONE;
+	}
+
+	for (int32 Input : CompiledInputs)
+	{
+		if (Input == INDEX_NONE)
+		{
+			return INDEX_NONE;
+		}
+	}
+
+	// If our selector is constant, we can skip the generation of code below and simply
+    // grab/return the appropriate input directly. This should greatly simplify the
+    // expression in constant parameter cases.
+	FMaterialUniformExpression* ExpressionSwitchValue = GetParameterUniformExpression(SwitchValueInput);
+	EMaterialValueType SwitchValueType = GetParameterType(SwitchValueInput);
+	if (SwitchValueType == MCT_Float && ExpressionSwitchValue && ExpressionSwitchValue->IsConstant())
+	{		
+		FLinearColor Value;
+		FMaterialRenderContext DummyContext(nullptr, *Material, nullptr);
+		ExpressionSwitchValue->GetNumberValue(DummyContext, Value);
+
+		for (int32 InputIndex = 0; InputIndex < CompiledInputs.Num(); ++InputIndex)
+		{
+			if (Value.R >= InputIndex && Value.R < InputIndex + 1)
+			{
+				return CompiledInputs[InputIndex];
+			}
+		}
+		return DefaultInput;
+	}
+	
+	EMaterialValueType ResultType = GetParameterType(CompiledInputs[0]);
+	for (int32 Input : CompiledInputs)
+	{
+		ResultType = GetArithmeticResultType(ResultType, GetParameterType(Input));
+	}
+
+	FString SwitchCode = TEXT("");
+	FString SwitchValueParameter = CoerceParameter(SwitchValueInput, MCT_Float);
+	FString DefaultParameter = CoerceParameter(DefaultInput, ResultType);
+
+
+	// TODO: The following represents two potential implementations for the dynamic switch logic in HLSL.
+	// The first option is an expression that mathamatically cancels out any "unmatched" case,
+	// leaving only the desired value passing through. This approach works, but is less optimal 
+	// due to needing to evalulate every potential case value as part of the process.
+	// The second option is potentially more optimal, as it uses a native switch statement in HLSL, something
+	// that the compiler should be able to optimize more easily with less unneeded evaluations. However,
+	// this approach currently doesn't work well with the existing material translation code in the engine.
+	// The existing design supports expressions returning floating point type values, which are used by
+	// the analyical derivative components elsewhere. Attempting to return a void type or insert type-less
+	// statements into the generated code has been problematic.
+	//
+	// Ideally, the second option can be made to function correctly and replace the first, which is why
+	// it's being left in as a preprocessor switch for the time being. However, for expedieceny, we want
+	//to use the first option for the time being, to provide a working, if less efficicent, switch implementation.
+#if 1
+	// We form the "switch" statement from an equivalent math expression,
+    // using step functions and additions to select between inputs.
+	for (int32 InputIndex = 0; InputIndex < CompiledInputs.Num(); ++InputIndex)
+	{
+		FString InputParameter = CoerceParameter(CompiledInputs[InputIndex], ResultType);
+		SwitchCode += FString::Printf(TEXT("step(%f, floor(%s) + 0.5f) * step(floor(%s) + 0.5f, %f) * %s +"),
+			float(InputIndex), *SwitchValueParameter,        // If greater than the lower bound
+			*SwitchValueParameter, float(InputIndex) + 1.0f, // and less  than the upper bound parameters
+			*InputParameter);                             // return the selected selected input
+	}
+	SwitchCode += FString::Printf(TEXT("(step(floor(%s) + 0.5f, %f) + step(%f, floor(%s) + 0.5f)) * %s"),
+	    *SwitchValueParameter, float(0),                    // If less than Zero,
+		float(CompiledInputs.Num()), *SwitchValueParameter, // and greater than number of inputs,
+		*DefaultParameter);                              // return the default
+	return AddCodeChunk(ResultType, *SwitchCode);
+#else
+	int32 ReturnValue = AddCodeChunk(ResultType, TEXT("0"));
+	FString ReturnParameter = CoerceParameter(ReturnValue, ResultType);
+
+	SwitchCode += FString::Printf(TEXT("switch(%s){"), *SwitchValueParameter);
+	for (int32 InputIndex = 0; InputIndex < CompiledInputs.Num(); ++InputIndex)
+	{
+		FString InputParameter = CoerceParameter(CompiledInputs[InputIndex], ResultType);
+		SwitchCode += FString::Printf(TEXT("case %d: %s = %s; break;"), InputIndex, *ReturnParameter, *InputParameter);
+	}
+	SwitchCode += FString::Printf(TEXT("default: %s = %s; }"), *ReturnParameter, *DefaultParameter);
+	
+	return AddCodeChunkInner(MCT_VoidStatement, EDerivativeStatus::NotAware, true, *SwitchCode);
+#endif
+}
+
 int32 FHLSLMaterialTranslator::Texture(UTexture* InTexture, int32& TextureReferenceIndex, EMaterialSamplerType SamplerType, ESamplerSourceMode SamplerSource, ETextureMipValueMode MipValueMode)
 {
 	EMaterialValueType ShaderType = InTexture->GetMaterialType();
diff --git a/Engine/Source/Runtime/Engine/Private/Materials/HLSLMaterialTranslator.h b/Engine/Source/Runtime/Engine/Private/Materials/HLSLMaterialTranslator.h
--- a/Engine/Source/Runtime/Engine/Private/Materials/HLSLMaterialTranslator.h
+++ b/Engine/Source/Runtime/Engine/Private/Materials/HLSLMaterialTranslator.h
@@ -892,6 +892,7 @@ protected:
 
 	virtual int32 SceneColor(int32 Offset, int32 ViewportUV, bool bUseOffset) override;
 
+	virtual int32 Switch(int32 SwitchValueInput, int32 DefaultInput, TArray<int32>& CompiledInputs) override;
 	virtual int32 Texture(UTexture* InTexture, int32& TextureReferenceIndex, EMaterialSamplerType SamplerType, ESamplerSourceMode SamplerSource = SSM_FromTextureAsset, ETextureMipValueMode MipValueMode = TMVM_None) override;
 	virtual int32 TextureParameter(FName ParameterName, UTexture* DefaultValue, int32& TextureReferenceIndex, EMaterialSamplerType SamplerType, ESamplerSourceMode SamplerSource = SSM_FromTextureAsset) override;
 
diff --git a/Engine/Source/Runtime/Engine/Private/Materials/MaterialExpressions.cpp b/Engine/Source/Runtime/Engine/Private/Materials/MaterialExpressions.cpp
--- a/Engine/Source/Runtime/Engine/Private/Materials/MaterialExpressions.cpp
+++ b/Engine/Source/Runtime/Engine/Private/Materials/MaterialExpressions.cpp
@@ -219,6 +219,7 @@
 #include "Materials/MaterialExpressionTextureBase.h"
 #include "Materials/MaterialExpressionTextureObject.h"
 #include "Materials/MaterialExpressionTextureProperty.h"
+#include "Materials/MaterialExpressionSwitch.h"
 #include "Materials/MaterialExpressionTextureSample.h"
 #include "Materials/MaterialExpressionParticleSubUV.h"
 #include "Materials/MaterialExpressionParticleSubUVProperties.h"
@@ -13264,6 +13265,167 @@ void UMaterialFunctionInterfaceEditorOnlyData::Serialize(FArchive& Ar)
 #endif
 }
 
+///////////////////////////////////////////////////////////////////////////////
+// UMaterialExpressionSwitch
+///////////////////////////////////////////////////////////////////////////////
+UMaterialExpressionSwitch::UMaterialExpressionSwitch(const FObjectInitializer& ObjectInitializer)
+	: Super(ObjectInitializer)
+{
+#if WITH_EDITORONLY_DATA
+	// Structure to hold one-time initialization
+	struct FConstructorStatics
+	{
+		FText NAME_Math;
+		FConstructorStatics()
+			: NAME_Math(LOCTEXT("Math", "Math"))
+		{
+		}
+	};
+	static FConstructorStatics ConstructorStatics;
+#endif // WITH_EDITORONLY_DATA
+
+	Description = TEXT("Switch");
+
+#if WITH_EDITORONLY_DATA
+	MenuCategories.Add(ConstructorStatics.NAME_Math);
+#endif
+
+	Inputs.Add(FSwitchCustomInput());
+	Inputs[0].InputName = TEXT("");
+
+#if WITH_EDITORONLY_DATA
+	bCollapsed = false;
+#endif // WITH_EDITORONLY_DATA
+}
+
+#if WITH_EDITOR
+int32 UMaterialExpressionSwitch::Compile(class FMaterialCompiler* Compiler, int32 OutputIndex)
+{
+	// if the input is hooked up, use it, otherwise use the internal constant
+	int32 CompiledSwitchValue = SwitchValue.GetTracedInput().Expression ? SwitchValue.Compile(Compiler) : Compiler->Constant(ConstSwitchValue);
+	int32 CompiledDefault = Default.GetTracedInput().Expression ? Default.Compile(Compiler) : Compiler->Constant(ConstDefault);
+
+	TArray<int32> CompiledInputs;
+
+	for (int32 i = 0; i < Inputs.Num(); i++)
+	{
+		if (!Inputs[i].Input.GetTracedInput().Expression)
+		{
+			return Compiler->Errorf(TEXT("Texture Multiplexer missing input %d (%s)"), i + 1, *Inputs[i].InputName.ToString());
+		}
+		int32 InputCode = Inputs[i].Input.Compile(Compiler);
+		if (InputCode < 0)
+		{
+			return InputCode;
+		}
+		CompiledInputs.Add(InputCode);
+	}
+
+	return Compiler->Switch(CompiledSwitchValue, CompiledDefault, CompiledInputs);
+}
+
+
+void UMaterialExpressionSwitch::GetCaption(TArray<FString>& OutCaptions) const
+{
+	OutCaptions.Add(Description);
+}
+
+
+const TArray<FExpressionInput*> UMaterialExpressionSwitch::GetInputs()
+{
+	TArray<FExpressionInput*> Result;
+	Result.Add(&SwitchValue);
+	Result.Add(&Default);
+	for (int32 i = 0; i < Inputs.Num(); i++)
+	{
+		Result.Add(&Inputs[i].Input);
+	}
+	return Result;
+}
+
+FExpressionInput* UMaterialExpressionSwitch::GetInput(int32 InputIndex)
+{
+	if (InputIndex < 0 || InputIndex > Inputs.Num() + 1)
+	{
+		return nullptr;
+	}
+	switch (InputIndex)
+	{
+	case 0:
+		return &SwitchValue;
+	case 1:
+		return &Default;
+	default:
+		return &Inputs[InputIndex - 2].Input;
+	}
+}
+
+FName UMaterialExpressionSwitch::GetInputName(int32 InputIndex) const
+{
+	if (InputIndex < 0 || InputIndex > Inputs.Num() + 1)
+	{
+		return NAME_None;
+	}
+	switch (InputIndex)
+	{
+	case 0:
+		return FName(TEXT("SwitchValue"));
+	case 1:
+		return FName(TEXT("Default"));
+	default:
+		return Inputs[InputIndex - 2].InputName;
+	}
+}
+
+void UMaterialExpressionSwitch::PostEditChangeProperty(FPropertyChangedEvent& PropertyChangedEvent)
+{
+	// strip any spaces from input name
+	FProperty* PropertyThatChanged = PropertyChangedEvent.Property;
+	if (PropertyThatChanged && PropertyThatChanged->GetFName() == GET_MEMBER_NAME_CHECKED(FCustomInput, InputName))
+	{
+		for (FSwitchCustomInput& Input : Inputs)
+		{
+			FString InputName = Input.InputName.ToString();
+			if (InputName.ReplaceInline(TEXT(" "), TEXT("")) > 0)
+			{
+				Input.InputName = *InputName;
+			}
+		}
+	}
+
+	RebuildOutputs();
+
+	if (PropertyChangedEvent.MemberProperty && GraphNode)
+	{
+		const FName PropertyName = PropertyChangedEvent.MemberProperty->GetFName();
+		if (PropertyName == GET_MEMBER_NAME_CHECKED(UMaterialExpressionCustom, Inputs))
+		{
+			GraphNode->ReconstructNode();
+		}
+	}
+
+	Super::PostEditChangeProperty(PropertyChangedEvent);
+}
+
+void UMaterialExpressionSwitch::RebuildOutputs()
+{
+	Outputs.Reset(1);
+	bShowOutputNameOnPin = false;
+	Outputs.Add(FExpressionOutput(TEXT("")));
+}
+
+#endif // WITH_EDITOR
+
+void UMaterialExpressionSwitch::Serialize(FStructuredArchive::FRecord Record)
+{
+	Super::Serialize(Record);
+	FArchive& UnderlyingArchive = Record.GetUnderlyingArchive();
+
+	UnderlyingArchive.UsingCustomVersion(FRenderingObjectVersion::GUID);
+	UnderlyingArchive.UsingCustomVersion(FUE5MainStreamObjectVersion::GUID);
+	UnderlyingArchive.UsingCustomVersion(FUE5LWCRenderingStreamObjectVersion::GUID);
+}
+
 ///////////////////////////////////////////////////////////////////////////////
 // UMaterialFunctionInterface
 ///////////////////////////////////////////////////////////////////////////////
diff --git a/Engine/Source/Runtime/Engine/Public/MaterialCompiler.h b/Engine/Source/Runtime/Engine/Public/MaterialCompiler.h
--- a/Engine/Source/Runtime/Engine/Public/MaterialCompiler.h
+++ b/Engine/Source/Runtime/Engine/Public/MaterialCompiler.h
@@ -351,6 +351,7 @@ public:
 	virtual int32 ParticleSize() = 0;
 
 	virtual int32 If(int32 A,int32 B,int32 AGreaterThanB,int32 AEqualsB,int32 ALessThanB,int32 Threshold) = 0;
+	virtual int32 Switch(int32 SwitchValueInput, int32 DefaultInput, TArray<int32>& CompiledInputs) = 0;
 
 	virtual int32 TextureCoordinate(uint32 CoordinateIndex, bool UnMirrorU, bool UnMirrorV) = 0;
 	virtual int32 TextureSample(int32 Texture,int32 Coordinate,enum EMaterialSamplerType SamplerType,int32 MipValue0Index=INDEX_NONE,int32 MipValue1Index=INDEX_NONE,ETextureMipValueMode MipValueMode=TMVM_None,ESamplerSourceMode SamplerSource=SSM_FromTextureAsset,int32 TextureReferenceIndex=INDEX_NONE, bool AutomaticViewMipBias=false, bool AdaptiveVirtualTexture=false, bool EnableFeedback = true) = 0;
@@ -863,7 +864,8 @@ public:
 	virtual int32 SphericalParticleOpacity(int32 Density) override { return Compiler->SphericalParticleOpacity(Density); }
 
 	virtual int32 If(int32 A, int32 B, int32 AGreaterThanB, int32 AEqualsB, int32 ALessThanB, int32 Threshold) override { return Compiler->If(A, B, AGreaterThanB, AEqualsB, ALessThanB, Threshold); }
-
+	virtual int32 Switch(int32 SwitchValueInput, int32 DefaultInput, TArray<int32>& CompiledInputs) override { return Compiler->Switch(SwitchValueInput, DefaultInput, CompiledInputs); };
+	
 	virtual int32 TextureSample(int32 InTexture, int32 Coordinate, enum EMaterialSamplerType SamplerType, int32 MipValue0Index, int32 MipValue1Index, ETextureMipValueMode MipValueMode, ESamplerSourceMode SamplerSource, int32 TextureReferenceIndex, bool AutomaticViewMipBias, bool AdaptiveVirtualTexture, bool EnableFeedback) override
 	{
 		return Compiler->TextureSample(InTexture, Coordinate, SamplerType, MipValue0Index, MipValue1Index, MipValueMode, SamplerSource, TextureReferenceIndex, AutomaticViewMipBias, AdaptiveVirtualTexture);
