MaterialExpressionSwitch: Fix NonUnity build error

diff --git a/Engine/Source/Runtime/Engine/Classes/Materials/MaterialExpressionSwitch.h b/Engine/Source/Runtime/Engine/Classes/Materials/MaterialExpressionSwitch.h
--- a/Engine/Source/Runtime/Engine/Classes/Materials/MaterialExpressionSwitch.h
+++ b/Engine/Source/Runtime/Engine/Classes/Materials/MaterialExpressionSwitch.h
@@ -7,6 +7,7 @@
 #include "UObject/ObjectMacros.h"
 #include "MaterialExpressionIO.h"
 #include "Materials/MaterialExpression.h"
+#include "MaterialValueType.h"
 #include "MaterialExpressionSwitch.generated.h"
 
 struct FPropertyChangedEvent;
