Fix static analyzer warnings.

diff --git a/Engine/Source/Runtime/Engine/Private/Tests/WorldPartition/WorldPartitionStreamingGenerationTests.cpp b/Engine/Source/Runtime/Engine/Private/Tests/WorldPartition/WorldPartitionStreamingGenerationTests.cpp
--- a/Engine/Source/Runtime/Engine/Private/Tests/WorldPartition/WorldPartitionStreamingGenerationTests.cpp
+++ b/Engine/Source/Runtime/Engine/Private/Tests/WorldPartition/WorldPartitionStreamingGenerationTests.cpp
@@ -47,12 +47,14 @@ namespace WorldPartitionTests
 			return false;
 		}
 
+		check(World);
 		UWorldPartition* WorldPartition = World->GetWorldPartition();
 		if (!TestTrue(TEXT("Missing World Partition Object"), !!WorldPartition))
 		{
 			return false;
 		}
-
+		
+		check(WorldPartition);
 		UActorDescContainer* ActorDescMainContainer = WorldPartition->GetActorDescContainer();
 		if (!TestTrue(TEXT("Missing World Partition Container"), !!ActorDescMainContainer))
 		{
