add support for geometry collections to UWorld::ComponentSweepMulti and UWorld::ComponentOverlapMulti via the Chaos::FPhysicsObject

- also added an implementation of both those functions to also take in an ECollisionChannel parameter.

diff --git a/Engine/Source/Runtime/Engine/Classes/Engine/World.h b/Engine/Source/Runtime/Engine/Classes/Engine/World.h
--- a/Engine/Source/Runtime/Engine/Classes/Engine/World.h
+++ b/Engine/Source/Runtime/Engine/Classes/Engine/World.h
@@ -2162,6 +2162,20 @@ public:
 	bool ComponentSweepMulti(TArray<struct FHitResult>& OutHits, class UPrimitiveComponent* PrimComp, const FVector& Start, const FVector& End, const FQuat& Rot,    const FComponentQueryParams& Params) const;
 	bool ComponentSweepMulti(TArray<struct FHitResult>& OutHits, class UPrimitiveComponent* PrimComp, const FVector& Start, const FVector& End, const FRotator& Rot, const FComponentQueryParams& Params) const;
 
+	/**
+	 *  Sweep the geometry of the supplied component using a specific channel, and determine the set of components that it hits.
+	 *  @note The overload taking rotation as an FQuat is slightly faster than the version using FRotator (which will be converted to an FQuat)..
+	 *  @param  OutHits         Array of hits found between ray and the world
+	 *  @param  PrimComp        Component's geometry to test against the world. Transform of this component is ignored
+	 *  @param  Start           Start location of the trace
+	 *  @param  End             End location of the trace
+	 *  @param  Rot             Rotation of PrimComp geometry for test against the world (rotation remains constant over sweep)
+	 *  @param  Params          Additional parameters used for the trace
+	 *  @return TRUE if OutHits contains any blocking hit entries
+	 */
+	bool ComponentSweepMultiByChannel(TArray<struct FHitResult>& OutHits, class UPrimitiveComponent* PrimComp, const FVector& Start, const FVector& End, const FQuat& Rot, ECollisionChannel TraceChannel, const FComponentQueryParams& Params) const;
+	bool ComponentSweepMultiByChannel(TArray<struct FHitResult>& OutHits, class UPrimitiveComponent* PrimComp, const FVector& Start, const FVector& End, const FRotator& Rot, ECollisionChannel TraceChannel, const FComponentQueryParams& Params) const;
+
 	// COMPONENT OVERLAP
 
 	/**
@@ -4235,6 +4249,12 @@ FORCEINLINE_DEBUGGABLE bool UWorld::ComponentSweepMulti(TArray<struct FHitResult
 	return ComponentSweepMulti(OutHits, PrimComp, Start, End, Rot.Quaternion(), Params);
 }
 
+FORCEINLINE_DEBUGGABLE bool UWorld::ComponentSweepMultiByChannel(TArray<struct FHitResult>& OutHits, class UPrimitiveComponent* PrimComp, const FVector& Start, const FVector& End, const FRotator& Rot, ECollisionChannel TraceChannel, const FComponentQueryParams& Params) const
+{
+	// Pass through to FQuat version.
+	return ComponentSweepMultiByChannel(OutHits, PrimComp, Start, End, Rot.Quaternion(), TraceChannel, Params);
+}
+
 FORCEINLINE_DEBUGGABLE ENetMode UWorld::GetNetMode() const
 {
 	// IsRunningDedicatedServer() is a compile-time check in optimized non-editor builds.
diff --git a/Engine/Source/Runtime/Engine/Private/Collision/WorldCollision.cpp b/Engine/Source/Runtime/Engine/Private/Collision/WorldCollision.cpp
--- a/Engine/Source/Runtime/Engine/Private/Collision/WorldCollision.cpp
+++ b/Engine/Source/Runtime/Engine/Private/Collision/WorldCollision.cpp
@@ -14,6 +14,7 @@
 #include "Framework/Docking/TabManager.h"
 #include "Collision.h"
 #include "Physics/PhysicsInterfaceCore.h"
+#include "PhysicsEngine/PhysicsObjectExternalInterface.h"
 #include "PhysXPublic.h"
 #include "Physics/PhysicsInterfaceTypes.h"
 #include "Chaos/ImplicitObject.h"
@@ -322,20 +323,109 @@ bool UWorld::ComponentOverlapMulti(TArray<struct FOverlapResult>& OutOverlaps, c
 	}
 }
 
-bool UWorld::ComponentOverlapMultiByChannel(TArray<struct FOverlapResult>& OutOverlaps, const class UPrimitiveComponent* PrimComp, const FVector& Pos, const FQuat& Quat, ECollisionChannel TraceChannel, const FComponentQueryParams& Params /* = FComponentQueryParams::DefaultComponentQueryParams */, const FCollisionObjectQueryParams& ObjectQueryParams/* =FCollisionObjectQueryParams::DefaultObjectQueryParam */) const
+bool UWorld::ComponentOverlapMultiByChannel(TArray<struct FOverlapResult>& OutOverlaps, const class UPrimitiveComponent* PrimComp, const FVector& Pos, const FQuat& Quat, ECollisionChannel TraceChannel, const FComponentQueryParams& Params, const FCollisionObjectQueryParams& ObjectQueryParams) const
+{
+	if (!GetPhysicsScene())
+	{
+		return false;
+	}
+
+	if (!PrimComp)
+	{
+		UE_LOG(LogCollision, Log, TEXT("ComponentOverlapMultiByChannel : No PrimComp"));
+		return false;
+	}
+
+	FComponentQueryParams ParamsWithSelf = Params;
+	ParamsWithSelf.AddIgnoredComponent_LikelyDuplicatedRoot(PrimComp);
+	OutOverlaps.Reset();
+
+	// Maintains compatibility with previous versions that primarily depended on the body instance.
+	FBodyInstance* BodyInstance = PrimComp->GetBodyInstance();
+	if (BodyInstance)
+	{
+		return BodyInstance->OverlapMulti(OutOverlaps, this, nullptr, Pos, Quat, TraceChannel, ParamsWithSelf, FCollisionResponseParams{ BodyInstance->GetResponseToChannels() }, ObjectQueryParams);
+	}
+
+	// New version that's more generic and relies on the physics object interface.
+	TArray<Chaos::FPhysicsObjectHandle> PhysicsObjects = PrimComp->GetAllPhysicsObjects();
+	FLockedReadPhysicsObjectExternalInterface Interface = FPhysicsObjectExternalInterface::LockRead(PhysicsObjects);
+
+	// For geometry collections we want to make sure we don't account for disabled particles since that means their parent cluster still hasn't released them yet.
+	// This should be safe for non-geometry collections and not filter out anything.
+	PhysicsObjects = PhysicsObjects.FilterByPredicate(
+		[&Interface](Chaos::FPhysicsObject* Object)
+		{
+			return !Interface->AreAllDisabled({ &Object, 1 });
+		}
+	);
+
+	if (PhysicsObjects.IsEmpty())
+	{
+		UE_LOG(LogCollision, Log, TEXT("UWorld::ComponentOverlapMultiByChannel : (%s) No physics data"), *PrimComp->GetName());
+		return false;
+	}
+
+	bool bHaveBlockingHit = false;
+	FCollisionResponseParams CollisionResponseParams{ PrimComp->GetCollisionResponseToChannels() };
+
+	const FTransform ComponentToTest{ Quat, Pos };
+	const FTransform WorldToComponent = PrimComp->GetComponentToWorld().Inverse();
+
+	TArray<struct FOverlapResult> TempOverlaps;
+	for (Chaos::FPhysicsObjectHandle Object : PhysicsObjects)
+	{
+		TArray<Chaos::FPerShapeData*> Shapes = Interface->GetAllShapes({&Object, 1});
+
+		// Determine how to convert the local space of this body instance to the test space
+		const FTransform ObjectToWorld = Interface->GetTransform(Object);
+		const FTransform ObjectToTest = ComponentToTest * WorldToComponent * ObjectToWorld;
+
+		for (Chaos::FPerShapeData* Shape : Shapes)
+		{
+			FPhysicsShapeHandle ShapeRef{ Shape, nullptr };
+
+			// TODO: Add support to be able to check if the shape collision is enabled for the generic physics object interface.
+			/*
+			// Skip this shape if it's CollisionEnabled setting was masked out
+			if (ParamsWithSelf.ShapeCollisionMask && !(ParamsWithSelf.ShapeCollisionMask & BodyInstance->GetShapeCollisionEnabled(ShapeIdx)))
+			{
+				continue;
+			}
+			*/
+			FPhysicsGeometryCollection GeomCollection = FPhysicsInterface::GetGeometryCollection(ShapeRef);
+			if (!ShapeRef.GetGeometry().IsConvex())
+			{
+				continue;	//we skip complex shapes - should this respect ComplexAsSimple?
+			}
+
+			TempOverlaps.Reset();
+			if (FPhysicsInterface::GeomOverlapMulti(this, GeomCollection, ObjectToTest.GetTranslation(), ObjectToTest.GetRotation(), TempOverlaps, TraceChannel, ParamsWithSelf, CollisionResponseParams, ObjectQueryParams))
+			{
+				bHaveBlockingHit = true;
+			}
+			OutOverlaps.Append(TempOverlaps);
+		}
+	}
+
+	return bHaveBlockingHit;
+}
+
+bool UWorld::ComponentSweepMulti(TArray<struct FHitResult>& OutHits, class UPrimitiveComponent* PrimComp, const FVector& Start, const FVector& End, const FQuat& Quat, const FComponentQueryParams& Params) const
 {
 	if (PrimComp)
 	{
-		return PrimComp->ComponentOverlapMulti(OutOverlaps, this, Pos, Quat, TraceChannel, Params, ObjectQueryParams);
+		ComponentSweepMultiByChannel(OutHits, PrimComp, Start, End, Quat, PrimComp->GetCollisionObjectType(), Params);
+		return (OutHits.Num() > 0);
 	}
 	else
 	{
-		UE_LOG(LogCollision, Log, TEXT("ComponentOverlapMulti : No PrimComp"));
+		UE_LOG(LogCollision, Log, TEXT("ComponentSweepMulti : No PrimComp"));
 		return false;
 	}
 }
 
-bool UWorld::ComponentSweepMulti(TArray<struct FHitResult>& OutHits, class UPrimitiveComponent* PrimComp, const FVector& Start, const FVector& End, const FQuat& Quat, const FComponentQueryParams& Params) const
+bool UWorld::ComponentSweepMultiByChannel(TArray<struct FHitResult>& OutHits, class UPrimitiveComponent* PrimComp, const FVector& Start, const FVector& End, const FQuat& Rot, ECollisionChannel TraceChannel, const FComponentQueryParams& Params) const
 {
 	OutHits.Reset();
 
@@ -346,68 +436,68 @@ bool UWorld::ComponentSweepMulti(TArray<struct FHitResult>& OutHits, class UPrim
 
 	if (PrimComp == NULL)
 	{
-		UE_LOG(LogCollision, Log, TEXT("ComponentSweepMulti : No PrimComp"));
+		UE_LOG(LogCollision, Log, TEXT("ComponentSweepMultiByChannel : No PrimComp"));
 		return false;
 	}
 
-	ECollisionChannel TraceChannel = PrimComp->GetCollisionObjectType();
-
 	// if extent is 0, do line trace
 	if (PrimComp->IsZeroExtent())
 	{
 		return FPhysicsInterface::RaycastMulti(this, OutHits, Start, End, TraceChannel, Params, FCollisionResponseParams(PrimComp->GetCollisionResponseToChannels()));
 	}
 
+	TArray<Chaos::FPhysicsObject*> PhysicsObjects = PrimComp->GetAllPhysicsObjects();
+	FLockedReadPhysicsObjectExternalInterface Interface = FPhysicsObjectExternalInterface::LockRead(PhysicsObjects);
 
-	const FBodyInstance* BodyInstance = PrimComp->GetBodyInstance();
-
-	if (!BodyInstance || !BodyInstance->IsValidBodyInstance())
-	{
-		UE_LOG(LogCollision, Log, TEXT("ComponentSweepMulti : (%s) No physics data"), *PrimComp->GetReadableName());
-		return false;
-	}
-
-#if !(UE_BUILD_SHIPPING || UE_BUILD_TEST)
-	if(PrimComp->IsA(USkeletalMeshComponent::StaticClass()))
-	{
-		UE_LOG(LogCollision, Log, TEXT("ComponentSweepMulti : SkeletalMeshComponent support only root body (%s) "), *PrimComp->GetReadableName());
-	}
-#endif
+	// For geometry collections we want to make sure we don't account for disabled particles since that means their parent cluster still hasn't released them yet.
+	// This should be safe for non-geometry collections and not filter out anything.
+	PhysicsObjects = PhysicsObjects.FilterByPredicate(
+		[&Interface](Chaos::FPhysicsObject* Object)
+		{
+			return !Interface->AreAllDisabled({ &Object, 1 });
+		}
+	);
 
 	SCOPE_CYCLE_COUNTER(STAT_Collision_GeomSweepMultiple);
 	bool bHaveBlockingHit = false;
 
-	FPhysicsCommand::ExecuteRead(BodyInstance->ActorHandle, [&](const FPhysicsActorHandle& Actor)
+	const FTransform ComponentToStart{ Rot, Start };
+	const FTransform ComponentToEnd{ Rot, End };
+
+	const FTransform WorldToComponent = PrimComp->GetComponentToWorld().Inverse();
+
+	TArray<struct FOverlapResult> TempOverlaps;
+	for (Chaos::FPhysicsObjectHandle Object : PhysicsObjects)
 	{
-		if(!FPhysicsInterface::IsValid(Actor))
-		{
-			return;
-		}
+		TArray<Chaos::FPerShapeData*> Shapes = Interface->GetAllShapes({ &Object, 1 });
 
-		// Get all the shapes from the actor
-		FInlineShapeArray PShapes;
-		const int32 NumShapes = FillInlineShapeArray_AssumesLocked(PShapes, Actor);
+		// Determine how to convert the local space of this body instance to the test space
+		const FTransform ObjectToWorld = Interface->GetTransform(Object);
+		const FTransform ObjectToStart = ComponentToStart * WorldToComponent * ObjectToWorld;
+		const FTransform ObjectToEnd = ComponentToEnd * WorldToComponent * ObjectToWorld;
 
-		for(FPhysicsShapeHandle& Shape : PShapes)
+		for (Chaos::FPerShapeData* Shape : Shapes)
 		{
-			check(Shape.IsValid());
-			ECollisionShapeType ShapeType = FPhysicsInterface::GetShapeType(Shape);
+			FPhysicsShapeHandle ShapeHandle{ Shape, nullptr };
 
-			if (!Shape.GetGeometry().IsConvex())
+			check(ShapeHandle.IsValid());
+			ECollisionShapeType ShapeType = FPhysicsInterface::GetShapeType(ShapeHandle);
+
+			if (!ShapeHandle.GetGeometry().IsConvex())
 			{
 				//We skip complex shapes. Should this respect complex as simple?
 				continue;
 			}
 
-			FPhysicsGeometryCollection GeomCollection = FPhysicsInterface::GetGeometryCollection(Shape);
+			FPhysicsGeometryCollection GeomCollection = FPhysicsInterface::GetGeometryCollection(ShapeHandle);
 			TArray<FHitResult> TmpHits;
-			if(FPhysicsInterface::GeomSweepMulti(this, GeomCollection, Quat, TmpHits, Start, End, TraceChannel, Params, FCollisionResponseParams(PrimComp->GetCollisionResponseToChannels())))
+			if (FPhysicsInterface::GeomSweepMulti(this, GeomCollection, ObjectToStart.GetRotation(), TmpHits, ObjectToStart.GetTranslation(), ObjectToEnd.GetTranslation(), TraceChannel, Params, FCollisionResponseParams(PrimComp->GetCollisionResponseToChannels())))
 			{
 				bHaveBlockingHit = true;
 			}
 			OutHits.Append(TmpHits);	//todo: should these be made unique?
 		}
-	});
+	}
 
 	return bHaveBlockingHit;
 }
diff --git a/Engine/Source/Runtime/Experimental/Chaos/Private/Chaos/PhysicsObjectInterface.cpp b/Engine/Source/Runtime/Experimental/Chaos/Private/Chaos/PhysicsObjectInterface.cpp
--- a/Engine/Source/Runtime/Experimental/Chaos/Private/Chaos/PhysicsObjectInterface.cpp
+++ b/Engine/Source/Runtime/Experimental/Chaos/Private/Chaos/PhysicsObjectInterface.cpp
@@ -1,6 +1,9 @@
 // Copyright Epic Games, Inc. All Rights Reserved.
 
 #include "Chaos/PhysicsObjectInterface.h"
+
+#include "Chaos/CastingUtilities.h"
+#include "Chaos/GeometryQueries.h"
 #include "Chaos/PhysicsObjectInternal.h"
 #include "PhysicsProxy/SingleParticlePhysicsProxy.h"
 #include "Math/Transform.h"
@@ -155,6 +158,111 @@ namespace Chaos
 		return FQuat::Identity;
 	}
 
+	template<EThreadContext Id>
+	TArray<FPerShapeData*> FReadPhysicsObjectInterface<Id>::GetAllShapes(TArrayView<FPhysicsObjectHandle> InObjects)
+	{
+		TArray<FPerShapeData*> AllShapes;
+
+		for (FPhysicsObjectHandle Object : InObjects)
+		{
+			if (!Object)
+			{
+				continue;
+			}
+
+			if (TThreadParticle<Id>* Particle = Object->GetParticle<Id>())
+			{
+				const Chaos::FShapesArray& ShapesArray = Particle->ShapesArray();
+				for (const TUniquePtr<Chaos::FPerShapeData>& Shape : ShapesArray)
+				{
+					AllShapes.Add(Shape.Get());
+				}
+			}
+		}
+
+		return AllShapes;
+	}
+
+	template<EThreadContext Id>
+	bool FReadPhysicsObjectInterface<Id>::GetPhysicsObjectOverlap(FPhysicsObjectHandle ObjectA, FPhysicsObjectHandle ObjectB, bool bTraceComplex, Chaos::FMTDInfo* OutMTD)
+	{
+		TArray<FPerShapeData*> ShapesA = GetAllShapes({ &ObjectA, 1 });
+		const FTransform TransformA(GetR(ObjectA), GetX(ObjectA));
+		const FBox BoxA = GetWorldBounds({&ObjectA, 1});
+
+		TArray<FPerShapeData*> ShapesB = GetAllShapes({ &ObjectB, 1 });
+		const FTransform TransformB(GetR(ObjectB), GetX(ObjectB));
+		const FBox BoxB = GetWorldBounds({ &ObjectB, 1 });
+
+		if (!BoxA.Intersect(BoxB))
+		{
+			return false;
+		}
+
+		for (FPerShapeData* A : ShapesA)
+		{
+			if (!A)
+			{
+				continue;
+			}
+
+			FCollisionFilterData ShapeFilterA = A->GetQueryData();
+			const bool bShapeIsComplexA = (ShapeFilterA.Word3 & static_cast<uint8>(EFilterFlags::ComplexCollision)) != 0;
+			const bool bShapeIsSimpleA = (ShapeFilterA.Word3 & static_cast<uint8>(EFilterFlags::SimpleCollision)) != 0;
+			const bool bShouldTraceA = (bTraceComplex && bShapeIsComplexA) || (!bTraceComplex && bShapeIsSimpleA);
+			if (!bShouldTraceA)
+			{
+				continue;
+			}
+
+			const FAABB3 BoxShapeA = A->GetGeometry()->CalculateTransformedBounds(TransformA);
+
+			for (FPerShapeData* B : ShapesB)
+			{
+				if (!B)
+				{
+					continue;
+				}
+
+				FCollisionFilterData ShapeFilterB = B->GetQueryData();
+				const bool bShapeIsComplexB = (ShapeFilterB.Word3 & static_cast<uint8>(EFilterFlags::ComplexCollision)) != 0;
+				const bool bShapeIsSimpleB = (ShapeFilterB.Word3 & static_cast<uint8>(EFilterFlags::SimpleCollision)) != 0;
+				const bool bShouldTraceB = (bTraceComplex && bShapeIsComplexB) || (!bTraceComplex && bShapeIsSimpleB);
+
+				if (!bShouldTraceB)
+				{
+					continue;
+				}
+
+				const FAABB3 BoxShapeB = B->GetGeometry()->CalculateTransformedBounds(TransformB);
+
+				if (!BoxShapeA.Intersects(BoxShapeB))
+				{
+					continue;
+				}
+
+				const bool bOverlap = Chaos::Utilities::CastHelper(
+					*B->GetGeometry(),
+					// ...Why is this transform passed?
+					TransformB,
+					[A, &TransformA, &TransformB, OutMTD](const auto& Downcast, const auto&)
+					{
+						return Chaos::OverlapQuery(*A->GetGeometry(), TransformA, Downcast, TransformB, 0, OutMTD);
+					}
+				);
+
+				// Is this actually correct? For now I'm mirroring the behavior in Overlap_GeomInternal
+				// but if the goal is to compute the MTD, shouldn't we need to iterate over all pairs of shapes?
+				if (bOverlap)
+				{
+					return true;
+				}
+			}
+		}
+
+		return false;
+	}
+
 	template<EThreadContext Id>
 	bool FReadPhysicsObjectInterface<Id>::AreAllValid(TArrayView<FPhysicsObjectHandle> InObjects)
 	{
@@ -250,7 +358,7 @@ namespace Chaos
 	}
 
 	template<EThreadContext Id>
-	FBox FReadPhysicsObjectInterface<Id>::GetBounds(TArrayView<FPhysicsObjectHandle> InObjects)
+	FBox FReadPhysicsObjectInterface<Id>::GetWorldBounds(TArrayView<FPhysicsObjectHandle> InObjects)
 	{
 		FBox RetBox(ForceInit);
 		for (FPhysicsObjectHandle Object : InObjects)
@@ -266,10 +374,12 @@ namespace Chaos
 				continue;
 			}
 
+			const FTransform WorldTransform = GetTransform(Object);
+
 			FBox ParticleBox(ForceInit);
 			if (const FImplicitObject* Geometry = Particle->Geometry().Get(); Geometry && Geometry->HasBoundingBox())
 			{
-				Chaos::FAABB3 Box = Geometry->BoundingBox();
+				const Chaos::FAABB3 Box = Geometry->BoundingBox().TransformedAABB(WorldTransform);
 				ParticleBox = FBox{ Box.Min(), Box.Max() };
 			}
 
diff --git a/Engine/Source/Runtime/Experimental/Chaos/Public/Chaos/PhysicsObjectInterface.h b/Engine/Source/Runtime/Experimental/Chaos/Public/Chaos/PhysicsObjectInterface.h
--- a/Engine/Source/Runtime/Experimental/Chaos/Public/Chaos/PhysicsObjectInterface.h
+++ b/Engine/Source/Runtime/Experimental/Chaos/Public/Chaos/PhysicsObjectInterface.h
@@ -36,6 +36,8 @@ struct CHAOS_API FClosestPhysicsObjectResult
 namespace Chaos
 {
 	class FPBDRigidsSolver;
+	class FPerShapeData;
+	struct FMTDInfo;
 
 	/**
 	 * FReadPhysicsObjectInterface will assume that these operations are safe to call (i.e. the relevant scenes have been read locked on the game thread).
@@ -53,6 +55,9 @@ namespace Chaos
 		FVector GetWorldCoM(FPhysicsObjectHandle Object);
 		FQuat GetR(FPhysicsObjectHandle Object);
 
+		TArray<FPerShapeData*> GetAllShapes(TArrayView<FPhysicsObjectHandle> InObjects);
+		bool GetPhysicsObjectOverlap(FPhysicsObjectHandle ObjectA, FPhysicsObjectHandle ObjectB, bool bTraceComplex, Chaos::FMTDInfo* OutMTD);
+
 		bool AreAllValid(TArrayView<FPhysicsObjectHandle> InObjects);
 		bool AreAllKinematic(TArrayView<FPhysicsObjectHandle> InObjects);
 		bool AreAllSleeping(TArrayView<FPhysicsObjectHandle> InObjects);
@@ -60,7 +65,7 @@ namespace Chaos
 		bool AreAllDynamic(TArrayView<FPhysicsObjectHandle> InObjects);
 		bool AreAllDisabled(TArrayView<FPhysicsObjectHandle> InObjects);
 		float GetMass(TArrayView<FPhysicsObjectHandle> InObjects);
-		FBox GetBounds(TArrayView<FPhysicsObjectHandle> InObjects);
+		FBox GetWorldBounds(TArrayView<FPhysicsObjectHandle> InObjects);
 		FClosestPhysicsObjectResult GetClosestPhysicsBodyFromLocation(TArrayView<FPhysicsObjectHandle> InObjects, const FVector& WorldLocation);
 
 		friend class FPhysicsObjectInterface;
diff --git a/Engine/Source/Runtime/PhysicsCore/Private/ChaosInterfaceWrapperCore.cpp b/Engine/Source/Runtime/PhysicsCore/Private/ChaosInterfaceWrapperCore.cpp
--- a/Engine/Source/Runtime/PhysicsCore/Private/ChaosInterfaceWrapperCore.cpp
+++ b/Engine/Source/Runtime/PhysicsCore/Private/ChaosInterfaceWrapperCore.cpp
@@ -31,10 +31,12 @@ namespace ChaosInterface
 	{
 		using namespace Chaos;
 		int32 ResultObjectType = GetInnerType(InGeometry.GetType());
+		const Chaos::FImplicitObject* CurrentGeometry = &InGeometry;
 
-		if (ResultObjectType == ImplicitObjectType::Transformed)
+		while (CurrentGeometry && ResultObjectType == ImplicitObjectType::Transformed)
 		{
-			ResultObjectType = GetInnerType(static_cast<const TImplicitObjectTransformed<FReal, 3>*>(&InGeometry)->Object()->GetType());
+			CurrentGeometry = static_cast<const TImplicitObjectTransformed<FReal, 3>*>(CurrentGeometry)->Object().Get();
+			ResultObjectType = GetInnerType(CurrentGeometry->GetType());
 		}
 
 		return ImplicitTypeToCollisionType(ResultObjectType);
