Adding key to new Curve within a Sequencer crashes Engine
#fix add d-tor to anim cache task ensuring any requets/childtasks have finished before deletion

diff --git a/Engine/Source/Runtime/Engine/Private/Animation/AnimationCompressionDerivedData.cpp b/Engine/Source/Runtime/Engine/Private/Animation/AnimationCompressionDerivedData.cpp
--- a/Engine/Source/Runtime/Engine/Private/Animation/AnimationCompressionDerivedData.cpp
+++ b/Engine/Source/Runtime/Engine/Private/Animation/AnimationCompressionDerivedData.cpp
@@ -34,6 +34,15 @@ FAnimationSequenceAsyncCacheTask::FAnimationSequenceAsyncCacheTask(const FIoHash
 	BeginCache(InKeyHash);
 }
 
+FAnimationSequenceAsyncCacheTask::~FAnimationSequenceAsyncCacheTask()
+{
+	if (!Poll())
+	{
+		Cancel();
+		Wait(false);
+	}
+}
+
 void FAnimationSequenceAsyncCacheTask::Cancel()
 {
 	if (BuildTask)
@@ -44,11 +53,11 @@ void FAnimationSequenceAsyncCacheTask::Cancel()
 	Owner.Cancel();
 }
 
-void FAnimationSequenceAsyncCacheTask::Wait()
+void FAnimationSequenceAsyncCacheTask::Wait(bool bPerformWork /*= true*/)
 {
 	if (BuildTask != nullptr)
 	{
-		BuildTask->EnsureCompletion();
+		BuildTask->EnsureCompletion(bPerformWork);
 	}
 
 	Owner.Wait();
diff --git a/Engine/Source/Runtime/Engine/Private/Animation/AnimationCompressionDerivedData.h b/Engine/Source/Runtime/Engine/Private/Animation/AnimationCompressionDerivedData.h
--- a/Engine/Source/Runtime/Engine/Private/Animation/AnimationCompressionDerivedData.h
+++ b/Engine/Source/Runtime/Engine/Private/Animation/AnimationCompressionDerivedData.h
@@ -70,8 +70,10 @@ namespace UE::Anim
 			UAnimSequence& InAnimSequence,
 			const ITargetPlatform* InTargetPlatform);
 
+		~FAnimationSequenceAsyncCacheTask();
+
 		void Cancel();
-		void Wait();
+		void Wait(bool bPerformWork = true);
 		bool Poll() const;
 		void Reschedule(FQueuedThreadPool* InThreadPool, EQueuedWorkPriority InPriority) const;
 	private:
