Fixed wrong missing refraction on water caused by 23371838.

diff --git a/Engine/Source/Runtime/Engine/Private/Materials/Material.cpp b/Engine/Source/Runtime/Engine/Private/Materials/Material.cpp
--- a/Engine/Source/Runtime/Engine/Private/Materials/Material.cpp
+++ b/Engine/Source/Runtime/Engine/Private/Materials/Material.cpp
@@ -6644,7 +6644,7 @@ static bool IsPropertyActive_Internal(EMaterialProperty InProperty,
 			switch (InProperty)
 			{
 			case MP_Refraction:
-				Active = ((bIsTranslucentBlendMode && !IsAlphaHoldoutBlendMode(BlendMode, StrataBlendMode) && !IsModulateBlendMode(BlendMode, StrataBlendMode)) || ShadingModels.HasShadingModel(MSM_SingleLayerWater)) && bUsesDistortion;
+				Active = (bIsTranslucentBlendMode && !IsAlphaHoldoutBlendMode(BlendMode, StrataBlendMode) && !IsModulateBlendMode(BlendMode, StrataBlendMode) && bUsesDistortion) || ShadingModels.HasShadingModel(MSM_SingleLayerWater);
 				break;
 			case MP_Opacity:
 				Active = (bIsTranslucentBlendMode && !IsModulateBlendMode(BlendMode, StrataBlendMode)) || ShadingModels.HasShadingModel(MSM_SingleLayerWater);
@@ -6679,7 +6679,7 @@ static bool IsPropertyActive_Internal(EMaterialProperty InProperty,
 			Active = false;
 			break;
 		case MP_Refraction:
-			Active = ((bIsTranslucentBlendMode && !IsAlphaHoldoutBlendMode(BlendMode, StrataBlendMode) && !IsModulateBlendMode(BlendMode, StrataBlendMode)) || ShadingModels.HasShadingModel(MSM_SingleLayerWater)) && bUsesDistortion;
+			Active = (bIsTranslucentBlendMode && !IsAlphaHoldoutBlendMode(BlendMode, StrataBlendMode) && !IsModulateBlendMode(BlendMode, StrataBlendMode) && bUsesDistortion) || ShadingModels.HasShadingModel(MSM_SingleLayerWater);
 			break;
 		case MP_Opacity:
 			Active = (bIsTranslucentBlendMode && !IsModulateBlendMode(BlendMode, StrataBlendMode)) || ShadingModels.HasShadingModel(MSM_SingleLayerWater);
