Remove game thread check on RHI resource set owner name. EngineTest can call on FSkinWeightVertexBuffer from other thread on consoles

diff --git a/Engine/Source/Runtime/Engine/Private/Rendering/SkinWeightVertexBuffer.cpp b/Engine/Source/Runtime/Engine/Private/Rendering/SkinWeightVertexBuffer.cpp
--- a/Engine/Source/Runtime/Engine/Private/Rendering/SkinWeightVertexBuffer.cpp
+++ b/Engine/Source/Runtime/Engine/Private/Rendering/SkinWeightVertexBuffer.cpp
@@ -805,7 +805,6 @@ void FSkinWeightVertexBuffer::ResetVertexBoneWeights(uint32 VertexIndex)
 
 void FSkinWeightVertexBuffer::SetOwnerName(const FName& OwnerName)
 {
-	check(IsInGameThread() || IsInParallelGameThread());
 	LookupVertexBuffer.SetOwnerName(OwnerName);
 	DataVertexBuffer.SetOwnerName(OwnerName);
 }
diff --git a/Engine/Source/Runtime/RenderCore/Private/RenderResource.cpp b/Engine/Source/Runtime/RenderCore/Private/RenderResource.cpp
--- a/Engine/Source/Runtime/RenderCore/Private/RenderResource.cpp
+++ b/Engine/Source/Runtime/RenderCore/Private/RenderResource.cpp
@@ -354,7 +354,6 @@ FBufferRHIRef FRenderResource::CreateRHIBufferInternal(
 void FRenderResource::SetOwnerName(const FName& InOwnerName)
 {
 #if RHI_ENABLE_RESOURCE_INFO
-	check(IsInGameThread() || IsInParallelGameThread());
 	OwnerName = InOwnerName;
 #endif
 }
