Fixing nonunity builds

diff --git a/Engine/Source/Runtime/Engine/Classes/Materials/MaterialExpressionSparseVolumeTextureSample.h b/Engine/Source/Runtime/Engine/Classes/Materials/MaterialExpressionSparseVolumeTextureSample.h
--- a/Engine/Source/Runtime/Engine/Classes/Materials/MaterialExpressionSparseVolumeTextureSample.h
+++ b/Engine/Source/Runtime/Engine/Classes/Materials/MaterialExpressionSparseVolumeTextureSample.h
@@ -5,6 +5,7 @@
 #include "CoreMinimal.h"
 #include "Materials/MaterialExpression.h"
 #include "UObject/ObjectMacros.h"
+#include "MaterialTypes.h"
 
 #include "MaterialExpressionSparseVolumeTextureSample.generated.h"
 
diff --git a/Engine/Source/Runtime/Engine/Classes/Materials/MaterialFunctionInterface.h b/Engine/Source/Runtime/Engine/Classes/Materials/MaterialFunctionInterface.h
--- a/Engine/Source/Runtime/Engine/Classes/Materials/MaterialFunctionInterface.h
+++ b/Engine/Source/Runtime/Engine/Classes/Materials/MaterialFunctionInterface.h
@@ -15,6 +15,11 @@
 #include "Materials/MaterialExpressionTextureSampleParameter.h"
 #endif
 #include "StaticParameterSet.h"
+
+#if WITH_EDITOR
+#include "Materials/MaterialExpression.h"
+#endif
+
 #include "MaterialFunctionInterface.generated.h"
 
 class UMaterial;
diff --git a/Engine/Source/Runtime/Experimental/Iris/Core/Private/Iris/Serialization/NetBitStreamReader.cpp b/Engine/Source/Runtime/Experimental/Iris/Core/Private/Iris/Serialization/NetBitStreamReader.cpp
--- a/Engine/Source/Runtime/Experimental/Iris/Core/Private/Iris/Serialization/NetBitStreamReader.cpp
+++ b/Engine/Source/Runtime/Experimental/Iris/Core/Private/Iris/Serialization/NetBitStreamReader.cpp
@@ -4,6 +4,7 @@
 #include "Iris/Serialization/NetBitStreamUtils.h"
 #include "HAL/PlatformMemory.h"
 #include "Misc/AssertionMacros.h"
+#include <stdint.h>
 
 namespace UE::Net
 {
