Fix UMaterial to properly populate the ShaderId with StaticParameters to allow merging them with Instances that use the default Parameter Set

diff --git a/Engine/Source/Runtime/Engine/Classes/Materials/MaterialInstance.h b/Engine/Source/Runtime/Engine/Classes/Materials/MaterialInstance.h
--- a/Engine/Source/Runtime/Engine/Classes/Materials/MaterialInstance.h
+++ b/Engine/Source/Runtime/Engine/Classes/Materials/MaterialInstance.h
@@ -7,7 +7,6 @@
 #include "UObject/ObjectMacros.h"
 #include "Misc/Guid.h"
 #include "Materials/MaterialInterface.h"
-#include "StaticParameterSet.h"
 #include "Materials/MaterialInstanceBasePropertyOverrides.h"
 
 #if UE_ENABLE_INCLUDE_ORDER_DEPRECATED_IN_5_2
@@ -690,11 +689,6 @@ protected:
 	TUniquePtr<FMaterialInstanceCachedData> CachedData;
 
 private:
-#if WITH_EDITOR
-	mutable TOptional<FStaticParameterSet> CachedStaticParameterValues;
-	mutable uint8 AllowCachingStaticParameterValuesCounter = 0;
-#endif // WITH_EDITOR
-
 	/** Inline material resources serialized from disk. To be processed on game thread in PostLoad. */
 	TArray<FMaterialResource> LoadedMaterialResources;
 
@@ -859,11 +853,6 @@ public:
 #if WITH_EDITORONLY_DATA
 
 	ENGINE_API void SetStaticSwitchParameterValueEditorOnly(const FMaterialParameterInfo& ParameterInfo, bool Value);
-
-	/**
-	 * Builds a composited set of static parameters, including inherited and overridden values
-	 */
-	ENGINE_API void GetStaticParameterValues(FStaticParameterSet& OutStaticParameters);
 #endif // WITH_EDITORONLY_DATA
 
 	ENGINE_API virtual void GetAllParametersOfType(EMaterialParameterType Type, TMap<FMaterialParameterInfo, FMaterialParameterMetadata>& OutParameters) const override;
diff --git a/Engine/Source/Runtime/Engine/Classes/Materials/MaterialInterface.h b/Engine/Source/Runtime/Engine/Classes/Materials/MaterialInterface.h
--- a/Engine/Source/Runtime/Engine/Classes/Materials/MaterialInterface.h
+++ b/Engine/Source/Runtime/Engine/Classes/Materials/MaterialInterface.h
@@ -28,6 +28,7 @@
 #include "MaterialRecursionGuard.h"
 #include "MaterialShaderPrecompileMode.h"
 #include "RHIFeatureLevel.h"
+#include "StaticParameterSet.h"
 
 #if UE_ENABLE_INCLUDE_ORDER_DEPRECATED_IN_5_2
 #include "RHI.h"
@@ -536,6 +537,10 @@ public:
 	virtual FGraphEventArray PrecachePSOs(const TConstArrayView<const FVertexFactoryType*>& VertexFactoryTypes, const struct FPSOPrecacheParams& PreCacheParams) { return FGraphEventArray(); }
 
 #if WITH_EDITORONLY_DATA
+	/**
+	* Builds a composited set of static parameters, including inherited and overridden values
+	*/
+	ENGINE_API void GetStaticParameterValues(FStaticParameterSet& OutStaticParameters);
 
 	/**
 	* Get the value of the given static switch parameter
@@ -1093,6 +1098,11 @@ private:
 	static UEnum* SamplerTypeEnum;
 
 #if WITH_EDITOR
+protected:
+	mutable TOptional<FStaticParameterSet> CachedStaticParameterValues;
+	mutable uint8 AllowCachingStaticParameterValuesCounter = 0;
+
+private:
 	/**
 	* Whether or not this material interface should force the preview to be a plane mesh.
 	*/
diff --git a/Engine/Source/Runtime/Engine/Private/Materials/Material.cpp b/Engine/Source/Runtime/Engine/Private/Materials/Material.cpp
--- a/Engine/Source/Runtime/Engine/Private/Materials/Material.cpp
+++ b/Engine/Source/Runtime/Engine/Private/Materials/Material.cpp
@@ -375,6 +375,12 @@ void FMaterialResource::GetShaderMapId(EShaderPlatform Platform, const ITargetPl
 		MaterialInstance->GetStaticParameterValues(CompositedStaticParameters);
 		OutId.UpdateFromParameterSet(CompositedStaticParameters);
 	}
+	else
+	{
+		FStaticParameterSet CompositedStaticParameters;
+		Material->GetStaticParameterValues(CompositedStaticParameters);
+		OutId.UpdateFromParameterSet(CompositedStaticParameters);
+	}
 #endif // WITH_EDITOR
 }
 
@@ -388,6 +394,10 @@ void FMaterialResource::GetStaticParameterSet(EShaderPlatform Platform, FStaticP
 	{
 		MaterialInstance->GetStaticParameterValues(OutSet);
 	}
+	else
+	{
+		Material->GetStaticParameterValues(OutSet);
+	}
 }
 #endif // WITH_EDITORONLY_DATA
 
diff --git a/Engine/Source/Runtime/Engine/Private/Materials/MaterialInstance.cpp b/Engine/Source/Runtime/Engine/Private/Materials/MaterialInstance.cpp
--- a/Engine/Source/Runtime/Engine/Private/Materials/MaterialInstance.cpp
+++ b/Engine/Source/Runtime/Engine/Private/Materials/MaterialInstance.cpp
@@ -1812,7 +1812,7 @@ void UMaterialInstance::SetStaticSwitchParameterValueEditorOnly(const FMaterialP
 	new(StaticParametersRuntime.StaticSwitchParameters) FStaticSwitchParameter(ParameterInfo, Value, true, FGuid());
 }
 
-void UMaterialInstance::GetStaticParameterValues(FStaticParameterSet& OutStaticParameters)
+void UMaterialInterface::GetStaticParameterValues(FStaticParameterSet& OutStaticParameters)
 {
 	check(IsInGameThread());
 
@@ -1822,25 +1822,24 @@ void UMaterialInstance::GetStaticParameterValues(FStaticParameterSet& OutStaticP
 		return;
 	}
 
-	if (Parent)
+	TMap<FMaterialParameterInfo, FMaterialParameterMetadata> ParameterValues;
+	for (int32 ParameterTypeIndex = 0; ParameterTypeIndex < NumMaterialParameterTypes; ++ParameterTypeIndex)
 	{
-		TMap<FMaterialParameterInfo, FMaterialParameterMetadata> ParameterValues;
-		for (int32 ParameterTypeIndex = 0; ParameterTypeIndex < NumMaterialParameterTypes; ++ParameterTypeIndex)
+		const EMaterialParameterType ParameterType = (EMaterialParameterType)ParameterTypeIndex;
+		if (IsStaticMaterialParameter(ParameterType))
 		{
-			const EMaterialParameterType ParameterType = (EMaterialParameterType)ParameterTypeIndex;
-			if (IsStaticMaterialParameter(ParameterType))
-			{
-				ParameterValues.Reset();
-				GetAllParametersOfType(ParameterType, ParameterValues);
-				OutStaticParameters.AddParametersOfType(ParameterType, ParameterValues);
-			}
+			ParameterValues.Reset();
+			GetAllParametersOfType(ParameterType, ParameterValues);
+			OutStaticParameters.AddParametersOfType(ParameterType, ParameterValues);
 		}
 	}
 
-	UMaterialInstanceEditorOnlyData* EditorOnly = GetEditorOnlyData();
-	if (EditorOnly)
+	if(UMaterialInstance* MaterialInstance = Cast<UMaterialInstance>(this))
 	{
-		OutStaticParameters.EditorOnly.TerrainLayerWeightParameters = EditorOnly->StaticParameters.TerrainLayerWeightParameters;
+		if (UMaterialInstanceEditorOnlyData* EditorOnly = MaterialInstance->GetEditorOnlyData())
+		{
+			OutStaticParameters.EditorOnly.TerrainLayerWeightParameters = EditorOnly->StaticParameters.TerrainLayerWeightParameters;
+		}
 	}
 
 	FMaterialLayersFunctions MaterialLayers;
