Move SceneOutliner Pin logic into FLoaderAdapterPinnedActors

#rnx

diff --git a/Engine/Source/Editor/SceneOutliner/Private/ActorDescTreeItem.cpp b/Engine/Source/Editor/SceneOutliner/Private/ActorDescTreeItem.cpp
--- a/Engine/Source/Editor/SceneOutliner/Private/ActorDescTreeItem.cpp
+++ b/Engine/Source/Editor/SceneOutliner/Private/ActorDescTreeItem.cpp
@@ -4,6 +4,7 @@
 #include "WorldPartition/ActorDescContainer.h"
 #include "WorldPartition/WorldPartition.h"
 #include "WorldPartition/WorldPartitionActorDesc.h"
+#include "WorldPartition/LoaderAdapter/LoaderAdapterPinnedActors.h"
 #include "Widgets/DeclarativeSyntaxSupport.h"
 #include "Layout/WidgetPath.h"
 #include "Framework/Application/MenuStack.h"
@@ -116,24 +117,13 @@ private:
 	{
 		if (TSharedPtr<FActorDescTreeItem> TreeItem = TreeItemPtr.Pin())
 		{
-			if (const FWorldPartitionActorDesc* ActorDesc = TreeItem->ActorDescHandle.Get())
+			if (FWorldPartitionActorDesc* ActorDesc = TreeItem->ActorDescHandle.Get())
 			{
 				FFormatNamedArguments Args;
 				Args.Add(TEXT("ActorLabel"), FText::FromString(TreeItem->GetDisplayString()));
-				
-				UActorDescContainer* ActorDescContainer = ActorDesc->GetContainer();
-				UWorld* World = ActorDescContainer != nullptr ? ActorDescContainer->GetWorld() : nullptr;
-				UWorldPartition* WorldPartition = World != nullptr ? World->GetWorldPartition() : nullptr;
-				if (WorldPartition && (!ActorDesc->GetIsSpatiallyLoaded() || WorldPartition->IsActorPinned(ActorDesc->GetGuid())))
-				{
-					Args.Add(TEXT("UnloadState"), LOCTEXT("UnloadedDataLayer", "(Unloaded DataLayer)"));
-				}
-				else
-				{
-					Args.Add(TEXT("UnloadState"), LOCTEXT("UnloadedActorLabel", "(Unloaded)"));
-				}
+				Args.Add(TEXT("UnloadState"), FLoaderAdapterPinnedActors::GetUnloadedReason(ActorDesc));
 
-				return FText::Format(LOCTEXT("UnloadedActorDisplay", "{ActorLabel} {UnloadState}"), Args);
+				return FText::Format(LOCTEXT("UnloadedActorDisplay", "{ActorLabel} ({UnloadState})"), Args);
 			}
 		}
 		return FText();
@@ -319,37 +309,7 @@ bool FActorDescTreeItem::GetVisibility() const
 
 bool FActorDescTreeItem::ShouldShowPinnedState() const
 {
-	if (ActorDescHandle.IsValid())
-	{
-		// Pinning of ActorDescs is only supported on the main world partition
-		if (UActorDescContainer* Container = ActorDescHandle->GetContainer())
-		{
-			if (FWorldPartitionActorDesc const* ActorDesc = ActorDescHandle.Get())
-			{
-				// Only Spatially loaded actors can be pinned with the exception of non spatially loaded, runtime only actors (ex: HLODs)
-				if (!ActorDesc->GetIsSpatiallyLoaded() && !ActorDesc->GetActorIsRuntimeOnly())
-				{
-					return false;
-				}
-			}
-
-			if (Container->IsMainPartitionContainer())
-			{
-				return true;
-			}
-			else if (ActorDescHandle->GetContentBundleGuid().IsValid())
-			{
-				const UWorld* ContainerWorld = Container->GetWorld();
-				const UWorldPartition* ContainerWorldPartition = ContainerWorld ? ContainerWorld->GetWorldPartition() : nullptr;
-				if (ContainerWorldPartition && ContainerWorldPartition->IsMainWorldPartition())
-				{
-					return ActorDescHandle->GetActorSoftPath().GetAssetPath().GetPackageName() == ContainerWorld->GetPackage()->GetFName();
-				}
-			}
-		}
-	}
-
-	return false;
+	return FLoaderAdapterPinnedActors::SupportsPinning(ActorDescHandle.Get());
 }
 
 bool FActorDescTreeItem::GetPinnedState() const
diff --git a/Engine/Source/Editor/SceneOutliner/Private/ActorTreeItem.cpp b/Engine/Source/Editor/SceneOutliner/Private/ActorTreeItem.cpp
--- a/Engine/Source/Editor/SceneOutliner/Private/ActorTreeItem.cpp
+++ b/Engine/Source/Editor/SceneOutliner/Private/ActorTreeItem.cpp
@@ -19,6 +19,7 @@
 #include "SSocketChooser.h"
 #include "LevelInstance/LevelInstanceInterface.h"
 #include "WorldPartition/WorldPartition.h"
+#include "WorldPartition/LoaderAdapter/LoaderAdapterPinnedActors.h"
 #include "ToolMenu.h"
 #include "Engine/Level.h"
 
@@ -362,17 +363,7 @@ TSharedRef<SWidget> FActorTreeItem::GenerateLabelWidget(ISceneOutliner& Outliner
 
 bool FActorTreeItem::ShouldShowPinnedState() const
 {
-	if (const AActor* ActorPtr = Actor.Get())
-	{
-		// Pinning of Actors is only supported on the main world partition
-		const ULevel* Level = ActorPtr->GetLevel();
-		const UWorld* World = Level->GetWorld();
-
-		// Only Spatially loaded actors can be pinned with the exception of non spatially loaded, runtime only actors (ex: HLODs)
-		return World && !World->IsGameWorld() && !!World->GetWorldPartition() && Level->IsPersistentLevel() && ActorPtr->IsPackageExternal() && (ActorPtr->GetIsSpatiallyLoaded() || ActorPtr->IsRuntimeOnly());
-	}
-
-	return false;
+	return FLoaderAdapterPinnedActors::SupportsPinning(Actor.Get());
 }
 
 bool FActorTreeItem::ShouldShowVisibilityState() const
diff --git a/Engine/Source/Runtime/Engine/Private/WorldPartition/LoaderAdapter/LoaderAdapterPinnedActors.cpp b/Engine/Source/Runtime/Engine/Private/WorldPartition/LoaderAdapter/LoaderAdapterPinnedActors.cpp
new file mode 100644
--- /dev/null
+++ b/Engine/Source/Runtime/Engine/Private/WorldPartition/LoaderAdapter/LoaderAdapterPinnedActors.cpp
@@ -0,0 +1,88 @@
+// Copyright Epic Games, Inc. All Rights Reserved.
+
+#include "WorldPartition/LoaderAdapter/LoaderAdapterPinnedActors.h"
+
+#if WITH_EDITOR
+
+#include "Engine/World.h"
+#include "Engine/Level.h"
+#include "WorldPartition/WorldPartitionActorDesc.h"
+#include "WorldPartition/ActorDescContainer.h"
+#include "WorldPartition/WorldPartition.h"
+
+#define LOCTEXT_NAMESPACE "FLoaderAdapterPinnedActors"
+
+bool FLoaderAdapterPinnedActors::PassActorDescFilter(const FWorldPartitionHandle& ActorHandle) const
+{
+	// We want to be able to pin any type of actors (HLODs, etc).
+	return ActorHandle.IsValid() && !ActorsToRemove.Contains(ActorHandle);
+}
+
+FText FLoaderAdapterPinnedActors::GetUnloadedReason(FWorldPartitionActorDesc* InActorDesc)
+{
+	if (InActorDesc)
+	{
+		UActorDescContainer* ActorDescContainer = InActorDesc->GetContainer();
+		UWorld* World = ActorDescContainer != nullptr ? ActorDescContainer->GetWorld() : nullptr;
+		UWorldPartition* WorldPartition = World != nullptr ? World->GetWorldPartition() : nullptr;
+		bool bShouldBeLoaded = !InActorDesc->GetIsSpatiallyLoaded() && !InActorDesc->GetActorIsRuntimeOnly();
+		if (WorldPartition && (bShouldBeLoaded || WorldPartition->IsActorPinned(InActorDesc->GetGuid())))
+		{
+			return LOCTEXT("UnloadedDataLayerReason", "Unloaded DataLayer");
+		}
+	}
+	
+	return LOCTEXT("UnloadedReason", "Unloaded");
+}
+
+bool FLoaderAdapterPinnedActors::SupportsPinning(FWorldPartitionActorDesc* InActorDesc)
+{
+	if (!InActorDesc)
+	{
+		return false;
+	}
+
+	// Only Spatially loaded actors can be pinned with the exception of non spatially loaded, runtime only actors (ex: HLODs)
+	if (!InActorDesc->GetIsSpatiallyLoaded() && !InActorDesc->GetActorIsRuntimeOnly())
+	{
+		return false;
+	}
+
+	if (UActorDescContainer* Container = InActorDesc->GetContainer())
+	{
+		if (Container->IsMainPartitionContainer())
+		{
+			return true;
+		}
+		else if (InActorDesc->GetContentBundleGuid().IsValid())
+		{
+			const UWorld* ContainerWorld = Container->GetWorld();
+			const UWorldPartition* ContainerWorldPartition = ContainerWorld ? ContainerWorld->GetWorldPartition() : nullptr;
+			if (ContainerWorldPartition && ContainerWorldPartition->IsMainWorldPartition())
+			{
+				return InActorDesc->GetActorSoftPath().GetAssetPath().GetPackageName() == ContainerWorld->GetPackage()->GetFName();
+			}
+		}
+	}
+
+	return false;
+}
+
+bool FLoaderAdapterPinnedActors::SupportsPinning(AActor* InActor)
+{
+	if (InActor)
+	{
+		// Pinning of Actors is only supported on the main world partition
+		const ULevel* Level = InActor->GetLevel();
+		const UWorld* World = Level->GetWorld();
+
+		// Only Spatially loaded actors can be pinned with the exception of non spatially loaded, runtime only actors (ex: HLODs)
+		return World && !World->IsGameWorld() && !!World->GetWorldPartition() && Level->IsPersistentLevel() && InActor->IsPackageExternal() && (InActor->GetIsSpatiallyLoaded() || InActor->IsRuntimeOnly());
+	}
+
+	return false;
+}
+
+#undef LOCTEXT_NAMESPACE
+
+#endif
\ No newline at end of file
diff --git a/Engine/Source/Runtime/Engine/Private/WorldPartition/WorldPartition.cpp b/Engine/Source/Runtime/Engine/Private/WorldPartition/WorldPartition.cpp
--- a/Engine/Source/Runtime/Engine/Private/WorldPartition/WorldPartition.cpp
+++ b/Engine/Source/Runtime/Engine/Private/WorldPartition/WorldPartition.cpp
@@ -56,6 +56,7 @@
 #include "WorldPartition/DataLayer/WorldDataLayers.h"
 #include "WorldPartition/LoaderAdapter/LoaderAdapterActorList.h"
 #include "WorldPartition/LoaderAdapter/LoaderAdapterShape.h"
+#include "WorldPartition/LoaderAdapter/LoaderAdapterPinnedActors.h"
 #include "WorldPartition/WorldPartitionActorDescViewProxy.h"
 #include "WorldPartition/HLOD/HLODLayer.h"
 #include "WorldPartition/Cook/WorldPartitionCookPackage.h"
@@ -283,21 +284,6 @@ public:
 	}
 };
 
-class FLoaderAdapterPinnedActors: public FLoaderAdapterActorList
-{
-public:
-	FLoaderAdapterPinnedActors(UWorld* InWorld)
-		: FLoaderAdapterActorList(InWorld)
-	{}
-
-protected:
-	virtual bool PassActorDescFilter(const FWorldPartitionHandle& ActorHandle) const override
-	{
-		// We want to be able to pin any type of actors (HLODs, etc).
-		return ActorHandle.IsValid() && !ActorsToRemove.Contains(ActorHandle);
-	}
-};
-
 #endif
 
 UWorldPartition::UWorldPartition(const FObjectInitializer& ObjectInitializer)
diff --git a/Engine/Source/Runtime/Engine/Public/WorldPartition/LoaderAdapter/LoaderAdapterPinnedActors.h b/Engine/Source/Runtime/Engine/Public/WorldPartition/LoaderAdapter/LoaderAdapterPinnedActors.h
new file mode 100644
--- /dev/null
+++ b/Engine/Source/Runtime/Engine/Public/WorldPartition/LoaderAdapter/LoaderAdapterPinnedActors.h
@@ -0,0 +1,26 @@
+// Copyright Epic Games, Inc. All Rights Reserved.
+#pragma once
+
+#if WITH_EDITOR
+
+#include "WorldPartition/LoaderAdapter/LoaderAdapterActorList.h"
+
+class AActor;
+class FWorldPartitionActorDesc;
+
+class ENGINE_API FLoaderAdapterPinnedActors : public FLoaderAdapterActorList
+{
+public:
+	FLoaderAdapterPinnedActors(UWorld* InWorld)
+		: FLoaderAdapterActorList(InWorld)
+	{}
+
+	static FText GetUnloadedReason(FWorldPartitionActorDesc* InActorDesc);
+	static bool SupportsPinning(FWorldPartitionActorDesc* InActorDesc);
+	static bool SupportsPinning(AActor* InActor);
+
+protected:
+	virtual bool PassActorDescFilter(const FWorldPartitionHandle& ActorHandle) const override;
+};
+
+#endif
\ No newline at end of file
