UWorld::AddToWorld will only call FAssetCompilingManager::Get().ProcessAsyncTasks() on a MakingVisible streaming level only if its time limit hasn't been reached.

diff --git a/Engine/Source/Runtime/Engine/Private/World.cpp b/Engine/Source/Runtime/Engine/Private/World.cpp
--- a/Engine/Source/Runtime/Engine/Private/World.cpp
+++ b/Engine/Source/Runtime/Engine/Private/World.cpp
@@ -2957,38 +2957,41 @@ void UWorld::AddToWorld( ULevel* Level, const FTransform& LevelTransform, bool b
 	}
 
 #if WITH_EDITOR
-	// Wait on any async DDC handles
-	if (bExecuteNextStep && AsyncPreRegisterDDCRequests.Num())
+	if (bExecuteNextStep)
 	{
-		if (!bConsiderTimeLimit)
-		{
-			QUICK_SCOPE_CYCLE_COUNTER(UWorld_AddToWorld_WaitFor_AsyncPreRegisterLevelStreamingTasks);
-
-			for (TSharedPtr<FAsyncPreRegisterDDCRequest>& Request : AsyncPreRegisterDDCRequests)
-			{
-				Request->WaitAsynchronousCompletion();
-			}
-			AsyncPreRegisterDDCRequests.Empty();
-		}
-		else
+		// Wait on any async DDC handles
+		if (AsyncPreRegisterDDCRequests.Num())
 		{
-			for (int32 Index = 0; Index < AsyncPreRegisterDDCRequests.Num(); Index++)
+			if (!bConsiderTimeLimit)
 			{
-				if (AsyncPreRegisterDDCRequests[Index]->PollAsynchronousCompletion())
+				QUICK_SCOPE_CYCLE_COUNTER(UWorld_AddToWorld_WaitFor_AsyncPreRegisterLevelStreamingTasks);
+
+				for (TSharedPtr<FAsyncPreRegisterDDCRequest>& Request : AsyncPreRegisterDDCRequests)
 				{
-					AsyncPreRegisterDDCRequests.RemoveAtSwap(Index--);
+					Request->WaitAsynchronousCompletion();
 				}
-				else
+				AsyncPreRegisterDDCRequests.Empty();
+			}
+			else
+			{
+				for (int32 Index = 0; Index < AsyncPreRegisterDDCRequests.Num(); Index++)
 				{
-					bExecuteNextStep = false;
-					break;
+					if (AsyncPreRegisterDDCRequests[Index]->PollAsynchronousCompletion())
+					{
+						AsyncPreRegisterDDCRequests.RemoveAtSwap(Index--);
+					}
+					else
+					{
+						bExecuteNextStep = false;
+						break;
+					}
 				}
 			}
 		}
-	}
 
-	// Gives a chance to any assets being used for PIE/game to complete
-	FAssetCompilingManager::Get().ProcessAsyncTasks();
+		// Gives a chance to any assets being used for PIE/game to complete
+		FAssetCompilingManager::Get().ProcessAsyncTasks();
+	}
 #endif
 
 	// Updates the level components (Actor components and UModelComponents).
