Iris - NetDriver should normally not call functions that iterates over all NetDrivers.

diff --git a/Engine/Source/Runtime/Engine/Private/Net/Iris/ReplicationSystem/ReplicationSystemUtil.cpp b/Engine/Source/Runtime/Engine/Private/Net/Iris/ReplicationSystem/ReplicationSystemUtil.cpp
--- a/Engine/Source/Runtime/Engine/Private/Net/Iris/ReplicationSystem/ReplicationSystemUtil.cpp
+++ b/Engine/Source/Runtime/Engine/Private/Net/Iris/ReplicationSystem/ReplicationSystemUtil.cpp
@@ -429,57 +429,43 @@ void FReplicationSystemUtil::BeginReplicationForActorsInWorld(UWorld* World)
 	}
 }
 
-void FReplicationSystemUtil::NotifyActorDormancyChange(AActor* Actor, ENetDormancy OldDormancyState)
+void FReplicationSystemUtil::NotifyActorDormancyChange(UReplicationSystem* ReplicationSystem, AActor* Actor, ENetDormancy OldDormancyState)
 {
-	FNetHandle ActorHandle = GetNetHandle(Actor);
-	if (!ActorHandle.IsValid())
+	if (!ReplicationSystem || !ReplicationSystem->IsServer())
 	{
 		return;
 	}
-	
+
 	const ENetDormancy Dormancy = Actor->NetDormancy;
 	const bool bIsPendingDormancy = (Dormancy > DORM_Awake);
 
-	if (const UWorld* World = Actor->GetWorld())
+	if (UObjectReplicationBridge* Bridge = ReplicationSystem->GetReplicationBridgeAs<UObjectReplicationBridge>())
 	{
-		ReplicationSystemUtil::ForEachReplicationSystem(GEngine, World, [ActorHandle, bIsPendingDormancy](UReplicationSystem* ReplicationSystem)
+		const FNetRefHandle ActorRefHandle = Bridge->GetReplicatedRefHandle(Actor);
+		if (ActorRefHandle.IsValid())
 		{
-			if (ReplicationSystem->IsServer())
-			{
-				if (UObjectReplicationBridge* Bridge = ReplicationSystem->GetReplicationBridgeAs<UObjectReplicationBridge>())
-				{
-					const FNetRefHandle ActorRefHandle = Bridge->GetReplicatedRefHandle(ActorHandle);
-					if (ActorRefHandle.IsValid())
-					{
-						Bridge->SetObjectWantsToBeDormant(ActorRefHandle, bIsPendingDormancy);
-					}
-				}
-			}
-		});
+			Bridge->SetObjectWantsToBeDormant(ActorRefHandle, bIsPendingDormancy);
+		}
 	}
 }
 
-void FReplicationSystemUtil::FlushNetDormancy(AActor* Actor, bool bWasDormInitial)
+void FReplicationSystemUtil::FlushNetDormancy(UReplicationSystem* ReplicationSystem, AActor* Actor, bool bWasDormInitial)
 {
+	if (!ReplicationSystem || !ReplicationSystem->IsServer())
+	{
+		return;
+	}
+
 	FNetHandle ActorHandle = GetNetHandle(Actor);
 	if (ActorHandle.IsValid())
 	{
-		if (const UWorld* World = Actor->GetWorld())
+		if (UObjectReplicationBridge* Bridge = ReplicationSystem->GetReplicationBridgeAs<UObjectReplicationBridge>())
 		{
-			ReplicationSystemUtil::ForEachReplicationSystem(GEngine, World, [ActorHandle, bWasDormInitial](UReplicationSystem* ReplicationSystem)
+			const FNetRefHandle ActorRefHandle = Bridge->GetReplicatedRefHandle(ActorHandle);
+			if (ActorRefHandle.IsValid())
 			{
-				if (ReplicationSystem->IsServer())
-				{
-					if (UObjectReplicationBridge* Bridge = ReplicationSystem->GetReplicationBridgeAs<UObjectReplicationBridge>())
-					{
-						const FNetRefHandle ActorRefHandle = Bridge->GetReplicatedRefHandle(ActorHandle);
-						if (ActorRefHandle.IsValid())
-						{
-							Bridge->ForceUpdateWantsToBeDormantObject(ActorRefHandle);
-						}
-					}
-				}
-			});
+				Bridge->ForceUpdateWantsToBeDormantObject(ActorRefHandle);
+			}
 		}
 	}
 	else
diff --git a/Engine/Source/Runtime/Engine/Private/NetDriver.cpp b/Engine/Source/Runtime/Engine/Private/NetDriver.cpp
--- a/Engine/Source/Runtime/Engine/Private/NetDriver.cpp
+++ b/Engine/Source/Runtime/Engine/Private/NetDriver.cpp
@@ -3623,9 +3623,9 @@ void UNetDriver::FlushActorDormancy(AActor* Actor, bool bWasDormInitial)
 	}
 
 #if UE_WITH_IRIS
-	if (ReplicationSystem && ReplicationSystem->IsServer())
+	if (ReplicationSystem)
 	{
-		UE::Net::FReplicationSystemUtil::FlushNetDormancy(Actor, bWasDormInitial);
+		UE::Net::FReplicationSystemUtil::FlushNetDormancy(ReplicationSystem, Actor, bWasDormInitial);
 	}
 	else
 #endif // UE_WITH_IRIS
@@ -3644,7 +3644,7 @@ void UNetDriver::NotifyActorDormancyChange(AActor* Actor, ENetDormancy OldDorman
 #if UE_WITH_IRIS
 	if (ReplicationSystem)
 	{
-		UE::Net::FReplicationSystemUtil::NotifyActorDormancyChange(Actor, OldDormancyState);
+		UE::Net::FReplicationSystemUtil::NotifyActorDormancyChange(ReplicationSystem, Actor, OldDormancyState);
 	}
 	else
 #endif // UE_WITH_IRIS
diff --git a/Engine/Source/Runtime/Engine/Public/Net/Iris/ReplicationSystem/ReplicationSystemUtil.h b/Engine/Source/Runtime/Engine/Public/Net/Iris/ReplicationSystem/ReplicationSystemUtil.h
--- a/Engine/Source/Runtime/Engine/Public/Net/Iris/ReplicationSystem/ReplicationSystemUtil.h
+++ b/Engine/Source/Runtime/Engine/Public/Net/Iris/ReplicationSystem/ReplicationSystemUtil.h
@@ -108,10 +108,10 @@ struct FReplicationSystemUtil
 	ENGINE_API static void BeginReplicationForActorsInWorld(UWorld* World);
 
 	/** Notify the ReplicationSystem of a dormancy change. */
-	ENGINE_API static void NotifyActorDormancyChange(AActor* Actor, ENetDormancy OldDormancyState);
+	ENGINE_API static void NotifyActorDormancyChange(UReplicationSystem* ReplicationSystem, AActor* Actor, ENetDormancy OldDormancyState);
 
 	/** Trigger replication of dirty state for actor wanting to be dormant. */
-	ENGINE_API static void FlushNetDormancy(AActor* Actor, bool bWasDormInitial);
+	ENGINE_API static void FlushNetDormancy(UReplicationSystem* ReplicationSystem, AActor* Actor, bool bWasDormInitial);
 
 	/** Enable or disable a replication condition. This will affect the replication of properties with conditions. */
 	ENGINE_API static void SetReplicationCondition(FNetHandle NetHandle, EReplicationCondition Condition, bool bEnableCondition);
