Fix BP ubergraph frame arrays of references incorrectly being treated as strong

Fixes EngineTest AssetManager tests

#rnx

diff --git a/Engine/Source/Runtime/Engine/Private/BlueprintGeneratedClass.cpp b/Engine/Source/Runtime/Engine/Private/BlueprintGeneratedClass.cpp
--- a/Engine/Source/Runtime/Engine/Private/BlueprintGeneratedClass.cpp
+++ b/Engine/Source/Runtime/Engine/Private/BlueprintGeneratedClass.cpp
@@ -1941,9 +1941,18 @@ public:
 
 	virtual void HandleObjectReferences(UObject** Objects, int32 Num, const UObject* ReferencingObject, const FProperty* ReferencingProperty) override
 	{
+		if (Num <= 0)
+		{
+			return;
+		}
+
+		// HandleObjectReference only calls MarkWeakObjectReferenceForClearing on object property references
+		const FProperty* InnerProperty = ReferencingProperty && ReferencingProperty->IsA<FArrayProperty>() ?
+			static_cast<const FArrayProperty*>(ReferencingProperty)->Inner : ReferencingProperty;
+
 		for (int32 Idx = 0; Idx < Num; ++Idx)
 		{
-			HandleObjectReference(Objects[Idx], ReferencingObject, ReferencingProperty);
+			HandleObjectReference(Objects[Idx], ReferencingObject, InnerProperty);
 		}
 	}
 
