World Partition: serialize actor to cell remap table in cooked builds.

diff --git a/Engine/Source/Runtime/Engine/Public/WorldPartition/WorldPartitionLevelStreamingPolicy.h b/Engine/Source/Runtime/Engine/Public/WorldPartition/WorldPartitionLevelStreamingPolicy.h
--- a/Engine/Source/Runtime/Engine/Public/WorldPartition/WorldPartitionLevelStreamingPolicy.h
+++ b/Engine/Source/Runtime/Engine/Public/WorldPartition/WorldPartitionLevelStreamingPolicy.h
@@ -42,10 +42,8 @@ protected:
 	void ForEachActiveRuntimeCell(TFunctionRef<void(const UWorldPartitionRuntimeCell*)> Func) const;
 
 private:
-#if WITH_EDITORONLY_DATA
 	UPROPERTY()
 	TMap<FName, FName> ActorToCellRemapping;
-#endif
 
 	UPROPERTY()
 	TMap<FName, FName> SubObjectsToCellRemapping;
