Do not release skeletal mesh rendering resources when there is pending streaming operations

diff --git a/Engine/Source/Runtime/Engine/Private/SkeletalMesh.cpp b/Engine/Source/Runtime/Engine/Private/SkeletalMesh.cpp
--- a/Engine/Source/Runtime/Engine/Private/SkeletalMesh.cpp
+++ b/Engine/Source/Runtime/Engine/Private/SkeletalMesh.cpp
@@ -1518,6 +1518,12 @@ bool USkeletalMesh::IsReadyForFinishDestroy()
 		return false;
 	}
 
+	// Match BeginDestroy() by checking for HasPendingInitOrStreaming().
+	if (HasPendingInitOrStreaming())
+	{
+		return false;
+	}
+
 	ReleaseResources();
 
 	// see if we have hit the resource flush fence
