Replicate SubObjects
* Added the possiblity to store registered subobjects via a WeakObjectPtr or a regular pointer
* Toggled by setting UE_NET_SUBOBJECTLIST_WEAKPTR.
* Default is to use weak pointers
* WeakObjects are useful in a context where some subobjects may get destroyed before getting unregistered which generally occurs on clients and PIE.
* Regular pointers are useful when you want to speed up the replication code and you know all your subobjects will be unregistered before getting destroyed. This is generally the case on dedicated servers.

diff --git a/Engine/Source/Runtime/Engine/Classes/Engine/ActorChannel.h b/Engine/Source/Runtime/Engine/Classes/Engine/ActorChannel.h
--- a/Engine/Source/Runtime/Engine/Classes/Engine/ActorChannel.h
+++ b/Engine/Source/Runtime/Engine/Classes/Engine/ActorChannel.h
@@ -430,10 +430,10 @@ private:
 	bool WriteSubObjectInBunch(UObject* Obj, FOutBunch& Bunch, FReplicationFlags RepFlags);
 
 	/** Find the replicated subobjects of the component and write them into the bunch */
-	bool WriteSubObjects(UActorComponent* Component, FOutBunch& Bunch, FReplicationFlags RepFlags, const TStaticBitArray<COND_Max>& ConditionMap);
+	bool WriteComponentSubObjects(UActorComponent* Component, FOutBunch& Bunch, FReplicationFlags RepFlags, const TStaticBitArray<COND_Max>& ConditionMap);
 
 	/** Replicate a list of subobjects */
-	bool WriteSubObjects(UActorComponent* ReplicatedComponent, const UE::Net::FSubObjectRegistry& SubObjectList, FOutBunch& Bunch, FReplicationFlags RepFlags, const TStaticBitArray<COND_Max>& ConditionMap);
+	bool WriteSubObjects(UObject* SubObjectOwner, const UE::Net::FSubObjectRegistry& SubObjectList, FOutBunch& Bunch, FReplicationFlags RepFlags, const TStaticBitArray<COND_Max>& ConditionMap);
 
 	bool CanSubObjectReplicateToClient(ELifetimeCondition NetCondition, FObjectKey SubObjectKey, const TStaticBitArray<COND_Max>& ConditionMap) const;
 
diff --git a/Engine/Source/Runtime/Engine/Private/DataChannel.cpp b/Engine/Source/Runtime/Engine/Private/DataChannel.cpp
--- a/Engine/Source/Runtime/Engine/Private/DataChannel.cpp
+++ b/Engine/Source/Runtime/Engine/Private/DataChannel.cpp
@@ -3762,12 +3762,6 @@ bool UActorChannel::ReplicateRegisteredSubObjects(FOutBunch& Bunch, FReplication
 	}
 #endif
 
-#if UE_NET_REPACTOR_NAME_DEBUG
-	// View address in memory viewer
-	const bool bSubObjectNameDebug = GCVarNetSubObjectNameDebug && GMetAsyncDemoNameDebugChance && Connection->IsReplay();
-	volatile PTRINT SubObjName = bSubObjectNameDebug ? UE_NET_ALLOCA_NAME_DEBUG_BUFFER() : 0;
-#endif
-
 	const TStaticBitArray<COND_Max> ConditionMap = FSendingRepState::BuildConditionMapFromRepFlags(RepFlags);
 
 	bool bWroteSomethingImportant = false;
@@ -3775,22 +3769,8 @@ bool UActorChannel::ReplicateRegisteredSubObjects(FOutBunch& Bunch, FReplication
 	// Start with the Actor's subobjects
 	{
 		SetCurrentSubObjectOwner(Actor);
-		for (const FSubObjectRegistry::FEntry& SubObjectInfo : FSubObjectRegistryGetter::GetSubObjects(Actor).GetRegistryList())
-		{
-			if (CanSubObjectReplicateToClient(SubObjectInfo.NetCondition, SubObjectInfo.Key, ConditionMap))
-			{
-				checkf(IsValid(SubObjectInfo.SubObject), TEXT("Found invalid subobject (%s) registered in %s"), *GetNameSafe(SubObjectInfo.SubObject), *Actor->GetName());
-
-#if UE_NET_REPACTOR_NAME_DEBUG
-				if (bSubObjectNameDebug)
-				{
-					StoreSubObjectName(SubObjName, SubObjectInfo);
-				}
-#endif
 
-				bWroteSomethingImportant |= WriteSubObjectInBunch(SubObjectInfo.SubObject, Bunch, RepFlags);
-			}
-		}
+		bWroteSomethingImportant |= WriteSubObjects(Actor, FSubObjectRegistryGetter::GetSubObjects(Actor), Bunch, RepFlags, ConditionMap);
 	}
 
 	// Now the replicated actor components
@@ -3802,13 +3782,6 @@ bool UActorChannel::ReplicateRegisteredSubObjects(FOutBunch& Bunch, FReplication
 
 			checkf(IsValid(ReplicatedComponent), TEXT("Found invalid replicated component (%s) registered in %s"), *GetNameSafe(ReplicatedComponent), *Actor->GetName());
 
-#if UE_NET_REPACTOR_NAME_DEBUG
-			if (bSubObjectNameDebug)
-			{
-				StoreSubObjectName(SubObjName, RepComponentInfo);
-			}
-#endif
-
 			SetCurrentSubObjectOwner(ReplicatedComponent);
 			if (ReplicatedComponent->IsUsingRegisteredSubObjectList())
 			{
@@ -3934,7 +3907,7 @@ bool UActorChannel::ReplicateSubobject(UActorComponent* ReplicatedComponent, FOu
 			*Actor->GetName(), *ReplicatedComponent->GetName());
 
 		// Write all the subobjects of the component
-		bWroteSomethingImportant |= WriteSubObjects(ReplicatedComponent, Bunch, RepFlags, ConditionMap);
+		bWroteSomethingImportant |= WriteComponentSubObjects(ReplicatedComponent, Bunch, RepFlags, ConditionMap);
 
 #if SUBOBJECT_TRANSITION_VALIDATION
 		if (UE::Net::GCVarCompareSubObjectsReplicated)
@@ -3957,7 +3930,7 @@ bool UActorChannel::ReplicateSubobject(UActorComponent* ReplicatedComponent, FOu
 	}
 }
 
-bool UActorChannel::WriteSubObjects(UActorComponent* ReplicatedComponent, FOutBunch& Bunch, FReplicationFlags RepFlags, const TStaticBitArray<COND_Max>& ConditionMap)
+bool UActorChannel::WriteComponentSubObjects(UActorComponent* ReplicatedComponent, FOutBunch& Bunch, FReplicationFlags RepFlags, const TStaticBitArray<COND_Max>& ConditionMap)
 {
 	using namespace UE::Net;
 	if (const FSubObjectRegistry* SubObjectList = FSubObjectRegistryGetter::GetSubObjectsOfActorCompoment(Actor, ReplicatedComponent))
@@ -3968,11 +3941,16 @@ bool UActorChannel::WriteSubObjects(UActorComponent* ReplicatedComponent, FOutBu
 	return false;
 }
 
-bool UActorChannel::WriteSubObjects(UActorComponent* ReplicatedComponent, const UE::Net::FSubObjectRegistry& SubObjectList, FOutBunch& Bunch, FReplicationFlags RepFlags, const TStaticBitArray<COND_Max>& ConditionMap)
+bool UActorChannel::WriteSubObjects(UObject* SubObjectOwner, const UE::Net::FSubObjectRegistry& SubObjectList, FOutBunch& Bunch, FReplicationFlags RepFlags, const TStaticBitArray<COND_Max>& ConditionMap)
 {
 	using namespace UE::Net;
 	using namespace UE::Net::Private;
 
+	if (SubObjectList.IsEmpty())
+	{
+		return false;
+	}
+
 	bool bWroteSomethingImportant = false;
 
 #if UE_NET_REPACTOR_NAME_DEBUG
@@ -3981,23 +3959,43 @@ bool UActorChannel::WriteSubObjects(UActorComponent* ReplicatedComponent, const
 	volatile PTRINT SubObjName = bSubObjectNameDebug ? UE_NET_ALLOCA_NAME_DEBUG_BUFFER() : 0;
 #endif
 
-	for (const FSubObjectRegistry::FEntry& SubObjectInfo : SubObjectList.GetRegistryList())
+#if UE_NET_SUBOBJECTLIST_WEAKPTR
+	TArray<int32, TInlineAllocator<8>> InvalidSubObjects;
+#endif
+
+	const TArray<FSubObjectRegistry::FEntry>& SubObjects = SubObjectList.GetRegistryList();
+
+	for (int32 Index=0; Index < SubObjects.Num(); ++Index)
 	{
-		checkf(IsValid(SubObjectInfo.SubObject), TEXT("Found invalid subobject (%s) registered in %s::%s"), *GetNameSafe(SubObjectInfo.SubObject), *Actor->GetName(), *GetNameSafe(ReplicatedComponent));
+		const FSubObjectRegistry::FEntry& SubObjectInfo = SubObjects[Index];
 
-		if (CanSubObjectReplicateToClient(SubObjectInfo.NetCondition, SubObjectInfo.Key, ConditionMap))
+		if (UObject* SubObjectToReplicate = SubObjectInfo.GetSubObject())
 		{
-#if UE_NET_REPACTOR_NAME_DEBUG
-			if (bSubObjectNameDebug)
+			checkf(IsValid(SubObjectToReplicate), TEXT("Found invalid subobject (%s) registered in %s::%s"), *GetNameSafe(SubObjectToReplicate), *Actor->GetName(), *GetNameSafe(SubObjectOwner));
+
+			if (CanSubObjectReplicateToClient(SubObjectInfo.NetCondition, SubObjectInfo.Key, ConditionMap))
 			{
-				StoreSubObjectName(SubObjName, SubObjectInfo);
-			}
+#if UE_NET_REPACTOR_NAME_DEBUG
+				if (bSubObjectNameDebug)
+				{
+					StoreSubObjectName(SubObjName, SubObjectInfo);
+				}
 #endif
-
-			bWroteSomethingImportant |= WriteSubObjectInBunch(SubObjectInfo.SubObject, Bunch, RepFlags);
+				bWroteSomethingImportant |= WriteSubObjectInBunch(SubObjectToReplicate, Bunch, RepFlags);
+			}
 		}
+#if UE_NET_SUBOBJECTLIST_WEAKPTR
+		else
+		{
+			InvalidSubObjects.Add(Index);
+		}
+#endif
 	}
 
+#if UE_NET_SUBOBJECTLIST_WEAKPTR
+	const_cast<FSubObjectRegistry&>(SubObjectList).CleanRegistryIndexes(InvalidSubObjects);
+#endif
+
 	return bWroteSomethingImportant;
 }
 
diff --git a/Engine/Source/Runtime/Engine/Private/Net/Iris/ReplicationSystem/ActorReplicationBridge.cpp b/Engine/Source/Runtime/Engine/Private/Net/Iris/ReplicationSystem/ActorReplicationBridge.cpp
--- a/Engine/Source/Runtime/Engine/Private/Net/Iris/ReplicationSystem/ActorReplicationBridge.cpp
+++ b/Engine/Source/Runtime/Engine/Private/Net/Iris/ReplicationSystem/ActorReplicationBridge.cpp
@@ -327,9 +327,10 @@ UE::Net::FNetRefHandle UActorReplicationBridge::BeginReplication(AActor* Actor,
 		// Start with the Actor's SubObjects (that is SubObjects that are not ActorComponents)
 		for (const FSubObjectRegistry::FEntry& SubObjectInfo : ActorSubObjects.GetRegistryList())
 		{
-			if (IsValid(SubObjectInfo.SubObject) && SubObjectInfo.NetCondition != ELifetimeCondition::COND_Never)
+			UObject* SubObjectToReplicate = SubObjectInfo.GetSubObject()
+			if (IsValid(SubObjectToReplicate) && SubObjectInfo.NetCondition != ELifetimeCondition::COND_Never)
 			{
-				FNetRefHandle SubObjectRefHandle = UObjectReplicationBridge::BeginReplication(ActorRefHandle, SubObjectInfo.SubObject);
+				FNetRefHandle SubObjectRefHandle = UObjectReplicationBridge::BeginReplication(ActorRefHandle, SubObjectToReplicate);
 				if (SubObjectRefHandle.IsValid() && SubObjectInfo.NetCondition != ELifetimeCondition::COND_None)
 				{
 					UObjectReplicationBridge::SetSubObjectNetCondition(SubObjectRefHandle, SubObjectInfo.NetCondition);
@@ -406,9 +407,10 @@ UE::Net::FNetRefHandle UActorReplicationBridge::BeginReplication(FNetRefHandle O
 	// Begin replication for any SubObjects registered by the component
 	for (const FSubObjectRegistry::FEntry& SubObjectInfo : RepComponentInfo->SubObjects.GetRegistryList())
 	{
-		if (IsValid(SubObjectInfo.SubObject) && SubObjectInfo.NetCondition != ELifetimeCondition::COND_Never)
+		UObject* SubObjectToReplicate = SubObjectInfo.GetSubObject()
+		if (IsValid(SubObjectToReplicate) && SubObjectInfo.NetCondition != ELifetimeCondition::COND_Never)
 		{
-			FNetRefHandle SubObjectHandle = UObjectReplicationBridge::BeginReplication(OwnerHandle, SubObjectInfo.SubObject, ReplicatedComponentHandle, UReplicationBridge::ESubObjectInsertionOrder::ReplicateWith);
+			FNetRefHandle SubObjectHandle = UObjectReplicationBridge::BeginReplication(OwnerHandle, SubObjectToReplicate, ReplicatedComponentHandle, UReplicationBridge::ESubObjectInsertionOrder::ReplicateWith);
 			if (SubObjectHandle.IsValid() && SubObjectInfo.NetCondition != ELifetimeCondition::COND_None)
 			{
 				SetSubObjectNetCondition(SubObjectHandle, SubObjectInfo.NetCondition);
diff --git a/Engine/Source/Runtime/Engine/Private/Net/NetNameDebug.h b/Engine/Source/Runtime/Engine/Private/Net/NetNameDebug.h
--- a/Engine/Source/Runtime/Engine/Private/Net/NetNameDebug.h
+++ b/Engine/Source/Runtime/Engine/Private/Net/NetNameDebug.h
@@ -51,7 +51,7 @@ namespace UE::Net::Private
 
 	inline void StoreSubObjectName(PTRINT BufferPtr, const FSubObjectRegistry::FEntry& SubObjectInfo)
 	{
-		StoreSubObjectName(BufferPtr, SubObjectInfo.SubObjectClassName, SubObjectInfo.SubObjectName, SubObjectInfo.SubObject);
+		StoreSubObjectName(BufferPtr, SubObjectInfo.SubObjectClassName, SubObjectInfo.SubObjectName, SubObjectInfo.GetSubObject());
 	}
 
 	inline void StoreSubObjectName(PTRINT BufferPtr, const FReplicatedComponentInfo& ComponentInfo)
diff --git a/Engine/Source/Runtime/Net/Core/Private/Net/Core/Misc/NetSubObjectRegistry.cpp b/Engine/Source/Runtime/Net/Core/Private/Net/Core/Misc/NetSubObjectRegistry.cpp
--- a/Engine/Source/Runtime/Net/Core/Private/Net/Core/Misc/NetSubObjectRegistry.cpp
+++ b/Engine/Source/Runtime/Net/Core/Private/Net/Core/Misc/NetSubObjectRegistry.cpp
@@ -50,4 +50,29 @@ bool FSubObjectRegistry::IsSubObjectInRegistry(const UObject* SubObject) const
 	return Registry.FindByKey(SubObject) != nullptr;
 }
 
+void FSubObjectRegistry::CleanRegistryIndexes(const TArrayView<int32>& IndexesToClean)
+{
+#if !UE_BUILD_SHIPPING
+	// Test to make sure the array is sorted 
+	int32 LastRemovedIndex = INT32_MAX;
+#endif
+	
+	for (int32 i=IndexesToClean.Num()-1; i >= 0; --i)
+	{
+		const int32 IndexToRemove = IndexesToClean[i];
+
+#if !UE_BUILD_SHIPPING
+		const bool bIsSorted = IndexToRemove < LastRemovedIndex;
+		ensureMsgf(bIsSorted, TEXT("CleanRegistryIndexes did not receive a sorted array. Index value %d (position %d) was >= then Index value %d (position %d)"), LastRemovedIndex, i + 1, IndexToRemove, i);
+		LastRemovedIndex = IndexToRemove;
+		if (!bIsSorted)
+		{
+			continue;
+		}
+#endif
+
+		Registry.RemoveAtSwap(IndexToRemove);
+	}
+}
+
 } //namespace UE::Net
\ No newline at end of file
diff --git a/Engine/Source/Runtime/Net/Core/Public/Net/Core/Misc/NetSubObjectRegistry.h b/Engine/Source/Runtime/Net/Core/Public/Net/Core/Misc/NetSubObjectRegistry.h
--- a/Engine/Source/Runtime/Net/Core/Public/Net/Core/Misc/NetSubObjectRegistry.h
+++ b/Engine/Source/Runtime/Net/Core/Public/Net/Core/Misc/NetSubObjectRegistry.h
@@ -10,6 +10,10 @@
 class UActorComponent;
 class UObject;
 
+#ifndef UE_NET_SUBOBJECTLIST_WEAKPTR
+	#define UE_NET_SUBOBJECTLIST_WEAKPTR 1
+#endif
+
 namespace UE::Net
 {
 	class FSubObjectRegistryGetter;
@@ -41,13 +45,32 @@ public:
 	/** Remove the subobject from the replicated list. Returns true if the subobject had been registered. */
 	bool RemoveSubObject(UObject* InSubObject);
 
+	bool IsEmpty() const { return Registry.Num() == 0; }
+
+	/**
+	* Remove all the indexes specified by the passed array
+	* @param IndexesToClean Array of indexes sorted from smallest to biggest. The sorted order is assumed to have been done by the caller
+	*/
+	void CleanRegistryIndexes(const TArrayView<int32>& IndexesToClean);
+
 	/** Find the NetCondition of a SubObject. Returns COND_MAX if not registered */
 	ELifetimeCondition GetNetCondition(UObject* SubObject) const;
 
 	struct FEntry
 	{
-		/** Raw pointer since users are obligated to call RemoveReplicatedSubobject before destroying the subobject otherwise it will cause a crash.  */
+	private:
+		/** 
+		* Subobjects can be stored via raw pointers or weakobject pointers.  
+		* Use raw pointers in a environment where you want the fastest replication code and can guarantee that every replicated subobject will be properly unregistered before getting deleted. Server's are one such place.
+		* Otherwise use weak pointers in a environment where the replicated subobjects might get garbage'd before being unregistered.  PIE or clients using the list for demo recording are exemples of that.
+		*/
+#if UE_NET_SUBOBJECTLIST_WEAKPTR
+		FWeakObjectPtr SubObject;
+#else
 		UObject* SubObject = nullptr;
+#endif
+
+	public:
 
 		/** Store the object key info for fast access later */
 		FObjectKey Key;
@@ -74,8 +97,17 @@ public:
 		{
 		}
 
-		bool operator==(const FEntry& rhs) const { return Key == rhs.Key; }
-		bool operator==(const UObject* rhs) const { return SubObject == rhs; }
+		inline bool operator==(const FEntry& rhs) const	{ return Key == rhs.Key; }
+		inline bool operator==(const UObject* rhs) const	{ return SubObject == rhs; }
+
+		inline UObject* GetSubObject() const
+		{
+#if UE_NET_SUBOBJECTLIST_WEAKPTR
+			return SubObject.Get();
+#else
+			return SubObject;
+#endif
+		}
 	};
 
 	/** Returns the list of registered subobjects */
