Fix compile error in Iris

diff --git a/Engine/Source/Runtime/Engine/Private/Net/Iris/ReplicationSystem/ActorReplicationBridge.cpp b/Engine/Source/Runtime/Engine/Private/Net/Iris/ReplicationSystem/ActorReplicationBridge.cpp
--- a/Engine/Source/Runtime/Engine/Private/Net/Iris/ReplicationSystem/ActorReplicationBridge.cpp
+++ b/Engine/Source/Runtime/Engine/Private/Net/Iris/ReplicationSystem/ActorReplicationBridge.cpp
@@ -327,7 +327,7 @@ UE::Net::FNetRefHandle UActorReplicationBridge::BeginReplication(AActor* Actor,
 		// Start with the Actor's SubObjects (that is SubObjects that are not ActorComponents)
 		for (const FSubObjectRegistry::FEntry& SubObjectInfo : ActorSubObjects.GetRegistryList())
 		{
-			UObject* SubObjectToReplicate = SubObjectInfo.GetSubObject()
+			UObject* SubObjectToReplicate = SubObjectInfo.GetSubObject();
 			if (IsValid(SubObjectToReplicate) && SubObjectInfo.NetCondition != ELifetimeCondition::COND_Never)
 			{
 				FNetRefHandle SubObjectRefHandle = UObjectReplicationBridge::BeginReplication(ActorRefHandle, SubObjectToReplicate);
@@ -407,7 +407,7 @@ UE::Net::FNetRefHandle UActorReplicationBridge::BeginReplication(FNetRefHandle O
 	// Begin replication for any SubObjects registered by the component
 	for (const FSubObjectRegistry::FEntry& SubObjectInfo : RepComponentInfo->SubObjects.GetRegistryList())
 	{
-		UObject* SubObjectToReplicate = SubObjectInfo.GetSubObject()
+		UObject* SubObjectToReplicate = SubObjectInfo.GetSubObject();
 		if (IsValid(SubObjectToReplicate) && SubObjectInfo.NetCondition != ELifetimeCondition::COND_Never)
 		{
 			FNetRefHandle SubObjectHandle = UObjectReplicationBridge::BeginReplication(OwnerHandle, SubObjectToReplicate, ReplicatedComponentHandle, UReplicationBridge::ESubObjectInsertionOrder::ReplicateWith);
