Moving EMaterialDomain to its own header.

diff --git a/Engine/Source/Runtime/Engine/Public/MaterialDomain.h b/Engine/Source/Runtime/Engine/Public/MaterialDomain.h
new file mode 100644
--- /dev/null
+++ b/Engine/Source/Runtime/Engine/Public/MaterialDomain.h
@@ -0,0 +1,30 @@
+// Copyright Epic Games, Inc. All Rights Reserved.
+
+#pragma once
+
+#include "UObject/ObjectMacros.h"
+#include "Containers/UnrealString.h"
+
+/** Defines the domain of a material. */
+UENUM()
+enum EMaterialDomain : int
+{
+	/** The material's attributes describe a 3d surface. */
+	MD_Surface UMETA(DisplayName = "Surface"),
+	/** The material's attributes describe a deferred decal, and will be mapped onto the decal's frustum. */
+	MD_DeferredDecal UMETA(DisplayName = "Deferred Decal"),
+	/** The material's attributes describe a light's distribution. */
+	MD_LightFunction UMETA(DisplayName = "Light Function"),
+	/** The material's attributes describe a 3d volume. */
+	MD_Volume UMETA(DisplayName = "Volume"),
+	/** The material will be used in a custom post process pass. */
+	MD_PostProcess UMETA(DisplayName = "Post Process"),
+	/** The material will be used for UMG or Slate UI */
+	MD_UI UMETA(DisplayName = "User Interface"),
+	/** The material will be used for runtime virtual texture (Deprecated). */
+	MD_RuntimeVirtualTexture UMETA(Hidden),
+
+	MD_MAX
+};
+
+ENGINE_API FString MaterialDomainString(EMaterialDomain MaterialDomain);
diff --git a/Engine/Source/Runtime/Engine/Public/MaterialShared.h b/Engine/Source/Runtime/Engine/Public/MaterialShared.h
--- a/Engine/Source/Runtime/Engine/Public/MaterialShared.h
+++ b/Engine/Source/Runtime/Engine/Public/MaterialShared.h
@@ -40,6 +40,7 @@
 #if UE_ENABLE_INCLUDE_ORDER_DEPRECATED_IN_5_2
 #include "RHI.h"
 #endif
+#include "MaterialDomain.h"
 
 #include "Shader/Preshader.h"
 
@@ -83,6 +84,8 @@ class FShaderMapLayout;
 class FMaterialCachedHLSLTree;
 #endif
 
+enum EMaterialDomain : int;
+
 namespace UE
 {
 namespace HLSLTree
@@ -245,30 +248,6 @@ enum struct EDeferredParamStrictness
 	EStrict, // throw warnings
 };
 
-/** Defines the domain of a material. */
-UENUM()
-enum EMaterialDomain : int
-{
-	/** The material's attributes describe a 3d surface. */
-	MD_Surface UMETA(DisplayName = "Surface"),
-	/** The material's attributes describe a deferred decal, and will be mapped onto the decal's frustum. */
-	MD_DeferredDecal UMETA(DisplayName = "Deferred Decal"),
-	/** The material's attributes describe a light's distribution. */
-	MD_LightFunction UMETA(DisplayName = "Light Function"),
-	/** The material's attributes describe a 3d volume. */
-	MD_Volume UMETA(DisplayName = "Volume"),
-	/** The material will be used in a custom post process pass. */
-	MD_PostProcess UMETA(DisplayName = "Post Process"),
-	/** The material will be used for UMG or Slate UI */
-	MD_UI UMETA(DisplayName = "User Interface"),
-	/** The material will be used for runtime virtual texture (Deprecated). */
-	MD_RuntimeVirtualTexture UMETA(Hidden),
-
-	MD_MAX
-};
-
-ENGINE_API FString MaterialDomainString(EMaterialDomain MaterialDomain);
-
 /** Fully describes a material compilation target */
 struct FMaterialCompileTargetParameters
 {
