Fix non-unity builds.

diff --git a/Engine/Source/Runtime/Engine/Private/Level.cpp b/Engine/Source/Runtime/Engine/Private/Level.cpp
--- a/Engine/Source/Runtime/Engine/Private/Level.cpp
+++ b/Engine/Source/Runtime/Engine/Private/Level.cpp
@@ -10,6 +10,8 @@ Level.cpp: Level-related functions
 #include "UObject/RenderingObjectVersion.h"
 #include "UObject/UE5MainStreamObjectVersion.h"
 #include "UObject/UE5ReleaseStreamObjectVersion.h"
+#include "UObject/FortniteMainBranchObjectVersion.h"
+#include "UObject/ReleaseObjectVersion.h"
 #include "Templates/UnrealTemplate.h"
 #include "UObject/Package.h"
 #include "EngineStats.h"
@@ -46,7 +48,6 @@ Level.cpp: Level-related functions
 #include "Engine/WorldComposition.h"
 #include "StaticLighting.h"
 #include "TickTaskManagerInterface.h"
-#include "UObject/ReleaseObjectVersion.h"
 #include "Physics/PhysicsInterfaceScene.h"
 #include "PhysicsEngine/BodySetup.h"
 #include "EngineGlobals.h"
