Fixed compile issues found after moving module files around in the unity files.
Also added missing generated.h includes.

diff --git a/Engine/Plugins/Compression/OodleNetwork/Source/Public/OodleNetworkHandlerComponent.h b/Engine/Plugins/Compression/OodleNetwork/Source/Public/OodleNetworkHandlerComponent.h
--- a/Engine/Plugins/Compression/OodleNetwork/Source/Public/OodleNetworkHandlerComponent.h
+++ b/Engine/Plugins/Compression/OodleNetwork/Source/Public/OodleNetworkHandlerComponent.h
@@ -14,6 +14,8 @@
 
 #include "oodle2net.h"
 
+#include "OodleNetworkHandlerComponent.generated.h"
+
 DECLARE_LOG_CATEGORY_EXTERN(OodleNetworkHandlerComponentLog, Log, All);
 
 // The maximum packet size that this component can handle - UNetConnection's should never allow MaxPacket to exceed MAX_PACKET_SIZE
diff --git a/Engine/Plugins/FX/Niagara/Source/Niagara/Classes/NiagaraCollision.h b/Engine/Plugins/FX/Niagara/Source/Niagara/Classes/NiagaraCollision.h
--- a/Engine/Plugins/FX/Niagara/Source/Niagara/Classes/NiagaraCollision.h
+++ b/Engine/Plugins/FX/Niagara/Source/Niagara/Classes/NiagaraCollision.h
@@ -6,6 +6,8 @@
 #include "NiagaraStats.h"
 #include "WorldCollision.h"
 
+#include "NiagaraCollision.generated.h"
+
 DECLARE_CYCLE_STAT(TEXT("Collision"), STAT_NiagaraCollision, STATGROUP_Niagara);
 
 class FNiagaraDataSet;
diff --git a/Engine/Plugins/Runtime/AudioWidgets/Source/AudioWidgets/Public/SAudioRadialSlider.h b/Engine/Plugins/Runtime/AudioWidgets/Source/AudioWidgets/Public/SAudioRadialSlider.h
--- a/Engine/Plugins/Runtime/AudioWidgets/Source/AudioWidgets/Public/SAudioRadialSlider.h
+++ b/Engine/Plugins/Runtime/AudioWidgets/Source/AudioWidgets/Public/SAudioRadialSlider.h
@@ -18,6 +18,8 @@
 #include "Widgets/Input/SEditableText.h"
 #include "Widgets/Layout/SWidgetSwitcher.h"
 
+#include "SAudioRadialSlider.generated.h"
+
 class SRadialSlider;
 
 UENUM()
diff --git a/Engine/Plugins/Runtime/MeshModelingToolset/Source/ModelingComponents/Public/Baking/BakingTypes.h b/Engine/Plugins/Runtime/MeshModelingToolset/Source/ModelingComponents/Public/Baking/BakingTypes.h
--- a/Engine/Plugins/Runtime/MeshModelingToolset/Source/ModelingComponents/Public/Baking/BakingTypes.h
+++ b/Engine/Plugins/Runtime/MeshModelingToolset/Source/ModelingComponents/Public/Baking/BakingTypes.h
@@ -8,6 +8,8 @@
 
 #include "CoreMinimal.h"
 
+#include "BakingTypes.generated.h"
+
 UENUM()
 enum class EBakeTextureResolution 
 {
diff --git a/Engine/Plugins/Runtime/ResonanceAudio/Source/ResonanceAudio/Public/ResonanceAudioEnums.h b/Engine/Plugins/Runtime/ResonanceAudio/Source/ResonanceAudio/Public/ResonanceAudioEnums.h
--- a/Engine/Plugins/Runtime/ResonanceAudio/Source/ResonanceAudio/Public/ResonanceAudioEnums.h
+++ b/Engine/Plugins/Runtime/ResonanceAudio/Source/ResonanceAudio/Public/ResonanceAudioEnums.h
@@ -7,6 +7,8 @@
 #include "CoreMinimal.h"
 #include "UObject/ObjectMacros.h"
 
+#include "ResonanceAudioEnums.generated.h"
+
 DECLARE_LOG_CATEGORY_EXTERN(LogResonanceAudio, Display, All);
 
 UENUM(BlueprintType)
diff --git a/Engine/Source/Editor/DetailCustomizations/Private/FbxImportUIDetails.h b/Engine/Source/Editor/DetailCustomizations/Private/FbxImportUIDetails.h
--- a/Engine/Source/Editor/DetailCustomizations/Private/FbxImportUIDetails.h
+++ b/Engine/Source/Editor/DetailCustomizations/Private/FbxImportUIDetails.h
@@ -15,6 +15,8 @@
 #include "UObject/ObjectMacros.h"
 #include "UObject/WeakObjectPtrTemplates.h"
 
+#include "FbxImportUIDetails.generated.h"
+
 class FString;
 class FText;
 class IDetailLayoutBuilder;
diff --git a/Engine/Source/Editor/Kismet/Public/FindInBlueprintManager.h b/Engine/Source/Editor/Kismet/Public/FindInBlueprintManager.h
--- a/Engine/Source/Editor/Kismet/Public/FindInBlueprintManager.h
+++ b/Engine/Source/Editor/Kismet/Public/FindInBlueprintManager.h
@@ -36,6 +36,8 @@
 #include "UObject/WeakObjectPtr.h"
 #include "UObject/WeakObjectPtrTemplates.h"
 
+#include "FindInBlueprintManager.generated.h"
+
 class FArchive;
 class SDockTab;
 class SWidget;
diff --git a/Engine/Source/Runtime/AIModule/Classes/BehaviorTree/Blackboard/BlackboardKeyEnums.h b/Engine/Source/Runtime/AIModule/Classes/BehaviorTree/Blackboard/BlackboardKeyEnums.h
--- a/Engine/Source/Runtime/AIModule/Classes/BehaviorTree/Blackboard/BlackboardKeyEnums.h
+++ b/Engine/Source/Runtime/AIModule/Classes/BehaviorTree/Blackboard/BlackboardKeyEnums.h
@@ -4,6 +4,7 @@
 
 #include "CoreTypes.h"
 
+#include "BlackboardKeyEnums.generated.h"
 
 namespace EBlackboardCompare
 {
diff --git a/Engine/Source/Runtime/CoreOnline/Private/Online/CoreOnlinePrivate.h b/Engine/Source/Runtime/CoreOnline/Private/Online/CoreOnlinePrivate.h
--- a/Engine/Source/Runtime/CoreOnline/Private/Online/CoreOnlinePrivate.h
+++ b/Engine/Source/Runtime/CoreOnline/Private/Online/CoreOnlinePrivate.h
@@ -5,6 +5,8 @@
 #include "Online/CoreOnline.h"
 #include "UObject/ObjectMacros.h"
 
+#include "CoreOnlinePrivate.generated.h"
+
 // This exists purely to guarantee a package is created, or the engine will not boot.
 UENUM(BlueprintType)
 enum class ECoreOnlineDummy : uint8
diff --git a/Engine/Source/Runtime/CoreOnline/Public/Online/CoreOnline.h b/Engine/Source/Runtime/CoreOnline/Public/Online/CoreOnline.h
--- a/Engine/Source/Runtime/CoreOnline/Public/Online/CoreOnline.h
+++ b/Engine/Source/Runtime/CoreOnline/Public/Online/CoreOnline.h
@@ -23,6 +23,8 @@
 #include "UObject/ObjectMacros.h"
 #include "UObject/UnrealNames.h"
 
+#include "CoreOnline.generated.h"
+
 class FDefaultSetAllocator;
 class FLazySingleton;
 
diff --git a/Engine/Source/Runtime/Engine/Classes/Engine/ViewportSplitScreen.h b/Engine/Source/Runtime/Engine/Classes/Engine/ViewportSplitScreen.h
--- a/Engine/Source/Runtime/Engine/Classes/Engine/ViewportSplitScreen.h
+++ b/Engine/Source/Runtime/Engine/Classes/Engine/ViewportSplitScreen.h
@@ -4,6 +4,8 @@
 
 #include "UObject/ObjectMacros.h"
 
+#include "ViewportSplitScreen.generated.h"
+
 /** 
  * Various data declarations relating to split screen on a GameViewportClient
  *
diff --git a/Engine/Source/Runtime/Engine/Classes/Sound/AudioOutputTarget.h b/Engine/Source/Runtime/Engine/Classes/Sound/AudioOutputTarget.h
--- a/Engine/Source/Runtime/Engine/Classes/Sound/AudioOutputTarget.h
+++ b/Engine/Source/Runtime/Engine/Classes/Sound/AudioOutputTarget.h
@@ -3,6 +3,8 @@
 
 #include "UObject/ObjectMacros.h"
 
+#include "AudioOutputTarget.generated.h"
+
 UENUM()
 namespace EAudioOutputTarget
 {
diff --git a/Engine/Source/Runtime/Engine/Classes/Sound/SoundWaveLoadingBehavior.h b/Engine/Source/Runtime/Engine/Classes/Sound/SoundWaveLoadingBehavior.h
--- a/Engine/Source/Runtime/Engine/Classes/Sound/SoundWaveLoadingBehavior.h
+++ b/Engine/Source/Runtime/Engine/Classes/Sound/SoundWaveLoadingBehavior.h
@@ -4,7 +4,7 @@
 #include "CoreTypes.h"
 #include "UObject/ObjectMacros.h"
 
-#include <UObject/ObjectMacros.h>
+#include "SoundWaveLoadingBehavior.generated.h"
 
 /**
  * Only used when stream caching is enabled. Determines how we are going to load or retain a given audio asset.
diff --git a/Engine/Source/Runtime/Engine/Private/ActorEditor.cpp b/Engine/Source/Runtime/Engine/Private/ActorEditor.cpp
--- a/Engine/Source/Runtime/Engine/Private/ActorEditor.cpp
+++ b/Engine/Source/Runtime/Engine/Private/ActorEditor.cpp
@@ -119,12 +119,12 @@ bool AActor::CanEditChange(const FProperty* PropertyThatWillChange) const
 	return Super::CanEditChange(PropertyThatWillChange);
 }
 
-static FName Name_RelativeLocation = USceneComponent::GetRelativeLocationPropertyName();
-static FName Name_RelativeRotation = USceneComponent::GetRelativeRotationPropertyName();
-static FName Name_RelativeScale3D = USceneComponent::GetRelativeScale3DPropertyName();
-
 void AActor::PostEditChangeProperty(FPropertyChangedEvent& PropertyChangedEvent)
 {
+	static const FName Name_RelativeLocation = USceneComponent::GetRelativeLocationPropertyName();
+	static const FName Name_RelativeRotation = USceneComponent::GetRelativeRotationPropertyName();
+	static const FName Name_RelativeScale3D = USceneComponent::GetRelativeScale3DPropertyName();
+
 	FProperty* MemberPropertyThatChanged = PropertyChangedEvent.MemberProperty;
 	const FName MemberPropertyName = MemberPropertyThatChanged != NULL ? MemberPropertyThatChanged->GetFName() : NAME_None;
 
diff --git a/Engine/Source/Runtime/Engine/Public/Animation/AnimPhysicsSolver.h b/Engine/Source/Runtime/Engine/Public/Animation/AnimPhysicsSolver.h
--- a/Engine/Source/Runtime/Engine/Public/Animation/AnimPhysicsSolver.h
+++ b/Engine/Source/Runtime/Engine/Public/Animation/AnimPhysicsSolver.h
@@ -34,6 +34,8 @@
 #include "Stats/Stats.h"
 #include "UObject/ObjectMacros.h"
 
+#include "AnimPhysicsSolver.generated.h"
+
 DECLARE_CYCLE_STAT_EXTERN(TEXT("Anim Dynamics Physics Update"), STAT_AnimDynamicsUpdate, STATGROUP_Physics, ENGINE_API);
 DECLARE_CYCLE_STAT_EXTERN(TEXT("Linear Limit Pre-Update"), STAT_AnimDynamicsLinearPre, STATGROUP_Physics, ENGINE_API);
 DECLARE_CYCLE_STAT_EXTERN(TEXT("Linear Limit Post-Update"), STAT_AnimDynamicsLinearPost, STATGROUP_Physics, ENGINE_API);
diff --git a/Engine/Source/Runtime/Engine/Public/HLOD/HLODBatchingPolicy.h b/Engine/Source/Runtime/Engine/Public/HLOD/HLODBatchingPolicy.h
--- a/Engine/Source/Runtime/Engine/Public/HLOD/HLODBatchingPolicy.h
+++ b/Engine/Source/Runtime/Engine/Public/HLOD/HLODBatchingPolicy.h
@@ -6,6 +6,8 @@
 
 #include "CoreTypes.h"
 
+#include "HLODBatchingPolicy.generated.h"
+
 /** Determines how the geometry of a component will be incorporated in proxy (simplified) HLODs. */
 UENUM()
 enum class EHLODBatchingPolicy : uint8
diff --git a/Engine/Source/Runtime/Engine/Public/HLOD/HLODLevelExclusion.h b/Engine/Source/Runtime/Engine/Public/HLOD/HLODLevelExclusion.h
--- a/Engine/Source/Runtime/Engine/Public/HLOD/HLODLevelExclusion.h
+++ b/Engine/Source/Runtime/Engine/Public/HLOD/HLODLevelExclusion.h
@@ -4,6 +4,8 @@
 
 #include "CoreMinimal.h"
 
+#include "HLODLevelExclusion.generated.h"
+
 /** Bitflag enum to allow editing of UPrimitiveComponent::ExcludeFromHLODLevels as a bitmask in the properties */
 UENUM(meta = (Bitflags, UseEnumValuesAsMaskValuesInEditor = "true"))
 enum class EHLODLevelExclusion : uint8
diff --git a/Engine/Source/Runtime/Engine/Public/HitProxies.h b/Engine/Source/Runtime/Engine/Public/HitProxies.h
--- a/Engine/Source/Runtime/Engine/Public/HitProxies.h
+++ b/Engine/Source/Runtime/Engine/Public/HitProxies.h
@@ -12,6 +12,8 @@
 #include "GenericPlatform/ICursor.h"
 #include "Elements/Framework/TypedElementHandle.h"
 
+#include "HitProxies.generated.h"
+
 class FReferenceCollector;
 
 /**
diff --git a/Engine/Source/Runtime/Engine/Public/LegacyScreenPercentageDriver.h b/Engine/Source/Runtime/Engine/Public/LegacyScreenPercentageDriver.h
--- a/Engine/Source/Runtime/Engine/Public/LegacyScreenPercentageDriver.h
+++ b/Engine/Source/Runtime/Engine/Public/LegacyScreenPercentageDriver.h
@@ -1,13 +1,10 @@
 // Copyright Epic Games, Inc. All Rights Reserved.
 
-/*=============================================================================
-	RendererPrivate.h: Renderer interface private definitions.
-=============================================================================*/
-
 #pragma once
 
 #include "SceneView.h"
 
+#include "LegacyScreenPercentageDriver.generated.h"
 
 /**
  * Default screen percentage interface that just apply View->FinalPostProcessSettings.ScreenPercentage.
diff --git a/Engine/Source/Runtime/Engine/Public/Net/NetPing.h b/Engine/Source/Runtime/Engine/Public/Net/NetPing.h
--- a/Engine/Source/Runtime/Engine/Public/Net/NetPing.h
+++ b/Engine/Source/Runtime/Engine/Public/Net/NetPing.h
@@ -6,6 +6,8 @@
 #include "Net/Core/Analytics/NetStatsUtils.h"
 #include "GameFramework/PlayerState.h"
 
+#include "NetPing.generated.h"
+
 // Forward declarations
 enum class EPingType : uint32;
 class FInternetAddr;
diff --git a/Engine/Source/Runtime/Engine/Public/ParticleHelper.h b/Engine/Source/Runtime/Engine/Public/ParticleHelper.h
--- a/Engine/Source/Runtime/Engine/Public/ParticleHelper.h
+++ b/Engine/Source/Runtime/Engine/Public/ParticleHelper.h
@@ -25,6 +25,8 @@
 #include "PrimitiveSceneProxy.h"
 #include "Particles/ParticlePerfStats.h"
 
+#include "ParticleHelper.generated.h"
+
 #define _ENABLE_PARTICLE_LOD_INGAME_
 
 class FColoredMaterialRenderProxy;
diff --git a/Engine/Source/Runtime/Engine/Public/Rendering/SkeletalMeshLODImporterData.h b/Engine/Source/Runtime/Engine/Public/Rendering/SkeletalMeshLODImporterData.h
--- a/Engine/Source/Runtime/Engine/Public/Rendering/SkeletalMeshLODImporterData.h
+++ b/Engine/Source/Runtime/Engine/Public/Rendering/SkeletalMeshLODImporterData.h
@@ -14,6 +14,8 @@
 #include "Animation/MorphTarget.h"
 #include "Templates/DontCopy.h"
 
+#include "SkeletalMeshLODImporterData.generated.h"
+
 struct FMeshDescription;
 class FSkeletalMeshLODModel;
 
diff --git a/Engine/Source/Runtime/Engine/Public/SceneUtils.h b/Engine/Source/Runtime/Engine/Public/SceneUtils.h
--- a/Engine/Source/Runtime/Engine/Public/SceneUtils.h
+++ b/Engine/Source/Runtime/Engine/Public/SceneUtils.h
@@ -17,6 +17,8 @@
 #endif
 #include "RHIFeatureLevel.h"
 
+#include "SceneUtils.generated.h"
+
 ENGINE_API DECLARE_LOG_CATEGORY_EXTERN(LogSceneUtils, Log, All);
 
 class FRHICommandListImmediate;
diff --git a/Engine/Source/Runtime/Engine/Public/Slate/SGameLayerManager.h b/Engine/Source/Runtime/Engine/Public/Slate/SGameLayerManager.h
--- a/Engine/Source/Runtime/Engine/Public/Slate/SGameLayerManager.h
+++ b/Engine/Source/Runtime/Engine/Public/Slate/SGameLayerManager.h
@@ -14,6 +14,8 @@
 #include "Styling/SlateTypes.h"
 #include "UObject/ObjectKey.h"
 
+#include "SGameLayerManager.generated.h"
+
 class FPaintArgs;
 class FSceneViewport;
 class FSlateWindowElementList;
diff --git a/Engine/Source/Runtime/Engine/Public/WorldPartition/ContentBundle/ContentBundleClient.h b/Engine/Source/Runtime/Engine/Public/WorldPartition/ContentBundle/ContentBundleClient.h
--- a/Engine/Source/Runtime/Engine/Public/WorldPartition/ContentBundle/ContentBundleClient.h
+++ b/Engine/Source/Runtime/Engine/Public/WorldPartition/ContentBundle/ContentBundleClient.h
@@ -7,6 +7,8 @@
 
 #include "WorldPartition/ContentBundle/ContentBundleStatus.h"
 
+#include "ContentBundleClient.generated.h"
+
 class UContentBundleDescriptor;
 class UWorld;
 
diff --git a/Engine/Source/Runtime/Experimental/Chaos/Public/Field/FieldSystemTypes.h b/Engine/Source/Runtime/Experimental/Chaos/Public/Field/FieldSystemTypes.h
--- a/Engine/Source/Runtime/Experimental/Chaos/Public/Field/FieldSystemTypes.h
+++ b/Engine/Source/Runtime/Experimental/Chaos/Public/Field/FieldSystemTypes.h
@@ -1,7 +1,10 @@
 // Copyright Epic Games, Inc. All Rights Reserved.
 #pragma once
+
 #include "UObject/ObjectMacros.h"
 
+#include "FieldSystemTypes.generated.h"
+
 /**
 *
 */
diff --git a/Engine/Source/Runtime/Experimental/Chaos/Public/GeometryCollection/GeometryCollectionConvexUtility.h b/Engine/Source/Runtime/Experimental/Chaos/Public/GeometryCollection/GeometryCollectionConvexUtility.h
--- a/Engine/Source/Runtime/Experimental/Chaos/Public/GeometryCollection/GeometryCollectionConvexUtility.h
+++ b/Engine/Source/Runtime/Experimental/Chaos/Public/GeometryCollection/GeometryCollectionConvexUtility.h
@@ -5,6 +5,8 @@
 #include "CoreMinimal.h"
 #include "GeometryCollection/ManagedArrayCollection.h"
 
+#include "GeometryCollectionConvexUtility.generated.h"
+
 class FGeometryCollection;
 
 namespace Chaos
diff --git a/Engine/Source/Runtime/Experimental/Chaos/Public/GeometryCollection/GeometryCollectionProximityUtility.h b/Engine/Source/Runtime/Experimental/Chaos/Public/GeometryCollection/GeometryCollectionProximityUtility.h
--- a/Engine/Source/Runtime/Experimental/Chaos/Public/GeometryCollection/GeometryCollectionProximityUtility.h
+++ b/Engine/Source/Runtime/Experimental/Chaos/Public/GeometryCollection/GeometryCollectionProximityUtility.h
@@ -5,6 +5,8 @@
 #include "CoreTypes.h"
 #include "UObject/ObjectMacros.h"
 
+#include "GeometryCollectionProximityUtility.generated.h"
+
 class FGeometryCollection;
 namespace UE::GeometryCollectionConvexUtility
 {
diff --git a/Engine/Source/Runtime/Experimental/GeometryCollectionEngine/Public/GeometryCollection/GeometryCollectionEngineTypes.h b/Engine/Source/Runtime/Experimental/GeometryCollectionEngine/Public/GeometryCollection/GeometryCollectionEngineTypes.h
--- a/Engine/Source/Runtime/Experimental/GeometryCollectionEngine/Public/GeometryCollection/GeometryCollectionEngineTypes.h
+++ b/Engine/Source/Runtime/Experimental/GeometryCollectionEngine/Public/GeometryCollection/GeometryCollectionEngineTypes.h
@@ -7,6 +7,8 @@
 #include "GeometryCollectionProxyData.h"
 #include "GeometryCollection/GeometryCollection.h"
 
+#include "GeometryCollectionEngineTypes.generated.h"
+
 UENUM()
 enum class ECollectionAttributeEnum : uint8
 {
diff --git a/Engine/Source/Runtime/HeadMountedDisplay/Public/IMotionController.h b/Engine/Source/Runtime/HeadMountedDisplay/Public/IMotionController.h
--- a/Engine/Source/Runtime/HeadMountedDisplay/Public/IMotionController.h
+++ b/Engine/Source/Runtime/HeadMountedDisplay/Public/IMotionController.h
@@ -7,6 +7,8 @@
 #include "UObject/ObjectPtr.h"
 #include "Features/IModularFeature.h"
 
+#include "IMotionController.generated.h"
+
 enum class EControllerHand : uint8;
 
 UENUM(BlueprintType)
diff --git a/Engine/Source/Runtime/InteractiveToolsFramework/Public/ToolContextInterfaces.h b/Engine/Source/Runtime/InteractiveToolsFramework/Public/ToolContextInterfaces.h
--- a/Engine/Source/Runtime/InteractiveToolsFramework/Public/ToolContextInterfaces.h
+++ b/Engine/Source/Runtime/InteractiveToolsFramework/Public/ToolContextInterfaces.h
@@ -18,6 +18,8 @@
 #include "UObject/WeakObjectPtr.h"
 #include "UObject/WeakObjectPtrTemplates.h"
 
+#include "ToolContextInterfaces.generated.h"
+
 // predeclarations so we don't have to include these in all tools
 class AActor;
 class FPrimitiveDrawInterface;
diff --git a/Engine/Source/Runtime/MoviePlayer/Public/MoviePlayer.h b/Engine/Source/Runtime/MoviePlayer/Public/MoviePlayer.h
--- a/Engine/Source/Runtime/MoviePlayer/Public/MoviePlayer.h
+++ b/Engine/Source/Runtime/MoviePlayer/Public/MoviePlayer.h
@@ -9,6 +9,8 @@
 #include "Widgets/SWidget.h"
 #include "Slate/SlateTextures.h"
 
+#include "MoviePlayer.generated.h"
+
 // In order for a platform to support early movie playback, the platform must support the rendering thread 
 // starting very early and support rendering as soon as it is started and the module containing the movie streamer for the platform must already be loaded
 
diff --git a/Engine/Source/Runtime/MovieScene/Public/KeyParams.h b/Engine/Source/Runtime/MovieScene/Public/KeyParams.h
--- a/Engine/Source/Runtime/MovieScene/Public/KeyParams.h
+++ b/Engine/Source/Runtime/MovieScene/Public/KeyParams.h
@@ -5,6 +5,8 @@
 #include "CoreMinimal.h"
 #include "UObject/ObjectMacros.h"
 
+#include "KeyParams.generated.h"
+
 UENUM()
 enum class EMovieSceneKeyInterpolation : uint8
 {
diff --git a/Engine/Source/Runtime/MovieScene/Public/MovieSceneFwd.h b/Engine/Source/Runtime/MovieScene/Public/MovieSceneFwd.h
--- a/Engine/Source/Runtime/MovieScene/Public/MovieSceneFwd.h
+++ b/Engine/Source/Runtime/MovieScene/Public/MovieSceneFwd.h
@@ -9,6 +9,8 @@
 #include "MovieSceneSequenceID.h"
 #include "HAL/PreprocessorHelpers.h"
 
+#include "MovieSceneFwd.generated.h"
+
 #ifndef UE_MOVIESCENE_EVENTS
 	#define UE_MOVIESCENE_EVENTS WITH_EDITOR
 #endif
diff --git a/Engine/Source/Runtime/Net/Core/Public/Net/Core/Connection/NetEnums.h b/Engine/Source/Runtime/Net/Core/Public/Net/Core/Connection/NetEnums.h
--- a/Engine/Source/Runtime/Net/Core/Public/Net/Core/Connection/NetEnums.h
+++ b/Engine/Source/Runtime/Net/Core/Public/Net/Core/Connection/NetEnums.h
@@ -5,6 +5,7 @@
 #include "CoreMinimal.h"
 #include "UObject/ObjectMacros.h"
 
+#include "NetEnums.generated.h"
 
 /** Types of network failures broadcast from the engine */
 UENUM(BlueprintType)
diff --git a/Engine/Source/Runtime/RigVM/Public/RigVMCore/RigVMFunction.h b/Engine/Source/Runtime/RigVM/Public/RigVMCore/RigVMFunction.h
--- a/Engine/Source/Runtime/RigVM/Public/RigVMCore/RigVMFunction.h
+++ b/Engine/Source/Runtime/RigVM/Public/RigVMCore/RigVMFunction.h
@@ -15,6 +15,8 @@
 #include "UObject/NameTypes.h"
 #include "UObject/ObjectMacros.h"
 
+#include "RigVMFunction.generated.h"
+
 class UScriptStruct;
 struct FRigVMExtendedExecuteContext;
 struct FRigVMMemoryHandle;
diff --git a/Engine/Source/Runtime/RigVM/Public/RigVMFunctions/RigVMFunctionDefines.h b/Engine/Source/Runtime/RigVM/Public/RigVMFunctions/RigVMFunctionDefines.h
--- a/Engine/Source/Runtime/RigVM/Public/RigVMFunctions/RigVMFunctionDefines.h
+++ b/Engine/Source/Runtime/RigVM/Public/RigVMFunctions/RigVMFunctionDefines.h
@@ -5,6 +5,8 @@
 #include "CoreMinimal.h"
 #include "UObject/ObjectMacros.h"
 
+#include "RigVMFunctionDefines.generated.h"
+
 UENUM()
 enum class ERigVMTransformSpace : uint8
 {
diff --git a/Engine/Source/Runtime/Slate/Public/Widgets/Layout/SScrollBox.h b/Engine/Source/Runtime/Slate/Public/Widgets/Layout/SScrollBox.h
--- a/Engine/Source/Runtime/Slate/Public/Widgets/Layout/SScrollBox.h
+++ b/Engine/Source/Runtime/Slate/Public/Widgets/Layout/SScrollBox.h
@@ -22,6 +22,8 @@
 #include "Framework/Layout/InertialScrollManager.h"
 #include "Framework/Layout/Overscroll.h"
 
+#include "SScrollBox.generated.h"
+
 class FPaintArgs;
 class FSlateWindowElementList;
 class SScrollPanel;
diff --git a/Engine/Source/Runtime/Slate/Public/Widgets/Text/ISlateEditableTextWidget.h b/Engine/Source/Runtime/Slate/Public/Widgets/Text/ISlateEditableTextWidget.h
--- a/Engine/Source/Runtime/Slate/Public/Widgets/Text/ISlateEditableTextWidget.h
+++ b/Engine/Source/Runtime/Slate/Public/Widgets/Text/ISlateEditableTextWidget.h
@@ -8,6 +8,8 @@
 #include "Framework/Application/IMenu.h"
 #include "Widgets/Input/IVirtualKeyboardEntry.h"
 
+#include "ISlateEditableTextWidget.generated.h"
+
 struct FTextLocation;
 
 enum class ECursorMoveMethod
diff --git a/Engine/Source/Runtime/VectorVM/Public/VectorVMCommon.h b/Engine/Source/Runtime/VectorVM/Public/VectorVMCommon.h
--- a/Engine/Source/Runtime/VectorVM/Public/VectorVMCommon.h
+++ b/Engine/Source/Runtime/VectorVM/Public/VectorVMCommon.h
@@ -6,6 +6,8 @@
 #include <atomic>
 #include "UObject/ObjectMacros.h"
 
+#include "VectorVMCommon.generated.h"
+
 //TODO: move to a per platform header and have VM scale vectorization according to vector width.
 #define VECTOR_WIDTH (128)
 #define VECTOR_WIDTH_BYTES (16)
