Optimized OnLodStrippingQualityLevelChanged by passing in an array instead of passing each item individually. This changes it so the expensive calls in FStaticMeshComponentRecreateRenderStateContext only happen once.

Before: 741ms
After: 122ms

diff --git a/Engine/Source/Runtime/Engine/Private/StaticMesh.cpp b/Engine/Source/Runtime/Engine/Private/StaticMesh.cpp
--- a/Engine/Source/Runtime/Engine/Private/StaticMesh.cpp
+++ b/Engine/Source/Runtime/Engine/Private/StaticMesh.cpp
@@ -8251,14 +8251,20 @@ void UStaticMesh::OnLodStrippingQualityLevelChanged(IConsoleVariable* Variable){
 #if WITH_EDITOR || PLATFORM_DESKTOP
 	if (GEngine && GEngine->UseStaticMeshMinLODPerQualityLevels)
 	{
+		TArray<UStaticMesh*> StaticMeshes;
 		for (TObjectIterator<UStaticMesh> It; It; ++It)
 		{
 			UStaticMesh* StaticMesh = *It;
 			if (StaticMesh && StaticMesh->GetQualityLevelMinLOD().PerQuality.Num() > 0)
 			{
-				FStaticMeshComponentRecreateRenderStateContext Context(StaticMesh, false);
+				StaticMeshes.Add(StaticMesh);
 			}
 		}
+
+		if (StaticMeshes.Num() > 0)
+		{
+			FStaticMeshComponentRecreateRenderStateContext Context(StaticMeshes, false);
+		}
 	}
 #endif
 }
