Fixed non unity build.

diff --git a/Engine/Source/Runtime/Engine/Classes/Materials/MaterialExpressionSparseVolumeTextureSample.h b/Engine/Source/Runtime/Engine/Classes/Materials/MaterialExpressionSparseVolumeTextureSample.h
--- a/Engine/Source/Runtime/Engine/Classes/Materials/MaterialExpressionSparseVolumeTextureSample.h
+++ b/Engine/Source/Runtime/Engine/Classes/Materials/MaterialExpressionSparseVolumeTextureSample.h
@@ -10,6 +10,7 @@
 #include "MaterialExpressionSparseVolumeTextureSample.generated.h"
 
 class USparseVolumeTexture;
+struct FMaterialParameterMetadata;
 
 /** Material expression for sampling from a runtime virtual texture. */
 UCLASS(collapsecategories, hidecategories=Object)
@@ -81,32 +82,8 @@ class UMaterialExpressionSparseVolumeTextureSampleParameter : public UMaterialEx
 	virtual void ValidateParameterName(const bool bAllowDuplicateName) override;
 	virtual void GetCaption(TArray<FString>& OutCaptions) const override;
 	virtual bool MatchesSearchQuery(const TCHAR* SearchQuery) override;
-	virtual bool GetParameterValue(FMaterialParameterMetadata& OutMeta) const override
-	{
-		OutMeta.Value = SparseVolumeTexture;
-		OutMeta.Description = Desc;
-		OutMeta.ExpressionGuid = ExpressionGUID;
-		OutMeta.Group = Group;
-		OutMeta.SortPriority = SortPriority;
-		OutMeta.AssetPath = GetAssetPathName();
-		return true;
-	}
-	virtual bool SetParameterValue(const FName& Name, const FMaterialParameterMetadata& Meta, EMaterialExpressionSetParameterValueFlags Flags) override
-	{
-		if (Meta.Value.Type == EMaterialParameterType::SparseVolumeTexture)
-		{
-			if (SetParameterValue(Name, Meta.Value.SparseVolumeTexture, Flags))
-			{
-				if (EnumHasAnyFlags(Flags, EMaterialExpressionSetParameterValueFlags::AssignGroupAndSortPriority))
-				{
-					Group = Meta.Group;
-					SortPriority = Meta.SortPriority;
-				}
-				return true;
-			}
-		}
-		return false;
-	}
+	virtual bool GetParameterValue(FMaterialParameterMetadata& OutMeta) const override;
+	virtual bool SetParameterValue(const FName& Name, const FMaterialParameterMetadata& Meta, EMaterialExpressionSetParameterValueFlags Flags) override;
 #endif
 	virtual FGuid& GetParameterExpressionId() override { return ExpressionGUID; }
 	//~ End UMaterialExpression Interface
diff --git a/Engine/Source/Runtime/Engine/Private/Materials/MaterialExpressions.cpp b/Engine/Source/Runtime/Engine/Private/Materials/MaterialExpressions.cpp
--- a/Engine/Source/Runtime/Engine/Private/Materials/MaterialExpressions.cpp
+++ b/Engine/Source/Runtime/Engine/Private/Materials/MaterialExpressions.cpp
@@ -25361,6 +25361,33 @@ void UMaterialExpressionSparseVolumeTextureSample::GetCaption(TArray<FString>& O
 	OutCaptions.Add(FString(TEXT("Sparse Volume Texture Sample")));
 }
 
+bool UMaterialExpressionSparseVolumeTextureSampleParameter::GetParameterValue(FMaterialParameterMetadata& OutMeta) const
+{
+	OutMeta.Value = SparseVolumeTexture;
+	OutMeta.Description = Desc;
+	OutMeta.ExpressionGuid = ExpressionGUID;
+	OutMeta.Group = Group;
+	OutMeta.SortPriority = SortPriority;
+	OutMeta.AssetPath = GetAssetPathName();
+	return true;
+}
+bool UMaterialExpressionSparseVolumeTextureSampleParameter::SetParameterValue(const FName& Name, const FMaterialParameterMetadata& Meta, EMaterialExpressionSetParameterValueFlags Flags)
+{
+	if (Meta.Value.Type == EMaterialParameterType::SparseVolumeTexture)
+	{
+		if (SetParameterValue(Name, Meta.Value.SparseVolumeTexture, Flags))
+		{
+			if (EnumHasAnyFlags(Flags, EMaterialExpressionSetParameterValueFlags::AssignGroupAndSortPriority))
+			{
+				Group = Meta.Group;
+				SortPriority = Meta.SortPriority;
+			}
+			return true;
+		}
+	}
+	return false;
+}
+
 #endif // WITH_EDITOR
 
 ///////////////////////////////////////////////////////////////////////////////
