World Partition: don't validate runtime grids for autogenerated HLODs.

diff --git a/Engine/Source/Runtime/Engine/Private/WorldPartition/WorldPartitionActorDescView.cpp b/Engine/Source/Runtime/Engine/Private/WorldPartition/WorldPartitionActorDescView.cpp
--- a/Engine/Source/Runtime/Engine/Private/WorldPartition/WorldPartitionActorDescView.cpp
+++ b/Engine/Source/Runtime/Engine/Private/WorldPartition/WorldPartitionActorDescView.cpp
@@ -188,6 +188,13 @@ FName FWorldPartitionActorDescView::GetActorLabelOrName() const
 	return ActorDesc->GetActorLabelOrName();
 }
 
+bool FWorldPartitionActorDescView::ShouldValidateRuntimeGrid() const
+{
+PRAGMA_DISABLE_DEPRECATION_WARNINGS
+	return ActorDesc->ShouldValidateRuntimeGrid();
+PRAGMA_ENABLE_DEPRECATION_WARNINGS
+}
+
 void FWorldPartitionActorDescView::SetForcedNonSpatiallyLoaded()
 {
 	if (!bIsForcedNonSpatiallyLoaded)
diff --git a/Engine/Source/Runtime/Engine/Private/WorldPartition/WorldPartitionStreamingGeneration.cpp b/Engine/Source/Runtime/Engine/Private/WorldPartition/WorldPartitionStreamingGeneration.cpp
--- a/Engine/Source/Runtime/Engine/Private/WorldPartition/WorldPartitionStreamingGeneration.cpp
+++ b/Engine/Source/Runtime/Engine/Private/WorldPartition/WorldPartitionStreamingGeneration.cpp
@@ -597,6 +597,13 @@ class FWorldPartitionStreamingGenerator
 				// Validate runtime grid
 				auto IsRuntimeGridValid = [this](const FWorldPartitionActorDescView& ActorDescView)
 				{
+PRAGMA_DISABLE_DEPRECATION_WARNINGS
+					if (!ActorDescView.ShouldValidateRuntimeGrid())
+					{
+						return true;
+					}
+PRAGMA_ENABLE_DEPRECATION_WARNINGS
+
 					return IsValidGrid(ActorDescView.GetRuntimeGrid());
 				};
 
diff --git a/Engine/Source/Runtime/Engine/Public/WorldPartition/HLOD/HLODActorDesc.h b/Engine/Source/Runtime/Engine/Public/WorldPartition/HLOD/HLODActorDesc.h
--- a/Engine/Source/Runtime/Engine/Public/WorldPartition/HLOD/HLODActorDesc.h
+++ b/Engine/Source/Runtime/Engine/Public/WorldPartition/HLOD/HLODActorDesc.h
@@ -39,6 +39,7 @@ protected:
 	virtual bool Equals(const FWorldPartitionActorDesc* Other) const override;
 	virtual void Serialize(FArchive& Ar) override;
 	virtual bool IsRuntimeRelevant(const FActorContainerID& InContainerID) const override { return !bIsForcedNonSpatiallyLoaded; }
+	virtual bool ShouldValidateRuntimeGrid() const override { return false; }
 	//~ End FWorldPartitionActorDesc Interface.
 
 	TArray<FHLODSubActorDesc> HLODSubActors;
diff --git a/Engine/Source/Runtime/Engine/Public/WorldPartition/WorldPartitionActorDesc.h b/Engine/Source/Runtime/Engine/Public/WorldPartition/WorldPartitionActorDesc.h
--- a/Engine/Source/Runtime/Engine/Public/WorldPartition/WorldPartitionActorDesc.h
+++ b/Engine/Source/Runtime/Engine/Public/WorldPartition/WorldPartitionActorDesc.h
@@ -143,6 +143,9 @@ PRAGMA_ENABLE_DEPRECATION_WARNINGS
 	virtual bool IsRuntimeRelevant(const FActorContainerID& InContainerID) const { return true; }
 	virtual void CheckForErrors(IStreamingGenerationErrorHandler* ErrorHandler) const;
 
+	UE_DEPRECATED(5.2, "ShouldValidateRuntimeGrid is deprecated and should not be used.")
+	virtual bool ShouldValidateRuntimeGrid() const { return true; }
+
 	bool operator==(const FWorldPartitionActorDesc& Other) const
 	{
 		return Guid == Other.Guid;
diff --git a/Engine/Source/Runtime/Engine/Public/WorldPartition/WorldPartitionActorDescView.h b/Engine/Source/Runtime/Engine/Public/WorldPartition/WorldPartitionActorDescView.h
--- a/Engine/Source/Runtime/Engine/Public/WorldPartition/WorldPartitionActorDescView.h
+++ b/Engine/Source/Runtime/Engine/Public/WorldPartition/WorldPartitionActorDescView.h
@@ -30,10 +30,8 @@ public:
 	FTopLevelAssetPath GetBaseClass() const;
 	FTopLevelAssetPath GetNativeClass() const;
 	UClass* GetActorNativeClass() const;
-PRAGMA_DISABLE_DEPRECATION_WARNINGS
 	UE_DEPRECATED(5.2, "GetOrigin is deprecated.")
 	FVector GetOrigin() const;
-PRAGMA_ENABLE_DEPRECATION_WARNINGS
 	FName GetRuntimeGrid() const;
 	bool GetIsSpatiallyLoaded() const;
 	bool GetActorIsEditorOnly() const;
@@ -53,10 +51,8 @@ PRAGMA_ENABLE_DEPRECATION_WARNINGS
 	FSoftObjectPath GetActorSoftPath() const;
 	FName GetActorLabel() const;
 
-PRAGMA_DISABLE_DEPRECATION_WARNINGS
 	UE_DEPRECATED(5.2, "GetBounds is deprecated, GetEditorBounds or GetRuntimeBounds should be used instead.")
 	FBox GetBounds() const;
-PRAGMA_ENABLE_DEPRECATION_WARNINGS
 
 	FBox GetEditorBounds() const;
 	FBox GetRuntimeBounds() const;
@@ -77,6 +73,9 @@ PRAGMA_ENABLE_DEPRECATION_WARNINGS
 
 	FName GetActorLabelOrName() const;
 
+	UE_DEPRECATED(5.2, "ShouldValidateRuntimeGrid is deprecated and should not be used.")
+	bool ShouldValidateRuntimeGrid() const;
+
 	void SetForcedNonSpatiallyLoaded();
 	void SetInvalidRuntimeGrid();
 	void SetInvalidDataLayers();
