Fix crash when we delete a compiling anim sequence and a skeleton asset.

#rnx

diff --git a/Engine/Source/Runtime/Engine/Private/Animation/AnimSequence.cpp b/Engine/Source/Runtime/Engine/Private/Animation/AnimSequence.cpp
--- a/Engine/Source/Runtime/Engine/Private/Animation/AnimSequence.cpp
+++ b/Engine/Source/Runtime/Engine/Private/Animation/AnimSequence.cpp
@@ -5208,7 +5208,10 @@ void UAnimSequence::FinishAsyncTasks()
 #if WITH_EDITOR
 		SynchronousAnimatedBoneAttributesCompression();
 		//This is only safe during sync anim compression
-		SetSkeletonVirtualBoneGuid(GetSkeleton()->GetVirtualBoneGuid());
+		if (GetSkeleton())
+		{
+			SetSkeletonVirtualBoneGuid(GetSkeleton()->GetVirtualBoneGuid());
+		}
 #endif
 		
 		FAssetRegistryModule& AssetRegistryModule = FModuleManager::LoadModuleChecked<FAssetRegistryModule>(TEXT("AssetRegistry"));
