USplineMeshComponent: Rename GetAxisValue to GetAxisValueRef

#9909

diff --git a/Engine/Plugins/MeshPainting/Source/MeshPaintingToolset/Private/MeshPaintSplineMeshAdapter.cpp b/Engine/Plugins/MeshPainting/Source/MeshPaintingToolset/Private/MeshPaintSplineMeshAdapter.cpp
--- a/Engine/Plugins/MeshPainting/Source/MeshPaintingToolset/Private/MeshPaintSplineMeshAdapter.cpp
+++ b/Engine/Plugins/MeshPainting/Source/MeshPaintingToolset/Private/MeshPaintSplineMeshAdapter.cpp
@@ -31,8 +31,8 @@ bool FMeshPaintSplineMeshComponentAdapter::InitializeVertexData()
 		for (int32 Index = 0; Index < NumVertices; Index++)
 		{
 			FVector Position = (FVector)LODModel->VertexBuffers.PositionVertexBuffer.VertexPosition(Index);
-			const FTransform SliceTransform = SplineMeshComponent->CalcSliceTransform(USplineMeshComponent::GetAxisValue(Position, SplineMeshComponent->ForwardAxis));
-			USplineMeshComponent::GetAxisValue(Position, SplineMeshComponent->ForwardAxis) = 0;
+			const FTransform SliceTransform = SplineMeshComponent->CalcSliceTransform(USplineMeshComponent::GetAxisValueRef(Position, SplineMeshComponent->ForwardAxis));
+			USplineMeshComponent::GetAxisValueRef(Position, SplineMeshComponent->ForwardAxis) = 0;
 			MeshVertices[Index] = SliceTransform.TransformPosition(Position);
 		}
 
diff --git a/Engine/Source/Developer/MeshMergeUtilities/Private/MeshMergeHelpers.cpp b/Engine/Source/Developer/MeshMergeUtilities/Private/MeshMergeHelpers.cpp
--- a/Engine/Source/Developer/MeshMergeUtilities/Private/MeshMergeHelpers.cpp
+++ b/Engine/Source/Developer/MeshMergeUtilities/Private/MeshMergeHelpers.cpp
@@ -741,7 +741,7 @@ void FMeshMergeHelpers::PropagateSplineDeformationToMesh(const USplineMeshCompon
 		{
 			const FVertexInstanceID VertexInstanceID = InOutMeshDescription.GetTriangleVertexInstance(TriangleID, Corner);
 			const FVertexID VertexID = InOutMeshDescription.GetVertexInstanceVertex(VertexInstanceID);
-			const float& AxisValue = USplineMeshComponent::GetAxisValue(VertexPositions[VertexID], InSplineMeshComponent->ForwardAxis);
+			const float& AxisValue = USplineMeshComponent::GetAxisValueRef(VertexPositions[VertexID], InSplineMeshComponent->ForwardAxis);
 			FTransform SliceTransform = InSplineMeshComponent->CalcSliceTransform(AxisValue);
 			FVector TangentY = FVector::CrossProduct((FVector)VertexInstanceNormals[VertexInstanceID], (FVector)VertexInstanceTangents[VertexInstanceID]).GetSafeNormal() * VertexInstanceBinormalSigns[VertexInstanceID];
 			VertexInstanceTangents[VertexInstanceID] = (FVector3f)SliceTransform.TransformVector((FVector)VertexInstanceTangents[VertexInstanceID]);
@@ -754,7 +754,7 @@ void FMeshMergeHelpers::PropagateSplineDeformationToMesh(const USplineMeshCompon
 	// Apply spline deformation for each vertex position
 	for (const FVertexID VertexID : InOutMeshDescription.Vertices().GetElementIDs())
 	{
-		float& AxisValue = USplineMeshComponent::GetAxisValue(VertexPositions[VertexID], InSplineMeshComponent->ForwardAxis);
+		float& AxisValue = USplineMeshComponent::GetAxisValueRef(VertexPositions[VertexID], InSplineMeshComponent->ForwardAxis);
 		FTransform SliceTransform = InSplineMeshComponent->CalcSliceTransform(AxisValue);
 		AxisValue = 0.0f;
 		VertexPositions[VertexID] = (FVector3f)SliceTransform.TransformPosition((FVector)VertexPositions[VertexID]);
@@ -769,7 +769,7 @@ void FMeshMergeHelpers::PropagateSplineDeformationToPhysicsGeometry(USplineMeshC
 	{
 		for (FVector& Position : Elem.VertexData)
 		{
-			const float& AxisValue = USplineMeshComponent::GetAxisValue(Position, SplineMeshComponent->ForwardAxis);
+			const float& AxisValue = USplineMeshComponent::GetAxisValueRef(Position, SplineMeshComponent->ForwardAxis);
 			FTransform SliceTransform = SplineMeshComponent->CalcSliceTransform(AxisValue);
 			Position = SliceTransform.TransformPosition(Position * Mask);
 		}
@@ -780,13 +780,13 @@ void FMeshMergeHelpers::PropagateSplineDeformationToPhysicsGeometry(USplineMeshC
 	for (FKSphereElem& Elem : InOutPhysicsGeometry.SphereElems)
 	{
 		const FVector WorldSpaceCenter = Elem.GetTransform().TransformPosition(Elem.Center);
-		Elem.Center = SplineMeshComponent->CalcSliceTransform(USplineMeshComponent::GetAxisValue(WorldSpaceCenter, SplineMeshComponent->ForwardAxis)).TransformPosition(Elem.Center * Mask);
+		Elem.Center = SplineMeshComponent->CalcSliceTransform(USplineMeshComponent::GetAxisValueRef(WorldSpaceCenter, SplineMeshComponent->ForwardAxis)).TransformPosition(Elem.Center * Mask);
 	}
 
 	for (FKSphylElem& Elem : InOutPhysicsGeometry.SphylElems)
 	{
 		const FVector WorldSpaceCenter = Elem.GetTransform().TransformPosition(Elem.Center);
-		Elem.Center = SplineMeshComponent->CalcSliceTransform(USplineMeshComponent::GetAxisValue(WorldSpaceCenter, SplineMeshComponent->ForwardAxis)).TransformPosition(Elem.Center * Mask);
+		Elem.Center = SplineMeshComponent->CalcSliceTransform(USplineMeshComponent::GetAxisValueRef(WorldSpaceCenter, SplineMeshComponent->ForwardAxis)).TransformPosition(Elem.Center * Mask);
 	}
 }
 
diff --git a/Engine/Source/Editor/FoliageEdit/Private/FoliageEdMode.cpp b/Engine/Source/Editor/FoliageEdit/Private/FoliageEdMode.cpp
--- a/Engine/Source/Editor/FoliageEdit/Private/FoliageEdMode.cpp
+++ b/Engine/Source/Editor/FoliageEdit/Private/FoliageEdMode.cpp
@@ -2822,7 +2822,7 @@ void FEdModeFoliage::ApplyPaintBucket_Add(AActor* Actor)
 			{
 				// Transform spline mesh verts correctly
 				FVector Mask = FVector(1, 1, 1);
-				USplineMeshComponent::GetAxisValue(Mask, SplineMesh->ForwardAxis) = 0;
+				USplineMeshComponent::GetAxisValueRef(Mask, SplineMesh->ForwardAxis) = 0;
 
 				for (int32 Idx = 0; Idx < Indices.Num(); Idx += 3)
 				{
@@ -2830,9 +2830,9 @@ void FEdModeFoliage::ApplyPaintBucket_Add(AActor* Actor)
 					const int32 Index1 = Indices[Idx + 1];
 					const int32 Index2 = Indices[Idx + 2];
 
-					const FVector Vert0 = SplineMesh->CalcSliceTransform(USplineMeshComponent::GetAxisValue(PositionVertexBuffer.VertexPosition(Index0), SplineMesh->ForwardAxis)).TransformPosition((FVector)PositionVertexBuffer.VertexPosition(Index0) * Mask);
-					const FVector Vert1 = SplineMesh->CalcSliceTransform(USplineMeshComponent::GetAxisValue(PositionVertexBuffer.VertexPosition(Index1), SplineMesh->ForwardAxis)).TransformPosition((FVector)PositionVertexBuffer.VertexPosition(Index1) * Mask);
-					const FVector Vert2 = SplineMesh->CalcSliceTransform(USplineMeshComponent::GetAxisValue(PositionVertexBuffer.VertexPosition(Index2), SplineMesh->ForwardAxis)).TransformPosition((FVector)PositionVertexBuffer.VertexPosition(Index2) * Mask);
+					const FVector Vert0 = SplineMesh->CalcSliceTransform(USplineMeshComponent::GetAxisValueRef(PositionVertexBuffer.VertexPosition(Index0), SplineMesh->ForwardAxis)).TransformPosition((FVector)PositionVertexBuffer.VertexPosition(Index0) * Mask);
+					const FVector Vert1 = SplineMesh->CalcSliceTransform(USplineMeshComponent::GetAxisValueRef(PositionVertexBuffer.VertexPosition(Index1), SplineMesh->ForwardAxis)).TransformPosition((FVector)PositionVertexBuffer.VertexPosition(Index1) * Mask);
+					const FVector Vert2 = SplineMesh->CalcSliceTransform(USplineMeshComponent::GetAxisValueRef(PositionVertexBuffer.VertexPosition(Index2), SplineMesh->ForwardAxis)).TransformPosition((FVector)PositionVertexBuffer.VertexPosition(Index2) * Mask);
 
 					new(PotentialTriangles)FFoliagePaintBucketTriangle(LocalToWorld
 						, Vert0
diff --git a/Engine/Source/Editor/MeshPaint/Private/MeshPaintSplineMeshAdapter.cpp b/Engine/Source/Editor/MeshPaint/Private/MeshPaintSplineMeshAdapter.cpp
--- a/Engine/Source/Editor/MeshPaint/Private/MeshPaintSplineMeshAdapter.cpp
+++ b/Engine/Source/Editor/MeshPaint/Private/MeshPaintSplineMeshAdapter.cpp
@@ -26,8 +26,8 @@ bool FMeshPaintGeometryAdapterForSplineMeshes::InitializeVertexData()
 		for (int32 Index = 0; Index < NumVertices; Index++)
 		{
 			FVector Position = (FVector)LODModel->VertexBuffers.PositionVertexBuffer.VertexPosition(Index);
-			const FTransform SliceTransform = SplineMeshComponent->CalcSliceTransform(USplineMeshComponent::GetAxisValue(Position, SplineMeshComponent->ForwardAxis));
-			USplineMeshComponent::GetAxisValue(Position, SplineMeshComponent->ForwardAxis) = 0;
+			const FTransform SliceTransform = SplineMeshComponent->CalcSliceTransform(USplineMeshComponent::GetAxisValueRef(Position, SplineMeshComponent->ForwardAxis));
+			USplineMeshComponent::GetAxisValueRef(Position, SplineMeshComponent->ForwardAxis) = 0;
 			MeshVertices[Index] = SliceTransform.TransformPosition(Position);
 		}
 
diff --git a/Engine/Source/Editor/UnrealEd/Private/Fbx/FbxMainExport.cpp b/Engine/Source/Editor/UnrealEd/Private/Fbx/FbxMainExport.cpp
--- a/Engine/Source/Editor/UnrealEd/Private/Fbx/FbxMainExport.cpp
+++ b/Engine/Source/Editor/UnrealEd/Private/Fbx/FbxMainExport.cpp
@@ -4244,8 +4244,8 @@ void FFbxExporter::ExportSplineMeshToFbx(const USplineMeshComponent* SplineMeshC
 		int32 UnrealPosIndex = UniqueVerts[PosIndex];
 		FVector Position = (FVector)RenderMesh.VertexBuffers.PositionVertexBuffer.VertexPosition(UnrealPosIndex);
 
-		const FTransform SliceTransform = SplineMeshComp->CalcSliceTransform(USplineMeshComponent::GetAxisValue(Position, SplineMeshComp->ForwardAxis));
-		USplineMeshComponent::GetAxisValue(Position, SplineMeshComp->ForwardAxis) = 0;
+		const FTransform SliceTransform = SplineMeshComp->CalcSliceTransform(USplineMeshComponent::GetAxisValueRef(Position, SplineMeshComp->ForwardAxis));
+		USplineMeshComponent::GetAxisValueRef(Position, SplineMeshComp->ForwardAxis) = 0;
 		Position = SliceTransform.TransformPosition(Position);
 
 		ControlPoints[PosIndex] = FbxVector4(Position.X, -Position.Y, Position.Z);
@@ -4291,7 +4291,7 @@ void FFbxExporter::ExportSplineMeshToFbx(const USplineMeshComponent* SplineMeshC
 	for (int32 VertIndex = 0; VertIndex < VertexCount; ++VertIndex)
 	{
 		FVector Position = (FVector)RenderMesh.VertexBuffers.PositionVertexBuffer.VertexPosition(VertIndex);
-		const FTransform SliceTransform = SplineMeshComp->CalcSliceTransform(USplineMeshComponent::GetAxisValue(Position, SplineMeshComp->ForwardAxis));
+		const FTransform SliceTransform = SplineMeshComp->CalcSliceTransform(USplineMeshComponent::GetAxisValueRef(Position, SplineMeshComp->ForwardAxis));
 		FVector Normal = FVector(RenderMesh.VertexBuffers.StaticMeshVertexBuffer.VertexTangentZ(VertIndex));
 		Normal = SliceTransform.TransformVector(Normal);
 		FbxVector4& FbxNormal = FbxNormals[VertIndex];
diff --git a/Engine/Source/Runtime/Engine/Classes/Components/SplineMeshComponent.h b/Engine/Source/Runtime/Engine/Classes/Components/SplineMeshComponent.h
--- a/Engine/Source/Runtime/Engine/Classes/Components/SplineMeshComponent.h
+++ b/Engine/Source/Runtime/Engine/Classes/Components/SplineMeshComponent.h
@@ -362,11 +362,21 @@ public:
 	 */
 	FTransform CalcSliceTransformAtSplineOffset(const float Alpha) const;
 
-	inline static const double& GetAxisValue(const FVector3d& InVector, ESplineMeshAxis::Type InAxis);
-	inline static double& GetAxisValue(FVector3d& InVector, ESplineMeshAxis::Type InAxis);
+	UE_DEPRECATED(5.2, "Use GetAxisValueRef() instead.")
+	static const double& GetAxisValue(const FVector3d& InVector, ESplineMeshAxis::Type InAxis) { return GetAxisValueRef(InVector, InAxis); }
+	UE_DEPRECATED(5.2, "Use GetAxisValueRef() instead.")
+	static double& GetAxisValue(FVector3d& InVector, ESplineMeshAxis::Type InAxis) { return GetAxisValueRef(InVector, InAxis); }
 
-	inline static const float& GetAxisValue(const FVector3f& InVector, ESplineMeshAxis::Type InAxis);
-	inline static float& GetAxisValue(FVector3f& InVector, ESplineMeshAxis::Type InAxis);
+	UE_DEPRECATED(5.2, "Use GetAxisValueRef() instead.")
+	static const float& GetAxisValue(const FVector3f& InVector, ESplineMeshAxis::Type InAxis) { return GetAxisValueRef(InVector, InAxis); }
+	UE_DEPRECATED(5.2, "Use GetAxisValueRef() instead.")
+	static float& GetAxisValue(FVector3f& InVector, ESplineMeshAxis::Type InAxis) { return GetAxisValueRef(InVector, InAxis); }
+
+	inline static const double& GetAxisValueRef(const FVector3d& InVector, ESplineMeshAxis::Type InAxis);
+	inline static double& GetAxisValueRef(FVector3d& InVector, ESplineMeshAxis::Type InAxis);
+
+	inline static const float& GetAxisValueRef(const FVector3f& InVector, ESplineMeshAxis::Type InAxis);
+	inline static float& GetAxisValueRef(FVector3f& InVector, ESplineMeshAxis::Type InAxis);
 
 	/** Returns a vector which, when componentwise-multiplied by another vector, will zero all the components not corresponding to the supplied ESplineMeshAxis */
 	inline static FVector GetAxisMask(ESplineMeshAxis::Type InAxis);
@@ -388,7 +398,7 @@ public:
 	FSplineMeshInstanceData() = default;
 	explicit FSplineMeshInstanceData(const USplineMeshComponent* SourceComponent);
 
-	virtual ~FSplineMeshInstanceData() = default;
+	virtual ~FSplineMeshInstanceData() override = default;
 
 	virtual bool ContainsData() const override
 	{
@@ -414,7 +424,7 @@ public:
 	FVector EndTangent = FVector::ZeroVector;
 };
 
-const double& USplineMeshComponent::GetAxisValue(const FVector3d& InVector, ESplineMeshAxis::Type InAxis)
+const double& USplineMeshComponent::GetAxisValueRef(const FVector3d& InVector, ESplineMeshAxis::Type InAxis)
 {
 	switch (InAxis)
 	{
@@ -430,7 +440,7 @@ const double& USplineMeshComponent::GetAxisValue(const FVector3d& InVector, ESpl
 	}
 }
 
-double& USplineMeshComponent::GetAxisValue(FVector3d& InVector, ESplineMeshAxis::Type InAxis)
+double& USplineMeshComponent::GetAxisValueRef(FVector3d& InVector, ESplineMeshAxis::Type InAxis)
 {
 	switch (InAxis)
 	{
@@ -447,7 +457,7 @@ double& USplineMeshComponent::GetAxisValue(FVector3d& InVector, ESplineMeshAxis:
 }
 
 
-const float& USplineMeshComponent::GetAxisValue(const FVector3f& InVector, ESplineMeshAxis::Type InAxis)
+const float& USplineMeshComponent::GetAxisValueRef(const FVector3f& InVector, ESplineMeshAxis::Type InAxis)
 {
 	switch (InAxis)
 	{
@@ -463,7 +473,7 @@ const float& USplineMeshComponent::GetAxisValue(const FVector3f& InVector, ESpli
 	}
 }
 
-float& USplineMeshComponent::GetAxisValue(FVector3f& InVector, ESplineMeshAxis::Type InAxis)
+float& USplineMeshComponent::GetAxisValueRef(FVector3f& InVector, ESplineMeshAxis::Type InAxis)
 {
 	switch (InAxis)
 	{
diff --git a/Engine/Source/Runtime/Engine/Private/Components/SplineMeshComponent.cpp b/Engine/Source/Runtime/Engine/Private/Components/SplineMeshComponent.cpp
--- a/Engine/Source/Runtime/Engine/Private/Components/SplineMeshComponent.cpp
+++ b/Engine/Source/Runtime/Engine/Private/Components/SplineMeshComponent.cpp
@@ -494,8 +494,8 @@ void USplineMeshComponent::CalculateScaleZAndMinZ(float& OutScaleZ, float& OutMi
 		if (FMath::IsNearlyEqual(SplineBoundaryMin, SplineBoundaryMax))
 		{
 			FBoxSphereBounds StaticMeshBounds = GetStaticMesh()->GetBounds();
-			OutScaleZ = 0.5f / USplineMeshComponent::GetAxisValue(StaticMeshBounds.BoxExtent, ForwardAxis); // 1/(2 * Extent)
-			OutMinZ = USplineMeshComponent::GetAxisValue(StaticMeshBounds.Origin, ForwardAxis) * OutScaleZ - 0.5f;
+			OutScaleZ = 0.5f / USplineMeshComponent::GetAxisValueRef(StaticMeshBounds.BoxExtent, ForwardAxis); // 1/(2 * Extent)
+			OutMinZ = USplineMeshComponent::GetAxisValueRef(StaticMeshBounds.Origin, ForwardAxis) * OutScaleZ - 0.5f;
 		}
 		else
 		{
@@ -721,8 +721,8 @@ FBoxSphereBounds USplineMeshComponent::CalcBounds(const FTransform& LocalToWorld
 	if (bHasCustomBoundary)
 	{
 		// If there's a custom boundary, alter the min/max of the spline we need to evaluate
-		const float MeshMin = GetAxisValue(MeshBounds.Origin - MeshBounds.BoxExtent, ForwardAxis);
-		const float MeshMax = GetAxisValue(MeshBounds.Origin + MeshBounds.BoxExtent, ForwardAxis);
+		const float MeshMin = GetAxisValueRef(MeshBounds.Origin - MeshBounds.BoxExtent, ForwardAxis);
+		const float MeshMax = GetAxisValueRef(MeshBounds.Origin + MeshBounds.BoxExtent, ForwardAxis);
 
 		const float MeshMinT = (MeshMin - SplineBoundaryMin) / (SplineBoundaryMax - SplineBoundaryMin);
 		const float MeshMaxT = (MeshMax - SplineBoundaryMin) / (SplineBoundaryMax - SplineBoundaryMin);
@@ -821,7 +821,7 @@ FTransform USplineMeshComponent::GetSocketTransform(FName InSocketName, ERelativ
 		{
 			FTransform SocketTransform;
 			SocketTransform = FTransform(Socket->RelativeRotation, Socket->RelativeLocation * GetAxisMask(ForwardAxis), Socket->RelativeScale);
-			SocketTransform = SocketTransform * CalcSliceTransform(GetAxisValue(Socket->RelativeLocation, ForwardAxis));
+			SocketTransform = SocketTransform * CalcSliceTransform(GetAxisValueRef(Socket->RelativeLocation, ForwardAxis));
 
 			switch (TransformSpace)
 			{
@@ -862,8 +862,8 @@ FTransform USplineMeshComponent::CalcSliceTransform(const float DistanceAlong) c
 	else if (GetStaticMesh())
 	{
 		const FBoxSphereBounds StaticMeshBounds = GetStaticMesh()->GetBounds();
-		const float MeshMinZ = GetAxisValue(StaticMeshBounds.Origin, ForwardAxis) - GetAxisValue(StaticMeshBounds.BoxExtent, ForwardAxis);
-		const float MeshRangeZ = 2.0f * GetAxisValue(StaticMeshBounds.BoxExtent, ForwardAxis);
+		const float MeshMinZ = GetAxisValueRef(StaticMeshBounds.Origin, ForwardAxis) - GetAxisValueRef(StaticMeshBounds.BoxExtent, ForwardAxis);
+		const float MeshRangeZ = 2.0f * GetAxisValueRef(StaticMeshBounds.BoxExtent, ForwardAxis);
 		if (MeshRangeZ > UE_SMALL_NUMBER)
 		{
 			Alpha = (DistanceAlong - MeshMinZ) / MeshRangeZ;
@@ -933,11 +933,11 @@ bool USplineMeshComponent::GetPhysicsTriMeshData(struct FTriMeshCollisionData* C
 		GetStaticMesh()->GetPhysicsTriMeshData(CollisionData, InUseAllTriData);
 
 		FVector3f Mask = FVector3f(1, 1, 1);
-		GetAxisValue(Mask, ForwardAxis) = 0;
+		GetAxisValueRef(Mask, ForwardAxis) = 0;
 
 		for (FVector3f& CollisionVert : CollisionData->Vertices)
 		{
-			CollisionVert = (FVector3f)CalcSliceTransform(GetAxisValue(CollisionVert, ForwardAxis)).TransformPosition(FVector(CollisionVert * Mask));
+			CollisionVert = (FVector3f)CalcSliceTransform(GetAxisValueRef(CollisionVert, ForwardAxis)).TransformPosition(FVector(CollisionVert * Mask));
 		}
 
 		CollisionData->bDeformableMesh = true;
@@ -1054,7 +1054,7 @@ bool USplineMeshComponent::DoCustomNavigableGeometryExport(FNavigableGeometryExp
 			if (NavCollision->HasConvexGeometry())
 			{
 				FVector Mask = FVector(1, 1, 1);
-				GetAxisValue(Mask, ForwardAxis) = 0;
+				GetAxisValueRef(Mask, ForwardAxis) = 0;
 
 				TArray<FVector> VertexBuffer;
 				VertexBuffer.Reserve(FMath::Max(NavCollision->GetConvexCollision().VertexBuffer.Num(), NavCollision->GetTriMeshCollision().VertexBuffer.Num()));
@@ -1062,7 +1062,7 @@ bool USplineMeshComponent::DoCustomNavigableGeometryExport(FNavigableGeometryExp
 				for (int32 i = 0; i < NavCollision->GetConvexCollision().VertexBuffer.Num(); ++i)
 				{
 					FVector Vertex = NavCollision->GetConvexCollision().VertexBuffer[i];
-					Vertex = CalcSliceTransform(GetAxisValue(Vertex, ForwardAxis)).TransformPosition(Vertex * Mask);
+					Vertex = CalcSliceTransform(GetAxisValueRef(Vertex, ForwardAxis)).TransformPosition(Vertex * Mask);
 					VertexBuffer.Add(Vertex);
 				}
 				GeomExport.ExportCustomMesh(VertexBuffer.GetData(), VertexBuffer.Num(),
@@ -1073,7 +1073,7 @@ bool USplineMeshComponent::DoCustomNavigableGeometryExport(FNavigableGeometryExp
 				for (int32 i = 0; i < NavCollision->GetTriMeshCollision().VertexBuffer.Num(); ++i)
 				{
 					FVector Vertex = NavCollision->GetTriMeshCollision().VertexBuffer[i];
-					Vertex = CalcSliceTransform(GetAxisValue(Vertex, ForwardAxis)).TransformPosition(Vertex * Mask);
+					Vertex = CalcSliceTransform(GetAxisValueRef(Vertex, ForwardAxis)).TransformPosition(Vertex * Mask);
 					VertexBuffer.Add(Vertex);
 				}
 				GeomExport.ExportCustomMesh(VertexBuffer.GetData(), VertexBuffer.Num(),
@@ -1130,12 +1130,12 @@ void USplineMeshComponent::RecreateCollision()
 		else
 		{
 			FVector Mask = FVector(1, 1, 1);
-			GetAxisValue(Mask, ForwardAxis) = 0;
+			GetAxisValueRef(Mask, ForwardAxis) = 0;
 
 			// distortion of a sphere can't be done nicely, so we just transform the origin and size
 			for (FKSphereElem& SphereElem : BodySetup->AggGeom.SphereElems)
 			{
-				const float Z = GetAxisValue(SphereElem.Center, ForwardAxis);
+				const float Z = GetAxisValueRef(SphereElem.Center, ForwardAxis);
 				FTransform SliceTransform = CalcSliceTransform(Z);
 				SphereElem.Center *= Mask;
 
@@ -1146,7 +1146,7 @@ void USplineMeshComponent::RecreateCollision()
 			// distortion of a sphyl can't be done nicely, so we just transform the origin and size
 			for (FKSphylElem& SphylElem : BodySetup->AggGeom.SphylElems)
 			{
-				const float Z = GetAxisValue(SphylElem.Center, ForwardAxis);
+				const float Z = GetAxisValueRef(SphylElem.Center, ForwardAxis);
 				FTransform SliceTransform = CalcSliceTransform(Z);
 				SphylElem.Center *= Mask;
 
@@ -1187,7 +1187,7 @@ void USplineMeshComponent::RecreateCollision()
 					// pretransform the point by its local transform so we are working in untransformed local space
 					FVector TransformedPoint = TM.TransformPosition(Point);
 					// apply the transform to spline space
-					Point = CalcSliceTransform(GetAxisValue(TransformedPoint, ForwardAxis)).TransformPosition(TransformedPoint * Mask);
+					Point = CalcSliceTransform(GetAxisValueRef(TransformedPoint, ForwardAxis)).TransformPosition(TransformedPoint * Mask);
 				}
 
 				// Set the local transform as an identity as points have already been transformed
diff --git a/Engine/Source/Runtime/Landscape/Private/LandscapeSplines.cpp b/Engine/Source/Runtime/Landscape/Private/LandscapeSplines.cpp
--- a/Engine/Source/Runtime/Landscape/Private/LandscapeSplines.cpp
+++ b/Engine/Source/Runtime/Landscape/Private/LandscapeSplines.cpp
@@ -2912,11 +2912,11 @@ void ULandscapeSplineSegment::UpdateSplinePoints(bool bUpdateCollision, bool bUp
 			FVector Scale = MeshEntry->Scale;
 			if (MeshEntry->bScaleToWidth)
 			{
-				Scale *= Width / USplineMeshComponent::GetAxisValue(MeshBounds.BoxExtent, CrossAxis(MeshEntry->ForwardAxis, MeshEntry->UpAxis));
+				Scale *= Width / USplineMeshComponent::GetAxisValueRef(MeshBounds.BoxExtent, CrossAxis(MeshEntry->ForwardAxis, MeshEntry->UpAxis));
 			}
 
-			const float MeshLength = static_cast<float>(FMath::Abs(USplineMeshComponent::GetAxisValue(MeshBounds.BoxExtent, MeshEntry->ForwardAxis) * 2.0 *
-				USplineMeshComponent::GetAxisValue(Scale, MeshEntry->ForwardAxis)));
+			const float MeshLength = static_cast<float>(FMath::Abs(USplineMeshComponent::GetAxisValueRef(MeshBounds.BoxExtent, MeshEntry->ForwardAxis) * 2.0 *
+				USplineMeshComponent::GetAxisValueRef(Scale, MeshEntry->ForwardAxis)));
 			float MeshT = (MeshLength / SplineLength);
 
 			// Improve our approximation if we're not going off the end of the spline
@@ -3048,7 +3048,7 @@ void ULandscapeSplineSegment::UpdateSplinePoints(bool bUpdateCollision, bool bUp
 			FVector Scale = MeshEntry->Scale;
 			if (MeshEntry->bScaleToWidth)
 			{
-				Scale *= Width / USplineMeshComponent::GetAxisValue(MeshBounds.BoxExtent, SideAxis);
+				Scale *= Width / USplineMeshComponent::GetAxisValueRef(MeshBounds.BoxExtent, SideAxis);
 			}
 
 			FVector2D Offset = MeshEntry->CenterAdjust;
@@ -3056,11 +3056,11 @@ void ULandscapeSplineSegment::UpdateSplinePoints(bool bUpdateCollision, bool bUp
 			{
 				if (bDoOrientationRoll)
 				{
-					Offset.Y -= USplineMeshComponent::GetAxisValue(MeshBounds.Origin, SideAxis);
+					Offset.Y -= USplineMeshComponent::GetAxisValueRef(MeshBounds.Origin, SideAxis);
 				}
 				else
 				{
-					Offset.X -= USplineMeshComponent::GetAxisValue(MeshBounds.Origin, SideAxis);
+					Offset.X -= USplineMeshComponent::GetAxisValueRef(MeshBounds.Origin, SideAxis);
 				}
 			}
 
@@ -3098,7 +3098,7 @@ void ULandscapeSplineSegment::UpdateSplinePoints(bool bUpdateCollision, bool bUp
 			FVector ScaleEnd = MeshEntry->Scale;
 			if (MeshEntry->bScaleToWidth)
 			{
-				ScaleEnd *= WidthEnd / USplineMeshComponent::GetAxisValue(MeshBounds.BoxExtent, SideAxis);
+				ScaleEnd *= WidthEnd / USplineMeshComponent::GetAxisValueRef(MeshBounds.BoxExtent, SideAxis);
 			}
 
 			FVector2D OffsetEnd = MeshEntry->CenterAdjust;
@@ -3106,11 +3106,11 @@ void ULandscapeSplineSegment::UpdateSplinePoints(bool bUpdateCollision, bool bUp
 			{
 				if (bDoOrientationRoll)
 				{
-					OffsetEnd.Y -= USplineMeshComponent::GetAxisValue(MeshBounds.Origin, SideAxis);
+					OffsetEnd.Y -= USplineMeshComponent::GetAxisValueRef(MeshBounds.Origin, SideAxis);
 				}
 				else
 				{
-					OffsetEnd.X -= USplineMeshComponent::GetAxisValue(MeshBounds.Origin, SideAxis);
+					OffsetEnd.X -= USplineMeshComponent::GetAxisValueRef(MeshBounds.Origin, SideAxis);
 				}
 			}
 
@@ -3151,7 +3151,7 @@ void ULandscapeSplineSegment::UpdateSplinePoints(bool bUpdateCollision, bool bUp
 				MeshComponent->SplineParams.EndPos   = RelativeTransform.TransformPosition(MeshComponent->SplineParams.EndPos);
 			}
 
-			if (USplineMeshComponent::GetAxisValue(MeshEntry->Scale, MeshEntry->ForwardAxis) < 0)
+			if (USplineMeshComponent::GetAxisValueRef(MeshEntry->Scale, MeshEntry->ForwardAxis) < 0)
 			{
 				Swap(MeshComponent->SplineParams.StartPos, MeshComponent->SplineParams.EndPos);
 				Swap(MeshComponent->SplineParams.StartTangent, MeshComponent->SplineParams.EndTangent);
