Engine dependency cleanup
- MaterialRenderProxy moved out of MaterialShared and into its own header.
- MaterialAttributeDefinitionMap moved out of MaterialShared and into its own header.
- ImmutableSamplerState moved out of RHIDefinitions and into its own header.

diff --git a/Engine/Source/Runtime/Engine/Private/Materials/MaterialAttributeDefinitionMap.cpp b/Engine/Source/Runtime/Engine/Private/Materials/MaterialAttributeDefinitionMap.cpp
new file mode 100644
--- /dev/null
+++ b/Engine/Source/Runtime/Engine/Private/Materials/MaterialAttributeDefinitionMap.cpp
@@ -0,0 +1,544 @@
+// Copyright Epic Games, Inc. All Rights Reserved.
+
+#include "Materials/MaterialAttributeDefinitionMap.h"
+
+#include "Engine/EngineTypes.h"
+#include "MaterialCompiler.h"
+#include "MaterialDomain.h"
+#include "Materials/Material.h"
+#include "Materials/MaterialExpressionShadingModel.h"
+#include "Materials/MaterialExpressionStrata.h"
+#include "SceneTypes.h"
+
+#define LOCTEXT_NAMESPACE "MaterialShared"
+
+FMaterialAttributeDefintion::FMaterialAttributeDefintion(
+	const FGuid& InAttributeID, const FString& InAttributeName, EMaterialProperty InProperty,
+	EMaterialValueType InValueType, const FVector4& InDefaultValue, EShaderFrequency InShaderFrequency,
+	int32 InTexCoordIndex /*= INDEX_NONE*/, bool bInIsHidden /*= false*/, MaterialAttributeBlendFunction InBlendFunction /*= nullptr*/)
+	: AttributeID(InAttributeID)
+	, DefaultValue(InDefaultValue)
+	, AttributeName(InAttributeName)
+	, Property(InProperty)
+	, ValueType(InValueType)
+	, ShaderFrequency(InShaderFrequency)
+	, TexCoordIndex(InTexCoordIndex)
+	, BlendFunction(InBlendFunction)
+	, bIsHidden(bInIsHidden)
+{
+	checkf(ValueType & MCT_Float || ValueType == MCT_ShadingModel || ValueType == MCT_Strata || ValueType == MCT_MaterialAttributes, TEXT("Unsupported material attribute type %d"), ValueType);
+}
+
+int32 FMaterialAttributeDefintion::CompileDefaultValue(FMaterialCompiler* Compiler)
+{
+	int32 Ret;
+
+	// TODO: Temporarily preserving hack from 4.13 to change default value for two-sided foliage model 
+	if (Property == MP_SubsurfaceColor && Compiler->GetCompiledShadingModels().HasShadingModel(MSM_TwoSidedFoliage))
+	{
+		check(ValueType == MCT_Float3);
+		return Compiler->Constant3(0, 0, 0);
+	}
+
+	if (Property == MP_ShadingModel)
+	{
+		check(ValueType == MCT_ShadingModel);
+		// Default to the first shading model of the material. If the material is using a single shading model selected through the dropdown, this is how it gets written to the shader as a constant (optimizing out all the dynamic branches)
+		return Compiler->ShadingModel(Compiler->GetMaterialShadingModels().GetFirstShadingModel());
+	}
+
+	if (Property == MP_FrontMaterial)
+	{
+		check(ValueType == MCT_Strata);
+		return Compiler->StrataCreateAndRegisterNullMaterial();
+	}
+
+	if (TexCoordIndex == INDEX_NONE)
+	{
+		// Standard value type
+		switch (ValueType)
+		{
+		case MCT_Float:
+		case MCT_Float1: Ret = Compiler->Constant(DefaultValue.X); break;
+		case MCT_Float2: Ret = Compiler->Constant2(DefaultValue.X, DefaultValue.Y); break;
+		case MCT_Float3: Ret = Compiler->Constant3(DefaultValue.X, DefaultValue.Y, DefaultValue.Z); break;
+		default: Ret = Compiler->Constant4(DefaultValue.X, DefaultValue.Y, DefaultValue.Z, DefaultValue.W);
+		}
+	}
+	else
+	{
+		// Texture coordinates allow pass through for default	
+		Ret = Compiler->TextureCoordinate(TexCoordIndex, false, false);
+	}
+
+	return Ret;
+}
+
+//////////////////////////////////////////////////////////////////////////
+
+FMaterialCustomOutputAttributeDefintion::FMaterialCustomOutputAttributeDefintion(
+	const FGuid& InAttributeID, const FString& InAttributeName, const FString& InFunctionName, EMaterialProperty InProperty,
+	EMaterialValueType InValueType, const FVector4& InDefaultValue, EShaderFrequency InShaderFrequency, MaterialAttributeBlendFunction InBlendFunction /*= nullptr*/)
+	: FMaterialAttributeDefintion(InAttributeID, InAttributeName, InProperty, InValueType, InDefaultValue, InShaderFrequency, INDEX_NONE, false, InBlendFunction)
+	, FunctionName(InFunctionName)
+{
+}
+
+//////////////////////////////////////////////////////////////////////////
+
+FMaterialAttributeDefinitionMap FMaterialAttributeDefinitionMap::GMaterialPropertyAttributesMap;
+
+FMaterialAttributeDefinitionMap::FMaterialAttributeDefinitionMap()
+	: AttributeDDCString(TEXT(""))
+	, bIsInitialized(false)
+{
+	AttributeMap.Empty(MP_MAX);
+	InitializeAttributeMap();
+}
+
+/** Compiles the default expression for a material attribute */
+int32 FMaterialAttributeDefinitionMap::CompileDefaultExpression(FMaterialCompiler* Compiler, EMaterialProperty Property)
+{
+	FMaterialAttributeDefintion* Attribute = GMaterialPropertyAttributesMap.Find(Property);
+	return Attribute->CompileDefaultValue(Compiler);
+}
+
+/** Compiles the default expression for a material attribute */
+int32 FMaterialAttributeDefinitionMap::CompileDefaultExpression(FMaterialCompiler* Compiler, const FGuid& AttributeID)
+{
+	FMaterialAttributeDefintion* Attribute = GMaterialPropertyAttributesMap.Find(AttributeID);
+	return Attribute->CompileDefaultValue(Compiler);
+}
+
+/** Returns the display name of a material attribute */
+const FString& FMaterialAttributeDefinitionMap::GetAttributeName(EMaterialProperty Property)
+{
+	FMaterialAttributeDefintion* Attribute = GMaterialPropertyAttributesMap.Find(Property);
+	return Attribute->AttributeName;
+}
+
+/** Returns the display name of a material attribute */
+const FString& FMaterialAttributeDefinitionMap::GetAttributeName(const FGuid& AttributeID)
+{
+	FMaterialAttributeDefintion* Attribute = GMaterialPropertyAttributesMap.Find(AttributeID);
+	return Attribute->AttributeName;
+}
+
+/** Returns the display name of a material attribute, accounting for overrides based on properties of a given material */
+FText FMaterialAttributeDefinitionMap::GetDisplayNameForMaterial(EMaterialProperty Property, UMaterial* Material)
+{
+	if (!Material)
+	{
+		return FText::FromString(GetAttributeName(Property));
+	}
+
+	FMaterialAttributeDefintion* Attribute = GMaterialPropertyAttributesMap.Find(Property);
+	return GetAttributeOverrideForMaterial(Attribute->AttributeID, Material);
+}
+
+/** Returns the display name of a material attribute, accounting for overrides based on properties of a given material */
+FText FMaterialAttributeDefinitionMap::GetDisplayNameForMaterial(const FGuid& AttributeID, UMaterial* Material)
+{
+	if (!Material)
+	{
+		return FText::FromString(GetAttributeName(AttributeID));
+	}
+
+	FMaterialAttributeDefintion* Attribute = GMaterialPropertyAttributesMap.Find(AttributeID);
+	return GetAttributeOverrideForMaterial(AttributeID, Material);
+}
+
+/** Returns the value type of a material attribute */
+EMaterialValueType FMaterialAttributeDefinitionMap::GetValueType(EMaterialProperty Property)
+{
+	FMaterialAttributeDefintion* Attribute = GMaterialPropertyAttributesMap.Find(Property);
+	return Attribute->ValueType;
+}
+
+/** Returns the value type of a material attribute */
+EMaterialValueType FMaterialAttributeDefinitionMap::GetValueType(const FGuid& AttributeID)
+{
+	FMaterialAttributeDefintion* Attribute = GMaterialPropertyAttributesMap.Find(AttributeID);
+	return Attribute->ValueType;
+}
+
+/** Returns the default value of a material property */
+FVector4f FMaterialAttributeDefinitionMap::GetDefaultValue(EMaterialProperty Property)
+{
+	FMaterialAttributeDefintion* Attribute = GMaterialPropertyAttributesMap.Find(Property);
+	return (FVector4f)Attribute->DefaultValue;
+}
+
+/** Returns the default value of a material attribute */
+FVector4f FMaterialAttributeDefinitionMap::GetDefaultValue(const FGuid& AttributeID)
+{
+	FMaterialAttributeDefintion* Attribute = GMaterialPropertyAttributesMap.Find(AttributeID);
+	return (FVector4f)Attribute->DefaultValue;
+}
+
+/** Returns the shader frequency of a material attribute */
+EShaderFrequency FMaterialAttributeDefinitionMap::GetShaderFrequency(EMaterialProperty Property)
+{
+	FMaterialAttributeDefintion* Attribute = GMaterialPropertyAttributesMap.Find(Property);
+	return Attribute->ShaderFrequency;
+}
+
+/** Returns the shader frequency of a material attribute */
+EShaderFrequency FMaterialAttributeDefinitionMap::GetShaderFrequency(const FGuid& AttributeID)
+{
+	FMaterialAttributeDefintion* Attribute = GMaterialPropertyAttributesMap.Find(AttributeID);
+	return Attribute->ShaderFrequency;
+}
+
+/** Returns the attribute ID for a matching material property */
+FGuid FMaterialAttributeDefinitionMap::GetID(EMaterialProperty Property)
+{
+	FMaterialAttributeDefintion* Attribute = GMaterialPropertyAttributesMap.Find(Property);
+	return Attribute->AttributeID;
+}
+
+/** Returns a the material property matching the specified attribute AttributeID */
+EMaterialProperty FMaterialAttributeDefinitionMap::GetProperty(const FGuid& AttributeID)
+{
+	if (FMaterialAttributeDefintion* Attribute = GMaterialPropertyAttributesMap.Find(AttributeID))
+	{
+		return Attribute->Property;
+	}
+	return MP_MAX;
+}
+
+/** Returns the custom blend function of a material attribute */
+MaterialAttributeBlendFunction FMaterialAttributeDefinitionMap::GetBlendFunction(const FGuid& AttributeID)
+{
+	FMaterialAttributeDefintion* Attribute = GMaterialPropertyAttributesMap.Find(AttributeID);
+	return Attribute->BlendFunction;
+}
+
+/** Returns a default attribute AttributeID */
+FGuid FMaterialAttributeDefinitionMap::GetDefaultID()
+{
+	return GMaterialPropertyAttributesMap.Find(MP_MAX)->AttributeID;
+}
+
+void FMaterialAttributeDefinitionMap::InitializeAttributeMap()
+{
+	check(!bIsInitialized);
+	bIsInitialized = true;
+	const bool bHideAttribute = true;
+	LLM_SCOPE(ELLMTag::Materials);
+
+	// All types plus default/missing attribute
+	AttributeMap.Empty(MP_MAX + 1);
+
+	// Basic attributes
+	Add(FGuid(0x69B8D336, 0x16ED4D49, 0x9AA49729, 0x2F050F7A), TEXT("BaseColor"), MP_BaseColor, MCT_Float3, FVector4(0, 0, 0, 0), SF_Pixel);
+	Add(FGuid(0x57C3A161, 0x7F064296, 0xB00B24A5, 0xA496F34C), TEXT("Metallic"), MP_Metallic, MCT_Float, FVector4(0, 0, 0, 0), SF_Pixel);
+	Add(FGuid(0x9FDAB399, 0x25564CC9, 0x8CD2D572, 0xC12C8FED), TEXT("Specular"), MP_Specular, MCT_Float, FVector4(.5, 0, 0, 0), SF_Pixel);
+	Add(FGuid(0xD1DD967C, 0x4CAD47D3, 0x9E6346FB, 0x08ECF210), TEXT("Roughness"), MP_Roughness, MCT_Float, FVector4(.5, 0, 0, 0), SF_Pixel);
+	Add(FGuid(0x55E2B4FB, 0xC1C54DB2, 0x9F11875F, 0x7231EB1E), TEXT("Anisotropy"), MP_Anisotropy, MCT_Float, FVector4(0, 0, 0, 0), SF_Pixel);
+	Add(FGuid(0xB769B54D, 0xD08D4440, 0xABC21BA6, 0xCD27D0E2), TEXT("EmissiveColor"), MP_EmissiveColor, MCT_Float3, FVector4(0, 0, 0, 0), SF_Pixel);
+	Add(FGuid(0xB8F50FBA, 0x2A754EC1, 0x9EF672CF, 0xEB27BF51), TEXT("Opacity"), MP_Opacity, MCT_Float, FVector4(1, 0, 0, 0), SF_Pixel);
+	Add(FGuid(0x679FFB17, 0x2BB5422C, 0xAD520483, 0x166E0C75), TEXT("OpacityMask"), MP_OpacityMask, MCT_Float, FVector4(1, 0, 0, 0), SF_Pixel);
+	Add(FGuid(0x0FA2821A, 0x200F4A4A, 0xB719B789, 0xC1259C64), TEXT("Normal"), MP_Normal, MCT_Float3, FVector4(0, 0, 1, 0), SF_Pixel);
+	Add(FGuid(0xD5F8E9CF, 0xCDC3468D, 0xB10E4465, 0x596A7BBA), TEXT("Tangent"), MP_Tangent, MCT_Float3, FVector4(1, 0, 0, 0), SF_Pixel);
+
+	// Advanced attributes
+	Add(FGuid(0xF905F895, 0xD5814314, 0x916D2434, 0x8C40CE9E), TEXT("WorldPositionOffset"), MP_WorldPositionOffset, MCT_Float3, FVector4(0, 0, 0, 0), SF_Vertex);
+	Add(FGuid(0x5B8FC679, 0x51CE4082, 0x9D777BEE, 0xF4F72C44), TEXT("SubsurfaceColor"), MP_SubsurfaceColor, MCT_Float3, FVector4(1, 1, 1, 0), SF_Pixel);
+	Add(FGuid(0x9E502E69, 0x3C8F48FA, 0x94645CFD, 0x28E5428D), TEXT("ClearCoat"), MP_CustomData0, MCT_Float, FVector4(1, 0, 0, 0), SF_Pixel);
+	Add(FGuid(0xBE4F2FFD, 0x12FC4296, 0xB0124EEA, 0x12C28D92), TEXT("ClearCoatRoughness"), MP_CustomData1, MCT_Float, FVector4(.1, 0, 0, 0), SF_Pixel);
+	Add(FGuid(0xE8EBD0AD, 0xB1654CBE, 0xB079C3A8, 0xB39B9F15), TEXT("AmbientOcclusion"), MP_AmbientOcclusion, MCT_Float, FVector4(1, 0, 0, 0), SF_Pixel);
+	Add(FGuid(0xD0B0FA03, 0x14D74455, 0xA851BAC5, 0x81A0788B), TEXT("Refraction"), MP_Refraction, MCT_Float2, FVector4(1, 0, 0, 0), SF_Pixel);
+	Add(FGuid(0x0AC97EC3, 0xE3D047BA, 0xB610167D, 0xC4D919FF), TEXT("PixelDepthOffset"), MP_PixelDepthOffset, MCT_Float, FVector4(0, 0, 0, 0), SF_Pixel);
+	Add(FGuid(0xD9423FFF, 0xD77E4D82, 0x8FF9CF5E, 0x055D1255), TEXT("ShadingModel"), MP_ShadingModel, MCT_ShadingModel, FVector4(0, 0, 0, 0), SF_Pixel, INDEX_NONE, false, &CompileShadingModelBlendFunction);
+	Add(FGuid(0x5973A03E, 0x13A74E08, 0x92D0CEDD, 0xF2936CF8), TEXT("FrontMaterial"), MP_FrontMaterial, MCT_Strata, FVector4(0, 0, 0, 0), SF_Pixel, INDEX_NONE, false, &CompileStrataBlendFunction);
+
+	// Used when compiling material with execution pins, which are compiling all attributes together
+	Add(FGuid(0xE0ED040B, 0x82794D93, 0xBD2D59B2, 0xA5BBF41C), TEXT("MaterialAttributes"), MP_MaterialAttributes, MCT_MaterialAttributes, FVector4(0, 0, 0, 0), SF_Pixel, INDEX_NONE, bHideAttribute);
+
+	// Texture coordinates
+	Add(FGuid(0xD30EC284, 0xE13A4160, 0x87BB5230, 0x2ED115DC), TEXT("CustomizedUV0"), MP_CustomizedUVs0, MCT_Float2, FVector4(0, 0, 0, 0), SF_Vertex, 0);
+	Add(FGuid(0xC67B093C, 0x2A5249AA, 0xABC97ADE, 0x4A1F49C5), TEXT("CustomizedUV1"), MP_CustomizedUVs1, MCT_Float2, FVector4(0, 0, 0, 0), SF_Vertex, 1);
+	Add(FGuid(0x85C15B24, 0xF3E047CA, 0x85856872, 0x01AE0F4F), TEXT("CustomizedUV2"), MP_CustomizedUVs2, MCT_Float2, FVector4(0, 0, 0, 0), SF_Vertex, 2);
+	Add(FGuid(0x777819DC, 0x31AE4676, 0xB864EF77, 0xB807E873), TEXT("CustomizedUV3"), MP_CustomizedUVs3, MCT_Float2, FVector4(0, 0, 0, 0), SF_Vertex, 3);
+	Add(FGuid(0xDA63B233, 0xDDF44CAD, 0xB93D867B, 0x8DAFDBCC), TEXT("CustomizedUV4"), MP_CustomizedUVs4, MCT_Float2, FVector4(0, 0, 0, 0), SF_Vertex, 4);
+	Add(FGuid(0xC2F52B76, 0x4A034388, 0x89119528, 0x2071B190), TEXT("CustomizedUV5"), MP_CustomizedUVs5, MCT_Float2, FVector4(0, 0, 0, 0), SF_Vertex, 5);
+	Add(FGuid(0x8214A8CA, 0x0CB944CF, 0x9DFD78DB, 0xE48BB55F), TEXT("CustomizedUV6"), MP_CustomizedUVs6, MCT_Float2, FVector4(0, 0, 0, 0), SF_Vertex, 6);
+	Add(FGuid(0xD8F8D01F, 0xC6F74715, 0xA3CFB4FF, 0x9EF51FAC), TEXT("CustomizedUV7"), MP_CustomizedUVs7, MCT_Float2, FVector4(0, 0, 0, 0), SF_Vertex, 7);
+
+	// Lightmass attributes	
+	Add(FGuid(0x68934E1B, 0x70EB411B, 0x86DF5AA5, 0xDF2F626C), TEXT("DiffuseColor"), MP_DiffuseColor, MCT_Float3, FVector4(0, 0, 0, 0), SF_Pixel, INDEX_NONE, bHideAttribute);
+	Add(FGuid(0xE89CBD84, 0x62EA48BE, 0x80F88521, 0x2B0C403C), TEXT("SpecularColor"), MP_SpecularColor, MCT_Float3, FVector4(0, 0, 0, 0), SF_Pixel, INDEX_NONE, bHideAttribute);
+
+	// Debug attributes
+	Add(FGuid(0x5BF6BA94, 0xA3264629, 0xA253A05B, 0x0EABBB86), TEXT("Missing"), MP_MAX, MCT_Float, FVector4(0, 0, 0, 0), SF_Pixel, INDEX_NONE, bHideAttribute);
+
+	// Removed attributes
+	Add(FGuid(0x2091ECA2, 0xB59248EE, 0x8E2CD578, 0xD371926D), TEXT("WorldDisplacement"), MP_WorldDisplacement_DEPRECATED, MCT_Float3, FVector4(0, 0, 0, 0), SF_Vertex, INDEX_NONE, bHideAttribute);
+	Add(FGuid(0xA0119D44, 0xC456450D, 0x9C39C933, 0x1F72D8D1), TEXT("TessellationMultiplier"), MP_TessellationMultiplier_DEPRECATED, MCT_Float, FVector4(1, 0, 0, 0), SF_Vertex, INDEX_NONE, bHideAttribute);
+
+	// UMaterialExpression custom outputs
+	AddCustomAttribute(FGuid(0xfbd7b46e, 0xb1234824, 0xbde76b23, 0x609f984c), "BentNormal", "GetBentNormal", MCT_Float3, FVector4(0, 0, 1, 0));
+	AddCustomAttribute(FGuid(0xAA3D5C04, 0x16294716, 0xBBDEC869, 0x6A27DD72), "ClearCoatBottomNormal", "ClearCoatBottomNormal", MCT_Float3, FVector4(0, 0, 1, 0));
+	AddCustomAttribute(FGuid(0x8EAB2CB2, 0x73634A24, 0x8CD14F47, 0x3F9C8E55), "CustomEyeTangent", "GetTangentOutput", MCT_Float3, FVector4(0, 0, 0, 0));
+}
+
+void FMaterialAttributeDefinitionMap::Add(const FGuid& AttributeID, const FString& AttributeName, EMaterialProperty Property,
+	EMaterialValueType ValueType, const FVector4& DefaultValue, EShaderFrequency ShaderFrequency,
+	int32 TexCoordIndex /*= INDEX_NONE*/, bool bIsHidden /*= false*/, MaterialAttributeBlendFunction BlendFunction /*= nullptr*/)
+{
+	checkf(!AttributeMap.Contains(Property), TEXT("Tried to add duplicate material property."));
+	AttributeMap.Add(Property, FMaterialAttributeDefintion(AttributeID, AttributeName, Property, ValueType, DefaultValue, ShaderFrequency, TexCoordIndex, bIsHidden, BlendFunction));
+	if (!bIsHidden)
+	{
+		OrderedVisibleAttributeList.Add(AttributeID);
+	}
+}
+
+FMaterialAttributeDefintion* FMaterialAttributeDefinitionMap::Find(const FGuid& AttributeID)
+{
+	for (auto& Attribute : CustomAttributes)
+	{
+		if (Attribute.AttributeID == AttributeID)
+		{
+			return &Attribute;
+		}
+	}
+
+	for (auto& Attribute : AttributeMap)
+	{
+		if (Attribute.Value.AttributeID == AttributeID)
+		{
+			return &Attribute.Value;
+		}
+	}
+
+	UE_LOG(LogMaterial, Warning, TEXT("Failed to find material attribute, AttributeID: %s."), *AttributeID.ToString(EGuidFormats::Digits));
+	return Find(MP_MAX);
+}
+
+FMaterialAttributeDefintion* FMaterialAttributeDefinitionMap::Find(EMaterialProperty Property)
+{
+	if (FMaterialAttributeDefintion* Attribute = AttributeMap.Find(Property))
+	{
+		return Attribute;
+	}
+
+	UE_LOG(LogMaterial, Warning, TEXT("Failed to find material attribute, PropertyType: %i."), (uint32)Property);
+	return Find(MP_MAX);
+}
+
+FText FMaterialAttributeDefinitionMap::GetAttributeOverrideForMaterial(const FGuid& AttributeID, UMaterial* Material)
+{
+	TArray<TKeyValuePair<EMaterialShadingModel, FString>> CustomPinNames;
+	EMaterialProperty Property = GMaterialPropertyAttributesMap.Find(AttributeID)->Property;
+
+	switch (Property)
+	{
+	case MP_EmissiveColor:
+		return Material->IsUIMaterial() ? LOCTEXT("UIOutputColor", "Final Color") : LOCTEXT("EmissiveColor", "Emissive Color");
+	case MP_Opacity:
+		return LOCTEXT("Opacity", "Opacity");
+	case MP_OpacityMask:
+		return LOCTEXT("OpacityMask", "Opacity Mask");
+	case MP_DiffuseColor:
+		return LOCTEXT("DiffuseColor", "Diffuse Color");
+	case MP_SpecularColor:
+		return LOCTEXT("SpecularColor", "Specular Color");
+	case MP_BaseColor:
+		return Material->MaterialDomain == MD_Volume ? LOCTEXT("Albedo", "Albedo") : LOCTEXT("BaseColor", "Base Color");
+	case MP_Metallic:
+		CustomPinNames.Add({ MSM_Hair, "Scatter" });
+		CustomPinNames.Add({ MSM_Eye, "Curvature" });
+		return FText::FromString(GetPinNameFromShadingModelField(Material->GetShadingModels(), CustomPinNames, "Metallic"));
+	case MP_Specular:
+		return LOCTEXT("Specular", "Specular");
+	case MP_Roughness:
+		return LOCTEXT("Roughness", "Roughness");
+	case MP_Anisotropy:
+		return LOCTEXT("Anisotropy", "Anisotropy");
+	case MP_Normal:
+		CustomPinNames.Add({ MSM_Hair, "Tangent" });
+		return FText::FromString(GetPinNameFromShadingModelField(Material->GetShadingModels(), CustomPinNames, "Normal"));
+	case MP_Tangent:
+		return LOCTEXT("Tangent", "Tangent");
+	case MP_WorldPositionOffset:
+		return Material->IsUIMaterial() ? LOCTEXT("ScreenPosition", "Screen Position") : LOCTEXT("WorldPositionOffset", "World Position Offset");
+	case MP_WorldDisplacement_DEPRECATED:
+		return LOCTEXT("WorldDisplacement", "World Displacement");
+	case MP_TessellationMultiplier_DEPRECATED:
+		return LOCTEXT("TessellationMultiplier", "Tessellation Multiplier");
+	case MP_SubsurfaceColor:
+		if (Material->MaterialDomain == MD_Volume)
+		{
+			return LOCTEXT("Extinction", "Extinction");
+		}
+		CustomPinNames.Add({ MSM_Cloth, "Fuzz Color" });
+		return FText::FromString(GetPinNameFromShadingModelField(Material->GetShadingModels(), CustomPinNames, "Subsurface Color"));
+	case MP_CustomData0:
+		CustomPinNames.Add({ MSM_ClearCoat, "Clear Coat" });
+		CustomPinNames.Add({ MSM_Hair, "Backlit" });
+		CustomPinNames.Add({ MSM_Cloth, "Cloth" });
+		CustomPinNames.Add({ MSM_Eye, "Iris Mask" });
+		CustomPinNames.Add({ MSM_SubsurfaceProfile, "Curvature" });
+		return FText::FromString(GetPinNameFromShadingModelField(Material->GetShadingModels(), CustomPinNames, "Custom Data 0"));
+	case MP_CustomData1:
+		CustomPinNames.Add({ MSM_ClearCoat, "Clear Coat Roughness" });
+		CustomPinNames.Add({ MSM_Eye, "Iris Distance" });
+		return FText::FromString(GetPinNameFromShadingModelField(Material->GetShadingModels(), CustomPinNames, "Custom Data 1"));
+	case MP_AmbientOcclusion:
+		return LOCTEXT("AmbientOcclusion", "Ambient Occlusion");
+	case MP_Refraction:
+		return LOCTEXT("Refraction", "Refraction");
+	case MP_CustomizedUVs0:
+		return LOCTEXT("CustomizedUV0", "Customized UV 0");
+	case MP_CustomizedUVs1:
+		return LOCTEXT("CustomizedUV1", "Customized UV 1");
+	case MP_CustomizedUVs2:
+		return LOCTEXT("CustomizedUV2", "Customized UV 2");
+	case MP_CustomizedUVs3:
+		return LOCTEXT("CustomizedUV3", "Customized UV 3");
+	case MP_CustomizedUVs4:
+		return LOCTEXT("CustomizedUV4", "Customized UV 4");
+	case MP_CustomizedUVs5:
+		return LOCTEXT("CustomizedUV5", "Customized UV 5");
+	case MP_CustomizedUVs6:
+		return LOCTEXT("CustomizedUV6", "Customized UV 6");
+	case MP_CustomizedUVs7:
+		return LOCTEXT("CustomizedUV7", "Customized UV 7");
+	case MP_PixelDepthOffset:
+		return LOCTEXT("PixelDepthOffset", "Pixel Depth Offset");
+	case MP_ShadingModel:
+		return LOCTEXT("ShadingModel", "Shading Model");
+	case MP_FrontMaterial:
+		return LOCTEXT("FrontMaterial", "Front Material");
+	case MP_CustomOutput:
+		return FText::FromString(GetAttributeName(AttributeID));
+
+	}
+	return  LOCTEXT("Missing", "Missing");
+}
+
+FString FMaterialAttributeDefinitionMap::GetPinNameFromShadingModelField(FMaterialShadingModelField InShadingModels, const TArray<TKeyValuePair<EMaterialShadingModel, FString>>& InCustomShadingModelPinNames, const FString& InDefaultPinName)
+{
+	FString OutPinName;
+	for (const TKeyValuePair<EMaterialShadingModel, FString>& CustomShadingModelPinName : InCustomShadingModelPinNames)
+	{
+		if (InShadingModels.HasShadingModel(CustomShadingModelPinName.Key))
+		{
+			// Add delimiter
+			if (!OutPinName.IsEmpty())
+			{
+				OutPinName.Append(" or ");
+			}
+
+			// Append the name and remove the shading model from the temp field
+			OutPinName.Append(CustomShadingModelPinName.Value);
+			InShadingModels.RemoveShadingModel(CustomShadingModelPinName.Key);
+		}
+	}
+
+	// There are other shading models present, these don't have their own specific name for this pin, so use a default one
+	if (InShadingModels.CountShadingModels() != 0)
+	{
+		// Add delimiter
+		if (!OutPinName.IsEmpty())
+		{
+			OutPinName.Append(" or ");
+		}
+
+		OutPinName.Append(InDefaultPinName);
+	}
+
+	ensure(!OutPinName.IsEmpty());
+	return OutPinName;
+}
+
+void FMaterialAttributeDefinitionMap::AppendDDCKeyString(FString& String)
+{
+	FString& DDCString = GMaterialPropertyAttributesMap.AttributeDDCString;
+
+	if (DDCString.Len() == 0)
+	{
+		FString AttributeIDs;
+
+		for (const auto& Attribute : GMaterialPropertyAttributesMap.AttributeMap)
+		{
+			AttributeIDs += Attribute.Value.AttributeID.ToString(EGuidFormats::Digits);
+		}
+
+		for (const auto& Attribute : GMaterialPropertyAttributesMap.CustomAttributes)
+		{
+			AttributeIDs += Attribute.AttributeID.ToString(EGuidFormats::Digits);
+		}
+
+		FSHA1 HashState;
+		HashState.UpdateWithString(*AttributeIDs, AttributeIDs.Len());
+		HashState.Final();
+
+		FSHAHash Hash;
+		HashState.GetHash(&Hash.Hash[0]);
+		DDCString = Hash.ToString();
+	}
+	else
+	{
+		// TODO: In debug force re-generate DDC string and compare to catch invalid runtime changes
+	}
+
+	String.Append(DDCString);
+}
+
+void FMaterialAttributeDefinitionMap::AddCustomAttribute(const FGuid& AttributeID, const FString& AttributeName, const FString& FunctionName, EMaterialValueType ValueType, const FVector4& DefaultValue, MaterialAttributeBlendFunction BlendFunction /*= nullptr*/)
+{
+	// Make sure that we init CustomAttributes before DDCString is initialized (before first shader load)
+	check(GMaterialPropertyAttributesMap.AttributeDDCString.Len() == 0);
+
+	FMaterialCustomOutputAttributeDefintion UserAttribute(AttributeID, AttributeName, FunctionName, MP_CustomOutput, ValueType, DefaultValue, SF_Pixel, BlendFunction);
+#if DO_CHECK
+	for (auto& Attribute : GMaterialPropertyAttributesMap.AttributeMap)
+	{
+		checkf(Attribute.Value.AttributeID != AttributeID, TEXT("Tried to add duplicate custom output attribute (%s) already in base attributes (%s)."), *AttributeName, *(Attribute.Value.AttributeName));
+	}
+	checkf(!GMaterialPropertyAttributesMap.CustomAttributes.Contains(UserAttribute), TEXT("Tried to add duplicate custom output attribute (%s)."), *AttributeName);
+#endif
+	GMaterialPropertyAttributesMap.CustomAttributes.Add(UserAttribute);
+
+	if (!UserAttribute.bIsHidden)
+	{
+		GMaterialPropertyAttributesMap.OrderedVisibleAttributeList.Add(AttributeID);
+	}
+}
+
+FGuid FMaterialAttributeDefinitionMap::GetCustomAttributeID(const FString& AttributeName)
+{
+	for (auto& Attribute : GMaterialPropertyAttributesMap.CustomAttributes)
+	{
+		if (Attribute.AttributeName == AttributeName)
+		{
+			return Attribute.AttributeID;
+		}
+	}
+
+	return GMaterialPropertyAttributesMap.Find(MP_MAX)->AttributeID;
+}
+
+void FMaterialAttributeDefinitionMap::GetCustomAttributeList(TArray<FMaterialCustomOutputAttributeDefintion>& CustomAttributeList)
+{
+	CustomAttributeList.Empty(GMaterialPropertyAttributesMap.CustomAttributes.Num());
+	for (auto& Attribute : GMaterialPropertyAttributesMap.CustomAttributes)
+	{
+		CustomAttributeList.Add(Attribute);
+	}
+}
+
+void FMaterialAttributeDefinitionMap::GetAttributeNameToIDList(TArray<TPair<FString, FGuid>>& NameToIDList)
+{
+	NameToIDList.Empty(GMaterialPropertyAttributesMap.OrderedVisibleAttributeList.Num());
+	for (const FGuid& AttributeID : GMaterialPropertyAttributesMap.OrderedVisibleAttributeList)
+	{
+		FMaterialAttributeDefintion* Attribute = GMaterialPropertyAttributesMap.Find(AttributeID);
+		NameToIDList.Emplace(Attribute->AttributeName, AttributeID);
+	}
+}
+
+#undef LOCTEXT_NAMESPACE
diff --git a/Engine/Source/Runtime/Engine/Private/Materials/MaterialRenderProxy.cpp b/Engine/Source/Runtime/Engine/Private/Materials/MaterialRenderProxy.cpp
new file mode 100644
--- /dev/null
+++ b/Engine/Source/Runtime/Engine/Private/Materials/MaterialRenderProxy.cpp
@@ -0,0 +1,805 @@
+// Copyright Epic Games, Inc. All Rights Reserved.
+
+#include "Materials/MaterialRenderProxy.h"
+#include "Async/TaskGraphFwd.h"
+#include "Engine/Texture.h"
+#include "EngineModule.h"
+#include "ExternalTexture.h"
+#include "HAL/IConsoleManager.h"
+#include "Materials/Material.h"
+#include "Materials/MaterialInterface.h"
+#include "MaterialDomain.h"
+#include "MaterialShared.h"
+#include "RendererInterface.h"
+#include "RenderingThread.h"
+#include "RHI.h"
+#include "TextureResource.h"
+#include "VT/RuntimeVirtualTexture.h"
+
+#if WITH_EDITOR
+#include "Rendering/StaticLightingSystemInterface.h"
+#endif
+
+int32 GDeferUniformExpressionCaching = 1;
+FAutoConsoleVariableRef CVarDeferUniformExpressionCaching(
+	TEXT("r.DeferUniformExpressionCaching"),
+	GDeferUniformExpressionCaching,
+	TEXT("Whether to defer caching of uniform expressions until a rendering command needs them up to date.  Deferring updates is more efficient because multiple SetVectorParameterValue calls in a frame will only result in one update."),
+	ECVF_RenderThreadSafe
+);
+
+int32 GUniformExpressionCacheAsyncUpdates = 1;
+FAutoConsoleVariableRef CVarUniformExpressionCacheAsyncUpdates(
+	TEXT("r.UniformExpressionCacheAsyncUpdates"),
+	GUniformExpressionCacheAsyncUpdates,
+	TEXT("Whether to allow async updates of uniform expression caches."),
+	ECVF_RenderThreadSafe);
+
+FUniformExpressionCache::~FUniformExpressionCache()
+{
+	ResetAllocatedVTs();
+	UniformBuffer.SafeRelease();
+}
+
+class FUniformExpressionCacheAsyncUpdateTask
+{
+public:
+	void Begin()
+	{
+		ReferenceCount++;
+	}
+
+	void End()
+	{
+		check(ReferenceCount > 0);
+		if (--ReferenceCount == 0)
+		{
+			Wait();
+		}
+	}
+
+	bool IsEnabled() const
+	{
+		return ReferenceCount > 0 && GUniformExpressionCacheAsyncUpdates > 0 && !GRHICommandList.Bypass();
+	}
+
+	void SetTask(const FGraphEventRef& InTask)
+	{
+		check(!Task);
+		check(IsEnabled());
+		Task = InTask;
+	}
+
+	void Wait()
+	{
+		if (Task)
+		{
+			Task->Wait();
+			Task = nullptr;
+		}
+	}
+
+private:
+	FGraphEventRef Task;
+	int32 ReferenceCount = 0;
+
+} GUniformExpressionCacheAsyncUpdateTask;
+
+FUniformExpressionCacheAsyncUpdateScope::FUniformExpressionCacheAsyncUpdateScope()
+{
+	ENQUEUE_RENDER_COMMAND(BeginAsyncUniformExpressionCacheUpdates)(
+		[](FRHICommandList&)
+		{
+			GUniformExpressionCacheAsyncUpdateTask.Begin();
+			FMaterialRenderProxy::UpdateDeferredCachedUniformExpressions();
+		});
+}
+
+FUniformExpressionCacheAsyncUpdateScope::~FUniformExpressionCacheAsyncUpdateScope()
+{
+	ENQUEUE_RENDER_COMMAND(EndAsyncUniformExpressionCacheUpdates)(
+		[](FRHICommandList&)
+		{
+			GUniformExpressionCacheAsyncUpdateTask.End();
+		});
+}
+
+void FUniformExpressionCacheAsyncUpdateScope::WaitForTask()
+{
+	GUniformExpressionCacheAsyncUpdateTask.Wait();
+}
+
+class FUniformExpressionCacheAsyncUpdater
+{
+public:
+	void Add(FUniformExpressionCache* UniformExpressionCache, const FUniformExpressionSet* UniformExpressionSet, const FRHIUniformBufferLayout* UniformBufferLayout, const FMaterialRenderContext& Context)
+	{
+		Items.Emplace(UniformExpressionCache, UniformExpressionSet, UniformBufferLayout, Context);
+	}
+
+	void Update(FRHICommandListImmediate& RHICmdListImmediate)
+	{
+		if (Items.IsEmpty())
+		{
+			return;
+		}
+
+		GUniformExpressionCacheAsyncUpdateTask.Wait();
+
+		FRHICommandList* RHICmdList = new FRHICommandList(FRHIGPUMask::All());
+
+		FGraphEventRef Event = FFunctionGraphTask::CreateAndDispatchWhenReady(
+			[Items = MoveTemp(Items), RHICmdList]
+			{
+				TRACE_CPUPROFILER_EVENT_SCOPE(FUniformExpressionCacheAsyncUpdater::Update);
+				FTaskTagScope Scope(ETaskTag::EParallelRenderingThread);
+				FMemMark Mark(FMemStack::Get());
+
+				for (const FItem& Item : Items)
+				{
+					uint8* TempBuffer = FMemStack::Get().PushBytes(Item.UniformBufferLayout->ConstantBufferSize, SHADER_PARAMETER_STRUCT_ALIGNMENT);
+
+					FMaterialRenderContext Context(Item.MaterialRenderProxy, *Item.Material, Item.bShowSelection);
+
+					Item.UniformExpressionSet->FillUniformBuffer(Context, Item.AllocatedVTs, Item.UniformBufferLayout, TempBuffer, Item.UniformBufferLayout->ConstantBufferSize);
+
+					RHICmdList->UpdateUniformBuffer(Item.UniformBuffer, TempBuffer);
+				}
+
+				RHICmdList->FinishRecording();
+
+			}, TStatId(), nullptr, ENamedThreads::AnyHiPriThreadHiPriTask);
+
+		RHICmdListImmediate.QueueAsyncCommandListSubmit(RHICmdList);
+
+		GUniformExpressionCacheAsyncUpdateTask.SetTask(Event);
+	}
+
+private:
+	struct FItem
+	{
+		FItem() = default;
+		FItem(FUniformExpressionCache* InUniformExpressionCache, const FUniformExpressionSet* InUniformExpressionSet, const FRHIUniformBufferLayout* InUniformBufferLayout, const FMaterialRenderContext& Context)
+			: UniformBuffer(InUniformExpressionCache->UniformBuffer)
+			, AllocatedVTs(InUniformExpressionCache->AllocatedVTs)
+			, UniformExpressionSet(InUniformExpressionSet)
+			, UniformBufferLayout(InUniformBufferLayout)
+			, MaterialRenderProxy(Context.MaterialRenderProxy)
+			, Material(&Context.Material)
+			, bShowSelection(Context.bShowSelection)
+		{}
+
+		TRefCountPtr<FRHIUniformBuffer> UniformBuffer;
+		TArray<IAllocatedVirtualTexture*, FConcurrentLinearArrayAllocator> AllocatedVTs;
+		const FUniformExpressionSet* UniformExpressionSet = nullptr;
+		const FRHIUniformBufferLayout* UniformBufferLayout = nullptr;
+		const FMaterialRenderProxy* MaterialRenderProxy = nullptr;
+		const FMaterial* Material = nullptr;
+		bool bShowSelection = false;
+	};
+
+	TArray<FItem, FConcurrentLinearArrayAllocator> Items;
+};
+
+void FUniformExpressionCache::ResetAllocatedVTs()
+{
+	for (int32 i = 0; i < OwnedAllocatedVTs.Num(); ++i)
+	{
+		GetRendererModule().DestroyVirtualTexture(OwnedAllocatedVTs[i]);
+	}
+	AllocatedVTs.Reset();
+	OwnedAllocatedVTs.Reset();
+}
+
+
+bool FMaterialRenderProxy::GetVectorValue(const FHashedMaterialParameterInfo& ParameterInfo, FLinearColor* OutValue, const FMaterialRenderContext& Context) const
+{
+	FMaterialParameterValue Value;
+	if (GetParameterValue(EMaterialParameterType::Vector, ParameterInfo, Value, Context))
+	{
+		*OutValue = Value.AsLinearColor();
+		return true;
+	}
+	return false;
+}
+
+bool FMaterialRenderProxy::GetScalarValue(const FHashedMaterialParameterInfo& ParameterInfo, float* OutValue, const FMaterialRenderContext& Context) const
+{
+	FMaterialParameterValue Value;
+	if (GetParameterValue(EMaterialParameterType::Scalar, ParameterInfo, Value, Context))
+	{
+		*OutValue = Value.AsScalar();
+		return true;
+	}
+	return false;
+}
+
+bool FMaterialRenderProxy::GetTextureValue(const FHashedMaterialParameterInfo& ParameterInfo, const UTexture** OutValue, const FMaterialRenderContext& Context) const
+{
+	FMaterialParameterValue Value;
+	if (GetParameterValue(EMaterialParameterType::Texture, ParameterInfo, Value, Context))
+	{
+		*OutValue = Value.Texture;
+		return true;
+	}
+	return false;
+}
+
+bool FMaterialRenderProxy::GetTextureValue(const FHashedMaterialParameterInfo& ParameterInfo, const URuntimeVirtualTexture** OutValue, const FMaterialRenderContext& Context) const
+{
+	FMaterialParameterValue Value;
+	if (GetParameterValue(EMaterialParameterType::RuntimeVirtualTexture, ParameterInfo, Value, Context))
+	{
+		*OutValue = Value.RuntimeVirtualTexture;
+		return true;
+	}
+	return false;
+}
+
+static void OnVirtualTextureDestroyedCB(const FVirtualTextureProducerHandle& InHandle, void* Baton)
+{
+	FMaterialRenderProxy* MaterialProxy = static_cast<FMaterialRenderProxy*>(Baton);
+
+	MaterialProxy->InvalidateUniformExpressionCache(false);
+	UMaterialInterface::IterateOverActiveFeatureLevels([&](ERHIFeatureLevel::Type InFeatureLevel)
+		{
+			MaterialProxy->UpdateUniformExpressionCacheIfNeeded(InFeatureLevel);
+		});
+}
+
+IAllocatedVirtualTexture* FMaterialRenderProxy::GetPreallocatedVTStack(const FMaterialRenderContext& Context, const FUniformExpressionSet& UniformExpressionSet, const FMaterialVirtualTextureStack& VTStack) const
+{
+	check(VTStack.IsPreallocatedStack());
+
+	const URuntimeVirtualTexture* Texture;
+	VTStack.GetTextureValue(Context, UniformExpressionSet, Texture);
+
+	if (Texture == nullptr)
+	{
+		return nullptr;
+	}
+
+	GetRendererModule().AddVirtualTextureProducerDestroyedCallback(Texture->GetProducerHandle(), &OnVirtualTextureDestroyedCB, const_cast<FMaterialRenderProxy*>(this));
+	HasVirtualTextureCallbacks = true;
+
+	return Texture->GetAllocatedVirtualTexture();
+}
+
+IAllocatedVirtualTexture* FMaterialRenderProxy::AllocateVTStack(const FMaterialRenderContext& Context, const FUniformExpressionSet& UniformExpressionSet, const FMaterialVirtualTextureStack& VTStack) const
+{
+	check(!VTStack.IsPreallocatedStack());
+	const uint32 NumLayers = VTStack.GetNumLayers();
+	if (NumLayers == 0u)
+	{
+		return nullptr;
+	}
+
+	const UTexture* LayerTextures[VIRTUALTEXTURE_SPACE_MAXLAYERS] = { nullptr };
+	VTStack.GetTextureValues(Context, UniformExpressionSet, LayerTextures);
+
+	const UMaterialInterface* MaterialInterface = GetMaterialInterface();
+
+	FAllocatedVTDescription VTDesc;
+	if (MaterialInterface)
+	{
+		VTDesc.Name = MaterialInterface->GetFName();
+	}
+	VTDesc.Dimensions = 2;
+	VTDesc.NumTextureLayers = NumLayers;
+	bool bFoundValidLayer = false;
+	for (uint32 LayerIndex = 0u; LayerIndex < NumLayers; ++LayerIndex)
+	{
+		const UTexture* Texture = LayerTextures[LayerIndex];
+		if (!Texture)
+		{
+			continue;
+		}
+
+		// GetResource() is safe to call from the render thread.
+		const FTextureResource* TextureResource = Texture->GetResource();
+		if (TextureResource)
+		{
+			const FVirtualTexture2DResource* VirtualTextureResourceForLayer = TextureResource->GetVirtualTexture2DResource();
+
+			if (VirtualTextureResourceForLayer == nullptr)
+			{
+				// The placeholder used during async texture compilation is expected to be of the wrong type since
+				// no VT infos are available until later in the compilation process. This will be resolved
+				// once the final texture resource is available.
+#if WITH_EDITOR
+				if (!TextureResource->IsProxy())
+#endif
+				{
+					UE_LOG(LogMaterial, Warning, TEXT("Material '%s' expects texture '%s' to be Virtual"),
+						*GetFriendlyName(), *Texture->GetName());
+				}
+				continue;
+			}
+			else
+			{
+				// All tile sizes need to match
+				check(!bFoundValidLayer || VTDesc.TileSize == VirtualTextureResourceForLayer->GetTileSize());
+				check(!bFoundValidLayer || VTDesc.TileBorderSize == VirtualTextureResourceForLayer->GetBorderSize());
+
+				const FVirtualTextureProducerHandle& ProducerHandle = VirtualTextureResourceForLayer->GetProducerHandle();
+				if (ProducerHandle.IsValid())
+				{
+					VTDesc.TileSize = VirtualTextureResourceForLayer->GetTileSize();
+					VTDesc.TileBorderSize = VirtualTextureResourceForLayer->GetBorderSize();
+					VTDesc.ProducerHandle[LayerIndex] = ProducerHandle;
+					VTDesc.ProducerLayerIndex[LayerIndex] = 0u;
+					GetRendererModule().AddVirtualTextureProducerDestroyedCallback(ProducerHandle, &OnVirtualTextureDestroyedCB, const_cast<FMaterialRenderProxy*>(this));
+					bFoundValidLayer = true;
+				}
+			}
+		}
+	}
+
+	if (bFoundValidLayer)
+	{
+		HasVirtualTextureCallbacks = true;
+		return GetRendererModule().AllocateVirtualTexture(VTDesc);
+	}
+	return nullptr;
+}
+
+
+void FMaterialRenderProxy::EvaluateUniformExpressions(FUniformExpressionCache& OutUniformExpressionCache, const FMaterialRenderContext& Context, FUniformExpressionCacheAsyncUpdater* Updater) const
+{
+	check(IsInRenderingThread());
+
+	SCOPE_CYCLE_COUNTER(STAT_CacheUniformExpressions);
+
+	TRACE_CPUPROFILER_EVENT_SCOPE(FMaterialRenderProxy::EvaluateUniformExpressions);
+
+	// Retrieve the material's uniform expression set.
+	FMaterialShaderMap* ShaderMap = Context.Material.GetRenderingThreadShaderMap();
+	const FUniformExpressionSet& UniformExpressionSet = ShaderMap->GetUniformExpressionSet();
+
+	OutUniformExpressionCache.CachedUniformExpressionShaderMap = ShaderMap;
+	OutUniformExpressionCache.ResetAllocatedVTs();
+	OutUniformExpressionCache.AllocatedVTs.Empty(UniformExpressionSet.VTStacks.Num());
+	OutUniformExpressionCache.OwnedAllocatedVTs.Empty(UniformExpressionSet.VTStacks.Num());
+
+	if (HasVirtualTextureCallbacks)
+	{
+		GetRendererModule().RemoveAllVirtualTextureProducerDestroyedCallbacks(this);
+		HasVirtualTextureCallbacks = false;
+	}
+
+	for (int32 i = 0; i < UniformExpressionSet.VTStacks.Num(); ++i)
+	{
+		const FMaterialVirtualTextureStack& VTStack = UniformExpressionSet.VTStacks[i];
+		IAllocatedVirtualTexture* AllocatedVT = nullptr;
+		if (VTStack.IsPreallocatedStack())
+		{
+			AllocatedVT = GetPreallocatedVTStack(Context, UniformExpressionSet, VTStack);
+		}
+		else
+		{
+			AllocatedVT = AllocateVTStack(Context, UniformExpressionSet, VTStack);
+			if (AllocatedVT != nullptr)
+			{
+				OutUniformExpressionCache.OwnedAllocatedVTs.Add(AllocatedVT);
+			}
+		}
+		OutUniformExpressionCache.AllocatedVTs.Add(AllocatedVT);
+	}
+
+	const FRHIUniformBufferLayout* UniformBufferLayout = ShaderMap->GetUniformBufferLayout();
+
+	if (IsValidRef(OutUniformExpressionCache.UniformBuffer))
+	{
+		if (!OutUniformExpressionCache.UniformBuffer->IsValid())
+		{
+			UE_LOG(LogMaterial, Fatal, TEXT("The Uniformbuffer needs to be valid if it has been set"));
+		}
+
+		// The actual pointer may not match because there are cases (in the editor, during the shader compilation) when material's shader map gets updated without proxy's cache
+		// getting invalidated, but the layout contents must match.
+		check(OutUniformExpressionCache.UniformBuffer->GetLayoutPtr() == UniformBufferLayout || *OutUniformExpressionCache.UniformBuffer->GetLayoutPtr() == *UniformBufferLayout);
+	}
+
+	if (Updater)
+	{
+		if (!IsValidRef(OutUniformExpressionCache.UniformBuffer))
+		{
+			OutUniformExpressionCache.UniformBuffer = RHICreateUniformBuffer(nullptr, UniformBufferLayout, UniformBuffer_MultiFrame);
+		}
+
+		Updater->Add(&OutUniformExpressionCache, &UniformExpressionSet, UniformBufferLayout, Context);
+	}
+	else
+	{
+		FMemMark Mark(FMemStack::Get());
+		uint8* TempBuffer = FMemStack::Get().PushBytes(UniformBufferLayout->ConstantBufferSize, SHADER_PARAMETER_STRUCT_ALIGNMENT);
+		UniformExpressionSet.FillUniformBuffer(Context, OutUniformExpressionCache, UniformBufferLayout, TempBuffer, UniformBufferLayout->ConstantBufferSize);
+
+		if (IsValidRef(OutUniformExpressionCache.UniformBuffer))
+		{
+			RHIUpdateUniformBuffer(OutUniformExpressionCache.UniformBuffer, TempBuffer);
+		}
+		else
+		{
+			OutUniformExpressionCache.UniformBuffer = RHICreateUniformBuffer(TempBuffer, UniformBufferLayout, UniformBuffer_MultiFrame);
+		}
+	}
+
+	OutUniformExpressionCache.ParameterCollections = UniformExpressionSet.ParameterCollections;
+
+	++UniformExpressionCacheSerialNumber;
+}
+
+void FMaterialRenderProxy::CacheUniformExpressions(bool bRecreateUniformBuffer)
+{
+	// Register the render proxy's as a render resource so it can receive notifications to free the uniform buffer.
+	InitResource();
+
+	bool bUsingNewLoader = EVENT_DRIVEN_ASYNC_LOAD_ACTIVE_AT_RUNTIME && GEventDrivenLoaderEnabled;
+
+	check((bUsingNewLoader && GIsInitialLoad) || // The EDL at boot time maybe not load the default materials first; we need to intialize materials before the default materials are done
+		UMaterial::GetDefaultMaterial(MD_Surface));
+
+
+	if (IsMarkedForGarbageCollection())
+	{
+		UE_LOG(LogMaterial, Fatal, TEXT("Cannot queue the Expression Cache for Material %s when it is about to be deleted"), *MaterialName);
+	}
+	StartCacheUniformExpressions();
+	DeferredUniformExpressionCacheRequests.Add(this);
+
+	InvalidateUniformExpressionCache(bRecreateUniformBuffer);
+
+	if (!GDeferUniformExpressionCaching)
+	{
+		FMaterialRenderProxy::UpdateDeferredCachedUniformExpressions();
+	}
+}
+
+void FMaterialRenderProxy::CacheUniformExpressions_GameThread(bool bRecreateUniformBuffer)
+{
+	if (FApp::CanEverRender())
+	{
+		UE_LOG(LogMaterial, VeryVerbose, TEXT("Caching uniform expressions for material: %s"), *GetFriendlyName());
+
+		FMaterialRenderProxy* RenderProxy = this;
+		ENQUEUE_RENDER_COMMAND(FCacheUniformExpressionsCommand)(
+			[RenderProxy, bRecreateUniformBuffer](FRHICommandListImmediate& RHICmdList)
+			{
+				RenderProxy->CacheUniformExpressions(bRecreateUniformBuffer);
+			});
+	}
+}
+
+void FMaterialRenderProxy::InvalidateUniformExpressionCache(bool bRecreateUniformBuffer)
+{
+	check(IsInRenderingThread());
+	GUniformExpressionCacheAsyncUpdateTask.Wait();
+
+#if WITH_EDITOR
+	FStaticLightingSystemInterface::OnMaterialInvalidated.Broadcast(this);
+#endif
+
+	if (HasVirtualTextureCallbacks)
+	{
+		GetRendererModule().RemoveAllVirtualTextureProducerDestroyedCallbacks(this);
+		HasVirtualTextureCallbacks = false;
+	}
+
+	++UniformExpressionCacheSerialNumber;
+	for (int32 i = 0; i < ERHIFeatureLevel::Num; ++i)
+	{
+		UniformExpressionCache[i].CachedUniformExpressionShaderMap = nullptr;
+		UniformExpressionCache[i].ResetAllocatedVTs();
+
+		if (bRecreateUniformBuffer)
+		{
+			// This is required if the FMaterial is being recompiled (the uniform buffer layout will change).
+			// This should only be done if the calling code is using FMaterialUpdateContext to recreate the rendering state of primitives using this material, 
+			// Since cached mesh commands also cache uniform buffer pointers.
+			UniformExpressionCache[i].UniformBuffer = nullptr;
+		}
+	}
+}
+
+void FMaterialRenderProxy::UpdateUniformExpressionCacheIfNeeded(ERHIFeatureLevel::Type InFeatureLevel) const
+{
+	// Don't cache uniform expressions if an entirely different FMaterialRenderProxy is going to be used for rendering
+	const FMaterial* Material = GetMaterialNoFallback(InFeatureLevel);
+	// Note: We would actually need to compare the FMaterialShaderMapId of both shader maps but a simple pointer compare also works
+	// because shader maps are currently swapped out whenever they are modified.
+	if (Material && Material->GetRenderingThreadShaderMap() != UniformExpressionCache[InFeatureLevel].CachedUniformExpressionShaderMap)
+	{
+		FMaterialRenderContext MaterialRenderContext(this, *Material, nullptr);
+		MaterialRenderContext.bShowSelection = GIsEditor;
+		EvaluateUniformExpressions(UniformExpressionCache[InFeatureLevel], MaterialRenderContext);
+	}
+}
+
+FMaterialRenderProxy::FMaterialRenderProxy(FString InMaterialName)
+	: SubsurfaceProfileRT(0)
+	, MaterialName(MoveTemp(InMaterialName))
+	, MarkedForGarbageCollection(0)
+	, DeletedFlag(0)
+	, HasVirtualTextureCallbacks(0)
+{
+}
+
+FMaterialRenderProxy::~FMaterialRenderProxy()
+{
+	// We only wait on deletions happening on the render thread. Async deletions can happen during scene render shutdown and those are waited on explicitly.
+	if (IsInRenderingThread())
+	{
+		GUniformExpressionCacheAsyncUpdateTask.Wait();
+	}
+
+	if (IsInitialized())
+	{
+		check(IsInRenderingThread());
+		ReleaseResource();
+	}
+
+	if (HasVirtualTextureCallbacks)
+	{
+		check(IsInRenderingThread());
+		GetRendererModule().RemoveAllVirtualTextureProducerDestroyedCallbacks(this);
+		HasVirtualTextureCallbacks = false;
+	}
+
+	DeletedFlag = 1;
+}
+
+void FMaterialRenderProxy::InitDynamicRHI()
+{
+#if WITH_EDITOR
+	// MaterialRenderProxyMap is only used by shader compiling
+	if (!FPlatformProperties::RequiresCookedData())
+	{
+		FScopeLock Locker(&MaterialRenderProxyMapLock);
+		FMaterialRenderProxy::MaterialRenderProxyMap.Add(this);
+	}
+#endif // WITH_EDITOR
+}
+
+void FMaterialRenderProxy::CancelCacheUniformExpressions()
+{
+	DeferredUniformExpressionCacheRequests.Remove(this);
+}
+
+void FMaterialRenderProxy::ReleaseDynamicRHI()
+{
+#if WITH_EDITOR
+	if (!FPlatformProperties::RequiresCookedData())
+	{
+		FScopeLock Locker(&MaterialRenderProxyMapLock);
+		FMaterialRenderProxy::MaterialRenderProxyMap.Remove(this);
+	}
+#endif // WITH_EDITOR
+
+	if (DeferredUniformExpressionCacheRequests.Remove(this))
+	{
+		// Notify that we're finished with this inflight cache request, because the object is being released
+		FinishCacheUniformExpressions();
+	}
+
+	InvalidateUniformExpressionCache(true);
+
+	FExternalTextureRegistry::Get().RemoveMaterialRenderProxyReference(this);
+}
+
+void FMaterialRenderProxy::ReleaseResource()
+{
+	ReleaseResourceFlag = true;
+	FRenderResource::ReleaseResource();
+	if (HasVirtualTextureCallbacks)
+	{
+		GetRendererModule().RemoveAllVirtualTextureProducerDestroyedCallbacks(this);
+		HasVirtualTextureCallbacks = false;
+	}
+}
+
+
+const FMaterial& FMaterialRenderProxy::GetMaterialWithFallback(ERHIFeatureLevel::Type InFeatureLevel, const FMaterialRenderProxy*& OutFallbackMaterialRenderProxy) const
+{
+	const FMaterial* BaseMaterial = GetMaterialNoFallback(InFeatureLevel);
+	const FMaterial* Material = BaseMaterial;
+	if (!Material || !Material->IsRenderingThreadShaderMapComplete())
+	{
+		const FMaterialRenderProxy* FallbackMaterialProxy = this;
+		do
+		{
+			FallbackMaterialProxy = FallbackMaterialProxy->GetFallback(InFeatureLevel);
+			check(FallbackMaterialProxy);
+			Material = FallbackMaterialProxy->GetMaterialNoFallback(InFeatureLevel);
+		} while (!Material || !Material->IsRenderingThreadShaderMapComplete());
+		OutFallbackMaterialRenderProxy = FallbackMaterialProxy;
+
+		if (BaseMaterial)
+		{
+			BaseMaterial->SubmitCompileJobs_RenderThread(EShaderCompileJobPriority::Normal);
+		}
+	}
+	return *Material;
+}
+
+const FMaterial& FMaterialRenderProxy::GetIncompleteMaterialWithFallback(ERHIFeatureLevel::Type InFeatureLevel) const
+{
+	const FMaterial* Material = GetMaterialNoFallback(InFeatureLevel);
+	if (!Material)
+	{
+		const FMaterialRenderProxy* FallbackMaterialProxy = this;
+		do
+		{
+			FallbackMaterialProxy = FallbackMaterialProxy->GetFallback(InFeatureLevel);
+			check(FallbackMaterialProxy);
+			Material = FallbackMaterialProxy->GetMaterialNoFallback(InFeatureLevel);
+		} while (!Material);
+	}
+	return *Material;
+}
+
+void FMaterialRenderProxy::UpdateDeferredCachedUniformExpressions()
+{
+	LLM_SCOPE(ELLMTag::Materials);
+
+	check(IsInRenderingThread());
+
+	TRACE_CPUPROFILER_EVENT_SCOPE(FMaterialRenderProxy::UpdateDeferredCachedUniformExpressions);
+	CSV_SCOPED_TIMING_STAT_EXCLUSIVE(Material_UpdateDeferredCachedUniformExpressions);
+	QUICK_SCOPE_CYCLE_COUNTER(STAT_UpdateDeferredCachedUniformExpressions);
+
+	FUniformExpressionCacheAsyncUpdater Updater;
+	FUniformExpressionCacheAsyncUpdater* UpdaterIfEnabled = GUniformExpressionCacheAsyncUpdateTask.IsEnabled() ? &Updater : nullptr;
+
+	for (TSet<FMaterialRenderProxy*>::TConstIterator It(DeferredUniformExpressionCacheRequests); It; ++It)
+	{
+		FMaterialRenderProxy* MaterialProxy = *It;
+		if (MaterialProxy->IsDeleted())
+		{
+			UE_LOG(LogMaterial, Fatal, TEXT("FMaterialRenderProxy deleted and GC mark was: %i"), MaterialProxy->IsMarkedForGarbageCollection());
+		}
+
+		UMaterialInterface::IterateOverActiveFeatureLevels([&](ERHIFeatureLevel::Type InFeatureLevel)
+			{
+				// Don't bother caching if we'll be falling back to a different FMaterialRenderProxy for rendering anyway
+				const FMaterial* Material = MaterialProxy->GetMaterialNoFallback(InFeatureLevel);
+		if (Material && Material->GetRenderingThreadShaderMap())
+		{
+			FMaterialRenderContext MaterialRenderContext(MaterialProxy, *Material, nullptr);
+			MaterialRenderContext.bShowSelection = GIsEditor;
+			MaterialProxy->EvaluateUniformExpressions(MaterialProxy->UniformExpressionCache[(int32)InFeatureLevel], MaterialRenderContext, UpdaterIfEnabled);
+		}
+			});
+
+		MaterialProxy->FinishCacheUniformExpressions();
+	}
+
+	if (UpdaterIfEnabled)
+	{
+		Updater.Update(FRHICommandListExecutor::GetImmediateCommandList());
+	}
+
+	DeferredUniformExpressionCacheRequests.Reset();
+}
+
+#if WITH_EDITOR
+TSet<FMaterialRenderProxy*> FMaterialRenderProxy::MaterialRenderProxyMap;
+FCriticalSection FMaterialRenderProxy::MaterialRenderProxyMapLock;
+#endif // WITH_EDITOR
+TSet<FMaterialRenderProxy*> FMaterialRenderProxy::DeferredUniformExpressionCacheRequests;
+
+/*-----------------------------------------------------------------------------
+	FColoredMaterialRenderProxy
+-----------------------------------------------------------------------------*/
+
+FColoredMaterialRenderProxy::FColoredMaterialRenderProxy(const FMaterialRenderProxy* InParent, const FLinearColor& InColor, FName InColorParamName)
+	: FMaterialRenderProxy(InParent ? InParent->GetMaterialName() : TEXT("FColoredMaterialRenderProxy"))
+	, Parent(InParent)
+	, Color(InColor)
+	, ColorParamName(InColorParamName)
+{
+}
+
+const FMaterial* FColoredMaterialRenderProxy::GetMaterialNoFallback(ERHIFeatureLevel::Type InFeatureLevel) const
+{
+	return Parent->GetMaterialNoFallback(InFeatureLevel);
+}
+
+const FMaterialRenderProxy* FColoredMaterialRenderProxy::GetFallback(ERHIFeatureLevel::Type InFeatureLevel) const
+{
+	return Parent->GetFallback(InFeatureLevel);
+}
+
+bool FColoredMaterialRenderProxy::GetParameterValue(EMaterialParameterType Type, const FHashedMaterialParameterInfo& ParameterInfo, FMaterialParameterValue& OutValue, const FMaterialRenderContext& Context) const
+{
+	if (Type == EMaterialParameterType::Vector && ParameterInfo.Name == ColorParamName)
+	{
+		OutValue = Color;
+		return true;
+	}
+	else
+	{
+		return Parent->GetParameterValue(Type, ParameterInfo, OutValue, Context);
+	}
+}
+
+/*-----------------------------------------------------------------------------
+	FColoredTexturedMaterialRenderProxy
+-----------------------------------------------------------------------------*/
+
+FColoredTexturedMaterialRenderProxy::FColoredTexturedMaterialRenderProxy(const FMaterialRenderProxy* InParent, const FLinearColor& InColor, FName InColorParamName, const UTexture* InTexture, FName InTextureParamName)
+	: FColoredMaterialRenderProxy(InParent, InColor, InColorParamName)
+	, Texture(InTexture)
+	, TextureParamName(InTextureParamName)
+{
+}
+
+bool FColoredTexturedMaterialRenderProxy::GetParameterValue(EMaterialParameterType Type, const FHashedMaterialParameterInfo& ParameterInfo, FMaterialParameterValue& OutValue, const FMaterialRenderContext& Context) const
+{
+	if (Type == EMaterialParameterType::Texture && ParameterInfo.Name == TextureParamName)
+	{
+		OutValue = Texture;
+		return true;
+	}
+	else
+	{
+		return Parent->GetParameterValue(Type, ParameterInfo, OutValue, Context);
+	}
+}
+
+/*-----------------------------------------------------------------------------
+	FOverrideSelectionColorMaterialRenderProxy
+-----------------------------------------------------------------------------*/
+
+FOverrideSelectionColorMaterialRenderProxy::FOverrideSelectionColorMaterialRenderProxy(const FMaterialRenderProxy* InParent, const FLinearColor& InSelectionColor)
+	: FMaterialRenderProxy(InParent ? InParent->GetMaterialName() : TEXT("FOverrideSelectionColorMaterialRenderProxy"))
+	, Parent(InParent)
+	, SelectionColor(FLinearColor(InSelectionColor.R, InSelectionColor.G, InSelectionColor.B, 1))
+{
+}
+
+const FMaterial* FOverrideSelectionColorMaterialRenderProxy::GetMaterialNoFallback(ERHIFeatureLevel::Type InFeatureLevel) const
+{
+	return Parent->GetMaterialNoFallback(InFeatureLevel);
+}
+
+const FMaterialRenderProxy* FOverrideSelectionColorMaterialRenderProxy::GetFallback(ERHIFeatureLevel::Type InFeatureLevel) const
+{
+	return Parent->GetFallback(InFeatureLevel);
+}
+
+bool FOverrideSelectionColorMaterialRenderProxy::GetParameterValue(EMaterialParameterType Type, const FHashedMaterialParameterInfo& ParameterInfo, FMaterialParameterValue& OutValue, const FMaterialRenderContext& Context) const
+{
+	if (Type == EMaterialParameterType::Vector && ParameterInfo.Name == FName(NAME_SelectionColor))
+	{
+		OutValue = SelectionColor;
+		return true;
+	}
+	else
+	{
+		return Parent->GetParameterValue(Type, ParameterInfo, OutValue, Context);
+	}
+}
+
+/*-----------------------------------------------------------------------------
+	FLightingDensityMaterialRenderProxy
+-----------------------------------------------------------------------------*/
+
+FLightingDensityMaterialRenderProxy::FLightingDensityMaterialRenderProxy(const FMaterialRenderProxy* InParent, const FLinearColor& InColor, const FVector2D& InLightmapResolution)
+	: FColoredMaterialRenderProxy(InParent, InColor)
+	, LightmapResolution(InLightmapResolution)
+{
+}
+
+static FName NAME_LightmapRes = FName(TEXT("LightmapRes"));
+
+bool FLightingDensityMaterialRenderProxy::GetParameterValue(EMaterialParameterType Type, const FHashedMaterialParameterInfo& ParameterInfo, FMaterialParameterValue& OutValue, const FMaterialRenderContext& Context) const
+{
+	if (Type == EMaterialParameterType::Vector && ParameterInfo.Name == NAME_LightmapRes)
+	{
+		OutValue = FLinearColor(LightmapResolution.X, LightmapResolution.Y, 0.0f, 0.0f);
+		return true;
+	}
+	return FColoredMaterialRenderProxy::GetParameterValue(Type, ParameterInfo, OutValue, Context);
+}
+
diff --git a/Engine/Source/Runtime/Engine/Private/Materials/MaterialShared.cpp b/Engine/Source/Runtime/Engine/Private/Materials/MaterialShared.cpp
--- a/Engine/Source/Runtime/Engine/Private/Materials/MaterialShared.cpp
+++ b/Engine/Source/Runtime/Engine/Private/Materials/MaterialShared.cpp
@@ -16,8 +16,10 @@
 #include "Materials/MaterialInterface.h"
 #include "MaterialExpressionIO.h"
 #include "Materials/Material.h"
+#include "Materials/MaterialAttributeDefinitionMap.h"
 #include "Materials/MaterialInstanceBasePropertyOverrides.h"
 #include "Materials/MaterialInstance.h"
+#include "Materials/MaterialRenderProxy.h"
 #include "Materials/MaterialUniformExpressions.h"
 #include "Materials/MaterialShaderMapLayout.h"
 #include "UObject/UObjectIterator.h"
@@ -85,21 +87,6 @@ IMPLEMENT_TYPE_LAYOUT(FMaterialTextureParameterInfo);
 IMPLEMENT_TYPE_LAYOUT(FMaterialExternalTextureParameterInfo);
 IMPLEMENT_TYPE_LAYOUT(FMaterialVirtualTextureStack);
 
-int32 GDeferUniformExpressionCaching = 1;
-FAutoConsoleVariableRef CVarDeferUniformExpressionCaching(
-	TEXT("r.DeferUniformExpressionCaching"),
-	GDeferUniformExpressionCaching,
-	TEXT("Whether to defer caching of uniform expressions until a rendering command needs them up to date.  Deferring updates is more efficient because multiple SetVectorParameterValue calls in a frame will only result in one update."),
-	ECVF_RenderThreadSafe
-	);
-
-int32 GUniformExpressionCacheAsyncUpdates = 1;
-FAutoConsoleVariableRef CVarUniformExpressionCacheAsyncUpdates(
-	TEXT("r.UniformExpressionCacheAsyncUpdates"),
-	GUniformExpressionCacheAsyncUpdates,
-	TEXT("Whether to allow async updates of uniform expression caches."),
-	ECVF_RenderThreadSafe);
-
 bool Engine_IsStrataEnabled();
 
 struct FAllowCachingStaticParameterValues
@@ -3308,20 +3295,6 @@ bool FMaterial::ShouldCachePipeline(EShaderPlatform Platform, const FShaderPipel
 	return true;
 }
 
-//
-// FColoredMaterialRenderProxy implementation.
-//
-
-const FMaterial* FColoredMaterialRenderProxy::GetMaterialNoFallback(ERHIFeatureLevel::Type InFeatureLevel) const
-{
-	return Parent->GetMaterialNoFallback(InFeatureLevel);
-}
-
-const FMaterialRenderProxy* FColoredMaterialRenderProxy::GetFallback(ERHIFeatureLevel::Type InFeatureLevel) const
-{
-	return Parent->GetFallback(InFeatureLevel);
-}
-
 /**
  * Finds the shader matching the template type and the passed in vertex factory, asserts if not found.
  */
@@ -3829,1034 +3802,288 @@ void FMaterialVirtualTextureStack::Serialize(FArchive& Ar)
 	Ar << PreallocatedStackTextureIndex;
 }
 
-/*-----------------------------------------------------------------------------
-	FMaterialRenderProxy
------------------------------------------------------------------------------*/
-
-bool FMaterialRenderProxy::GetVectorValue(const FHashedMaterialParameterInfo& ParameterInfo, FLinearColor* OutValue, const FMaterialRenderContext& Context) const
+#if WITH_EDITOR
+void FMaterial::SubmitCompileJobs_GameThread(EShaderCompileJobPriority Priority)
 {
-	FMaterialParameterValue Value;
-	if (GetParameterValue(EMaterialParameterType::Vector, ParameterInfo, Value, Context))
+	check(IsInGameThread());
+
+	if (GameThreadCompilingShaderMapId != 0u && GameThreadShaderMap)
 	{
-		*OutValue = Value.AsLinearColor();
-		return true;
+		const EShaderCompileJobPriority SubmittedPriority = GameThreadShaderMapSubmittedPriority;
+
+		// To avoid as much useless work as possible, we make sure to submit our compile jobs only once per priority upgrade.
+		if (GameThreadShaderMapSubmittedPriority == EShaderCompileJobPriority::None || Priority > SubmittedPriority)
+		{
+			check(GameThreadPendingCompilerEnvironment.IsValid());
+
+			GameThreadShaderMapSubmittedPriority = Priority;
+			GameThreadShaderMap->SubmitCompileJobs(GameThreadCompilingShaderMapId, this, GameThreadPendingCompilerEnvironment, Priority);
+		}
 	}
-	return false;
 }
 
-bool FMaterialRenderProxy::GetScalarValue(const FHashedMaterialParameterInfo& ParameterInfo, float* OutValue, const FMaterialRenderContext& Context) const
+void FMaterial::SubmitCompileJobs_RenderThread(EShaderCompileJobPriority Priority) const
 {
-	FMaterialParameterValue Value;
-	if (GetParameterValue(EMaterialParameterType::Scalar, ParameterInfo, Value, Context))
+	check(IsInParallelRenderingThread());
+	if (RenderingThreadCompilingShaderMapId != 0u && RenderingThreadShaderMap)
 	{
-		*OutValue = Value.AsScalar();
-		return true;
+		// std::atomic don't support enum class, so we have to make sure our cast assumptions are respected.
+		static_assert((int8)EShaderCompileJobPriority::None == -1 && EShaderCompileJobPriority::Low < EShaderCompileJobPriority::ForceLocal, "Revise EShaderCompileJobPriority cast assumptions");
+		EShaderCompileJobPriority SubmittedPriority = (EShaderCompileJobPriority)RenderingThreadShaderMapSubmittedPriority.load(std::memory_order_relaxed);
+
+		// To avoid as much useless work as possible, we make sure to submit our compile jobs only once per priority upgrade.
+		if (SubmittedPriority == EShaderCompileJobPriority::None || Priority > SubmittedPriority)
+		{
+			RenderingThreadShaderMapSubmittedPriority = (int8)Priority;
+			RenderingThreadShaderMap->SubmitCompileJobs(RenderingThreadCompilingShaderMapId, this, RenderingThreadPendingCompilerEnvironment, Priority);
+		}
 	}
-	return false;
 }
+#endif // WITH_EDITOR
 
-bool FMaterialRenderProxy::GetTextureValue(const FHashedMaterialParameterInfo& ParameterInfo, const UTexture** OutValue, const FMaterialRenderContext& Context) const
+#if WITH_EDITOR
+/** Returns the number of samplers used in this material, or -1 if the material does not have a valid shader map (compile error or still compiling). */
+int32 FMaterialResource::GetSamplerUsage() const
 {
-	FMaterialParameterValue Value;
-	if (GetParameterValue(EMaterialParameterType::Texture, ParameterInfo, Value, Context))
+	if (GetGameThreadShaderMap())
 	{
-		*OutValue = Value.Texture;
-		return true;
+		return GetGameThreadShaderMap()->GetMaxTextureSamplers();
 	}
-	return false;
+
+	return -1;
 }
+#endif // WITH_EDITOR
 
-bool FMaterialRenderProxy::GetTextureValue(const FHashedMaterialParameterInfo& ParameterInfo, const URuntimeVirtualTexture** OutValue, const FMaterialRenderContext& Context) const
+#if WITH_EDITOR
+void FMaterialResource::GetUserInterpolatorUsage(uint32& NumUsedUVScalars, uint32& NumUsedCustomInterpolatorScalars) const
 {
-	FMaterialParameterValue Value;
-	if (GetParameterValue(EMaterialParameterType::RuntimeVirtualTexture, ParameterInfo, Value, Context))
+	NumUsedUVScalars = NumUsedCustomInterpolatorScalars = 0;
+
+	if (const FMaterialShaderMap* ShaderMap = GetGameThreadShaderMap())
 	{
-		*OutValue = Value.RuntimeVirtualTexture;
-		return true;
+		NumUsedUVScalars = ShaderMap->GetNumUsedUVScalars();
+		NumUsedCustomInterpolatorScalars = ShaderMap->GetNumUsedCustomInterpolatorScalars();
 	}
-	return false;
 }
 
-bool FMaterialRenderProxy::GetTextureValue(const FHashedMaterialParameterInfo& ParameterInfo, const USparseVolumeTexture** OutValue, const FMaterialRenderContext& Context) const
+void FMaterialResource::GetEstimatedNumTextureSamples(uint32& VSSamples, uint32& PSSamples) const
 {
-	FMaterialParameterValue Value;
-	if (GetParameterValue(EMaterialParameterType::SparseVolumeTexture, ParameterInfo, Value, Context))
+	VSSamples = PSSamples = 0;
+	if (const FMaterialShaderMap* ShaderMap = GetGameThreadShaderMap())
 	{
-		*OutValue = Value.SparseVolumeTexture;
-		return true;
+		ShaderMap->GetEstimatedNumTextureSamples(VSSamples, PSSamples);
 	}
-	return false;
 }
 
-static void OnVirtualTextureDestroyedCB(const FVirtualTextureProducerHandle& InHandle, void* Baton)
+uint32 FMaterialResource::GetEstimatedNumVirtualTextureLookups() const
 {
-	FMaterialRenderProxy* MaterialProxy = static_cast<FMaterialRenderProxy*>(Baton);
+	if (const FMaterialShaderMap* ShaderMap = GetGameThreadShaderMap())
+	{
+		return ShaderMap->GetEstimatedNumVirtualTextureLookups();
+	}
+	return 0;
+}
+#endif // WITH_EDITOR
 
-	MaterialProxy->InvalidateUniformExpressionCache(false);
-	UMaterialInterface::IterateOverActiveFeatureLevels([&](ERHIFeatureLevel::Type InFeatureLevel)
+uint32 FMaterialResource::GetNumVirtualTextureStacks() const
+{
+	if (const FMaterialShaderMap* ShaderMap = GetGameThreadShaderMap())
 	{
-		MaterialProxy->UpdateUniformExpressionCacheIfNeeded(InFeatureLevel);
-	});
+		return ShaderMap->GetNumVirtualTextureStacks();
+	}
+	return 0;
 }
 
-IAllocatedVirtualTexture* FMaterialRenderProxy::GetPreallocatedVTStack(const FMaterialRenderContext& Context, const FUniformExpressionSet& UniformExpressionSet, const FMaterialVirtualTextureStack& VTStack) const
+
+FString FMaterialResource::GetMaterialUsageDescription() const
 {
-	check(VTStack.IsPreallocatedStack())
+	check(Material);
+	FString BaseDescription = FString::Printf(
+		TEXT("LightingModel=%s, BlendMode=%s, "),
+		*GetShadingModelFieldString(GetShadingModels()), *GetBlendModeString(GetBlendMode()));
 
-	URuntimeVirtualTexture const* Texture;
-	VTStack.GetTextureValue(Context, UniformExpressionSet, Texture);
+	// this changed from ",SpecialEngine, TwoSided" to ",SpecialEngine=1, TwoSided=1, TSNormal=0, ..." to be more readable
+	BaseDescription += FString::Printf(
+		TEXT("SpecialEngine=%d, TwoSided=%d, TSNormal=%d, Masked=%d, Distorted=%d, WritesEveryPixel=%d, ModifiesMeshPosition=%d")
+		TEXT(", Usage={"),
+		(int32)IsSpecialEngineMaterial(), (int32)IsTwoSided(), (int32)IsTangentSpaceNormal(), (int32)IsMasked(), (int32)IsDistorted(), (int32)WritesEveryPixel(), (int32)MaterialMayModifyMeshPosition()
+		);
 
-	if (Texture == nullptr)
+	bool bFirst = true;
+	for (int32 MaterialUsageIndex = 0; MaterialUsageIndex < MATUSAGE_MAX; MaterialUsageIndex++)
 	{
-		return nullptr;
+		if (Material->GetUsageByFlag((EMaterialUsage)MaterialUsageIndex))
+		{
+			if (!bFirst)
+			{
+				BaseDescription += FString(TEXT(","));
+			}
+			BaseDescription += Material->GetUsageName((EMaterialUsage)MaterialUsageIndex);
+			bFirst = false;
+		}
 	}
+	BaseDescription += FString(TEXT("}"));
 
-	GetRendererModule().AddVirtualTextureProducerDestroyedCallback(Texture->GetProducerHandle(), &OnVirtualTextureDestroyedCB, const_cast<FMaterialRenderProxy*>(this));
-	HasVirtualTextureCallbacks = true;
-
-	return Texture->GetAllocatedVirtualTexture();
+	return BaseDescription;
 }
 
-IAllocatedVirtualTexture* FMaterialRenderProxy::AllocateVTStack(const FMaterialRenderContext& Context, const FUniformExpressionSet& UniformExpressionSet, const FMaterialVirtualTextureStack& VTStack) const
+static void AddSortedShader(TArray<FShaderType*>& Shaders, FShaderType* Shader)
 {
-	check(!VTStack.IsPreallocatedStack());
-	const uint32 NumLayers = VTStack.GetNumLayers();
-	if (NumLayers == 0u)
+	const int32 SortedIndex = Algo::LowerBoundBy(Shaders, Shader->GetHashedName(), [](const FShaderType* InShaderType) { return InShaderType->GetHashedName(); });
+	if (SortedIndex >= Shaders.Num() || Shaders[SortedIndex] != Shader)
 	{
-		return nullptr;
+		Shaders.Insert(Shader, SortedIndex);
 	}
+}
 
-	const UTexture* LayerTextures[VIRTUALTEXTURE_SPACE_MAXLAYERS] = { nullptr };
-	VTStack.GetTextureValues(Context, UniformExpressionSet, LayerTextures);
+static void AddSortedShaderPipeline(TArray<const FShaderPipelineType*>& Pipelines, const FShaderPipelineType* Pipeline)
+{
+	const int32 SortedIndex = Algo::LowerBoundBy(Pipelines, Pipeline->GetHashedName(), [](const FShaderPipelineType* InPiplelineType) { return InPiplelineType->GetHashedName(); });
+	if (SortedIndex >= Pipelines.Num() || Pipelines[SortedIndex] != Pipeline)
+	{
+		Pipelines.Insert(Pipeline, SortedIndex);
+	}
+}
 
-	const UMaterialInterface* MaterialInterface = GetMaterialInterface();
+void FMaterial::GetDependentShaderAndVFTypes(EShaderPlatform Platform, const FPlatformTypeLayoutParameters& LayoutParams, TArray<FShaderType*>& OutShaderTypes, TArray<const FShaderPipelineType*>& OutShaderPipelineTypes, TArray<FVertexFactoryType*>& OutVFTypes) const
+{
+	const FMaterialShaderParameters MaterialParameters(this);
+	const FMaterialShaderMapLayout& Layout = AcquireMaterialShaderMapLayout(Platform, GetShaderPermutationFlags(LayoutParams), MaterialParameters);
 
-	FAllocatedVTDescription VTDesc;
-	if (MaterialInterface)
+	for (const FShaderLayoutEntry& Shader : Layout.Shaders)
 	{
-		VTDesc.Name = MaterialInterface->GetFName();
+		if (ShouldCache(Platform, Shader.ShaderType, nullptr))
+		{
+			AddSortedShader(OutShaderTypes, Shader.ShaderType);
+		}
 	}
-	VTDesc.Dimensions = 2;
-	VTDesc.NumTextureLayers = NumLayers;
-	bool bFoundValidLayer = false;
-	for (uint32 LayerIndex = 0u; LayerIndex < NumLayers; ++LayerIndex)
+
+	for (const FShaderPipelineType* Pipeline : Layout.ShaderPipelines)
 	{
-		const UTexture* Texture = LayerTextures[LayerIndex];
-		if (!Texture)
+		if (ShouldCachePipeline(Platform, Pipeline, nullptr))
 		{
-			continue;
+			AddSortedShaderPipeline(OutShaderPipelineTypes, Pipeline);
+			for (const FShaderType* Type : Pipeline->GetStages())
+			{
+				AddSortedShader(OutShaderTypes, (FShaderType*)Type);
+			}
 		}
+	}
 
-		// GetResource() is safe to call from the render thread.
-		const FTextureResource* TextureResource = Texture->GetResource();
-		if (TextureResource)
+	for (const FMeshMaterialShaderMapLayout& MeshLayout : Layout.MeshShaderMaps)
+	{
+		bool bIncludeVertexFactory = false;
+		for (const FShaderLayoutEntry& Shader : MeshLayout.Shaders)
 		{
-			const FVirtualTexture2DResource* VirtualTextureResourceForLayer = TextureResource->GetVirtualTexture2DResource();
-
-			if (VirtualTextureResourceForLayer == nullptr)
+			if (ShouldCache(Platform, Shader.ShaderType, MeshLayout.VertexFactoryType))
 			{
-				// The placeholder used during async texture compilation is expected to be of the wrong type since
-				// no VT infos are available until later in the compilation process. This will be resolved
-				// once the final texture resource is available.
-#if WITH_EDITOR
-				if (!TextureResource->IsProxy())
-#endif
-				{
-					UE_LOG(LogMaterial, Warning, TEXT("Material '%s' expects texture '%s' to be Virtual"),
-						*GetFriendlyName(), *Texture->GetName());
-				}
-				continue;
+				bIncludeVertexFactory = true;
+				AddSortedShader(OutShaderTypes, Shader.ShaderType);
 			}
-			else
-			{
-				// All tile sizes need to match
-				check(!bFoundValidLayer || VTDesc.TileSize == VirtualTextureResourceForLayer->GetTileSize());
-				check(!bFoundValidLayer || VTDesc.TileBorderSize == VirtualTextureResourceForLayer->GetBorderSize());
+		}
 
-				const FVirtualTextureProducerHandle& ProducerHandle = VirtualTextureResourceForLayer->GetProducerHandle();
-				if (ProducerHandle.IsValid())
+		for (const FShaderPipelineType* Pipeline : MeshLayout.ShaderPipelines)
+		{
+			if (ShouldCachePipeline(Platform, Pipeline, MeshLayout.VertexFactoryType))
+			{
+				bIncludeVertexFactory = true;
+				AddSortedShaderPipeline(OutShaderPipelineTypes, Pipeline);
+				for (const FShaderType* Type : Pipeline->GetStages())
 				{
-					VTDesc.TileSize = VirtualTextureResourceForLayer->GetTileSize();
-					VTDesc.TileBorderSize = VirtualTextureResourceForLayer->GetBorderSize();
-					VTDesc.ProducerHandle[LayerIndex] = ProducerHandle;
-					VTDesc.ProducerLayerIndex[LayerIndex] = 0u;
-					GetRendererModule().AddVirtualTextureProducerDestroyedCallback(ProducerHandle, &OnVirtualTextureDestroyedCB, const_cast<FMaterialRenderProxy*>(this));
-					bFoundValidLayer = true;
+					AddSortedShader(OutShaderTypes, (FShaderType*)Type);
 				}
 			}
 		}
-	}
 
-	if (bFoundValidLayer)
-	{
-		HasVirtualTextureCallbacks = true;
-		return GetRendererModule().AllocateVirtualTexture(VTDesc);
+		if (bIncludeVertexFactory)
+		{
+			// Vertex factories are already sorted
+			OutVFTypes.Add(MeshLayout.VertexFactoryType);
+		}
 	}
-	return nullptr;
 }
 
-PRAGMA_DISABLE_DEPRECATION_WARNINGS
-FUniformExpressionCache::~FUniformExpressionCache()
+void FMaterial::GetReferencedTexturesHash(EShaderPlatform Platform, FSHAHash& OutHash) const
 {
-	ResetAllocatedVTs();
-	UniformBuffer.SafeRelease();
-}
-PRAGMA_ENABLE_DEPRECATION_WARNINGS
+	FSHA1 HashState;
 
-class FUniformExpressionCacheAsyncUpdateTask
-{
-public:
-	void Begin()
+	const TArrayView<const TObjectPtr<UObject>> ReferencedTextures = GetReferencedTextures();
+	// Hash the names of the uniform expression textures to capture changes in their order or values resulting from material compiler code changes
+	for (int32 TextureIndex = 0; TextureIndex < ReferencedTextures.Num(); TextureIndex++)
 	{
-		ReferenceCount++;
-	}
+		FString TextureName;
 
-	void End()
-	{
-		check(ReferenceCount > 0);
-		if (--ReferenceCount == 0)
+		if (ReferencedTextures[TextureIndex])
 		{
-			Wait();
+			TextureName = ReferencedTextures[TextureIndex]->GetName();
 		}
-	}
 
-	bool IsEnabled() const
-	{
-		return ReferenceCount > 0 && GUniformExpressionCacheAsyncUpdates > 0 && !GRHICommandList.Bypass();
+		HashState.UpdateWithString(*TextureName, TextureName.Len());
 	}
 
-	void SetTask(const FGraphEventRef& InTask)
+	UMaterialShaderQualitySettings* MaterialShaderQualitySettings = UMaterialShaderQualitySettings::Get();
+	if(MaterialShaderQualitySettings->HasPlatformQualitySettings(Platform, QualityLevel))
 	{
-		check(!Task);
-		check(IsEnabled());
-		Task = InTask;
+		MaterialShaderQualitySettings->GetShaderPlatformQualitySettings(Platform)->AppendToHashState(QualityLevel, HashState);
 	}
 
-	void Wait()
-	{
-		if (Task)
-		{
-			Task->Wait();
-			Task = nullptr;
-		}
-	}
+	HashState.Final();
+	HashState.GetHash(&OutHash.Hash[0]);
+}
 
-private:
-	FGraphEventRef Task;
-	int32 ReferenceCount = 0;
+/**
+ * Get user source code for the material, with a list of code snippets to highlight representing the code for each MaterialExpression
+ * @param OutSource - generated source code
+ * @param OutHighlightMap - source code highlight list
+ * @return - true on Success
+ */
+bool FMaterial::GetMaterialExpressionSource( FString& OutSource )
+{
+#if WITH_EDITORONLY_DATA
+	const EShaderPlatform ShaderPlatform = GMaxRHIShaderPlatform;
+	const ITargetPlatform* TargetPlatform = nullptr;
 
-} GUniformExpressionCacheAsyncUpdateTask;
+	FMaterialShaderMapId ShaderMapId;
+	GetShaderMapId(ShaderPlatform, TargetPlatform, ShaderMapId);
+	FStaticParameterSet StaticParameterSet;
+	GetStaticParameterSet(ShaderPlatform, StaticParameterSet);
 
-FUniformExpressionCacheAsyncUpdateScope::FUniformExpressionCacheAsyncUpdateScope()
-{
-	ENQUEUE_RENDER_COMMAND(BeginAsyncUniformExpressionCacheUpdates)(
-		[](FRHICommandList&)
-	{
-		GUniformExpressionCacheAsyncUpdateTask.Begin();
-		FMaterialRenderProxy::UpdateDeferredCachedUniformExpressions();
-	});
-}
+	FMaterialCompilationOutput NewCompilationOutput;
+	TRefCountPtr<FSharedShaderCompilerEnvironment> MaterialEnvironment;
+	const bool bSuccess = Translate(ShaderMapId, StaticParameterSet, ShaderPlatform, TargetPlatform, NewCompilationOutput, MaterialEnvironment);
 
-FUniformExpressionCacheAsyncUpdateScope::~FUniformExpressionCacheAsyncUpdateScope()
-{
-	ENQUEUE_RENDER_COMMAND(EndAsyncUniformExpressionCacheUpdates)(
-		[](FRHICommandList&)
+	if (bSuccess)
 	{
-		GUniformExpressionCacheAsyncUpdateTask.End();
-	});
+		FString* Source = MaterialEnvironment->IncludeVirtualPathToContentsMap.Find(TEXT("/Engine/Generated/Material.ush"));
+		if (Source)
+		{
+			OutSource = MoveTemp(*Source);
+			return true;
+		}
+	}
+	return false;
+#else
+	UE_LOG(LogMaterial, Fatal,TEXT("Not supported."));
+	return false;
+#endif
 }
 
-void FUniformExpressionCacheAsyncUpdateScope::WaitForTask()
+bool FMaterial::WritesEveryPixel(bool bShadowPass) const
 {
-	GUniformExpressionCacheAsyncUpdateTask.Wait();
+	bool bLocalStencilDitheredLOD = FeatureLevel >= ERHIFeatureLevel::SM5 && bStencilDitheredLOD;
+	return !IsMasked()
+		// Render dithered material as masked if a stencil prepass is not used (UE-50064, UE-49537)
+		&& !((bShadowPass || !bLocalStencilDitheredLOD) && IsDitheredLODTransition())
+		&& !IsWireframe()
+		&& !(bLocalStencilDitheredLOD && IsDitheredLODTransition() && IsUsedWithInstancedStaticMeshes())
+		&& !IsStencilTestEnabled();
 }
 
-class FUniformExpressionCacheAsyncUpdater
+#if WITH_EDITOR
+/** Recompiles any materials in the EditorLoadedMaterialResources list if they are not complete. */
+void FMaterial::UpdateEditorLoadedMaterialResources(EShaderPlatform InShaderPlatform)
 {
-public:
-	void Add(FUniformExpressionCache* UniformExpressionCache, const FUniformExpressionSet* UniformExpressionSet, const FRHIUniformBufferLayout* UniformBufferLayout, const FMaterialRenderContext& Context)
-	{
-		Items.Emplace(UniformExpressionCache, UniformExpressionSet, UniformBufferLayout, Context);
-	}
-
-	void Update(FRHICommandListImmediate& RHICmdListImmediate)
-	{
-		if (Items.IsEmpty())
-		{
-			return;
-		}
-
-		GUniformExpressionCacheAsyncUpdateTask.Wait();
-
-		FRHICommandList* RHICmdList = new FRHICommandList(FRHIGPUMask::All());
-
-		FGraphEventRef Event = FFunctionGraphTask::CreateAndDispatchWhenReady(
-			[Items = MoveTemp(Items), RHICmdList]
-		{
-			TRACE_CPUPROFILER_EVENT_SCOPE(FUniformExpressionCacheAsyncUpdater::Update);
-			FTaskTagScope Scope(ETaskTag::EParallelRenderingThread);
-			FMemMark Mark(FMemStack::Get());
-			RHICmdList->SwitchPipeline(ERHIPipeline::Graphics);
-
-			for (const FItem& Item : Items)
-			{
-				uint8* TempBuffer = FMemStack::Get().PushBytes(Item.UniformBufferLayout->ConstantBufferSize, SHADER_PARAMETER_STRUCT_ALIGNMENT);
-
-				FMaterialRenderContext Context(Item.MaterialRenderProxy, *Item.Material, Item.bShowSelection);
-
-				Item.UniformExpressionSet->FillUniformBuffer(Context, Item.AllocatedVTs, Item.UniformBufferLayout, TempBuffer, Item.UniformBufferLayout->ConstantBufferSize);
-
-				RHICmdList->UpdateUniformBuffer(Item.UniformBuffer, TempBuffer);
-			}
-
-			RHICmdList->FinishRecording();
-
-		}, TStatId(), nullptr, ENamedThreads::AnyHiPriThreadHiPriTask);
-
-		RHICmdListImmediate.QueueAsyncCommandListSubmit(RHICmdList);
-
-		GUniformExpressionCacheAsyncUpdateTask.SetTask(Event);
-	}
-
-private:
-	struct FItem
-	{
-		FItem() = default;
-		FItem(FUniformExpressionCache* InUniformExpressionCache, const FUniformExpressionSet* InUniformExpressionSet, const FRHIUniformBufferLayout* InUniformBufferLayout, const FMaterialRenderContext& Context)
-			: UniformBuffer(InUniformExpressionCache->UniformBuffer)
-			, AllocatedVTs(InUniformExpressionCache->AllocatedVTs)
-			, UniformExpressionSet(InUniformExpressionSet)
-			, UniformBufferLayout(InUniformBufferLayout)
-			, MaterialRenderProxy(Context.MaterialRenderProxy)
-			, Material(&Context.Material)
-			, bShowSelection(Context.bShowSelection)
-		{}
-
-		TRefCountPtr<FRHIUniformBuffer> UniformBuffer;
-		TArray<IAllocatedVirtualTexture*, FConcurrentLinearArrayAllocator> AllocatedVTs;
-		const FUniformExpressionSet* UniformExpressionSet = nullptr;
-		const FRHIUniformBufferLayout* UniformBufferLayout = nullptr;
-		const FMaterialRenderProxy* MaterialRenderProxy = nullptr;
-		const FMaterial* Material = nullptr;
-		bool bShowSelection = false;
-	};
-
-	TArray<FItem, FConcurrentLinearArrayAllocator> Items;
-};
-
-void FUniformExpressionCache::ResetAllocatedVTs()
-{
-	for (int32 i=0; i< OwnedAllocatedVTs.Num(); ++i)
-	{
-		GetRendererModule().DestroyVirtualTexture(OwnedAllocatedVTs[i]);
-	}
-	AllocatedVTs.Reset();
-	OwnedAllocatedVTs.Reset();
-}
-
-void FMaterialRenderProxy::EvaluateUniformExpressions(FUniformExpressionCache& OutUniformExpressionCache, const FMaterialRenderContext& Context, FUniformExpressionCacheAsyncUpdater* Updater) const
-{
-	check(IsInRenderingThread());
-
-	SCOPE_CYCLE_COUNTER(STAT_CacheUniformExpressions);
-
-	TRACE_CPUPROFILER_EVENT_SCOPE(FMaterialRenderProxy::EvaluateUniformExpressions);
-	
-	// Retrieve the material's uniform expression set.
-	FMaterialShaderMap* ShaderMap = Context.Material.GetRenderingThreadShaderMap();
-	const FUniformExpressionSet& UniformExpressionSet = ShaderMap->GetUniformExpressionSet();
-
-	OutUniformExpressionCache.CachedUniformExpressionShaderMap = ShaderMap;
-	OutUniformExpressionCache.ResetAllocatedVTs();
-	OutUniformExpressionCache.AllocatedVTs.Empty(UniformExpressionSet.VTStacks.Num());
-	OutUniformExpressionCache.OwnedAllocatedVTs.Empty(UniformExpressionSet.VTStacks.Num());
-	
-	if (HasVirtualTextureCallbacks)
-	{
-		GetRendererModule().RemoveAllVirtualTextureProducerDestroyedCallbacks(this);
-		HasVirtualTextureCallbacks = false;
-	}
-
-	for (int32 i = 0; i < UniformExpressionSet.VTStacks.Num(); ++i)
-	{
-		const FMaterialVirtualTextureStack& VTStack = UniformExpressionSet.VTStacks[i];
-		IAllocatedVirtualTexture* AllocatedVT = nullptr;
-		if (VTStack.IsPreallocatedStack())
-		{
-			AllocatedVT = GetPreallocatedVTStack(Context, UniformExpressionSet, VTStack);
-		}
-		else
-		{
-			AllocatedVT = AllocateVTStack(Context, UniformExpressionSet, VTStack);
-			if (AllocatedVT != nullptr)
-			{
-				OutUniformExpressionCache.OwnedAllocatedVTs.Add(AllocatedVT);
-			}
-		}
-		OutUniformExpressionCache.AllocatedVTs.Add(AllocatedVT);
-	}
-
-	const FRHIUniformBufferLayout* UniformBufferLayout = ShaderMap->GetUniformBufferLayout();
-
-	if (IsValidRef(OutUniformExpressionCache.UniformBuffer))
-	{
-		if (!OutUniformExpressionCache.UniformBuffer->IsValid())
-		{
-			UE_LOG(LogMaterial, Fatal, TEXT("The Uniformbuffer needs to be valid if it has been set"));
-		}
-
-		// The actual pointer may not match because there are cases (in the editor, during the shader compilation) when material's shader map gets updated without proxy's cache
-		// getting invalidated, but the layout contents must match.
-		check(OutUniformExpressionCache.UniformBuffer->GetLayoutPtr() == UniformBufferLayout || *OutUniformExpressionCache.UniformBuffer->GetLayoutPtr() == *UniformBufferLayout);
-	}
-
-	if (Updater)
-	{
-		if (!IsValidRef(OutUniformExpressionCache.UniformBuffer))
-		{
-			OutUniformExpressionCache.UniformBuffer = RHICreateUniformBuffer(nullptr, UniformBufferLayout, UniformBuffer_MultiFrame);
-		}
-
-		Updater->Add(&OutUniformExpressionCache, &UniformExpressionSet, UniformBufferLayout, Context);
-	}
-	else
-	{
-		FMemMark Mark(FMemStack::Get());
-		uint8* TempBuffer = FMemStack::Get().PushBytes(UniformBufferLayout->ConstantBufferSize, SHADER_PARAMETER_STRUCT_ALIGNMENT);
-		UniformExpressionSet.FillUniformBuffer(Context, OutUniformExpressionCache, UniformBufferLayout, TempBuffer, UniformBufferLayout->ConstantBufferSize);
-
-		if (IsValidRef(OutUniformExpressionCache.UniformBuffer))
-		{
-			RHIUpdateUniformBuffer(OutUniformExpressionCache.UniformBuffer, TempBuffer);
-		}
-		else
-		{
-			OutUniformExpressionCache.UniformBuffer = RHICreateUniformBuffer(TempBuffer, UniformBufferLayout, UniformBuffer_MultiFrame);
-		}
-	}
-
-	OutUniformExpressionCache.ParameterCollections = UniformExpressionSet.ParameterCollections;
-
-	++UniformExpressionCacheSerialNumber;
-}
-
-void FMaterialRenderProxy::CacheUniformExpressions(bool bRecreateUniformBuffer)
-{
-	// Register the render proxy's as a render resource so it can receive notifications to free the uniform buffer.
-	InitResource();
-
-	bool bUsingNewLoader = EVENT_DRIVEN_ASYNC_LOAD_ACTIVE_AT_RUNTIME && GEventDrivenLoaderEnabled;
-
-	check((bUsingNewLoader && GIsInitialLoad) || // The EDL at boot time maybe not load the default materials first; we need to intialize materials before the default materials are done
-		UMaterial::GetDefaultMaterial(MD_Surface));
-
-
-	if (IsMarkedForGarbageCollection())
-	{
-		UE_LOG(LogMaterial, Fatal, TEXT("Cannot queue the Expression Cache for Material %s when it is about to be deleted"), *MaterialName);
-	}
-	StartCacheUniformExpressions();
-	DeferredUniformExpressionCacheRequests.Add(this);
-
-	InvalidateUniformExpressionCache(bRecreateUniformBuffer);
-
-	if (!GDeferUniformExpressionCaching)
-	{
-		FMaterialRenderProxy::UpdateDeferredCachedUniformExpressions();
-	}
-}
-
-void FMaterialRenderProxy::CacheUniformExpressions_GameThread(bool bRecreateUniformBuffer)
-{
-	if (FApp::CanEverRender())
-	{
-		UE_LOG(LogMaterial, VeryVerbose, TEXT("Caching uniform expressions for material: %s"), *GetFriendlyName());
-
-		FMaterialRenderProxy* RenderProxy = this;
-		ENQUEUE_RENDER_COMMAND(FCacheUniformExpressionsCommand)(
-			[RenderProxy, bRecreateUniformBuffer](FRHICommandListImmediate& RHICmdList)
-			{
-				RenderProxy->CacheUniformExpressions(bRecreateUniformBuffer);
-			});
-	}
-}
-
-void FMaterialRenderProxy::InvalidateUniformExpressionCache(bool bRecreateUniformBuffer)
-{
-	check(IsInRenderingThread());
-	GUniformExpressionCacheAsyncUpdateTask.Wait();
-
-#if WITH_EDITOR
-	FStaticLightingSystemInterface::OnMaterialInvalidated.Broadcast(this);
-#endif
-
-	if (HasVirtualTextureCallbacks)
-	{
-		GetRendererModule().RemoveAllVirtualTextureProducerDestroyedCallbacks(this);
-		HasVirtualTextureCallbacks = false;
-	}
-
-	++UniformExpressionCacheSerialNumber;
-	for (int32 i = 0; i < ERHIFeatureLevel::Num; ++i)
-	{
-		UniformExpressionCache[i].CachedUniformExpressionShaderMap = nullptr;
-		UniformExpressionCache[i].ResetAllocatedVTs();
-
-		if (bRecreateUniformBuffer)
-		{
-			// This is required if the FMaterial is being recompiled (the uniform buffer layout will change).
-			// This should only be done if the calling code is using FMaterialUpdateContext to recreate the rendering state of primitives using this material, 
-			// Since cached mesh commands also cache uniform buffer pointers.
-			UniformExpressionCache[i].UniformBuffer = nullptr;
-		}
-	}
-}
-
-void FMaterialRenderProxy::UpdateUniformExpressionCacheIfNeeded(ERHIFeatureLevel::Type InFeatureLevel) const
-{
-	// Don't cache uniform expressions if an entirely different FMaterialRenderProxy is going to be used for rendering
-	const FMaterial* Material = GetMaterialNoFallback(InFeatureLevel);
-	// Note: We would actually need to compare the FMaterialShaderMapId of both shader maps but a simple pointer compare also works
-	// because shader maps are currently swapped out whenever they are modified.
-	if (Material && Material->GetRenderingThreadShaderMap() != UniformExpressionCache[InFeatureLevel].CachedUniformExpressionShaderMap)
-	{
-		FMaterialRenderContext MaterialRenderContext(this, *Material, nullptr);
-		MaterialRenderContext.bShowSelection = GIsEditor;
-		EvaluateUniformExpressions(UniformExpressionCache[InFeatureLevel], MaterialRenderContext);
-	}
-}
-
-FMaterialRenderProxy::FMaterialRenderProxy(FString InMaterialName)
-	: SubsurfaceProfileRT(0)
-	, MaterialName(MoveTemp(InMaterialName))
-	, MarkedForGarbageCollection(0)
-	, DeletedFlag(0)
-	, HasVirtualTextureCallbacks(0)
-{
-}
-
-FMaterialRenderProxy::~FMaterialRenderProxy()
-{
-	// We only wait on deletions happening on the render thread. Async deletions can happen during scene render shutdown and those are waited on explicitly.
-	if (IsInRenderingThread())
-	{
-		GUniformExpressionCacheAsyncUpdateTask.Wait();
-	}
-
-	if(IsInitialized())
-	{
-		check(IsInRenderingThread());
-		ReleaseResource();
-	}
-
-	if (HasVirtualTextureCallbacks)
-	{
-		check(IsInRenderingThread());
-		GetRendererModule().RemoveAllVirtualTextureProducerDestroyedCallbacks(this);
-		HasVirtualTextureCallbacks = false;
-	}
-
-	DeletedFlag = 1;
-}
-
-void FMaterialRenderProxy::InitDynamicRHI()
-{
-#if WITH_EDITOR
-	// MaterialRenderProxyMap is only used by shader compiling
-	if (!FPlatformProperties::RequiresCookedData())
-	{
-		FScopeLock Locker(&MaterialRenderProxyMapLock);
-		FMaterialRenderProxy::MaterialRenderProxyMap.Add(this);
-	}
-#endif // WITH_EDITOR
-}
-
-void FMaterialRenderProxy::CancelCacheUniformExpressions()
-{
-	DeferredUniformExpressionCacheRequests.Remove(this);
-}
-
-void FMaterialRenderProxy::ReleaseDynamicRHI()
-{
-#if WITH_EDITOR
-	if (!FPlatformProperties::RequiresCookedData())
-	{
-		FScopeLock Locker(&MaterialRenderProxyMapLock);
-		FMaterialRenderProxy::MaterialRenderProxyMap.Remove(this);
-	}
-#endif // WITH_EDITOR
-
-	if (DeferredUniformExpressionCacheRequests.Remove(this))
-	{
-		// Notify that we're finished with this inflight cache request, because the object is being released
-		FinishCacheUniformExpressions();
-	}
-
-	InvalidateUniformExpressionCache(true);
-
-	FExternalTextureRegistry::Get().RemoveMaterialRenderProxyReference(this);
-}
-
-void FMaterialRenderProxy::ReleaseResource()
-{
-	ReleaseResourceFlag = true;
-	FRenderResource::ReleaseResource();
-	if (HasVirtualTextureCallbacks)
-	{
-		GetRendererModule().RemoveAllVirtualTextureProducerDestroyedCallbacks(this);
-		HasVirtualTextureCallbacks = false;
-	}
-}
-
-#if WITH_EDITOR
-void FMaterial::SubmitCompileJobs_GameThread(EShaderCompileJobPriority Priority)
-{
-	check(IsInGameThread());
-
-	if (GameThreadCompilingShaderMapId != 0u && GameThreadShaderMap)
-	{
-		const EShaderCompileJobPriority SubmittedPriority = GameThreadShaderMapSubmittedPriority;
-
-		// To avoid as much useless work as possible, we make sure to submit our compile jobs only once per priority upgrade.
-		if (GameThreadShaderMapSubmittedPriority == EShaderCompileJobPriority::None || Priority > SubmittedPriority)
-		{
-			check(GameThreadPendingCompilerEnvironment.IsValid());
-
-			GameThreadShaderMapSubmittedPriority = Priority;
-			GameThreadShaderMap->SubmitCompileJobs(GameThreadCompilingShaderMapId, this, GameThreadPendingCompilerEnvironment, Priority);
-		}
-	}
-}
-
-void FMaterial::SubmitCompileJobs_RenderThread(EShaderCompileJobPriority Priority) const
-{
-	check(IsInParallelRenderingThread());
-	if (RenderingThreadCompilingShaderMapId != 0u && RenderingThreadShaderMap)
-	{
-		// std::atomic don't support enum class, so we have to make sure our cast assumptions are respected.
-		static_assert((int8)EShaderCompileJobPriority::None == -1 && EShaderCompileJobPriority::Low < EShaderCompileJobPriority::ForceLocal, "Revise EShaderCompileJobPriority cast assumptions");
-		EShaderCompileJobPriority SubmittedPriority = (EShaderCompileJobPriority)RenderingThreadShaderMapSubmittedPriority.load(std::memory_order_relaxed);
-
-		// To avoid as much useless work as possible, we make sure to submit our compile jobs only once per priority upgrade.
-		if (SubmittedPriority == EShaderCompileJobPriority::None || Priority > SubmittedPriority)
-		{
-			RenderingThreadShaderMapSubmittedPriority = (int8)Priority;
-			RenderingThreadShaderMap->SubmitCompileJobs(RenderingThreadCompilingShaderMapId, this, RenderingThreadPendingCompilerEnvironment, Priority);
-		}
-	}
-}
-#endif // WITH_EDITOR
-
-const FMaterial& FMaterialRenderProxy::GetMaterialWithFallback(ERHIFeatureLevel::Type InFeatureLevel, const FMaterialRenderProxy*& OutFallbackMaterialRenderProxy) const
-{
-	const FMaterial* BaseMaterial = GetMaterialNoFallback(InFeatureLevel);
-	const FMaterial* Material = BaseMaterial;
-	if (!Material || !Material->IsRenderingThreadShaderMapComplete())
-	{
-		const FMaterialRenderProxy* FallbackMaterialProxy = this;
-		do
-		{
-			FallbackMaterialProxy = FallbackMaterialProxy->GetFallback(InFeatureLevel);
-			check(FallbackMaterialProxy);
-			Material = FallbackMaterialProxy->GetMaterialNoFallback(InFeatureLevel);
-		}
-		while (!Material || !Material->IsRenderingThreadShaderMapComplete());
-		OutFallbackMaterialRenderProxy = FallbackMaterialProxy;
-
-#if WITH_EDITOR
-		if (BaseMaterial)
-		{
-			BaseMaterial->SubmitCompileJobs_RenderThread(EShaderCompileJobPriority::Normal);
-		}
-#endif // WITH_EDITOR
-	}
-	return *Material;
-}
-
-const FMaterial& FMaterialRenderProxy::GetIncompleteMaterialWithFallback(ERHIFeatureLevel::Type InFeatureLevel) const
-{
-	const FMaterial* Material = GetMaterialNoFallback(InFeatureLevel);
-	if (!Material)
-	{
-		const FMaterialRenderProxy* FallbackMaterialProxy = this;
-		do
-		{
-			FallbackMaterialProxy = FallbackMaterialProxy->GetFallback(InFeatureLevel);
-			check(FallbackMaterialProxy);
-			Material = FallbackMaterialProxy->GetMaterialNoFallback(InFeatureLevel);
-		} while (!Material);
-	}
-	return *Material;
-}
-
-void FMaterialRenderProxy::UpdateDeferredCachedUniformExpressions()
-{
-	LLM_SCOPE(ELLMTag::Materials);
-
-	check(IsInRenderingThread());
-
-	TRACE_CPUPROFILER_EVENT_SCOPE(FMaterialRenderProxy::UpdateDeferredCachedUniformExpressions);
-	CSV_SCOPED_TIMING_STAT_EXCLUSIVE(Material_UpdateDeferredCachedUniformExpressions);
-	QUICK_SCOPE_CYCLE_COUNTER(STAT_UpdateDeferredCachedUniformExpressions);
-
-	FUniformExpressionCacheAsyncUpdater Updater;
-	FUniformExpressionCacheAsyncUpdater* UpdaterIfEnabled = GUniformExpressionCacheAsyncUpdateTask.IsEnabled() ? &Updater : nullptr;
-
-	for (TSet<FMaterialRenderProxy*>::TConstIterator It(DeferredUniformExpressionCacheRequests); It; ++It)
-	{
-		FMaterialRenderProxy* MaterialProxy = *It;
-		if (MaterialProxy->IsDeleted())
-		{
-			UE_LOG(LogMaterial, Fatal, TEXT("FMaterialRenderProxy deleted and GC mark was: %i"), MaterialProxy->IsMarkedForGarbageCollection());
-		}
-
-		UMaterialInterface::IterateOverActiveFeatureLevels([&](ERHIFeatureLevel::Type InFeatureLevel)
-		{
-			// Don't bother caching if we'll be falling back to a different FMaterialRenderProxy for rendering anyway
-			const FMaterial* Material = MaterialProxy->GetMaterialNoFallback(InFeatureLevel);
-			if (Material && Material->GetRenderingThreadShaderMap())
-			{
-				FMaterialRenderContext MaterialRenderContext(MaterialProxy, *Material, nullptr);
-				MaterialRenderContext.bShowSelection = GIsEditor;
-				MaterialProxy->EvaluateUniformExpressions(MaterialProxy->UniformExpressionCache[(int32)InFeatureLevel], MaterialRenderContext, UpdaterIfEnabled);
-			}
-		});
-
-		MaterialProxy->FinishCacheUniformExpressions();
-	}
-
-	if (UpdaterIfEnabled)
-	{
-		Updater.Update(FRHICommandListExecutor::GetImmediateCommandList());
-	}
-
-	DeferredUniformExpressionCacheRequests.Reset();
-}
-
-#if WITH_EDITOR
-TSet<FMaterialRenderProxy*> FMaterialRenderProxy::MaterialRenderProxyMap;
-FCriticalSection FMaterialRenderProxy::MaterialRenderProxyMapLock;
-#endif // WITH_EDITOR
-TSet<FMaterialRenderProxy*> FMaterialRenderProxy::DeferredUniformExpressionCacheRequests;
-
-/*-----------------------------------------------------------------------------
-	FColoredMaterialRenderProxy
------------------------------------------------------------------------------*/
-
-bool FColoredMaterialRenderProxy::GetParameterValue(EMaterialParameterType Type, const FHashedMaterialParameterInfo& ParameterInfo, FMaterialParameterValue& OutValue, const FMaterialRenderContext& Context) const
-{
-	if (Type == EMaterialParameterType::Vector && ParameterInfo.Name == ColorParamName)
-	{
-		OutValue = Color;
-		return true;
-	}
-	else
-	{
-		return Parent->GetParameterValue(Type, ParameterInfo, OutValue, Context);
-	}
-}
-
-/*-----------------------------------------------------------------------------
-	FColoredTexturedMaterialRenderProxy
------------------------------------------------------------------------------*/
-
-bool FColoredTexturedMaterialRenderProxy::GetParameterValue(EMaterialParameterType Type, const FHashedMaterialParameterInfo& ParameterInfo, FMaterialParameterValue& OutValue, const FMaterialRenderContext& Context) const
-{
-	if (Type == EMaterialParameterType::Texture && ParameterInfo.Name == TextureParamName)
-	{
-		OutValue = Texture;
-		return true;
-	}
-	else
-	{
-		return Parent->GetParameterValue(Type, ParameterInfo, OutValue, Context);
-	}
-}
-
-/*-----------------------------------------------------------------------------
-	FOverrideSelectionColorMaterialRenderProxy
------------------------------------------------------------------------------*/
-const FMaterial* FOverrideSelectionColorMaterialRenderProxy::GetMaterialNoFallback(ERHIFeatureLevel::Type InFeatureLevel) const
-{
-	return Parent->GetMaterialNoFallback(InFeatureLevel);
-}
-
-const FMaterialRenderProxy* FOverrideSelectionColorMaterialRenderProxy::GetFallback(ERHIFeatureLevel::Type InFeatureLevel) const
-{
-	return Parent->GetFallback(InFeatureLevel);
-}
-
-bool FOverrideSelectionColorMaterialRenderProxy::GetParameterValue(EMaterialParameterType Type, const FHashedMaterialParameterInfo& ParameterInfo, FMaterialParameterValue& OutValue, const FMaterialRenderContext& Context) const
-{
-	if (Type == EMaterialParameterType::Vector && ParameterInfo.Name == FName(NAME_SelectionColor))
-	{
-		OutValue = SelectionColor;
-		return true;
-	}
-	else
-	{
-		return Parent->GetParameterValue(Type, ParameterInfo, OutValue, Context);
-	}
-}
-
-/*-----------------------------------------------------------------------------
-	FLightingDensityMaterialRenderProxy
------------------------------------------------------------------------------*/
-static FName NAME_LightmapRes = FName(TEXT("LightmapRes"));
-
-bool FLightingDensityMaterialRenderProxy::GetParameterValue(EMaterialParameterType Type, const FHashedMaterialParameterInfo& ParameterInfo, FMaterialParameterValue& OutValue, const FMaterialRenderContext& Context) const
-{
-	if (Type == EMaterialParameterType::Vector && ParameterInfo.Name == NAME_LightmapRes)
-	{
-		OutValue = FLinearColor(LightmapResolution.X, LightmapResolution.Y, 0.0f, 0.0f);
-		return true;
-	}
-	return FColoredMaterialRenderProxy::GetParameterValue(Type, ParameterInfo, OutValue, Context);
-}
-
-#if WITH_EDITOR
-/** Returns the number of samplers used in this material, or -1 if the material does not have a valid shader map (compile error or still compiling). */
-int32 FMaterialResource::GetSamplerUsage() const
-{
-	if (GetGameThreadShaderMap())
-	{
-		return GetGameThreadShaderMap()->GetMaxTextureSamplers();
-	}
-
-	return -1;
-}
-#endif // WITH_EDITOR
-
-#if WITH_EDITOR
-void FMaterialResource::GetUserInterpolatorUsage(uint32& NumUsedUVScalars, uint32& NumUsedCustomInterpolatorScalars) const
-{
-	NumUsedUVScalars = NumUsedCustomInterpolatorScalars = 0;
-
-	if (const FMaterialShaderMap* ShaderMap = GetGameThreadShaderMap())
-	{
-		NumUsedUVScalars = ShaderMap->GetNumUsedUVScalars();
-		NumUsedCustomInterpolatorScalars = ShaderMap->GetNumUsedCustomInterpolatorScalars();
-	}
-}
-
-void FMaterialResource::GetEstimatedNumTextureSamples(uint32& VSSamples, uint32& PSSamples) const
-{
-	VSSamples = PSSamples = 0;
-	if (const FMaterialShaderMap* ShaderMap = GetGameThreadShaderMap())
-	{
-		ShaderMap->GetEstimatedNumTextureSamples(VSSamples, PSSamples);
-	}
-}
-
-uint32 FMaterialResource::GetEstimatedNumVirtualTextureLookups() const
-{
-	if (const FMaterialShaderMap* ShaderMap = GetGameThreadShaderMap())
-	{
-		return ShaderMap->GetEstimatedNumVirtualTextureLookups();
-	}
-	return 0;
-}
-#endif // WITH_EDITOR
-
-uint32 FMaterialResource::GetNumVirtualTextureStacks() const
-{
-	if (const FMaterialShaderMap* ShaderMap = GetGameThreadShaderMap())
-	{
-		return ShaderMap->GetNumVirtualTextureStacks();
-	}
-	return 0;
-}
-
-
-FString FMaterialResource::GetMaterialUsageDescription() const
-{
-	check(Material);
-	FString BaseDescription = FString::Printf(
-		TEXT("LightingModel=%s, BlendMode=%s, "),
-		*GetShadingModelFieldString(GetShadingModels()), *GetBlendModeString(GetBlendMode()));
-
-	// this changed from ",SpecialEngine, TwoSided" to ",SpecialEngine=1, TwoSided=1, TSNormal=0, ..." to be more readable
-	BaseDescription += FString::Printf(
-		TEXT("SpecialEngine=%d, TwoSided=%d, TSNormal=%d, Masked=%d, Distorted=%d, WritesEveryPixel=%d, ModifiesMeshPosition=%d")
-		TEXT(", Usage={"),
-		(int32)IsSpecialEngineMaterial(), (int32)IsTwoSided(), (int32)IsTangentSpaceNormal(), (int32)IsMasked(), (int32)IsDistorted(), (int32)WritesEveryPixel(), (int32)MaterialMayModifyMeshPosition()
-		);
-
-	bool bFirst = true;
-	for (int32 MaterialUsageIndex = 0; MaterialUsageIndex < MATUSAGE_MAX; MaterialUsageIndex++)
-	{
-		if (Material->GetUsageByFlag((EMaterialUsage)MaterialUsageIndex))
-		{
-			if (!bFirst)
-			{
-				BaseDescription += FString(TEXT(","));
-			}
-			BaseDescription += Material->GetUsageName((EMaterialUsage)MaterialUsageIndex);
-			bFirst = false;
-		}
-	}
-	BaseDescription += FString(TEXT("}"));
-
-	return BaseDescription;
-}
-
-static void AddSortedShader(TArray<FShaderType*>& Shaders, FShaderType* Shader)
-{
-	const int32 SortedIndex = Algo::LowerBoundBy(Shaders, Shader->GetHashedName(), [](const FShaderType* InShaderType) { return InShaderType->GetHashedName(); });
-	if (SortedIndex >= Shaders.Num() || Shaders[SortedIndex] != Shader)
-	{
-		Shaders.Insert(Shader, SortedIndex);
-	}
-}
-
-static void AddSortedShaderPipeline(TArray<const FShaderPipelineType*>& Pipelines, const FShaderPipelineType* Pipeline)
-{
-	const int32 SortedIndex = Algo::LowerBoundBy(Pipelines, Pipeline->GetHashedName(), [](const FShaderPipelineType* InPiplelineType) { return InPiplelineType->GetHashedName(); });
-	if (SortedIndex >= Pipelines.Num() || Pipelines[SortedIndex] != Pipeline)
-	{
-		Pipelines.Insert(Pipeline, SortedIndex);
-	}
-}
-
-void FMaterial::GetDependentShaderAndVFTypes(EShaderPlatform Platform, const FPlatformTypeLayoutParameters& LayoutParams, TArray<FShaderType*>& OutShaderTypes, TArray<const FShaderPipelineType*>& OutShaderPipelineTypes, TArray<FVertexFactoryType*>& OutVFTypes) const
-{
-	const FMaterialShaderParameters MaterialParameters(this);
-	const FMaterialShaderMapLayout& Layout = AcquireMaterialShaderMapLayout(Platform, GetShaderPermutationFlags(LayoutParams), MaterialParameters);
-
-	for (const FShaderLayoutEntry& Shader : Layout.Shaders)
-	{
-		if (ShouldCache(Platform, Shader.ShaderType, nullptr))
-		{
-			AddSortedShader(OutShaderTypes, Shader.ShaderType);
-		}
-	}
-
-	for (const FShaderPipelineType* Pipeline : Layout.ShaderPipelines)
-	{
-		if (ShouldCachePipeline(Platform, Pipeline, nullptr))
-		{
-			AddSortedShaderPipeline(OutShaderPipelineTypes, Pipeline);
-			for (const FShaderType* Type : Pipeline->GetStages())
-			{
-				AddSortedShader(OutShaderTypes, (FShaderType*)Type);
-			}
-		}
-	}
-
-	for (const FMeshMaterialShaderMapLayout& MeshLayout : Layout.MeshShaderMaps)
-	{
-		bool bIncludeVertexFactory = false;
-		for (const FShaderLayoutEntry& Shader : MeshLayout.Shaders)
-		{
-			if (ShouldCache(Platform, Shader.ShaderType, MeshLayout.VertexFactoryType))
-			{
-				bIncludeVertexFactory = true;
-				AddSortedShader(OutShaderTypes, Shader.ShaderType);
-			}
-		}
-
-		for (const FShaderPipelineType* Pipeline : MeshLayout.ShaderPipelines)
-		{
-			if (ShouldCachePipeline(Platform, Pipeline, MeshLayout.VertexFactoryType))
-			{
-				bIncludeVertexFactory = true;
-				AddSortedShaderPipeline(OutShaderPipelineTypes, Pipeline);
-				for (const FShaderType* Type : Pipeline->GetStages())
-				{
-					AddSortedShader(OutShaderTypes, (FShaderType*)Type);
-				}
-			}
-		}
-
-		if (bIncludeVertexFactory)
-		{
-			// Vertex factories are already sorted
-			OutVFTypes.Add(MeshLayout.VertexFactoryType);
-		}
-	}
-}
-
-void FMaterial::GetReferencedTexturesHash(EShaderPlatform Platform, FSHAHash& OutHash) const
-{
-	FSHA1 HashState;
-
-	const TArrayView<const TObjectPtr<UObject>> ReferencedTextures = GetReferencedTextures();
-	// Hash the names of the uniform expression textures to capture changes in their order or values resulting from material compiler code changes
-	for (int32 TextureIndex = 0; TextureIndex < ReferencedTextures.Num(); TextureIndex++)
-	{
-		FString TextureName;
-
-		if (ReferencedTextures[TextureIndex])
-		{
-			TextureName = ReferencedTextures[TextureIndex]->GetName();
-		}
-
-		HashState.UpdateWithString(*TextureName, TextureName.Len());
-	}
-
-	UMaterialShaderQualitySettings* MaterialShaderQualitySettings = UMaterialShaderQualitySettings::Get();
-	if(MaterialShaderQualitySettings->HasPlatformQualitySettings(Platform, QualityLevel))
-	{
-		MaterialShaderQualitySettings->GetShaderPlatformQualitySettings(Platform)->AppendToHashState(QualityLevel, HashState);
-	}
-
-	HashState.Final();
-	HashState.GetHash(&OutHash.Hash[0]);
-}
-
-/**
- * Get user source code for the material, with a list of code snippets to highlight representing the code for each MaterialExpression
- * @param OutSource - generated source code
- * @param OutHighlightMap - source code highlight list
- * @return - true on Success
- */
-bool FMaterial::GetMaterialExpressionSource( FString& OutSource )
-{
-#if WITH_EDITORONLY_DATA
-	const EShaderPlatform ShaderPlatform = GMaxRHIShaderPlatform;
-	const ITargetPlatform* TargetPlatform = nullptr;
-
-	FMaterialShaderMapId ShaderMapId;
-	GetShaderMapId(ShaderPlatform, TargetPlatform, ShaderMapId);
-	FStaticParameterSet StaticParameterSet;
-	GetStaticParameterSet(ShaderPlatform, StaticParameterSet);
-
-	FMaterialCompilationOutput NewCompilationOutput;
-	TRefCountPtr<FSharedShaderCompilerEnvironment> MaterialEnvironment;
-	const bool bSuccess = Translate(ShaderMapId, StaticParameterSet, ShaderPlatform, TargetPlatform, NewCompilationOutput, MaterialEnvironment);
-
-	if (bSuccess)
-	{
-		FString* Source = MaterialEnvironment->IncludeVirtualPathToContentsMap.Find(TEXT("/Engine/Generated/Material.ush"));
-		if (Source)
-		{
-			OutSource = MoveTemp(*Source);
-			return true;
-		}
-	}
-	return false;
-#else
-	UE_LOG(LogMaterial, Fatal,TEXT("Not supported."));
-	return false;
-#endif
-}
-
-bool FMaterial::WritesEveryPixel(bool bShadowPass) const
-{
-	bool bLocalStencilDitheredLOD = FeatureLevel >= ERHIFeatureLevel::SM5 && bStencilDitheredLOD;
-	return !IsMasked()
-		// Render dithered material as masked if a stencil prepass is not used (UE-50064, UE-49537)
-		&& !((bShadowPass || !bLocalStencilDitheredLOD) && IsDitheredLODTransition())
-		&& !IsWireframe()
-		&& !(bLocalStencilDitheredLOD && IsDitheredLODTransition() && IsUsedWithInstancedStaticMeshes())
-		&& !IsStencilTestEnabled();
-}
-
-#if WITH_EDITOR
-/** Recompiles any materials in the EditorLoadedMaterialResources list if they are not complete. */
-void FMaterial::UpdateEditorLoadedMaterialResources(EShaderPlatform InShaderPlatform)
-{
-	for (TSet<FMaterial*>::TIterator It(EditorLoadedMaterialResources); It; ++It)
+	for (TSet<FMaterial*>::TIterator It(EditorLoadedMaterialResources); It; ++It)
 	{
 		FMaterial* CurrentMaterial = *It;
 		if (!CurrentMaterial->GetGameThreadShaderMap() || !CurrentMaterial->GetGameThreadShaderMap()->IsComplete(CurrentMaterial, true))
@@ -5447,424 +4674,6 @@ bool FMaterialShaderMapId::ContainsVertexFactoryType(const FVertexFactoryType* V
 #endif // WITH_EDITOR
 //////////////////////////////////////////////////////////////////////////
 
-FMaterialAttributeDefintion::FMaterialAttributeDefintion(
-		const FGuid& InAttributeID, const FString& InAttributeName, EMaterialProperty InProperty,
-		EMaterialValueType InValueType, const FVector4& InDefaultValue, EShaderFrequency InShaderFrequency,
-		int32 InTexCoordIndex /*= INDEX_NONE*/, bool bInIsHidden /*= false*/, MaterialAttributeBlendFunction InBlendFunction /*= nullptr*/)
-	: AttributeID(InAttributeID)
-	, DefaultValue(InDefaultValue)
-	, AttributeName(InAttributeName)
-	, Property(InProperty)
-	, ValueType(InValueType)
-	, ShaderFrequency(InShaderFrequency)
-	, TexCoordIndex(InTexCoordIndex)
-	, BlendFunction(InBlendFunction)
-	, bIsHidden(bInIsHidden)
-{
-	checkf(ValueType & MCT_Float || ValueType == MCT_ShadingModel || ValueType == MCT_Strata || ValueType == MCT_MaterialAttributes, TEXT("Unsupported material attribute type %d"), ValueType);
-}
-
-int32 FMaterialAttributeDefintion::CompileDefaultValue(FMaterialCompiler* Compiler)
-{
-	int32 Ret;
-
-	// TODO: Temporarily preserving hack from 4.13 to change default value for two-sided foliage model 
-	if (Property == MP_SubsurfaceColor && Compiler->GetCompiledShadingModels().HasShadingModel(MSM_TwoSidedFoliage))
-	{
-		check(ValueType == MCT_Float3);
-		return Compiler->Constant3(0, 0, 0);
-	}
-
-	if (Property == MP_ShadingModel)
-	{
-		check(ValueType == MCT_ShadingModel);
-		// Default to the first shading model of the material. If the material is using a single shading model selected through the dropdown, this is how it gets written to the shader as a constant (optimizing out all the dynamic branches)
-		return Compiler->ShadingModel(Compiler->GetMaterialShadingModels().GetFirstShadingModel());
-	}
-
-	if (Property == MP_Normal)
-	{
-		if (Compiler->IsTangentSpaceNormal())
-		{
-			return Compiler->Constant3(0.0f, 0.0f, 1.0f);	// Tangent space normal
-		}
-		return Compiler->VertexNormal();					// World space normal
-	}
-
-	if (Property == MP_Tangent)
-	{
-		if (Compiler->IsTangentSpaceNormal())
-		{
-			return Compiler->Constant3(1.0f, 0.0f, 0.0f);	// Tangent space tangent
-		}
-		return Compiler->VertexTangent();					// World space tangent
-	}
-
-	if (Property == MP_FrontMaterial)
-	{
-		check(ValueType == MCT_Strata);
-		return Compiler->StrataCreateAndRegisterNullMaterial();
-	}
-
-	if (TexCoordIndex == INDEX_NONE)
-	{
-		// Standard value type
-		switch (ValueType)
-		{
-		case MCT_Float:
-		case MCT_Float1: Ret = Compiler->Constant(DefaultValue.X); break;
-		case MCT_Float2: Ret = Compiler->Constant2(DefaultValue.X, DefaultValue.Y); break;
-		case MCT_Float3: Ret = Compiler->Constant3(DefaultValue.X, DefaultValue.Y, DefaultValue.Z); break;
-		default: Ret = Compiler->Constant4(DefaultValue.X, DefaultValue.Y, DefaultValue.Z, DefaultValue.W);
-		}
-	}
-	else
-	{
-		// Texture coordinates allow pass through for default	
-		Ret = Compiler->TextureCoordinate(TexCoordIndex, false, false);
-	}
-
-	return Ret;
-}
-
-//////////////////////////////////////////////////////////////////////////
-
-FMaterialCustomOutputAttributeDefintion::FMaterialCustomOutputAttributeDefintion(
-		const FGuid& InAttributeID, const FString& InAttributeName, const FString& InFunctionName, EMaterialProperty InProperty,
-		EMaterialValueType InValueType, const FVector4& InDefaultValue, EShaderFrequency InShaderFrequency, MaterialAttributeBlendFunction InBlendFunction /*= nullptr*/)
-	: FMaterialAttributeDefintion(InAttributeID, InAttributeName, InProperty, InValueType, InDefaultValue, InShaderFrequency, INDEX_NONE, false, InBlendFunction)
-	, FunctionName(InFunctionName)
-{
-}
-
-//////////////////////////////////////////////////////////////////////////
-FMaterialAttributeDefinitionMap FMaterialAttributeDefinitionMap::GMaterialPropertyAttributesMap;
-
-void FMaterialAttributeDefinitionMap::InitializeAttributeMap()
-{
-	check(!bIsInitialized);
-	bIsInitialized = true;
-	const bool bHideAttribute = true;
-	LLM_SCOPE(ELLMTag::Materials);
-
-	// All types plus default/missing attribute
-	AttributeMap.Empty(MP_MAX + 1);
-
-	// Basic attributes
-	Add(FGuid(0x69B8D336, 0x16ED4D49, 0x9AA49729, 0x2F050F7A), TEXT("BaseColor"),		MP_BaseColor,		MCT_Float3,	FVector4(0,0,0,0),	SF_Pixel);
-	Add(FGuid(0x57C3A161, 0x7F064296, 0xB00B24A5, 0xA496F34C), TEXT("Metallic"),		MP_Metallic,		MCT_Float,	FVector4(0,0,0,0),	SF_Pixel);
-	Add(FGuid(0x9FDAB399, 0x25564CC9, 0x8CD2D572, 0xC12C8FED), TEXT("Specular"),		MP_Specular,		MCT_Float,	FVector4(.5,0,0,0), SF_Pixel);
-	Add(FGuid(0xD1DD967C, 0x4CAD47D3, 0x9E6346FB, 0x08ECF210), TEXT("Roughness"),		MP_Roughness,		MCT_Float,	FVector4(.5,0,0,0), SF_Pixel);
-	Add(FGuid(0x55E2B4FB, 0xC1C54DB2, 0x9F11875F, 0x7231EB1E), TEXT("Anisotropy"),		MP_Anisotropy,		MCT_Float,	FVector4(0,0,0,0),  SF_Pixel);
-	Add(FGuid(0xB769B54D, 0xD08D4440, 0xABC21BA6, 0xCD27D0E2), TEXT("EmissiveColor"),	MP_EmissiveColor,	MCT_Float3,	FVector4(0,0,0,0),	SF_Pixel);
-	Add(FGuid(0xB8F50FBA, 0x2A754EC1, 0x9EF672CF, 0xEB27BF51), TEXT("Opacity"),			MP_Opacity,			MCT_Float,	FVector4(1,0,0,0),	SF_Pixel);
-	Add(FGuid(0x679FFB17, 0x2BB5422C, 0xAD520483, 0x166E0C75), TEXT("OpacityMask"),		MP_OpacityMask,		MCT_Float,	FVector4(1,0,0,0),	SF_Pixel);
-	Add(FGuid(0x0FA2821A, 0x200F4A4A, 0xB719B789, 0xC1259C64), TEXT("Normal"),			MP_Normal,			MCT_Float3,	FVector4(0,0,1,0),	SF_Pixel);
-	Add(FGuid(0xD5F8E9CF, 0xCDC3468D, 0xB10E4465, 0x596A7BBA), TEXT("Tangent"),			MP_Tangent,			MCT_Float3,	FVector4(1,0,0,0),	SF_Pixel);
-
-	// Advanced attributes
-	Add(FGuid(0xF905F895, 0xD5814314, 0x916D2434, 0x8C40CE9E), TEXT("WorldPositionOffset"),		MP_WorldPositionOffset,		MCT_Float3,	FVector4(0,0,0,0),	SF_Vertex);
-	Add(FGuid(0x5B8FC679, 0x51CE4082, 0x9D777BEE, 0xF4F72C44), TEXT("SubsurfaceColor"),			MP_SubsurfaceColor,			MCT_Float3,	FVector4(1,1,1,0),	SF_Pixel);
-	Add(FGuid(0x9E502E69, 0x3C8F48FA, 0x94645CFD, 0x28E5428D), TEXT("ClearCoat"),				MP_CustomData0,				MCT_Float,	FVector4(1,0,0,0),	SF_Pixel);
-	Add(FGuid(0xBE4F2FFD, 0x12FC4296, 0xB0124EEA, 0x12C28D92), TEXT("ClearCoatRoughness"),		MP_CustomData1,				MCT_Float,	FVector4(.1,0,0,0),	SF_Pixel);
-	Add(FGuid(0xE8EBD0AD, 0xB1654CBE, 0xB079C3A8, 0xB39B9F15), TEXT("AmbientOcclusion"),		MP_AmbientOcclusion,		MCT_Float,	FVector4(1,0,0,0),	SF_Pixel);
-	Add(FGuid(0xD0B0FA03, 0x14D74455, 0xA851BAC5, 0x81A0788B), TEXT("Refraction"),				MP_Refraction,				MCT_Float3,	FVector4(1,0,0,0),	SF_Pixel);
-	Add(FGuid(0x0AC97EC3, 0xE3D047BA, 0xB610167D, 0xC4D919FF), TEXT("PixelDepthOffset"),		MP_PixelDepthOffset,		MCT_Float,	FVector4(0,0,0,0),	SF_Pixel);
-	Add(FGuid(0xD9423FFF, 0xD77E4D82, 0x8FF9CF5E, 0x055D1255), TEXT("ShadingModel"),			MP_ShadingModel,			MCT_ShadingModel, FVector4(0, 0, 0, 0), SF_Pixel, INDEX_NONE, false, &CompileShadingModelBlendFunction);
-	Add(FGuid(0x42BDD2E0, 0xBE714189, 0xA0984BC3, 0xDD0BE872), TEXT("SurfaceThickness"),		MP_SurfaceThickness,		MCT_Float,  FVector4(STRATA_LAYER_DEFAULT_THICKNESS_CM, 0, 0, 0), SF_Pixel);
-	Add(FGuid(0x5973A03E, 0x13A74E08, 0x92D0CEDD, 0xF2936CF8), TEXT("FrontMaterial"),			MP_FrontMaterial,			MCT_Strata, FVector4(0,0,0,0),	SF_Pixel, INDEX_NONE, false, &CompileStrataBlendFunction);
-
-	// Used when compiling material with execution pins, which are compiling all attributes together
-	Add(FGuid(0xE0ED040B, 0x82794D93, 0xBD2D59B2, 0xA5BBF41C), TEXT("MaterialAttributes"),		MP_MaterialAttributes,		MCT_MaterialAttributes, FVector4(0,0,0,0), SF_Pixel, INDEX_NONE, bHideAttribute);
-
-	// Texture coordinates
-	Add(FGuid(0xD30EC284, 0xE13A4160, 0x87BB5230, 0x2ED115DC), TEXT("CustomizedUV0"), MP_CustomizedUVs0, MCT_Float2, FVector4(0,0,0,0), SF_Vertex, 0);
-	Add(FGuid(0xC67B093C, 0x2A5249AA, 0xABC97ADE, 0x4A1F49C5), TEXT("CustomizedUV1"), MP_CustomizedUVs1, MCT_Float2, FVector4(0,0,0,0), SF_Vertex, 1);
-	Add(FGuid(0x85C15B24, 0xF3E047CA, 0x85856872, 0x01AE0F4F), TEXT("CustomizedUV2"), MP_CustomizedUVs2, MCT_Float2, FVector4(0,0,0,0), SF_Vertex, 2);
-	Add(FGuid(0x777819DC, 0x31AE4676, 0xB864EF77, 0xB807E873), TEXT("CustomizedUV3"), MP_CustomizedUVs3, MCT_Float2, FVector4(0,0,0,0), SF_Vertex, 3);
-	Add(FGuid(0xDA63B233, 0xDDF44CAD, 0xB93D867B, 0x8DAFDBCC), TEXT("CustomizedUV4"), MP_CustomizedUVs4, MCT_Float2, FVector4(0,0,0,0), SF_Vertex, 4);
-	Add(FGuid(0xC2F52B76, 0x4A034388, 0x89119528, 0x2071B190), TEXT("CustomizedUV5"), MP_CustomizedUVs5, MCT_Float2, FVector4(0,0,0,0), SF_Vertex, 5);
-	Add(FGuid(0x8214A8CA, 0x0CB944CF, 0x9DFD78DB, 0xE48BB55F), TEXT("CustomizedUV6"), MP_CustomizedUVs6, MCT_Float2, FVector4(0,0,0,0), SF_Vertex, 6);
-	Add(FGuid(0xD8F8D01F, 0xC6F74715, 0xA3CFB4FF, 0x9EF51FAC), TEXT("CustomizedUV7"), MP_CustomizedUVs7, MCT_Float2, FVector4(0,0,0,0), SF_Vertex, 7);
-
-	// Lightmass attributes	
-	Add(FGuid(0x68934E1B, 0x70EB411B, 0x86DF5AA5, 0xDF2F626C), TEXT("DiffuseColor"),	MP_DiffuseColor,	MCT_Float3, FVector4(0,0,0,0), SF_Pixel, INDEX_NONE, bHideAttribute);
-	Add(FGuid(0xE89CBD84, 0x62EA48BE, 0x80F88521, 0x2B0C403C), TEXT("SpecularColor"),	MP_SpecularColor,	MCT_Float3, FVector4(0,0,0,0), SF_Pixel, INDEX_NONE, bHideAttribute);
-
-	// Debug attributes
-	Add(FGuid(0x5BF6BA94, 0xA3264629, 0xA253A05B, 0x0EABBB86), TEXT("Missing"), MP_MAX, MCT_Float, FVector4(0,0,0,0), SF_Pixel, INDEX_NONE, bHideAttribute);
-
-	// Removed attributes
-	Add(FGuid(0x2091ECA2, 0xB59248EE, 0x8E2CD578, 0xD371926D), TEXT("WorldDisplacement"), MP_WorldDisplacement_DEPRECATED, MCT_Float3, FVector4(0, 0, 0, 0), SF_Vertex, INDEX_NONE, bHideAttribute);
-	Add(FGuid(0xA0119D44, 0xC456450D, 0x9C39C933, 0x1F72D8D1), TEXT("TessellationMultiplier"), MP_TessellationMultiplier_DEPRECATED, MCT_Float, FVector4(1, 0, 0, 0), SF_Vertex, INDEX_NONE, bHideAttribute);
-
-	// UMaterialExpression custom outputs
-	AddCustomAttribute(FGuid(0xfbd7b46e, 0xb1234824, 0xbde76b23, 0x609f984c), "BentNormal", "GetBentNormal", MCT_Float3, FVector4(0, 0, 1, 0));
-	AddCustomAttribute(FGuid(0xAA3D5C04, 0x16294716, 0xBBDEC869, 0x6A27DD72), "ClearCoatBottomNormal", "ClearCoatBottomNormal", MCT_Float3, FVector4(0, 0, 1, 0));
-	AddCustomAttribute(FGuid(0x8EAB2CB2, 0x73634A24, 0x8CD14F47, 0x3F9C8E55), "CustomEyeTangent", "GetTangentOutput", MCT_Float3, FVector4(0, 0, 0, 0));
-}
-
-
-void FMaterialAttributeDefinitionMap::Add(const FGuid& AttributeID, const FString& AttributeName, EMaterialProperty Property,
-	EMaterialValueType ValueType, const FVector4& DefaultValue, EShaderFrequency ShaderFrequency,
-	int32 TexCoordIndex /*= INDEX_NONE*/, bool bIsHidden /*= false*/, MaterialAttributeBlendFunction BlendFunction /*= nullptr*/)
-{
-	checkf(!AttributeMap.Contains(Property), TEXT("Tried to add duplicate material property."));
-	AttributeMap.Add(Property, FMaterialAttributeDefintion(AttributeID, AttributeName, Property, ValueType, DefaultValue, ShaderFrequency, TexCoordIndex, bIsHidden, BlendFunction));
-	if (!bIsHidden)
-	{
-		OrderedVisibleAttributeList.Add(AttributeID);
-	}
-}
-
-FMaterialAttributeDefintion* FMaterialAttributeDefinitionMap::Find(const FGuid& AttributeID)
-{
-	for (auto& Attribute : CustomAttributes)
-	{
-		if (Attribute.AttributeID == AttributeID)
-		{
-			return &Attribute;
-		}
-	}
-
-	for (auto& Attribute : AttributeMap)
-	{
-		if (Attribute.Value.AttributeID == AttributeID)
-		{
-			return &Attribute.Value;
-		}
-	}
-	
-	UE_LOG(LogMaterial, Warning, TEXT("Failed to find material attribute, AttributeID: %s."), *AttributeID.ToString(EGuidFormats::Digits));
-	return Find(MP_MAX);
-}
-
-FMaterialAttributeDefintion* FMaterialAttributeDefinitionMap::Find(EMaterialProperty Property)
-{
-	if (FMaterialAttributeDefintion* Attribute = AttributeMap.Find(Property))
-	{
-		return Attribute;
-	}
-
-	UE_LOG(LogMaterial, Warning, TEXT("Failed to find material attribute, PropertyType: %i."), (uint32)Property);
-	return Find(MP_MAX);
-}
-
-FText FMaterialAttributeDefinitionMap::GetAttributeOverrideForMaterial(const FGuid& AttributeID, UMaterial* Material)
-{
-	TArray<TKeyValuePair<EMaterialShadingModel, FString>> CustomPinNames;
-	EMaterialProperty Property = GMaterialPropertyAttributesMap.Find(AttributeID)->Property;
-
-	switch (Property)
-	{
-	case MP_EmissiveColor:
-		return Material->IsUIMaterial() ? LOCTEXT("UIOutputColor", "Final Color") : LOCTEXT("EmissiveColor", "Emissive Color");
-	case MP_Opacity:
-		return LOCTEXT("Opacity", "Opacity");
-	case MP_OpacityMask:
-		return LOCTEXT("OpacityMask", "Opacity Mask");
-	case MP_DiffuseColor:
-		return LOCTEXT("DiffuseColor", "Diffuse Color");
-	case MP_SpecularColor:
-		return LOCTEXT("SpecularColor", "Specular Color");
-	case MP_BaseColor:
-		return Material->MaterialDomain == MD_Volume ? LOCTEXT("Albedo", "Albedo") : LOCTEXT("BaseColor", "Base Color");
-	case MP_Metallic:
-		CustomPinNames.Add({MSM_Hair, "Scatter"});
-		CustomPinNames.Add({ MSM_Eye, "Curvature" });
-		return FText::FromString(GetPinNameFromShadingModelField(Material->GetShadingModels(), CustomPinNames, "Metallic"));
-	case MP_Specular:
-		return LOCTEXT("Specular", "Specular");
-	case MP_Roughness:
-		return LOCTEXT("Roughness", "Roughness");
-	case MP_Anisotropy:
-		return LOCTEXT("Anisotropy", "Anisotropy");
-	case MP_Normal:
-		CustomPinNames.Add({MSM_Hair, "Tangent"});
-		return FText::FromString(GetPinNameFromShadingModelField(Material->GetShadingModels(), CustomPinNames, "Normal"));
-	case MP_Tangent:
-		return LOCTEXT("Tangent", "Tangent");
-	case MP_WorldPositionOffset:
-		return Material->IsUIMaterial() ? LOCTEXT("ScreenPosition", "Screen Position") : LOCTEXT("WorldPositionOffset", "World Position Offset");
-	case MP_WorldDisplacement_DEPRECATED:
-		return LOCTEXT("WorldDisplacement", "World Displacement");
-	case MP_TessellationMultiplier_DEPRECATED:
-		return LOCTEXT("TessellationMultiplier", "Tessellation Multiplier");
-	case MP_SubsurfaceColor:
-		if (Material->MaterialDomain == MD_Volume)
-		{
-			return LOCTEXT("Extinction", "Extinction");
-		}
-		CustomPinNames.Add({MSM_Cloth, "Fuzz Color"});
-		return FText::FromString(GetPinNameFromShadingModelField(Material->GetShadingModels(), CustomPinNames, "Subsurface Color"));
-	case MP_CustomData0:	
-		CustomPinNames.Add({ MSM_ClearCoat, "Clear Coat" });
-		CustomPinNames.Add({MSM_Hair, "Backlit"});
-		CustomPinNames.Add({MSM_Cloth, "Cloth"});
-		CustomPinNames.Add({MSM_Eye, "Iris Mask"});
-		CustomPinNames.Add({MSM_SubsurfaceProfile, "Curvature" });
-		return FText::FromString(GetPinNameFromShadingModelField(Material->GetShadingModels(), CustomPinNames, "Custom Data 0"));
-	case MP_CustomData1:
-		CustomPinNames.Add({ MSM_ClearCoat, "Clear Coat Roughness" });
-		CustomPinNames.Add({MSM_Eye, "Iris Distance"});
-		return FText::FromString(GetPinNameFromShadingModelField(Material->GetShadingModels(), CustomPinNames, "Custom Data 1"));
-	case MP_AmbientOcclusion:
-		return LOCTEXT("AmbientOcclusion", "Ambient Occlusion");
-	case MP_Refraction:
-		return LOCTEXT("Refraction", "Refraction");
-	case MP_CustomizedUVs0:
-		return LOCTEXT("CustomizedUV0", "Customized UV 0");
-	case MP_CustomizedUVs1:
-		return LOCTEXT("CustomizedUV1", "Customized UV 1");
-	case MP_CustomizedUVs2:
-		return LOCTEXT("CustomizedUV2", "Customized UV 2");
-	case MP_CustomizedUVs3:
-		return LOCTEXT("CustomizedUV3", "Customized UV 3");
-	case MP_CustomizedUVs4:
-		return LOCTEXT("CustomizedUV4", "Customized UV 4");
-	case MP_CustomizedUVs5:
-		return LOCTEXT("CustomizedUV5", "Customized UV 5");
-	case MP_CustomizedUVs6:
-		return LOCTEXT("CustomizedUV6", "Customized UV 6");
-	case MP_CustomizedUVs7:
-		return LOCTEXT("CustomizedUV7", "Customized UV 7");
-	case MP_PixelDepthOffset:
-		return LOCTEXT("PixelDepthOffset", "Pixel Depth Offset");
-	case MP_ShadingModel:
-		return LOCTEXT("ShadingModel", "Shading Model");
-	case MP_SurfaceThickness:
-		return LOCTEXT("SurfaceThickness", "Surface Thickness");
-	case MP_FrontMaterial:
-		return LOCTEXT("FrontMaterial", "Front Material");
-	case MP_CustomOutput:
-		return FText::FromString(GetAttributeName(AttributeID));
-		
-	}
-	return  LOCTEXT("Missing", "Missing");
-}
-
-FString FMaterialAttributeDefinitionMap::GetPinNameFromShadingModelField(FMaterialShadingModelField InShadingModels, const TArray<TKeyValuePair<EMaterialShadingModel, FString>>& InCustomShadingModelPinNames, const FString& InDefaultPinName) 
-{
-	FString OutPinName;
-	for (const TKeyValuePair<EMaterialShadingModel, FString>& CustomShadingModelPinName : InCustomShadingModelPinNames)
-	{
-		if (InShadingModels.HasShadingModel(CustomShadingModelPinName.Key))
-		{
-			// Add delimiter
-			if (!OutPinName.IsEmpty())
-			{
-				OutPinName.Append(" or ");
-			}
-
-			// Append the name and remove the shading model from the temp field
-			OutPinName.Append(CustomShadingModelPinName.Value);
-			InShadingModels.RemoveShadingModel(CustomShadingModelPinName.Key);
-		}
-	}
-
-	// There are other shading models present, these don't have their own specific name for this pin, so use a default one
-	if (InShadingModels.CountShadingModels() != 0)
-	{
-		// Add delimiter
-		if (!OutPinName.IsEmpty())
-		{
-			OutPinName.Append(" or ");
-		}
-
-		OutPinName.Append(InDefaultPinName);
-	}
-
-	ensure(!OutPinName.IsEmpty());
-	return OutPinName;
-}
-
-void FMaterialAttributeDefinitionMap::AppendDDCKeyString(FString& String)
-{
-	FString& DDCString = GMaterialPropertyAttributesMap.AttributeDDCString;
-
-	if (DDCString.Len() == 0)
-	{
-		FString AttributeIDs;
-
-		for (const auto& Attribute : GMaterialPropertyAttributesMap.AttributeMap)
-		{
-			AttributeIDs += Attribute.Value.AttributeID.ToString(EGuidFormats::Digits);
-		}
-
-		for (const auto& Attribute : GMaterialPropertyAttributesMap.CustomAttributes)
-		{
-			AttributeIDs += Attribute.AttributeID.ToString(EGuidFormats::Digits);
-		}
-
-		FSHA1 HashState;
-		HashState.UpdateWithString(*AttributeIDs, AttributeIDs.Len());
-		HashState.Final();
-
-		FSHAHash Hash;
-		HashState.GetHash(&Hash.Hash[0]);
-		DDCString = Hash.ToString();
-	}
-	else
-	{
-		// TODO: In debug force re-generate DDC string and compare to catch invalid runtime changes
-	}
-
-	String.Append(DDCString);
-}
-
-void FMaterialAttributeDefinitionMap::AddCustomAttribute(const FGuid& AttributeID, const FString& AttributeName, const FString& FunctionName, EMaterialValueType ValueType, const FVector4& DefaultValue, MaterialAttributeBlendFunction BlendFunction /*= nullptr*/)
-{
-	// Make sure that we init CustomAttributes before DDCString is initialized (before first shader load)
-	check(GMaterialPropertyAttributesMap.AttributeDDCString.Len() == 0);
-
-	FMaterialCustomOutputAttributeDefintion UserAttribute(AttributeID, AttributeName, FunctionName, MP_CustomOutput, ValueType, DefaultValue, SF_Pixel, BlendFunction);
-#if DO_CHECK
-	for (auto& Attribute : GMaterialPropertyAttributesMap.AttributeMap)
-	{
-		checkf(Attribute.Value.AttributeID != AttributeID, TEXT("Tried to add duplicate custom output attribute (%s) already in base attributes (%s)."), *AttributeName, *(Attribute.Value.AttributeName));
-	}
-	checkf(!GMaterialPropertyAttributesMap.CustomAttributes.Contains(UserAttribute), TEXT("Tried to add duplicate custom output attribute (%s)."), *AttributeName);
-#endif
-	GMaterialPropertyAttributesMap.CustomAttributes.Add(UserAttribute);
-
-	if (!UserAttribute.bIsHidden)
-	{
-		GMaterialPropertyAttributesMap.OrderedVisibleAttributeList.Add(AttributeID);
-	}
-}
-
-FGuid FMaterialAttributeDefinitionMap::GetCustomAttributeID(const FString& AttributeName)
-{
-	for (auto& Attribute : GMaterialPropertyAttributesMap.CustomAttributes)
-	{
-		if (Attribute.AttributeName == AttributeName)
-		{
-			return Attribute.AttributeID;
-		}
-	}
-
-	return GMaterialPropertyAttributesMap.Find(MP_MAX)->AttributeID;
-}
-
-void FMaterialAttributeDefinitionMap::GetCustomAttributeList(TArray<FMaterialCustomOutputAttributeDefintion>& CustomAttributeList)
-{
-	CustomAttributeList.Empty(GMaterialPropertyAttributesMap.CustomAttributes.Num());
-	for (auto& Attribute : GMaterialPropertyAttributesMap.CustomAttributes)
-	{
-		CustomAttributeList.Add(Attribute);
-	}
-}
-
-void FMaterialAttributeDefinitionMap::GetAttributeNameToIDList(TArray<TPair<FString, FGuid>>& NameToIDList)
-{
-	NameToIDList.Empty(GMaterialPropertyAttributesMap.OrderedVisibleAttributeList.Num());
-	for (const FGuid& AttributeID : GMaterialPropertyAttributesMap.OrderedVisibleAttributeList)
-	{
-		FMaterialAttributeDefintion* Attribute = GMaterialPropertyAttributesMap.Find(AttributeID);
-		NameToIDList.Emplace(Attribute->AttributeName, AttributeID);
-	}
-}
-
 FMaterialResourceMemoryWriter::FMaterialResourceMemoryWriter(FArchive& Ar) :
 	FMemoryWriter(Bytes, Ar.IsPersistent(), false, TEXT("FShaderMapMemoryWriter")),
 	ParentAr(&Ar)
diff --git a/Engine/Source/Runtime/Engine/Public/MaterialShared.h b/Engine/Source/Runtime/Engine/Public/MaterialShared.h
--- a/Engine/Source/Runtime/Engine/Public/MaterialShared.h
+++ b/Engine/Source/Runtime/Engine/Public/MaterialShared.h
@@ -41,6 +41,8 @@
 #include "RHI.h"
 #include "MaterialDomain.h"
 #endif
+#include "Materials/MaterialAttributeDefinitionMap.h"
+#include "Materials/MaterialRenderProxy.h"
 
 #include "Shader/Preshader.h"
 
@@ -2441,329 +2443,7 @@ private:
 	friend class FMaterialHLSLErrorHandler;
 };
 
-/**
- * Cached uniform expression values.
- */
-struct FUniformExpressionCache
-{
-	/** Material uniform buffer. */
-	FUniformBufferRHIRef UniformBuffer;
-	UE_DEPRECATED(5.1, "Local uniform buffers have been deprecated.")
-	FLocalUniformBuffer LocalUniformBuffer;
-	/** Allocated virtual textures, one for each entry in FUniformExpressionSet::VTStacks */
-	TArray<IAllocatedVirtualTexture*> AllocatedVTs;
-	/** Allocated virtual textures that will need destroying during a call to ResetAllocatedVTs() */
-	TArray<IAllocatedVirtualTexture*> OwnedAllocatedVTs;
-	/** Ids of parameter collections needed for rendering. */
-	TArray<FGuid> ParameterCollections;
-	/** Shader map that was used to cache uniform expressions on this material. This is used for debugging, verifying correct behavior and checking if the cache is up to date. */
-	const FMaterialShaderMap* CachedUniformExpressionShaderMap = nullptr;
-
-	/** Destructor. */
-	ENGINE_API ~FUniformExpressionCache();
-
-	void ResetAllocatedVTs();
-};
-
-struct FUniformExpressionCacheContainer
-{
-	inline FUniformExpressionCache& operator[](int32 Index)
-	{
-#if WITH_EDITOR
-		return Elements[Index];
-#else
-		return Elements;
-#endif
-	}
-private:
-#if WITH_EDITOR
-	FUniformExpressionCache Elements[ERHIFeatureLevel::Num];
-#else
-	FUniformExpressionCache Elements;
-#endif
-};
-
-class USubsurfaceProfile;
-class FUniformExpressionCacheAsyncUpdater;
-
-/** Defines a scope to update deferred uniform expression caches using an async task to fill uniform buffers. The scope
- *  attempts to launch async updates immediately, but further async updates can launch within the scope. Only one async
- *  task is launched at a time though. The async task is synced when the scope destructs. The scope enqueues render commands
- *  so it can be used on the game or render threads.
- */
-class ENGINE_API FUniformExpressionCacheAsyncUpdateScope
-{
-public:
-	FUniformExpressionCacheAsyncUpdateScope();
-	~FUniformExpressionCacheAsyncUpdateScope();
-
-	/** Call if a wait is required within the scope. */
-	static void WaitForTask();
-};
-
-/**
- * A material render proxy used by the renderer.
- */
-class ENGINE_API FMaterialRenderProxy : public FRenderResource, public FNoncopyable
-{
-public:
-
-	/** Cached uniform expressions. */
-	mutable FUniformExpressionCacheContainer UniformExpressionCache;
-
-	/** Cached external texture immutable samplers */
-	mutable FImmutableSamplerState ImmutableSamplerState;
-
-	/** Default constructor. */
-	FMaterialRenderProxy(FString InMaterialName);
-
-	/** Destructor. */
-	virtual ~FMaterialRenderProxy();
-
-	UE_DEPRECATED(5.1, "EvaluateUniformExpressions with a command list is deprecated.")
-	void EvaluateUniformExpressions(FUniformExpressionCache& OutUniformExpressionCache, const FMaterialRenderContext& Context, class FRHIComputeCommandList*) const
-	{
-		EvaluateUniformExpressions(OutUniformExpressionCache, Context);
-	}
-
-	/**
-	 * Evaluates uniform expressions and stores them in OutUniformExpressionCache.
-	 * @param OutUniformExpressionCache - The uniform expression cache to build.
-	 * @param MaterialRenderContext - The context for which to cache expressions.
-	 */
-	void EvaluateUniformExpressions(FUniformExpressionCache& OutUniformExpressionCache, const FMaterialRenderContext& Context, FUniformExpressionCacheAsyncUpdater* Updater = nullptr) const;
-
-	/**
-	 * Caches uniform expressions for efficient runtime evaluation.
-	 */
-	void CacheUniformExpressions(bool bRecreateUniformBuffer);
-
-	/** Cancels an in-flight cache operation. */
-	void CancelCacheUniformExpressions();
-
-	/**
-	 * Enqueues a rendering command to cache uniform expressions for efficient runtime evaluation.
-	 * bRecreateUniformBuffer - whether to recreate the material uniform buffer.  
-	 *		This is required if the FMaterial is being recompiled (the uniform buffer layout will change).
-	 *		This should only be done if the calling code is using FMaterialUpdateContext to recreate the rendering state of primitives using this material, since cached mesh commands also cache uniform buffer pointers.
-	 */
-	void CacheUniformExpressions_GameThread(bool bRecreateUniformBuffer);
-
-	/**
-	 * Invalidates the uniform expression cache.
-	 */
-	void InvalidateUniformExpressionCache(bool bRecreateUniformBuffer);
-
-	void UpdateUniformExpressionCacheIfNeeded(ERHIFeatureLevel::Type InFeatureLevel) const;
-
-	
-	/** Returns the FMaterial, without using a fallback if the FMaterial doesn't have a valid shader map. Can return NULL. */
-	virtual const FMaterial* GetMaterialNoFallback(ERHIFeatureLevel::Type InFeatureLevel) const = 0;
-	virtual const FMaterialRenderProxy* GetFallback(ERHIFeatureLevel::Type InFeatureLevel) const = 0;
-
-	// These functions should only be called by the rendering thread.
-
-	/**
-	 * Finds the FMaterial to use for rendering this FMaterialRenderProxy.  Will fall back to a default material if needed due to a content error, or async compilation.
-	 * The returned FMaterial is guaranteed to have a complete shader map, so all relevant shaders should be available
-	 * OutFallbackMaterialRenderProxy - The proxy that corresponds to the returned FMaterial, should be used for further rendering.  May be a fallback material, or 'this' if no fallback was needed
-	 */
-	const FMaterial& GetMaterialWithFallback(ERHIFeatureLevel::Type InFeatureLevel, const FMaterialRenderProxy*& OutFallbackMaterialRenderProxy) const;
-
-	/**
-	 * Finds the FMaterial to use for rendering this FMaterialRenderProxy.  Will fall back to a default material if needed due to a content error, or async compilation.
-	 * Will always return a valid FMaterial, but unlike GetMaterialWithFallback, FMaterial's shader map may be incomplete
-	 */
-	const FMaterial& GetIncompleteMaterialWithFallback(ERHIFeatureLevel::Type InFeatureLevel) const;
-
-	UE_DEPRECATED(4.26, "This function is deprecated. Use GetIncompleteMaterialWithFallback() instead.")
-	inline const FMaterial* GetMaterial(ERHIFeatureLevel::Type InFeatureLevel) const
-	{
-		return &GetIncompleteMaterialWithFallback(InFeatureLevel);
-	}
 
-	virtual UMaterialInterface* GetMaterialInterface() const { return NULL; }
-
-	bool GetVectorValue(const FHashedMaterialParameterInfo& ParameterInfo, FLinearColor* OutValue, const FMaterialRenderContext& Context) const;
-	bool GetScalarValue(const FHashedMaterialParameterInfo& ParameterInfo, float* OutValue, const FMaterialRenderContext& Context) const;
-	bool GetTextureValue(const FHashedMaterialParameterInfo& ParameterInfo, const UTexture** OutValue, const FMaterialRenderContext& Context) const;
-	bool GetTextureValue(const FHashedMaterialParameterInfo& ParameterInfo, const URuntimeVirtualTexture** OutValue, const FMaterialRenderContext& Context) const;
-	bool GetTextureValue(const FHashedMaterialParameterInfo& ParameterInfo, const USparseVolumeTexture** OutValue, const FMaterialRenderContext& Context) const;
-	virtual bool GetParameterValue(EMaterialParameterType Type, const FHashedMaterialParameterInfo& ParameterInfo, FMaterialParameterValue& OutValue, const FMaterialRenderContext& Context) const = 0;
-
-	bool IsDeleted() const
-	{
-		return DeletedFlag != 0;
-	}
-
-	void MarkForGarbageCollection()
-	{
-		MarkedForGarbageCollection = 1;
-	}
-
-	bool IsMarkedForGarbageCollection() const
-	{
-		return MarkedForGarbageCollection != 0;
-	}
-
-	// FRenderResource interface.
-	virtual void InitDynamicRHI() override;
-	virtual void ReleaseDynamicRHI() override;
-	virtual void ReleaseResource() override;
-
-#if WITH_EDITOR
-	static const TSet<FMaterialRenderProxy*>& GetMaterialRenderProxyMap() 
-	{
-		check(!FPlatformProperties::RequiresCookedData());
-		return MaterialRenderProxyMap;
-	}
-
-	static FCriticalSection& GetMaterialRenderProxyMapLock()
-	{
-		return MaterialRenderProxyMapLock;
-	}
-#endif
-
-	void SetSubsurfaceProfileRT(const USubsurfaceProfile* Ptr) { SubsurfaceProfileRT = Ptr; }
-	const USubsurfaceProfile* GetSubsurfaceProfileRT() const { return SubsurfaceProfileRT; }
-
-	static void UpdateDeferredCachedUniformExpressions();
-
-	static inline bool HasDeferredUniformExpressionCacheRequests() 
-	{
-		return DeferredUniformExpressionCacheRequests.Num() > 0;
-	}
-
-	int32 GetExpressionCacheSerialNumber() const { return UniformExpressionCacheSerialNumber; }
-
-	const FString& GetMaterialName() const { return MaterialName; }
-
-private:
-	IAllocatedVirtualTexture* GetPreallocatedVTStack(const FMaterialRenderContext& Context, const FUniformExpressionSet& UniformExpressionSet, const FMaterialVirtualTextureStack& VTStack) const;
-	IAllocatedVirtualTexture* AllocateVTStack(const FMaterialRenderContext& Context, const FUniformExpressionSet& UniformExpressionSet, const FMaterialVirtualTextureStack& VTStack) const;
-
-	virtual void StartCacheUniformExpressions() const {}
-	virtual void FinishCacheUniformExpressions() const {}
-
-	/** 0 if not set, game thread pointer, do not dereference, only for comparison */
-	const USubsurfaceProfile* SubsurfaceProfileRT;
-	FString MaterialName;
-
-	/** Incremented each time UniformExpressionCache is modified */
-	mutable int32 UniformExpressionCacheSerialNumber = 0;
-
-	/** For tracking down a bug accessing a deleted proxy. */
-	mutable int8 MarkedForGarbageCollection : 1;
-	mutable int8 DeletedFlag : 1;
-	mutable int8 ReleaseResourceFlag : 1;
-	/** If any VT producer destroyed callbacks have been registered */
-	mutable int8 HasVirtualTextureCallbacks : 1;
-
-#if WITH_EDITOR
-	/** 
-	 * Tracks all material render proxies in all scenes.
-	 * This is used to propagate new shader maps to materials being used for rendering.
-	 */
-	static TSet<FMaterialRenderProxy*> MaterialRenderProxyMap;
-
-	/**
-	 * Lock that guards the access to the render proxy map
-	 */
-	static FCriticalSection MaterialRenderProxyMapLock;
-#endif
-
-	static TSet<FMaterialRenderProxy*> DeferredUniformExpressionCacheRequests;
-};
-
-/**
- * An material render proxy which overrides the material's Color vector parameter.
- */
-class FColoredMaterialRenderProxy : public FMaterialRenderProxy
-{
-public:
-
-	const FMaterialRenderProxy* const Parent;
-	const FLinearColor Color;
-	FName ColorParamName;
-
-	/** Initialization constructor. */
-	FColoredMaterialRenderProxy(const FMaterialRenderProxy* InParent,const FLinearColor& InColor, FName InColorParamName = NAME_Color):
-		FMaterialRenderProxy(InParent ? InParent->GetMaterialName() : TEXT("FColoredMaterialRenderProxy")),
-		Parent(InParent),
-		Color(InColor),
-		ColorParamName(InColorParamName)
-	{}
-
-	// FMaterialRenderProxy interface.
-	ENGINE_API virtual const FMaterial* GetMaterialNoFallback(ERHIFeatureLevel::Type InFeatureLevel) const override;
-	ENGINE_API virtual const FMaterialRenderProxy* GetFallback(ERHIFeatureLevel::Type InFeatureLevel) const override;
-	ENGINE_API virtual bool GetParameterValue(EMaterialParameterType Type, const FHashedMaterialParameterInfo& ParameterInfo, FMaterialParameterValue& OutValue, const FMaterialRenderContext& Context) const override;
-};
-
-
-/**
- * An material render proxy which overrides the material's Color vector and Texture parameter (mixed together).
- */
-class FColoredTexturedMaterialRenderProxy : public FColoredMaterialRenderProxy
-{
-public:
-
-	const UTexture* Texture;
-	FName TextureParamName;
-
-	/** Initialization constructor. */
-	FColoredTexturedMaterialRenderProxy(const FMaterialRenderProxy* InParent, const FLinearColor& InColor, FName InColorParamName, const UTexture* InTexture, FName InTextureParamName) :
-		FColoredMaterialRenderProxy(InParent, InColor, InColorParamName),
-		Texture(InTexture),
-		TextureParamName(InTextureParamName)
-	{}
-
-	ENGINE_API virtual bool GetParameterValue(EMaterialParameterType Type, const FHashedMaterialParameterInfo& ParameterInfo, FMaterialParameterValue& OutValue, const FMaterialRenderContext& Context) const override;
-};
-
-
-/**
- * A material render proxy which overrides the selection color
- */
-class FOverrideSelectionColorMaterialRenderProxy : public FMaterialRenderProxy
-{
-public:
-
-	const FMaterialRenderProxy* const Parent;
-	const FLinearColor SelectionColor;
-
-	/** Initialization constructor. */
-	FOverrideSelectionColorMaterialRenderProxy(const FMaterialRenderProxy* InParent, const FLinearColor& InSelectionColor) :
-		FMaterialRenderProxy(InParent ? InParent->GetMaterialName() : TEXT("FOverrideSelectionColorMaterialRenderProxy")),
-		Parent(InParent),
-		SelectionColor(FLinearColor(InSelectionColor.R, InSelectionColor.G, InSelectionColor.B, 1))
-	{
-	}
-
-	// FMaterialRenderProxy interface.
-	ENGINE_API virtual const FMaterial* GetMaterialNoFallback(ERHIFeatureLevel::Type InFeatureLevel) const override;
-	ENGINE_API virtual const FMaterialRenderProxy* GetFallback(ERHIFeatureLevel::Type InFeatureLevel) const override;
-	ENGINE_API virtual bool GetParameterValue(EMaterialParameterType Type, const FHashedMaterialParameterInfo& ParameterInfo, FMaterialParameterValue& OutValue, const FMaterialRenderContext& Context) const override;
-};
-
-
-/**
- * An material render proxy which overrides the material's Color and Lightmap resolution vector parameter.
- */
-class FLightingDensityMaterialRenderProxy : public FColoredMaterialRenderProxy
-{
-public:
-	const FVector2D LightmapResolution;
-
-	/** Initialization constructor. */
-	FLightingDensityMaterialRenderProxy(const FMaterialRenderProxy* InParent,const FLinearColor& InColor, const FVector2D& InLightmapResolution) :
-		FColoredMaterialRenderProxy(InParent, InColor), 
-		LightmapResolution(InLightmapResolution)
-	{}
-
-	// FMaterialRenderProxy interface.
-	virtual bool GetParameterValue(EMaterialParameterType Type, const FHashedMaterialParameterInfo& ParameterInfo, FMaterialParameterValue& OutValue, const FMaterialRenderContext& Context) const override;
-};
 
 /**
  * Implementation of the FMaterial interface for a UMaterial or UMaterialInstance.
@@ -3054,247 +2734,6 @@ ENGINE_API bool DoesMaterialUseTexture(const UMaterialInterface* Material,const
 ENGINE_API void DoMaterialAttributeReorder(FExpressionInput* Input, const FPackageFileVersion&, int32 RenderObjVer, int32 UE5MainVer);
 #endif // WITH_EDITORONLY_DATA
 
-/**
- * Custom attribute blend functions
- */
-typedef int32 (*MaterialAttributeBlendFunction)(FMaterialCompiler* Compiler, int32 A, int32 B, int32 Alpha);
-
-/**
- * Attribute data describing a material property
- */
-class FMaterialAttributeDefintion
-{
-public:
-	FMaterialAttributeDefintion(const FGuid& InGUID, const FString& AttributeName, EMaterialProperty InProperty,
-		EMaterialValueType InValueType, const FVector4& InDefaultValue, EShaderFrequency InShaderFrequency,
-		int32 InTexCoordIndex = INDEX_NONE, bool bInIsHidden = false, MaterialAttributeBlendFunction InBlendFunction = nullptr);
-
-	ENGINE_API int32 CompileDefaultValue(FMaterialCompiler* Compiler);
-
-	bool operator==(const FMaterialAttributeDefintion& Other) const
-	{
-		return (AttributeID == Other.AttributeID);
-	}
-
-	FGuid				AttributeID;
-	FVector4			DefaultValue;
-	FString				AttributeName;
-	EMaterialProperty	Property;	
-	EMaterialValueType	ValueType;
-	EShaderFrequency	ShaderFrequency;
-	int32				TexCoordIndex;
-
-	// Optional function pointer for custom blend behavior
-	MaterialAttributeBlendFunction BlendFunction;
-
-	// Hidden from auto-generated lists but valid for manual material creation
-	bool				bIsHidden;
-};
-
-/**
- * Attribute data describing a material property used for a custom output
- */
-class FMaterialCustomOutputAttributeDefintion : public FMaterialAttributeDefintion
-{
-public:
-	FMaterialCustomOutputAttributeDefintion(const FGuid& InGUID, const FString& InAttributeName, const FString& InFunctionName, EMaterialProperty InProperty,
-		EMaterialValueType InValueType, const FVector4& InDefaultValue, EShaderFrequency InShaderFrequency, MaterialAttributeBlendFunction InBlendFunction = nullptr);
-
-	bool operator==(const FMaterialCustomOutputAttributeDefintion& Other) const
-	{
-		return (AttributeID == Other.AttributeID);
-	}
-
-	// Name of function used to access attribute in shader code
-	FString							FunctionName;
-};
-
-/**
- * Material property to attribute data mappings
- */
-class FMaterialAttributeDefinitionMap
-{
-public:
-	FMaterialAttributeDefinitionMap()
-	: AttributeDDCString(TEXT(""))
-	, bIsInitialized(false)
-	{
-		AttributeMap.Empty(MP_MAX);
-		InitializeAttributeMap();
-	}
-
-	/** Compiles the default expression for a material attribute */
-	ENGINE_API static int32 CompileDefaultExpression(FMaterialCompiler* Compiler, EMaterialProperty Property)
-	{
-		FMaterialAttributeDefintion* Attribute = GMaterialPropertyAttributesMap.Find(Property);
-		return Attribute->CompileDefaultValue(Compiler);
-	}
-
-	/** Compiles the default expression for a material attribute */
-	ENGINE_API static int32 CompileDefaultExpression(FMaterialCompiler* Compiler, const FGuid& AttributeID)
-	{
-		FMaterialAttributeDefintion* Attribute = GMaterialPropertyAttributesMap.Find(AttributeID);
-		return Attribute->CompileDefaultValue(Compiler);
-	}
-
-	/** Returns the display name of a material attribute */
-	ENGINE_API static const FString& GetAttributeName(EMaterialProperty Property)
-	{
-		FMaterialAttributeDefintion* Attribute = GMaterialPropertyAttributesMap.Find(Property);
-		return Attribute->AttributeName;
-	}
-
-	/** Returns the display name of a material attribute */
-	ENGINE_API static const FString& GetAttributeName(const FGuid& AttributeID)
-	{
-		FMaterialAttributeDefintion* Attribute = GMaterialPropertyAttributesMap.Find(AttributeID);
-		return Attribute->AttributeName;
-	}
-
-	/** Returns the display name of a material attribute, accounting for overrides based on properties of a given material */
-	ENGINE_API static FText GetDisplayNameForMaterial(EMaterialProperty Property, UMaterial* Material)
-	{
-		if (!Material)
-		{
-			return FText::FromString(GetAttributeName(Property));
-		}
-
-		FMaterialAttributeDefintion* Attribute = GMaterialPropertyAttributesMap.Find(Property);
-		return GetAttributeOverrideForMaterial(Attribute->AttributeID, Material);
-	}
-
-	/** Returns the display name of a material attribute, accounting for overrides based on properties of a given material */
-	ENGINE_API static FText GetDisplayNameForMaterial(const FGuid& AttributeID, UMaterial* Material)
-	{
-		if (!Material)
-		{
-			return FText::FromString(GetAttributeName(AttributeID));
-		}
-
-		FMaterialAttributeDefintion* Attribute = GMaterialPropertyAttributesMap.Find(AttributeID);
-		return GetAttributeOverrideForMaterial(AttributeID, Material);
-	}
-
-	/** Returns the value type of a material attribute */
-	ENGINE_API static EMaterialValueType GetValueType(EMaterialProperty Property)
-	{
-		FMaterialAttributeDefintion* Attribute = GMaterialPropertyAttributesMap.Find(Property);
-		return Attribute->ValueType;
-	}
-
-	/** Returns the value type of a material attribute */
-	ENGINE_API static EMaterialValueType GetValueType(const FGuid& AttributeID)
-	{
-		FMaterialAttributeDefintion* Attribute = GMaterialPropertyAttributesMap.Find(AttributeID);
-		return Attribute->ValueType;
-	}
-
-	/** Returns the default value of a material property */
-	ENGINE_API static FVector4f GetDefaultValue(EMaterialProperty Property)
-	{
-		FMaterialAttributeDefintion* Attribute = GMaterialPropertyAttributesMap.Find(Property);
-		return (FVector4f)Attribute->DefaultValue;
-	}
-
-	/** Returns the default value of a material attribute */
-	ENGINE_API static FVector4f GetDefaultValue(const FGuid& AttributeID)
-	{
-		FMaterialAttributeDefintion* Attribute = GMaterialPropertyAttributesMap.Find(AttributeID);
-		return (FVector4f)Attribute->DefaultValue;
-	}
-	
-	/** Returns the shader frequency of a material attribute */
-	ENGINE_API static EShaderFrequency GetShaderFrequency(EMaterialProperty Property)
-	{
-		FMaterialAttributeDefintion* Attribute = GMaterialPropertyAttributesMap.Find(Property);
-		return Attribute->ShaderFrequency;
-	}
-
-	/** Returns the shader frequency of a material attribute */
-	ENGINE_API static EShaderFrequency GetShaderFrequency(const FGuid& AttributeID)
-	{
-		FMaterialAttributeDefintion* Attribute = GMaterialPropertyAttributesMap.Find(AttributeID);
-		return Attribute->ShaderFrequency;
-	}
-
-	/** Returns the attribute ID for a matching material property */
-	ENGINE_API static FGuid GetID(EMaterialProperty Property)
-	{
-		FMaterialAttributeDefintion* Attribute = GMaterialPropertyAttributesMap.Find(Property);
-		return Attribute->AttributeID;
-	}
-
-	/** Returns a the material property matching the specified attribute AttributeID */
-	ENGINE_API static EMaterialProperty GetProperty(const FGuid& AttributeID)
-	{
-		if (FMaterialAttributeDefintion* Attribute = GMaterialPropertyAttributesMap.Find(AttributeID))
-		{
-			return Attribute->Property;
-		}
-		return MP_MAX;
-	}
-
-	/** Returns the custom blend function of a material attribute */
-	ENGINE_API static MaterialAttributeBlendFunction GetBlendFunction(const FGuid& AttributeID)
-	{
-		FMaterialAttributeDefintion* Attribute = GMaterialPropertyAttributesMap.Find(AttributeID);
-		return Attribute->BlendFunction;
-	}
-
-	/** Returns a default attribute AttributeID */
-	ENGINE_API static FGuid GetDefaultID()
-	{
-		return GMaterialPropertyAttributesMap.Find(MP_MAX)->AttributeID;
-	}
-
-	/** Appends a hash of the property map intended for use with the DDC key */
-	ENGINE_API static void AppendDDCKeyString(FString& String);
-
-	/** Appends a new attribute definition to the custom output list */
-	ENGINE_API static void AddCustomAttribute(const FGuid& AttributeID, const FString& AttributeName, const FString& FunctionName, EMaterialValueType ValueType, const FVector4& DefaultValue, MaterialAttributeBlendFunction BlendFunction = nullptr);
-
-	/** Returns the first custom attribute ID that has the specificed attribute name */
-	ENGINE_API static FGuid GetCustomAttributeID(const FString& AttributeName);
-
-	/** Returns a list of registered custom attributes */
-	ENGINE_API static void GetCustomAttributeList(TArray<FMaterialCustomOutputAttributeDefintion>& CustomAttributeList);
-
-	ENGINE_API static const TArray<FGuid>& GetOrderedVisibleAttributeList()
-	{
-		return GMaterialPropertyAttributesMap.OrderedVisibleAttributeList;
-	}
-
-private:
-	// Customization class for displaying data in the material editor
-	friend class FMaterialAttributePropertyDetails;
-
-	/** Returns a list of display names and their associated GUIDs for material properties */
-	ENGINE_API static void GetAttributeNameToIDList(TArray<TPair<FString, FGuid>>& NameToIDList);
-
-	// Internal map management
-	void InitializeAttributeMap();
-
-	void Add(const FGuid& AttributeID, const FString& AttributeName, EMaterialProperty Property,
-		EMaterialValueType ValueType, const FVector4& DefaultValue, EShaderFrequency ShaderFrequency,
-		int32 TexCoordIndex = INDEX_NONE, bool bIsHidden = false, MaterialAttributeBlendFunction BlendFunction = nullptr);
-
-	ENGINE_API FMaterialAttributeDefintion* Find(const FGuid& AttributeID);
-	ENGINE_API FMaterialAttributeDefintion* Find(EMaterialProperty Property);
-
-	// Helper functions to determine display name based on shader model, material domain, etc.
-	ENGINE_API static FText GetAttributeOverrideForMaterial(const FGuid& AttributeID, UMaterial* Material);
-	ENGINE_API static FString GetPinNameFromShadingModelField(FMaterialShadingModelField InShadingModels, const TArray<TKeyValuePair<EMaterialShadingModel, FString>>& InCustomShadingModelPinNames, const FString& InDefaultPinName);
-
-	ENGINE_API static FMaterialAttributeDefinitionMap GMaterialPropertyAttributesMap;
-
-	TMap<EMaterialProperty, FMaterialAttributeDefintion>	AttributeMap; // Fixed map of compile-time definitions
-	TArray<FMaterialCustomOutputAttributeDefintion>			CustomAttributes; // Array of custom output definitions
-	TArray<FGuid>											OrderedVisibleAttributeList; // List used for consistency with e.g. combobox filling
-
-	FString													AttributeDDCString;
-	bool bIsInitialized;
-};
-
 struct FMaterialResourceLocOnDisk
 {
 	// Relative offset to package (uasset/umap + uexp) beginning
diff --git a/Engine/Source/Runtime/Engine/Public/Materials/MaterialAttributeDefinitionMap.h b/Engine/Source/Runtime/Engine/Public/Materials/MaterialAttributeDefinitionMap.h
new file mode 100644
--- /dev/null
+++ b/Engine/Source/Runtime/Engine/Public/Materials/MaterialAttributeDefinitionMap.h
@@ -0,0 +1,177 @@
+// Copyright Epic Games, Inc. All Rights Reserved.
+
+#pragma once
+
+#include "CoreTypes.h"
+#include "Containers/UnrealString.h"
+#include "MaterialValueType.h"
+#include "Math/Vector4.h"
+#include "Misc/Guid.h"
+
+enum EMaterialProperty : int;
+enum EMaterialShadingModel : int;
+enum EShaderFrequency : uint8;
+
+class FMaterialCompiler;
+class UMaterial;
+
+struct FMaterialShadingModelField;
+
+/**
+ * Custom attribute blend functions
+ */
+typedef int32(*MaterialAttributeBlendFunction)(FMaterialCompiler* Compiler, int32 A, int32 B, int32 Alpha);
+
+/**
+ * Attribute data describing a material property
+ */
+class FMaterialAttributeDefintion
+{
+public:
+	FMaterialAttributeDefintion(const FGuid& InGUID, const FString& AttributeName, EMaterialProperty InProperty,
+		EMaterialValueType InValueType, const FVector4& InDefaultValue, EShaderFrequency InShaderFrequency,
+		int32 InTexCoordIndex = INDEX_NONE, bool bInIsHidden = false, MaterialAttributeBlendFunction InBlendFunction = nullptr);
+
+	ENGINE_API int32 CompileDefaultValue(FMaterialCompiler* Compiler);
+
+	bool operator==(const FMaterialAttributeDefintion& Other) const
+	{
+		return (AttributeID == Other.AttributeID);
+	}
+
+	FGuid				AttributeID;
+	FVector4			DefaultValue;
+	FString				AttributeName;
+	EMaterialProperty	Property;
+	EMaterialValueType	ValueType;
+	EShaderFrequency	ShaderFrequency;
+	int32				TexCoordIndex;
+
+	// Optional function pointer for custom blend behavior
+	MaterialAttributeBlendFunction BlendFunction;
+
+	// Hidden from auto-generated lists but valid for manual material creation
+	bool				bIsHidden;
+};
+
+/**
+ * Attribute data describing a material property used for a custom output
+ */
+class FMaterialCustomOutputAttributeDefintion : public FMaterialAttributeDefintion
+{
+public:
+	FMaterialCustomOutputAttributeDefintion(const FGuid& InGUID, const FString& InAttributeName, const FString& InFunctionName, EMaterialProperty InProperty,
+		EMaterialValueType InValueType, const FVector4& InDefaultValue, EShaderFrequency InShaderFrequency, MaterialAttributeBlendFunction InBlendFunction = nullptr);
+
+	bool operator==(const FMaterialCustomOutputAttributeDefintion& Other) const
+	{
+		return (AttributeID == Other.AttributeID);
+	}
+
+	// Name of function used to access attribute in shader code
+	FString							FunctionName;
+};
+
+/**
+ * Material property to attribute data mappings
+ */
+class ENGINE_API FMaterialAttributeDefinitionMap
+{
+public:
+	FMaterialAttributeDefinitionMap();
+
+	/** Compiles the default expression for a material attribute */
+	static int32 CompileDefaultExpression(FMaterialCompiler* Compiler, EMaterialProperty Property);
+
+	/** Compiles the default expression for a material attribute */
+	static int32 CompileDefaultExpression(FMaterialCompiler* Compiler, const FGuid& AttributeID);
+
+	/** Returns the display name of a material attribute */
+	static const FString& GetAttributeName(EMaterialProperty Property);
+
+	/** Returns the display name of a material attribute */
+	static const FString& GetAttributeName(const FGuid& AttributeID);
+
+	/** Returns the display name of a material attribute, accounting for overrides based on properties of a given material */
+	static FText GetDisplayNameForMaterial(EMaterialProperty Property, UMaterial* Material);
+
+	/** Returns the display name of a material attribute, accounting for overrides based on properties of a given material */
+	static FText GetDisplayNameForMaterial(const FGuid& AttributeID, UMaterial* Material);
+
+	/** Returns the value type of a material attribute */
+	static EMaterialValueType GetValueType(EMaterialProperty Property);
+
+	/** Returns the value type of a material attribute */
+	static EMaterialValueType GetValueType(const FGuid& AttributeID);
+
+	/** Returns the default value of a material property */
+	static FVector4f GetDefaultValue(EMaterialProperty Property);
+
+	/** Returns the default value of a material attribute */
+	static FVector4f GetDefaultValue(const FGuid& AttributeID);
+
+	/** Returns the shader frequency of a material attribute */
+	static EShaderFrequency GetShaderFrequency(EMaterialProperty Property);
+
+	/** Returns the shader frequency of a material attribute */
+	static EShaderFrequency GetShaderFrequency(const FGuid& AttributeID);
+
+	/** Returns the attribute ID for a matching material property */
+	static FGuid GetID(EMaterialProperty Property);
+
+	/** Returns a the material property matching the specified attribute AttributeID */
+	static EMaterialProperty GetProperty(const FGuid& AttributeID);
+
+	/** Returns the custom blend function of a material attribute */
+	static MaterialAttributeBlendFunction GetBlendFunction(const FGuid& AttributeID);
+
+	/** Returns a default attribute AttributeID */
+	static FGuid GetDefaultID();
+
+	/** Appends a hash of the property map intended for use with the DDC key */
+	static void AppendDDCKeyString(FString& String);
+
+	/** Appends a new attribute definition to the custom output list */
+	static void AddCustomAttribute(const FGuid& AttributeID, const FString& AttributeName, const FString& FunctionName, EMaterialValueType ValueType, const FVector4& DefaultValue, MaterialAttributeBlendFunction BlendFunction = nullptr);
+
+	/** Returns the first custom attribute ID that has the specificed attribute name */
+	static FGuid GetCustomAttributeID(const FString& AttributeName);
+
+	/** Returns a list of registered custom attributes */
+	static void GetCustomAttributeList(TArray<FMaterialCustomOutputAttributeDefintion>& CustomAttributeList);
+
+	static const TArray<FGuid>& GetOrderedVisibleAttributeList()
+	{
+		return GMaterialPropertyAttributesMap.OrderedVisibleAttributeList;
+	}
+
+private:
+	// Customization class for displaying data in the material editor
+	friend class FMaterialAttributePropertyDetails;
+
+	/** Returns a list of display names and their associated GUIDs for material properties */
+	static void GetAttributeNameToIDList(TArray<TPair<FString, FGuid>>& NameToIDList);
+
+	// Internal map management
+	void InitializeAttributeMap();
+
+	void Add(const FGuid& AttributeID, const FString& AttributeName, EMaterialProperty Property,
+		EMaterialValueType ValueType, const FVector4& DefaultValue, EShaderFrequency ShaderFrequency,
+		int32 TexCoordIndex = INDEX_NONE, bool bIsHidden = false, MaterialAttributeBlendFunction BlendFunction = nullptr);
+
+	FMaterialAttributeDefintion* Find(const FGuid& AttributeID);
+	FMaterialAttributeDefintion* Find(EMaterialProperty Property);
+
+	// Helper functions to determine display name based on shader model, material domain, etc.
+	static FText GetAttributeOverrideForMaterial(const FGuid& AttributeID, UMaterial* Material);
+	static FString GetPinNameFromShadingModelField(FMaterialShadingModelField InShadingModels, const TArray<TKeyValuePair<EMaterialShadingModel, FString>>& InCustomShadingModelPinNames, const FString& InDefaultPinName);
+
+	static FMaterialAttributeDefinitionMap GMaterialPropertyAttributesMap;
+
+	TMap<EMaterialProperty, FMaterialAttributeDefintion>	AttributeMap; // Fixed map of compile-time definitions
+	TArray<FMaterialCustomOutputAttributeDefintion>			CustomAttributes; // Array of custom output definitions
+	TArray<FGuid>											OrderedVisibleAttributeList; // List used for consistency with e.g. combobox filling
+
+	FString													AttributeDDCString;
+	bool bIsInitialized;
+};
diff --git a/Engine/Source/Runtime/Engine/Public/Materials/MaterialRenderProxy.h b/Engine/Source/Runtime/Engine/Public/Materials/MaterialRenderProxy.h
new file mode 100644
--- /dev/null
+++ b/Engine/Source/Runtime/Engine/Public/Materials/MaterialRenderProxy.h
@@ -0,0 +1,318 @@
+// Copyright Epic Games, Inc. All Rights Reserved.
+
+#pragma once
+
+#include "CoreTypes.h"
+#include "RenderResource.h"
+#include "RHIImmutableSamplerState.h"
+
+enum class EMaterialParameterType : uint8;
+
+class FMaterial;
+class FMaterialRenderProxy;
+class FMaterialShaderMap;
+class FMaterialVirtualTextureStack;
+class FRHIComputeCommandList;
+class FUniformExpressionCacheAsyncUpdater;
+class FUniformExpressionSet;
+class IAllocatedVirtualTexture;
+class UMaterialInterface;
+class URuntimeVirtualTexture;
+class USubsurfaceProfile;
+class UTexture;
+
+struct FMaterialParameterValue;
+struct FMaterialRenderContext;
+
+struct FMemoryImageMaterialParameterInfo;
+using FHashedMaterialParameterInfo = FMemoryImageMaterialParameterInfo;
+
+
+/**
+ * Cached uniform expression values.
+ */
+struct FUniformExpressionCache
+{
+	/** Material uniform buffer. */
+	FUniformBufferRHIRef UniformBuffer;
+	/** Allocated virtual textures, one for each entry in FUniformExpressionSet::VTStacks */
+	TArray<IAllocatedVirtualTexture*> AllocatedVTs;
+	/** Allocated virtual textures that will need destroying during a call to ResetAllocatedVTs() */
+	TArray<IAllocatedVirtualTexture*> OwnedAllocatedVTs;
+	/** Ids of parameter collections needed for rendering. */
+	TArray<FGuid> ParameterCollections;
+	/** Shader map that was used to cache uniform expressions on this material. This is used for debugging, verifying correct behavior and checking if the cache is up to date. */
+	const FMaterialShaderMap* CachedUniformExpressionShaderMap = nullptr;
+
+	/** Destructor. */
+	ENGINE_API ~FUniformExpressionCache();
+
+	void ResetAllocatedVTs();
+};
+
+struct FUniformExpressionCacheContainer
+{
+	inline FUniformExpressionCache& operator[](int32 Index)
+	{
+#if WITH_EDITOR
+		return Elements[Index];
+#else
+		return Elements;
+#endif
+	}
+private:
+#if WITH_EDITOR
+	FUniformExpressionCache Elements[ERHIFeatureLevel::Num];
+#else
+	FUniformExpressionCache Elements;
+#endif
+};
+
+/** Defines a scope to update deferred uniform expression caches using an async task to fill uniform buffers. The scope
+ *  attempts to launch async updates immediately, but further async updates can launch within the scope. Only one async
+ *  task is launched at a time though. The async task is synced when the scope destructs. The scope enqueues render commands
+ *  so it can be used on the game or render threads.
+ */
+class ENGINE_API FUniformExpressionCacheAsyncUpdateScope
+{
+public:
+	FUniformExpressionCacheAsyncUpdateScope();
+	~FUniformExpressionCacheAsyncUpdateScope();
+
+	/** Call if a wait is required within the scope. */
+	static void WaitForTask();
+};
+
+/**
+ * A material render proxy used by the renderer.
+ */
+class ENGINE_API FMaterialRenderProxy : public FRenderResource, public FNoncopyable
+{
+public:
+
+	/** Cached uniform expressions. */
+	mutable FUniformExpressionCacheContainer UniformExpressionCache;
+
+	/** Cached external texture immutable samplers */
+	mutable FImmutableSamplerState ImmutableSamplerState;
+
+	/** Default constructor. */
+	FMaterialRenderProxy(FString InMaterialName);
+
+	/** Destructor. */
+	virtual ~FMaterialRenderProxy();
+
+	UE_DEPRECATED(5.1, "EvaluateUniformExpressions with a command list is deprecated.")
+	void EvaluateUniformExpressions(FUniformExpressionCache& OutUniformExpressionCache, const FMaterialRenderContext& Context, class FRHIComputeCommandList*) const
+	{
+		EvaluateUniformExpressions(OutUniformExpressionCache, Context);
+	}
+
+	/**
+	 * Evaluates uniform expressions and stores them in OutUniformExpressionCache.
+	 * @param OutUniformExpressionCache - The uniform expression cache to build.
+	 * @param MaterialRenderContext - The context for which to cache expressions.
+	 */
+	void EvaluateUniformExpressions(FUniformExpressionCache& OutUniformExpressionCache, const FMaterialRenderContext& Context, FUniformExpressionCacheAsyncUpdater* Updater = nullptr) const;
+
+	/**
+	 * Caches uniform expressions for efficient runtime evaluation.
+	 */
+	void CacheUniformExpressions(bool bRecreateUniformBuffer);
+
+	/** Cancels an in-flight cache operation. */
+	void CancelCacheUniformExpressions();
+
+	/**
+	 * Enqueues a rendering command to cache uniform expressions for efficient runtime evaluation.
+	 * bRecreateUniformBuffer - whether to recreate the material uniform buffer.
+	 *		This is required if the FMaterial is being recompiled (the uniform buffer layout will change).
+	 *		This should only be done if the calling code is using FMaterialUpdateContext to recreate the rendering state of primitives using this material, since cached mesh commands also cache uniform buffer pointers.
+	 */
+	void CacheUniformExpressions_GameThread(bool bRecreateUniformBuffer);
+
+	/**
+	 * Invalidates the uniform expression cache.
+	 */
+	void InvalidateUniformExpressionCache(bool bRecreateUniformBuffer);
+
+	void UpdateUniformExpressionCacheIfNeeded(ERHIFeatureLevel::Type InFeatureLevel) const;
+
+
+	/** Returns the FMaterial, without using a fallback if the FMaterial doesn't have a valid shader map. Can return NULL. */
+	virtual const FMaterial* GetMaterialNoFallback(ERHIFeatureLevel::Type InFeatureLevel) const = 0;
+	virtual const FMaterialRenderProxy* GetFallback(ERHIFeatureLevel::Type InFeatureLevel) const = 0;
+
+	// These functions should only be called by the rendering thread.
+
+	/**
+	 * Finds the FMaterial to use for rendering this FMaterialRenderProxy.  Will fall back to a default material if needed due to a content error, or async compilation.
+	 * The returned FMaterial is guaranteed to have a complete shader map, so all relevant shaders should be available
+	 * OutFallbackMaterialRenderProxy - The proxy that corresponds to the returned FMaterial, should be used for further rendering.  May be a fallback material, or 'this' if no fallback was needed
+	 */
+	const FMaterial& GetMaterialWithFallback(ERHIFeatureLevel::Type InFeatureLevel, const FMaterialRenderProxy*& OutFallbackMaterialRenderProxy) const;
+
+	/**
+	 * Finds the FMaterial to use for rendering this FMaterialRenderProxy.  Will fall back to a default material if needed due to a content error, or async compilation.
+	 * Will always return a valid FMaterial, but unlike GetMaterialWithFallback, FMaterial's shader map may be incomplete
+	 */
+	const FMaterial& GetIncompleteMaterialWithFallback(ERHIFeatureLevel::Type InFeatureLevel) const;
+
+	virtual UMaterialInterface* GetMaterialInterface() const { return NULL; }
+
+	bool GetVectorValue(const FHashedMaterialParameterInfo& ParameterInfo, FLinearColor* OutValue, const FMaterialRenderContext& Context) const;
+	bool GetScalarValue(const FHashedMaterialParameterInfo& ParameterInfo, float* OutValue, const FMaterialRenderContext& Context) const;
+	bool GetTextureValue(const FHashedMaterialParameterInfo& ParameterInfo, const UTexture** OutValue, const FMaterialRenderContext& Context) const;
+	bool GetTextureValue(const FHashedMaterialParameterInfo& ParameterInfo, const URuntimeVirtualTexture** OutValue, const FMaterialRenderContext& Context) const;
+	virtual bool GetParameterValue(EMaterialParameterType Type, const FHashedMaterialParameterInfo& ParameterInfo, FMaterialParameterValue& OutValue, const FMaterialRenderContext& Context) const = 0;
+
+	bool IsDeleted() const
+	{
+		return DeletedFlag != 0;
+	}
+
+	void MarkForGarbageCollection()
+	{
+		MarkedForGarbageCollection = 1;
+	}
+
+	bool IsMarkedForGarbageCollection() const
+	{
+		return MarkedForGarbageCollection != 0;
+	}
+
+	// FRenderResource interface.
+	virtual void InitDynamicRHI() override;
+	virtual void ReleaseDynamicRHI() override;
+	virtual void ReleaseResource() override;
+
+#if WITH_EDITOR
+	static const TSet<FMaterialRenderProxy*>& GetMaterialRenderProxyMap()
+	{
+		check(!FPlatformProperties::RequiresCookedData());
+		return MaterialRenderProxyMap;
+	}
+
+	static FCriticalSection& GetMaterialRenderProxyMapLock()
+	{
+		return MaterialRenderProxyMapLock;
+	}
+#endif
+
+	void SetSubsurfaceProfileRT(const USubsurfaceProfile* Ptr) { SubsurfaceProfileRT = Ptr; }
+	const USubsurfaceProfile* GetSubsurfaceProfileRT() const { return SubsurfaceProfileRT; }
+
+	static void UpdateDeferredCachedUniformExpressions();
+
+	static inline bool HasDeferredUniformExpressionCacheRequests()
+	{
+		return DeferredUniformExpressionCacheRequests.Num() > 0;
+	}
+
+	int32 GetExpressionCacheSerialNumber() const { return UniformExpressionCacheSerialNumber; }
+
+	const FString& GetMaterialName() const { return MaterialName; }
+
+private:
+	IAllocatedVirtualTexture* GetPreallocatedVTStack(const FMaterialRenderContext& Context, const FUniformExpressionSet& UniformExpressionSet, const FMaterialVirtualTextureStack& VTStack) const;
+	IAllocatedVirtualTexture* AllocateVTStack(const FMaterialRenderContext& Context, const FUniformExpressionSet& UniformExpressionSet, const FMaterialVirtualTextureStack& VTStack) const;
+
+	virtual void StartCacheUniformExpressions() const {}
+	virtual void FinishCacheUniformExpressions() const {}
+
+	/** 0 if not set, game thread pointer, do not dereference, only for comparison */
+	const USubsurfaceProfile* SubsurfaceProfileRT;
+	FString MaterialName;
+
+	/** Incremented each time UniformExpressionCache is modified */
+	mutable int32 UniformExpressionCacheSerialNumber = 0;
+
+	/** For tracking down a bug accessing a deleted proxy. */
+	mutable int8 MarkedForGarbageCollection : 1;
+	mutable int8 DeletedFlag : 1;
+	mutable int8 ReleaseResourceFlag : 1;
+	/** If any VT producer destroyed callbacks have been registered */
+	mutable int8 HasVirtualTextureCallbacks : 1;
+
+#if WITH_EDITOR
+	/**
+	 * Tracks all material render proxies in all scenes.
+	 * This is used to propagate new shader maps to materials being used for rendering.
+	 */
+	static TSet<FMaterialRenderProxy*> MaterialRenderProxyMap;
+
+	/**
+	 * Lock that guards the access to the render proxy map
+	 */
+	static FCriticalSection MaterialRenderProxyMapLock;
+#endif
+
+	static TSet<FMaterialRenderProxy*> DeferredUniformExpressionCacheRequests;
+};
+
+/**
+ * An material render proxy which overrides the material's Color vector parameter.
+ */
+class ENGINE_API FColoredMaterialRenderProxy : public FMaterialRenderProxy
+{
+public:
+	const FMaterialRenderProxy* const Parent;
+	const FLinearColor Color;
+	FName ColorParamName;
+
+	/** Initialization constructor. */
+	FColoredMaterialRenderProxy(const FMaterialRenderProxy* InParent, const FLinearColor& InColor, FName InColorParamName = NAME_Color);
+
+	// FMaterialRenderProxy interface.
+	virtual const FMaterial* GetMaterialNoFallback(ERHIFeatureLevel::Type InFeatureLevel) const override;
+	virtual const FMaterialRenderProxy* GetFallback(ERHIFeatureLevel::Type InFeatureLevel) const override;
+	virtual bool GetParameterValue(EMaterialParameterType Type, const FHashedMaterialParameterInfo& ParameterInfo, FMaterialParameterValue& OutValue, const FMaterialRenderContext& Context) const override;
+};
+
+/**
+ * An material render proxy which overrides the material's Color vector and Texture parameter (mixed together).
+ */
+class ENGINE_API FColoredTexturedMaterialRenderProxy : public FColoredMaterialRenderProxy
+{
+public:
+	const UTexture* Texture;
+	FName TextureParamName;
+
+	/** Initialization constructor. */
+	FColoredTexturedMaterialRenderProxy(const FMaterialRenderProxy* InParent, const FLinearColor& InColor, FName InColorParamName, const UTexture* InTexture, FName InTextureParamName);
+
+	virtual bool GetParameterValue(EMaterialParameterType Type, const FHashedMaterialParameterInfo& ParameterInfo, FMaterialParameterValue& OutValue, const FMaterialRenderContext& Context) const override;
+};
+
+/**
+ * A material render proxy which overrides the selection color
+ */
+class ENGINE_API FOverrideSelectionColorMaterialRenderProxy : public FMaterialRenderProxy
+{
+public:
+	const FMaterialRenderProxy* const Parent;
+	const FLinearColor SelectionColor;
+
+	/** Initialization constructor. */
+	FOverrideSelectionColorMaterialRenderProxy(const FMaterialRenderProxy* InParent, const FLinearColor& InSelectionColor);
+
+	// FMaterialRenderProxy interface.
+	virtual const FMaterial* GetMaterialNoFallback(ERHIFeatureLevel::Type InFeatureLevel) const override;
+	virtual const FMaterialRenderProxy* GetFallback(ERHIFeatureLevel::Type InFeatureLevel) const override;
+	virtual bool GetParameterValue(EMaterialParameterType Type, const FHashedMaterialParameterInfo& ParameterInfo, FMaterialParameterValue& OutValue, const FMaterialRenderContext& Context) const override;
+};
+
+/**
+ * An material render proxy which overrides the material's Color and Lightmap resolution vector parameter.
+ */
+class ENGINE_API FLightingDensityMaterialRenderProxy : public FColoredMaterialRenderProxy
+{
+public:
+	const FVector2D LightmapResolution;
+
+	/** Initialization constructor. */
+	FLightingDensityMaterialRenderProxy(const FMaterialRenderProxy* InParent, const FLinearColor& InColor, const FVector2D& InLightmapResolution);
+
+	// FMaterialRenderProxy interface.
+	virtual bool GetParameterValue(EMaterialParameterType Type, const FHashedMaterialParameterInfo& ParameterInfo, FMaterialParameterValue& OutValue, const FMaterialRenderContext& Context) const override;
+};
diff --git a/Engine/Source/Runtime/RHI/Public/RHIDefinitions.h b/Engine/Source/Runtime/RHI/Public/RHIDefinitions.h
--- a/Engine/Source/Runtime/RHI/Public/RHIDefinitions.h
+++ b/Engine/Source/Runtime/RHI/Public/RHIDefinitions.h
@@ -15,6 +15,7 @@
 #include "HAL/IConsoleManager.h"
 #include "RHIShaderPlatform.h"
 #include "RHIFeatureLevel.h"
+#include "RHIImmutableSamplerState.h"
 #endif
 
 /** Alignment of the shader parameters struct is required to be 16-byte boundaries. */
@@ -136,12 +137,6 @@ enum { MAX_TEXTURE_MIP_COUNT = 15 };
 /** Maximum number of static/skeletal mesh LODs */
 enum { MAX_MESH_LOD_COUNT = 8 };
 
-/** Maximum number of immutable samplers in a PSO. */
-enum
-{
-	MaxImmutableSamplers = 2
-};
-
 /** The maximum number of vertex elements which can be used by a vertex declaration. */
 enum
 {
diff --git a/Engine/Source/Runtime/RHI/Public/RHIImmutableSamplerState.h b/Engine/Source/Runtime/RHI/Public/RHIImmutableSamplerState.h
new file mode 100644
--- /dev/null
+++ b/Engine/Source/Runtime/RHI/Public/RHIImmutableSamplerState.h
@@ -0,0 +1,43 @@
+// Copyright Epic Games, Inc. All Rights Reserved.
+
+#pragma once
+
+#include "CoreTypes.h"
+#include "Containers/StaticArray.h"
+
+class FRHISamplerState;
+
+/** Maximum number of immutable samplers in a PSO. */
+enum
+{
+	MaxImmutableSamplers = 2
+};
+
+struct FImmutableSamplerState
+{
+	using TImmutableSamplers = TStaticArray<FRHISamplerState*, MaxImmutableSamplers>;
+
+	FImmutableSamplerState()
+		: ImmutableSamplers(InPlace, nullptr)
+	{}
+
+	void Reset()
+	{
+		for (uint32 Index = 0; Index < MaxImmutableSamplers; ++Index)
+		{
+			ImmutableSamplers[Index] = nullptr;
+		}
+	}
+
+	bool operator==(const FImmutableSamplerState& rhs) const
+	{
+		return ImmutableSamplers == rhs.ImmutableSamplers;
+	}
+
+	bool operator!=(const FImmutableSamplerState& rhs) const
+	{
+		return ImmutableSamplers != rhs.ImmutableSamplers;
+	}
+
+	TImmutableSamplers ImmutableSamplers;
+};
diff --git a/Engine/Source/Runtime/RHI/Public/RHIResources.h b/Engine/Source/Runtime/RHI/Public/RHIResources.h
--- a/Engine/Source/Runtime/RHI/Public/RHIResources.h
+++ b/Engine/Source/Runtime/RHI/Public/RHIResources.h
@@ -26,6 +26,7 @@
 #include "Containers/ConsumeAllMpmcQueue.h"
 #include "Misc/CoreDelegates.h"
 #include "RHIFwd.h"
+#include "RHIImmutableSamplerState.h"
 
 class FRHIComputeCommandList;
 class FRHICommandListImmediate;
@@ -3134,35 +3135,6 @@ private:
 #endif
 };
 
-struct FImmutableSamplerState
-{
-	using TImmutableSamplers = TStaticArray<FRHISamplerState*, MaxImmutableSamplers>;
-
-	FImmutableSamplerState()
-		: ImmutableSamplers(InPlace, nullptr)
-	{}
-
-	void Reset()
-	{
-		for (uint32 Index = 0; Index < MaxImmutableSamplers; ++Index)
-		{
-			ImmutableSamplers[Index] = nullptr;
-		}
-	}
-
-	bool operator==(const FImmutableSamplerState& rhs) const
-	{
-		return ImmutableSamplers == rhs.ImmutableSamplers;
-	}
-
-	bool operator!=(const FImmutableSamplerState& rhs) const
-	{
-		return ImmutableSamplers != rhs.ImmutableSamplers;
-	}
-
-	TImmutableSamplers ImmutableSamplers;
-};
-
 // Hints for some RHIs that support subpasses
 enum class ESubpassHint : uint8
 {
diff --git a/Engine/Source/Runtime/Renderer/Private/ShaderBaseClasses.cpp b/Engine/Source/Runtime/Renderer/Private/ShaderBaseClasses.cpp
--- a/Engine/Source/Runtime/Renderer/Private/ShaderBaseClasses.cpp
+++ b/Engine/Source/Runtime/Renderer/Private/ShaderBaseClasses.cpp
@@ -136,40 +136,12 @@ void FMaterialShader::VerifyExpressionAndShaderMaps(const FMaterialRenderProxy*
 		|| UniformExpressionCache->CachedUniformExpressionShaderMap != ShaderMap;
 	if (!bUniformExpressionSetMismatch)
 	{
-	PRAGMA_DISABLE_DEPRECATION_WARNINGS
-		if (UniformExpressionCache->LocalUniformBuffer.IsValid())
+		if (DebugUniformExpressionUBLayout.GetHash() != UniformExpressionCache->UniformBuffer->GetLayout().GetHash())
 		{
-			if (UniformExpressionCache->LocalUniformBuffer.BypassUniform)
-			{
-				if (DebugUniformExpressionUBLayout.GetHash() != UniformExpressionCache->LocalUniformBuffer.BypassUniform->GetLayout().GetHash())
-				{
-					UE_LOG(LogShaders, Warning, TEXT("Material Expression UB mismatch!"));
-					DumpUB(DebugUniformExpressionUBLayout);
-					DumpUB(UniformExpressionCache->LocalUniformBuffer.BypassUniform->GetLayout());
-					bUniformExpressionSetMismatch = true;
-				}
-			}
-			else
-			{
-				if (DebugUniformExpressionUBLayout.GetHash() != UniformExpressionCache->LocalUniformBuffer.WorkArea->Layout->GetHash())
-				{
-					UE_LOG(LogShaders, Warning, TEXT("Material Expression UB mismatch!"));
-					DumpUB(DebugUniformExpressionUBLayout);
-					DumpUB(*UniformExpressionCache->LocalUniformBuffer.WorkArea->Layout);
-					bUniformExpressionSetMismatch = true;
-				}
-			}
-		}
-	PRAGMA_ENABLE_DEPRECATION_WARNINGS
-		else
-		{
-			if (DebugUniformExpressionUBLayout.GetHash() != UniformExpressionCache->UniformBuffer->GetLayout().GetHash())
-			{
-				UE_LOG(LogShaders, Warning, TEXT("Material Expression UB mismatch!"));
-				DumpUB(DebugUniformExpressionUBLayout);
-				DumpUB(UniformExpressionCache->UniformBuffer->GetLayout());
-				bUniformExpressionSetMismatch = true;
-			}
+			UE_LOG(LogShaders, Warning, TEXT("Material Expression UB mismatch!"));
+			DumpUB(DebugUniformExpressionUBLayout);
+			DumpUB(UniformExpressionCache->UniformBuffer->GetLayout());
+			bUniformExpressionSetMismatch = true;
 		}
 	}
 	if (bUniformExpressionSetMismatch)
diff --git a/Engine/Source/Runtime/Renderer/Public/MeshPassProcessor.h b/Engine/Source/Runtime/Renderer/Public/MeshPassProcessor.h
--- a/Engine/Source/Runtime/Renderer/Public/MeshPassProcessor.h
+++ b/Engine/Source/Runtime/Renderer/Public/MeshPassProcessor.h
@@ -12,6 +12,7 @@
 #include "PSOPrecache.h"
 #include "Hash/CityHash.h"
 #include "Experimental/Containers/RobinHoodHashTable.h"
+#include "RHIImmutableSamplerState.h"
 #include <atomic>
 
 #define MESH_DRAW_COMMAND_DEBUG_DATA ((!UE_BUILD_SHIPPING && !UE_BUILD_TEST) || VALIDATE_MESH_COMMAND_BINDINGS || WANTS_DRAW_MESH_EVENTS)
