Restoring missing code after code was moved.

diff --git a/Engine/Source/Runtime/Engine/Private/Materials/MaterialAttributeDefinitionMap.cpp b/Engine/Source/Runtime/Engine/Private/Materials/MaterialAttributeDefinitionMap.cpp
--- a/Engine/Source/Runtime/Engine/Private/Materials/MaterialAttributeDefinitionMap.cpp
+++ b/Engine/Source/Runtime/Engine/Private/Materials/MaterialAttributeDefinitionMap.cpp
@@ -47,6 +47,24 @@ int32 FMaterialAttributeDefintion::CompileDefaultValue(FMaterialCompiler* Compil
 		return Compiler->ShadingModel(Compiler->GetMaterialShadingModels().GetFirstShadingModel());
 	}
 
+	if (Property == MP_Normal)
+	{
+		if (Compiler->IsTangentSpaceNormal())
+		{
+			return Compiler->Constant3(0.0f, 0.0f, 1.0f);	// Tangent space normal
+		}
+		return Compiler->VertexNormal();					// World space normal
+	}
+
+	if (Property == MP_Tangent)
+	{
+		if (Compiler->IsTangentSpaceNormal())
+		{
+			return Compiler->Constant3(1.0f, 0.0f, 0.0f);	// Tangent space tangent
+		}
+		return Compiler->VertexTangent();					// World space tangent
+	}
+
 	if (Property == MP_FrontMaterial)
 	{
 		check(ValueType == MCT_Strata);
@@ -231,47 +249,48 @@ void FMaterialAttributeDefinitionMap::InitializeAttributeMap()
 	AttributeMap.Empty(MP_MAX + 1);
 
 	// Basic attributes
-	Add(FGuid(0x69B8D336, 0x16ED4D49, 0x9AA49729, 0x2F050F7A), TEXT("BaseColor"), MP_BaseColor, MCT_Float3, FVector4(0, 0, 0, 0), SF_Pixel);
-	Add(FGuid(0x57C3A161, 0x7F064296, 0xB00B24A5, 0xA496F34C), TEXT("Metallic"), MP_Metallic, MCT_Float, FVector4(0, 0, 0, 0), SF_Pixel);
-	Add(FGuid(0x9FDAB399, 0x25564CC9, 0x8CD2D572, 0xC12C8FED), TEXT("Specular"), MP_Specular, MCT_Float, FVector4(.5, 0, 0, 0), SF_Pixel);
-	Add(FGuid(0xD1DD967C, 0x4CAD47D3, 0x9E6346FB, 0x08ECF210), TEXT("Roughness"), MP_Roughness, MCT_Float, FVector4(.5, 0, 0, 0), SF_Pixel);
-	Add(FGuid(0x55E2B4FB, 0xC1C54DB2, 0x9F11875F, 0x7231EB1E), TEXT("Anisotropy"), MP_Anisotropy, MCT_Float, FVector4(0, 0, 0, 0), SF_Pixel);
-	Add(FGuid(0xB769B54D, 0xD08D4440, 0xABC21BA6, 0xCD27D0E2), TEXT("EmissiveColor"), MP_EmissiveColor, MCT_Float3, FVector4(0, 0, 0, 0), SF_Pixel);
-	Add(FGuid(0xB8F50FBA, 0x2A754EC1, 0x9EF672CF, 0xEB27BF51), TEXT("Opacity"), MP_Opacity, MCT_Float, FVector4(1, 0, 0, 0), SF_Pixel);
-	Add(FGuid(0x679FFB17, 0x2BB5422C, 0xAD520483, 0x166E0C75), TEXT("OpacityMask"), MP_OpacityMask, MCT_Float, FVector4(1, 0, 0, 0), SF_Pixel);
-	Add(FGuid(0x0FA2821A, 0x200F4A4A, 0xB719B789, 0xC1259C64), TEXT("Normal"), MP_Normal, MCT_Float3, FVector4(0, 0, 1, 0), SF_Pixel);
-	Add(FGuid(0xD5F8E9CF, 0xCDC3468D, 0xB10E4465, 0x596A7BBA), TEXT("Tangent"), MP_Tangent, MCT_Float3, FVector4(1, 0, 0, 0), SF_Pixel);
+	Add(FGuid(0x69B8D336, 0x16ED4D49, 0x9AA49729, 0x2F050F7A), TEXT("BaseColor"),		MP_BaseColor,		MCT_Float3,	FVector4(0,0,0,0),	SF_Pixel);
+	Add(FGuid(0x57C3A161, 0x7F064296, 0xB00B24A5, 0xA496F34C), TEXT("Metallic"),		MP_Metallic,		MCT_Float,	FVector4(0,0,0,0),	SF_Pixel);
+	Add(FGuid(0x9FDAB399, 0x25564CC9, 0x8CD2D572, 0xC12C8FED), TEXT("Specular"),		MP_Specular,		MCT_Float,	FVector4(.5,0,0,0), SF_Pixel);
+	Add(FGuid(0xD1DD967C, 0x4CAD47D3, 0x9E6346FB, 0x08ECF210), TEXT("Roughness"),		MP_Roughness,		MCT_Float,	FVector4(.5,0,0,0), SF_Pixel);
+	Add(FGuid(0x55E2B4FB, 0xC1C54DB2, 0x9F11875F, 0x7231EB1E), TEXT("Anisotropy"),		MP_Anisotropy,		MCT_Float,	FVector4(0,0,0,0),  SF_Pixel);
+	Add(FGuid(0xB769B54D, 0xD08D4440, 0xABC21BA6, 0xCD27D0E2), TEXT("EmissiveColor"),	MP_EmissiveColor,	MCT_Float3,	FVector4(0,0,0,0),	SF_Pixel);
+	Add(FGuid(0xB8F50FBA, 0x2A754EC1, 0x9EF672CF, 0xEB27BF51), TEXT("Opacity"),			MP_Opacity,			MCT_Float,	FVector4(1,0,0,0),	SF_Pixel);
+	Add(FGuid(0x679FFB17, 0x2BB5422C, 0xAD520483, 0x166E0C75), TEXT("OpacityMask"),		MP_OpacityMask,		MCT_Float,	FVector4(1,0,0,0),	SF_Pixel);
+	Add(FGuid(0x0FA2821A, 0x200F4A4A, 0xB719B789, 0xC1259C64), TEXT("Normal"),			MP_Normal,			MCT_Float3,	FVector4(0,0,1,0),	SF_Pixel);
+	Add(FGuid(0xD5F8E9CF, 0xCDC3468D, 0xB10E4465, 0x596A7BBA), TEXT("Tangent"),			MP_Tangent,			MCT_Float3,	FVector4(1,0,0,0),	SF_Pixel);
 
 	// Advanced attributes
-	Add(FGuid(0xF905F895, 0xD5814314, 0x916D2434, 0x8C40CE9E), TEXT("WorldPositionOffset"), MP_WorldPositionOffset, MCT_Float3, FVector4(0, 0, 0, 0), SF_Vertex);
-	Add(FGuid(0x5B8FC679, 0x51CE4082, 0x9D777BEE, 0xF4F72C44), TEXT("SubsurfaceColor"), MP_SubsurfaceColor, MCT_Float3, FVector4(1, 1, 1, 0), SF_Pixel);
-	Add(FGuid(0x9E502E69, 0x3C8F48FA, 0x94645CFD, 0x28E5428D), TEXT("ClearCoat"), MP_CustomData0, MCT_Float, FVector4(1, 0, 0, 0), SF_Pixel);
-	Add(FGuid(0xBE4F2FFD, 0x12FC4296, 0xB0124EEA, 0x12C28D92), TEXT("ClearCoatRoughness"), MP_CustomData1, MCT_Float, FVector4(.1, 0, 0, 0), SF_Pixel);
-	Add(FGuid(0xE8EBD0AD, 0xB1654CBE, 0xB079C3A8, 0xB39B9F15), TEXT("AmbientOcclusion"), MP_AmbientOcclusion, MCT_Float, FVector4(1, 0, 0, 0), SF_Pixel);
-	Add(FGuid(0xD0B0FA03, 0x14D74455, 0xA851BAC5, 0x81A0788B), TEXT("Refraction"), MP_Refraction, MCT_Float2, FVector4(1, 0, 0, 0), SF_Pixel);
-	Add(FGuid(0x0AC97EC3, 0xE3D047BA, 0xB610167D, 0xC4D919FF), TEXT("PixelDepthOffset"), MP_PixelDepthOffset, MCT_Float, FVector4(0, 0, 0, 0), SF_Pixel);
-	Add(FGuid(0xD9423FFF, 0xD77E4D82, 0x8FF9CF5E, 0x055D1255), TEXT("ShadingModel"), MP_ShadingModel, MCT_ShadingModel, FVector4(0, 0, 0, 0), SF_Pixel, INDEX_NONE, false, &CompileShadingModelBlendFunction);
-	Add(FGuid(0x5973A03E, 0x13A74E08, 0x92D0CEDD, 0xF2936CF8), TEXT("FrontMaterial"), MP_FrontMaterial, MCT_Strata, FVector4(0, 0, 0, 0), SF_Pixel, INDEX_NONE, false, &CompileStrataBlendFunction);
+	Add(FGuid(0xF905F895, 0xD5814314, 0x916D2434, 0x8C40CE9E), TEXT("WorldPositionOffset"),		MP_WorldPositionOffset,		MCT_Float3,	FVector4(0,0,0,0),	SF_Vertex);
+	Add(FGuid(0x5B8FC679, 0x51CE4082, 0x9D777BEE, 0xF4F72C44), TEXT("SubsurfaceColor"),			MP_SubsurfaceColor,			MCT_Float3,	FVector4(1,1,1,0),	SF_Pixel);
+	Add(FGuid(0x9E502E69, 0x3C8F48FA, 0x94645CFD, 0x28E5428D), TEXT("ClearCoat"),				MP_CustomData0,				MCT_Float,	FVector4(1,0,0,0),	SF_Pixel);
+	Add(FGuid(0xBE4F2FFD, 0x12FC4296, 0xB0124EEA, 0x12C28D92), TEXT("ClearCoatRoughness"),		MP_CustomData1,				MCT_Float,	FVector4(.1,0,0,0),	SF_Pixel);
+	Add(FGuid(0xE8EBD0AD, 0xB1654CBE, 0xB079C3A8, 0xB39B9F15), TEXT("AmbientOcclusion"),		MP_AmbientOcclusion,		MCT_Float,	FVector4(1,0,0,0),	SF_Pixel);
+	Add(FGuid(0xD0B0FA03, 0x14D74455, 0xA851BAC5, 0x81A0788B), TEXT("Refraction"),				MP_Refraction,				MCT_Float3,	FVector4(1,0,0,0),	SF_Pixel);
+	Add(FGuid(0x0AC97EC3, 0xE3D047BA, 0xB610167D, 0xC4D919FF), TEXT("PixelDepthOffset"),		MP_PixelDepthOffset,		MCT_Float,	FVector4(0,0,0,0),	SF_Pixel);
+	Add(FGuid(0xD9423FFF, 0xD77E4D82, 0x8FF9CF5E, 0x055D1255), TEXT("ShadingModel"),			MP_ShadingModel,			MCT_ShadingModel, FVector4(0, 0, 0, 0), SF_Pixel, INDEX_NONE, false, &CompileShadingModelBlendFunction);
+	Add(FGuid(0x42BDD2E0, 0xBE714189, 0xA0984BC3, 0xDD0BE872), TEXT("SurfaceThickness"),		MP_SurfaceThickness,		MCT_Float,  FVector4(STRATA_LAYER_DEFAULT_THICKNESS_CM, 0, 0, 0), SF_Pixel);
+	Add(FGuid(0x5973A03E, 0x13A74E08, 0x92D0CEDD, 0xF2936CF8), TEXT("FrontMaterial"),			MP_FrontMaterial,			MCT_Strata, FVector4(0,0,0,0),	SF_Pixel, INDEX_NONE, false, &CompileStrataBlendFunction);
 
 	// Used when compiling material with execution pins, which are compiling all attributes together
-	Add(FGuid(0xE0ED040B, 0x82794D93, 0xBD2D59B2, 0xA5BBF41C), TEXT("MaterialAttributes"), MP_MaterialAttributes, MCT_MaterialAttributes, FVector4(0, 0, 0, 0), SF_Pixel, INDEX_NONE, bHideAttribute);
+	Add(FGuid(0xE0ED040B, 0x82794D93, 0xBD2D59B2, 0xA5BBF41C), TEXT("MaterialAttributes"),		MP_MaterialAttributes,		MCT_MaterialAttributes, FVector4(0,0,0,0), SF_Pixel, INDEX_NONE, bHideAttribute);
 
 	// Texture coordinates
-	Add(FGuid(0xD30EC284, 0xE13A4160, 0x87BB5230, 0x2ED115DC), TEXT("CustomizedUV0"), MP_CustomizedUVs0, MCT_Float2, FVector4(0, 0, 0, 0), SF_Vertex, 0);
-	Add(FGuid(0xC67B093C, 0x2A5249AA, 0xABC97ADE, 0x4A1F49C5), TEXT("CustomizedUV1"), MP_CustomizedUVs1, MCT_Float2, FVector4(0, 0, 0, 0), SF_Vertex, 1);
-	Add(FGuid(0x85C15B24, 0xF3E047CA, 0x85856872, 0x01AE0F4F), TEXT("CustomizedUV2"), MP_CustomizedUVs2, MCT_Float2, FVector4(0, 0, 0, 0), SF_Vertex, 2);
-	Add(FGuid(0x777819DC, 0x31AE4676, 0xB864EF77, 0xB807E873), TEXT("CustomizedUV3"), MP_CustomizedUVs3, MCT_Float2, FVector4(0, 0, 0, 0), SF_Vertex, 3);
-	Add(FGuid(0xDA63B233, 0xDDF44CAD, 0xB93D867B, 0x8DAFDBCC), TEXT("CustomizedUV4"), MP_CustomizedUVs4, MCT_Float2, FVector4(0, 0, 0, 0), SF_Vertex, 4);
-	Add(FGuid(0xC2F52B76, 0x4A034388, 0x89119528, 0x2071B190), TEXT("CustomizedUV5"), MP_CustomizedUVs5, MCT_Float2, FVector4(0, 0, 0, 0), SF_Vertex, 5);
-	Add(FGuid(0x8214A8CA, 0x0CB944CF, 0x9DFD78DB, 0xE48BB55F), TEXT("CustomizedUV6"), MP_CustomizedUVs6, MCT_Float2, FVector4(0, 0, 0, 0), SF_Vertex, 6);
-	Add(FGuid(0xD8F8D01F, 0xC6F74715, 0xA3CFB4FF, 0x9EF51FAC), TEXT("CustomizedUV7"), MP_CustomizedUVs7, MCT_Float2, FVector4(0, 0, 0, 0), SF_Vertex, 7);
+	Add(FGuid(0xD30EC284, 0xE13A4160, 0x87BB5230, 0x2ED115DC), TEXT("CustomizedUV0"), MP_CustomizedUVs0, MCT_Float2, FVector4(0,0,0,0), SF_Vertex, 0);
+	Add(FGuid(0xC67B093C, 0x2A5249AA, 0xABC97ADE, 0x4A1F49C5), TEXT("CustomizedUV1"), MP_CustomizedUVs1, MCT_Float2, FVector4(0,0,0,0), SF_Vertex, 1);
+	Add(FGuid(0x85C15B24, 0xF3E047CA, 0x85856872, 0x01AE0F4F), TEXT("CustomizedUV2"), MP_CustomizedUVs2, MCT_Float2, FVector4(0,0,0,0), SF_Vertex, 2);
+	Add(FGuid(0x777819DC, 0x31AE4676, 0xB864EF77, 0xB807E873), TEXT("CustomizedUV3"), MP_CustomizedUVs3, MCT_Float2, FVector4(0,0,0,0), SF_Vertex, 3);
+	Add(FGuid(0xDA63B233, 0xDDF44CAD, 0xB93D867B, 0x8DAFDBCC), TEXT("CustomizedUV4"), MP_CustomizedUVs4, MCT_Float2, FVector4(0,0,0,0), SF_Vertex, 4);
+	Add(FGuid(0xC2F52B76, 0x4A034388, 0x89119528, 0x2071B190), TEXT("CustomizedUV5"), MP_CustomizedUVs5, MCT_Float2, FVector4(0,0,0,0), SF_Vertex, 5);
+	Add(FGuid(0x8214A8CA, 0x0CB944CF, 0x9DFD78DB, 0xE48BB55F), TEXT("CustomizedUV6"), MP_CustomizedUVs6, MCT_Float2, FVector4(0,0,0,0), SF_Vertex, 6);
+	Add(FGuid(0xD8F8D01F, 0xC6F74715, 0xA3CFB4FF, 0x9EF51FAC), TEXT("CustomizedUV7"), MP_CustomizedUVs7, MCT_Float2, FVector4(0,0,0,0), SF_Vertex, 7);
 
 	// Lightmass attributes	
-	Add(FGuid(0x68934E1B, 0x70EB411B, 0x86DF5AA5, 0xDF2F626C), TEXT("DiffuseColor"), MP_DiffuseColor, MCT_Float3, FVector4(0, 0, 0, 0), SF_Pixel, INDEX_NONE, bHideAttribute);
-	Add(FGuid(0xE89CBD84, 0x62EA48BE, 0x80F88521, 0x2B0C403C), TEXT("SpecularColor"), MP_SpecularColor, MCT_Float3, FVector4(0, 0, 0, 0), SF_Pixel, INDEX_NONE, bHideAttribute);
+	Add(FGuid(0x68934E1B, 0x70EB411B, 0x86DF5AA5, 0xDF2F626C), TEXT("DiffuseColor"),	MP_DiffuseColor,	MCT_Float3, FVector4(0,0,0,0), SF_Pixel, INDEX_NONE, bHideAttribute);
+	Add(FGuid(0xE89CBD84, 0x62EA48BE, 0x80F88521, 0x2B0C403C), TEXT("SpecularColor"),	MP_SpecularColor,	MCT_Float3, FVector4(0,0,0,0), SF_Pixel, INDEX_NONE, bHideAttribute);
 
 	// Debug attributes
-	Add(FGuid(0x5BF6BA94, 0xA3264629, 0xA253A05B, 0x0EABBB86), TEXT("Missing"), MP_MAX, MCT_Float, FVector4(0, 0, 0, 0), SF_Pixel, INDEX_NONE, bHideAttribute);
+	Add(FGuid(0x5BF6BA94, 0xA3264629, 0xA253A05B, 0x0EABBB86), TEXT("Missing"), MP_MAX, MCT_Float, FVector4(0,0,0,0), SF_Pixel, INDEX_NONE, bHideAttribute);
 
 	// Removed attributes
 	Add(FGuid(0x2091ECA2, 0xB59248EE, 0x8E2CD578, 0xD371926D), TEXT("WorldDisplacement"), MP_WorldDisplacement_DEPRECATED, MCT_Float3, FVector4(0, 0, 0, 0), SF_Vertex, INDEX_NONE, bHideAttribute);
@@ -410,6 +429,8 @@ FText FMaterialAttributeDefinitionMap::GetAttributeOverrideForMaterial(const FGu
 		return LOCTEXT("PixelDepthOffset", "Pixel Depth Offset");
 	case MP_ShadingModel:
 		return LOCTEXT("ShadingModel", "Shading Model");
+	case MP_SurfaceThickness:
+		return LOCTEXT("SurfaceThickness", "Surface Thickness");
 	case MP_FrontMaterial:
 		return LOCTEXT("FrontMaterial", "Front Material");
 	case MP_CustomOutput:
diff --git a/Engine/Source/Runtime/Engine/Private/Materials/MaterialRenderProxy.cpp b/Engine/Source/Runtime/Engine/Private/Materials/MaterialRenderProxy.cpp
--- a/Engine/Source/Runtime/Engine/Private/Materials/MaterialRenderProxy.cpp
+++ b/Engine/Source/Runtime/Engine/Private/Materials/MaterialRenderProxy.cpp
@@ -134,6 +134,7 @@ public:
 				TRACE_CPUPROFILER_EVENT_SCOPE(FUniformExpressionCacheAsyncUpdater::Update);
 				FTaskTagScope Scope(ETaskTag::EParallelRenderingThread);
 				FMemMark Mark(FMemStack::Get());
+				RHICmdList->SwitchPipeline(ERHIPipeline::Graphics);
 
 				for (const FItem& Item : Items)
 				{
@@ -236,6 +237,17 @@ bool FMaterialRenderProxy::GetTextureValue(const FHashedMaterialParameterInfo& P
 	return false;
 }
 
+bool FMaterialRenderProxy::GetTextureValue(const FHashedMaterialParameterInfo& ParameterInfo, const USparseVolumeTexture** OutValue, const FMaterialRenderContext& Context) const
+{
+	FMaterialParameterValue Value;
+	if (GetParameterValue(EMaterialParameterType::SparseVolumeTexture, ParameterInfo, Value, Context))
+	{
+		*OutValue = Value.SparseVolumeTexture;
+		return true;
+	}
+	return false;
+}
+
 static void OnVirtualTextureDestroyedCB(const FVirtualTextureProducerHandle& InHandle, void* Baton)
 {
 	FMaterialRenderProxy* MaterialProxy = static_cast<FMaterialRenderProxy*>(Baton);
@@ -614,10 +626,12 @@ const FMaterial& FMaterialRenderProxy::GetMaterialWithFallback(ERHIFeatureLevel:
 		} while (!Material || !Material->IsRenderingThreadShaderMapComplete());
 		OutFallbackMaterialRenderProxy = FallbackMaterialProxy;
 
+#if WITH_EDITOR
 		if (BaseMaterial)
 		{
 			BaseMaterial->SubmitCompileJobs_RenderThread(EShaderCompileJobPriority::Normal);
 		}
+#endif // WITH_EDITOR
 	}
 	return *Material;
 }
diff --git a/Engine/Source/Runtime/Engine/Public/Materials/MaterialRenderProxy.h b/Engine/Source/Runtime/Engine/Public/Materials/MaterialRenderProxy.h
--- a/Engine/Source/Runtime/Engine/Public/Materials/MaterialRenderProxy.h
+++ b/Engine/Source/Runtime/Engine/Public/Materials/MaterialRenderProxy.h
@@ -18,6 +18,7 @@ class FUniformExpressionSet;
 class IAllocatedVirtualTexture;
 class UMaterialInterface;
 class URuntimeVirtualTexture;
+class USparseVolumeTexture;
 class USubsurfaceProfile;
 class UTexture;
 
@@ -164,6 +165,7 @@ public:
 	bool GetScalarValue(const FHashedMaterialParameterInfo& ParameterInfo, float* OutValue, const FMaterialRenderContext& Context) const;
 	bool GetTextureValue(const FHashedMaterialParameterInfo& ParameterInfo, const UTexture** OutValue, const FMaterialRenderContext& Context) const;
 	bool GetTextureValue(const FHashedMaterialParameterInfo& ParameterInfo, const URuntimeVirtualTexture** OutValue, const FMaterialRenderContext& Context) const;
+	bool GetTextureValue(const FHashedMaterialParameterInfo& ParameterInfo, const USparseVolumeTexture** OutValue, const FMaterialRenderContext& Context) const;
 	virtual bool GetParameterValue(EMaterialParameterType Type, const FHashedMaterialParameterInfo& ParameterInfo, FMaterialParameterValue& OutValue, const FMaterialRenderContext& Context) const = 0;
 
 	bool IsDeleted() const
