- PR #9911: Fix Cascade Particle System - ModuleLocationDirect (Contributed by [REDACTED_PERSON])

#rnx

diff --git a/Engine/Source/Runtime/Engine/Private/Particles/ParticleModules.cpp b/Engine/Source/Runtime/Engine/Private/Particles/ParticleModules.cpp
--- a/Engine/Source/Runtime/Engine/Private/Particles/ParticleModules.cpp
+++ b/Engine/Source/Runtime/Engine/Private/Particles/ParticleModules.cpp
@@ -2876,7 +2876,7 @@ void UParticleModuleAcceleration::Update(FParticleEmitterInstance* Owner, int32
 
 uint32 UParticleModuleAcceleration::RequiredBytes(UParticleModuleTypeDataBase* TypeData)
 {
-	// FVector UsedAcceleration
+	// FVector3f UsedAcceleration
 	return sizeof(FVector3f);
 }
 
diff --git a/Engine/Source/Runtime/Engine/Private/Particles/ParticleModules_Location.cpp b/Engine/Source/Runtime/Engine/Private/Particles/ParticleModules_Location.cpp
--- a/Engine/Source/Runtime/Engine/Private/Particles/ParticleModules_Location.cpp
+++ b/Engine/Source/Runtime/Engine/Private/Particles/ParticleModules_Location.cpp
@@ -356,7 +356,7 @@ void UParticleModuleLocationDirect::Update(FParticleEmitterInstance* Owner, int3
 
 uint32 UParticleModuleLocationDirect::RequiredBytes(UParticleModuleTypeDataBase* TypeData)
 {
-	return sizeof(FVector3f);
+	return sizeof(FVector);
 }
 
 /*-----------------------------------------------------------------------------
