#UE Prevent infinite loop in ValidateContainerDescriptor when there is an error due to a missing parent

diff --git a/Engine/Source/Runtime/Engine/Private/WorldPartition/WorldPartitionStreamingGeneration.cpp b/Engine/Source/Runtime/Engine/Private/WorldPartition/WorldPartitionStreamingGeneration.cpp
--- a/Engine/Source/Runtime/Engine/Private/WorldPartition/WorldPartitionStreamingGeneration.cpp
+++ b/Engine/Source/Runtime/Engine/Private/WorldPartition/WorldPartitionStreamingGeneration.cpp
@@ -780,6 +780,11 @@ PRAGMA_ENABLE_DEPRECATION_WARNINGS
 					}
 				}
 			});		
+			
+			if (!ensureMsgf(NbValidationPasses < 100000, TEXT("Infinite loop in validation passes for %s"), ContainerDescriptor.Container ? *ContainerDescriptor.Container->GetContainerPackage().ToString() : TEXT("[unknown world]")))
+			{
+				break;
+			}
 		}
 	}
 
