Optimus: Move ownership of deformed geometry out of passthrough vertex factory and into the mesh object.
This allows multiple deformers to write to different parts of the same final vertex buffer.
Simplify the passthrough vertex factory, and the logic around uniform buffer updates.

diff --git a/Engine/Plugins/Animation/DeformerGraph/Source/OptimusCore/Private/DataInterfaces/OptimusDataInterfaceSkinnedMeshWrite.cpp b/Engine/Plugins/Animation/DeformerGraph/Source/OptimusCore/Private/DataInterfaces/OptimusDataInterfaceSkinnedMeshWrite.cpp
--- a/Engine/Plugins/Animation/DeformerGraph/Source/OptimusCore/Private/DataInterfaces/OptimusDataInterfaceSkinnedMeshWrite.cpp
+++ b/Engine/Plugins/Animation/DeformerGraph/Source/OptimusCore/Private/DataInterfaces/OptimusDataInterfaceSkinnedMeshWrite.cpp
@@ -136,23 +136,15 @@ void FOptimusSkinnedMeshWriteDataProviderProxy::AllocateResources(FRDGBuilder& G
 {
 	// Allocate required buffers
 	const int32 LodIndex = SkeletalMeshObject->GetLOD();
-	FSkeletalMeshRenderData const& SkeletalMeshRenderData = SkeletalMeshObject->GetSkeletalMeshRenderData();
-	FSkeletalMeshLODRenderData const* LodRenderData = &SkeletalMeshRenderData.LODRenderData[LodIndex];
-	const int32 NumVertices = LodRenderData->GetNumVertices();
 
 	// We will extract buffers from RDG.
 	// It could be better to for memory to use QueueBufferExtraction instead of ConvertToExternalBuffer but that will require an extra hook after graph execution.
 	TRefCountPtr<FRDGPooledBuffer> PositionBufferExternal;
-	TRefCountPtr<FRDGPooledBuffer> TangentBufferExternal;
-	TRefCountPtr<FRDGPooledBuffer> ColorBufferExternal;
 
 	if (OutputMask & 1)
 	{
-		const uint32 PosBufferBytesPerElement = 4;
-		PositionBuffer = GraphBuilder.CreateBuffer(FRDGBufferDesc::CreateBufferDesc(PosBufferBytesPerElement, NumVertices * 3), TEXT("SkinnedMeshPositionBuffer"), ERDGBufferFlags::None);
+		PositionBuffer = FSkeletalMeshDeformerHelpers::AllocateVertexFactoryPositionBuffer(GraphBuilder, SkeletalMeshObject, LodIndex, TEXT("OptimusSkinnedMeshPosition"));
 		PositionBufferUAV = GraphBuilder.CreateUAV(PositionBuffer, PF_R32_FLOAT, ERDGUnorderedAccessViewFlags::SkipBarrier);
-		PositionBufferExternal = GraphBuilder.ConvertToExternalBuffer(PositionBuffer);
-		GraphBuilder.SetBufferAccessFinal(PositionBuffer, ERHIAccess::VertexOrIndexBuffer | ERHIAccess::SRVMask);
 	}
 	else
 	{
@@ -164,11 +156,8 @@ void FOptimusSkinnedMeshWriteDataProviderProxy::AllocateResources(FRDGBuilder& G
 
 	if (OutputMask & 2)
 	{
-		const uint32 TangentBufferBytesPerElement = 8;
-		TangentBuffer = GraphBuilder.CreateBuffer(FRDGBufferDesc::CreateBufferDesc(TangentBufferBytesPerElement, NumVertices * 2), TEXT("SkinnedMeshTangentBuffer"), ERDGBufferFlags::None);
+		TangentBuffer = FSkeletalMeshDeformerHelpers::AllocateVertexFactoryTangentBuffer(GraphBuilder, SkeletalMeshObject, LodIndex, TEXT("OptimusSkinnedMeshTangent"));
 		TangentBufferUAV = GraphBuilder.CreateUAV(TangentBuffer, TangentsFormat, ERDGUnorderedAccessViewFlags::SkipBarrier);
-		TangentBufferExternal = GraphBuilder.ConvertToExternalBuffer(TangentBuffer);
-		GraphBuilder.SetBufferAccessFinal(TangentBuffer, ERHIAccess::VertexOrIndexBuffer | ERHIAccess::SRVMask);
 	}
 	else
 	{
@@ -177,29 +166,15 @@ void FOptimusSkinnedMeshWriteDataProviderProxy::AllocateResources(FRDGBuilder& G
 
 	if (OutputMask & 8)
 	{
-		const uint32 ColorBufferBytesPerElement = 4;
-		ColorBuffer = GraphBuilder.CreateBuffer(FRDGBufferDesc::CreateBufferDesc(ColorBufferBytesPerElement, NumVertices), TEXT("SkinnedMeshColorBuffer"), ERDGBufferFlags::None);
+		ColorBuffer = FSkeletalMeshDeformerHelpers::AllocateVertexFactoryColorBuffer(GraphBuilder, SkeletalMeshObject, LodIndex, TEXT("OptimusSkinnedMeshColor"));
 		ColorBufferUAV = GraphBuilder.CreateUAV(ColorBuffer, PF_B8G8R8A8, ERDGUnorderedAccessViewFlags::SkipBarrier);
-		ColorBufferExternal = GraphBuilder.ConvertToExternalBuffer(ColorBuffer);
-		GraphBuilder.SetBufferAccessFinal(ColorBuffer, ERHIAccess::VertexOrIndexBuffer | ERHIAccess::SRVMask);
 	}
 	else
 	{
 		ColorBufferUAV = GraphBuilder.CreateUAV(GraphBuilder.RegisterExternalBuffer(GWhiteVertexBufferWithRDG->Buffer), PF_B8G8R8A8);
 	}
 
-	// Set to vertex factories
-	const int32 NumSections = LodRenderData->RenderSections.Num();
-	FSkeletalMeshDeformerHelpers::SetVertexFactoryBufferOverrides(SkeletalMeshObject, LodIndex, FSkeletalMeshDeformerHelpers::EOverrideType::Partial, PositionBufferExternal, TangentBufferExternal, ColorBufferExternal);
-
-#if RHI_RAYTRACING
-	if (PositionBufferExternal.IsValid())
-	{
-		// This can create RHI resources but it queues and doesn't actually build ray tracing structures. Not sure if we need to put inside a render graph pass?
-		// Also note that for ray tracing we may want to support a second graph execution if ray tracing LOD needs to be different to render LOD.
-		FSkeletalMeshDeformerHelpers::UpdateRayTracingGeometry(SkeletalMeshObject, LodIndex, PositionBufferExternal);
-	}
-#endif
+	FSkeletalMeshDeformerHelpers::UpdateVertexFactoryBufferOverrides(SkeletalMeshObject, LodIndex);
 }
 
 void FOptimusSkinnedMeshWriteDataProviderProxy::GatherDispatchData(FDispatchSetup const& InDispatchSetup, FCollectedDispatchData& InOutDispatchData)
diff --git a/Engine/Shaders/Private/LocalVertexFactory.ush b/Engine/Shaders/Private/LocalVertexFactory.ush
--- a/Engine/Shaders/Private/LocalVertexFactory.ush
+++ b/Engine/Shaders/Private/LocalVertexFactory.ush
@@ -171,8 +171,6 @@ bool IsGPUSkinPassThrough()
 }
 
 #if RAYHITGROUPSHADER || COMPUTESHADER
-Buffer<float> GPUSkinCachePositionBuffer;
-
 uint GetNumRayTracingDynamicMeshVerticesIndirect()
 {
 	return 0;
@@ -221,11 +219,15 @@ FVertexFactoryInput LoadVertexFactoryInputForDynamicUpdate(uint TriangleIndex, i
 
 	if (IsGPUSkinPassThrough())
 	{
+		// GPUSkinPassThrough can change the PositionBuffer every frame but doesn't always update the corresponding 
+		// LocalVF.VertexFetch_PositionBuffer since it relies on fixed function vertex buffer fetch and so doesn't need
+		// to pay the cost of recreating the UniformBuffer.
+		// For the compute shader case we store the PositionBuffer SRV in the loose uniform buffer.
 		Input.VertexId = TriangleIndex * 3 + VertexIndex;
 		uint VertexOffset = LocalVF.VertexFetch_Parameters[VF_VertexOffset] + Input.VertexId;
-		Input.Position.x = GPUSkinCachePositionBuffer[VertexOffset * 3 + 0];
-		Input.Position.y = GPUSkinCachePositionBuffer[VertexOffset * 3 + 1];
-		Input.Position.z = GPUSkinCachePositionBuffer[VertexOffset * 3 + 2];
+		Input.Position.x = LocalVFLooseParameters.GPUSkinPassThroughPositionBuffer[VertexOffset * 3 + 0];
+		Input.Position.y = LocalVFLooseParameters.GPUSkinPassThroughPositionBuffer[VertexOffset * 3 + 1];
+		Input.Position.z = LocalVFLooseParameters.GPUSkinPassThroughPositionBuffer[VertexOffset * 3 + 2];
 	}
 	else
 	{
diff --git a/Engine/Source/Runtime/Engine/Private/Animation/MeshDeformerGeometry.cpp b/Engine/Source/Runtime/Engine/Private/Animation/MeshDeformerGeometry.cpp
new file mode 100644
--- /dev/null
+++ b/Engine/Source/Runtime/Engine/Private/Animation/MeshDeformerGeometry.cpp
@@ -0,0 +1,15 @@
+// Copyright Epic Games, Inc. All Rights Reserved.
+
+#include "Animation/MeshDeformerGeometry.h"
+
+FMeshDeformerGeometry::FMeshDeformerGeometry()
+{
+	VertexFactoryUserData.DeformerGeometry = this;
+}
+
+void FMeshDeformerGeometry::Reset()
+{
+	PositionUpdatedFrame = TangentUpdatedFrame = ColorUpdatedFrame = 0;
+	Position = PrevPosition = Tangent = Color = nullptr;
+	PositionSRV = PrevPositionSRV = TangentSRV = ColorSRV = nullptr;
+}
diff --git a/Engine/Source/Runtime/Engine/Private/Animation/MeshDeformerGeometry.h b/Engine/Source/Runtime/Engine/Private/Animation/MeshDeformerGeometry.h
new file mode 100644
--- /dev/null
+++ b/Engine/Source/Runtime/Engine/Private/Animation/MeshDeformerGeometry.h
@@ -0,0 +1,37 @@
+// Copyright Epic Games, Inc. All Rights Reserved.
+
+#pragma once
+
+#include "RenderGraphResources.h"
+#include "SkeletalRenderPublic.h"
+
+/**
+ * Owner class for Mesh Deformer generated geometry.
+ */
+class FMeshDeformerGeometry
+{
+public:
+	FMeshDeformerGeometry();
+	
+	/** Reset the stored data. */
+	void Reset();
+
+	/** Associated VertexFactoryUserData set to point at this object. */
+	FSkinBatchVertexFactoryUserData VertexFactoryUserData;
+
+	// Frame numbers of last update.
+	uint32 PositionUpdatedFrame = 0;
+	uint32 TangentUpdatedFrame = 0;
+	uint32 ColorUpdatedFrame = 0;
+
+	// Buffers containing deformed geometry data.
+	TRefCountPtr<FRDGPooledBuffer> Position;
+	TRefCountPtr<FRDGPooledBuffer> PrevPosition;
+	TRefCountPtr<FRDGPooledBuffer> Tangent;
+	TRefCountPtr<FRDGPooledBuffer> Color;
+	// Shader resource views to the buffers.
+	TRefCountPtr<FRHIShaderResourceView> PositionSRV;
+	TRefCountPtr<FRHIShaderResourceView> PrevPositionSRV;
+	TRefCountPtr<FRHIShaderResourceView> TangentSRV;
+	TRefCountPtr<FRHIShaderResourceView> ColorSRV;
+};
diff --git a/Engine/Source/Runtime/Engine/Private/GPUSkinCache.cpp b/Engine/Source/Runtime/Engine/Private/GPUSkinCache.cpp
--- a/Engine/Source/Runtime/Engine/Private/GPUSkinCache.cpp
+++ b/Engine/Source/Runtime/Engine/Private/GPUSkinCache.cpp
@@ -301,13 +301,12 @@ public:
 		, MorphBuffer(0)
 		, LOD(InLOD)
 	{
-		
 		const TArray<FSkelMeshRenderSection>& Sections = InGPUSkin->GetRenderSections(LOD);
 		DispatchData.AddDefaulted(Sections.Num());
 		BatchElementsUserData.AddZeroed(Sections.Num());
 		for (int32 Index = 0; Index < Sections.Num(); ++Index)
 		{
-			BatchElementsUserData[Index].Entry = this;
+			BatchElementsUserData[Index].SkinCacheEntry = this;
 			BatchElementsUserData[Index].Section = Index;
 		}
 
@@ -340,7 +339,8 @@ public:
 		// See EGPUSkinCacheDispatchFlags
 		uint16 DispatchFlags = 0;
 
-		//
+		uint32 UpdatedFrameNumber = 0;
+		
 		uint32 NumBoneInfluences = 0;
 
 		// in floats (4 bytes)
@@ -416,12 +416,24 @@ public:
 			return IntermediateAccumulatedTangentBuffer;
 		}
 
+		void UpdateFrameNumber()
+		{
+			UpdatedFrameNumber = SourceVertexFactory->GetShaderData().UpdatedFrameNumber;
+		}
+
 		void UpdateVertexFactoryDeclaration()
 		{
-			TargetVertexFactory->UpdateVertexDeclaration(SourceVertexFactory, &GetPositionRWBuffer()->Buffer, &GetTangentRWBuffer()->Buffer);
+			FGPUSkinPassthroughVertexFactory::EVertexAtttribute Attribs[2] = { FGPUSkinPassthroughVertexFactory::EVertexAtttribute::Position, FGPUSkinPassthroughVertexFactory::EVertexAtttribute::Tangent };
+			FRHIShaderResourceView* SRVs[2] = { nullptr, GetTangentRWBuffer()->Buffer.SRV };
+			TargetVertexFactory->AddVertexAttributes(SourceVertexFactory, Attribs, SRVs);
 		}
 	};
 
+	void UpdateFrameNumber(int32 Section)
+	{
+		DispatchData[Section].UpdateFrameNumber();
+	}
+
 	void UpdateVertexFactoryDeclaration(int32 Section)
 	{
 		DispatchData[Section].UpdateVertexFactoryDeclaration();
@@ -573,7 +585,7 @@ public:
 		Data.SourceVertexFactory = InSourceVertexFactory;
 		Data.TargetVertexFactory = InTargetVertexFactory;
 
-		InTargetVertexFactory->InvalidateStreams();
+		InTargetVertexFactory->ResetVertexAttributes();
 
 		int32 RecomputeTangentsMode = GSkinCacheRecomputeTangents;
 		if (RecomputeTangentsMode > 0)
@@ -610,7 +622,7 @@ protected:
 	EGPUSkinCacheEntryMode Mode;
 	FGPUSkinCache::FRWBuffersAllocation* PositionAllocation;
 	FGPUSkinCache* SkinCache;
-	TArray<FGPUSkinBatchElementUserData> BatchElementsUserData;
+	TArray<FSkinBatchVertexFactoryUserData> BatchElementsUserData;
 	TArray<FSectionDispatchData> DispatchData;
 	FSkeletalMeshObjectGPUSkin* GPUSkin;
 	int BoneInfluenceType;
@@ -1400,6 +1412,7 @@ void FGPUSkinCache::DoDispatch(FRHICommandListImmediate& RHICmdList)
 	for (int32 i = 0; i < BatchCount; ++i)
 	{
 		FDispatchEntry& DispatchItem = BatchDispatches[i];
+		DispatchItem.SkinCacheEntry->UpdateFrameNumber(DispatchItem.Section);
 		DispatchItem.SkinCacheEntry->UpdateVertexFactoryDeclaration(DispatchItem.Section);
 	}
 
@@ -1454,6 +1467,7 @@ void FGPUSkinCache::DoDispatch(FRHICommandListImmediate& RHICmdList, FGPUSkinCac
 		DispatchUpdateSkinTangents(RHICmdList, SkinCacheEntry, Section, StagingBuffer, false);
 	}
 
+	SkinCacheEntry->UpdateFrameNumber(Section);
 	SkinCacheEntry->UpdateVertexFactoryDeclaration(Section);
 
 	TransitionAllToReadable(RHICmdList, BuffersToTransitionToRead);
@@ -1756,7 +1770,6 @@ void FGPUSkinCache::GetShaderBindings(
 	const FGPUSkinCacheEntry* Entry, 
 	int32 Section,
 	const FGPUSkinPassthroughVertexFactory* VertexFactory,
-	FShaderResourceParameter GPUSkinCachePositionBuffer,
 	FMeshDrawSingleShaderBindings& ShaderBindings,
 	FVertexInputStreamArray& VertexStreams)
 {
@@ -1769,14 +1782,14 @@ void FGPUSkinCache::GetShaderBindings(
 
 	//UE_LOG(LogSkinCache, Warning, TEXT("*** SetVertexStreams E %p Sec %d(%p) LOD %d"), Entry, Section, Entry->DispatchData[Section].Section, Entry->LOD);
 
-	VertexStreams.Add(FVertexInputStream(VertexFactory->GetPositionStreamIndex(), 0, DispatchData.GetPositionRWBuffer()->Buffer.Buffer));
+	int32 PositionStreamIndex = VertexFactory->GetAttributeStreamIndex(FGPUSkinPassthroughVertexFactory::EVertexAtttribute::Position);
+	VertexStreams.Add(FVertexInputStream(PositionStreamIndex, 0, DispatchData.GetPositionRWBuffer()->Buffer.Buffer));
 
-	if (VertexFactory->GetTangentStreamIndex() > -1 && DispatchData.GetTangentRWBuffer())
+	int32 TangentStreamIndex = VertexFactory->GetAttributeStreamIndex(FGPUSkinPassthroughVertexFactory::EVertexAtttribute::Tangent);
+	if (TangentStreamIndex > -1 && DispatchData.GetTangentRWBuffer())
 	{
-		VertexStreams.Add(FVertexInputStream(VertexFactory->GetTangentStreamIndex(), 0, DispatchData.GetTangentRWBuffer()->Buffer.Buffer));
+		VertexStreams.Add(FVertexInputStream(TangentStreamIndex, 0, DispatchData.GetTangentRWBuffer()->Buffer.Buffer));
 	}
-
-	ShaderBindings.Add(GPUSkinCachePositionBuffer, DispatchData.GetPositionRWBuffer()->Buffer.SRV);
 }
 
 void FGPUSkinCache::PrepareUpdateSkinning(FGPUSkinCacheEntry* Entry, int32 Section, uint32 RevisionNumber, TArray<FSkinCacheRWBuffer*>* OverlappedUAVs)
@@ -1966,7 +1979,7 @@ void FGPUSkinCache::ReleaseSkinCacheEntry(FGPUSkinCacheEntry* SkinCacheEntry)
 
 	for (FGPUSkinCacheEntry::FSectionDispatchData& SectionData : SkinCacheEntry->GetDispatchData())
 	{
-		SectionData.TargetVertexFactory->InvalidateStreams();
+		SectionData.TargetVertexFactory->ResetVertexAttributes();
 	}
 
 	FRWBuffersAllocation* PositionAllocation = SkinCacheEntry->PositionAllocation;
@@ -1992,7 +2005,7 @@ bool FGPUSkinCache::IsEntryValid(FGPUSkinCacheEntry* SkinCacheEntry, int32 Secti
 	return SkinCacheEntry->IsSectionValid(Section);
 }
 
-FGPUSkinBatchElementUserData* FGPUSkinCache::InternalGetFactoryUserData(FGPUSkinCacheEntry* Entry, int32 Section)
+const FSkinBatchVertexFactoryUserData* FGPUSkinCache::InternalGetVertexFactoryUserData(FGPUSkinCacheEntry* Entry, int32 Section)
 {
 	return &Entry->BatchElementsUserData[Section];
 }
@@ -2046,6 +2059,11 @@ FRWBuffer* FGPUSkinCache::GetPreviousPositionBuffer(FGPUSkinCacheEntry const* En
 	return nullptr;
 }
 
+uint32 FGPUSkinCache::GetUpdatedFrame(FGPUSkinCacheEntry const* Entry, uint32 SectionIndex)
+{
+	return Entry != nullptr ? Entry->GetDispatchData()[SectionIndex].UpdatedFrameNumber : 0;
+}
+
 void FGPUSkinCache::UpdateSkinWeightBuffer(FGPUSkinCacheEntry* Entry)
 {
 	if (Entry)
diff --git a/Engine/Source/Runtime/Engine/Private/GPUSkinVertexFactory.cpp b/Engine/Source/Runtime/Engine/Private/GPUSkinVertexFactory.cpp
--- a/Engine/Source/Runtime/Engine/Private/GPUSkinVertexFactory.cpp
+++ b/Engine/Source/Runtime/Engine/Private/GPUSkinVertexFactory.cpp
@@ -537,28 +537,25 @@ void FGPUBaseSkinVertexFactory::SetData(const FGPUSkinDataType* InData)
 	UpdateRHI();
 }
 
-void FGPUBaseSkinVertexFactory::CopyDataTypeForPassthroughFactory(FGPUSkinPassthroughVertexFactory* PassthroughVertexFactory)
+void FGPUBaseSkinVertexFactory::CopyDataTypeForLocalVertexFactory(FLocalVertexFactory::FDataType& OutDestData) const
 {
-	FGPUSkinPassthroughVertexFactory::FDataType DestDataType;
 	check(Data.IsValid());
 
-	DestDataType.PositionComponent = Data->PositionComponent;
-	DestDataType.TangentBasisComponents[0] = Data->TangentBasisComponents[0];
-	DestDataType.TangentBasisComponents[1] = Data->TangentBasisComponents[1];
-	DestDataType.TextureCoordinates = Data->TextureCoordinates;
-	DestDataType.ColorComponent = Data->ColorComponent;
-	DestDataType.PreSkinPositionComponent = Data->PositionComponent;
-	DestDataType.PositionComponentSRV = Data->PositionComponentSRV;
-	DestDataType.PreSkinPositionComponentSRV = Data->PositionComponentSRV;
-	DestDataType.TangentsSRV = Data->TangentsSRV;
-	DestDataType.ColorComponentsSRV = Data->ColorComponentsSRV;
-	DestDataType.ColorIndexMask = Data->ColorIndexMask;
-	DestDataType.TextureCoordinatesSRV = Data->TextureCoordinatesSRV;
-	DestDataType.LightMapCoordinateIndex = Data->LightMapCoordinateIndex;
-	DestDataType.NumTexCoords = Data->NumTexCoords;
-	DestDataType.LODLightmapDataIndex = Data->LODLightmapDataIndex;
-
-	PassthroughVertexFactory->SetData(DestDataType);
+	OutDestData.PositionComponent = Data->PositionComponent;
+	OutDestData.TangentBasisComponents[0] = Data->TangentBasisComponents[0];
+	OutDestData.TangentBasisComponents[1] = Data->TangentBasisComponents[1];
+	OutDestData.TextureCoordinates = Data->TextureCoordinates;
+	OutDestData.ColorComponent = Data->ColorComponent;
+	OutDestData.PreSkinPositionComponent = Data->PositionComponent;
+	OutDestData.PositionComponentSRV = Data->PositionComponentSRV;
+	OutDestData.PreSkinPositionComponentSRV = Data->PositionComponentSRV;
+	OutDestData.TangentsSRV = Data->TangentsSRV;
+	OutDestData.ColorComponentsSRV = Data->ColorComponentsSRV;
+	OutDestData.ColorIndexMask = Data->ColorIndexMask;
+	OutDestData.TextureCoordinatesSRV = Data->TextureCoordinatesSRV;
+	OutDestData.LightMapCoordinateIndex = Data->LightMapCoordinateIndex;
+	OutDestData.NumTexCoords = Data->NumTexCoords;
+	OutDestData.LODLightmapDataIndex = Data->LODLightmapDataIndex;
 }
 
 
@@ -932,168 +929,6 @@ IMPLEMENT_GPUSKINNING_VERTEX_FACTORY_TYPE(TGPUSkinVertexFactory, "/Engine/Privat
 );
 
 
-/*-----------------------------------------------------------------------------
-FGPUSkinPassthroughVertexFactory
------------------------------------------------------------------------------*/
-FGPUSkinPassthroughVertexFactory::FGPUSkinPassthroughVertexFactory(ERHIFeatureLevel::Type InFeatureLevel)
-	: FLocalVertexFactory(InFeatureLevel, "FGPUSkinPassthroughVertexFactory")
-{
-	bGPUSkinPassThrough = true;
-}
-
-FGPUSkinPassthroughVertexFactory::~FGPUSkinPassthroughVertexFactory()
-{
-}
-
-void FGPUSkinPassthroughVertexFactory::ReleaseRHI()
-{
-	FLocalVertexFactory::ReleaseRHI();
-
-	// When adding anything else to this function be aware of the bypassing code in InternalUpdateVertexDeclaration.
-	PositionRDG = nullptr;
-	PrevPositionRDG = nullptr;
-	TangentRDG = nullptr;
-	ColorRDG = nullptr;
-	PositionVBAlias.ReleaseRHI();
-	TangentVBAlias.ReleaseRHI();
-	ColorVBAlias.ReleaseRHI();
-}
-
-void FGPUSkinPassthroughVertexFactory::InternalUpdateVertexDeclaration(FGPUBaseSkinVertexFactory const* SourceVertexFactory)
-{
-	if (PositionVBAlias.VertexBufferRHI != nullptr)
-	{
-		Data.PositionComponent.VertexBuffer = &PositionVBAlias;
-		Data.PositionComponent.Offset = 0;
-		Data.PositionComponent.VertexStreamUsage = EVertexStreamUsage::Overridden;
-		Data.PositionComponent.Stride = 3 * sizeof(float);
-	}
-	else
-	{
-		Data.PositionComponent = SourceVertexFactory->GetPositionStreamComponent();
-	}
-
-	if (TangentVBAlias.VertexBufferRHI != nullptr)
-	{
-		Data.TangentBasisComponents[0].VertexBuffer = &TangentVBAlias;
-		Data.TangentBasisComponents[0].Offset = 0;
-		Data.TangentBasisComponents[0].Type = VET_Short4N;
-		Data.TangentBasisComponents[0].Stride = 16;
-		Data.TangentBasisComponents[0].VertexStreamUsage = EVertexStreamUsage::Overridden | EVertexStreamUsage::ManualFetch;
-
-		Data.TangentBasisComponents[1].VertexBuffer = &TangentVBAlias;
-		Data.TangentBasisComponents[1].Offset = 8;
-		Data.TangentBasisComponents[1].Type = VET_Short4N;
-		Data.TangentBasisComponents[1].Stride = 16;
-		Data.TangentBasisComponents[1].VertexStreamUsage = EVertexStreamUsage::Overridden | EVertexStreamUsage::ManualFetch;
-	}
-	else
-	{
-		Data.TangentBasisComponents[0] = SourceVertexFactory->GetTangentStreamComponent(0);
-		Data.TangentBasisComponents[1] = SourceVertexFactory->GetTangentStreamComponent(1);
-	}
-
-	if (ColorVBAlias.VertexBufferRHI)
-	{
-		Data.ColorComponent.VertexBuffer = &ColorVBAlias;
-		Data.ColorComponent.Offset = 0;
-		Data.ColorComponent.Type = VET_Color;
-		Data.ColorComponent.Stride = 4;
-		Data.ColorComponent.VertexStreamUsage = EVertexStreamUsage::Overridden | EVertexStreamUsage::ManualFetch;
-		
-		Data.ColorIndexMask = ~0u;
-	}
-
-	Data.PositionComponentSRV = PositionSRVAlias ? PositionSRVAlias : SourceVertexFactory->GetPositionsSRV().GetReference();
-	Data.PreSkinPositionComponentSRV = SourceVertexFactory->GetPositionsSRV().GetReference();
-	Data.TangentsSRV = TangentSRVAlias ? TangentSRVAlias : SourceVertexFactory->GetTangentsSRV().GetReference();
-	Data.ColorComponentsSRV = ColorSRVAlias ? ColorSRVAlias : SourceVertexFactory->GetColorComponentsSRV().GetReference();
-
-	// hack to allow us to release the alias pointers properly in ReleaseRHI.
-	// To be cleaned up in UE-68826
-	FLocalVertexFactory::ReleaseRHI();
-	FLocalVertexFactory::ReleaseDynamicRHI();
-	FLocalVertexFactory::InitDynamicRHI();
-	FLocalVertexFactory::InitRHI();
-
-	// Find added streams
-	PositionStreamIndex = -1;
-	TangentStreamIndex = -1;
-
-	for (int32 Index = 0; Index < Streams.Num(); ++Index)
-	{
-		if (Streams[Index].VertexBuffer->VertexBufferRHI.GetReference() == Data.PositionComponent.VertexBuffer->VertexBufferRHI.GetReference())
-		{
-			PositionStreamIndex = Index;
-		}
-		if (TangentVBAlias.VertexBufferRHI != nullptr)
-		{
-			if (Streams[Index].VertexBuffer->VertexBufferRHI.GetReference() == Data.TangentBasisComponents[0].VertexBuffer->VertexBufferRHI.GetReference())
-			{
-				TangentStreamIndex = Index;
-			}
-		}
-	}
-	checkf(PositionStreamIndex != -1, TEXT("Unable to find stream for RWBuffer Vertex buffer!"));
-}
-
-void FGPUSkinPassthroughVertexFactory::InternalUpdateVertexDeclaration(
-	FGPUBaseSkinVertexFactory const* SourceVertexFactory, 
-	struct FRWBuffer* PositionRWBuffer, 
-	struct FRWBuffer* TangentRWBuffer)
-{
-	PositionRDG = nullptr;
-	PrevPositionRDG = nullptr;
-	TangentRDG = nullptr;
-	ColorRDG = nullptr;
-
-	PositionVBAlias.VertexBufferRHI = PositionRWBuffer ? PositionRWBuffer->Buffer : nullptr;
-	TangentVBAlias.VertexBufferRHI = TangentRWBuffer ? TangentRWBuffer->Buffer : nullptr;
-	ColorVBAlias.VertexBufferRHI = nullptr;
-
-	PositionSRVAlias = PositionRWBuffer ? PositionRWBuffer->SRV : nullptr;
-	PrevPositionSRVAlias = nullptr;
-	TangentSRVAlias = TangentRWBuffer ? TangentRWBuffer->SRV : nullptr;
-	ColorSRVAlias = nullptr;
-
-	InternalUpdateVertexDeclaration(SourceVertexFactory);
-}
-
-void FGPUSkinPassthroughVertexFactory::InternalUpdateVertexDeclaration(
-	EOverrideFlags OverrideFlags,
-	FGPUBaseSkinVertexFactory const* SourceVertexFactory, 
-	TRefCountPtr<FRDGPooledBuffer> const& PositionBuffer, 
-	TRefCountPtr<FRDGPooledBuffer> const& TangentBuffer,
-	TRefCountPtr<FRDGPooledBuffer> const& ColorBuffer)
-{
-	if (EnumHasAnyFlags(OverrideFlags, EOverrideFlags::Position))
-	{
-		PrevPositionRDG = PositionBuffer.IsValid() ? PositionRDG : nullptr;
-		PositionRDG = PositionBuffer;
-	}
-	if (EnumHasAnyFlags(OverrideFlags, EOverrideFlags::Tangent))
-	{
-		TangentRDG = TangentBuffer;
-	}
-	if (EnumHasAnyFlags(OverrideFlags, EOverrideFlags::Color))
-	{
-		ColorRDG = ColorBuffer;
-	}
-
-	PositionVBAlias.VertexBufferRHI = PositionRDG.IsValid() ? PositionRDG->GetRHI() : nullptr;
-	TangentVBAlias.VertexBufferRHI = TangentRDG.IsValid() ? TangentRDG->GetRHI() : nullptr;
-	ColorVBAlias.VertexBufferRHI = ColorRDG.IsValid() ? ColorRDG->GetRHI() : nullptr;
-
-	PositionSRVAlias = PositionRDG.IsValid() ? PositionRDG->GetOrCreateSRV(FRHIBufferSRVCreateInfo(PF_R32_FLOAT)) : nullptr;
-	PrevPositionSRVAlias = PrevPositionRDG.IsValid() ? PrevPositionRDG->GetOrCreateSRV(FRHIBufferSRVCreateInfo(PF_R32_FLOAT)) : nullptr;
-	const EPixelFormat TangentsFormat = IsOpenGLPlatform(GMaxRHIShaderPlatform) ? PF_R16G16B16A16_SINT : PF_R16G16B16A16_SNORM;
-	TangentSRVAlias = TangentRDG.IsValid() ? TangentRDG->GetOrCreateSRV(FRHIBufferSRVCreateInfo(TangentsFormat)) : nullptr;
-	ColorSRVAlias = ColorRDG.IsValid() ? ColorRDG->GetOrCreateSRV(FRHIBufferSRVCreateInfo(PF_R8G8B8A8)) : nullptr;
-
-	InternalUpdateVertexDeclaration(SourceVertexFactory);
-}
-
-
 /*-----------------------------------------------------------------------------
 	FGPUBaseSkinAPEXClothVertexFactory
 -----------------------------------------------------------------------------*/
@@ -1424,5 +1259,163 @@ IMPLEMENT_GPUSKINNING_VERTEX_FACTORY_TYPE(TGPUSkinAPEXClothVertexFactory, "/Engi
 	| EVertexFactoryFlags::SupportsPSOPrecaching
 );
 
+
+/*-----------------------------------------------------------------------------
+FGPUSkinPassthroughVertexFactory
+-----------------------------------------------------------------------------*/
+FGPUSkinPassthroughVertexFactory::FGPUSkinPassthroughVertexFactory(ERHIFeatureLevel::Type InFeatureLevel)
+	: FLocalVertexFactory(InFeatureLevel, "FGPUSkinPassthroughVertexFactory")
+{
+	bGPUSkinPassThrough = true;
+
+	ResetVertexAttributes();
+}
+
+void FGPUSkinPassthroughVertexFactory::ResetVertexAttributes()
+{
+	VertexAttributeMask = 0;
+	for (int32 Index = 0; Index < EVertexAtttribute::NumAttributes; ++Index)
+	{
+		StreamIndices[Index] = -1;
+		SRVs[Index] = nullptr;
+	}
+}
+
+// We don't set actual vertex buffers in our FDataType because we always expect to override the values in GetElementShaderBindings().
+// We use some empty dummy objects instead so that the correct vertex declaration is built and so that we can track which attributes are at which stream index.
+TStaticArray<FVertexBuffer, FGPUSkinPassthroughVertexFactory::EVertexAtttribute::NumAttributes> FGPUSkinPassthroughVertexFactory::DummyVBs;
+
+void FGPUSkinPassthroughVertexFactory::OverrideAttributeData()
+{
+	if (VertexAttributeMask & (1 << EVertexAtttribute::Position))
+	{
+		Data.PositionComponent.VertexBuffer = &DummyVBs[EVertexAtttribute::Position];
+		Data.PositionComponent.Offset = 0;
+		Data.PositionComponent.VertexStreamUsage = EVertexStreamUsage::Overridden;
+		Data.PositionComponent.Stride = 3 * sizeof(float);
+	}
+
+	if (VertexAttributeMask & (1 << EVertexAtttribute::Tangent))
+	{
+		Data.TangentBasisComponents[0].VertexBuffer = &DummyVBs[EVertexAtttribute::Tangent];
+		Data.TangentBasisComponents[0].Offset = 0;
+		Data.TangentBasisComponents[0].Type = VET_Short4N;
+		Data.TangentBasisComponents[0].Stride = 16;
+		Data.TangentBasisComponents[0].VertexStreamUsage = EVertexStreamUsage::Overridden | EVertexStreamUsage::ManualFetch;
+
+		Data.TangentBasisComponents[1].VertexBuffer = &DummyVBs[EVertexAtttribute::Tangent];
+		Data.TangentBasisComponents[1].Offset = 8;
+		Data.TangentBasisComponents[1].Type = VET_Short4N;
+		Data.TangentBasisComponents[1].Stride = 16;
+		Data.TangentBasisComponents[1].VertexStreamUsage = EVertexStreamUsage::Overridden | EVertexStreamUsage::ManualFetch;
+	}
+
+	if (VertexAttributeMask & (1 << EVertexAtttribute::Color))
+	{
+		Data.ColorComponent.VertexBuffer = &DummyVBs[EVertexAtttribute::Color];
+		Data.ColorComponent.Offset = 0;
+		Data.ColorComponent.VertexStreamUsage = EVertexStreamUsage::Overridden | EVertexStreamUsage::ManualFetch;
+		Data.ColorComponent.Stride = sizeof(uint32);
+	}
+
+	if (VertexAttributeMask & (1 << EVertexAtttribute::TexCoord))
+	{
+		Data.TextureCoordinates[0].VertexBuffer = &DummyVBs[EVertexAtttribute::TexCoord];
+		Data.TextureCoordinates[0].Offset = 0;
+		Data.TextureCoordinates[0].VertexStreamUsage = EVertexStreamUsage::Overridden | EVertexStreamUsage::ManualFetch;
+		Data.TextureCoordinates[0].Stride = 2 * sizeof(float);
+
+		Data.NumTexCoords = 1;
+	}
+}
+
+void FGPUSkinPassthroughVertexFactory::OverrideSRVs(FGPUBaseSkinVertexFactory const* InSourceVertexFactory)
+{
+	Data.PositionComponentSRV = SRVs[EVertexAtttribute::Position] ? SRVs[EVertexAtttribute::Position] : (FRHIShaderResourceView*)InSourceVertexFactory->GetPositionsSRV();
+	Data.TangentsSRV = SRVs[EVertexAtttribute::Tangent] ? SRVs[EVertexAtttribute::Tangent] : (FRHIShaderResourceView*)InSourceVertexFactory->GetTangentsSRV();
+	Data.ColorComponentsSRV = SRVs[EVertexAtttribute::Color] ? SRVs[EVertexAtttribute::Color] : (FRHIShaderResourceView*)InSourceVertexFactory->GetColorComponentsSRV();
+	Data.TextureCoordinatesSRV = SRVs[EVertexAtttribute::TexCoord] ? SRVs[EVertexAtttribute::TexCoord] : (FRHIShaderResourceView*)InSourceVertexFactory->GetTextureCoordinatesSRV();
+}
+
+void FGPUSkinPassthroughVertexFactory::BuildStreamIndices()
+{
+	for (int32 AttributeIndex = 0; AttributeIndex < EVertexAtttribute::NumAttributes; ++AttributeIndex)
+	{
+		StreamIndices[AttributeIndex] = -1;
+
+		for (int32 StreamIndex = 0; StreamIndex < Streams.Num(); ++StreamIndex)
+		{
+			if (Streams[StreamIndex].VertexBuffer == &DummyVBs[AttributeIndex])
+			{
+				StreamIndices[AttributeIndex] = StreamIndex;
+				break;
+			}
+		}
+	}
+}
+
+void FGPUSkinPassthroughVertexFactory::CreateUniformBuffer()
+{
+	if (RHISupportsManualVertexFetch(GetFeatureLevelShaderPlatform(GetFeatureLevel())))
+	{
+		const int32 DefaultBaseVertexIndex = 0;
+		const int32 DefaultPreSkinBaseVertexIndex = 0;
+		UniformBuffer = CreateLocalVFUniformBuffer(this, Data.LODLightmapDataIndex, nullptr, DefaultBaseVertexIndex, DefaultPreSkinBaseVertexIndex);
+	}
+}
+
+void FGPUSkinPassthroughVertexFactory::AddVertexAttributes(FGPUBaseSkinVertexFactory const* InSourceVertexFactory, TArrayView<EVertexAtttribute> const& InAttributes, TArrayView<FRHIShaderResourceView*> const& InSRVs)
+{
+	// Check for new attributes or modified SRVs.
+	bool bNeedFullUpdate = false;
+	bool bNeedSRVUpdate = false;
+	
+	for (int32 Index = 0; Index < InAttributes.Num(); ++Index)
+	{
+		const uint32 MaskBit = 1 << InAttributes[Index];
+		if ((VertexAttributeMask & MaskBit) == 0)
+		{
+			VertexAttributeMask |= MaskBit;
+			bNeedFullUpdate = true;
+		}
+
+		if (SRVs[InAttributes[Index]] != InSRVs[Index])
+		{
+			SRVs[InAttributes[Index]] = InSRVs[Index];
+			bNeedSRVUpdate = true;
+		}
+	}
+
+	if (bNeedFullUpdate)
+	{
+		// Reset the vertex factory data from the source vertex factory.
+		InSourceVertexFactory->CopyDataTypeForLocalVertexFactory(Data);
+
+		// Override the vertex factory data for our added attributes.
+		OverrideAttributeData();
+		OverrideSRVs(InSourceVertexFactory);
+
+		// Rebuild the vertex declaration.
+		UpdateRHI();
+
+		// Rebuild the vertex stream indices.
+		BuildStreamIndices();
+	}
+	else if (bNeedSRVUpdate)
+	{
+		// Override the vertex factory data for our added attributes.
+		OverrideSRVs(InSourceVertexFactory);
+
+		// Only need to recreate the vertex factory uniform buffer.
+		CreateUniformBuffer();
+	}
+}
+
+int32 FGPUSkinPassthroughVertexFactory::GetAttributeStreamIndex(EVertexAtttribute InAttribute) const
+{
+	return StreamIndices[InAttribute];
+}
+
+
 #undef IMPLEMENT_GPUSKINNING_VERTEX_FACTORY_PARAMETER_TYPE
 #undef IMPLEMENT_GPUSKINNING_VERTEX_FACTORY_TYPE
diff --git a/Engine/Source/Runtime/Engine/Private/LocalVertexFactory.cpp b/Engine/Source/Runtime/Engine/Private/LocalVertexFactory.cpp
--- a/Engine/Source/Runtime/Engine/Private/LocalVertexFactory.cpp
+++ b/Engine/Source/Runtime/Engine/Private/LocalVertexFactory.cpp
@@ -18,6 +18,7 @@
 #include "GPUSkinVertexFactory.h"
 #include "Animation/MeshDeformerProvider.h"
 #include "RenderGraphResources.h"
+#include "SkeletalRenderGPUSkin.h"
 
 IMPLEMENT_TYPE_LAYOUT(FLocalVertexFactoryShaderParametersBase);
 IMPLEMENT_TYPE_LAYOUT(FLocalVertexFactoryShaderParameters);
@@ -166,7 +167,6 @@ void FLocalVertexFactoryShaderParametersBase::GetElementShaderBindingsBase(
 void FLocalVertexFactoryShaderParameters::Bind(const FShaderParameterMap& ParameterMap)
 {
 	FLocalVertexFactoryShaderParametersBase::Bind(ParameterMap);
-	GPUSkinCachePositionBuffer.Bind(ParameterMap, TEXT("GPUSkinCachePositionBuffer"));
 	IsGPUSkinPassThrough.Bind(ParameterMap, TEXT("bIsGPUSkinPassThrough"));
 }
 
@@ -237,52 +237,32 @@ void FLocalVertexFactoryShaderParameters::GetElementShaderBindingsGPUSkinPassThr
 	class FMeshDrawSingleShaderBindings& ShaderBindings,
 	FVertexInputStreamArray& VertexStreams) const
 {
-	// #dxr_todo do we need this call to the base?
-	FLocalVertexFactoryShaderParametersBase::GetElementShaderBindingsBase(Scene, View, Shader, InputStreamType, FeatureLevel, VertexFactory, BatchElement, nullptr, ShaderBindings, VertexStreams);
-
-	// todo: Add more context into VertexFactoryUserData about whether this is skin cache/mesh deformer/ray tracing etc.
-	// For now it just holds a skin cache pointer which is null if using other mesh deformers.
-	FGPUSkinBatchElementUserData* BatchUserData = (FGPUSkinBatchElementUserData*)BatchElement.VertexFactoryUserData;
-	const bool bUsesSkinCache = BatchUserData != nullptr;
-
-	check(VertexFactory->GetType() == &FGPUSkinPassthroughVertexFactory::StaticType);
 	FGPUSkinPassthroughVertexFactory const* PassthroughVertexFactory = static_cast<FGPUSkinPassthroughVertexFactory const*>(VertexFactory);
-	if (bUsesSkinCache)
-	{
-		GetElementShaderBindingsSkinCache(PassthroughVertexFactory, BatchUserData, ShaderBindings, VertexStreams);
-	}
-	else
+	if (PassthroughVertexFactory->SupportsManualVertexFetch(FeatureLevel) || UseGPUScene(GMaxRHIShaderPlatform, FeatureLevel))
 	{
-		GetElementShaderBindingsMeshDeformer(PassthroughVertexFactory, ShaderBindings, VertexStreams);
+		ShaderBindings.Add(Shader->GetUniformBufferParameter<FLocalVertexFactoryUniformShaderParameters>(), PassthroughVertexFactory->GetUniformBuffer());
 	}
-}
-
-void FLocalVertexFactoryShaderParameters::GetElementShaderBindingsSkinCache(
-	FGPUSkinPassthroughVertexFactory const* PassthroughVertexFactory,
-	FGPUSkinBatchElementUserData* BatchUserData,
-	FMeshDrawSingleShaderBindings& ShaderBindings,
-	FVertexInputStreamArray& VertexStreams) const
-{
-	FGPUSkinCache::GetShaderBindings(
-		BatchUserData->Entry, BatchUserData->Section,
-		PassthroughVertexFactory,
-		GPUSkinCachePositionBuffer,
-		ShaderBindings, VertexStreams);
-}
 
-void FLocalVertexFactoryShaderParameters::GetElementShaderBindingsMeshDeformer(
-	FGPUSkinPassthroughVertexFactory const* PassthroughVertexFactory,
-	FMeshDrawSingleShaderBindings& ShaderBindings,
-	FVertexInputStreamArray& VertexStreams) const
-{
-	if (PassthroughVertexFactory->PositionRDG.IsValid())
+ 	FSkinBatchVertexFactoryUserData* BatchUserData = (FSkinBatchVertexFactoryUserData*)BatchElement.VertexFactoryUserData;
+	if (!ensure(BatchUserData != nullptr))
 	{
-		VertexStreams.Add(FVertexInputStream(PassthroughVertexFactory->GetPositionStreamIndex(), 0, PassthroughVertexFactory->PositionRDG->GetRHI()));
-		ShaderBindings.Add(GPUSkinCachePositionBuffer, PassthroughVertexFactory->GetPositionsSRV());
+		return;
 	}
-	if (PassthroughVertexFactory->TangentRDG.IsValid() && PassthroughVertexFactory->GetTangentStreamIndex() > -1)
-	{
-		VertexStreams.Add(FVertexInputStream(PassthroughVertexFactory->GetTangentStreamIndex(), 0, PassthroughVertexFactory->TangentRDG->GetRHI()));
+
+ 	if (BatchUserData->SkinCacheEntry != nullptr)
+ 	{
+		// Using Skin Cache.
+		FGPUSkinCache::GetShaderBindings(BatchUserData->SkinCacheEntry, BatchUserData->Section, PassthroughVertexFactory, ShaderBindings, VertexStreams);
+ 	}
+ 	else if (BatchUserData->DeformerGeometry != nullptr)
+ 	{
+		// Using Mesh Deformers.
+		FMeshDeformerGeometry* DeformerGeometry = BatchUserData->DeformerGeometry;
+		if (DeformerGeometry->Position.IsValid())
+		{
+			int32 PositionStreamIndex = PassthroughVertexFactory->GetAttributeStreamIndex(FGPUSkinPassthroughVertexFactory::EVertexAtttribute::Position);
+			VertexStreams.Add(FVertexInputStream(PositionStreamIndex, 0, DeformerGeometry->Position->GetRHI()));
+		}
 	}
 }
 
@@ -539,6 +519,7 @@ void FLocalVertexFactory::InitRHI()
 	}
 
 	FLocalVertexFactoryLooseParameters LooseParameters;
+	LooseParameters.GPUSkinPassThroughPositionBuffer = GNullVertexBuffer.VertexBufferSRV;
 	LooseParameters.GPUSkinPassThroughPreviousPositionBuffer = GNullVertexBuffer.VertexBufferSRV;
 	LooseParametersUniformBuffer = TUniformBufferRef<FLocalVertexFactoryLooseParameters>::CreateUniformBufferImmediate(LooseParameters, UniformBuffer_MultiFrame);
 
diff --git a/Engine/Source/Runtime/Engine/Private/SkeletalMesh.cpp b/Engine/Source/Runtime/Engine/Private/SkeletalMesh.cpp
--- a/Engine/Source/Runtime/Engine/Private/SkeletalMesh.cpp
+++ b/Engine/Source/Runtime/Engine/Private/SkeletalMesh.cpp
@@ -65,6 +65,7 @@
 #include "Templates/UniquePtr.h"
 #include "AnimationRuntime.h"
 #include "Animation/AnimSequence.h"
+#include "Animation/MeshDeformerGeometry.h"
 #include "Animation/SkinWeightProfile.h"
 #include "Streaming/SkeletalMeshUpdate.h"
 #include "HAL/FileManager.h"
@@ -6422,57 +6423,71 @@ void FSkeletalMeshSceneProxy::CreateBaseMeshBatch(const FSceneView* View, const
 	BatchElement.IndexBuffer = LODData.MultiSizeIndexContainer.GetIndexBuffer();
 	BatchElement.MinVertexIndex = LODData.RenderSections[SectionIndex].GetVertexBufferIndex();
 	BatchElement.MaxVertexIndex = LODData.RenderSections[SectionIndex].GetVertexBufferIndex() + LODData.RenderSections[SectionIndex].GetNumVertices() - 1;
-#if RHI_RAYTRACING
-	FGPUSkinBatchElementUserData* VertexFactoryUserData = FGPUSkinCache::GetFactoryUserData(VFMode == ESkinVertexFactoryMode::RayTracing ? MeshObject->GetSkinCacheEntryForRayTracing() : MeshObject->SkinCacheEntry, SectionIndex);
-#else
-	FGPUSkinBatchElementUserData* VertexFactoryUserData = FGPUSkinCache::GetFactoryUserData(MeshObject->SkinCacheEntry, SectionIndex);
-#endif
-	BatchElement.VertexFactoryUserData = VertexFactoryUserData;
+
+	FSkinBatchVertexFactoryUserData const* VertexFactoryUserData = MeshObject->GetVertexFactoryUserData(LODIndex, SectionIndex, VFMode);
+	BatchElement.VertexFactoryUserData = (void*)VertexFactoryUserData;
+
 	BatchElement.PrimitiveUniformBuffer = GetUniformBuffer();
 	BatchElement.NumPrimitives = LODData.RenderSections[SectionIndex].NumTriangles;
 
-	UpdateLooseParametersUniformBuffer(View, SectionIndex, Mesh, VertexFactoryUserData);
+	// todo: Move this uniform buffer update to somewhere that doesn't require a const cast?
+	FVertexFactory* NonConstVertexFactory = const_cast<FVertexFactory*>(Mesh.VertexFactory);
+	UpdateLooseParametersUniformBuffer(View, SectionIndex, VertexFactoryUserData, NonConstVertexFactory);
 }
 
-void FSkeletalMeshSceneProxy::UpdateLooseParametersUniformBuffer(const FSceneView* View, const int32 SectionIndex, const FMeshBatch& Mesh, const FGPUSkinBatchElementUserData* BatchUserData) const
+void FSkeletalMeshSceneProxy::UpdateLooseParametersUniformBuffer(const FSceneView* View, const int32 SectionIndex, const FSkinBatchVertexFactoryUserData* BatchUserData, FVertexFactory* VertexFactory) const
 {
 	// Loose parameters uniform buffer is only needed for PassThroughVF which is a type of LocalVF.
 	// Check it is a LocalVF StaticType and bGPUSkinPassThrough flag is set. Early out if either condition is not met.
-	if (!Mesh.VertexFactory || Mesh.VertexFactory->GetType() != &FLocalVertexFactory::StaticType)
+	if (VertexFactory == nullptr || VertexFactory->GetType() != &FLocalVertexFactory::StaticType)
 	{
 		return;
 	}
-	const FLocalVertexFactory* LocalVertexFactory = static_cast<const FLocalVertexFactory*>(Mesh.VertexFactory);
+	FLocalVertexFactory* LocalVertexFactory = static_cast<FLocalVertexFactory*>(VertexFactory);
 	if (!LocalVertexFactory->bGPUSkinPassThrough)
 	{
 		return;
 	}
-	
-	FLocalVertexFactoryLooseParameters Parameters;
+	// We always expect a valid BatchUserData for PassthroughVF.
+	if (!ensure(BatchUserData != nullptr))
+	{
+		return;
+	}
 
-	FVertexFactory* NonConstVertexFactory = const_cast<FVertexFactory*>(Mesh.VertexFactory);
-	FGPUSkinPassthroughVertexFactory* PassthroughVertexFactory = static_cast<FGPUSkinPassthroughVertexFactory*>(NonConstVertexFactory);
-	// Bone data is updated whenever animation triggers a dynamic update, animation can skip frames hence the frequency is not necessary every frame.
-	// So check if bone data is updated this frame, if not then the previous frame data is stale and not suitable for motion blur.
-	bool bBoneDataUpdatedThisFrame = (View->Family->FrameNumber == PassthroughVertexFactory->GetUpdatedFrameNumber());
-	// If world is paused, use current frame bone matrices, so velocity is canceled and skeletal mesh isn't blurred from motion.
-	bool bVerticesInMotion = !View->Family->bWorldIsPaused && bBoneDataUpdatedThisFrame;
+	uint32 UpdatedFrameNumber = 0;
+	FRHIShaderResourceView* PositionSRV = nullptr;
+	FRHIShaderResourceView* PrevPositionSRV = nullptr;
 
-	const bool bUsesSkinCache = BatchUserData != nullptr;
-	if (bUsesSkinCache)
+	if (BatchUserData->SkinCacheEntry != nullptr)
 	{
-		Parameters.GPUSkinPassThroughPreviousPositionBuffer = bVerticesInMotion ?
-			FGPUSkinCache::GetPreviousPositionBuffer(BatchUserData->Entry, SectionIndex)->SRV :
-			FGPUSkinCache::GetPositionBuffer(BatchUserData->Entry, SectionIndex)->SRV;
+		// Using Skin Cache.
+		UpdatedFrameNumber = FGPUSkinCache::GetUpdatedFrame(BatchUserData->SkinCacheEntry, SectionIndex);
+		PositionSRV = FGPUSkinCache::GetPositionBuffer(BatchUserData->SkinCacheEntry, SectionIndex)->SRV;
+		PrevPositionSRV = FGPUSkinCache::GetPreviousPositionBuffer(BatchUserData->SkinCacheEntry, SectionIndex)->SRV;
 	}
-	else // Mesh deformer
+	else if (BatchUserData->DeformerGeometry != nullptr)
 	{
-		Parameters.GPUSkinPassThroughPreviousPositionBuffer = (bVerticesInMotion && PassthroughVertexFactory->PrevPositionRDG) ?
-			PassthroughVertexFactory->PrevPositionRDG->GetOrCreateSRV(FRHIBufferSRVCreateInfo(PF_R32_FLOAT)) :
-			PassthroughVertexFactory->GetPositionsSRV();
+		// Using Mesh Deformers.
+		FMeshDeformerGeometry& DeformerGeometry = *BatchUserData->DeformerGeometry;
+		UpdatedFrameNumber = DeformerGeometry.PositionUpdatedFrame;
+		PositionSRV = DeformerGeometry.PositionSRV;
+		PrevPositionSRV = DeformerGeometry.PrevPositionSRV;
 	}
 
-	PassthroughVertexFactory->LooseParametersUniformBuffer.UpdateUniformBufferImmediate(Parameters);
+	// Bone data is updated whenever animation triggers a dynamic update, animation can skip frames hence the frequency is not necessary every frame.
+	// So check if bone data is updated this frame, if not then the previous frame data is stale and not suitable for motion blur.
+	bool bBoneDataUpdatedThisFrame = (View->Family->FrameNumber == UpdatedFrameNumber);
+	// If world is paused, use current frame bone matrices, so velocity is canceled and skeletal mesh isn't blurred from motion.
+	bool bVerticesInMotion = !View->Family->bWorldIsPaused && bBoneDataUpdatedThisFrame;
+
+	PositionSRV = (PositionSRV == nullptr) ? LocalVertexFactory->GetPositionsSRV() : PositionSRV;
+	PrevPositionSRV = (PrevPositionSRV == nullptr || !bVerticesInMotion) ? PositionSRV : PrevPositionSRV;
+
+	FLocalVertexFactoryLooseParameters Parameters;
+	Parameters.GPUSkinPassThroughPositionBuffer = PositionSRV;
+	Parameters.GPUSkinPassThroughPreviousPositionBuffer = PrevPositionSRV;
+	
+	LocalVertexFactory->LooseParametersUniformBuffer.UpdateUniformBufferImmediate(Parameters);
 }
 
 uint8 FSkeletalMeshSceneProxy::GetCurrentFirstLODIdx_Internal() const
diff --git a/Engine/Source/Runtime/Engine/Private/SkeletalMeshDeformerHelpers.cpp b/Engine/Source/Runtime/Engine/Private/SkeletalMeshDeformerHelpers.cpp
--- a/Engine/Source/Runtime/Engine/Private/SkeletalMeshDeformerHelpers.cpp
+++ b/Engine/Source/Runtime/Engine/Private/SkeletalMeshDeformerHelpers.cpp
@@ -1,157 +1,300 @@
 // Copyright Epic Games, Inc. All Rights Reserved.
 #include "SkeletalMeshDeformerHelpers.h"
 
+#include "DataDrivenShaderPlatformInfo.h"
 #include "GPUSkinVertexFactory.h"
 #include "SkeletalRenderGPUSkin.h"
+#include "RenderGraphBuilder.h"
 
 FRHIShaderResourceView* FSkeletalMeshDeformerHelpers::GetBoneBufferForReading(
-	FSkeletalMeshObject const* MeshObject,
-	int32 LODIndex,
-	int32 SectionIndex,
-	bool bPreviousFrame)
+	FSkeletalMeshObject const* InMeshObject,
+	int32 InLodIndex,
+	int32 InSectionIndex,
+	bool bInPreviousFrame)
 {
-	if (MeshObject->IsCPUSkinned())
+	if (InMeshObject->IsCPUSkinned())
 	{
 		return nullptr;
 	}
 
-	FSkeletalMeshObjectGPUSkin const* MeshObjectGPU = static_cast<FSkeletalMeshObjectGPUSkin const*>(MeshObject);
-	FGPUBaseSkinVertexFactory const* BaseVertexFactory = MeshObjectGPU->GetBaseSkinVertexFactory(LODIndex, SectionIndex);
-	bool bHasBoneBuffer = BaseVertexFactory != nullptr && BaseVertexFactory->GetShaderData().HasBoneBufferForReading(bPreviousFrame);
+	FSkeletalMeshObjectGPUSkin const* MeshObjectGPU = static_cast<FSkeletalMeshObjectGPUSkin const*>(InMeshObject);
+	FGPUBaseSkinVertexFactory const* BaseVertexFactory = MeshObjectGPU->GetBaseSkinVertexFactory(InLodIndex, InSectionIndex);
+	bool bHasBoneBuffer = BaseVertexFactory != nullptr && BaseVertexFactory->GetShaderData().HasBoneBufferForReading(bInPreviousFrame);
 	if (!bHasBoneBuffer)
 	{
 		return nullptr;
 	}
 
-	return BaseVertexFactory->GetShaderData().GetBoneBufferForReading(bPreviousFrame).VertexBufferSRV;
+	return BaseVertexFactory->GetShaderData().GetBoneBufferForReading(bInPreviousFrame).VertexBufferSRV;
 }
 
 FRHIShaderResourceView* FSkeletalMeshDeformerHelpers::GetMorphTargetBufferForReading(
-	FSkeletalMeshObject const* MeshObject,
-	int32 LODIndex,
-	int32 SectionIndex,
-	uint32 FrameNumber,
-	bool bPreviousFrame)
+	FSkeletalMeshObject const* InMeshObject,
+	int32 InLodIndex,
+	int32 InSectionIndex,
+	uint32 InFrameNumber,
+	bool bInPreviousFrame)
 {
-	if (MeshObject->IsCPUSkinned())
+	if (InMeshObject->IsCPUSkinned())
 	{
 		return nullptr;
 	}
 
-	FSkeletalMeshObjectGPUSkin const* MeshObjectGPU = static_cast<FSkeletalMeshObjectGPUSkin const*>(MeshObject);
-	FGPUBaseSkinVertexFactory const* BaseVertexFactory = MeshObjectGPU->GetBaseSkinVertexFactory(LODIndex, SectionIndex);
-	FMorphVertexBuffer const* MorphVertexBuffer = BaseVertexFactory != nullptr ? BaseVertexFactory->GetMorphVertexBuffer(bPreviousFrame) : nullptr;
+	FSkeletalMeshObjectGPUSkin const* MeshObjectGPU = static_cast<FSkeletalMeshObjectGPUSkin const*>(InMeshObject);
+	FGPUBaseSkinVertexFactory const* BaseVertexFactory = MeshObjectGPU->GetBaseSkinVertexFactory(InLodIndex, InSectionIndex);
+	FMorphVertexBuffer const* MorphVertexBuffer = BaseVertexFactory != nullptr ? BaseVertexFactory->GetMorphVertexBuffer(bInPreviousFrame) : nullptr;
 
 	return MorphVertexBuffer != nullptr ? MorphVertexBuffer->GetSRV() : nullptr;
 }
 
 FSkeletalMeshDeformerHelpers::FClothBuffers FSkeletalMeshDeformerHelpers::GetClothBuffersForReading(
-	FSkeletalMeshObject const* MeshObject,
-	int32 LODIndex,
-	int32 SectionIndex,
-	uint32 FrameNumber,
-	bool bPreviousFrame)
+	FSkeletalMeshObject const* InMeshObject,
+	int32 InLodIndex,
+	int32 InSectionIndex,
+	uint32 InFrameNumber,
+	bool bInPreviousFrame)
 {
-	if (MeshObject->IsCPUSkinned())
+	if (InMeshObject->IsCPUSkinned())
 	{
 		return FClothBuffers();
 	}
 
-	FSkeletalMeshObjectGPUSkin const* MeshObjectGPU = static_cast<FSkeletalMeshObjectGPUSkin const*>(MeshObject);
-	FGPUBaseSkinVertexFactory const* BaseVertexFactory = MeshObjectGPU->GetBaseSkinVertexFactory(LODIndex, SectionIndex);
+	FSkeletalMeshObjectGPUSkin const* MeshObjectGPU = static_cast<FSkeletalMeshObjectGPUSkin const*>(InMeshObject);
+	FGPUBaseSkinVertexFactory const* BaseVertexFactory = MeshObjectGPU->GetBaseSkinVertexFactory(InLodIndex, InSectionIndex);
 	FGPUBaseSkinAPEXClothVertexFactory const* ClothVertexFactory = BaseVertexFactory != nullptr ? BaseVertexFactory->GetClothVertexFactory() : nullptr;
 
-	if (ClothVertexFactory == nullptr || !ClothVertexFactory->GetClothShaderData().HasClothBufferForReading(bPreviousFrame))
+	if (ClothVertexFactory == nullptr || !ClothVertexFactory->GetClothShaderData().HasClothBufferForReading(bInPreviousFrame))
 	{
 		return FClothBuffers();
 	}
 
-	FSkeletalMeshRenderData const& SkeletalMeshRenderData = MeshObject->GetSkeletalMeshRenderData();
-	FSkeletalMeshLODRenderData const* LodRenderData = SkeletalMeshRenderData.GetPendingFirstLOD(LODIndex);
-	FSkelMeshRenderSection const& RenderSection = LodRenderData->RenderSections[SectionIndex];
+	FSkeletalMeshRenderData const& SkeletalMeshRenderData = InMeshObject->GetSkeletalMeshRenderData();
+	FSkeletalMeshLODRenderData const* LodRenderData = SkeletalMeshRenderData.GetPendingFirstLOD(InLodIndex);
+	FSkelMeshRenderSection const& RenderSection = LodRenderData->RenderSections[InSectionIndex];
 
 	FClothBuffers Ret;
 	Ret.ClothInfluenceBuffer = ClothVertexFactory->GetClothBuffer();
 	Ret.ClothInfluenceBufferOffset = ClothVertexFactory->GetClothIndexOffset(RenderSection.BaseVertexIndex);
-	Ret.ClothSimulatedPositionAndNormalBuffer = ClothVertexFactory->GetClothShaderData().GetClothBufferForReading(bPreviousFrame).VertexBufferSRV;
-	Ret.ClothToLocal = ClothVertexFactory->GetClothShaderData().GetClothToLocalForReading(bPreviousFrame);
+	Ret.ClothSimulatedPositionAndNormalBuffer = ClothVertexFactory->GetClothShaderData().GetClothBufferForReading(bInPreviousFrame).VertexBufferSRV;
+	Ret.ClothToLocal = ClothVertexFactory->GetClothShaderData().GetClothToLocalForReading(bInPreviousFrame);
 	return Ret;
 }
 
-void FSkeletalMeshDeformerHelpers::SetVertexFactoryBufferOverrides(
-	FSkeletalMeshObject* MeshObject,
-	int32 LODIndex,
-	EOverrideType OverrideType,
-	TRefCountPtr<FRDGPooledBuffer> const& PositionBuffer,
-	TRefCountPtr<FRDGPooledBuffer> const& TangentBuffer,
-	TRefCountPtr<FRDGPooledBuffer> const& ColorBuffer)
+FRDGBuffer* FSkeletalMeshDeformerHelpers::AllocateVertexFactoryPositionBuffer(FRDGBuilder& GraphBuilder, FSkeletalMeshObject* InMeshObject, int32 InLodIndex, TCHAR const* InBufferName)
 {
-	if (MeshObject->IsCPUSkinned())
+	if (InMeshObject->IsCPUSkinned())
 	{
-		return;
+		return nullptr;
 	}
 
-	FGPUSkinPassthroughVertexFactory::EOverrideFlags OverrideFlags = FGPUSkinPassthroughVertexFactory::EOverrideFlags::All;
-	if (OverrideType == EOverrideType::Partial)
+	FSkeletalMeshRenderData const& SkeletalMeshRenderData = InMeshObject->GetSkeletalMeshRenderData();
+	FSkeletalMeshLODRenderData const* LodRenderData = SkeletalMeshRenderData.GetPendingFirstLOD(InLodIndex);
+	const int32 NumVertices = LodRenderData->GetNumVertices();
+
+	FSkeletalMeshObjectGPUSkin* MeshObjectGPU = static_cast<FSkeletalMeshObjectGPUSkin*>(InMeshObject);
+	FGPUBaseSkinVertexFactory const* BaseVertexFactory = MeshObjectGPU->GetBaseSkinVertexFactory(InLodIndex, 0);
+	uint32 Frame = BaseVertexFactory->GetShaderData().UpdatedFrameNumber;
+
+	FMeshDeformerGeometry& DeformerGeometry = MeshObjectGPU->GetDeformerGeometry(InLodIndex);
+
+	FRDGBuffer* PositionBuffer = nullptr;
+	if (DeformerGeometry.Position && Frame == DeformerGeometry.PositionUpdatedFrame)
 	{
-		OverrideFlags = FGPUSkinPassthroughVertexFactory::EOverrideFlags::None;
- 		OverrideFlags |= (PositionBuffer.IsValid() ? FGPUSkinPassthroughVertexFactory::EOverrideFlags::Position : FGPUSkinPassthroughVertexFactory::EOverrideFlags::None);
- 		OverrideFlags |= (TangentBuffer.IsValid() ? FGPUSkinPassthroughVertexFactory::EOverrideFlags::Tangent : FGPUSkinPassthroughVertexFactory::EOverrideFlags::None);
- 		OverrideFlags |= (ColorBuffer.IsValid() ? FGPUSkinPassthroughVertexFactory::EOverrideFlags::Color : FGPUSkinPassthroughVertexFactory::EOverrideFlags::None);
+		PositionBuffer = GraphBuilder.RegisterExternalBuffer(DeformerGeometry.Position);
 	}
-
-	FSkeletalMeshObjectGPUSkin* MeshObjectGPU = static_cast<FSkeletalMeshObjectGPUSkin*>(MeshObject);
-	const FSkeletalMeshObjectGPUSkin::FSkeletalMeshObjectLOD& LOD = MeshObjectGPU->LODs[LODIndex];
-	
-	const int32 NumSections = MeshObject->GetRenderSections(LODIndex).Num();
-	for (int32 SectionIndex = 0; SectionIndex < NumSections; ++SectionIndex)
+	else
 	{
-		FGPUBaseSkinVertexFactory const* BaseVertexFactory = MeshObjectGPU->GetBaseSkinVertexFactory(LODIndex, SectionIndex);
-		FGPUSkinPassthroughVertexFactory* TargetVertexFactory = LOD.GPUSkinVertexFactories.PassthroughVertexFactories[SectionIndex].Get();
-		TargetVertexFactory->InvalidateStreams();
-		TargetVertexFactory->UpdateVertexDeclaration(OverrideFlags, BaseVertexFactory, PositionBuffer, TangentBuffer, ColorBuffer);
+		DeformerGeometry.PrevPosition = DeformerGeometry.Position;
+		DeformerGeometry.PrevPositionSRV = DeformerGeometry.PositionSRV;
+
+		const uint32 PosBufferBytesPerElement = 4;
+		PositionBuffer = GraphBuilder.CreateBuffer(FRDGBufferDesc::CreateBufferDesc(PosBufferBytesPerElement, NumVertices * 3), InBufferName, ERDGBufferFlags::None);
+		DeformerGeometry.Position = GraphBuilder.ConvertToExternalBuffer(PositionBuffer);
+		DeformerGeometry.PositionSRV = DeformerGeometry.Position->GetOrCreateSRV(FRHIBufferSRVCreateInfo(PF_R32_FLOAT));
+		DeformerGeometry.PositionUpdatedFrame = Frame;
+		GraphBuilder.SetBufferAccessFinal(PositionBuffer, ERHIAccess::VertexOrIndexBuffer | ERHIAccess::SRVMask);
+
+#if RHI_RAYTRACING
+		// Update ray tracing geometry whenever we recreate the position buffer.
+		FSkeletalMeshRenderData& SkelMeshRenderData = MeshObjectGPU->GetSkeletalMeshRenderData();
+		FSkeletalMeshLODRenderData& LODModel = SkelMeshRenderData.LODRenderData[InLodIndex];
+
+		const int32 NumSections = InMeshObject->GetRenderSections(InLodIndex).Num();
+
+		TArray<FBufferRHIRef> VertexBuffers;
+		VertexBuffers.SetNum(NumSections);
+		for (int32 SectionIndex = 0; SectionIndex < NumSections; ++SectionIndex)
+		{
+			VertexBuffers[SectionIndex] = DeformerGeometry.Position->GetRHI();
+		}
+
+		MeshObjectGPU->UpdateRayTracingGeometry(LODModel, InLodIndex, VertexBuffers);
+#endif // RHI_RAYTRACING
 	}
+
+	return PositionBuffer;
 }
 
-void FSkeletalMeshDeformerHelpers::ResetVertexFactoryBufferOverrides_GameThread(FSkeletalMeshObject* MeshObject, int32 LODIndex)
+FRDGBuffer* FSkeletalMeshDeformerHelpers::AllocateVertexFactoryTangentBuffer(FRDGBuilder& GraphBuilder, FSkeletalMeshObject* InMeshObject, int32 InLodIndex, TCHAR const* InBufferName)
 {
-	if (MeshObject->IsCPUSkinned())
+	if (InMeshObject->IsCPUSkinned())
 	{
-		return;
+		return nullptr;
 	}
 
-	ENQUEUE_RENDER_COMMAND(ResetSkinPassthroughVertexFactory)([MeshObject, LODIndex](FRHICommandList& CmdList)
+	FSkeletalMeshRenderData const& SkeletalMeshRenderData = InMeshObject->GetSkeletalMeshRenderData();
+	FSkeletalMeshLODRenderData const* LodRenderData = SkeletalMeshRenderData.GetPendingFirstLOD(InLodIndex);
+	const int32 NumVertices = LodRenderData->GetNumVertices();
+
+	FSkeletalMeshObjectGPUSkin* MeshObjectGPU = static_cast<FSkeletalMeshObjectGPUSkin*>(InMeshObject);
+	FGPUBaseSkinVertexFactory const* BaseVertexFactory = MeshObjectGPU->GetBaseSkinVertexFactory(InLodIndex, 0);
+	uint32 Frame = BaseVertexFactory->GetShaderData().UpdatedFrameNumber;
+
+	FMeshDeformerGeometry& DeformerGeometry = MeshObjectGPU->GetDeformerGeometry(InLodIndex);
+
+	FRDGBuffer* TangentBuffer = nullptr;
+	if (DeformerGeometry.Tangent && Frame == DeformerGeometry.TangentUpdatedFrame)
 	{
-		SetVertexFactoryBufferOverrides(MeshObject, LODIndex, EOverrideType::All, nullptr, nullptr, nullptr);
-	});
+		TangentBuffer = GraphBuilder.RegisterExternalBuffer(DeformerGeometry.Tangent);
+	}
+	else
+	{
+		const uint32 TangentBufferBytesPerElement = 8;
+		TangentBuffer = GraphBuilder.CreateBuffer(FRDGBufferDesc::CreateBufferDesc(TangentBufferBytesPerElement, NumVertices * 2), InBufferName, ERDGBufferFlags::None);
+		DeformerGeometry.Tangent = GraphBuilder.ConvertToExternalBuffer(TangentBuffer);
+		const EPixelFormat TangentsFormat = IsOpenGLPlatform(GMaxRHIShaderPlatform) ? PF_R16G16B16A16_SINT : PF_R16G16B16A16_SNORM;
+		DeformerGeometry.TangentSRV = DeformerGeometry.Tangent->GetOrCreateSRV(FRHIBufferSRVCreateInfo(TangentsFormat));
+		DeformerGeometry.TangentUpdatedFrame = Frame;
+		GraphBuilder.SetBufferAccessFinal(TangentBuffer, ERHIAccess::VertexOrIndexBuffer | ERHIAccess::SRVMask);
+	}
+
+	return TangentBuffer;
 }
 
-#if RHI_RAYTRACING
+FRDGBuffer* FSkeletalMeshDeformerHelpers::AllocateVertexFactoryColorBuffer(FRDGBuilder& GraphBuilder, FSkeletalMeshObject* InMeshObject, int32 InLodIndex, TCHAR const* InBufferName)
+{
+	if (InMeshObject->IsCPUSkinned())
+	{
+		return nullptr;
+	}
+
+	FSkeletalMeshRenderData const& SkeletalMeshRenderData = InMeshObject->GetSkeletalMeshRenderData();
+	FSkeletalMeshLODRenderData const* LodRenderData = SkeletalMeshRenderData.GetPendingFirstLOD(InLodIndex);
+	const int32 NumVertices = LodRenderData->GetNumVertices();
+
+	FSkeletalMeshObjectGPUSkin* MeshObjectGPU = static_cast<FSkeletalMeshObjectGPUSkin*>(InMeshObject);
+	FGPUBaseSkinVertexFactory const* BaseVertexFactory = MeshObjectGPU->GetBaseSkinVertexFactory(InLodIndex, 0);
+	uint32 Frame = BaseVertexFactory->GetShaderData().UpdatedFrameNumber;
+
+	FMeshDeformerGeometry& DeformerGeometry = MeshObjectGPU->GetDeformerGeometry(InLodIndex);
 
-void FSkeletalMeshDeformerHelpers::UpdateRayTracingGeometry(
-	FSkeletalMeshObject* MeshObject,
-	int32 LODIndex,
-	TRefCountPtr<FRDGPooledBuffer> const& PositionBuffer)
+	FRDGBuffer* ColorBuffer = nullptr;
+	if (DeformerGeometry.Color && Frame == DeformerGeometry.ColorUpdatedFrame)
+	{
+		ColorBuffer = GraphBuilder.RegisterExternalBuffer(DeformerGeometry.Color);
+	}
+	else
+	{
+		const uint32 ColorBufferBytesPerElement = 4;
+		ColorBuffer = GraphBuilder.CreateBuffer(FRDGBufferDesc::CreateBufferDesc(ColorBufferBytesPerElement, NumVertices), InBufferName, ERDGBufferFlags::None);
+		DeformerGeometry.Color = GraphBuilder.ConvertToExternalBuffer(ColorBuffer);
+		DeformerGeometry.ColorSRV = DeformerGeometry.Color->GetOrCreateSRV(FRHIBufferSRVCreateInfo(PF_R8G8B8A8));
+		DeformerGeometry.ColorUpdatedFrame = Frame;
+		GraphBuilder.SetBufferAccessFinal(ColorBuffer, ERHIAccess::VertexOrIndexBuffer | ERHIAccess::SRVMask);
+	}
+
+	return ColorBuffer;
+}
+
+void FSkeletalMeshDeformerHelpers::UpdateVertexFactoryBufferOverrides(FSkeletalMeshObject* InMeshObject, int32 InLodIndex)
 {
-	if (MeshObject->IsCPUSkinned())
+	if (InMeshObject->IsCPUSkinned())
 	{
 		return;
 	}
 
-	FSkeletalMeshObjectGPUSkin* MeshObjectGPU = static_cast<FSkeletalMeshObjectGPUSkin*>(MeshObject);
-	FSkeletalMeshRenderData& SkelMeshRenderData = MeshObjectGPU->GetSkeletalMeshRenderData();
-	FSkeletalMeshLODRenderData& LODModel = SkelMeshRenderData.LODRenderData[LODIndex];
+	FSkeletalMeshObjectGPUSkin* MeshObjectGPU = static_cast<FSkeletalMeshObjectGPUSkin*>(InMeshObject);
+	FMeshDeformerGeometry& DeformerGeometry = MeshObjectGPU->GetDeformerGeometry(InLodIndex);
+	
+	TArray<FGPUSkinPassthroughVertexFactory::EVertexAtttribute, TInlineAllocator<4>> Attributes;
+	TArray<FRHIShaderResourceView*, TInlineAllocator<4>> SRVs;
 
-	const int32 NumSections = MeshObject->GetRenderSections(LODIndex).Num();
+	if (DeformerGeometry.PositionSRV)
+	{
+		Attributes.Add(FGPUSkinPassthroughVertexFactory::EVertexAtttribute::Position);
+		SRVs.Add(DeformerGeometry.PositionSRV);
+	}
+	if (DeformerGeometry.TangentSRV)
+	{
+		Attributes.Add(FGPUSkinPassthroughVertexFactory::EVertexAtttribute::Tangent);
+		SRVs.Add(DeformerGeometry.TangentSRV);
+	}
+	if (DeformerGeometry.ColorSRV)
+	{
+		Attributes.Add(FGPUSkinPassthroughVertexFactory::EVertexAtttribute::Color);
+		SRVs.Add(DeformerGeometry.ColorSRV);
+	}
 
-	TArray<FBufferRHIRef> VertexBufffers;
-	VertexBufffers.SetNum(NumSections);
-	for (int32 SectionIndex = 0; SectionIndex < NumSections; ++SectionIndex)
+	if (Attributes.Num() == 0)
 	{
-		VertexBufffers[SectionIndex] = PositionBuffer->GetRHI();
+		return;
 	}
 
-	MeshObjectGPU->UpdateRayTracingGeometry(LODModel, LODIndex, VertexBufffers);
+	const FSkeletalMeshObjectGPUSkin::FSkeletalMeshObjectLOD& LOD = MeshObjectGPU->LODs[InLodIndex];
+	const int32 NumSections = InMeshObject->GetRenderSections(InLodIndex).Num();
+	for (int32 SectionIndex = 0; SectionIndex < NumSections; ++SectionIndex)
+	{
+		FGPUBaseSkinVertexFactory const* BaseVertexFactory = MeshObjectGPU->GetBaseSkinVertexFactory(InLodIndex, SectionIndex);
+		FGPUSkinPassthroughVertexFactory* TargetVertexFactory = LOD.GPUSkinVertexFactories.PassthroughVertexFactories[SectionIndex].Get();
+		
+		TargetVertexFactory->AddVertexAttributes(BaseVertexFactory, Attributes, SRVs);
+	}
 }
 
+void FSkeletalMeshDeformerHelpers::ResetVertexFactoryBufferOverrides_GameThread(FSkeletalMeshObject* InMeshObject, int32 LODIndex)
+{
+	if (InMeshObject->IsCPUSkinned())
+	{
+		return;
+	}
+
+	ENQUEUE_RENDER_COMMAND(ResetSkinPassthroughVertexFactory)([InMeshObject, LODIndex](FRHICommandList& CmdList)
+	{
+		FSkeletalMeshObjectGPUSkin* MeshObjectGPU = static_cast<FSkeletalMeshObjectGPUSkin*>(InMeshObject);
+		FMeshDeformerGeometry& DeformerGeometry = MeshObjectGPU->GetDeformerGeometry(LODIndex);
+
+		// This can be called per frame even when already reset. So early out if we don't need to do anything.
+		const bool bIsReset = DeformerGeometry.PositionUpdatedFrame == 0 && DeformerGeometry.TangentUpdatedFrame == 0 && DeformerGeometry.ColorUpdatedFrame == 0;
+		if (bIsReset)
+		{
+			return;
+		}
+
+		// Reset stored buffers.
+		DeformerGeometry.Reset();
+
+		// Reset vertex factories.
+		const FSkeletalMeshObjectGPUSkin::FSkeletalMeshObjectLOD& LOD = MeshObjectGPU->LODs[LODIndex];
+		const int32 NumSections = InMeshObject->GetRenderSections(LODIndex).Num();
+		for (int32 SectionIndex = 0; SectionIndex < NumSections; ++SectionIndex)
+		{
+			FGPUBaseSkinVertexFactory const* BaseVertexFactory = MeshObjectGPU->GetBaseSkinVertexFactory(LODIndex, SectionIndex);
+			FGPUSkinPassthroughVertexFactory* TargetVertexFactory = LOD.GPUSkinVertexFactories.PassthroughVertexFactories[SectionIndex].Get();
+			TargetVertexFactory->ResetVertexAttributes();
+			FGPUSkinPassthroughVertexFactory::FDataType Data;
+			BaseVertexFactory->CopyDataTypeForLocalVertexFactory(Data);
+			TargetVertexFactory->SetData(Data);
+		}
+
+#if RHI_RAYTRACING
+		// Reset ray tracing geometry.
+		FSkeletalMeshRenderData& SkelMeshRenderData = MeshObjectGPU->GetSkeletalMeshRenderData();
+		FSkeletalMeshLODRenderData& LODModel = SkelMeshRenderData.LODRenderData[LODIndex];
+		FBufferRHIRef VertexBuffer = LODModel.StaticVertexBuffers.PositionVertexBuffer.VertexBufferRHI;
+
+		TArray<FBufferRHIRef> VertexBuffers;
+		VertexBuffers.Init(VertexBuffer, NumSections);
+		MeshObjectGPU->UpdateRayTracingGeometry(LODModel, LODIndex, VertexBuffers);
 #endif // RHI_RAYTRACING
+	});
+}
diff --git a/Engine/Source/Runtime/Engine/Private/SkeletalRenderGPUSkin.cpp b/Engine/Source/Runtime/Engine/Private/SkeletalRenderGPUSkin.cpp
--- a/Engine/Source/Runtime/Engine/Private/SkeletalRenderGPUSkin.cpp
+++ b/Engine/Source/Runtime/Engine/Private/SkeletalRenderGPUSkin.cpp
@@ -1544,6 +1544,21 @@ const FVertexFactory* FSkeletalMeshObjectGPUSkin::GetSkinVertexFactory(const FSc
 	return GetBaseSkinVertexFactory(LODIndex, ChunkIdx);
 }
 
+const FSkinBatchVertexFactoryUserData* FSkeletalMeshObjectGPUSkin::GetVertexFactoryUserData(const int32 LODIndex, int32 ChunkIdx, ESkinVertexFactoryMode VFMode) const
+{
+#if RHI_RAYTRACING
+	FGPUSkinCacheEntry* Entry = VFMode == ESkinVertexFactoryMode::RayTracing ? GetSkinCacheEntryForRayTracing() : SkinCacheEntry;
+#else
+	FGPUSkinCacheEntry* Entry = SkinCacheEntry;
+#endif
+	if (Entry != nullptr)
+	{
+		return FGPUSkinCache::GetVertexFactoryUserData(SkinCacheEntry, ChunkIdx);
+	}
+
+	return &LODs[LODIndex].DeformerGeometry.VertexFactoryUserData;
+}
+
 FGPUBaseSkinVertexFactory const* FSkeletalMeshObjectGPUSkin::GetBaseSkinVertexFactory(int32 LODIndex, int32 ChunkIdx) const
 {
 	const FSkeletalMeshObjectLOD& LOD = LODs[LODIndex];
@@ -1801,17 +1816,6 @@ static FGPUBaseSkinVertexFactory* CreateVertexFactory(TArray<TUniquePtr<FGPUBase
 	return VertexFactory;
 }
 
-void FGPUSkinPassthroughVertexFactory::SetData(const FDataType& InData)
-{
-	FLocalVertexFactory::SetData(InData);
-	const int32 DefaultBaseVertexIndex = 0;
-	const int32 DefaultPreSkinBaseVertexIndex = 0;
-	if (RHISupportsManualVertexFetch(GetFeatureLevelShaderPlatform(GetFeatureLevel())))
-	{
-		UniformBuffer = CreateLocalVFUniformBuffer(this, Data.LODLightmapDataIndex, nullptr, DefaultBaseVertexIndex, DefaultPreSkinBaseVertexIndex);
-	}
-}
-
 void UpdateVertexFactory(TArray<TUniquePtr<FGPUBaseSkinVertexFactory>>& VertexFactories,
 	const FSkeletalMeshObjectGPUSkin::FVertexFactoryBuffers& InVertexBuffers)
 {
@@ -1846,7 +1850,9 @@ static void CreatePassthroughVertexFactory(ERHIFeatureLevel::Type InFeatureLevel
 	ENQUEUE_RENDER_COMMAND(InitPassthroughGPUSkinVertexFactory)(
 		[NewPassthroughVertexFactory, SourceVertexFactory](FRHICommandList& RHICmdList)
 		{
-			SourceVertexFactory->CopyDataTypeForPassthroughFactory(NewPassthroughVertexFactory);
+			FLocalVertexFactory::FDataType Data;
+			SourceVertexFactory->CopyDataTypeForLocalVertexFactory(Data);
+			NewPassthroughVertexFactory->SetData(Data);
 			NewPassthroughVertexFactory->InitResource();
 		}
 	);
@@ -2260,6 +2266,11 @@ const TArray<FMatrix44f>& FSkeletalMeshObjectGPUSkin::GetReferenceToLocalMatrice
 	return DynamicData->ReferenceToLocal;
 }
 
+FMeshDeformerGeometry& FSkeletalMeshObjectGPUSkin::GetDeformerGeometry(int32 LODIndex)
+{
+	return LODs[LODIndex].DeformerGeometry;
+}
+
 bool FSkeletalMeshObjectGPUSkin::GetCachedGeometry(FCachedGeometry& OutCachedGeometry) const
 {
 	OutCachedGeometry = FCachedGeometry();
@@ -2297,22 +2308,20 @@ bool FSkeletalMeshObjectGPUSkin::GetCachedGeometry(FCachedGeometry& OutCachedGeo
 		}
 		else
 		{
-			// Pull the cached geometry SRV directly out of the vertex factory.
-			// This catches cases where we have setup deformed data outside of the skin cache (eg mesh deformers).
-			// This approach _could_ be used for the skin cache case as well, except that skin cache only 
-			// actually pokes the vertex factory once instead of every frame. It flips current/previous buffers
-			// by setting SRVs directly in GetShaderBindings().
-			FGPUSkinPassthroughVertexFactory* VertexFactory = VertexFactories.PassthroughVertexFactories[SectionIndex].Get();
-			if (VertexFactory == nullptr || VertexFactory->GetPositionsSRV() == nullptr)
+			
+			// Get the cached geometry SRVs from the deformer geometry.
+			FMeshDeformerGeometry const& DeformerGeometry = LODs[LodIndex].DeformerGeometry;
+
+			if (!DeformerGeometry.Position.IsValid())
 			{
 				// Reset all output if one section isn't available.
 				OutCachedGeometry.Sections.Reset();
 				return false;
 			}
 
-			CachedSection.PositionBuffer = VertexFactory->GetPositionsSRV();
-			FRHIShaderResourceView* PreviousPositionBuffer = VertexFactory->GetPreviousPositionsSRV();
-			CachedSection.PreviousPositionBuffer = PreviousPositionBuffer != nullptr ? PreviousPositionBuffer : CachedSection.PositionBuffer;
+			CachedSection.PositionBuffer = DeformerGeometry.PositionSRV;
+			CachedSection.PreviousPositionBuffer = DeformerGeometry.PrevPositionSRV;
+			CachedSection.PreviousPositionBuffer = CachedSection.PreviousPositionBuffer != nullptr ? CachedSection.PreviousPositionBuffer : CachedSection.PositionBuffer;
 		}
 				
 		CachedSection.IndexBuffer = LODRenderData.MultiSizeIndexContainer.GetIndexBuffer()->GetSRV();
diff --git a/Engine/Source/Runtime/Engine/Private/SkeletalRenderGPUSkin.h b/Engine/Source/Runtime/Engine/Private/SkeletalRenderGPUSkin.h
--- a/Engine/Source/Runtime/Engine/Private/SkeletalRenderGPUSkin.h
+++ b/Engine/Source/Runtime/Engine/Private/SkeletalRenderGPUSkin.h
@@ -19,6 +19,7 @@
 #include "ClothingSystemRuntimeTypes.h"
 #include "Rendering/SkeletalMeshRenderData.h"
 #include "Rendering/SkeletalMeshLODRenderData.h"
+#include "Animation/MeshDeformerGeometry.h"
 
 enum class EGPUSkinCacheEntryMode;
 class FGPUSkinCache;
@@ -350,11 +351,14 @@ public:
 	void UpdateDynamicData_RenderThread(FGPUSkinCache* GPUSkinCache, FRHICommandListImmediate& RHICmdList, FDynamicSkelMeshObjectDataGPUSkin* InDynamicData, FSceneInterface* Scene, uint32 FrameNumberToPrepare, uint32 RevisionNumber);
 	virtual void PreGDMECallback(FGPUSkinCache* GPUSkinCache, uint32 FrameNumber) override;
 	virtual const FVertexFactory* GetSkinVertexFactory(const FSceneView* View, int32 LODIndex,int32 ChunkIdx, ESkinVertexFactoryMode VFMode = ESkinVertexFactoryMode::Default) const override;
+	virtual const FSkinBatchVertexFactoryUserData* GetVertexFactoryUserData(const int32 LODIndex, int32 ChunkIdx, ESkinVertexFactoryMode VFMode) const override;
 	virtual void CacheVertices(int32 LODIndex, bool bForce) const override {}
 	virtual bool IsCPUSkinned() const override { return false; }
 	virtual TArray<FTransform>* GetComponentSpaceTransforms() const override;
 	virtual const TArray<FMatrix44f>& GetReferenceToLocalMatrices() const override;
 	virtual bool GetCachedGeometry(FCachedGeometry& OutCachedGeometry) const override;
+	
+	FMeshDeformerGeometry& GetDeformerGeometry(int32 LODIndex);
 
 #if RHI_RAYTRACING
 	/** Geometry for ray tracing. */
@@ -619,6 +623,9 @@ protected:
 		/** Color buffer to user, could be from asset or component override */
 		FColorVertexBuffer* MeshObjectColorBuffer;
 
+		/** Mesh deformer output buffers */
+		FMeshDeformerGeometry DeformerGeometry;
+
 		/**
 		 * Update the contents of the morphtarget vertex buffer by accumulating all 
 		 * delta positions and delta normals from the set of active morph targets
diff --git a/Engine/Source/Runtime/Engine/Public/GPUSkinCache.h b/Engine/Source/Runtime/Engine/Public/GPUSkinCache.h
--- a/Engine/Source/Runtime/Engine/Public/GPUSkinCache.h
+++ b/Engine/Source/Runtime/Engine/Public/GPUSkinCache.h
@@ -56,6 +56,7 @@ struct FClothSimulData;
 struct FSkelMeshRenderSection;
 struct FVertexBufferAndSRV;
 struct FRayTracingGeometrySegment;
+struct FSkinBatchVertexFactoryUserData;
 
 // Can the skin cache be used (ie shaders added, etc)
 extern ENGINE_API bool IsGPUSkinCacheAvailable(EShaderPlatform Platform);
@@ -90,12 +91,6 @@ struct FClothSimulEntry
 	}
 };
 
-struct FGPUSkinBatchElementUserData
-{
-	FGPUSkinCacheEntry* Entry;
-	int32 Section;
-};
-
 enum class EGPUSkinCacheEntryMode
 {
 	Raster,
@@ -154,17 +149,16 @@ public:
 		const FGPUSkinCacheEntry* Entry,
 		int32 Section,
 		const FGPUSkinPassthroughVertexFactory* VertexFactory,
-		FShaderResourceParameter GPUSkinCachePositionBuffer,
 		class FMeshDrawSingleShaderBindings& ShaderBindings,
 		FVertexInputStreamArray& VertexStreams);
 
 	static void Release(FGPUSkinCacheEntry*& SkinCacheEntry);
 
-	static inline FGPUSkinBatchElementUserData* GetFactoryUserData(FGPUSkinCacheEntry* Entry, int32 Section)
+	static inline const FSkinBatchVertexFactoryUserData* GetVertexFactoryUserData(FGPUSkinCacheEntry* Entry, int32 Section)
 	{
 		if (Entry)
 		{
-			return InternalGetFactoryUserData(Entry, Section);
+			return InternalGetVertexFactoryUserData(Entry, Section);
 		}
 		return nullptr;
 	}
@@ -394,6 +388,7 @@ public:
 	FGPUSkinCacheEntry const* GetSkinCacheEntry(uint32 ComponentId) const;
 	static FRWBuffer* GetPositionBuffer(FGPUSkinCacheEntry const* Entry, uint32 SectionIndex);
 	static FRWBuffer* GetPreviousPositionBuffer(FGPUSkinCacheEntry const* Entry, uint32 SectionIndex);
+	static uint32 GetUpdatedFrame(FGPUSkinCacheEntry const* Entry, uint32 SectionIndex);
 
 	// Deprecated function. Can remove include of CachedGeometry.h when this is removed.
 	UE_DEPRECATED(5.1, "Use GetPositionBuffer() or similar instead.")
@@ -440,7 +435,7 @@ protected:
 	void Cleanup();
 	static void TransitionAllToReadable(FRHICommandList& RHICmdList, const TSet<FSkinCacheRWBuffer*>& BuffersToTransitionToRead);
 	static void ReleaseSkinCacheEntry(FGPUSkinCacheEntry* SkinCacheEntry);
-	static FGPUSkinBatchElementUserData* InternalGetFactoryUserData(FGPUSkinCacheEntry* Entry, int32 Section);
+	static const FSkinBatchVertexFactoryUserData* InternalGetVertexFactoryUserData(FGPUSkinCacheEntry* Entry, int32 Section);
 	void InvalidateAllEntries();
 	uint64 UsedMemoryInBytes;
 	uint64 ExtraRequiredMemory;
diff --git a/Engine/Source/Runtime/Engine/Public/GPUSkinVertexFactory.h b/Engine/Source/Runtime/Engine/Public/GPUSkinVertexFactory.h
--- a/Engine/Source/Runtime/Engine/Public/GPUSkinVertexFactory.h
+++ b/Engine/Source/Runtime/Engine/Public/GPUSkinVertexFactory.h
@@ -19,8 +19,6 @@
 #include "Matrix3x4.h"
 #include "SkeletalMeshTypes.h"
 
-class FRDGPooledBuffer;
-
 template <class T> class TConsoleVariableData;
 
 // Uniform buffer for APEX cloth
@@ -420,7 +418,7 @@ public:
 		return Data->TangentBasisComponents[Index];
 	}
 
-	void CopyDataTypeForPassthroughFactory(class FGPUSkinPassthroughVertexFactory* PassthroughVertexFactory);
+	void CopyDataTypeForLocalVertexFactory(FLocalVertexFactory::FDataType& OutDestData) const;
 
 protected:
 	/**
@@ -487,130 +485,6 @@ private:
 	int32 MorphDeltaStreamIndex = -1;
 };
 
-/** 
- * Vertex factory with vertex stream components for GPU-skinned streams, enabled for passthrough mode when vertices have been pre-skinned 
- */
-class FGPUSkinPassthroughVertexFactory : public FLocalVertexFactory
-{
-	typedef FLocalVertexFactory Super;
-
-public:
-	FGPUSkinPassthroughVertexFactory(ERHIFeatureLevel::Type InFeatureLevel);
-	~FGPUSkinPassthroughVertexFactory();
-
-	inline int32 GetPositionStreamIndex() const
-	{
-		check(PositionStreamIndex > -1);
-		return PositionStreamIndex;
-	}
-
-	inline int32 GetTangentStreamIndex() const
-	{
-		return TangentStreamIndex;
-	}
-
-	void SetData(const FDataType& InData);
-
-	uint32 GetUpdatedFrameNumber() const { return UpdatedFrameNumber; }
-
-	//TODO should be supported
-	bool SupportsPositionOnlyStream() const override { return false; }
-	bool SupportsPositionAndNormalOnlyStream() const override { return false; }
-
-	inline void InvalidateStreams()
-	{
-		PositionStreamIndex = -1;
-		TangentStreamIndex = -1;
-
-		PositionVBAlias.ReleaseRHI();
-		TangentVBAlias.ReleaseRHI();
-		ColorVBAlias.ReleaseRHI();
-	}
-
-	virtual void ReleaseRHI() override;
-
-	/** Poke values into the vertex factory. */
-	inline void UpdateVertexDeclaration(
-		FGPUBaseSkinVertexFactory const* SourceVertexFactory, 
-		struct FRWBuffer* PositionRWBuffer, 
-		struct FRWBuffer* TangentRWBuffer)
-	{
-		UpdatedFrameNumber = SourceVertexFactory->GetShaderData().UpdatedFrameNumber;
-		if (PositionStreamIndex == -1)
-		{
-			InternalUpdateVertexDeclaration(SourceVertexFactory, PositionRWBuffer, TangentRWBuffer);
-		}
-	}
-
-	/** Flags for override bitmask passed to UpdateVertexDeclaration(). */
-	enum class EOverrideFlags
-	{
-		None		= 0,
-		Position	= 1 << 0,
-		Tangent		= 1 << 1,
-		Color		= 1 << 2,
-		All			= 0xff,
-	};
-
-	/** Poke values into the vertex factory. */
-	inline void UpdateVertexDeclaration(
-		EOverrideFlags OverrideFlags,
-		FGPUBaseSkinVertexFactory const* SourceVertexFactory,
-		TRefCountPtr<FRDGPooledBuffer> const& PositionBuffer,
-		TRefCountPtr<FRDGPooledBuffer> const& TangentBuffer,
-		TRefCountPtr<FRDGPooledBuffer> const& ColorBuffer)
-	{
-		UpdatedFrameNumber = SourceVertexFactory->GetShaderData().UpdatedFrameNumber;
-		if (PositionStreamIndex == -1)
-		{
-			InternalUpdateVertexDeclaration(OverrideFlags, SourceVertexFactory, PositionBuffer, TangentBuffer, ColorBuffer);
-		}
-	}
-
-	inline FRHIShaderResourceView* GetPreviousPositionsSRV() const
-	{
-		return PrevPositionSRVAlias;
-	}
-
-protected:
-	friend class FLocalVertexFactoryShaderParameters;
-	friend class FSkeletalMeshSceneProxy;
-
-	// Reference holders for RDG buffers
-	TRefCountPtr<FRDGPooledBuffer> PositionRDG;
-	TRefCountPtr<FRDGPooledBuffer> PrevPositionRDG;
-	TRefCountPtr<FRDGPooledBuffer> TangentRDG;
-	TRefCountPtr<FRDGPooledBuffer> ColorRDG;
-	// Vertex buffer required for creating the Vertex Declaration
-	FVertexBuffer PositionVBAlias;
-	FVertexBuffer TangentVBAlias;
-	FVertexBuffer ColorVBAlias;
-	// SRVs required for binding
-	FRHIShaderResourceView* PositionSRVAlias = nullptr;
-	FRHIShaderResourceView* PrevPositionSRVAlias = nullptr;
-	FRHIShaderResourceView* TangentSRVAlias = nullptr;
-	FRHIShaderResourceView* ColorSRVAlias = nullptr;
-	// Cached stream indices
-	int32 PositionStreamIndex = -1;
-	int32 TangentStreamIndex = -1;
-	// Frame number of the bone data that is last updated
-	uint32 UpdatedFrameNumber = 0;
-
-	void InternalUpdateVertexDeclaration(
-		FGPUBaseSkinVertexFactory const* SourceVertexFactory);
-	void InternalUpdateVertexDeclaration(
-		FGPUBaseSkinVertexFactory const* SourceVertexFactory, 
-		struct FRWBuffer* PositionRWBuffer, 
-		struct FRWBuffer* TangentRWBuffer);
-	void InternalUpdateVertexDeclaration(
-		EOverrideFlags OverrideFlags,
-		FGPUBaseSkinVertexFactory const* SourceVertexFactory,
-		TRefCountPtr<FRDGPooledBuffer> const& PositionBuffer,
-		TRefCountPtr<FRDGPooledBuffer> const& TangentBuffer,
-		TRefCountPtr<FRDGPooledBuffer> const& ColorBuffer);
-};
-
-ENUM_CLASS_FLAGS(FGPUSkinPassthroughVertexFactory::EOverrideFlags)
 
 /** Vertex factory with vertex stream components for GPU-skinned and morph target streams */
 class FGPUBaseSkinAPEXClothVertexFactory
@@ -838,3 +712,65 @@ protected:
 	FGPUSkinAPEXClothDataType* ClothDataPtr = nullptr;
 };
 
+
+/**
+ * Vertex factory with vertex stream components for GPU-skinned streams.
+ * This enables Passthrough mode where vertices have been pre-skinned.
+ * Individual vertex attributes can be flagged so that they can be overriden by externally owned buffers.
+ */
+class FGPUSkinPassthroughVertexFactory : public FLocalVertexFactory
+{
+public:
+	FGPUSkinPassthroughVertexFactory(ERHIFeatureLevel::Type InFeatureLevel);
+
+	// Begin FVertexFactory Interface.
+	bool SupportsPositionOnlyStream() const override { return false; }
+	bool SupportsPositionAndNormalOnlyStream() const override { return false; }
+	// End FVertexFactory Interface.
+
+	/** Vertex attributes that we can override. */
+	enum EVertexAtttribute
+	{
+		Position,
+		Tangent,
+		Color,
+		TexCoord,
+		NumAttributes
+	};
+
+	/** 
+	 * Reset all added vertex attributes and SRVs. 
+	 * This doesn't reset the vertex factory itself. Call SetData() to do that.
+	 */
+	void ResetVertexAttributes();
+
+	/** 
+	 * Add vertex attributes that will be in the vertex stream. 
+	 * We accumulate with any attributes that have already been added by previous calls to this function.
+	 * If any attributes are being added for the first time then we recreate the vertex declaration here.
+	 * The InSRVs array size should match the InAttributes array size. 
+	 * SRVs will only be used by platforms that support manual vertex fetch.
+	 * The SRVs are cached per attribute. If any passed in SRV is changed from the cached value then we recreate the vertex factory uniform buffer here.
+	 * Note that on platforms that support manual vertex fetch, only Position will be in the final vertex stream and other attributes will be read through an SRV.
+	 */
+	void AddVertexAttributes(FGPUBaseSkinVertexFactory const* InSourceVertexFactory, TArrayView<EVertexAtttribute> const& InAttributes, TArrayView<FRHIShaderResourceView*> const& InSRVs);
+
+	/** 
+	 * Get the vertex stream index for a vertex attribute.
+	 * This will be -1 for attributes that haven't been set in AddVertexAttributes().
+	 * It may also be -1 for attributes that are read through manual vertex fetch.
+	 */
+	int32 GetAttributeStreamIndex(EVertexAtttribute InAttribute) const;
+
+private:
+	void OverrideAttributeData();
+	void OverrideSRVs(FGPUBaseSkinVertexFactory const* InSourceVertexFactory);
+	void BuildStreamIndices();
+	void CreateUniformBuffer();
+
+	uint32 VertexAttributeMask = 0;
+	TStaticArray<int32, EVertexAtttribute::NumAttributes> StreamIndices;
+	TStaticArray<FRHIShaderResourceView*, EVertexAtttribute::NumAttributes> SRVs;
+	
+	static TStaticArray<FVertexBuffer, EVertexAtttribute::NumAttributes> DummyVBs;
+};
diff --git a/Engine/Source/Runtime/Engine/Public/LocalVertexFactory.h b/Engine/Source/Runtime/Engine/Public/LocalVertexFactory.h
--- a/Engine/Source/Runtime/Engine/Public/LocalVertexFactory.h
+++ b/Engine/Source/Runtime/Engine/Public/LocalVertexFactory.h
@@ -28,6 +28,7 @@ BEGIN_GLOBAL_SHADER_PARAMETER_STRUCT(FLocalVertexFactoryUniformShaderParameters,
 END_GLOBAL_SHADER_PARAMETER_STRUCT()
 
 BEGIN_GLOBAL_SHADER_PARAMETER_STRUCT(FLocalVertexFactoryLooseParameters,)
+	SHADER_PARAMETER_SRV(Buffer<float>, GPUSkinPassThroughPositionBuffer)
 	SHADER_PARAMETER_SRV(Buffer<float>, GPUSkinPassThroughPreviousPositionBuffer)
 END_GLOBAL_SHADER_PARAMETER_STRUCT()
 
@@ -254,17 +255,5 @@ private:
 		class FMeshDrawSingleShaderBindings& ShaderBindings,
 		FVertexInputStreamArray& VertexStreams) const;
 
-	void GetElementShaderBindingsSkinCache(
-		class FGPUSkinPassthroughVertexFactory const* PassthroughVertexFactory,
-		struct FGPUSkinBatchElementUserData* BatchUserData,
-		FMeshDrawSingleShaderBindings& ShaderBindings,
-		FVertexInputStreamArray& VertexStreams) const;
-
-	void GetElementShaderBindingsMeshDeformer(
-		class FGPUSkinPassthroughVertexFactory const* PassthroughVertexFactory,
-		FMeshDrawSingleShaderBindings& ShaderBindings,
-		FVertexInputStreamArray& VertexStreams) const;
-
-	LAYOUT_FIELD(FShaderResourceParameter, GPUSkinCachePositionBuffer);
 	LAYOUT_FIELD(FShaderParameter, IsGPUSkinPassThrough);
 };
\ No newline at end of file
diff --git a/Engine/Source/Runtime/Engine/Public/SkeletalMeshDeformerHelpers.h b/Engine/Source/Runtime/Engine/Public/SkeletalMeshDeformerHelpers.h
--- a/Engine/Source/Runtime/Engine/Public/SkeletalMeshDeformerHelpers.h
+++ b/Engine/Source/Runtime/Engine/Public/SkeletalMeshDeformerHelpers.h
@@ -5,9 +5,9 @@
 #include "CoreMinimal.h"
 #include "RenderGraphResources.h"
 
-class FSkeletalMeshObject;
-class FRDGPooledBuffer;
+class FRDGBuffer;
 class FRHIShaderResourceView;
+class FSkeletalMeshObject;
 
 /** Functions that expose some internal functionality of FSkeletalMeshObject required by MeshDeformer systems. */
 class FSkeletalMeshDeformerHelpers
@@ -18,18 +18,18 @@ public:
 
 	/** Get direct access to bone matrix buffer SRV. */
 	ENGINE_API static FRHIShaderResourceView* GetBoneBufferForReading(
-		FSkeletalMeshObject const* MeshObject,
-		int32 LODIndex,
-		int32 SectionIndex,
-		bool bPreviousFrame);
+		FSkeletalMeshObject const* InMeshObject,
+		int32 InLodIndex,
+		int32 InSectionIndex,
+		bool bInPreviousFrame);
 
 	/** Get direct access to morph target buffer SRV. */
 	ENGINE_API static FRHIShaderResourceView* GetMorphTargetBufferForReading(
-		FSkeletalMeshObject const* MeshObject,
-		int32 LODIndex,
-		int32 SectionIndex,
-		uint32 FrameNumber,
-		bool bPreviousFrame);
+		FSkeletalMeshObject const* InMeshObject,
+		int32 InLodIndex,
+		int32 InSectionIndex,
+		uint32 InFrameNumber,
+		bool bInPreviousFrame);
 
 	/** Buffer SRVs from the cloth system. */
 	struct ENGINE_API FClothBuffers
@@ -42,44 +42,49 @@ public:
 
 	/** Get direct access to cloth buffer SRVs. */
 	ENGINE_API static FClothBuffers GetClothBuffersForReading(
-		FSkeletalMeshObject const* MeshObject,
-		int32 LODIndex,
-		int32 SectionIndex,
-		uint32 FrameNumber,
-		bool bPreviousFrame);
+		FSkeletalMeshObject const* InMeshObject,
+		int32 InLodIndex,
+		int32 InSectionIndex,
+		uint32 InFrameNumber,
+		bool bInPreviousFrame);
 
 #pragma endregion GetInternals
 
 #pragma region SetInternals
 
-	/** Buffer override behavior for SetVertexFactoryBufferOverrides. */
-	enum class EOverrideType
-	{
-		All,		// Clear overrides for each buffer input that is null.
-		Partial,	// Leave existing overrides for each buffer input that is null.
-	};
-
-	/** Apply buffer overrides to the pass through vertex factory. */
-	ENGINE_API static void SetVertexFactoryBufferOverrides(
-		FSkeletalMeshObject* MeshObject,
-		int32 LODIndex, 
-		EOverrideType OverrideType,
-		TRefCountPtr<FRDGPooledBuffer> const& PositionBuffer, 
-		TRefCountPtr<FRDGPooledBuffer> const& TangentBuffer, 
-		TRefCountPtr<FRDGPooledBuffer> const& ColorBuffer);
-
-	/** Reset all buffer overrides that were applied through SetVertexFactoryBufferOverrides. */
-	ENGINE_API static void ResetVertexFactoryBufferOverrides_GameThread(FSkeletalMeshObject* MeshObject, int32 LODIndex);
-
-#if RHI_RAYTRACING
-
-	/** Update ray tracing geometry with new position vertex buffer. */
-	ENGINE_API static void UpdateRayTracingGeometry(
-		FSkeletalMeshObject* MeshObject,
-		int32 LODIndex,
-		TRefCountPtr<FRDGPooledBuffer> const& PositionBuffer);
-
-#endif // RHI_RAYTRACING
+	/** 
+	 * Allocate and bind a new position buffer and return it for writing.
+	 * Ownership is handled by the MeshObject.
+	 * If we call this more than once for the same MeshObject in the same frame then we return the allocation from the first call.
+	 */
+	ENGINE_API static FRDGBuffer* AllocateVertexFactoryPositionBuffer(FRDGBuilder& GraphBuilder, FSkeletalMeshObject* InMeshObject, int32 InLodIndex, TCHAR const* InBufferName);
+
+	/**
+	 * Allocate and bind a new tangent buffer and return it for writing.
+	 * Ownership is handled by the MeshObject.
+	 * If we call this more than once for the same MeshObject in the same frame then we return the allocation from the first call.
+	 */
+	ENGINE_API static FRDGBuffer* AllocateVertexFactoryTangentBuffer(FRDGBuilder& GraphBuilder,	FSkeletalMeshObject* InMeshObject, int32 InLodIndex, TCHAR const* InBufferName);
+
+	/**
+	 * Allocate and bind a new color buffer and return it for writing.
+	 * Ownership is handled by the MeshObject.
+	 * If we call this more than once for the same MeshObject in the same frame then we return the allocation from the first call.
+	 */
+	ENGINE_API static FRDGBuffer* AllocateVertexFactoryColorBuffer(FRDGBuilder& GraphBuilder, FSkeletalMeshObject* InMeshObject, int32 InLodIndex, TCHAR const* InBufferName);
+
+	/** 
+	 * Update all of the MeshObject's passthrough vertex factories with the currently allocated vertex buffers. 
+	 * Usually call this after all AllocateVertexFactory*() functions for a frame.
+	 */
+	ENGINE_API static void UpdateVertexFactoryBufferOverrides(FSkeletalMeshObject* InMeshObject, int32 InLodIndex);
+
+	/** 
+	 * Release all of the the buffers that have been allocated through the AllocateVertexFactory*() functions.
+	 * Reset all of the MeshObject's passthrough vertex factories.  
+	 * Usually called on the game thread when we know that the mesh deformer is not able to update to force bind pose rendering.
+	 */
+	ENGINE_API static void ResetVertexFactoryBufferOverrides_GameThread(FSkeletalMeshObject* InMeshObject, int32 InLodIndex);
 
 #pragma endregion SetInternals
 };
diff --git a/Engine/Source/Runtime/Engine/Public/SkeletalMeshSceneProxy.h b/Engine/Source/Runtime/Engine/Public/SkeletalMeshSceneProxy.h
--- a/Engine/Source/Runtime/Engine/Public/SkeletalMeshSceneProxy.h
+++ b/Engine/Source/Runtime/Engine/Public/SkeletalMeshSceneProxy.h
@@ -211,7 +211,7 @@ protected:
 	uint8 GetCurrentFirstLODIdx_Internal() const;
 private:
 	void CreateBaseMeshBatch(const FSceneView* View, const FSkeletalMeshLODRenderData& LODData, const int32 LODIndex, const int32 SectionIndex, const FSectionElementInfo& SectionElementInfo, FMeshBatch& Mesh, ESkinVertexFactoryMode VFMode = ESkinVertexFactoryMode::Default) const;
-	void UpdateLooseParametersUniformBuffer(const FSceneView* View, const int32 SectionIndex, const FMeshBatch& Mesh, const struct FGPUSkinBatchElementUserData* BatchUserData) const;
+	void UpdateLooseParametersUniformBuffer(const FSceneView* View, const int32 SectionIndex, const struct FSkinBatchVertexFactoryUserData* BatchUserData, class FVertexFactory* VertexFactory) const;
 
 public:
 #if WITH_EDITORONLY_DATA
diff --git a/Engine/Source/Runtime/Engine/Public/SkeletalRenderPublic.h b/Engine/Source/Runtime/Engine/Public/SkeletalRenderPublic.h
--- a/Engine/Source/Runtime/Engine/Public/SkeletalRenderPublic.h
+++ b/Engine/Source/Runtime/Engine/Public/SkeletalRenderPublic.h
@@ -26,6 +26,7 @@ class UMorphTarget;
 struct FSkelMeshRenderSection;
 struct FCachedGeometry;
 class FGPUSkinCacheEntry;
+class FMeshDeformerGeometry;
 class FRayTracingGeometry;
 
 /** data for a single skinned skeletal mesh vertex */
@@ -53,6 +54,13 @@ enum class EPreviousBoneTransformUpdateMode
 	DuplicateCurrentToPrevious,
 };
 
+struct FSkinBatchVertexFactoryUserData
+{
+	FGPUSkinCacheEntry* SkinCacheEntry = nullptr;
+	FMeshDeformerGeometry* DeformerGeometry = nullptr;
+	int32 Section = -1;
+};
+
 /**
 * Interface for mesh rendering data
 */
@@ -94,7 +102,14 @@ public:
 	 * @param	ChunkIdx - not used
 	 * @return	vertex factory for rendering the LOD, 0 to suppress rendering
 	 */
-	virtual const FVertexFactory* GetSkinVertexFactory(const FSceneView* View, int32 LODIndex,int32 ChunkIdx, ESkinVertexFactoryMode VFMode = ESkinVertexFactoryMode::Default) const = 0;
+	virtual const FVertexFactory* GetSkinVertexFactory(const FSceneView* View, int32 LODIndex, int32 ChunkIdx, ESkinVertexFactoryMode VFMode = ESkinVertexFactoryMode::Default) const = 0;
+
+	/**
+	 * @param	LODIndex - Index to LODs
+	 * @param	ChunkIdx - Index to render sections.
+	 * @return	VertexFactoryUserData for storing on a FMeshBatch.
+	 */
+	virtual const FSkinBatchVertexFactoryUserData* GetVertexFactoryUserData(const int32 LODIndex, int32 ChunkIdx, ESkinVertexFactoryMode VFMode) const { return nullptr; }
 
 	/**
 	 * Re-skin cached vertices for an LOD and update the vertex buffer. Note that this
