World Partition: FStreamingGenerationLogErrorHandler now uses UE_ASSET_LOG to report errors during cook.

diff --git a/Engine/Source/Runtime/Engine/Private/WorldPartition/ErrorHandling/WorldPartitionStreamingGenerationErrorHandler.cpp b/Engine/Source/Runtime/Engine/Private/WorldPartition/ErrorHandling/WorldPartitionStreamingGenerationErrorHandler.cpp
--- a/Engine/Source/Runtime/Engine/Private/WorldPartition/ErrorHandling/WorldPartitionStreamingGenerationErrorHandler.cpp
+++ b/Engine/Source/Runtime/Engine/Private/WorldPartition/ErrorHandling/WorldPartitionStreamingGenerationErrorHandler.cpp
@@ -4,12 +4,16 @@
 #include "WorldPartition/ErrorHandling/WorldPartitionStreamingGenerationErrorHandler.h"
 #include "UObject/SoftObjectPath.h"
 
-// Helpers
-FString IStreamingGenerationErrorHandler::GetFullActorName(const FWorldPartitionActorDescView& ActorDescView)
+FString IStreamingGenerationErrorHandler::GetActorName(const FWorldPartitionActorDescView& ActorDescView)
 {
 	const FString ActorPath = ActorDescView.GetActorSoftPath().GetLongPackageName();
 	const FString ActorLabel = ActorDescView.GetActorLabelOrName().ToString();
+	return ActorPath + TEXT(".") + ActorLabel;
+}
+
+FString IStreamingGenerationErrorHandler::GetFullActorName(const FWorldPartitionActorDescView& ActorDescView)
+{
 	const FString ActorPackage = ActorDescView.GetActorPackage().ToString();
-	return ActorPath + TEXT(".") + ActorLabel + TEXT(" (") + ActorPackage + TEXT(")");
+	return GetActorName(ActorDescView) + TEXT(" (") + ActorPackage + TEXT(")");
 }
 #endif
\ No newline at end of file
diff --git a/Engine/Source/Runtime/Engine/Private/WorldPartition/ErrorHandling/WorldPartitionStreamingGenerationLogErrorHandler.cpp b/Engine/Source/Runtime/Engine/Private/WorldPartition/ErrorHandling/WorldPartitionStreamingGenerationLogErrorHandler.cpp
--- a/Engine/Source/Runtime/Engine/Private/WorldPartition/ErrorHandling/WorldPartitionStreamingGenerationLogErrorHandler.cpp
+++ b/Engine/Source/Runtime/Engine/Private/WorldPartition/ErrorHandling/WorldPartitionStreamingGenerationLogErrorHandler.cpp
@@ -6,81 +6,84 @@
 #include "WorldPartition/DataLayer/DataLayerInstanceWithAsset.h"
 #include "GameFramework/Actor.h"
 
+#define UE_ASSET_LOG_ACTORDESCVIEW(CategoryName, Verbosity, ActorDescView, Format, ...) \
+	UE_ASSET_LOG(LogWorldPartition, Log, *ActorDescView.GetActorPackage().ToString(), Format, ##__VA_ARGS__)
+
 void FStreamingGenerationLogErrorHandler::OnInvalidRuntimeGrid(const FWorldPartitionActorDescView& ActorDescView, FName GridName)
 {
-	UE_LOG(LogWorldPartition, Log, TEXT("Actor %s has an invalid runtime grid %s"), *GetFullActorName(ActorDescView), *GridName.ToString());
+	UE_ASSET_LOG_ACTORDESCVIEW(LogWorldPartition, Log, ActorDescView, TEXT("Actor %s has an invalid runtime grid %s"), *GetActorName(ActorDescView), *GridName.ToString());
 }
 
 void FStreamingGenerationLogErrorHandler::OnInvalidReference(const FWorldPartitionActorDescView& ActorDescView, const FGuid& ReferenceGuid, FWorldPartitionActorDescView* ReferenceActorDescView)
 {
-	UE_LOG(LogWorldPartition, Log, TEXT("Actor %s has an invalid reference to %s"), *GetFullActorName(ActorDescView), ReferenceActorDescView ? *GetFullActorName(*ReferenceActorDescView) : *ReferenceGuid.ToString());
+	UE_ASSET_LOG_ACTORDESCVIEW(LogWorldPartition, Log, ActorDescView, TEXT("Actor %s has an invalid reference to %s"), *GetActorName(ActorDescView), ReferenceActorDescView ? *GetActorName(*ReferenceActorDescView) : *ReferenceGuid.ToString());
 }
 
 void FStreamingGenerationLogErrorHandler::OnInvalidReferenceGridPlacement(const FWorldPartitionActorDescView& ActorDescView, const FWorldPartitionActorDescView& ReferenceActorDescView)
 {
-	const FString SpatiallyLoadedActor(TEXT("Spatially loaded actor"));
-	const FString NonSpatiallyLoadedActor(TEXT("Non-spatially loaded loaded actor"));
+	static const FString SpatiallyLoadedActor(TEXT("Spatially loaded actor"));
+	static const FString NonSpatiallyLoadedActor(TEXT("Non-spatially loaded loaded actor"));
 
-	UE_LOG(LogWorldPartition, Log, TEXT("%s %s reference %s %s"), ActorDescView.GetIsSpatiallyLoaded() ? *SpatiallyLoadedActor : *NonSpatiallyLoadedActor, *GetFullActorName(ActorDescView), ReferenceActorDescView.GetIsSpatiallyLoaded() ? *SpatiallyLoadedActor : *NonSpatiallyLoadedActor, *GetFullActorName(ReferenceActorDescView));
+	UE_ASSET_LOG_ACTORDESCVIEW(LogWorldPartition, Log, ActorDescView, TEXT("%s %s reference %s %s"), ActorDescView.GetIsSpatiallyLoaded() ? *SpatiallyLoadedActor : *NonSpatiallyLoadedActor, *GetActorName(ActorDescView), ReferenceActorDescView.GetIsSpatiallyLoaded() ? *SpatiallyLoadedActor : *NonSpatiallyLoadedActor, *GetActorName(ReferenceActorDescView));
 }
 
 void FStreamingGenerationLogErrorHandler::OnInvalidReferenceDataLayers(const FWorldPartitionActorDescView& ActorDescView, const FWorldPartitionActorDescView& ReferenceActorDescView)
 {
-	UE_LOG(LogWorldPartition, Log, TEXT("Actor %s references an actor in a different set of runtime data layers %s"), *GetFullActorName(ActorDescView), *GetFullActorName(ReferenceActorDescView));
+	UE_ASSET_LOG_ACTORDESCVIEW(LogWorldPartition, Log, ActorDescView, TEXT("Actor %s references an actor in a different set of runtime data layers %s"), *GetActorName(ActorDescView), *GetActorName(ReferenceActorDescView));
 }
 
 void FStreamingGenerationLogErrorHandler::OnInvalidReferenceRuntimeGrid(const FWorldPartitionActorDescView& ActorDescView, const FWorldPartitionActorDescView& ReferenceActorDescView)
 {
-	UE_LOG(LogWorldPartition, Log, TEXT("Actor %s references an actor in a different runtime grid %s"), *GetFullActorName(ActorDescView), *GetFullActorName(ReferenceActorDescView));
+	UE_ASSET_LOG_ACTORDESCVIEW(LogWorldPartition, Log, ActorDescView, TEXT("Actor %s references an actor in a different runtime grid %s"), *GetActorName(ActorDescView), *GetActorName(ReferenceActorDescView));
 }
 
 void FStreamingGenerationLogErrorHandler::OnInvalidReferenceLevelScriptStreamed(const FWorldPartitionActorDescView& ActorDescView)
 {
-	UE_LOG(LogWorldPartition, Log, TEXT("Level Script Blueprint references streamed actor %s"), *GetFullActorName(ActorDescView));
+	UE_ASSET_LOG_ACTORDESCVIEW(LogWorldPartition, Log, ActorDescView, TEXT("Level Script Blueprint references streamed actor %s"), *GetActorName(ActorDescView));
 }
 
 void FStreamingGenerationLogErrorHandler::OnInvalidReferenceLevelScriptDataLayers(const FWorldPartitionActorDescView& ActorDescView)
 {
-	UE_LOG(LogWorldPartition, Log, TEXT("Level Script Blueprint references streamed actor %s with a non empty set of data layers"), *GetFullActorName(ActorDescView));
+	UE_ASSET_LOG_ACTORDESCVIEW(LogWorldPartition, Log, ActorDescView, TEXT("Level Script Blueprint references streamed actor %s with a non empty set of data layers"), *GetActorName(ActorDescView));
 }
 
 void FStreamingGenerationLogErrorHandler::OnInvalidReferenceDataLayerAsset(const UDataLayerInstanceWithAsset* DataLayerInstance)
 {
-	UE_LOG(LogWorldPartition, Log, TEXT("Data Layer %s does not have a Data Layer asset"), *DataLayerInstance->GetDataLayerFName().ToString());
+	UE_ASSET_LOG(LogWorldPartition, Log, DataLayerInstance, TEXT("Data Layer does not have a Data Layer asset"));
 }
 
 void FStreamingGenerationLogErrorHandler::OnDataLayerHierarchyTypeMismatch(const UDataLayerInstance* DataLayerInstance, const UDataLayerInstance* Parent)
 {
-	UE_LOG(LogWorldPartition, Log, TEXT("Data Layer %s is of Type %s and its parent %s is of type %s"), *DataLayerInstance->GetDataLayerFullName(), *UEnum::GetValueAsString(DataLayerInstance->GetType()), *Parent->GetDataLayerFullName(), *UEnum::GetValueAsString(Parent->GetType()));
+	UE_ASSET_LOG(LogWorldPartition, Log, DataLayerInstance, TEXT("Data Layer %s is of Type %s and its parent %s is of type %s"), *DataLayerInstance->GetDataLayerFullName(), *UEnum::GetValueAsString(DataLayerInstance->GetType()), *Parent->GetDataLayerFullName(), *UEnum::GetValueAsString(Parent->GetType()));
 }
 
 void FStreamingGenerationLogErrorHandler::OnDataLayerAssetConflict(const UDataLayerInstanceWithAsset* DataLayerInstance, const UDataLayerInstanceWithAsset* ConflictingDataLayerInstance)
 {
-	UE_LOG(LogWorldPartition, Log, TEXT("Data Layer Instance %s and Data Layer Instance %s are both referencing Data Layer Asset %s"), *DataLayerInstance->GetDataLayerFName().ToString(), *ConflictingDataLayerInstance->GetDataLayerFName().ToString(), *DataLayerInstance->GetAsset()->GetFullName());
+	UE_ASSET_LOG(LogWorldPartition, Log, DataLayerInstance, TEXT("Data Layer Instance %s and Data Layer Instance %s are both referencing Data Layer Asset %s"), *DataLayerInstance->GetDataLayerFName().ToString(), *ConflictingDataLayerInstance->GetDataLayerFName().ToString(), *DataLayerInstance->GetAsset()->GetFullName());
 }
 
 void FStreamingGenerationLogErrorHandler::OnActorNeedsResave(const FWorldPartitionActorDescView& ActorDescView)
 {
-	UE_LOG(LogWorldPartition, Log, TEXT("Actor %s needs to be resaved"), *GetFullActorName(ActorDescView));
+	UE_ASSET_LOG_ACTORDESCVIEW(LogWorldPartition, Log, ActorDescView, TEXT("Actor %s needs to be resaved"), *GetActorName(ActorDescView));
 }
 
 void FStreamingGenerationLogErrorHandler::OnLevelInstanceInvalidWorldAsset(const FWorldPartitionActorDescView& ActorDescView, FName WorldAsset, ELevelInstanceInvalidReason Reason)
 {
-	const FString ActorName = GetFullActorName(ActorDescView);
+	const FString ActorName = GetActorName(ActorDescView);
 
 	switch (Reason)
 	{
 	case ELevelInstanceInvalidReason::WorldAssetNotFound:
-		UE_LOG(LogWorldPartition, Log, TEXT("Actor %s has an invalid world asset '%s'"), *ActorName, *WorldAsset.ToString());
+		UE_ASSET_LOG_ACTORDESCVIEW(LogWorldPartition, Log, ActorDescView, TEXT("Actor %s has an invalid world asset %s"), *ActorName, *WorldAsset.ToString());
 		break;
 	case ELevelInstanceInvalidReason::WorldAssetNotUsingExternalActors:
-		UE_LOG(LogWorldPartition, Log, TEXT("Actor %s world asset '%s' is not using external actors"), *ActorName, *WorldAsset.ToString());
+		UE_ASSET_LOG_ACTORDESCVIEW(LogWorldPartition, Log, ActorDescView, TEXT("Actor %s world asset %s is not using external actors"), *ActorName, *WorldAsset.ToString());
 		break;
 	case ELevelInstanceInvalidReason::WorldAssetImcompatiblePartitioned:
-		UE_LOG(LogWorldPartition, Log, TEXT("Actor %s world asset '%s' is partitioned but not marked as compatible"), *ActorName, *WorldAsset.ToString());
+		UE_ASSET_LOG_ACTORDESCVIEW(LogWorldPartition, Log, ActorDescView, TEXT("Actor %s world asset %s is partitioned but not marked as compatible"), *ActorName, *WorldAsset.ToString());
 		break;
 	case ELevelInstanceInvalidReason::WorldAssetHasInvalidContainer:
-		UE_LOG(LogWorldPartition, Log, TEXT("Actor %s world asset '%s' has an invalid container"), *ActorName, *WorldAsset.ToString());
+		UE_ASSET_LOG_ACTORDESCVIEW(LogWorldPartition, Log, ActorDescView, TEXT("Actor %s world asset %s has an invalid container"), *ActorName, *WorldAsset.ToString());
 		break;
 	};
 }
diff --git a/Engine/Source/Runtime/Engine/Public/WorldPartition/ErrorHandling/WorldPartitionStreamingGenerationErrorHandler.h b/Engine/Source/Runtime/Engine/Public/WorldPartition/ErrorHandling/WorldPartitionStreamingGenerationErrorHandler.h
--- a/Engine/Source/Runtime/Engine/Public/WorldPartition/ErrorHandling/WorldPartitionStreamingGenerationErrorHandler.h
+++ b/Engine/Source/Runtime/Engine/Public/WorldPartition/ErrorHandling/WorldPartitionStreamingGenerationErrorHandler.h
@@ -85,6 +85,7 @@ public:
 	virtual void OnLevelInstanceInvalidWorldAsset(const FWorldPartitionActorDescView& ActorDescView, FName WorldAsset, ELevelInstanceInvalidReason Reason) = 0;
 
 	// Helpers
+	static FString GetActorName(const FWorldPartitionActorDescView& ActorDescView);
 	static FString GetFullActorName(const FWorldPartitionActorDescView& ActorDescView);
 };
 #endif
