World Partition: changed some logs to verbose to avoid log spamming when spectating.

diff --git a/Engine/Source/Runtime/Engine/Private/WorldPartition/WorldPartitionStreamingPolicy.cpp b/Engine/Source/Runtime/Engine/Private/WorldPartition/WorldPartitionStreamingPolicy.cpp
--- a/Engine/Source/Runtime/Engine/Private/WorldPartition/WorldPartitionStreamingPolicy.cpp
+++ b/Engine/Source/Runtime/Engine/Private/WorldPartition/WorldPartitionStreamingPolicy.cpp
@@ -1110,7 +1110,7 @@ float FStreamingSourceVelocity::GetAverageVelocity(const FVector& NewPosition, c
 	const double Distance = bIsFirstCall ? 0.f : ((NewPosition - LastPosition) * 0.01).Size();
 	if (bIsFirstCall)
 	{
-		UE_LOG(LogWorldPartition, Log, TEXT("New Streaming Source: %s -> Position: %s"), *SourceName.ToString(), *NewPosition.ToString());
+		UE_LOG(LogWorldPartition, Verbose, TEXT("New Streaming Source: %s -> Position: %s"), *SourceName.ToString(), *NewPosition.ToString());
 		LastIndex = 0;
 	}
 
@@ -1123,7 +1123,7 @@ float FStreamingSourceVelocity::GetAverageVelocity(const FVector& NewPosition, c
 	// Handle invalid cases
 	if (bIsFirstCall || (DeltaSeconds <= 0.f) || (DeltaSeconds > MaxDeltaSeconds) || (Distance > TeleportDistance))
 	{
-		UE_CLOG(Distance > TeleportDistance, LogWorldPartition, Log, TEXT("Detected Streaming Source Teleport: %s -> Last Position: %s -> New Position: %s"), *SourceName.ToString(), *LastPosition.ToString(), *NewPosition.ToString());
+		UE_CLOG(Distance > TeleportDistance, LogWorldPartition, Verbose, TEXT("Detected Streaming Source Teleport: %s -> Last Position: %s -> New Position: %s"), *SourceName.ToString(), *LastPosition.ToString(), *NewPosition.ToString());
 		return 0.f;
 	}
 
