[Engine]
* Added // IWYU pragma: begin_keep/end_keep to monolithic boiler plate headers

diff --git a/Engine/Source/Runtime/Engine/Public/Engine.h b/Engine/Source/Runtime/Engine/Public/Engine.h
--- a/Engine/Source/Runtime/Engine/Public/Engine.h
+++ b/Engine/Source/Runtime/Engine/Public/Engine.h
@@ -6,6 +6,8 @@
 
 #pragma once
 
+// IWYU pragma: begin_keep
+
 #include "Misc/MonolithicHeaderBoilerplate.h"
 MONOLITHIC_HEADER_BOILERPLATE()
 
@@ -329,3 +331,5 @@ MONOLITHIC_HEADER_BOILERPLATE()
 #include "TimerManager.h"
 #include "SlateCore.h"
 #include "SlateBasics.h"
+
+// IWYU pragma: end_keep
diff --git a/Engine/Source/Runtime/Engine/Public/EngineMinimal.h b/Engine/Source/Runtime/Engine/Public/EngineMinimal.h
--- a/Engine/Source/Runtime/Engine/Public/EngineMinimal.h
+++ b/Engine/Source/Runtime/Engine/Public/EngineMinimal.h
@@ -6,6 +6,8 @@
 
 #pragma once
 
+// IWYU pragma: begin_keep
+
 // Boilerplate
 #include "Misc/MonolithicHeaderBoilerplate.h"
 MONOLITHIC_HEADER_BOILERPLATE()
@@ -112,3 +114,5 @@ MONOLITHIC_HEADER_BOILERPLATE()
 #include "Materials/MaterialInstanceDynamic.h"
 #include "Sound/SoundBase.h"
 #include "Sound/SoundWave.h"
+
+// IWYU pragma: end_keep
diff --git a/Engine/Source/Runtime/Engine/Public/ParticleDefinitions.h b/Engine/Source/Runtime/Engine/Public/ParticleDefinitions.h
--- a/Engine/Source/Runtime/Engine/Public/ParticleDefinitions.h
+++ b/Engine/Source/Runtime/Engine/Public/ParticleDefinitions.h
@@ -6,6 +6,8 @@
 
 #pragma once
 
+// IWYU pragma: begin_keep
+
 #include "Misc/MonolithicHeaderBoilerplate.h"
 MONOLITHIC_HEADER_BOILERPLATE()
 
@@ -42,3 +44,5 @@ MONOLITHIC_HEADER_BOILERPLATE()
 #include "ParticleHelper.h"
 #include "Distributions.h"
 #include "ParticleEmitterInstances.h"
+
+// IWYU pragma: end_keep
diff --git a/Engine/Source/Runtime/Engine/Public/SoundDefinitions.h b/Engine/Source/Runtime/Engine/Public/SoundDefinitions.h
--- a/Engine/Source/Runtime/Engine/Public/SoundDefinitions.h
+++ b/Engine/Source/Runtime/Engine/Public/SoundDefinitions.h
@@ -2,6 +2,8 @@
 
 #pragma once
 
+// IWYU pragma: begin_keep
+
 #include "Misc/MonolithicHeaderBoilerplate.h"
 MONOLITHIC_HEADER_BOILERPLATE()
 
@@ -25,3 +27,5 @@ MONOLITHIC_HEADER_BOILERPLATE()
 #include "AudioThread.h"
 #include "AudioDeviceManager.h"
 #include "AudioDevice.h"
+
+// IWYU pragma: end_keep
