[Engine]
* Fixed compile errors caused by previous merge

diff --git a/Engine/Source/Runtime/Engine/Private/DataBunch.cpp b/Engine/Source/Runtime/Engine/Private/DataBunch.cpp
--- a/Engine/Source/Runtime/Engine/Private/DataBunch.cpp
+++ b/Engine/Source/Runtime/Engine/Private/DataBunch.cpp
@@ -5,6 +5,7 @@
 =============================================================================*/
 
 #include "Net/DataBunch.h"
+#include "Engine/ControlChannel.h"
 #include "Engine/NetConnection.h"
 #include "Net/Core/Trace/NetTrace.h"
 #include "Net/Core/Trace/Private/NetTraceInternal.h"
diff --git a/Engine/Source/Runtime/Engine/Private/Net/Experimental/Iris/DataStreamChannel.cpp b/Engine/Source/Runtime/Engine/Private/Net/Experimental/Iris/DataStreamChannel.cpp
--- a/Engine/Source/Runtime/Engine/Private/Net/Experimental/Iris/DataStreamChannel.cpp
+++ b/Engine/Source/Runtime/Engine/Private/Net/Experimental/Iris/DataStreamChannel.cpp
@@ -17,14 +17,11 @@
 #include "Iris/Serialization/NetSerializationContext.h"
 #include "Iris/Core/IrisProfiler.h"
 #include "Iris/Core/IrisMemoryTracker.h"
+#include "Net/Core/Trace/NetTrace.h"
 #include "PacketHandler.h"
 #include "ProfilingDebugging/CsvProfiler.h"
 #endif // UE_WITH_IRIS
 
-#if UE_WITH_IRIS && UE_NET_TRACE_ENABLED
-#include "Net/Core/Trace/NetTrace.h"
-#endif
-
 namespace UE::Net::Private
 {
 	static bool bIrisSaturateBandwidth = false;
