Moved passthrough vertex factory previous frame logic from CPU to GPU.
This allows us to create the loose uniform buffer when we update the vertex factory instead of during the mesh batch creation.

diff --git a/Engine/Shaders/Private/LocalVertexFactory.ush b/Engine/Shaders/Private/LocalVertexFactory.ush
--- a/Engine/Shaders/Private/LocalVertexFactory.ush
+++ b/Engine/Shaders/Private/LocalVertexFactory.ush
@@ -1152,34 +1152,35 @@ float3 VertexFactoryGetWorldNormal(FVertexFactoryInput Input, FVertexFactoryInte
 float4 VertexFactoryGetPreviousWorldPosition(FVertexFactoryInput Input, FVertexFactoryIntermediates Intermediates)
 {
 	FLWCMatrix PreviousLocalToWorld = GetInstanceData(Intermediates).PrevLocalToWorld;
+	float4 PrevLocalPosition = Input.Position;
 
-	float4 PrevLocalPosition;
 #if SUPPORT_GPUSKIN_PASSTHROUGH
 	if (IsGPUSkinPassThrough())
 	{
-		uint Offset = Input.VertexId * 3;
-		float3 PreviousPos;
-		PreviousPos.x = LocalVFLooseParameters.GPUSkinPassThroughPreviousPositionBuffer[Offset + 0];
-		PreviousPos.y = LocalVFLooseParameters.GPUSkinPassThroughPreviousPositionBuffer[Offset + 1];
-		PreviousPos.z = LocalVFLooseParameters.GPUSkinPassThroughPreviousPositionBuffer[Offset + 2];
-		PrevLocalPosition = float4(PreviousPos, 1);	
+		// Duplicates logic in FGPUSkinVertexFactoryShaderParameters::GetElementShaderBindings().
+		if (ResolvedView.FrameNumber == LocalVFLooseParameters.FrameNumber && ResolvedView.WorldIsPaused == 0)
+		{
+			uint Offset = Input.VertexId * 3;
+			float3 PreviousPos;
+			PreviousPos.x = LocalVFLooseParameters.GPUSkinPassThroughPreviousPositionBuffer[Offset + 0];
+			PreviousPos.y = LocalVFLooseParameters.GPUSkinPassThroughPreviousPositionBuffer[Offset + 1];
+			PreviousPos.z = LocalVFLooseParameters.GPUSkinPassThroughPreviousPositionBuffer[Offset + 2];
+			PrevLocalPosition = float4(PreviousPos, 1);	
+		}
 	}
 	else
 #endif // SUPPORT_GPUSKIN_PASSTHROUGH
 	{
-	#if USE_INSTANCING
+#if USE_INSTANCING
 		float4x4 InstanceTransform = GetInstancePrevTransform(Intermediates);
 		PrevLocalPosition = mul(Input.Position, InstanceTransform);
 #elif USE_SPLINEDEFORM
-	// Just like CalcWorldPosition...
-	float4x3 SliceTransform = CalcSliceTransform(dot(Input.Position.xyz, SplineMeshDir));
-
-	// Transform into mesh space
+		// Just like CalcWorldPosition...
+		float4x3 SliceTransform = CalcSliceTransform(dot(Input.Position.xyz, SplineMeshDir));
+		// Transform into mesh space
 		PrevLocalPosition = float4(mul(Input.Position, SliceTransform), Input.Position.w);
-#else
-		PrevLocalPosition = Input.Position;
-#endif	// USE_INSTANCING
-}
+#endif
+	}
 
 	return TransformPreviousLocalPositionToTranslatedWorld(PrevLocalPosition.xyz, PreviousLocalToWorld);
 }
diff --git a/Engine/Source/Runtime/Engine/Private/GPUSkinCache.cpp b/Engine/Source/Runtime/Engine/Private/GPUSkinCache.cpp
--- a/Engine/Source/Runtime/Engine/Private/GPUSkinCache.cpp
+++ b/Engine/Source/Runtime/Engine/Private/GPUSkinCache.cpp
@@ -305,7 +305,7 @@ public:
 		for (int32 Index = 0; Index < Sections.Num(); ++Index)
 		{
 			BatchElementsUserData[Index].SkinCacheEntry = this;
-			BatchElementsUserData[Index].Section = Index;
+			BatchElementsUserData[Index].SectionIndex = Index;
 		}
 
 		UpdateSkinWeightBuffer();
@@ -414,24 +414,19 @@ public:
 			return IntermediateAccumulatedTangentBuffer;
 		}
 
-		void UpdateFrameNumber()
-		{
-			UpdatedFrameNumber = SourceVertexFactory->GetShaderData().UpdatedFrameNumber;
-		}
-
 		void UpdateVertexFactoryDeclaration()
 		{
-			FGPUSkinPassthroughVertexFactory::EVertexAtttribute Attribs[2] = { FGPUSkinPassthroughVertexFactory::EVertexAtttribute::Position, FGPUSkinPassthroughVertexFactory::EVertexAtttribute::Tangent };
-			FRHIShaderResourceView* SRVs[2] = { nullptr, GetTangentRWBuffer()->Buffer.SRV };
-			TargetVertexFactory->AddVertexAttributes(SourceVertexFactory, Attribs, SRVs);
+			FGPUSkinPassthroughVertexFactory::FAddVertexAttributeDesc Desc;
+			Desc.FrameNumber = SourceVertexFactory->GetShaderData().UpdatedFrameNumber;
+			Desc.VertexAttributes.Add(FGPUSkinPassthroughVertexFactory::VertexPosition);
+			Desc.VertexAttributes.Add(FGPUSkinPassthroughVertexFactory::VertexTangent);
+			Desc.SRVs[FGPUSkinPassthroughVertexFactory::Position] = GetPositionRWBuffer()->Buffer.SRV;
+			Desc.SRVs[FGPUSkinPassthroughVertexFactory::PreviousPosition] = GetPreviousPositionRWBuffer()->Buffer.SRV;
+			Desc.SRVs[FGPUSkinPassthroughVertexFactory::Tangent] = GetTangentRWBuffer()->Buffer.SRV;
+			TargetVertexFactory->SetVertexAttributes(SourceVertexFactory, Desc);
 		}
 	};
 
-	void UpdateFrameNumber(int32 Section)
-	{
-		DispatchData[Section].UpdateFrameNumber();
-	}
-
 	void UpdateVertexFactoryDeclaration(int32 Section)
 	{
 		DispatchData[Section].UpdateVertexFactoryDeclaration();
@@ -1410,7 +1405,6 @@ void FGPUSkinCache::DoDispatch(FRHICommandListImmediate& RHICmdList)
 	for (int32 i = 0; i < BatchCount; ++i)
 	{
 		FDispatchEntry& DispatchItem = BatchDispatches[i];
-		DispatchItem.SkinCacheEntry->UpdateFrameNumber(DispatchItem.Section);
 		DispatchItem.SkinCacheEntry->UpdateVertexFactoryDeclaration(DispatchItem.Section);
 	}
 
@@ -1465,7 +1459,6 @@ void FGPUSkinCache::DoDispatch(FRHICommandListImmediate& RHICmdList, FGPUSkinCac
 		DispatchUpdateSkinTangents(RHICmdList, SkinCacheEntry, Section, StagingBuffer, false);
 	}
 
-	SkinCacheEntry->UpdateFrameNumber(Section);
 	SkinCacheEntry->UpdateVertexFactoryDeclaration(Section);
 
 	TransitionAllToReadable(RHICmdList, BuffersToTransitionToRead);
@@ -1764,11 +1757,10 @@ void FGPUSkinCache::Release(FGPUSkinCacheEntry*& SkinCacheEntry)
 	}
 }
 
-void FGPUSkinCache::GetShaderBindings(
+void FGPUSkinCache::GetShaderVertexStreams(
 	const FGPUSkinCacheEntry* Entry, 
 	int32 Section,
 	const FGPUSkinPassthroughVertexFactory* VertexFactory,
-	FMeshDrawSingleShaderBindings& ShaderBindings,
 	FVertexInputStreamArray& VertexStreams)
 {
 	INC_DWORD_STAT(STAT_GPUSkinCache_NumSetVertexStreams);
@@ -1778,12 +1770,11 @@ void FGPUSkinCache::GetShaderBindings(
 
 	FGPUSkinCacheEntry::FSectionDispatchData const& DispatchData = Entry->DispatchData[Section];
 
-	//UE_LOG(LogSkinCache, Warning, TEXT("*** SetVertexStreams E %p Sec %d(%p) LOD %d"), Entry, Section, Entry->DispatchData[Section].Section, Entry->LOD);
-
-	int32 PositionStreamIndex = VertexFactory->GetAttributeStreamIndex(FGPUSkinPassthroughVertexFactory::EVertexAtttribute::Position);
+	const int32 PositionStreamIndex = VertexFactory->GetAttributeStreamIndex(FGPUSkinPassthroughVertexFactory::EVertexAtttribute::VertexPosition);
+	check(PositionStreamIndex > -1);
 	VertexStreams.Add(FVertexInputStream(PositionStreamIndex, 0, DispatchData.GetPositionRWBuffer()->Buffer.Buffer));
 
-	int32 TangentStreamIndex = VertexFactory->GetAttributeStreamIndex(FGPUSkinPassthroughVertexFactory::EVertexAtttribute::Tangent);
+	const int32 TangentStreamIndex = VertexFactory->GetAttributeStreamIndex(FGPUSkinPassthroughVertexFactory::EVertexAtttribute::VertexTangent);
 	if (TangentStreamIndex > -1 && DispatchData.GetTangentRWBuffer())
 	{
 		VertexStreams.Add(FVertexInputStream(TangentStreamIndex, 0, DispatchData.GetTangentRWBuffer()->Buffer.Buffer));
diff --git a/Engine/Source/Runtime/Engine/Private/GPUSkinVertexFactory.cpp b/Engine/Source/Runtime/Engine/Private/GPUSkinVertexFactory.cpp
--- a/Engine/Source/Runtime/Engine/Private/GPUSkinVertexFactory.cpp
+++ b/Engine/Source/Runtime/Engine/Private/GPUSkinVertexFactory.cpp
@@ -1276,6 +1276,7 @@ void FGPUSkinPassthroughVertexFactory::ResetVertexAttributes()
 		StreamIndices[Index] = -1;
 		SRVs[Index] = nullptr;
 	}
+	UpdatedFrameNumber = ~0U;
 }
 
 // We don't set actual vertex buffers in our FDataType because we always expect to override the values in GetElementShaderBindings().
@@ -1284,65 +1285,92 @@ TStaticArray<FVertexBuffer, FGPUSkinPassthroughVertexFactory::EVertexAtttribute:
 
 void FGPUSkinPassthroughVertexFactory::OverrideAttributeData()
 {
-	if (VertexAttributeMask & (1 << EVertexAtttribute::Position))
+	if (VertexAttributeMask & (1 << EVertexAtttribute::VertexPosition))
 	{
-		Data.PositionComponent.VertexBuffer = &DummyVBs[EVertexAtttribute::Position];
+		Data.PositionComponent.VertexBuffer = &DummyVBs[EVertexAtttribute::VertexPosition];
 		Data.PositionComponent.Offset = 0;
 		Data.PositionComponent.VertexStreamUsage = EVertexStreamUsage::Overridden;
 		Data.PositionComponent.Stride = 3 * sizeof(float);
 	}
 
-	if (VertexAttributeMask & (1 << EVertexAtttribute::Tangent))
+	if (VertexAttributeMask & (1 << EVertexAtttribute::VertexTangent))
 	{
-		Data.TangentBasisComponents[0].VertexBuffer = &DummyVBs[EVertexAtttribute::Tangent];
+		Data.TangentBasisComponents[0].VertexBuffer = &DummyVBs[EVertexAtttribute::VertexTangent];
 		Data.TangentBasisComponents[0].Offset = 0;
 		Data.TangentBasisComponents[0].Type = VET_Short4N;
 		Data.TangentBasisComponents[0].Stride = 16;
 		Data.TangentBasisComponents[0].VertexStreamUsage = EVertexStreamUsage::Overridden | EVertexStreamUsage::ManualFetch;
 
-		Data.TangentBasisComponents[1].VertexBuffer = &DummyVBs[EVertexAtttribute::Tangent];
+		Data.TangentBasisComponents[1].VertexBuffer = &DummyVBs[EVertexAtttribute::VertexTangent];
 		Data.TangentBasisComponents[1].Offset = 8;
 		Data.TangentBasisComponents[1].Type = VET_Short4N;
 		Data.TangentBasisComponents[1].Stride = 16;
 		Data.TangentBasisComponents[1].VertexStreamUsage = EVertexStreamUsage::Overridden | EVertexStreamUsage::ManualFetch;
 	}
 
-	if (VertexAttributeMask & (1 << EVertexAtttribute::Color))
+	if (VertexAttributeMask & (1 << EVertexAtttribute::VertexColor))
 	{
-		Data.ColorComponent.VertexBuffer = &DummyVBs[EVertexAtttribute::Color];
+		Data.ColorComponent.VertexBuffer = &DummyVBs[EVertexAtttribute::VertexColor];
 		Data.ColorComponent.Offset = 0;
 		Data.ColorComponent.VertexStreamUsage = EVertexStreamUsage::Overridden | EVertexStreamUsage::ManualFetch;
 		Data.ColorComponent.Stride = sizeof(uint32);
 	}
 
-	if (VertexAttributeMask & (1 << EVertexAtttribute::TexCoord))
+	// TexCoord vertex attributes are written in pairs.
+	for (uint32 TexCoordPairIndex = 0; TexCoordPairIndex < 4; ++TexCoordPairIndex)
 	{
-		Data.TextureCoordinates[0].VertexBuffer = &DummyVBs[EVertexAtttribute::TexCoord];
-		Data.TextureCoordinates[0].Offset = 0;
-		Data.TextureCoordinates[0].VertexStreamUsage = EVertexStreamUsage::Overridden | EVertexStreamUsage::ManualFetch;
-		Data.TextureCoordinates[0].Stride = 2 * sizeof(float);
+		const uint32 TexCoordIndex0 = 2 * TexCoordPairIndex;
+		const uint32 MaskTexCoord0 = (1 << (EVertexAtttribute::VertexTexCoord0 + TexCoordIndex0));
+		const uint32 MaskTexCoord1 = MaskTexCoord0 << 1;
 
-		Data.NumTexCoords = 1;
+		if (VertexAttributeMask & (MaskTexCoord0 | MaskTexCoord1))
+		{
+			const bool bHasTexCoord1 = (VertexAttributeMask & MaskTexCoord1);
+
+			// Note that we use the dummy VB from TexCoord0, so that is what we need to search for later in BuildStreamIndices().
+			Data.TextureCoordinates[TexCoordPairIndex].VertexBuffer = &DummyVBs[EVertexAtttribute::VertexTexCoord0 + TexCoordIndex0];
+			Data.TextureCoordinates[TexCoordPairIndex].Offset = 0;
+			Data.TextureCoordinates[TexCoordPairIndex].VertexStreamUsage = EVertexStreamUsage::Overridden | EVertexStreamUsage::ManualFetch;
+			Data.TextureCoordinates[TexCoordPairIndex].Stride = sizeof(float) * (bHasTexCoord1 ? 4 : 2);
+
+			Data.NumTexCoords = FMath::Max<int32>(Data.NumTexCoords, TexCoordIndex0 + (bHasTexCoord1 ? 2 : 1));
+		}
 	}
 }
 
 void FGPUSkinPassthroughVertexFactory::OverrideSRVs(FGPUBaseSkinVertexFactory const* InSourceVertexFactory)
 {
-	Data.PositionComponentSRV = SRVs[EVertexAtttribute::Position] ? SRVs[EVertexAtttribute::Position] : (FRHIShaderResourceView*)InSourceVertexFactory->GetPositionsSRV();
-	Data.TangentsSRV = SRVs[EVertexAtttribute::Tangent] ? SRVs[EVertexAtttribute::Tangent] : (FRHIShaderResourceView*)InSourceVertexFactory->GetTangentsSRV();
-	Data.ColorComponentsSRV = SRVs[EVertexAtttribute::Color] ? SRVs[EVertexAtttribute::Color] : (FRHIShaderResourceView*)InSourceVertexFactory->GetColorComponentsSRV();
-	Data.TextureCoordinatesSRV = SRVs[EVertexAtttribute::TexCoord] ? SRVs[EVertexAtttribute::TexCoord] : (FRHIShaderResourceView*)InSourceVertexFactory->GetTextureCoordinatesSRV();
+	Data.TangentsSRV = SRVs[EShaderResource::Tangent] ? SRVs[EShaderResource::Tangent] : (FRHIShaderResourceView*)InSourceVertexFactory->GetTangentsSRV();
+	Data.ColorComponentsSRV = SRVs[EShaderResource::Color] ? SRVs[EShaderResource::Color] : (FRHIShaderResourceView*)InSourceVertexFactory->GetColorComponentsSRV();
+	Data.TextureCoordinatesSRV = SRVs[EShaderResource::TexCoord] ? SRVs[EShaderResource::TexCoord] : (FRHIShaderResourceView*)InSourceVertexFactory->GetTextureCoordinatesSRV();
 }
 
 void FGPUSkinPassthroughVertexFactory::BuildStreamIndices()
 {
+	// Iterate to find matching streams.
 	for (int32 AttributeIndex = 0; AttributeIndex < EVertexAtttribute::NumAttributes; ++AttributeIndex)
 	{
 		StreamIndices[AttributeIndex] = -1;
 
+		// Don't search if vertex attribute is not requested.
+		if ((VertexAttributeMask & (1 << AttributeIndex)) == 0)
+		{
+			continue;
+		}
+
+		FVertexBuffer const* BufferToFind = &DummyVBs[AttributeIndex];
+		
+		// Each TexCoord streams each containing two attributes, so map the second attribute to the first one when searching.
+		if (AttributeIndex == EVertexAtttribute::VertexTexCoord1 || AttributeIndex == EVertexAtttribute::VertexTexCoord3 ||
+			AttributeIndex == EVertexAtttribute::VertexTexCoord5 || AttributeIndex == EVertexAtttribute::VertexTexCoord7)
+		{
+			BufferToFind = &DummyVBs[AttributeIndex - 1];
+		}
+
+		// Search for matching stream.
 		for (int32 StreamIndex = 0; StreamIndex < Streams.Num(); ++StreamIndex)
 		{
-			if (Streams[StreamIndex].VertexBuffer == &DummyVBs[AttributeIndex])
+			if (Streams[StreamIndex].VertexBuffer == BufferToFind)
 			{
 				StreamIndices[AttributeIndex] = StreamIndex;
 				break;
@@ -1361,28 +1389,61 @@ void FGPUSkinPassthroughVertexFactory::CreateUniformBuffer()
 	}
 }
 
-void FGPUSkinPassthroughVertexFactory::AddVertexAttributes(FGPUBaseSkinVertexFactory const* InSourceVertexFactory, TArrayView<EVertexAtttribute> const& InAttributes, TArrayView<FRHIShaderResourceView*> const& InSRVs)
+void FGPUSkinPassthroughVertexFactory::CreateLooseUniformBuffer(FGPUBaseSkinVertexFactory const* InSourceVertexFactory, uint32 InFrameNumber)
+{
+	FRHIShaderResourceView* PositionSRV = SRVs[EShaderResource::Position] != nullptr ? SRVs[EShaderResource::Position] : (FRHIShaderResourceView*)InSourceVertexFactory->GetPositionsSRV();
+	FRHIShaderResourceView* PrevPositionSRV = SRVs[EShaderResource::PreviousPosition] != nullptr ? SRVs[EShaderResource::PreviousPosition] : PositionSRV;
+
+	FLocalVertexFactoryLooseParameters Parameters;
+	Parameters.FrameNumber = InFrameNumber;
+	Parameters.GPUSkinPassThroughPositionBuffer = PositionSRV;
+	Parameters.GPUSkinPassThroughPreviousPositionBuffer = PrevPositionSRV;
+	LooseParametersUniformBuffer.UpdateUniformBufferImmediate(Parameters);
+}
+
+void FGPUSkinPassthroughVertexFactory::SetVertexAttributes(FGPUBaseSkinVertexFactory const* InSourceVertexFactory, FAddVertexAttributeDesc const& InDesc)
 {
-	// Check for new attributes or modified SRVs.
+	// Check for new vertex attributes.
 	bool bNeedFullUpdate = false;
-	bool bNeedSRVUpdate = false;
-	
-	for (int32 Index = 0; Index < InAttributes.Num(); ++Index)
+	for (int32 Index = 0; Index < InDesc.VertexAttributes.Num(); ++Index)
 	{
-		const uint32 MaskBit = 1 << InAttributes[Index];
+		const uint32 MaskBit = 1 << InDesc.VertexAttributes[Index];
 		if ((VertexAttributeMask & MaskBit) == 0)
 		{
 			VertexAttributeMask |= MaskBit;
 			bNeedFullUpdate = true;
 		}
+	}
 
-		if (SRVs[InAttributes[Index]] != InSRVs[Index])
+	// Check for modified SRVs.
+	bool bNeedUniformBufferUpdate = false;
+	bool bNeedLooseUniformBufferUpdate = false;
+	for (int32 Index = 0; Index < EShaderResource::NumShaderResources; ++Index)
+	{
+		if (SRVs[Index] != InDesc.SRVs[Index])
 		{
-			SRVs[InAttributes[Index]] = InSRVs[Index];
-			bNeedSRVUpdate = true;
+			SRVs[Index] = InDesc.SRVs[Index];
+
+			if (Index == EShaderResource::Position || Index == EShaderResource::PreviousPosition)
+			{
+				// Position SRVs are stored in the special "loose" uniform buffer used only by the passthrough vertex factory.
+				bNeedLooseUniformBufferUpdate = true;
+			}
+			else
+			{
+				// All other SRVs are stored in the main vertex factory uniform buffer.
+				bNeedUniformBufferUpdate = true;
+			}
 		}
 	}
 
+	if (UpdatedFrameNumber != InDesc.FrameNumber)
+	{
+		// Loose uniform buffer include the latest frame number.
+		UpdatedFrameNumber = InDesc.FrameNumber;
+		bNeedLooseUniformBufferUpdate = true;
+	}
+
 	if (bNeedFullUpdate)
 	{
 		// Reset the vertex factory data from the source vertex factory.
@@ -1393,12 +1454,13 @@ void FGPUSkinPassthroughVertexFactory::AddVertexAttributes(FGPUBaseSkinVertexFac
 		OverrideSRVs(InSourceVertexFactory);
 
 		// Rebuild the vertex declaration.
+		// This will also update the uniform buffer.
 		UpdateRHI();
 
 		// Rebuild the vertex stream indices.
 		BuildStreamIndices();
 	}
-	else if (bNeedSRVUpdate)
+	else if (bNeedUniformBufferUpdate)
 	{
 		// Override the vertex factory data for our added attributes.
 		OverrideSRVs(InSourceVertexFactory);
@@ -1406,6 +1468,12 @@ void FGPUSkinPassthroughVertexFactory::AddVertexAttributes(FGPUBaseSkinVertexFac
 		// Only need to recreate the vertex factory uniform buffer.
 		CreateUniformBuffer();
 	}
+
+	if (bNeedLooseUniformBufferUpdate)
+	{
+		// Update the loose uniform buffer.
+		CreateLooseUniformBuffer(InSourceVertexFactory, InDesc.FrameNumber);
+	}
 }
 
 int32 FGPUSkinPassthroughVertexFactory::GetAttributeStreamIndex(EVertexAtttribute InAttribute) const
diff --git a/Engine/Source/Runtime/Engine/Private/LocalVertexFactory.cpp b/Engine/Source/Runtime/Engine/Private/LocalVertexFactory.cpp
--- a/Engine/Source/Runtime/Engine/Private/LocalVertexFactory.cpp
+++ b/Engine/Source/Runtime/Engine/Private/LocalVertexFactory.cpp
@@ -172,6 +172,65 @@ void FLocalVertexFactoryShaderParameters::Bind(const FShaderParameterMap& Parame
 	IsGPUSkinPassThrough.Bind(ParameterMap, TEXT("bIsGPUSkinPassThrough"));
 }
 
+static bool IsGPUSkinPassThroughSupported(EShaderPlatform Platform)
+{
+	// Enable the GPUSkin passthrough path if we might use the GPUSkinCache or MeshDeformers.
+	static IMeshDeformerProvider* MeshDeformerProvider = IMeshDeformerProvider::Get();
+	bool bMeshDeformersAvailable = MeshDeformerProvider && MeshDeformerProvider->IsSupported(Platform);
+	return bMeshDeformersAvailable || IsGPUSkinCacheAvailable(Platform);
+}
+
+static void GetMeshDeformerVertexStreams(FMeshDeformerGeometry const& InDeformerGeometry, FGPUSkinPassthroughVertexFactory const* InVertexFactory, FVertexInputStreamArray& InOutVertexStreams)
+{
+	const int32 PositionStreamIndex = InVertexFactory->GetAttributeStreamIndex(FGPUSkinPassthroughVertexFactory::VertexPosition);
+	if (PositionStreamIndex > -1 && InDeformerGeometry.Position.IsValid())
+	{
+		InOutVertexStreams.Add(FVertexInputStream(PositionStreamIndex, 0, InDeformerGeometry.Position->GetRHI()));
+	}
+
+	const int32 TangentStreamIndex = InVertexFactory->GetAttributeStreamIndex(FGPUSkinPassthroughVertexFactory::VertexTangent);
+	if (TangentStreamIndex > -1 && InDeformerGeometry.Tangent.IsValid())
+	{
+		InOutVertexStreams.Add(FVertexInputStream(TangentStreamIndex, 0, InDeformerGeometry.Tangent->GetRHI()));
+	}
+
+	const int32 ColorStreamIndex = InVertexFactory->GetAttributeStreamIndex(FGPUSkinPassthroughVertexFactory::VertexColor);
+	if (ColorStreamIndex > -1 && InDeformerGeometry.Color.IsValid())
+	{
+		InOutVertexStreams.Add(FVertexInputStream(ColorStreamIndex, 0, InDeformerGeometry.Color->GetRHI()));
+	}
+}
+
+static void GetElementShaderBindingsGPUSkinPassThrough(
+	const FMeshMaterialShader* Shader,
+	ERHIFeatureLevel::Type FeatureLevel,
+	const FVertexFactory* VertexFactory,
+	const FMeshBatchElement& BatchElement,
+	class FMeshDrawSingleShaderBindings& ShaderBindings,
+	FVertexInputStreamArray& VertexStreams)
+{
+	// Bind vertex streams.
+	FSkinBatchVertexFactoryUserData* BatchUserData = (FSkinBatchVertexFactoryUserData*)BatchElement.VertexFactoryUserData;
+	FGPUSkinPassthroughVertexFactory const* PassthroughVertexFactory = static_cast<FGPUSkinPassthroughVertexFactory const*>(VertexFactory);
+
+	if (BatchUserData != nullptr && BatchUserData->SkinCacheEntry != nullptr)
+	{
+		// Using Skin Cache.
+		FGPUSkinCache::GetShaderVertexStreams(BatchUserData->SkinCacheEntry, BatchUserData->SectionIndex, PassthroughVertexFactory, VertexStreams);
+	}
+	else if (BatchUserData != nullptr && BatchUserData->DeformerGeometry != nullptr)
+	{
+		// Using Mesh Deformers.
+		GetMeshDeformerVertexStreams(*BatchUserData->DeformerGeometry, PassthroughVertexFactory, VertexStreams);
+	}
+
+	// Bind the vertex factory uniform buffer.
+	if (PassthroughVertexFactory->SupportsManualVertexFetch(FeatureLevel) || UseGPUScene(GMaxRHIShaderPlatform, FeatureLevel))
+	{
+		ShaderBindings.Add(Shader->GetUniformBufferParameter<FLocalVertexFactoryUniformShaderParameters>(), PassthroughVertexFactory->GetUniformBuffer());
+	}
+}
+
 void FLocalVertexFactoryShaderParameters::GetElementShaderBindings(
 	const FSceneInterface* Scene,
 	const FSceneView* View,
@@ -189,10 +248,7 @@ void FLocalVertexFactoryShaderParameters::GetElementShaderBindings(
 	if (LocalVertexFactory->bGPUSkinPassThrough)
 	{
 		GetElementShaderBindingsGPUSkinPassThrough(
-			Scene,
-			View,
 			Shader,
-			InputStreamType,
 			FeatureLevel,
 			VertexFactory,
 			BatchElement,
@@ -220,54 +276,6 @@ void FLocalVertexFactoryShaderParameters::GetElementShaderBindings(
 	ShaderBindings.Add(Shader->GetUniformBufferParameter<FLocalVertexFactoryLooseParameters>(), LocalVertexFactory->LooseParametersUniformBuffer);
 }
 
-bool IsGPUSkinPassThroughSupported(EShaderPlatform Platform)
-{
-	// Enable the GPUSkin passthrough path if we might use the GPUSkinCache or MeshDeformers.
-	static IMeshDeformerProvider* MeshDeformerProvider = IMeshDeformerProvider::Get();
-	bool bMeshDeformersAvailable = MeshDeformerProvider && MeshDeformerProvider->IsSupported(Platform);
-	return bMeshDeformersAvailable || IsGPUSkinCacheAvailable(Platform);
-}
-
-void FLocalVertexFactoryShaderParameters::GetElementShaderBindingsGPUSkinPassThrough(
-	const FSceneInterface* Scene,
-	const FSceneView* View,
-	const FMeshMaterialShader* Shader,
-	const EVertexInputStreamType InputStreamType,
-	ERHIFeatureLevel::Type FeatureLevel,
-	const FVertexFactory* VertexFactory,
-	const FMeshBatchElement& BatchElement,
-	class FMeshDrawSingleShaderBindings& ShaderBindings,
-	FVertexInputStreamArray& VertexStreams) const
-{
-	FGPUSkinPassthroughVertexFactory const* PassthroughVertexFactory = static_cast<FGPUSkinPassthroughVertexFactory const*>(VertexFactory);
-	if (PassthroughVertexFactory->SupportsManualVertexFetch(FeatureLevel) || UseGPUScene(GMaxRHIShaderPlatform, FeatureLevel))
-	{
-		ShaderBindings.Add(Shader->GetUniformBufferParameter<FLocalVertexFactoryUniformShaderParameters>(), PassthroughVertexFactory->GetUniformBuffer());
-	}
-
- 	FSkinBatchVertexFactoryUserData* BatchUserData = (FSkinBatchVertexFactoryUserData*)BatchElement.VertexFactoryUserData;
-	if (!ensure(BatchUserData != nullptr))
-	{
-		return;
-	}
-
- 	if (BatchUserData->SkinCacheEntry != nullptr)
- 	{
-		// Using Skin Cache.
-		FGPUSkinCache::GetShaderBindings(BatchUserData->SkinCacheEntry, BatchUserData->Section, PassthroughVertexFactory, ShaderBindings, VertexStreams);
- 	}
- 	else if (BatchUserData->DeformerGeometry != nullptr)
- 	{
-		// Using Mesh Deformers.
-		FMeshDeformerGeometry* DeformerGeometry = BatchUserData->DeformerGeometry;
-		if (DeformerGeometry->Position.IsValid())
-		{
-			int32 PositionStreamIndex = PassthroughVertexFactory->GetAttributeStreamIndex(FGPUSkinPassthroughVertexFactory::EVertexAtttribute::Position);
-			VertexStreams.Add(FVertexInputStream(PositionStreamIndex, 0, DeformerGeometry->Position->GetRHI()));
-		}
-	}
-}
-
 /**
  * Should we cache the material's shadertype on this platform with this vertex factory? 
  */
@@ -521,6 +529,7 @@ void FLocalVertexFactory::InitRHI()
 	}
 
 	FLocalVertexFactoryLooseParameters LooseParameters;
+	LooseParameters.FrameNumber = -1;
 	LooseParameters.GPUSkinPassThroughPositionBuffer = GNullVertexBuffer.VertexBufferSRV;
 	LooseParameters.GPUSkinPassThroughPreviousPositionBuffer = GNullVertexBuffer.VertexBufferSRV;
 	LooseParametersUniformBuffer = TUniformBufferRef<FLocalVertexFactoryLooseParameters>::CreateUniformBufferImmediate(LooseParameters, UniformBuffer_MultiFrame);
diff --git a/Engine/Source/Runtime/Engine/Private/SceneView.cpp b/Engine/Source/Runtime/Engine/Private/SceneView.cpp
--- a/Engine/Source/Runtime/Engine/Private/SceneView.cpp
+++ b/Engine/Source/Runtime/Engine/Private/SceneView.cpp
@@ -2686,6 +2686,7 @@ void FSceneView::SetupCommonViewUniformBufferParameters(
 	ViewUniformShaderParameters.Random = FMath::Rand();
 	ViewUniformShaderParameters.FrameNumber = Family->FrameNumber;
 
+	ViewUniformShaderParameters.WorldIsPaused = Family->bWorldIsPaused;
 	ViewUniformShaderParameters.CameraCut = bCameraCut ? 1 : 0;
 
 	ViewUniformShaderParameters.MinRoughness = FMath::Clamp(CVarGlobalMinRoughnessOverride.GetValueOnRenderThread(), 0.02f, 1.0f);
diff --git a/Engine/Source/Runtime/Engine/Private/SkeletalMesh.cpp b/Engine/Source/Runtime/Engine/Private/SkeletalMesh.cpp
--- a/Engine/Source/Runtime/Engine/Private/SkeletalMesh.cpp
+++ b/Engine/Source/Runtime/Engine/Private/SkeletalMesh.cpp
@@ -53,6 +53,7 @@
 
 #if WITH_EDITOR
 #include "AssetRegistry/IAssetRegistry.h"
+#include "AssetRegistry/AssetData.h"
 #include "Rendering/SkeletalMeshModel.h"
 #include "IMeshBuilderModule.h"
 #include "IMeshReductionManagerModule.h"
@@ -6399,65 +6400,6 @@ void FSkeletalMeshSceneProxy::CreateBaseMeshBatch(const FSceneView* View, const
 
 	BatchElement.PrimitiveUniformBuffer = GetUniformBuffer();
 	BatchElement.NumPrimitives = LODData.RenderSections[SectionIndex].NumTriangles;
-
-	// todo: Move this uniform buffer update to somewhere that doesn't require a const cast?
-	FVertexFactory* NonConstVertexFactory = const_cast<FVertexFactory*>(Mesh.VertexFactory);
-	UpdateLooseParametersUniformBuffer(View, SectionIndex, VertexFactoryUserData, NonConstVertexFactory);
-}
-
-void FSkeletalMeshSceneProxy::UpdateLooseParametersUniformBuffer(const FSceneView* View, const int32 SectionIndex, const FSkinBatchVertexFactoryUserData* BatchUserData, FVertexFactory* VertexFactory) const
-{
-	// Loose parameters uniform buffer is only needed for PassThroughVF which is a type of LocalVF.
-	// Check it is a LocalVF StaticType and bGPUSkinPassThrough flag is set. Early out if either condition is not met.
-	if (VertexFactory == nullptr || VertexFactory->GetType() != &FLocalVertexFactory::StaticType)
-	{
-		return;
-	}
-	FLocalVertexFactory* LocalVertexFactory = static_cast<FLocalVertexFactory*>(VertexFactory);
-	if (!LocalVertexFactory->bGPUSkinPassThrough)
-	{
-		return;
-	}
-	// We always expect a valid BatchUserData for PassthroughVF.
-	if (!ensure(BatchUserData != nullptr))
-	{
-		return;
-	}
-
-	uint32 UpdatedFrameNumber = 0;
-	FRHIShaderResourceView* PositionSRV = nullptr;
-	FRHIShaderResourceView* PrevPositionSRV = nullptr;
-
-	if (BatchUserData->SkinCacheEntry != nullptr)
-	{
-		// Using Skin Cache.
-		UpdatedFrameNumber = FGPUSkinCache::GetUpdatedFrame(BatchUserData->SkinCacheEntry, SectionIndex);
-		PositionSRV = FGPUSkinCache::GetPositionBuffer(BatchUserData->SkinCacheEntry, SectionIndex)->SRV;
-		PrevPositionSRV = FGPUSkinCache::GetPreviousPositionBuffer(BatchUserData->SkinCacheEntry, SectionIndex)->SRV;
-	}
-	else if (BatchUserData->DeformerGeometry != nullptr)
-	{
-		// Using Mesh Deformers.
-		FMeshDeformerGeometry& DeformerGeometry = *BatchUserData->DeformerGeometry;
-		UpdatedFrameNumber = DeformerGeometry.PositionUpdatedFrame;
-		PositionSRV = DeformerGeometry.PositionSRV;
-		PrevPositionSRV = DeformerGeometry.PrevPositionSRV;
-	}
-
-	// Bone data is updated whenever animation triggers a dynamic update, animation can skip frames hence the frequency is not necessary every frame.
-	// So check if bone data is updated this frame, if not then the previous frame data is stale and not suitable for motion blur.
-	bool bBoneDataUpdatedThisFrame = (View->Family->FrameNumber == UpdatedFrameNumber);
-	// If world is paused, use current frame bone matrices, so velocity is canceled and skeletal mesh isn't blurred from motion.
-	bool bVerticesInMotion = !View->Family->bWorldIsPaused && bBoneDataUpdatedThisFrame;
-
-	PositionSRV = (PositionSRV == nullptr) ? LocalVertexFactory->GetPositionsSRV() : PositionSRV;
-	PrevPositionSRV = (PrevPositionSRV == nullptr || !bVerticesInMotion) ? PositionSRV : PrevPositionSRV;
-
-	FLocalVertexFactoryLooseParameters Parameters;
-	Parameters.GPUSkinPassThroughPositionBuffer = PositionSRV;
-	Parameters.GPUSkinPassThroughPreviousPositionBuffer = PrevPositionSRV;
-	
-	LocalVertexFactory->LooseParametersUniformBuffer.UpdateUniformBufferImmediate(Parameters);
 }
 
 uint8 FSkeletalMeshSceneProxy::GetCurrentFirstLODIdx_Internal() const
diff --git a/Engine/Source/Runtime/Engine/Private/SkeletalMeshDeformerHelpers.cpp b/Engine/Source/Runtime/Engine/Private/SkeletalMeshDeformerHelpers.cpp
--- a/Engine/Source/Runtime/Engine/Private/SkeletalMeshDeformerHelpers.cpp
+++ b/Engine/Source/Runtime/Engine/Private/SkeletalMeshDeformerHelpers.cpp
@@ -217,26 +217,27 @@ void FSkeletalMeshDeformerHelpers::UpdateVertexFactoryBufferOverrides(FSkeletalM
 	FSkeletalMeshObjectGPUSkin* MeshObjectGPU = static_cast<FSkeletalMeshObjectGPUSkin*>(InMeshObject);
 	FMeshDeformerGeometry& DeformerGeometry = MeshObjectGPU->GetDeformerGeometry(InLodIndex);
 	
-	TArray<FGPUSkinPassthroughVertexFactory::EVertexAtttribute, TInlineAllocator<4>> Attributes;
-	TArray<FRHIShaderResourceView*, TInlineAllocator<4>> SRVs;
+	FGPUSkinPassthroughVertexFactory::FAddVertexAttributeDesc Desc;
+	Desc.FrameNumber = DeformerGeometry.PositionUpdatedFrame;
 
 	if (DeformerGeometry.PositionSRV)
 	{
-		Attributes.Add(FGPUSkinPassthroughVertexFactory::EVertexAtttribute::Position);
-		SRVs.Add(DeformerGeometry.PositionSRV);
+		Desc.VertexAttributes.Add(FGPUSkinPassthroughVertexFactory::VertexPosition);
+		Desc.SRVs[FGPUSkinPassthroughVertexFactory::Position] = DeformerGeometry.PositionSRV;
+		Desc.SRVs[FGPUSkinPassthroughVertexFactory::PreviousPosition] = DeformerGeometry.PrevPositionSRV;
 	}
 	if (DeformerGeometry.TangentSRV)
 	{
-		Attributes.Add(FGPUSkinPassthroughVertexFactory::EVertexAtttribute::Tangent);
-		SRVs.Add(DeformerGeometry.TangentSRV);
+		Desc.VertexAttributes.Add(FGPUSkinPassthroughVertexFactory::VertexTangent);
+		Desc.SRVs[FGPUSkinPassthroughVertexFactory::Tangent] = DeformerGeometry.TangentSRV;
 	}
 	if (DeformerGeometry.ColorSRV)
 	{
-		Attributes.Add(FGPUSkinPassthroughVertexFactory::EVertexAtttribute::Color);
-		SRVs.Add(DeformerGeometry.ColorSRV);
+		Desc.VertexAttributes.Add(FGPUSkinPassthroughVertexFactory::VertexColor);
+		Desc.SRVs[FGPUSkinPassthroughVertexFactory::Color] = DeformerGeometry.ColorSRV;
 	}
 
-	if (Attributes.Num() == 0)
+	if (Desc.VertexAttributes.Num() == 0)
 	{
 		return;
 	}
@@ -247,8 +248,7 @@ void FSkeletalMeshDeformerHelpers::UpdateVertexFactoryBufferOverrides(FSkeletalM
 	{
 		FGPUBaseSkinVertexFactory const* BaseVertexFactory = MeshObjectGPU->GetBaseSkinVertexFactory(InLodIndex, SectionIndex);
 		FGPUSkinPassthroughVertexFactory* TargetVertexFactory = LOD.GPUSkinVertexFactories.PassthroughVertexFactories[SectionIndex].Get();
-		
-		TargetVertexFactory->AddVertexAttributes(BaseVertexFactory, Attributes, SRVs);
+		TargetVertexFactory->SetVertexAttributes(BaseVertexFactory, Desc);
 	}
 }
 
diff --git a/Engine/Source/Runtime/Engine/Public/GPUSkinCache.h b/Engine/Source/Runtime/Engine/Public/GPUSkinCache.h
--- a/Engine/Source/Runtime/Engine/Public/GPUSkinCache.h
+++ b/Engine/Source/Runtime/Engine/Public/GPUSkinCache.h
@@ -145,11 +145,10 @@ public:
 		FGPUSkinCacheEntry*& InOutEntry
 		);
 
-	static void GetShaderBindings(
+	static void GetShaderVertexStreams(
 		const FGPUSkinCacheEntry* Entry,
 		int32 Section,
 		const FGPUSkinPassthroughVertexFactory* VertexFactory,
-		class FMeshDrawSingleShaderBindings& ShaderBindings,
 		FVertexInputStreamArray& VertexStreams);
 
 	static void Release(FGPUSkinCacheEntry*& SkinCacheEntry);
diff --git a/Engine/Source/Runtime/Engine/Public/GPUSkinVertexFactory.h b/Engine/Source/Runtime/Engine/Public/GPUSkinVertexFactory.h
--- a/Engine/Source/Runtime/Engine/Public/GPUSkinVertexFactory.h
+++ b/Engine/Source/Runtime/Engine/Public/GPUSkinVertexFactory.h
@@ -728,32 +728,61 @@ public:
 	bool SupportsPositionAndNormalOnlyStream() const override { return false; }
 	// End FVertexFactory Interface.
 
+	/**
+	 * Reset all added vertex attributes and SRVs.
+	 * This doesn't reset the vertex factory itself. Call SetData() to do that.
+	 */
+	void ResetVertexAttributes();
+
 	/** Vertex attributes that we can override. */
 	enum EVertexAtttribute
+	{
+		VertexPosition,
+		VertexTangent,
+		VertexColor,
+		VertexTexCoord0,
+		VertexTexCoord1,
+		VertexTexCoord2,
+		VertexTexCoord3,
+		VertexTexCoord4,
+		VertexTexCoord5,
+		VertexTexCoord6,
+		VertexTexCoord7,
+		NumAttributes
+	};
+	
+	/** SRVs that we can provide. */
+	enum EShaderResource
 	{
 		Position,
+		PreviousPosition,
 		Tangent,
 		Color,
 		TexCoord,
-		NumAttributes
+		NumShaderResources
 	};
 
-	/** 
-	 * Reset all added vertex attributes and SRVs. 
-	 * This doesn't reset the vertex factory itself. Call SetData() to do that.
-	 */
-	void ResetVertexAttributes();
+	/** Structure used for calls to SetVertexAttributes(). */
+	struct FAddVertexAttributeDesc
+	{
+		FAddVertexAttributeDesc() : SRVs(InPlace, nullptr) {}
+
+		/** Frame number at animation update. Used to determine if animation motion is valid and needs to output velocity. */
+		uint32 FrameNumber = ~0U;
+		/** Vertex attributes to use in vertex declaration. */
+		TArray<EVertexAtttribute, TFixedAllocator<EVertexAtttribute::NumAttributes>> VertexAttributes;
+		/** SRVs for binding. These are only be used by platforms that support manual vertex fetch. */
+		TStaticArray<FRHIShaderResourceView*, EShaderResource::NumShaderResources> SRVs;
+	};
 
 	/** 
-	 * Add vertex attributes that will be in the vertex stream. 
-	 * We accumulate with any attributes that have already been added by previous calls to this function.
-	 * If any attributes are being added for the first time then we recreate the vertex declaration here.
-	 * The InSRVs array size should match the InAttributes array size. 
-	 * SRVs will only be used by platforms that support manual vertex fetch.
+	 * Set vertex attributes and SRVs to be used. 
+	 * The vertex declaration is made by accumulating attributes set here along with all that already been added by previous calls to this function.
+	 * If any attributes are being added for the first time then we pay the cost to recreate the vertex declaration here.
 	 * The SRVs are cached per attribute. If any passed in SRV is changed from the cached value then we recreate the vertex factory uniform buffer here.
 	 * Note that on platforms that support manual vertex fetch, only Position will be in the final vertex stream and other attributes will be read through an SRV.
 	 */
-	void AddVertexAttributes(FGPUBaseSkinVertexFactory const* InSourceVertexFactory, TArrayView<EVertexAtttribute> const& InAttributes, TArrayView<FRHIShaderResourceView*> const& InSRVs);
+	void SetVertexAttributes(FGPUBaseSkinVertexFactory const* InSourceVertexFactory, FAddVertexAttributeDesc const& InDesc);
 
 	/** 
 	 * Get the vertex stream index for a vertex attribute.
@@ -767,10 +796,12 @@ private:
 	void OverrideSRVs(FGPUBaseSkinVertexFactory const* InSourceVertexFactory);
 	void BuildStreamIndices();
 	void CreateUniformBuffer();
+	void CreateLooseUniformBuffer(FGPUBaseSkinVertexFactory const* InSourceVertexFactory, uint32 InFrameNumber);
 
-	uint32 VertexAttributeMask = 0;
+	uint32 VertexAttributeMask;
 	TStaticArray<int32, EVertexAtttribute::NumAttributes> StreamIndices;
-	TStaticArray<FRHIShaderResourceView*, EVertexAtttribute::NumAttributes> SRVs;
-	
+	TStaticArray<FRHIShaderResourceView*, EShaderResource::NumShaderResources> SRVs;
+	uint32 UpdatedFrameNumber;
+
 	static TStaticArray<FVertexBuffer, EVertexAtttribute::NumAttributes> DummyVBs;
 };
diff --git a/Engine/Source/Runtime/Engine/Public/LocalVertexFactory.h b/Engine/Source/Runtime/Engine/Public/LocalVertexFactory.h
--- a/Engine/Source/Runtime/Engine/Public/LocalVertexFactory.h
+++ b/Engine/Source/Runtime/Engine/Public/LocalVertexFactory.h
@@ -28,6 +28,7 @@ BEGIN_GLOBAL_SHADER_PARAMETER_STRUCT(FLocalVertexFactoryUniformShaderParameters,
 END_GLOBAL_SHADER_PARAMETER_STRUCT()
 
 BEGIN_GLOBAL_SHADER_PARAMETER_STRUCT(FLocalVertexFactoryLooseParameters,)
+	SHADER_PARAMETER(uint32, FrameNumber)
 	SHADER_PARAMETER_SRV(Buffer<float>, GPUSkinPassThroughPositionBuffer)
 	SHADER_PARAMETER_SRV(Buffer<float>, GPUSkinPassThroughPreviousPositionBuffer)
 END_GLOBAL_SHADER_PARAMETER_STRUCT()
@@ -244,16 +245,5 @@ public:
 	) const; 
 
 private:
-	void GetElementShaderBindingsGPUSkinPassThrough(
-		const FSceneInterface* Scene,
-		const FSceneView* View,
-		const FMeshMaterialShader* Shader,
-		const EVertexInputStreamType InputStreamType,
-		ERHIFeatureLevel::Type FeatureLevel,
-		const FVertexFactory* VertexFactory,
-		const FMeshBatchElement& BatchElement,
-		class FMeshDrawSingleShaderBindings& ShaderBindings,
-		FVertexInputStreamArray& VertexStreams) const;
-
 	LAYOUT_FIELD(FShaderParameter, IsGPUSkinPassThrough);
 };
\ No newline at end of file
diff --git a/Engine/Source/Runtime/Engine/Public/SceneView.h b/Engine/Source/Runtime/Engine/Public/SceneView.h
--- a/Engine/Source/Runtime/Engine/Public/SceneView.h
+++ b/Engine/Source/Runtime/Engine/Public/SceneView.h
@@ -701,6 +701,7 @@ enum ETranslucencyVolumeCascade
 	VIEW_UNIFORM_BUFFER_MEMBER(uint32, StateFrameIndexMod8) \
 	VIEW_UNIFORM_BUFFER_MEMBER(uint32, StateFrameIndex) \
 	VIEW_UNIFORM_BUFFER_MEMBER(uint32, DebugViewModeMask) \
+	VIEW_UNIFORM_BUFFER_MEMBER(uint32, WorldIsPaused) \
 	VIEW_UNIFORM_BUFFER_MEMBER_EX(float, CameraCut, EShaderPrecisionModifier::Half) \
 	VIEW_UNIFORM_BUFFER_MEMBER_EX(float, UnlitViewmodeMask, EShaderPrecisionModifier::Half) \
 	VIEW_UNIFORM_BUFFER_MEMBER_EX(FLinearColor, DirectionalLightColor, EShaderPrecisionModifier::Half) \
diff --git a/Engine/Source/Runtime/Engine/Public/SkeletalMeshSceneProxy.h b/Engine/Source/Runtime/Engine/Public/SkeletalMeshSceneProxy.h
--- a/Engine/Source/Runtime/Engine/Public/SkeletalMeshSceneProxy.h
+++ b/Engine/Source/Runtime/Engine/Public/SkeletalMeshSceneProxy.h
@@ -211,7 +211,6 @@ protected:
 	uint8 GetCurrentFirstLODIdx_Internal() const;
 private:
 	void CreateBaseMeshBatch(const FSceneView* View, const FSkeletalMeshLODRenderData& LODData, const int32 LODIndex, const int32 SectionIndex, const FSectionElementInfo& SectionElementInfo, FMeshBatch& Mesh, ESkinVertexFactoryMode VFMode = ESkinVertexFactoryMode::Default) const;
-	void UpdateLooseParametersUniformBuffer(const FSceneView* View, const int32 SectionIndex, const struct FSkinBatchVertexFactoryUserData* BatchUserData, class FVertexFactory* VertexFactory) const;
 
 public:
 #if WITH_EDITORONLY_DATA
diff --git a/Engine/Source/Runtime/Engine/Public/SkeletalRenderPublic.h b/Engine/Source/Runtime/Engine/Public/SkeletalRenderPublic.h
--- a/Engine/Source/Runtime/Engine/Public/SkeletalRenderPublic.h
+++ b/Engine/Source/Runtime/Engine/Public/SkeletalRenderPublic.h
@@ -58,7 +58,7 @@ struct FSkinBatchVertexFactoryUserData
 {
 	FGPUSkinCacheEntry* SkinCacheEntry = nullptr;
 	FMeshDeformerGeometry* DeformerGeometry = nullptr;
-	int32 Section = -1;
+	int32 SectionIndex = -1;
 };
 
 /**
