[Engine]
* Last fixups of Engine module with IWYU to make it run multiple times in a row without changing anything. This included adding a few "// IWYU pragma: keep" and move some includes to better places in the file

diff --git a/Engine/Source/Runtime/Engine/Private/AssetCompilingManager.cpp b/Engine/Source/Runtime/Engine/Private/AssetCompilingManager.cpp
--- a/Engine/Source/Runtime/Engine/Private/AssetCompilingManager.cpp
+++ b/Engine/Source/Runtime/Engine/Private/AssetCompilingManager.cpp
@@ -2,8 +2,8 @@
 
 #include "AssetCompilingManager.h"
 
-#include "HAL/LowLevelMemStats.h"
-#include "HAL/LowLevelMemTracker.h"
+#include "HAL/LowLevelMemStats.h" // IWYU pragma: keep
+#include "HAL/LowLevelMemTracker.h" // IWYU pragma: keep
 
 DECLARE_LLM_MEMORY_STAT(TEXT("AssetCompilation"), STAT_AssetCompilationLLM, STATGROUP_LLMFULL);
 DECLARE_LLM_MEMORY_STAT(TEXT("AssetCompilation"), STAT_AssetCompilationSummaryLLM, STATGROUP_LLM);
diff --git a/Engine/Source/Runtime/Engine/Private/BlueprintGeneratedClass.cpp b/Engine/Source/Runtime/Engine/Private/BlueprintGeneratedClass.cpp
--- a/Engine/Source/Runtime/Engine/Private/BlueprintGeneratedClass.cpp
+++ b/Engine/Source/Runtime/Engine/Private/BlueprintGeneratedClass.cpp
@@ -19,7 +19,7 @@
 #include "Engine/InheritableComponentHandler.h"
 #include "Net/Core/PushModel/PushModel.h"
 #include "UObject/FortniteMainBranchObjectVersion.h"
-#include "UObject/Package.h"
+#include "UObject/Package.h" // IWYU pragma: keep
 #include "UObject/PrimaryAssetId.h"
 #include "UObject/UE5MainStreamObjectVersion.h"
 #include "GenericPlatform/GenericPlatformCrashContext.h"
@@ -28,11 +28,10 @@
 
 #if WITH_EDITOR
 #include "CookerSettings.h"
-#include "Editor/EditorEngine.h"
 #include "Kismet2/BlueprintEditorUtils.h"
 #include "BlueprintCompilationManager.h"
 #include "Engine/LevelScriptBlueprint.h"
-extern UNREALED_API UEditorEngine* GEditor;
+extern UNREALED_API class UEditorEngine* GEditor;
 #endif //WITH_EDITOR
 
 DEFINE_STAT(STAT_PersistentUberGraphFrameMemory);
diff --git a/Engine/Source/Runtime/Engine/Private/Components/ApplicationLifecycleComponent.cpp b/Engine/Source/Runtime/Engine/Private/Components/ApplicationLifecycleComponent.cpp
--- a/Engine/Source/Runtime/Engine/Private/Components/ApplicationLifecycleComponent.cpp
+++ b/Engine/Source/Runtime/Engine/Private/Components/ApplicationLifecycleComponent.cpp
@@ -2,6 +2,7 @@
 // ApplicationLifecycleComponent.cpp: Component to handle receiving notifications from the OS about application state (activated, suspended, termination, etc)
 
 #include "Components/ApplicationLifecycleComponent.h"
+#include "Templates/IsConst.h"
 
 #include UE_INLINE_GENERATED_CPP_BY_NAME(ApplicationLifecycleComponent)
 
diff --git a/Engine/Source/Runtime/Engine/Private/ContentStreaming.cpp b/Engine/Source/Runtime/Engine/Private/ContentStreaming.cpp
--- a/Engine/Source/Runtime/Engine/Private/ContentStreaming.cpp
+++ b/Engine/Source/Runtime/Engine/Private/ContentStreaming.cpp
@@ -6,7 +6,6 @@
 
 #include "ContentStreaming.h"
 #include "Components/PrimitiveComponent.h"
-#include "Engine/Engine.h"
 #include "Engine/Texture2D.h"
 #include "Misc/ConfigCacheIni.h"
 #include "RHI.h"
@@ -20,6 +19,12 @@
 #include "VT/VirtualTextureChunkManager.h"
 #include "Rendering/NaniteCoarseMeshStreamingManager.h"
 
+#if WITH_EDITOR
+#include "AudioDevice.h"
+#else
+#include "Engine/Engine.h"
+#endif
+
 /*-----------------------------------------------------------------------------
 	Globals.
 -----------------------------------------------------------------------------*/
@@ -1231,7 +1236,6 @@ void FStreamingManagerCollection::PropagateLightingScenarioChange()
 }
 
 #if WITH_EDITOR
-#include "AudioDevice.h"
 
 void FStreamingManagerCollection::OnAudioStreamingParamsChanged()
 {
diff --git a/Engine/Source/Runtime/Engine/Private/LevelTick.cpp b/Engine/Source/Runtime/Engine/Private/LevelTick.cpp
--- a/Engine/Source/Runtime/Engine/Private/LevelTick.cpp
+++ b/Engine/Source/Runtime/Engine/Private/LevelTick.cpp
@@ -5,7 +5,6 @@
 =============================================================================*/
 
 #include "Engine/Level.h"
-#include "Collision.h"
 #include "Misc/TimeGuard.h"
 #include "GameFramework/WorldSettings.h"
 #include "UObject/UObjectStats.h"
@@ -23,10 +22,8 @@
 #include "IXRTrackingSystem.h"
 #include "Camera/CameraPhotography.h"
 #include "UObject/Stack.h"
-
-#if ENABLE_COLLISION_ANALYZER
 #include "PhysicsEngine/CollisionAnalyzerCapture.h"
-#endif
+
 #if !UE_SERVER
 #include "IMediaModule.h"
 #endif
@@ -1195,10 +1192,6 @@ void UWorld::FlushDeferredParameterCollectionInstanceUpdates()
 	}
 }
 
-#if ENABLE_COLLISION_ANALYZER
-#endif // ENABLE_COLLISION_ANALYZER
-
-
 #if (CSV_PROFILER && !UE_BUILD_SHIPPING)
 static TAutoConsoleVariable<int32> CVarRecordTickCountsToCSV(
 	TEXT("csv.RecordTickCounts"),
diff --git a/Engine/Source/Runtime/Engine/Private/Particles/SubUVAnimation.cpp b/Engine/Source/Runtime/Engine/Private/Particles/SubUVAnimation.cpp
--- a/Engine/Source/Runtime/Engine/Private/Particles/SubUVAnimation.cpp
+++ b/Engine/Source/Runtime/Engine/Private/Particles/SubUVAnimation.cpp
@@ -11,7 +11,6 @@
 #include "DerivedDataCacheInterface.h"
 #include "ComponentReregisterContext.h"
 #include "RHI.h"
-#include "RHIDefinitions.h"
 
 #define LOCTEXT_NAMESPACE "FSubUVDerivedData"
 
diff --git a/Engine/Source/Runtime/Engine/Private/Rendering/SkeletalMeshVertexClothBuffer.cpp b/Engine/Source/Runtime/Engine/Private/Rendering/SkeletalMeshVertexClothBuffer.cpp
--- a/Engine/Source/Runtime/Engine/Private/Rendering/SkeletalMeshVertexClothBuffer.cpp
+++ b/Engine/Source/Runtime/Engine/Private/Rendering/SkeletalMeshVertexClothBuffer.cpp
@@ -6,10 +6,8 @@
 #include "EngineUtils.h"
 #include "Experimental/Containers/HazardPointer.h"
 #include "ProfilingDebugging/LoadTimeTracker.h"
-#include "PixelFormat.h"
 #include "RHIResourceUpdates.h"
 #include "RHI.h"
-#include "RHICommandList.h"
 
 /**
 * Constructor
diff --git a/Engine/Source/Runtime/Engine/Private/SampleBufferIO.cpp b/Engine/Source/Runtime/Engine/Private/SampleBufferIO.cpp
--- a/Engine/Source/Runtime/Engine/Private/SampleBufferIO.cpp
+++ b/Engine/Source/Runtime/Engine/Private/SampleBufferIO.cpp
@@ -1,6 +1,7 @@
 // Copyright Epic Games, Inc. All Rights Reserved.
 
 #include "Sound/SampleBufferIO.h"
+#include "HAL/PlatformFile.h"
 #include "HAL/PlatformFileManager.h"
 #include "AssetRegistry/AssetRegistryModule.h"
 #include "Misc/PackageName.h"
diff --git a/Engine/Source/Runtime/Engine/Private/SceneTexturesConfig.cpp b/Engine/Source/Runtime/Engine/Private/SceneTexturesConfig.cpp
--- a/Engine/Source/Runtime/Engine/Private/SceneTexturesConfig.cpp
+++ b/Engine/Source/Runtime/Engine/Private/SceneTexturesConfig.cpp
@@ -1,12 +1,9 @@
 // Copyright Epic Games, Inc. All Rights Reserved.
 
 #include "SceneTexturesConfig.h"
-#include "RenderUtils.h"
 #include "ShaderCompiler.h"
-#include "SceneInterface.h"
 #include "Engine/Engine.h"
-#include "ShadowRendering.h"
-#include "StereoRendering.h"
+#include "ShadowRendering.h" // IWYU pragma: keep
 #include "StereoRenderTargetManager.h"
 
 FSceneTexturesConfig FSceneTexturesConfig::GlobalInstance;
diff --git a/Engine/Source/Runtime/Engine/Private/ShaderCompiler/ShaderCompiler.cpp b/Engine/Source/Runtime/Engine/Private/ShaderCompiler/ShaderCompiler.cpp
--- a/Engine/Source/Runtime/Engine/Private/ShaderCompiler/ShaderCompiler.cpp
+++ b/Engine/Source/Runtime/Engine/Private/ShaderCompiler/ShaderCompiler.cpp
@@ -16,6 +16,7 @@
 #include "EditorSupportDelegates.h"
 #include "Engine/RendererSettings.h"
 #include "Features/IModularFeatures.h"
+#include "HAL/PlatformFile.h"
 #include "HAL/PlatformFileManager.h"
 #include "Interfaces/IShaderFormat.h"
 #include "Interfaces/ITargetPlatform.h"
diff --git a/Engine/Source/Runtime/Engine/Private/SkeletalRenderCPUSkin.cpp b/Engine/Source/Runtime/Engine/Private/SkeletalRenderCPUSkin.cpp
--- a/Engine/Source/Runtime/Engine/Private/SkeletalRenderCPUSkin.cpp
+++ b/Engine/Source/Runtime/Engine/Private/SkeletalRenderCPUSkin.cpp
@@ -6,7 +6,6 @@
 
 #include "SkeletalRenderCPUSkin.h"
 #include "EngineStats.h"
-#include "Engine/SkinnedAssetCommon.h"
 #include "Components/SkeletalMeshComponent.h"
 #include "PrimitiveSceneProxy.h"
 #include "RenderUtils.h"
diff --git a/Engine/Source/Runtime/Engine/Private/SkeletalRenderStatic.cpp b/Engine/Source/Runtime/Engine/Private/SkeletalRenderStatic.cpp
--- a/Engine/Source/Runtime/Engine/Private/SkeletalRenderStatic.cpp
+++ b/Engine/Source/Runtime/Engine/Private/SkeletalRenderStatic.cpp
@@ -5,7 +5,6 @@
 =============================================================================*/
 
 #include "SkeletalRenderStatic.h"
-#include "Engine/SkinnedAssetCommon.h"
 #include "RenderUtils.h"
 #include "Rendering/SkeletalMeshRenderData.h"
 #include "RenderingThread.h"
diff --git a/Engine/Source/Runtime/Engine/Private/SoundSourceBus.cpp b/Engine/Source/Runtime/Engine/Private/SoundSourceBus.cpp
--- a/Engine/Source/Runtime/Engine/Private/SoundSourceBus.cpp
+++ b/Engine/Source/Runtime/Engine/Private/SoundSourceBus.cpp
@@ -2,8 +2,6 @@
 
 
 #include "Sound/SoundSourceBus.h"
-#include "Sound/SoundSourceBusSend.h"
-#include "Sound/SoundSubmixSend.h"
 
 #include UE_INLINE_GENERATED_CPP_BY_NAME(SoundSourceBus)
 
diff --git a/Engine/Source/Runtime/Engine/Private/VT/VirtualTextureChunkManager.cpp b/Engine/Source/Runtime/Engine/Private/VT/VirtualTextureChunkManager.cpp
--- a/Engine/Source/Runtime/Engine/Private/VT/VirtualTextureChunkManager.cpp
+++ b/Engine/Source/Runtime/Engine/Private/VT/VirtualTextureChunkManager.cpp
@@ -5,6 +5,7 @@
 #include "RHI.h"
 #include "RenderUtils.h"
 #include "UploadingVirtualTexture.h"
+#include "VT/VirtualTextureBuiltData.h"
 #include "VT/VirtualTextureTranscodeCache.h"
 #include "VT/VirtualTextureUploadCache.h"
 
diff --git a/Engine/Source/Runtime/Engine/Private/VT/VirtualTextureTranscodeCache.cpp b/Engine/Source/Runtime/Engine/Private/VT/VirtualTextureTranscodeCache.cpp
--- a/Engine/Source/Runtime/Engine/Private/VT/VirtualTextureTranscodeCache.cpp
+++ b/Engine/Source/Runtime/Engine/Private/VT/VirtualTextureTranscodeCache.cpp
@@ -8,6 +8,7 @@
 #include "RendererInterface.h"
 #include "UploadingVirtualTexture.h"
 #include "VT/VirtualTextureBuiltData.h"
+#include "VT/VirtualTextureUploadCache.h"
 #include "VirtualTextureChunkManager.h"
 
 static int32 TranscodeRetireAge = 60; //1 second @ 60 fps
