[Engine]
* Fixed win64 non-unity compile error (or when RHI_RAYTRACING was enabled)

diff --git a/Engine/Source/Runtime/Engine/Private/SkeletalRenderCPUSkin.cpp b/Engine/Source/Runtime/Engine/Private/SkeletalRenderCPUSkin.cpp
--- a/Engine/Source/Runtime/Engine/Private/SkeletalRenderCPUSkin.cpp
+++ b/Engine/Source/Runtime/Engine/Private/SkeletalRenderCPUSkin.cpp
@@ -16,6 +16,10 @@
 #include "SceneInterface.h"
 #include "Stats/StatsTrace.h"
 
+#if RHI_RAYTRACING
+#include "Engine/SkinnedAssetCommon.h"
+#endif
+
 struct FMorphTargetDelta;
 
 template<typename VertexType, int32 NumberOfUVs>
diff --git a/Engine/Source/Runtime/Engine/Private/SkeletalRenderStatic.cpp b/Engine/Source/Runtime/Engine/Private/SkeletalRenderStatic.cpp
--- a/Engine/Source/Runtime/Engine/Private/SkeletalRenderStatic.cpp
+++ b/Engine/Source/Runtime/Engine/Private/SkeletalRenderStatic.cpp
@@ -9,6 +9,10 @@
 #include "Rendering/SkeletalMeshRenderData.h"
 #include "RenderingThread.h"
 
+#if RHI_RAYTRACING
+#include "Engine/SkinnedAssetCommon.h"
+#endif
+
 FSkeletalMeshObjectStatic::FSkeletalMeshObjectStatic(USkinnedMeshComponent* InMeshComponent, FSkeletalMeshRenderData* InSkelMeshRenderData, ERHIFeatureLevel::Type InFeatureLevel)
 	: FSkeletalMeshObject(InMeshComponent, InSkelMeshRenderData, InFeatureLevel)
 {
