World Partition Editor/Runtime Conversions:
- Updated ConvertEditorPathToRuntimePath to use the actor to cells remapping available in cooked builds instead of the one which is editor-only.
- Added support for instanced levels, newly spawned actors and paths to actor components.
- Updated the streaming generation unit tests to test new cases.

diff --git a/Engine/Source/Runtime/Engine/Private/Tests/WorldPartition/WorldPartitionStreamingGenerationTests.cpp b/Engine/Source/Runtime/Engine/Private/Tests/WorldPartition/WorldPartitionStreamingGenerationTests.cpp
--- a/Engine/Source/Runtime/Engine/Private/Tests/WorldPartition/WorldPartitionStreamingGenerationTests.cpp
+++ b/Engine/Source/Runtime/Engine/Private/Tests/WorldPartition/WorldPartitionStreamingGenerationTests.cpp
@@ -73,41 +73,86 @@ namespace WorldPartitionTests
 			return false;
 		}
 
-		FSoftObjectPath ActorRuntimePath;
-		if (!TestTrue(TEXT("Actor Path Editor to Runtime Conversion Failed"), FWorldPartitionHelpers::ConvertEditorPathToRuntimePath(ActorRef->GetActorSoftPath(), ActorRuntimePath)))
+		FWorldPartitionHandle ActorHandle(ActorDescMainContainer, FGuid(TEXT("0D2B04D240BE5DE58FE437A8D2DBF5C9")));
+		if (!TestTrue(TEXT("Invalid Actor Handle"), ActorHandle.IsValid()))
 		{
 			return false;
 		}
 
-		FSoftObjectPath ActorEditorPath;
-		if (!TestTrue(TEXT("Actor Path Runtime to Editor Conversion Failed"), FWorldPartitionHelpers::ConvertRuntimePathToEditorPath(ActorRuntimePath, ActorEditorPath)))
+		if (!TestTrue(TEXT("Actor Handle Not Loaded"), !ActorHandle->GetActor()))
 		{
 			return false;
 		}
 
-		if (!TestTrue(TEXT("Actor Path Editor to Runtime to Editor Conversion Failed"), ActorEditorPath == ActorRef->GetActorSoftPath()))
+		UObject* ResolvedObject = ActorHandle->GetActorSoftPath().TryLoad();
+		if (!TestTrue(TEXT("Actor Handle Loaded"), !!ActorHandle->GetActor()))
 		{
 			return false;
 		}
 
-		FWorldPartitionHandle ActorHandle(ActorDescMainContainer, FGuid(TEXT("0D2B04D240BE5DE58FE437A8D2DBF5C9")));
-		if (!TestTrue(TEXT("Invalid Actor Handle"), ActorHandle.IsValid()))
+		if (!TestTrue(TEXT("Resolving Runtime Actor From Editor Path Failed"), ResolvedObject == ActorHandle->GetActor()))
 		{
 			return false;
 		}
 
-		if (!TestTrue(TEXT("Actor Handle Not Loaded"), !ActorHandle->GetActor()))
+		auto TestEditorRuntimePathConversion = [this](const FSoftObjectPath& InEditorPath) -> bool
+		{
+			FSoftObjectPath RuntimePath;
+			if (!TestTrue(TEXT("Path Editor to Runtime Conversion Failed"), FWorldPartitionHelpers::ConvertEditorPathToRuntimePath(InEditorPath, RuntimePath)))
+			{
+				return false;
+			}
+
+			FSoftObjectPath EditorPath;
+			if (!TestTrue(TEXT("Path Runtime to Editor Conversion Failed"), FWorldPartitionHelpers::ConvertRuntimePathToEditorPath(RuntimePath, EditorPath)))
+			{
+				return false;
+			}
+
+			if (!TestTrue(TEXT("Path Editor to Runtime to Editor Conversion Failed"), EditorPath == InEditorPath))
+			{
+				return false;
+			}
+
+			return true;
+		};
+
+		// Test path conversions to an actor
+		if (!TestEditorRuntimePathConversion(FSoftObjectPath(Actor)))
 		{
 			return false;
 		}
 
-		UObject* ResolvedObject = ActorHandle->GetActorSoftPath().TryLoad();
-		if (!TestTrue(TEXT("Actor Handle Loaded"), !!ActorHandle->GetActor()))
+		// Test path conversions to an actor component
+		if (!TestEditorRuntimePathConversion(FSoftObjectPath(Actor->GetRootComponent())))
 		{
 			return false;
 		}
 
-		if (!TestTrue(TEXT("Resolving Runtime Actor From Editor Path Failed"), ResolvedObject == ActorHandle->GetActor()))
+		// Test path conversions to an instanced level
+		UPackage* WorldPackage = World->GetPackage();
+		WorldPackage->Rename(*FString(WorldPackage->GetName() + TEXT("_LevelInstance_1")));
+
+		// Test path conversions to an instanced actor
+		if (!TestEditorRuntimePathConversion(FSoftObjectPath(Actor)))
+		{
+			return false;
+		}
+
+		// Test path conversions to an instanced actor component
+		if (!TestEditorRuntimePathConversion(FSoftObjectPath(Actor->GetRootComponent())))
+		{
+			return false;
+		}
+
+		// Test path conversions to a newly spawned actor
+		AActor* NewActor = World->SpawnActor(AActor::StaticClass());
+		if (!TestTrue(TEXT("Sapwning Actor Failed"), !!NewActor))
+		{
+			return false;
+		}
+
+		if (!TestEditorRuntimePathConversion(FSoftObjectPath(NewActor)))
 		{
 			return false;
 		}
diff --git a/Engine/Source/Runtime/Engine/Private/WorldPartition/WorldPartitionHelpers.cpp b/Engine/Source/Runtime/Engine/Private/WorldPartition/WorldPartitionHelpers.cpp
--- a/Engine/Source/Runtime/Engine/Private/WorldPartition/WorldPartitionHelpers.cpp
+++ b/Engine/Source/Runtime/Engine/Private/WorldPartition/WorldPartitionHelpers.cpp
@@ -243,6 +243,16 @@ bool FWorldPartitionHelpers::ConvertRuntimePathToEditorPath(const FSoftObjectPat
 			if (int32 NextColonPos = InPathView.Find(TEXTVIEW(":"), NextDotPos); NextColonPos != INDEX_NONE)
 			{
 				OutPathString.RemoveAt(GeneratedPos, NextDotPos - GeneratedPos);
+
+				// In the editor, the _LevelInstance_ID is appended to the persistent level, while at runtime it is appended to each cell package, so we need to remap it there if present.
+				FString WorldAssetPackageName = InPath.GetAssetPath().GetPackageName().ToString();
+				const int32 LevelInstancePos = WorldAssetPackageName.Find(TEXT("_LevelInstance_"), ESearchCase::IgnoreCase, ESearchDir::FromEnd);
+				if (LevelInstancePos != INDEX_NONE)
+				{
+					FString LevelInstanceTag = WorldAssetPackageName.RightChop(LevelInstancePos);
+					OutPathString.InsertAt(GeneratedPos, LevelInstanceTag);
+				}
+
 				OutPath = OutPathString;
 				return true;
 			}
diff --git a/Engine/Source/Runtime/Engine/Private/WorldPartition/WorldPartitionLevelStreamingPolicy.cpp b/Engine/Source/Runtime/Engine/Private/WorldPartition/WorldPartitionLevelStreamingPolicy.cpp
--- a/Engine/Source/Runtime/Engine/Private/WorldPartition/WorldPartitionLevelStreamingPolicy.cpp
+++ b/Engine/Source/Runtime/Engine/Private/WorldPartition/WorldPartitionLevelStreamingPolicy.cpp
@@ -180,15 +180,51 @@ void UWorldPartitionLevelStreamingPolicy::RemapSoftObjectPath(FSoftObjectPath& O
 
 bool UWorldPartitionLevelStreamingPolicy::ConvertEditorPathToRuntimePath(const FSoftObjectPath& InPath, FSoftObjectPath& OutPath) const
 {
-	if (const FName* CellName = ActorToCellRemapping.Find(*InPath.ToString()))
+	const UWorld* OuterWorld = WorldPartition->GetTypedOuter<UWorld>();
+	const UPackage* OuterWorldPackage = OuterWorld->GetPackage();
+	const FTopLevelAssetPath WorldAssetPath(OuterWorld);
+
+	if (!OuterWorldPackage->HasAnyPackageFlags(PKG_PlayInEditor))
 	{
-		const UWorld* OuterWorld = WorldPartition->GetTypedOuter<UWorld>();
-		const UPackage* OuterWorldPackage = OuterWorld->GetPackage();
-		const FString OuterWorldPackageName = FPackageName::GetShortName(OuterWorldPackage->GetName());
-		const FString OuterWorldPackagePath = FPackageName::GetLongPackagePath(OuterWorldPackage->GetPathName());
-		OutPath = FString::Printf(TEXT("%s/%s/_Generated_/%s.%s:%s"), *OuterWorldPackagePath, *OuterWorldPackageName, *(CellName->ToString()), *OuterWorldPackageName, *InPath.GetSubPathString());
-		return true;
+		if (InPath.GetAssetPath() == WorldAssetPath)
+		{
+			FString SubAssetName;
+			FString SubAssetContext;
+			if (InPath.GetSubPathString().Split(TEXT("."), &SubAssetContext, &SubAssetName))
+			{
+				if (SubAssetContext == TEXT("PersistentLevel"))
+				{
+					FString SubObjectName;
+					FString SubObjectContext(SubAssetName);
+					SubAssetName.Split(TEXT("."), &SubObjectContext, &SubObjectName);
+
+					FString LevelInstanceTag;
+					FString WorldAssetPackageName = *WorldAssetPath.GetPackageName().ToString();
+					FString WorldAssetName = WorldAssetPath.GetAssetName().ToString();
+
+					// In the editor, the _LevelInstance_ID is appended to the persistent level, while at runtime it is appended to each cell package, so we need to remap it there if present.
+					const int32 LevelInstancePos = WorldAssetPackageName.Find(TEXT("_LevelInstance_"), ESearchCase::IgnoreCase, ESearchDir::FromEnd);
+					if (LevelInstancePos != INDEX_NONE)
+					{
+						LevelInstanceTag = WorldAssetPackageName.RightChop(LevelInstancePos);
+						WorldAssetPackageName.LeftInline(LevelInstancePos);
+					}
+
+					// Try to find the corresponding streaming cell, if it doesn't exists the actor must be in the persistent level.
+					if (const FName* CellName = SubObjectsToCellRemapping.Find(*SubObjectContext))
+					{
+						OutPath = FString::Printf(TEXT("%s/_Generated_/%s%s.%s:%s"), *WorldAssetPackageName, *(CellName->ToString()), *LevelInstanceTag, *WorldAssetPath.GetAssetName().ToString(), *InPath.GetSubPathString());
+					}
+					else
+					{
+						OutPath = FString::Printf(TEXT("%s/_Generated_/%s%s.%s:%s"), *WorldAssetPackageName, *WorldAssetName, *LevelInstanceTag, *WorldAssetPath.GetAssetName().ToString(), *InPath.GetSubPathString());
+					}
+					return true;
+				}
+			}
+		}
 	}
+
 	return false;
 }
 #endif
