Minor change: renamed the category of user defined struct attribute from custom attribute to animation attribute.

diff --git a/Engine/Source/Runtime/Engine/Classes/Animation/AnimationSettings.h b/Engine/Source/Runtime/Engine/Classes/Animation/AnimationSettings.h
--- a/Engine/Source/Runtime/Engine/Classes/Animation/AnimationSettings.h
+++ b/Engine/Source/Runtime/Engine/Classes/Animation/AnimationSettings.h
@@ -94,7 +94,7 @@ class ENGINE_API UAnimationSettings : public UDeveloperSettings
 	TArray<FString> TransformAttributeNames;
 
 	/** Register user defined structs as animation attributes*/
-	UPROPERTY(config, EditAnywhere, DisplayName="User Defined Struct Animation Attributes (Runtime only, Non-blendable)", Category = CustomAttributes, meta=(AllowedClasses="/Script/Engine.UserDefinedStruct"))
+	UPROPERTY(config, EditAnywhere, DisplayName="User Defined Struct Animation Attributes (Runtime only, Non-blendable)", Category = AnimationAttributes, meta=(AllowedClasses="/Script/Engine.UserDefinedStruct"))
 	TArray<TSoftObjectPtr<UUserDefinedStruct>> UserDefinedStructAttributes;
 
 	/** Find and Replace Expressions used for mirroring  */
