Fixed compile errors due to IWYU removed includes.

#pregflight none

diff --git a/Engine/Source/Runtime/Engine/Private/Net/Experimental/Iris/DataStreamChannel.cpp b/Engine/Source/Runtime/Engine/Private/Net/Experimental/Iris/DataStreamChannel.cpp
--- a/Engine/Source/Runtime/Engine/Private/Net/Experimental/Iris/DataStreamChannel.cpp
+++ b/Engine/Source/Runtime/Engine/Private/Net/Experimental/Iris/DataStreamChannel.cpp
@@ -7,6 +7,7 @@
 #include UE_INLINE_GENERATED_CPP_BY_NAME(DataStreamChannel)
 
 #if UE_WITH_IRIS
+// IWYU pragma: begin_keep
 #include "Engine/NetConnection.h"
 #include "Iris/IrisConfig.h"
 #include "Iris/DataStream/DataStreamManager.h"
@@ -20,6 +21,7 @@
 #include "Net/Core/Trace/NetTrace.h"
 #include "PacketHandler.h"
 #include "ProfilingDebugging/CsvProfiler.h"
+// IWYU pragma: end_keep
 #endif // UE_WITH_IRIS
 
 namespace UE::Net::Private
