Only allow the option to enable/disable VT feedback on a RVT sample when we are using fixed mip level.
This is because we usually only disable VT feedback to save on the mip level calculation when some other system is already generating feedback. (For example for grass VT sampling on landscape).
Also change behviour so that mip level and mip bias are used even if there is no connnection on the MipValue pin. Instead the level or bias defaults to 0.

diff --git a/Engine/Source/Runtime/Engine/Classes/Materials/MaterialExpressionRuntimeVirtualTextureSample.h b/Engine/Source/Runtime/Engine/Classes/Materials/MaterialExpressionRuntimeVirtualTextureSample.h
--- a/Engine/Source/Runtime/Engine/Classes/Materials/MaterialExpressionRuntimeVirtualTextureSample.h
+++ b/Engine/Source/Runtime/Engine/Classes/Materials/MaterialExpressionRuntimeVirtualTextureSample.h
@@ -90,7 +90,11 @@ class ENGINE_API UMaterialExpressionRuntimeVirtualTextureSample : public UMateri
 	UPROPERTY(EditAnywhere, Category = VirtualTexture, meta = (DisplayName = "Enable adaptive page table"))
 	bool bAdaptive = false;
 
-	/** Enable virtual texture feedback. */
+	/** 
+	 * Enable virtual texture feedback. 
+	 * Disabling this can result in the virtual texture not reaching the correct mip level. 
+	 * It should only be used in cases where we don't care about the correct mip level being resident, or some other process is maintaining the correct level.
+	 */
 	UPROPERTY(EditAnywhere, Category = VirtualTexture)
 	bool bEnableFeedback = true;
 
@@ -117,6 +121,7 @@ protected:
 	virtual void PostLoad() override;
 	virtual int32 Compile(class FMaterialCompiler* Compiler, int32 OutputIndex) override;
 	virtual void GetCaption(TArray<FString>& OutCaptions) const override;
+	virtual bool CanEditChange(const FProperty* InProperty) const override;
 public:
 	virtual void PostEditChangeProperty(FPropertyChangedEvent& PropertyChangedEvent) override;
 #endif
diff --git a/Engine/Source/Runtime/Engine/Private/Materials/MaterialExpressions.cpp b/Engine/Source/Runtime/Engine/Private/Materials/MaterialExpressions.cpp
--- a/Engine/Source/Runtime/Engine/Private/Materials/MaterialExpressions.cpp
+++ b/Engine/Source/Runtime/Engine/Private/Materials/MaterialExpressions.cpp
@@ -3035,6 +3035,24 @@ void UMaterialExpressionRuntimeVirtualTextureSample::PostLoad()
 	InitOutputs();
 }
 
+bool UMaterialExpressionRuntimeVirtualTextureSample::CanEditChange(const FProperty* InProperty) const
+{
+	bool bIsEditable = Super::CanEditChange(InProperty);
+	if (bIsEditable && InProperty != nullptr)
+	{
+		const FString PropertyName = InProperty->GetName();
+
+		if (PropertyName == GET_MEMBER_NAME_STRING_CHECKED(UMaterialExpressionRuntimeVirtualTextureSample, bEnableFeedback))
+		{
+			// We can support disabling feedback for MipLevel mode.
+			// We could allow for other modes too, but it's not a good idea to freely expose this option since it makes it easy could break things by accicent.
+			// Instead the user has to explicitly set the mip level mode before disabling feedback.
+			bIsEditable &= MipValueMode == RVTMVM_MipLevel;
+		}
+	}
+	return bIsEditable;
+}
+
 void UMaterialExpressionRuntimeVirtualTextureSample::PostEditChangeProperty(FPropertyChangedEvent& PropertyChangedEvent)
 {
 	// Update MaterialType setting to match VirtualTexture
@@ -3283,15 +3301,15 @@ int32 UMaterialExpressionRuntimeVirtualTextureSample::Compile(class FMaterialCom
 	int32 MipValue0Index = INDEX_NONE;
 	int32 MipValue1Index = INDEX_NONE;
 	const bool bMipValueExpressionValid = MipValue.GetTracedInput().Expression != nullptr;
-	if (MipValueMode == RVTMVM_MipLevel && bMipValueExpressionValid)
+	if (MipValueMode == RVTMVM_MipLevel)
 	{
 		TextureMipLevelMode = TMVM_MipLevel;
-		MipValue0Index = MipValue.Compile(Compiler);
+		MipValue0Index = bMipValueExpressionValid ? MipValue.Compile(Compiler) : Compiler->Constant(0);
 	}
-	else if (MipValueMode == RVTMVM_MipBias && bMipValueExpressionValid)
+	else if (MipValueMode == RVTMVM_MipBias)
 	{
 		TextureMipLevelMode = TMVM_MipBias;
-		MipValue0Index = MipValue.Compile(Compiler);
+		MipValue0Index = bMipValueExpressionValid ? MipValue.Compile(Compiler) : Compiler->Constant(0);
 	}
 	else if (MipValueMode == RVTMVM_RecalculateDerivatives)
 	{
@@ -3321,6 +3339,9 @@ int32 UMaterialExpressionRuntimeVirtualTextureSample::Compile(class FMaterialCom
 		break;
 	}
 
+	// We can support disabling feedback for MipLevel mode.
+	const bool bForceEnableFeedback = TextureMipLevelMode != TMVM_MipLevel;
+
 	// Compile the texture sample code
 	const bool bAutomaticMipViewBias = true;
 	int32 SampleCodeIndex[RuntimeVirtualTexture::MaxTextureLayers] = { INDEX_NONE };
@@ -3332,7 +3353,7 @@ int32 UMaterialExpressionRuntimeVirtualTextureSample::Compile(class FMaterialCom
 			SAMPLERTYPE_VirtualMasks,
 			MipValue0Index, MipValue1Index, TextureMipLevelMode, SamplerSourceMode,
 			TextureReferenceIndex[TexureLayerIndex],
-			bAutomaticMipViewBias, bAdaptive, bEnableFeedback);
+			bAutomaticMipViewBias, bAdaptive, (bEnableFeedback || bForceEnableFeedback));
 	}
 
 	// Compile any unpacking code
