Fix bug where we don't use the ray tracing specific skin cache entry.

diff --git a/Engine/Source/Runtime/Engine/Private/SkeletalRenderGPUSkin.cpp b/Engine/Source/Runtime/Engine/Private/SkeletalRenderGPUSkin.cpp
--- a/Engine/Source/Runtime/Engine/Private/SkeletalRenderGPUSkin.cpp
+++ b/Engine/Source/Runtime/Engine/Private/SkeletalRenderGPUSkin.cpp
@@ -1549,7 +1549,7 @@ const FSkinBatchVertexFactoryUserData* FSkeletalMeshObjectGPUSkin::GetVertexFact
 #endif
 	if (Entry != nullptr)
 	{
-		return FGPUSkinCache::GetVertexFactoryUserData(SkinCacheEntry, ChunkIdx);
+		return FGPUSkinCache::GetVertexFactoryUserData(Entry, ChunkIdx);
 	}
 
 	return &LODs[LODIndex].DeformerGeometry.VertexFactoryUserData;
