Move a function out of header.
It is there to inline the Entry pointer check and early out. But in the only calling location we check the Entry pointer before calling anyway.

diff --git a/Engine/Source/Runtime/Engine/Private/GPUSkinCache.cpp b/Engine/Source/Runtime/Engine/Private/GPUSkinCache.cpp
--- a/Engine/Source/Runtime/Engine/Private/GPUSkinCache.cpp
+++ b/Engine/Source/Runtime/Engine/Private/GPUSkinCache.cpp
@@ -1994,9 +1994,9 @@ bool FGPUSkinCache::IsEntryValid(FGPUSkinCacheEntry* SkinCacheEntry, int32 Secti
 	return SkinCacheEntry->IsSectionValid(Section);
 }
 
-const FSkinBatchVertexFactoryUserData* FGPUSkinCache::InternalGetVertexFactoryUserData(FGPUSkinCacheEntry* Entry, int32 Section)
+const FSkinBatchVertexFactoryUserData* FGPUSkinCache::GetVertexFactoryUserData(FGPUSkinCacheEntry* Entry, int32 Section)
 {
-	return &Entry->BatchElementsUserData[Section];
+	return Entry != nullptr ? &Entry->BatchElementsUserData[Section] : nullptr;
 }
 
 void FGPUSkinCache::InvalidateAllEntries()
diff --git a/Engine/Source/Runtime/Engine/Public/GPUSkinCache.h b/Engine/Source/Runtime/Engine/Public/GPUSkinCache.h
--- a/Engine/Source/Runtime/Engine/Public/GPUSkinCache.h
+++ b/Engine/Source/Runtime/Engine/Public/GPUSkinCache.h
@@ -153,14 +153,7 @@ public:
 
 	static void Release(FGPUSkinCacheEntry*& SkinCacheEntry);
 
-	static inline const FSkinBatchVertexFactoryUserData* GetVertexFactoryUserData(FGPUSkinCacheEntry* Entry, int32 Section)
-	{
-		if (Entry)
-		{
-			return InternalGetVertexFactoryUserData(Entry, Section);
-		}
-		return nullptr;
-	}
+	static const FSkinBatchVertexFactoryUserData* GetVertexFactoryUserData(FGPUSkinCacheEntry* Entry, int32 Section);
 
 	static bool IsEntryValid(FGPUSkinCacheEntry* SkinCacheEntry, int32 Section);
 	static FColor GetVisualizationDebugColor(const FName& GPUSkinCacheVisualizationMode, FGPUSkinCacheEntry* Entry, FGPUSkinCacheEntry* RayTracingEntry, uint32 SectionIndex);
@@ -434,7 +427,6 @@ protected:
 	void Cleanup();
 	static void TransitionAllToReadable(FRHICommandList& RHICmdList, const TSet<FSkinCacheRWBuffer*>& BuffersToTransitionToRead);
 	static void ReleaseSkinCacheEntry(FGPUSkinCacheEntry* SkinCacheEntry);
-	static const FSkinBatchVertexFactoryUserData* InternalGetVertexFactoryUserData(FGPUSkinCacheEntry* Entry, int32 Section);
 	void InvalidateAllEntries();
 	uint64 UsedMemoryInBytes;
 	uint64 ExtraRequiredMemory;
