Fix up a couple of places where we don't need to include GPUSkinCache.h
Fix a comment.

diff --git a/Engine/Source/Runtime/Engine/Private/GPUSkinVertexFactory.cpp b/Engine/Source/Runtime/Engine/Private/GPUSkinVertexFactory.cpp
--- a/Engine/Source/Runtime/Engine/Private/GPUSkinVertexFactory.cpp
+++ b/Engine/Source/Runtime/Engine/Private/GPUSkinVertexFactory.cpp
@@ -1360,7 +1360,9 @@ void FGPUSkinPassthroughVertexFactory::BuildStreamIndices()
 
 		FVertexBuffer const* BufferToFind = &DummyVBs[AttributeIndex];
 		
-		// Each TexCoord streams each containing two attributes, so map the second attribute to the first one when searching.
+		// Each TexCoord stream can contain two attributes. 
+		// We always create stream with the dummy VB of the first attribute.
+		// Map the second attribute to the first one when searching.
 		if (AttributeIndex == EVertexAtttribute::VertexTexCoord1 || AttributeIndex == EVertexAtttribute::VertexTexCoord3 ||
 			AttributeIndex == EVertexAtttribute::VertexTexCoord5 || AttributeIndex == EVertexAtttribute::VertexTexCoord7)
 		{
diff --git a/Engine/Source/Runtime/Engine/Private/SkeletalMesh.cpp b/Engine/Source/Runtime/Engine/Private/SkeletalMesh.cpp
--- a/Engine/Source/Runtime/Engine/Private/SkeletalMesh.cpp
+++ b/Engine/Source/Runtime/Engine/Private/SkeletalMesh.cpp
@@ -7,6 +7,7 @@
 #include "Engine/SkeletalMesh.h"
 
 #include "Animation/MeshDeformerGeometry.h"
+#include "DataDrivenShaderPlatformInfo.h"
 #include "Engine/Blueprint.h"
 #include "Engine/SkeletalMeshSampling.h"
 #include "EngineStats.h"
@@ -42,7 +43,6 @@
 #include "PhysicsEngine/PhysicsAsset.h"
 #include "Engine/AssetUserData.h"
 #include "Animation/NodeMappingContainer.h"
-#include "GPUSkinCache.h"
 #include "Rendering/SkeletalMeshRenderData.h"
 #include "AnimationRuntime.h"
 #include "Animation/AnimSequence.h"
diff --git a/Engine/Source/Runtime/Renderer/Private/SceneRendering.cpp b/Engine/Source/Runtime/Renderer/Private/SceneRendering.cpp
--- a/Engine/Source/Runtime/Renderer/Private/SceneRendering.cpp
+++ b/Engine/Source/Runtime/Renderer/Private/SceneRendering.cpp
@@ -3275,7 +3275,6 @@ void FSceneRenderer::RenderFinish(FRDGBuilder& GraphBuilder, FRDGTextureRef View
 		static auto* CVarSkinCacheOOM = IConsoleManager::Get().FindTConsoleVariableDataFloat(TEXT("r.SkinCache.SceneMemoryLimitInMB"));
 
 		uint64 GPUSkinCacheExtraRequiredMemory = 0;
-		extern ENGINE_API bool IsGPUSkinCacheAvailable(EShaderPlatform Platform);
 		if (IsGPUSkinCacheAvailable(ShaderPlatform))
 		{
 			if (FGPUSkinCache* SkinCache = Scene->GetGPUSkinCache())
