In some places where we use IsGPUSkinCacheAvailable() to determine if a buffer requires an SRV, replace with IsGPUSkinPassThroughSupported().
This covers both SkinCache and MeshDeformer cases.

diff --git a/Engine/Plugins/FX/Niagara/Source/Niagara/Private/NiagaraCommon.cpp b/Engine/Plugins/FX/Niagara/Source/Niagara/Private/NiagaraCommon.cpp
--- a/Engine/Plugins/FX/Niagara/Source/Niagara/Private/NiagaraCommon.cpp
+++ b/Engine/Plugins/FX/Niagara/Source/Niagara/Private/NiagaraCommon.cpp
@@ -2,6 +2,8 @@
 
 #include "NiagaraCommon.h"
 
+#include "DataDrivenShaderPlatformInfo.h"
+#include "LocalVertexFactory.h"
 #include "Misc/StringBuilder.h"
 #include "NiagaraComponent.h"
 #include "NiagaraConstants.h"
@@ -15,7 +17,6 @@
 #include "NiagaraWorldManager.h"
 #include "String/ParseTokens.h"
 #include "UObject/Class.h"
-#include "GPUSkinCache.h"
 
 DECLARE_CYCLE_STAT(TEXT("Niagara - Utilities - PrepareRapidIterationParameters"), STAT_Niagara_Utilities_PrepareRapidIterationParameters, STATGROUP_Niagara);
 
@@ -921,7 +922,7 @@ bool FNiagaraUtilities::AllowGPUCulling(EShaderPlatform ShaderPlatform)
 
 bool FNiagaraUtilities::AreBufferSRVsAlwaysCreated(EShaderPlatform ShaderPlatform)
 {
-	return RHISupportsManualVertexFetch(ShaderPlatform) || IsGPUSkinCacheAvailable(ShaderPlatform);
+	return RHISupportsManualVertexFetch(ShaderPlatform) || FLocalVertexFactory::IsGPUSkinPassThroughSupported(ShaderPlatform);
 }
 
 ENiagaraCompileUsageStaticSwitch FNiagaraUtilities::ConvertScriptUsageToStaticSwitchUsage(ENiagaraScriptUsage ScriptUsage)
diff --git a/Engine/Source/Editor/Persona/Private/PersonaMeshDetails.cpp b/Engine/Source/Editor/Persona/Private/PersonaMeshDetails.cpp
--- a/Engine/Source/Editor/Persona/Private/PersonaMeshDetails.cpp
+++ b/Engine/Source/Editor/Persona/Private/PersonaMeshDetails.cpp
@@ -5164,8 +5164,6 @@ TSharedRef<SWidget> FPersonaMeshDetails::OnGenerateCustomNameWidgetsForSection(i
 
 TSharedRef<SWidget> FPersonaMeshDetails::OnGenerateCustomSectionWidgetsForSection(int32 LODIndex, int32 SectionIndex)
 {
-	extern ENGINE_API bool IsGPUSkinCacheAvailable(EShaderPlatform Platform);
-
 	TSharedRef<SVerticalBox> SectionWidget = SNew(SVerticalBox);
 	
 	//If we have a chunk section, prevent editing of cloth cast shadow and recompute tangent
@@ -5241,7 +5239,6 @@ TSharedRef<SWidget> FPersonaMeshDetails::OnGenerateCustomSectionWidgetsForSectio
 		.Padding(5, 2, 0, 0)
 		[
 			SNew(SComboButton)
-			.IsEnabled(IsGPUSkinCacheAvailable(GMaxRHIShaderPlatform))
 			.OnGetMenuContent(this, &FPersonaMeshDetails::OnGenerateRecomputeTangentsSetting, LODIndex, SectionIndex)
 			.VAlign(VAlign_Center)
 			.ContentPadding(2.f)
diff --git a/Engine/Source/Runtime/Engine/Private/GPUSkinCache.cpp b/Engine/Source/Runtime/Engine/Private/GPUSkinCache.cpp
--- a/Engine/Source/Runtime/Engine/Private/GPUSkinCache.cpp
+++ b/Engine/Source/Runtime/Engine/Private/GPUSkinCache.cpp
@@ -263,16 +263,6 @@ ENGINE_API bool GPUSkinCacheNeedsDuplicatedVertices()
 #endif
 }
 
-// We don't have it always enabled as it's not clear if this has a performance cost
-// Call on render thread only!
-// Should only be called if SM5 (compute shaders, atomics) are supported.
-ENGINE_API bool DoSkeletalMeshIndexBuffersNeedSRV()
-{
-	// currently only implemented and tested on Window SM5 (needs Compute, Atomics, SRV for index buffers, UAV for VertexBuffers)
-	//#todo-gpuskin: Enable on PS4 when SRVs for IB exist
-	return IsGPUSkinCacheAvailable(GMaxRHIShaderPlatform);
-}
-
 UE_DEPRECATED(5.0, "This function is no longer in use and will be removed.")
 ENGINE_API bool DoRecomputeSkinTangentsOnGPU_RT()
 {
diff --git a/Engine/Source/Runtime/Engine/Private/InstancedStaticMesh.cpp b/Engine/Source/Runtime/Engine/Private/InstancedStaticMesh.cpp
--- a/Engine/Source/Runtime/Engine/Private/InstancedStaticMesh.cpp
+++ b/Engine/Source/Runtime/Engine/Private/InstancedStaticMesh.cpp
@@ -956,10 +956,9 @@ void FInstancedStaticMeshVertexFactory::InitRHI()
 			}
 		}
 
-		// PreSkinPosition is only used when skin cache is on, in which case skinned mesh uses local vertex factory as pass through.
+		// PreSkinPosition attribute is only used for GPUSkinPassthrough variation of local vertex factory.
 		// It is not used by ISM so fill with dummy buffer.
-		extern ENGINE_API bool IsGPUSkinCacheAvailable(EShaderPlatform Platform);
-		if (IsGPUSkinCacheAvailable(GMaxRHIShaderPlatform))
+		if (FLocalVertexFactory::IsGPUSkinPassThroughSupported(GMaxRHIShaderPlatform))
 		{
 			FVertexStreamComponent NullComponent(&GNullVertexBuffer, 0, 0, VET_Float4);
 			Elements.Add(AccessStreamComponent(NullComponent, 14));
diff --git a/Engine/Source/Runtime/Engine/Private/LocalVertexFactory.cpp b/Engine/Source/Runtime/Engine/Private/LocalVertexFactory.cpp
--- a/Engine/Source/Runtime/Engine/Private/LocalVertexFactory.cpp
+++ b/Engine/Source/Runtime/Engine/Private/LocalVertexFactory.cpp
@@ -172,7 +172,7 @@ void FLocalVertexFactoryShaderParameters::Bind(const FShaderParameterMap& Parame
 	IsGPUSkinPassThrough.Bind(ParameterMap, TEXT("bIsGPUSkinPassThrough"));
 }
 
-static bool IsGPUSkinPassThroughSupported(EShaderPlatform Platform)
+bool FLocalVertexFactory::IsGPUSkinPassThroughSupported(EShaderPlatform Platform)
 {
 	// Enable the GPUSkin passthrough path if we might use the GPUSkinCache or MeshDeformers.
 	static IMeshDeformerProvider* MeshDeformerProvider = IMeshDeformerProvider::Get();
diff --git a/Engine/Source/Runtime/Engine/Private/PositionVertexBuffer.cpp b/Engine/Source/Runtime/Engine/Private/PositionVertexBuffer.cpp
--- a/Engine/Source/Runtime/Engine/Private/PositionVertexBuffer.cpp
+++ b/Engine/Source/Runtime/Engine/Private/PositionVertexBuffer.cpp
@@ -3,10 +3,10 @@
 #include "Rendering/PositionVertexBuffer.h"
 
 #include "Components.h"
-
-#include "StaticMeshVertexData.h"
-#include "GPUSkinCache.h"
+#include "DataDrivenShaderPlatformInfo.h"
+#include "LocalVertexFactory.h"
 #include "RHIResourceUpdates.h"
+#include "StaticMeshVertexData.h"
 
 /*-----------------------------------------------------------------------------
 FPositionVertexBuffer
@@ -265,7 +265,7 @@ void FPositionVertexBuffer::InitRHI()
 	if (VertexBufferRHI)
 	{
 		// we have decide to create the SRV based on GMaxRHIShaderPlatform because this is created once and shared between feature levels for editor preview.
-		bool bSRV = RHISupportsManualVertexFetch(GMaxRHIShaderPlatform) || IsGPUSkinCacheAvailable(GMaxRHIShaderPlatform);
+		bool bSRV = RHISupportsManualVertexFetch(GMaxRHIShaderPlatform) || FLocalVertexFactory::IsGPUSkinPassThroughSupported(GMaxRHIShaderPlatform);
 
 		// When bAllowCPUAccess is true, the meshes is likely going to be used for Niagara to spawn particles on mesh surface.
 		// And it can be the case for CPU *and* GPU access: no differenciation today. That is why we create a SRV in this case.
diff --git a/Engine/Source/Runtime/Engine/Private/RawIndexBuffer.cpp b/Engine/Source/Runtime/Engine/Private/RawIndexBuffer.cpp
--- a/Engine/Source/Runtime/Engine/Private/RawIndexBuffer.cpp
+++ b/Engine/Source/Runtime/Engine/Private/RawIndexBuffer.cpp
@@ -5,10 +5,12 @@
 =============================================================================*/
 
 #include "RawIndexBuffer.h"
+
+#include "DataDrivenShaderPlatformInfo.h"
 #include "Interfaces/ITargetPlatform.h"
+#include "LocalVertexFactory.h"
 #include "Modules/ModuleManager.h"
 #include "RenderingThread.h"
-#include "DataDrivenShaderPlatformInfo.h"
 #include "RHIResourceUpdates.h"
 
 #if WITH_EDITOR
@@ -326,8 +328,8 @@ FBufferRHIRef FRawStaticIndexBuffer::CreateRHIBuffer_Internal()
 
 	if (GetNumIndices() > 0)
 	{
-		extern ENGINE_API bool DoSkeletalMeshIndexBuffersNeedSRV();
-		bool bSRV = RHISupportsManualVertexFetch(GMaxRHIShaderPlatform) || DoSkeletalMeshIndexBuffersNeedSRV();
+		// Systems that generate data for GPUSkinPassThrough use index buffer as SRV.
+		bool bSRV = RHISupportsManualVertexFetch(GMaxRHIShaderPlatform) || FLocalVertexFactory::IsGPUSkinPassThroughSupported(GMaxRHIShaderPlatform);
 
 		// When bAllowCPUAccess is true, the meshes is likely going to be used for Niagara to spawn particles on mesh surface.
 		// And it can be the case for CPU *and* GPU access: no differenciation today. That is why we create a SRV in this case.
@@ -461,6 +463,17 @@ void FRawStaticIndexBuffer::Discard()
     IndexStorage.Discard();
 }
 
+bool FRawStaticIndexBuffer16or32Interface::IsSRVNeeded(bool bAllowCPUAccess) const
+{
+	// Systems that generate data for GPUSkinPassThrough use index buffer as SRV.
+	bool bSRV = RHISupportsManualVertexFetch(GMaxRHIShaderPlatform) || FLocalVertexFactory::IsGPUSkinPassThroughSupported(GMaxRHIShaderPlatform);
+	// When bAllowCPUAccess is true, the meshes is likely going to be used for Niagara to spawn particles on mesh surface.
+	// And it can be the case for CPU *and* GPU access: no differenciation today. That is why we create a SRV in this case.
+	// This also avoid setting lots of states on all the members of all the different buffers used by meshes. Follow up: https://jira.it.epicgames.net/browse/UE-69376.
+	bSRV |= bAllowCPUAccess;
+	return bSRV;
+}
+
 void FRawStaticIndexBuffer16or32Interface::InitRHIForStreaming(FRHIBuffer* IntermediateBuffer, size_t IndexSize, FRHIResourceUpdateBatcher& Batcher)
 {
 	if (IndexBufferRHI && IntermediateBuffer)
diff --git a/Engine/Source/Runtime/Engine/Private/Rendering/StaticMeshVertexBuffer.cpp b/Engine/Source/Runtime/Engine/Private/Rendering/StaticMeshVertexBuffer.cpp
--- a/Engine/Source/Runtime/Engine/Private/Rendering/StaticMeshVertexBuffer.cpp
+++ b/Engine/Source/Runtime/Engine/Private/Rendering/StaticMeshVertexBuffer.cpp
@@ -1,9 +1,11 @@
 // Copyright Epic Games, Inc. All Rights Reserved.
 
 #include "Rendering/StaticMeshVertexBuffer.h"
-#include "EngineUtils.h"
+
 #include "Components.h"
-#include "GPUSkinCache.h"
+#include "DataDrivenShaderPlatformInfo.h"
+#include "EngineUtils.h"
+#include "LocalVertexFactory.h"
 #include "MeshUVChannelInfo.h"
 #include "ProfilingDebugging/LoadTimeTracker.h"
 #include "RHIResourceUpdates.h"
@@ -388,7 +390,7 @@ void FStaticMeshVertexBuffer::InitRHI()
 	const bool bHadTangentsData = TangentsData != nullptr;
 	const bool bCreateTangentsSRV = bHadTangentsData && TangentsData->GetAllowCPUAccess();
 	TangentsVertexBuffer.VertexBufferRHI = CreateTangentsRHIBuffer_RenderThread();
-	if (TangentsVertexBuffer.VertexBufferRHI && (bCreateTangentsSRV || RHISupportsManualVertexFetch(GMaxRHIShaderPlatform) || IsGPUSkinCacheAvailable(GMaxRHIShaderPlatform)))
+	if (TangentsVertexBuffer.VertexBufferRHI && (bCreateTangentsSRV || RHISupportsManualVertexFetch(GMaxRHIShaderPlatform) || FLocalVertexFactory::IsGPUSkinPassThroughSupported(GMaxRHIShaderPlatform)))
 	{
 		// When TangentsData is null, this buffer hasn't been streamed in yet. We still need to create a FRHIShaderResourceView which will be
 		// cached in a vertex factory uniform buffer later. The nullptr tells the RHI that the SRV doesn't view on anything yet.
diff --git a/Engine/Source/Runtime/Engine/Public/LocalVertexFactory.h b/Engine/Source/Runtime/Engine/Public/LocalVertexFactory.h
--- a/Engine/Source/Runtime/Engine/Public/LocalVertexFactory.h
+++ b/Engine/Source/Runtime/Engine/Public/LocalVertexFactory.h
@@ -77,6 +77,12 @@ public:
 
 	static void GetPSOPrecacheVertexFetchElements(EVertexInputStreamType VertexInputStreamType, FVertexDeclarationElementList& Elements);
 
+	/** 
+	 * Does the platform support GPUSkinPassthrough permutations.
+	 * This knowledge can be used to indicate if we need to create SRV for index/vertex buffers.
+	 */
+	static bool IsGPUSkinPassThroughSupported(EShaderPlatform Platform);
+
 	/**
 	 * An implementation of the interface used by TSynchronizedResource to update the resource with new data from the game thread.
 	 */
diff --git a/Engine/Source/Runtime/Engine/Public/RawIndexBuffer.h b/Engine/Source/Runtime/Engine/Public/RawIndexBuffer.h
--- a/Engine/Source/Runtime/Engine/Public/RawIndexBuffer.h
+++ b/Engine/Source/Runtime/Engine/Public/RawIndexBuffer.h
@@ -360,6 +360,8 @@ public:
 	}
 
 protected:
+	ENGINE_API bool IsSRVNeeded(bool bAllowCPUAccess) const;
+
 	/** Similar to Init/ReleaseRHI but only update existing SRV so references to the SRV stays valid */
 	void InitRHIForStreaming(FRHIBuffer* IntermediateBuffer, size_t IndexSize, FRHIResourceUpdateBatcher& Batcher);
 	void ReleaseRHIForStreaming(FRHIResourceUpdateBatcher& Batcher);
@@ -401,7 +403,7 @@ public:
 		const bool bHadIndexData = Num() > 0;
 		IndexBufferRHI = CreateRHIBuffer_RenderThread();
 
-		if (IndexBufferRHI && IsSRVNeeded() && bHadIndexData)
+		if (IndexBufferRHI && IsSRVNeeded(Indices.GetAllowCPUAccess()) && bHadIndexData)
 		{
 			// If the index buffer is a placeholder we still need to create a FRHIShaderResourceView.
 			SRVValue = RHICreateShaderResourceView(FShaderResourceViewInitializer(IndexBufferRHI->GetSize() > 0 ? IndexBufferRHI : nullptr, sizeof(INDEX_TYPE) == 2 ? PF_R16_UINT : PF_R32_UINT));
@@ -516,17 +518,6 @@ private:
 
 	int32 CachedNumIndices;
 
-	bool IsSRVNeeded() const
-	{
-		extern ENGINE_API bool DoSkeletalMeshIndexBuffersNeedSRV();
-		bool bSRV = DoSkeletalMeshIndexBuffersNeedSRV();
-		// When bAllowCPUAccess is true, the meshes is likely going to be used for Niagara to spawn particles on mesh surface.
-		// And it can be the case for CPU *and* GPU access: no differenciation today. That is why we create a SRV in this case.
-		// This also avoid setting lots of states on all the members of all the different buffers used by meshes. Follow up: https://jira.it.epicgames.net/browse/UE-69376.
-		bSRV |= Indices.GetAllowCPUAccess();
-		return bSRV;
-	}
-
 	template <bool bRenderThread>
 	FBufferRHIRef CreateRHIBuffer_Internal()
 	{
@@ -542,7 +533,7 @@ private:
 				Indices.Num(),
 				sizeof(INDEX_TYPE),
 				&Indices,
-				IsSRVNeeded(),
+				IsSRVNeeded(Indices.GetAllowCPUAccess()),
 				bRenderThread
 			);
 		}
