Fixed Out of Bounds accessing a static array at GPUSkinVertexFactory (was initializing the array with wrong size)

diff --git a/Engine/Source/Runtime/Engine/Private/GPUSkinVertexFactory.cpp b/Engine/Source/Runtime/Engine/Private/GPUSkinVertexFactory.cpp
--- a/Engine/Source/Runtime/Engine/Private/GPUSkinVertexFactory.cpp
+++ b/Engine/Source/Runtime/Engine/Private/GPUSkinVertexFactory.cpp
@@ -1274,6 +1274,9 @@ void FGPUSkinPassthroughVertexFactory::ResetVertexAttributes()
 	for (int32 Index = 0; Index < EVertexAtttribute::NumAttributes; ++Index)
 	{
 		StreamIndices[Index] = -1;
+	}
+	for (int32 Index = 0; Index < EShaderResource::NumShaderResources; ++Index)
+	{
 		SRVs[Index] = nullptr;
 	}
 	UpdatedFrameNumber = ~0U;
