Do not attempt to get the texture source mips for generating DDC for texture source with no bulkdata payload and emit a warning

diff --git a/Engine/Source/Runtime/Engine/Private/TextureDerivedDataTask.cpp b/Engine/Source/Runtime/Engine/Private/TextureDerivedDataTask.cpp
--- a/Engine/Source/Runtime/Engine/Private/TextureDerivedDataTask.cpp
+++ b/Engine/Source/Runtime/Engine/Private/TextureDerivedDataTask.cpp
@@ -8,7 +8,6 @@
 #include "IImageWrapperModule.h"
 #include "Misc/ScopedSlowTask.h"
 #include "TextureResource.h"
-#include "EngineLogs.h"
 #include "Engine/Texture2DArray.h"
 
 #if WITH_EDITOR
@@ -264,8 +263,13 @@ void FTextureSourceData::GetSourceMips(FTextureSource& Source, IImageWrapperModu
 	{
 		if (Source.HasHadBulkDataCleared())
 		{	// don't do any work we can't reload this
-			UE_LOG(LogTexture, Error, TEXT("Unable to get texture source mips because its bulk data was released. %s"), *TextureFullName)
-				return;
+			UE_LOG(LogTexture, Error, TEXT("Unable to get texture source mips because its bulk data was released. %s"), *TextureFullName);
+			return;
+		}
+		if (!Source.HasPayloadData())
+		{	// don't do any work we can't reload this
+			UE_LOG(LogTexture, Warning, TEXT("Unable to get texture source mips because its bulk data has no payload. This may happen if it was duplicated from cooked data. %s"), *TextureFullName);
+			return;
 		}
 
 		const FTextureSource::FMipData ScopedMipData = Source.GetMipData(InImageWrapper);
@@ -1957,6 +1961,12 @@ void FTextureCacheDerivedDataWorker::Finalize()
 	// --	if the texture compiler cvar disallows async texture compilation
 	if (!bSucceeded)
 	{
+		if (!TextureData.HasPayload())
+		{
+			UE_LOG(LogTexture, Warning, TEXT("Unable to build texture source data, no available payload for %s. This may happen if it was duplicated from cooked data."), *TexturePathName);
+			return;
+		}
+
 		TextureData.GetSourceMips(Texture.Source, ImageWrapper);
 		if (Texture.CompositeTexture)
 		{
diff --git a/Engine/Source/Runtime/Engine/Private/TextureDerivedDataTask.h b/Engine/Source/Runtime/Engine/Private/TextureDerivedDataTask.h
--- a/Engine/Source/Runtime/Engine/Private/TextureDerivedDataTask.h
+++ b/Engine/Source/Runtime/Engine/Private/TextureDerivedDataTask.h
@@ -91,6 +91,11 @@ struct FTextureSourceData
 
 	void Init(UTexture& InTexture, TextureMipGenSettings InMipGenSettings, bool bInCubeMap, bool bInTextureArray, bool bInVolumeTexture, bool bAllowAsyncLoading);
 	bool IsValid() const { return bValid; }
+
+	bool HasPayload() const
+	{
+		return AsyncSource.HasPayloadData();
+	}
 	
 	// ImageWrapperModule is not used
 	void GetSourceMips(FTextureSource& Source, IImageWrapperModule* InImageWrapper = nullptr);
