[SplineComponent] Enable better bounds calculations on SplineComponents in the event that they are required.

diff --git a/Engine/Source/Runtime/Engine/Private/Components/SplineComponent.cpp b/Engine/Source/Runtime/Engine/Private/Components/SplineComponent.cpp
--- a/Engine/Source/Runtime/Engine/Private/Components/SplineComponent.cpp
+++ b/Engine/Source/Runtime/Engine/Private/Components/SplineComponent.cpp
@@ -1856,14 +1856,10 @@ void USplineComponent::Draw(FPrimitiveDrawInterface* PDI, const FSceneView* View
 		OldKeyPos = NewKeyPos;
 	}
 }
+#endif
 
 FBoxSphereBounds USplineComponent::CalcBounds(const FTransform& LocalToWorld) const
 {
-	if (!bDrawDebug)
-	{
-		// Do as little as possible if not rendering anything
-		return Super::CalcBounds(LocalToWorld);
-	}
 
 #if SPLINE_FAST_BOUNDS_CALCULATION
 	FBox BoundingBox(0);
@@ -1909,8 +1905,6 @@ FBoxSphereBounds USplineComponent::CalcBounds(const FTransform& LocalToWorld) co
 #endif
 }
 
-#endif
-
 #if WITH_EDITOR
 bool USplineComponent::IgnoreBoundsForEditorFocus() const
 {
