Fix missing include for non unity compile

diff --git a/Engine/Source/Runtime/Engine/Private/TextureDerivedDataTask.cpp b/Engine/Source/Runtime/Engine/Private/TextureDerivedDataTask.cpp
--- a/Engine/Source/Runtime/Engine/Private/TextureDerivedDataTask.cpp
+++ b/Engine/Source/Runtime/Engine/Private/TextureDerivedDataTask.cpp
@@ -24,6 +24,7 @@
 #include "DerivedDataRequestOwner.h"
 #include "DerivedDataThreadPoolTask.h"
 #include "Engine/VolumeTexture.h"
+#include "EngineLogs.h"
 #include "ImageCoreUtils.h"
 #include "Interfaces/ITextureFormat.h"
 #include "Serialization/BulkDataRegistry.h"
