Chaos : fix multi sweep causing character movement to jitter

diff --git a/Engine/Source/Runtime/Engine/Private/Collision/WorldCollision.cpp b/Engine/Source/Runtime/Engine/Private/Collision/WorldCollision.cpp
--- a/Engine/Source/Runtime/Engine/Private/Collision/WorldCollision.cpp
+++ b/Engine/Source/Runtime/Engine/Private/Collision/WorldCollision.cpp
@@ -464,11 +464,6 @@ bool UWorld::ComponentSweepMultiByChannel(TArray<struct FHitResult>& OutHits, cl
 	{
 		TArray<Chaos::FPerShapeData*> Shapes = Interface->GetAllShapes({ &Object, 1 });
 
-		// Determine how to convert the local space of this body instance to the test space
-		const FTransform ObjectToWorld = Interface->GetTransform(Object);
-		const FTransform ObjectToStart = ComponentToStart * WorldToComponent * ObjectToWorld;
-		const FTransform ObjectToEnd = ComponentToEnd * WorldToComponent * ObjectToWorld;
-
 		for (Chaos::FPerShapeData* Shape : Shapes)
 		{
 			FPhysicsShapeHandle ShapeHandle{ Shape, nullptr };
@@ -484,7 +479,7 @@ bool UWorld::ComponentSweepMultiByChannel(TArray<struct FHitResult>& OutHits, cl
 
 			FPhysicsGeometryCollection GeomCollection = FPhysicsInterface::GetGeometryCollection(ShapeHandle);
 			TArray<FHitResult> TmpHits;
-			if (FPhysicsInterface::GeomSweepMulti(this, GeomCollection, ObjectToStart.GetRotation(), TmpHits, ObjectToStart.GetTranslation(), ObjectToEnd.GetTranslation(), TraceChannel, Params, FCollisionResponseParams(PrimComp->GetCollisionResponseToChannels())))
+			if (FPhysicsInterface::GeomSweepMulti(this, GeomCollection, Rot, TmpHits, Start, End, TraceChannel, Params, FCollisionResponseParams(PrimComp->GetCollisionResponseToChannels())))
 			{
 				bHaveBlockingHit = true;
 			}
