Removed extraneous Modify calls from TryCreateConnection.

MakeLinkTo&BreakLinkTo already call Modify on the source and destination pins, which makes these calls redundant.

diff --git a/Engine/Source/Runtime/Engine/Private/EdGraph/EdGraphSchema.cpp b/Engine/Source/Runtime/Engine/Private/EdGraph/EdGraphSchema.cpp
--- a/Engine/Source/Runtime/Engine/Private/EdGraph/EdGraphSchema.cpp
+++ b/Engine/Source/Runtime/Engine/Private/EdGraph/EdGraphSchema.cpp
@@ -399,31 +399,23 @@ bool UEdGraphSchema::TryCreateConnection(UEdGraphPin* PinA, UEdGraphPin* PinB) c
 	switch (Response.Response)
 	{
 	case CONNECT_RESPONSE_MAKE:
-		PinA->Modify();
-		PinB->Modify();
 		PinA->MakeLinkTo(PinB);
 		bModified = true;
 		break;
 
 	case CONNECT_RESPONSE_BREAK_OTHERS_A:
-		PinA->Modify();
-		PinB->Modify();
 		PinA->BreakAllPinLinks(true);
 		PinA->MakeLinkTo(PinB);
 		bModified = true;
 		break;
 
 	case CONNECT_RESPONSE_BREAK_OTHERS_B:
-		PinA->Modify();
-		PinB->Modify();
 		PinB->BreakAllPinLinks(true);
 		PinA->MakeLinkTo(PinB);
 		bModified = true;
 		break;
 
 	case CONNECT_RESPONSE_BREAK_OTHERS_AB:
-		PinA->Modify();
-		PinB->Modify();
 		PinA->BreakAllPinLinks(true);
 		PinB->BreakAllPinLinks(true);
 		PinA->MakeLinkTo(PinB);
