Fix missing vertex attribute type in pass through vertex factory.
This went missing in recent changes.

diff --git a/Engine/Source/Runtime/Engine/Private/GPUSkinVertexFactory.cpp b/Engine/Source/Runtime/Engine/Private/GPUSkinVertexFactory.cpp
--- a/Engine/Source/Runtime/Engine/Private/GPUSkinVertexFactory.cpp
+++ b/Engine/Source/Runtime/Engine/Private/GPUSkinVertexFactory.cpp
@@ -1315,6 +1315,7 @@ void FGPUSkinPassthroughVertexFactory::OverrideAttributeData()
 	{
 		Data.ColorComponent.VertexBuffer = &DummyVBs[EVertexAtttribute::VertexColor];
 		Data.ColorComponent.Offset = 0;
+		Data.ColorComponent.Type = VET_Color;
 		Data.ColorComponent.VertexStreamUsage = EVertexStreamUsage::Overridden | EVertexStreamUsage::ManualFetch;
 		Data.ColorComponent.Stride = sizeof(uint32);
 	}
@@ -1333,8 +1334,9 @@ void FGPUSkinPassthroughVertexFactory::OverrideAttributeData()
 			// Note that we use the dummy VB from TexCoord0, so that is what we need to search for later in BuildStreamIndices().
 			Data.TextureCoordinates[TexCoordPairIndex].VertexBuffer = &DummyVBs[EVertexAtttribute::VertexTexCoord0 + TexCoordIndex0];
 			Data.TextureCoordinates[TexCoordPairIndex].Offset = 0;
+			Data.TextureCoordinates[TexCoordPairIndex].Type = bHasTexCoord1 ? VET_Float4 : VET_Float2;
 			Data.TextureCoordinates[TexCoordPairIndex].VertexStreamUsage = EVertexStreamUsage::Overridden | EVertexStreamUsage::ManualFetch;
-			Data.TextureCoordinates[TexCoordPairIndex].Stride = sizeof(float) * (bHasTexCoord1 ? 4 : 2);
+			Data.TextureCoordinates[TexCoordPairIndex].Stride = bHasTexCoord1 ? 16 : 8;
 
 			Data.NumTexCoords = FMath::Max<int32>(Data.NumTexCoords, TexCoordIndex0 + (bHasTexCoord1 ? 2 : 1));
 		}
