[NO NON-PII COMMIT MESSAGE]

diff --git a/Engine/Source/Runtime/Engine/Classes/Engine/Texture.h b/Engine/Source/Runtime/Engine/Classes/Engine/Texture.h
--- a/Engine/Source/Runtime/Engine/Classes/Engine/Texture.h
+++ b/Engine/Source/Runtime/Engine/Classes/Engine/Texture.h
@@ -1279,6 +1279,10 @@ public:
 	UPROPERTY(EditAnywhere, BlueprintReadWrite, Category = Texture, AssetRegistrySearchable, AdvancedDisplay)
 	ETextureMipLoadOptions MipLoadOptions;
 
+	/** If the platform supports it, tile the texture when cooking, or keep it linear and tile it when it's actually submitted to the GPU. */
+	UPROPERTY(EditAnywhere, BlueprintReadWrite, Category = Compression, AdvancedDisplay)
+	TEnumAsByte<enum TextureCookPlatformTilingSettings> CookPlatformTilingSettings;
+
 	/** Texture group this texture belongs to */
 	UPROPERTY(EditAnywhere, BlueprintReadWrite, Category=LevelOfDetail, meta=(DisplayName="Texture Group"), AssetRegistrySearchable)
 	TEnumAsByte<enum TextureGroup> LODGroup;
@@ -1750,6 +1754,16 @@ public:
 		return UE::TextureDefines::IsUncompressed(CompressionSettings);
 	}
 
+	/** 
+	 * Checks whether this texture should be tiled to a platform-specific format during cook, or whether the bNotOfflineProcessed flag 
+	 * should be set to true at runtime because it has not been tiled at cook
+	 * 
+	 * @param  TargetPlatform	The platform for which the texture is being cooked and texture group info will be extracted from. 
+	 *                          If null, this info will be extracted from UDeviceProfileManager::Get().GetActiveProfile(), possibly at runtime
+	 * @return true if platform tiling during cook is disabled for this texture 
+	 */
+	ENGINE_API bool IsCookPlatformTilingDisabled(const ITargetPlatform* TargetPlatform) const;
+
 	/**
 	 * Calculates the size of this texture if it had MipCount miplevels streamed in.
 	 * This is the size in GPU memory of the built platformdata, accounting for LODBias, etc.
diff --git a/Engine/Source/Runtime/Engine/Classes/Engine/TextureDefines.h b/Engine/Source/Runtime/Engine/Classes/Engine/TextureDefines.h
--- a/Engine/Source/Runtime/Engine/Classes/Engine/TextureDefines.h
+++ b/Engine/Source/Runtime/Engine/Classes/Engine/TextureDefines.h
@@ -395,6 +395,18 @@ enum class ETextureColorSpace : uint8
 	TCS_MAX,
 };
 
+UENUM()
+enum TextureCookPlatformTilingSettings : uint8
+{
+	/** Get the tiling setting from the texture's group CookPlatformTilingDisabled setting. By default it's to tile during cook, unless it has been changed in the texture group */
+	TCPTS_FromTextureGroup UMETA(DisplayName = "FromTextureGroup"),
+	/** The texture will be tiled during the cook process if the platform supports it. */
+	TCPTS_Tile UMETA(DisplayName = "Tile during cook"),
+	/** The texture will not be tiled during the cook process, and will be tiled when uploaded to the GPU if the platform supports it. */
+	TCPTS_DoNotTile UMETA(DisplayName = "Do not tile during cook"),
+	TCPTS_MAX,
+};
+
 /** List of chromatic adaptation methods, matching the list in ColorManagementDefines.h. */
 UENUM()
 enum class ETextureChromaticAdaptationMethod : uint8
diff --git a/Engine/Source/Runtime/Engine/Classes/Engine/TextureLODSettings.h b/Engine/Source/Runtime/Engine/Classes/Engine/TextureLODSettings.h
--- a/Engine/Source/Runtime/Engine/Classes/Engine/TextureLODSettings.h
+++ b/Engine/Source/Runtime/Engine/Classes/Engine/TextureLODSettings.h
@@ -41,6 +41,7 @@ struct ENGINE_API FTextureLODGroup
 		, VirtualTextureTileCountBias(0)
 		, VirtualTextureTileSizeBias(0)
 		, LossyCompressionAmount(TLCA_Default)
+		, CookPlatformTilingDisabled(false)
 	{
 		SetupGroup();
 	}
@@ -127,6 +128,10 @@ struct ENGINE_API FTextureLODGroup
 	UPROPERTY()
 	TEnumAsByte<enum ETextureLossyCompressionAmount> LossyCompressionAmount;
 
+	/** If true textures with CookPlatformTilingSettings set to TCPTS_FromTextureGroup will not be tiled during cook. They will be tiled when uploaded to the GPU if necessary */
+	UPROPERTY()
+	bool CookPlatformTilingDisabled;
+
 	void SetupGroup();
 
 	bool operator==(const FTextureLODGroup& Other) const;
diff --git a/Engine/Source/Runtime/Engine/Private/Texture.cpp b/Engine/Source/Runtime/Engine/Private/Texture.cpp
--- a/Engine/Source/Runtime/Engine/Private/Texture.cpp
+++ b/Engine/Source/Runtime/Engine/Private/Texture.cpp
@@ -33,6 +33,8 @@
 #include "ImageCoreUtils.h"
 #include "ImageUtils.h"
 #include "Algo/Unique.h"
+#include "DeviceProfiles/DeviceProfile.h"
+#include "DeviceProfiles/DeviceProfileManager.h"
 
 #if WITH_EDITOR
 #include "DerivedDataBuildVersion.h"
@@ -157,6 +159,7 @@ UTexture::UTexture(const FObjectInitializer& ObjectInitializer)
 	CompressionYCoCg = 0;
 	Downscale = 0.f;
 	DownscaleOptions = ETextureDownscaleOptions::Default;
+	CookPlatformTilingSettings = TextureCookPlatformTilingSettings::TCPTS_FromTextureGroup;
 	Source.SetOwner(this);
 #endif // #if WITH_EDITORONLY_DATA
 
@@ -1040,6 +1043,12 @@ void UTexture::PostLoad()
 
 #endif
 
+	if (IsCookPlatformTilingDisabled(nullptr)) // nullptr TargetPlatform means it will use UDeviceProfileManager::Get().GetActiveProfile() to get the tiling settings
+	{
+		// The texture was not processed/tiled during cook, so it has to be tiled when uploaded to the GPU if necessary
+		bNotOfflineProcessed = true;
+	}
+
 	if( !IsTemplate() )
 	{
 		// The texture will be cached by the cubemap it is contained within on consoles.
@@ -1304,6 +1313,38 @@ TextureMipGenSettings UTexture::GetMipGenSettingsFromString(const TCHAR* InStr,
 	return bTextureGroup ? TMGS_SimpleAverage : TMGS_FromTextureGroup;
 }
 
+bool UTexture::IsCookPlatformTilingDisabled(const ITargetPlatform* TargetPlatform) const
+{
+	if (CookPlatformTilingSettings.GetValue() == TextureCookPlatformTilingSettings::TCPTS_FromTextureGroup)
+	{
+		const UTextureLODSettings* TextureLODSettings = nullptr;
+
+		if (TargetPlatform)
+		{
+			TextureLODSettings = &TargetPlatform->GetTextureLODSettings();
+		}
+		else
+		{
+			TextureLODSettings = UDeviceProfileManager::Get().GetActiveProfile()->GetTextureLODSettings();
+
+			if (!TextureLODSettings)
+			{
+				return false;
+			}
+		}
+
+		check(TextureLODSettings);
+
+		checkf(LODGroup < TextureLODSettings->TextureLODGroups.Num(), 
+			TEXT("A texture had passed a bad LODGroup to UTexture::IsCookPlatformTilingDisabled (%d, out of %d groups). The texture name is '%s'."), 
+			LODGroup, TextureLODSettings->TextureLODGroups.Num(), *GetPathName());
+
+		return TextureLODSettings->TextureLODGroups[LODGroup].CookPlatformTilingDisabled;
+	}
+
+	return CookPlatformTilingSettings.GetValue() == TextureCookPlatformTilingSettings::TCPTS_DoNotTile;
+}
+
 void UTexture::SetDeterministicLightingGuid()
 {
 #if WITH_EDITORONLY_DATA
