5.1.1 Fix for live issue with INI file settings are missing for various platforms.

diff --git a/Engine/Source/Runtime/Engine/Private/AudioFormatSettings.cpp b/Engine/Source/Runtime/Engine/Private/AudioFormatSettings.cpp
--- a/Engine/Source/Runtime/Engine/Private/AudioFormatSettings.cpp
+++ b/Engine/Source/Runtime/Engine/Private/AudioFormatSettings.cpp
@@ -60,25 +60,50 @@ namespace Audio
 
 	void FAudioFormatSettings::ReadConfiguration(FConfigCacheIni* InConfigSystem, const FString& InConfigFilename, const FString& InPlatformIdentifierForLogging)
 	{
+		auto MakePrettyArrayToString = [](const TArray<FName>& InNames) -> FString 
+		{
+			return FString::JoinBy(InNames, TEXT(", "),[](const FName& i) -> FString { return i.GetPlainNameString(); });
+		};
+		
 		auto ToFName = [](const FString& InName) -> FName { return { *InName }; };
+			
+		using namespace Audio;
 
-		TArray<FString> FormatNames;
-		if (ensure(InConfigSystem->GetArray(TEXT("Audio"), TEXT("AllWaveFormats"), FormatNames, InConfigFilename)))
-		{
-			Algo::Transform(FormatNames, AllWaveFormats, ToFName);
+		// AllWaveFormats.
+		{		
+			TArray<FString> FormatNames;
+			if (InConfigSystem->GetArray(TEXT("Audio"), TEXT("AllWaveFormats"), FormatNames, InConfigFilename))
+			{
+				Algo::Transform(FormatNames, AllWaveFormats, ToFName);
+			}
+			else
+			{
+				AllWaveFormats = { NAME_BINKA, NAME_ADPCM, NAME_PCM };
+				UE_LOG(LogAudio, Warning, TEXT("Audio:AllWaveFormats is not defined, defaulting to built in formats. (%s)"), *MakePrettyArrayToString(AllWaveFormats));
+			}
 		}
 
-		TArray<FString> FormatModuleHints;
-		if (InConfigSystem->GetArray(TEXT("Audio"), TEXT("FormatModuleHints"), FormatModuleHints, InConfigFilename))
-		{
-			Algo::Transform(FormatNames, WaveFormatModuleHints, ToFName);
+		// FormatModuleHints
+		{		
+			TArray<FString> FormatModuleHints;
+			if (InConfigSystem->GetArray(TEXT("Audio"), TEXT("FormatModuleHints"), FormatModuleHints, InConfigFilename))
+			{
+				Algo::Transform(FormatModuleHints, WaveFormatModuleHints, ToFName);
+			}
 		}
 
-		FString FallbackFormatString;
-		if (ensure(InConfigSystem->GetString(TEXT("Audio"), TEXT("FallbackFormat"), FallbackFormatString, InConfigFilename)))
+		// FallbackFormat
 		{
-			FallbackFormat = *FallbackFormatString;
-
+			FString FallbackFormatString;
+			if (InConfigSystem->GetString(TEXT("Audio"), TEXT("FallbackFormat"), FallbackFormatString, InConfigFilename))
+			{
+				FallbackFormat = *FallbackFormatString;
+			}
+			else
+			{
+				FallbackFormat = NAME_ADPCM;
+				UE_LOG(LogAudio, Warning, TEXT("Audio:FallbackFormat is not defined, defaulting to '%s'."), *FallbackFormat.GetPlainNameString());
+			}
 			if (!AllWaveFormats.Contains(FallbackFormat) && AllWaveFormats.Num() > 0)
 			{
 				UE_LOG(LogAudio, Warning, TEXT("FallbackFormat '%s' not defined in 'AllWaveFormats'. Using first format listed '%s'"), *FallbackFormatString, *AllWaveFormats[0].ToString());
@@ -86,22 +111,37 @@ namespace Audio
 			}
 		}
 
-		FString PlatformFormatString;
-		if (ensure(InConfigSystem->GetString(TEXT("Audio"), TEXT("PlatformFormat"), PlatformFormatString, InConfigFilename)))
+		// PlatformFormat
 		{
-			PlatformFormat = *PlatformFormatString;
-
+			FString PlatformFormatString;
+			if (InConfigSystem->GetString(TEXT("Audio"), TEXT("PlatformFormat"), PlatformFormatString, InConfigFilename))
+			{
+				PlatformFormat = *PlatformFormatString;
+			}
+			else
+			{
+				PlatformFormat = NAME_ADPCM;
+				UE_LOG(LogAudio, Warning, TEXT("Audio:PlatformFormat is not defined, defaulting to '%s'."), *PlatformFormat.GetPlainNameString());
+			}
 			if (!AllWaveFormats.Contains(PlatformFormat))
 			{
-				UE_LOG(LogAudio, Warning, TEXT("PlatformStreamingFormat '%s' not defined in 'AllWaveFormats'. Using fallback format '%s'"), *PlatformFormatString, *FallbackFormat.ToString());
+				UE_LOG(LogAudio, Warning, TEXT("PlatformFormat '%s' not defined in 'AllWaveFormats'. Using fallback format '%s'"), *PlatformFormatString, *FallbackFormat.ToString());
 				PlatformFormat = FallbackFormat;
 			}
 		}
 
-		FString PlatformStreamingFormatString;
-		if (ensure(InConfigSystem->GetString(TEXT("Audio"), TEXT("PlatformStreamingFormat"), PlatformStreamingFormatString, InConfigFilename)))
+		// PlatformStreamingFormat
 		{
-			PlatformStreamingFormat = *PlatformStreamingFormatString;
+			FString PlatformStreamingFormatString;
+			if (InConfigSystem->GetString(TEXT("Audio"), TEXT("PlatformStreamingFormat"), PlatformStreamingFormatString, InConfigFilename))
+			{
+				PlatformStreamingFormat = *PlatformStreamingFormatString;		
+			}
+			else
+			{
+				PlatformStreamingFormat = NAME_ADPCM;
+				UE_LOG(LogAudio, Warning, TEXT("Audio:PlatformStreamingFormat is not defined, defaulting to '%s'."), *PlatformStreamingFormat.GetPlainNameString());
+			}
 			if (!AllWaveFormats.Contains(PlatformStreamingFormat))
 			{
 				UE_LOG(LogAudio, Warning, TEXT("PlatformStreamingFormat '%s' not defined in 'AllWaveFormats'. Using fallback format '%s'"), *PlatformStreamingFormatString, *FallbackFormat.ToString());
@@ -109,19 +149,8 @@ namespace Audio
 			}
 		}
 
-	#if !NO_LOGGING
-
-		// Display log for sanity
-		FString AllFormatsConcat;
-		for (FName i : AllWaveFormats)
-		{
-			AllFormatsConcat += i.ToString() + TEXT(" ");
-		}
-
-		UE_LOG(LogAudio, Verbose, TEXT("AudioFormatSettings: TargetName='%s', AllWaveFormats=( %s), PlatformFormat='%s', PlatformStreamingFormat='%s', FallbackFormat='%s'"),
-			*InPlatformIdentifierForLogging, *AllFormatsConcat, *PlatformFormat.ToString(), *PlatformStreamingFormat.ToString(), *FallbackFormat.ToString());
-
-	#endif //!NO_LOGGING
+		UE_LOG(LogAudio, Verbose, TEXT("AudioFormatSettings: TargetName='%s', AllWaveFormats=(%s), Hints=(%s), PlatformFormat='%s', PlatformStreamingFormat='%s', FallbackFormat='%s'"),
+			*InPlatformIdentifierForLogging, *MakePrettyArrayToString(AllWaveFormats), *MakePrettyArrayToString(WaveFormatModuleHints), *PlatformFormat.ToString(), *PlatformStreamingFormat.ToString(), *FallbackFormat.ToString());	
 	}
 
 }// namespace Audio
