Analytics: Backed out 22909100 and 22951911

#rnx

diff --git a/Engine/Source/Developer/DerivedDataCache/DerivedDataCache.Build.cs b/Engine/Source/Developer/DerivedDataCache/DerivedDataCache.Build.cs
--- a/Engine/Source/Developer/DerivedDataCache/DerivedDataCache.Build.cs
+++ b/Engine/Source/Developer/DerivedDataCache/DerivedDataCache.Build.cs
@@ -11,7 +11,7 @@ public class DerivedDataCache : ModuleRules
 		PrivateDependencyModuleNames.Add("Zen");
 
 		// Dependencies for "S3" and "HTTP" backends
-		PrivateDependencyModuleNames.AddRange(new string[] { "SSL", "Json", "Analytics", "Zen" });
+		PrivateDependencyModuleNames.AddRange(new string[] { "SSL", "Json", "Zen" });
 		PrivateIncludePathModuleNames.AddRange(new string[] { "DesktopPlatform", "Zen"});
 		AddEngineThirdPartyPrivateStaticDependencies(Target, "libcurl");
 		AddEngineThirdPartyPrivateStaticDependencies(Target, "OpenSSL");
diff --git a/Engine/Source/Developer/DerivedDataCache/Private/DerivedDataCache.cpp b/Engine/Source/Developer/DerivedDataCache/Private/DerivedDataCache.cpp
--- a/Engine/Source/Developer/DerivedDataCache/Private/DerivedDataCache.cpp
+++ b/Engine/Source/Developer/DerivedDataCache/Private/DerivedDataCache.cpp
@@ -4,8 +4,6 @@
 #include "DerivedDataCacheInterface.h"
 
 #include "Algo/AllOf.h"
-#include "Analytics.h"
-#include "AnalyticsEventAttribute.h"
 #include "Async/AsyncWork.h"
 #include "Async/InheritedContext.h"
 #include "Async/TaskGraphInterfaces.h"
@@ -573,19 +571,12 @@ public:
 
 		GVerifyDDC = FParse::Param(FCommandLine::Get(), TEXT("VerifyDDC"));
 
-		FAnalytics::Get().GetEventCallback(TEXT("Core.Loading"))->AddRaw(this, &FDerivedDataCache::OnAnalyticsEvent);
-		
 		UE_CLOG(GVerifyDDC, LogDerivedDataCache, Display, TEXT("Items retrieved from the DDC will be verified (-VerifyDDC)"));
 	}
 
 	/** Destructor, flushes all sync tasks **/
 	~FDerivedDataCache()
 	{
-		if (FAnalytics::IsAvailable())
-		{
-			FAnalytics::Get().GetEventCallback(TEXT("Core.Loading"))->RemoveAll(this);
-		}
-		
 		WaitForQuiescence(true);
 		FScopeLock ScopeLock(&SynchronizationObject);
 		for (TMap<uint32,FAsyncTask<FBuildAsyncWorker>*>::TIterator It(PendingTasks); It; ++It)
@@ -953,77 +944,6 @@ public:
 		GatherDerivedDataCacheSummaryStats(DDCSummaryStats);
 	}
 
-	void OnAnalyticsEvent(TArray<FAnalyticsEventAttribute>& Attributes)
-	{
-#if ENABLE_COOK_STATS
-
-		// Gather the latest resource stats
-		TArray<FDerivedDataCacheResourceStat> ResourceStats;
-
-		GatherDerivedDataCacheResourceStats(ResourceStats);
-
-		FDerivedDataCacheResourceStat ResourceStatsTotal(TEXT("Total"));
-
-		// Accumulate Totals
-		for (const FDerivedDataCacheResourceStat& Stat : ResourceStats)
-		{
-			ResourceStatsTotal += Stat;
-		}
-
-		ResourceStats.Emplace(ResourceStatsTotal);
-
-		// Append to the attributes
-		for (const FDerivedDataCacheResourceStat& Stat : ResourceStats)
-		{
-			FString BaseName = TEXT("DDC.Resource.") + Stat.AssetType;
-
-			BaseName = BaseName.Replace(TEXT("("), TEXT("")).Replace(TEXT(")"), TEXT(""));
-
-			{
-				FString AttrName = BaseName + TEXT(".BuildCount");
-				Attributes.Emplace(MoveTemp(AttrName), Stat.BuildCount);
-			}
-
-			{
-				FString AttrName = BaseName + TEXT(".BuildTimeSec");
-				Attributes.Emplace(MoveTemp(AttrName), Stat.BuildTimeSec);
-			}
-
-			{
-				FString AttrName = BaseName + TEXT(".BuildSizeMB");
-				Attributes.Emplace(MoveTemp(AttrName), Stat.BuildSizeMB);
-			}
-
-			{
-				FString AttrName = BaseName + TEXT(".LoadCount");
-				Attributes.Emplace(MoveTemp(AttrName), Stat.LoadCount);
-			}
-
-			{
-				FString AttrName = BaseName + TEXT(".LoadTimeSecLoadTimeSec");
-				Attributes.Emplace(MoveTemp(AttrName), Stat.LoadTimeSec);
-			}
-
-			{
-				FString AttrName = BaseName + TEXT(".LoadSizeMB");
-				Attributes.Emplace(MoveTemp(AttrName), Stat.LoadSizeMB);
-			}
-		}
-
-		// Gather the summary stats
-		FDerivedDataCacheSummaryStats SummaryStats;
-
-		GatherDerivedDataCacheSummaryStats(SummaryStats);
-
-		// Append to the attributes
-		for (const FDerivedDataCacheSummaryStat& Stat : SummaryStats.Stats)
-		{
-			FString FormattedAttrName = "DDC.Summary." + Stat.Key;
-			Attributes.Emplace(FormattedAttrName, Stat.Value);
-		}
-#endif
-	}
-
 	/** Get event delegate for data cache notifications */
 	virtual FOnDDCNotification& GetDDCNotificationEvent()
 	{
diff --git a/Engine/Source/Developer/DerivedDataCache/Private/HttpCacheStore.cpp b/Engine/Source/Developer/DerivedDataCache/Private/HttpCacheStore.cpp
--- a/Engine/Source/Developer/DerivedDataCache/Private/HttpCacheStore.cpp
+++ b/Engine/Source/Developer/DerivedDataCache/Private/HttpCacheStore.cpp
@@ -34,7 +34,6 @@
 #include "Misc/ScopeLock.h"
 #include "Misc/ScopeRWLock.h"
 #include "Misc/StringBuilder.h"
-#include "ProfilingDebugging/CookStats.h"
 #include "ProfilingDebugging/CountersTrace.h"
 #include "ProfilingDebugging/CpuProfilerTrace.h"
 #include "Serialization/CompactBinary.h"
@@ -45,8 +44,6 @@
 #include "Serialization/JsonReader.h"
 #include "Serialization/JsonSerializer.h"
 #include "String/Find.h"
-#include "Analytics.h"
-#include "AnalyticsEventAttribute.h"
 
 #if PLATFORM_MICROSOFT
 #include "Microsoft/WindowsHWrapper.h"
@@ -335,7 +332,6 @@ private:
 	bool EndIsServiceReady(THttpUniquePtr<IHttpResponse>& Response, TArray64<uint8>& Body);
 	bool AcquireAccessToken(IHttpClient* Client = nullptr);
 	void SetAccessToken(FStringView Token, double RefreshDelay = 0.0);
-	void OnAnalyticsEvent(TArray<FAnalyticsEventAttribute>& Attributes);
 
 	enum class EOperationCategory
 	{
@@ -1618,50 +1614,11 @@ FHttpCacheStore::FHttpCacheStore(const FHttpCacheStoreParams& Params)
 		bIsUsable = true;
 	}
 
-	// Add analytics callbacks for the Loading and Cooking events
-	FAnalytics::Get().GetEventCallback("Core.Loading")->AddRaw(this, &FHttpCacheStore::OnAnalyticsEvent);
-	FAnalytics::Get().GetEventCallback("Core.Cooking")->AddRaw(this, &FHttpCacheStore::OnAnalyticsEvent);
-
 	AnyInstance = this;
 }
 
-void FHttpCacheStore::OnAnalyticsEvent(TArray<FAnalyticsEventAttribute>& Attributes)
-{
-#if ENABLE_COOK_STATS
-	const FString BaseName(TEXTVIEW("CloudDDC."));
-
-	Attributes.Emplace(FString(WriteToString<64>(BaseName, TEXTVIEW("Domain"))), Domain);
-	Attributes.Emplace(FString(WriteToString<64>(BaseName, TEXTVIEW("EffectiveDomain"))), *EffectiveDomain);
-	Attributes.Emplace(FString(WriteToString<64>(BaseName, TEXTVIEW("Namespace"))), Namespace);
-	Attributes.Emplace(FString(WriteToString<64>(BaseName, TEXTVIEW("LoginAttempts"))), LoginAttempts);
-	Attributes.Emplace(FString(WriteToString<64>(BaseName, TEXTVIEW("InteractiveLoginAttempts"))), InteractiveLoginAttempts);
-	Attributes.Emplace(FString(WriteToString<64>(BaseName, TEXTVIEW("FailedLoginAttempts"))), FailedLoginAttempts);
-
-	const int64 GetHits = UsageStats.GetStats.GetAccumulatedValueAnyThread(FCookStats::CallStats::EHitOrMiss::Hit, FCookStats::CallStats::EStatType::Counter);
-	const int64 GetMisses = UsageStats.GetStats.GetAccumulatedValueAnyThread(FCookStats::CallStats::EHitOrMiss::Miss, FCookStats::CallStats::EStatType::Counter);
-	const int64 PutHits = UsageStats.PutStats.GetAccumulatedValueAnyThread(FCookStats::CallStats::EHitOrMiss::Hit, FCookStats::CallStats::EStatType::Counter);
-	const int64 PutMisses = UsageStats.PutStats.GetAccumulatedValueAnyThread(FCookStats::CallStats::EHitOrMiss::Miss, FCookStats::CallStats::EStatType::Counter);
-	const int64 TotalGets = GetHits + GetMisses;
-	const int64 TotalPuts = PutHits + PutMisses;
-
-	Attributes.Emplace(FString(WriteToString<64>(BaseName, TEXTVIEW("GetHits"))), GetHits);
-	Attributes.Emplace(FString(WriteToString<64>(BaseName, TEXTVIEW("GetMisses"))), GetMisses);
-	Attributes.Emplace(FString(WriteToString<64>(BaseName, TEXTVIEW("TotalGets"))), TotalGets);
-	Attributes.Emplace(FString(WriteToString<64>(BaseName, TEXTVIEW("PutHits"))), PutHits);
-	Attributes.Emplace(FString(WriteToString<64>(BaseName, TEXTVIEW("PutMisses"))), TotalPuts);
-	Attributes.Emplace(FString(WriteToString<64>(BaseName, TEXTVIEW("TotalGets"))), TotalPuts);
-#endif
-}
-
 FHttpCacheStore::~FHttpCacheStore()
 {
-	// Remove any callbacks we registered if analytics are still available
-	if (FAnalytics::IsAvailable())
-	{
-		FAnalytics::Get().GetEventCallback("Core.Loading")->RemoveAll(this);
-		FAnalytics::Get().GetEventCallback("Core.Cooking")->RemoveAll(this);
-	}
-
 	if (RefreshAccessTokenHandle.IsValid())
 	{
 		FTSTicker::GetCoreTicker().RemoveTicker(RefreshAccessTokenHandle);
diff --git a/Engine/Source/Developer/Virtualization/Private/VirtualizationManager.cpp b/Engine/Source/Developer/Virtualization/Private/VirtualizationManager.cpp
--- a/Engine/Source/Developer/Virtualization/Private/VirtualizationManager.cpp
+++ b/Engine/Source/Developer/Virtualization/Private/VirtualizationManager.cpp
@@ -18,8 +18,6 @@
 #include "PackageVirtualizationProcess.h"
 #include "ProfilingDebugging/CookStats.h"
 #include "VirtualizationFilterSettings.h"
-#include "Analytics.h"
-#include "AnalyticsEventAttribute.h"
 
 #define LOCTEXT_NAMESPACE "Virtualization"
 
@@ -456,11 +454,6 @@ FVirtualizationManager::FVirtualizationManager()
 
 FVirtualizationManager::~FVirtualizationManager()
 {
-	if (FAnalytics::IsAvailable())
-	{
-		FAnalytics::Get().GetEventCallback(TEXT("Core.Loading"))->RemoveAll(this);
-	}
-
 	for (const TPair<IConsoleVariable*, FDelegateHandle>& KV : DebugValues.ConsoleDelegateHandles)
 	{
 		IConsoleVariable* ConsoleVariable = KV.Key;
@@ -513,8 +506,6 @@ bool FVirtualizationManager::Initialize(const FInitParams& InitParams)
 
 	MountBackends(InitParams.ConfigFile);
 
-	FAnalytics::Get().GetEventCallback(TEXT("Core.Loading"))->AddRaw(this, &FVirtualizationManager::OnAnalyticsEvent);
-
 	UE_LOG(LogVirtualization, Display, TEXT("Virtualization manager initialization completed"));
 
 	return true;
@@ -2030,71 +2021,6 @@ void FVirtualizationManager::BroadcastEvent(TConstArrayView<FPullRequest> Reques
 	}
 }
 
-void FVirtualizationManager::OnAnalyticsEvent(TArray<FAnalyticsEventAttribute>& Attributes)
-{
-	using namespace UE::Virtualization;
-
-	// Grab the Virtualization stats
-	if (IVirtualizationSystem::IsInitialized())
-	{
-		IVirtualizationSystem& System = IVirtualizationSystem::Get();
-
-		FPayloadActivityInfo PayloadActivityInfo = System.GetAccumualtedPayloadActivityInfo();
-
-		const FString BaseName = TEXT("Virtualization");
-
-		{
-			FString AttrName = BaseName + TEXT(".Enabled");
-			Attributes.Emplace(MoveTemp(AttrName), System.IsEnabled());
-		}
-
-		{
-			FString AttrName = BaseName + TEXT(".Cache.TimeSpent");
-			Attributes.Emplace(MoveTemp(AttrName), (double)PayloadActivityInfo.Cache.CyclesSpent * FPlatformTime::GetSecondsPerCycle());
-		}
-
-		{
-			FString AttrName = BaseName + TEXT(".Cache.PayloadCount");
-			Attributes.Emplace(MoveTemp(AttrName), (double)PayloadActivityInfo.Cache.PayloadCount);
-		}
-
-		{
-			FString AttrName = BaseName + TEXT(".Cache.TotalBytes");
-			Attributes.Emplace(MoveTemp(AttrName), (double)PayloadActivityInfo.Cache.TotalBytes);
-		}
-
-		{
-			FString AttrName = BaseName + TEXT(".Push.TimeSpent");
-			Attributes.Emplace(MoveTemp(AttrName), (double)PayloadActivityInfo.Push.CyclesSpent * FPlatformTime::GetSecondsPerCycle());
-		}
-
-		{
-			FString AttrName = BaseName + TEXT(".Push.PayloadCount");
-			Attributes.Emplace(MoveTemp(AttrName), (double)PayloadActivityInfo.Push.PayloadCount);
-		}
-
-		{
-			FString AttrName = BaseName + TEXT(".Push.TotalBytes");
-			Attributes.Emplace(MoveTemp(AttrName), (double)PayloadActivityInfo.Push.TotalBytes);
-		}
-
-		{
-			FString AttrName = BaseName + TEXT(".Pull.TimeSpent");
-			Attributes.Emplace(MoveTemp(AttrName), (double)PayloadActivityInfo.Pull.CyclesSpent * FPlatformTime::GetSecondsPerCycle());
-		}
-
-		{
-			FString AttrName = BaseName + TEXT(".Pull.PayloadCount");
-			Attributes.Emplace(MoveTemp(AttrName), (double)PayloadActivityInfo.Pull.PayloadCount);
-		}
-
-		{
-			FString AttrName = BaseName + TEXT(".Pull.TotalBytes");
-			Attributes.Emplace(MoveTemp(AttrName), (double)PayloadActivityInfo.Pull.TotalBytes);
-		}
-	}
-}
-
 } // namespace UE::Virtualization
 
 #undef UE_INLINE_ALLOCATION_COUNT
diff --git a/Engine/Source/Developer/Virtualization/Private/VirtualizationManager.h b/Engine/Source/Developer/Virtualization/Private/VirtualizationManager.h
--- a/Engine/Source/Developer/Virtualization/Private/VirtualizationManager.h
+++ b/Engine/Source/Developer/Virtualization/Private/VirtualizationManager.h
@@ -7,12 +7,10 @@
 #include "Logging/LogMacros.h"
 #include "Templates/UniquePtr.h"
 
-
 #include "Virtualization/VirtualizationSystem.h"
 
 class IConsoleObject;
 class FOutputDevice;
-struct FAnalyticsEventAttribute;
 
 /**
  * Configuring the backend hierarchy
@@ -156,7 +154,6 @@ private:
 	void OnUpdateDebugMissBackendsFromConsole(const TArray<FString>& Args, FOutputDevice& OutputDevice);
 	void OnUpdateDebugMissChanceFromConsole(const TArray<FString>& Args, FOutputDevice& OutputDevice);
 	void OnUpdateDebugMissCountFromConsole(const TArray<FString>& Args, FOutputDevice& OutputDevice);
-	void OnAnalyticsEvent(TArray<FAnalyticsEventAttribute>& Attributes);
 
 	void UpdateBackendDebugState();
 
@@ -178,7 +175,7 @@ private:
 
 	void PullDataFromAllBackends(TArrayView<FPullRequest> Requests);
 	void PullDataFromBackend(IVirtualizationBackend& Backend, TArrayView<FPullRequest> Requests);
-	
+
 	bool ShouldVirtualizeAsset(const UObject* Owner) const;
 
 	/** 
diff --git a/Engine/Source/Developer/Virtualization/Virtualization.Build.cs b/Engine/Source/Developer/Virtualization/Virtualization.Build.cs
--- a/Engine/Source/Developer/Virtualization/Virtualization.Build.cs
+++ b/Engine/Source/Developer/Virtualization/Virtualization.Build.cs
@@ -14,7 +14,6 @@ public class Virtualization : ModuleRules
 				"CoreUObject",
 				"ApplicationCore",
 				"DerivedDataCache",
-				"Analytics",
 				"MessageLog",
 				"Projects",
 				"SourceControl"
diff --git a/Engine/Source/Developer/Zen/Zen.Build.cs b/Engine/Source/Developer/Zen/Zen.Build.cs
--- a/Engine/Source/Developer/Zen/Zen.Build.cs
+++ b/Engine/Source/Developer/Zen/Zen.Build.cs
@@ -7,7 +7,7 @@ public class Zen : ModuleRules
 {
 	public Zen(ReadOnlyTargetRules Target) : base(Target)
 	{
-		PrivateDependencyModuleNames.AddRange(new string[] { "Core", "Sockets", "SSL", "Json", "Analytics" });
+		PrivateDependencyModuleNames.AddRange(new string[] { "Core", "Sockets", "SSL", "Json" });
 		PrivateIncludePathModuleNames.Add("DesktopPlatform");
 
 		AddEngineThirdPartyPrivateStaticDependencies(Target, "libcurl");
diff --git a/Engine/Source/Runtime/Analytics/Analytics/Public/Analytics.h b/Engine/Source/Runtime/Analytics/Analytics/Public/Analytics.h
--- a/Engine/Source/Runtime/Analytics/Analytics/Public/Analytics.h
+++ b/Engine/Source/Runtime/Analytics/Analytics/Public/Analytics.h
@@ -5,7 +5,6 @@
 #include "AnalyticsBuildType.h"
 #include "AnalyticsProviderConfigurationDelegate.h"
 #include "Containers/UnrealString.h"
-#include "Containers/Map.h"
 #include "CoreGlobals.h"
 #include "CoreMinimal.h"
 #include "Delegates/Delegate.h"
@@ -20,10 +19,6 @@
 ANALYTICS_API DECLARE_LOG_CATEGORY_EXTERN(LogAnalytics, Display, All);
 
 class IAnalyticsProvider;
-struct FAnalyticsEventAttribute;
-
-/** Analytics Event Callback */
-DECLARE_MULTICAST_DELEGATE_OneParam(FAnalyticsEventCallback, TArray<FAnalyticsEventAttribute>&);
 
 /**
  * The public interface for interacting with analytics.
@@ -68,37 +63,6 @@ public:
 		return FModuleManager::Get().IsModuleLoaded( "Analytics" );
 	}
 
-	TSharedPtr<FAnalyticsEventCallback> GetEventCallback(const FString& EventName)
-	{
-		// Check to see if we already have an event callback registered with this name
-		TSharedPtr<FAnalyticsEventCallback>* CallBackPtr = EventCallbacks.Find(EventName);
-
-		if (CallBackPtr != nullptr)
-		{
-			// Previously registered event callback so return it
-			return *CallBackPtr;
-		}
-
-		// Create and register a new event callback
-		TSharedPtr<FAnalyticsEventCallback> CallBack = MakeShareable(new FAnalyticsEventCallback());
-
-		EventCallbacks.Add(EventName, CallBack);
-
-		return CallBack;
-	}
-
-	void FireEventCallbacks(const FString& EventName, TArray<FAnalyticsEventAttribute>& Attributes)
-	{
-		// Check to see if we already have an event callback registered with this name
-		TSharedPtr<FAnalyticsEventCallback>* CallBackPtr = EventCallbacks.Find(EventName);
-
-		if (CallBackPtr != nullptr)
-		{
-			// There are valid callbacks for this event so fire it
-			(*CallBackPtr)->Broadcast(Attributes);
-		}
-	}
-
 	//--------------------------------------------------------------------------
 	// Provider Factory functions.
 	//--------------------------------------------------------------------------
@@ -229,12 +193,7 @@ public:
 	 */
 	virtual void WriteConfigValueToIni(const FString& IniName, const FString& SectionName, const FString& KeyName, const FString& Value);
 
-
 private:
 	virtual void StartupModule() override;
 	virtual void ShutdownModule() override;
-
-	TMap< FString, TSharedPtr<FAnalyticsEventCallback> > EventCallbacks;
-
 };
-
diff --git a/Engine/Source/Runtime/Analytics/AnalyticsET/Private/IAnalyticsProviderET.cpp b/Engine/Source/Runtime/Analytics/AnalyticsET/Private/IAnalyticsProviderET.cpp
--- a/Engine/Source/Runtime/Analytics/AnalyticsET/Private/IAnalyticsProviderET.cpp
+++ b/Engine/Source/Runtime/Analytics/AnalyticsET/Private/IAnalyticsProviderET.cpp
@@ -484,25 +484,16 @@ void FAnalyticsProviderET::RecordEvent(FString&& EventName, const TArray<FAnalyt
 	// let higher level code filter the decision of whether to send the event
 	if (ShouldRecordEvent(EventName))
 	{
-		// Make a copy of the incoming attributes so we can append any new attributes from the the callback functions
-		TArray<FAnalyticsEventAttribute> FinalAttributes;
-		
-		// Append the incoming attributes to the final attribute list
-		FinalAttributes.Append(Attributes);
-
-		// Fire any analytics callbacks that append attributes to the list
-		FAnalytics::Get().FireEventCallbacks(EventName, FinalAttributes);
-		
-		// Fire any read only callbacks
+		// fire any callbacks
 		for (const auto& Cb : EventRecordedCallbacks)
 		{
 			// we no longer track if the event was Json, each attribute does.
-			Cb(EventName, FinalAttributes, false);
+			Cb(EventName, Attributes, false);
 		}
 
 		if (!Config.UseLegacyProtocol)
 		{
-			EventCache.AddToCache(MoveTemp(EventName), FinalAttributes);
+			EventCache.AddToCache(MoveTemp(EventName), Attributes);
 
 			// if we aren't caching events, flush immediately. This is really only for debugging as it will significantly affect bandwidth.
 			if (!bShouldCacheEvents)
@@ -512,7 +503,7 @@ void FAnalyticsProviderET::RecordEvent(FString&& EventName, const TArray<FAnalyt
 		}
 		else
 		{
-			FlushEventLegacy(EventName, FinalAttributes);
+			FlushEventLegacy(EventName, Attributes);
 		}
 	}
 	else
diff --git a/Engine/Source/Runtime/Engine/Private/Analytics/StudioAnalytics.cpp b/Engine/Source/Runtime/Engine/Private/Analytics/StudioAnalytics.cpp
--- a/Engine/Source/Runtime/Engine/Private/Analytics/StudioAnalytics.cpp
+++ b/Engine/Source/Runtime/Engine/Private/Analytics/StudioAnalytics.cpp
@@ -1,10 +1,31 @@
 // Copyright Epic Games, Inc. All Rights Reserved.
 
 #include "StudioAnalytics.h"
-#include "HAL/PlatformTime.h"
+#include "Misc/Guid.h"
+#include "Stats/Stats.h"
+#include "Misc/ConfigCacheIni.h"
+#include "EngineGlobals.h"
+#include "Engine/Engine.h"
+#include "Misc/EngineBuildSettings.h"
+#include "AnalyticsBuildType.h"
+#include "AnalyticsEventAttribute.h"
 #include "IAnalyticsProviderET.h"
+#include "Analytics.h"
+#include "AnalyticsET.h"
+#include "GeneralProjectSettings.h"
+#include "Misc/EngineVersion.h"
+#include "RHI.h"
+#include "GenericPlatform/GenericPlatformCrashContext.h"
+#include "Interfaces/IAnalyticsProvider.h"
+#include "Templates/SharedPointer.h"
 #include "HAL/PlatformProcess.h"
-#include "HAL/Thread.h"
+#include "DerivedDataCacheInterface.h"
+#include "DerivedDataCacheUsageStats.h"
+#include "Virtualization/VirtualizationSystem.h"
+
+#if WITH_EDITOR
+#include "Experimental/ZenServerInterface.h"
+#endif
 
 bool FStudioAnalytics::bInitialized = false;
 volatile double FStudioAnalytics::TimeEstimation = 0;
@@ -128,6 +149,8 @@ void FStudioAnalytics::RecordEvent(const FString& EventName, const TArray<FAnaly
 
 void FStudioAnalytics::FireEvent_Loading(const FString& LoadingName, double SecondsSpentLoading, const TArray<FAnalyticsEventAttribute>& InAttributes)
 {
+	using namespace UE::Virtualization;
+
 	// Ignore anything less than a 1/4th a second.
 	if (SecondsSpentLoading < 0.250)
 	{
@@ -148,10 +171,222 @@ void FStudioAnalytics::FireEvent_Loading(const FString& LoadingName, double Seco
 		Attributes.Emplace(TEXT("LoadingSeconds"), SecondsSpentLoading);
 		Attributes.Append(InAttributes);
 
-		// Original loading event
 		FStudioAnalytics::GetProvider().RecordEvent(TEXT("Performance.Loading"), Attributes);
 
-		// New loading event which has callbacks from DDC, VA, Zen etc..
+#if ENABLE_COOK_STATS
+
+		TArray<FDerivedDataCacheResourceStat> DDCResourceStats;
+
+		// Grab the latest resource stats
+		GetDerivedDataCacheRef().GatherResourceStats(DDCResourceStats);
+
+		FDerivedDataCacheResourceStat DDCResourceStatsTotal(TEXT("Total"));
+
+		// Accumulate Totals
+		for (const FDerivedDataCacheResourceStat& Stat : DDCResourceStats)
+		{
+			DDCResourceStatsTotal += Stat;
+		}
+
+		DDCResourceStats.Emplace(DDCResourceStatsTotal);
+
+		for (const FDerivedDataCacheResourceStat& Stat : DDCResourceStats)
+		{
+			FString BaseName = TEXT("DDC.Resource.") + Stat.AssetType;
+
+			BaseName = BaseName.Replace(TEXT("("), TEXT("")).Replace(TEXT(")"), TEXT(""));
+			
+			{
+				FString AttrName = BaseName + TEXT(".BuildCount");
+				Attributes.Emplace(MoveTemp(AttrName), Stat.BuildCount);
+			}
+			
+			{
+				FString AttrName = BaseName + TEXT(".BuildTimeSec");
+				Attributes.Emplace(MoveTemp(AttrName), Stat.BuildTimeSec);
+			}
+
+			{
+				FString AttrName = BaseName + TEXT(".BuildSizeMB");
+				Attributes.Emplace(MoveTemp(AttrName), Stat.BuildSizeMB);
+			}
+
+			{
+				FString AttrName = BaseName + TEXT(".LoadCount");
+				Attributes.Emplace(MoveTemp(AttrName), Stat.LoadCount);
+			}
+
+			{
+				FString AttrName = BaseName + TEXT(".LoadTimeSec");
+				Attributes.Emplace(MoveTemp(AttrName), Stat.LoadTimeSec);
+			}
+
+			{
+				FString AttrName = BaseName + TEXT(".LoadSizeMB");
+				Attributes.Emplace(MoveTemp(AttrName), Stat.LoadSizeMB);
+			}
+		}
+
+		// Grab the DDC summary stats
+		{
+			FDerivedDataCacheSummaryStats DDCSummaryStats;
+			GetDerivedDataCacheRef().GatherSummaryStats(DDCSummaryStats);
+
+			for (const FDerivedDataCacheSummaryStat& Stat : DDCSummaryStats.Stats)
+			{
+				FString FormattedAttrName = "DDC.Summary." + Stat.Key;
+				Attributes.Emplace(FormattedAttrName, Stat.Value);
+			}
+
+			// Grab the Virtualization stats
+			if (IVirtualizationSystem::IsInitialized())
+			{
+				IVirtualizationSystem& System = IVirtualizationSystem::Get();
+
+				FPayloadActivityInfo PayloadActivityInfo = System.GetAccumualtedPayloadActivityInfo();
+
+				const FString BaseName = TEXT("Virtualization");
+
+				{
+					FString AttrName = BaseName + TEXT(".Enabled");
+					Attributes.Emplace(MoveTemp(AttrName), System.IsEnabled());
+				}
+
+				{
+					FString AttrName = BaseName + TEXT(".Cache.TimeSpent");
+					Attributes.Emplace(MoveTemp(AttrName), (double)PayloadActivityInfo.Cache.CyclesSpent * FPlatformTime::GetSecondsPerCycle());
+				}
+
+				{
+					FString AttrName = BaseName + TEXT(".Cache.PayloadCount");
+					Attributes.Emplace(MoveTemp(AttrName), (double)PayloadActivityInfo.Cache.PayloadCount);
+				}
+
+				{
+					FString AttrName = BaseName + TEXT(".Cache.TotalBytes");
+					Attributes.Emplace(MoveTemp(AttrName), (double)PayloadActivityInfo.Cache.TotalBytes);
+				}
+
+				{
+					FString AttrName = BaseName + TEXT(".Push.TimeSpent");
+					Attributes.Emplace(MoveTemp(AttrName), (double)PayloadActivityInfo.Push.CyclesSpent * FPlatformTime::GetSecondsPerCycle());
+				}
+
+				{
+					FString AttrName = BaseName + TEXT(".Push.PayloadCount");
+					Attributes.Emplace(MoveTemp(AttrName), (double)PayloadActivityInfo.Push.PayloadCount);
+				}
+
+				{
+					FString AttrName = BaseName + TEXT(".Push.TotalBytes");
+					Attributes.Emplace(MoveTemp(AttrName), (double)PayloadActivityInfo.Push.TotalBytes);
+				}
+
+				{
+					FString AttrName = BaseName + TEXT(".Pull.TimeSpent");
+					Attributes.Emplace(MoveTemp(AttrName), (double)PayloadActivityInfo.Pull.CyclesSpent * FPlatformTime::GetSecondsPerCycle());
+				}
+
+				{
+					FString AttrName = BaseName + TEXT(".Pull.PayloadCount");
+					Attributes.Emplace(MoveTemp(AttrName), (double)PayloadActivityInfo.Pull.PayloadCount);
+				}
+
+				{
+					FString AttrName = BaseName + TEXT(".Pull.TotalBytes");
+					Attributes.Emplace(MoveTemp(AttrName), (double)PayloadActivityInfo.Pull.TotalBytes);
+				}
+			}
+		}
+
+#if UE_WITH_ZEN
+		if (UE::Zen::IsDefaultServicePresent())
+		{
+			// Grab the Zen summary stats
+			UE::Zen::FZenStats ZenStats;
+			UE::Zen::GetDefaultServiceInstance().GetStats(ZenStats);
+
+			const FString BaseName = TEXT("Zen");
+
+			{
+				FString AttrName = BaseName + TEXT(".Enabled");
+				Attributes.Emplace(MoveTemp(AttrName), ZenStats.IsValid);
+			}
+
+			{
+				FString AttrName = BaseName + TEXT(".Cache.HitRatio");
+				Attributes.Emplace(MoveTemp(AttrName), ZenStats.CacheStats.HitRatio);
+			}
+
+			{
+				FString AttrName = BaseName + TEXT(".Cache.Hits");
+				Attributes.Emplace(MoveTemp(AttrName), ZenStats.CacheStats.Hits);
+			}
+
+			{
+				FString AttrName = BaseName + TEXT(".Cache.Misses");
+				Attributes.Emplace(MoveTemp(AttrName), ZenStats.CacheStats.Misses);
+			}
+
+			{
+				FString AttrName = BaseName + TEXT(".Cache.Size.Disk");
+				Attributes.Emplace(MoveTemp(AttrName), ZenStats.CacheStats.Size.Disk);
+			}
+
+			{
+				FString AttrName = BaseName + TEXT(".Cache.Size.Memory");
+				Attributes.Emplace(MoveTemp(AttrName), ZenStats.CacheStats.Size.Memory);
+			}
+
+			{
+				FString AttrName = BaseName + TEXT(".Cache.UpstreamHits");
+				Attributes.Emplace(MoveTemp(AttrName), ZenStats.CacheStats.UpstreamHits);
+			}
+
+			{
+				FString AttrName = BaseName + TEXT(".Cache.UpstreamRatio");
+				Attributes.Emplace(MoveTemp(AttrName), ZenStats.CacheStats.UpstreamRatio);
+			}
+
+			{
+				FString AttrName = BaseName + TEXT(".Cache.TotalUploadedMB");
+				Attributes.Emplace(MoveTemp(AttrName), ZenStats.UpstreamStats.TotalUploadedMB);
+			}
+
+			{
+				FString AttrName = BaseName + TEXT(".Upstream.TotalDownloadedMB");
+				Attributes.Emplace(MoveTemp(AttrName), ZenStats.UpstreamStats.TotalDownloadedMB);
+			}
+
+			{
+				FString AttrName = BaseName + TEXT(".Upstream.TotalUploadedMB");
+				Attributes.Emplace(MoveTemp(AttrName), ZenStats.UpstreamStats.TotalUploadedMB);
+			}
+
+			{
+				FString AttrName = BaseName + TEXT(".Cas.Size.Large");
+				Attributes.Emplace(MoveTemp(AttrName), ZenStats.CASStats.Size.Large);
+			}
+
+			{
+				FString AttrName = BaseName + TEXT(".Cas.Size.Small");
+				Attributes.Emplace(MoveTemp(AttrName), ZenStats.CASStats.Size.Small);
+			}
+
+			{
+				FString AttrName = BaseName + TEXT(".Cas.Size.Tiny");
+				Attributes.Emplace(MoveTemp(AttrName), ZenStats.CASStats.Size.Tiny);
+			}
+
+			{
+				FString AttrName = BaseName + TEXT(".Cas.Size.Total");
+				Attributes.Emplace(MoveTemp(AttrName), ZenStats.CASStats.Size.Total);
+			}
+		}
+#endif
+
+		// Store it all in the loading event
 		FStudioAnalytics::GetProvider().RecordEvent(TEXT("Core.Loading"), Attributes);
+#endif
 	}
 }
