Header dependency cleanup around HitProxies and BatchedElements

diff --git a/Engine/Plugins/Animation/IKRig/Source/IKRigEditor/Private/RigEditor/IKRigHitProxies.cpp b/Engine/Plugins/Animation/IKRig/Source/IKRigEditor/Private/RigEditor/IKRigHitProxies.cpp
--- a/Engine/Plugins/Animation/IKRig/Source/IKRigEditor/Private/RigEditor/IKRigHitProxies.cpp
+++ b/Engine/Plugins/Animation/IKRig/Source/IKRigEditor/Private/RigEditor/IKRigHitProxies.cpp
@@ -1,6 +1,40 @@
 // Copyright Epic Games, Inc. All Rights Reserved.
 
 #include "RigEditor/IKRigHitProxies.h"
+#include "GenericPlatform/ICursor.h"
 
 IMPLEMENT_HIT_PROXY(HIKRigEditorBoneProxy,HHitProxy);
-IMPLEMENT_HIT_PROXY(HIKRigEditorGoalProxy,HHitProxy);
\ No newline at end of file
+IMPLEMENT_HIT_PROXY(HIKRigEditorGoalProxy,HHitProxy);
+
+HIKRigEditorBoneProxy::HIKRigEditorBoneProxy(const FName& InBoneName)
+	: HHitProxy(HPP_World)
+	, BoneName(InBoneName)
+{
+}
+
+EMouseCursor::Type HIKRigEditorBoneProxy::GetMouseCursor()
+{
+	return EMouseCursor::Crosshairs;
+}
+
+bool HIKRigEditorBoneProxy::AlwaysAllowsTranslucentPrimitives() const
+{
+	return true;
+}
+
+HIKRigEditorGoalProxy::HIKRigEditorGoalProxy(const FName& InGoalName)
+	: HHitProxy(HPP_World)
+	, GoalName(InGoalName)
+{
+}
+
+EMouseCursor::Type HIKRigEditorGoalProxy::GetMouseCursor()
+{
+	return EMouseCursor::Crosshairs;
+}
+
+bool HIKRigEditorGoalProxy::AlwaysAllowsTranslucentPrimitives() const
+{
+	return true;
+}
+
diff --git a/Engine/Plugins/Animation/IKRig/Source/IKRigEditor/Public/RigEditor/IKRigHitProxies.h b/Engine/Plugins/Animation/IKRig/Source/IKRigEditor/Public/RigEditor/IKRigHitProxies.h
--- a/Engine/Plugins/Animation/IKRig/Source/IKRigEditor/Public/RigEditor/IKRigHitProxies.h
+++ b/Engine/Plugins/Animation/IKRig/Source/IKRigEditor/Public/RigEditor/IKRigHitProxies.h
@@ -2,8 +2,12 @@
 
 #pragma once
 
-#include "CoreMinimal.h"
 #include "HitProxies.h"
+#include "UObject/NameTypes.h"
+
+#if UE_ENABLE_INCLUDE_ORDER_DEPRECATED_IN_5_2
+#include "CoreMinimal.h"
+#endif
 
 struct HIKRigEditorBoneProxy : public HHitProxy
 {
@@ -11,19 +15,10 @@ struct HIKRigEditorBoneProxy : public HHitProxy
 
 	FName BoneName;
 
-	HIKRigEditorBoneProxy(const FName& InBoneName)
-		: HHitProxy(HPP_World)
-		, BoneName(InBoneName) {}
+	HIKRigEditorBoneProxy(const FName& InBoneName);
 
-	virtual EMouseCursor::Type GetMouseCursor()
-	{
-		return EMouseCursor::Crosshairs;
-	}
-
-	virtual bool AlwaysAllowsTranslucentPrimitives() const override
-	{
-		return true;
-	}
+	virtual EMouseCursor::Type GetMouseCursor() override;
+	virtual bool AlwaysAllowsTranslucentPrimitives() const override;
 };
 
 struct HIKRigEditorGoalProxy : public HHitProxy
@@ -32,17 +27,8 @@ struct HIKRigEditorGoalProxy : public HHitProxy
 
 	FName GoalName;
 	
-	HIKRigEditorGoalProxy(const FName& InGoalName)
-		: HHitProxy(HPP_World)
-		, GoalName(InGoalName) {}
-
-	virtual EMouseCursor::Type GetMouseCursor()
-	{
-		return EMouseCursor::Crosshairs;
-	}
-
-	virtual bool AlwaysAllowsTranslucentPrimitives() const override
-	{
-		return true;
-	}
+	HIKRigEditorGoalProxy(const FName& InGoalName);
+
+	virtual EMouseCursor::Type GetMouseCursor() override;
+	virtual bool AlwaysAllowsTranslucentPrimitives() const override;
 };
\ No newline at end of file
diff --git a/Engine/Plugins/Experimental/Dataflow/Source/DataflowEnginePlugin/Private/Dataflow/DataflowEngineSceneHitProxies.cpp b/Engine/Plugins/Experimental/Dataflow/Source/DataflowEnginePlugin/Private/Dataflow/DataflowEngineSceneHitProxies.cpp
--- a/Engine/Plugins/Experimental/Dataflow/Source/DataflowEnginePlugin/Private/Dataflow/DataflowEngineSceneHitProxies.cpp
+++ b/Engine/Plugins/Experimental/Dataflow/Source/DataflowEnginePlugin/Private/Dataflow/DataflowEngineSceneHitProxies.cpp
@@ -1,7 +1,41 @@
 // Copyright Epic Games, Inc. All Rights Reserved.
 
 #include "Dataflow/DataflowEngineSceneHitProxies.h"
+#include "GenericPlatform/ICursor.h"
 
 IMPLEMENT_HIT_PROXY(HDataflowDefault, HActor);
 IMPLEMENT_HIT_PROXY(HDataflowNode, HActor);
 IMPLEMENT_HIT_PROXY(HDataflowVertex, HActor);
+
+HDataflowDefault::HDataflowDefault(AActor* InActor, const UPrimitiveComponent* InPrimitiveComponent)
+	: HActor(InActor, InPrimitiveComponent) {}
+
+EMouseCursor::Type HDataflowDefault::GetMouseCursor()
+{
+	return EMouseCursor::Default;
+}
+
+HDataflowNode::HDataflowNode(AActor* InActor, const UPrimitiveComponent* InPrimitiveComponent, FString InNodeName, int32 InGeometryIndex)
+	: HActor(InActor, InPrimitiveComponent)
+{
+	NodeName = InNodeName;
+	GeometryIndex = InGeometryIndex;
+	SectionIndex = GeometryIndex;
+}
+
+EMouseCursor::Type HDataflowNode::GetMouseCursor()
+{
+	return EMouseCursor::Default;
+}
+
+HDataflowVertex::HDataflowVertex(AActor* InActor, const UPrimitiveComponent* InPrimitiveComponent, int32 InVertexIndex)
+	: HActor(InActor, InPrimitiveComponent)
+{
+	SectionIndex = InVertexIndex;
+}
+
+EMouseCursor::Type HDataflowVertex::GetMouseCursor()
+{
+	return EMouseCursor::Crosshairs;
+}
+
diff --git a/Engine/Plugins/Experimental/Dataflow/Source/DataflowEnginePlugin/Public/Dataflow/DataflowEngineSceneHitProxies.h b/Engine/Plugins/Experimental/Dataflow/Source/DataflowEnginePlugin/Public/Dataflow/DataflowEngineSceneHitProxies.h
--- a/Engine/Plugins/Experimental/Dataflow/Source/DataflowEnginePlugin/Public/Dataflow/DataflowEngineSceneHitProxies.h
+++ b/Engine/Plugins/Experimental/Dataflow/Source/DataflowEnginePlugin/Public/Dataflow/DataflowEngineSceneHitProxies.h
@@ -10,19 +10,15 @@
 class UDataflowComponent;
 class UPrimitiveComponent;
 
-
 /** HitProxy with for dataflow actor.
  */
 struct HDataflowDefault : public HActor
 {
 	DECLARE_HIT_PROXY(DATAFLOWENGINEPLUGIN_API)
-	HDataflowDefault(AActor* InActor, const UPrimitiveComponent* InPrimitiveComponent)
-		: HActor(InActor, InPrimitiveComponent){}
 
-	virtual EMouseCursor::Type GetMouseCursor() override
-	{
-		return EMouseCursor::Default;
-	}
+	HDataflowDefault(AActor* InActor, const UPrimitiveComponent* InPrimitiveComponent);
+
+	virtual EMouseCursor::Type GetMouseCursor() override;
 };
 
 
@@ -31,21 +27,13 @@ struct HDataflowDefault : public HActor
 struct HDataflowNode : public HActor
 {
 	DECLARE_HIT_PROXY(DATAFLOWENGINEPLUGIN_API)
+
 	int32 GeometryIndex = INDEX_NONE;
 	FString NodeName = FString("");
 
-	HDataflowNode(AActor* InActor, const UPrimitiveComponent* InPrimitiveComponent, FString InNodeName, int32 InGeometryIndex)
-		: HActor(InActor, InPrimitiveComponent)
-	{
-		NodeName = InNodeName;
-		GeometryIndex = InGeometryIndex;
-		SectionIndex = GeometryIndex;
-	}
+	HDataflowNode(AActor* InActor, const UPrimitiveComponent* InPrimitiveComponent, FString InNodeName, int32 InGeometryIndex);
 
-	virtual EMouseCursor::Type GetMouseCursor() override
-	{
-		return EMouseCursor::Default;
-	}
+	virtual EMouseCursor::Type GetMouseCursor() override;
 };
 
 
@@ -55,14 +43,7 @@ struct HDataflowVertex : public HActor
 {
 	DECLARE_HIT_PROXY(DATAFLOWENGINEPLUGIN_API)
 
-	HDataflowVertex(AActor* InActor, const UPrimitiveComponent* InPrimitiveComponent, int32 InVertexIndex)
-		: HActor(InActor, InPrimitiveComponent)
-	{
-		SectionIndex = InVertexIndex;
-	}
+	HDataflowVertex(AActor* InActor, const UPrimitiveComponent* InPrimitiveComponent, int32 InVertexIndex);
 
-	virtual EMouseCursor::Type GetMouseCursor() override
-	{
-		return EMouseCursor::Crosshairs;
-	}
+	virtual EMouseCursor::Type GetMouseCursor() override;
 };
\ No newline at end of file
diff --git a/Engine/Source/Editor/UnrealEd/Private/UnrealWidget.cpp b/Engine/Source/Editor/UnrealEd/Private/UnrealWidget.cpp
--- a/Engine/Source/Editor/UnrealEd/Private/UnrealWidget.cpp
+++ b/Engine/Source/Editor/UnrealEd/Private/UnrealWidget.cpp
@@ -8,8 +8,10 @@
 #include "EdMode.h"
 #include "EditorModeManager.h"
 #include "EditorViewportClient.h"
+#include "GenericPlatform/ICursor.h"
 #include "Materials/Material.h"
 #include "Materials/MaterialInstanceDynamic.h"
+#include "SceneView.h"
 #include "Settings/LevelEditorViewportSettings.h"
 #include "SnappingUtils.h"
 
@@ -968,3 +970,24 @@ bool FWidget::IsWidgetDisabled() const
 {
 	return EditorModeTools ? (EditorModeTools->IsDefaultModeActive() && GEditor->HasLockedActors()) : false;
 }
+
+HWidgetAxis::HWidgetAxis(EAxisList::Type InAxis, bool InbDisabled, EHitProxyPriority InHitProxy)
+	: HHitProxy(InHitProxy)
+	, Axis(InAxis)
+	, bDisabled(InbDisabled)
+{
+}
+
+EMouseCursor::Type HWidgetAxis::GetMouseCursor()
+{
+	if (bDisabled)
+	{
+		return EMouseCursor::SlashedCircle;
+	}
+	return EMouseCursor::CardinalCross;
+}
+
+bool HWidgetAxis::AlwaysAllowsTranslucentPrimitives() const
+{
+	return true;
+}
diff --git a/Engine/Source/Editor/UnrealEd/Public/UnrealWidget.h b/Engine/Source/Editor/UnrealEd/Public/UnrealWidget.h
--- a/Engine/Source/Editor/UnrealEd/Public/UnrealWidget.h
+++ b/Engine/Source/Editor/UnrealEd/Public/UnrealWidget.h
@@ -478,18 +478,9 @@ struct HWidgetAxis : public HHitProxy
 	EAxisList::Type Axis;
 	uint32 bDisabled : 1;
 
-	HWidgetAxis(EAxisList::Type InAxis, bool InbDisabled = false, EHitProxyPriority InHitProxy = HPP_UI)
-	    : HHitProxy(InHitProxy), Axis(InAxis), bDisabled(InbDisabled)
-	{}
+	UNREALED_API HWidgetAxis(EAxisList::Type InAxis, bool InbDisabled = false, EHitProxyPriority InHitProxy = HPP_UI);
 
-	virtual EMouseCursor::Type GetMouseCursor() override
-	{
-		if (bDisabled)
-		{
-			return EMouseCursor::SlashedCircle;
-		}
-		return EMouseCursor::CardinalCross;
-	}
+	virtual EMouseCursor::Type GetMouseCursor() override;
 
 	/**
 	 * Method that specifies whether the hit proxy *always* allows translucent primitives to be associated with it or not,
@@ -499,8 +490,5 @@ struct HWidgetAxis : public HHitProxy
 	 *
 	 * @return	true if translucent primitives are always allowed with this hit proxy; false otherwise
 	 */
-	virtual bool AlwaysAllowsTranslucentPrimitives() const override
-	{
-		return true;
-	}
+	virtual bool AlwaysAllowsTranslucentPrimitives() const override;
 };
diff --git a/Engine/Source/Runtime/ApplicationCore/Public/GenericPlatform/ICursor.h b/Engine/Source/Runtime/ApplicationCore/Public/GenericPlatform/ICursor.h
--- a/Engine/Source/Runtime/ApplicationCore/Public/GenericPlatform/ICursor.h
+++ b/Engine/Source/Runtime/ApplicationCore/Public/GenericPlatform/ICursor.h
@@ -14,7 +14,7 @@ typedef struct tagRECT RECT;
  */
 namespace EMouseCursor
 {
-	enum Type
+	enum Type : int
 	{
 		/** Causes no mouse cursor to be visible */
 		None,
diff --git a/Engine/Source/Runtime/Engine/Classes/Components/InstancedStaticMeshComponent.h b/Engine/Source/Runtime/Engine/Classes/Components/InstancedStaticMeshComponent.h
--- a/Engine/Source/Runtime/Engine/Classes/Components/InstancedStaticMeshComponent.h
+++ b/Engine/Source/Runtime/Engine/Classes/Components/InstancedStaticMeshComponent.h
@@ -525,11 +525,7 @@ struct HInstancedStaticMeshInstance : public HHitProxy
 
 	virtual void AddReferencedObjects(FReferenceCollector& Collector) override;
 	virtual FTypedElementHandle GetElementHandle() const override;
-
-	virtual EMouseCursor::Type GetMouseCursor() override
-	{
-		return EMouseCursor::Crosshairs;
-	}
+	virtual EMouseCursor::Type GetMouseCursor() override;
 };
 
 /** Used to store lightmap data during RerunConstructionScripts */
diff --git a/Engine/Source/Runtime/Engine/Private/BatchedElements.cpp b/Engine/Source/Runtime/Engine/Private/BatchedElements.cpp
--- a/Engine/Source/Runtime/Engine/Private/BatchedElements.cpp
+++ b/Engine/Source/Runtime/Engine/Private/BatchedElements.cpp
@@ -2,15 +2,113 @@
 
 
 #include "BatchedElements.h"
+
 #include "SimpleElementShaders.h"
 #include "RHIStaticStates.h"
 #include "MeshPassProcessor.h"
+#include "Misc/LargeWorldRenderPosition.h"
+#include "PipelineStateCache.h"
 #include "SceneRelativeViewMatrices.h"
 #include "GlobalRenderResources.h"
 #include "HDRHelper.h"
 
 DEFINE_LOG_CATEGORY_STATIC(LogBatchedElements, Log, All);
 
+FSimpleElementVertex::FSimpleElementVertex() = default;
+
+FSimpleElementVertex::FSimpleElementVertex(const FVector4f& InPosition, const FVector2f& InTextureCoordinate, const FLinearColor& InColor, const FColor& InHitProxyColor)
+	: RelativePosition(InPosition)
+	, TilePosition(ForceInitToZero)
+	, TextureCoordinate(InTextureCoordinate)
+	, Color(InColor)
+	, HitProxyIdColor(InHitProxyColor)
+{
+}
+
+FSimpleElementVertex::FSimpleElementVertex(const FVector4f& InPosition, const FVector2D& InTextureCoordinate, const FLinearColor& InColor, const FColor& InHitProxyColor)
+	: RelativePosition(InPosition)
+	, TilePosition(ForceInitToZero)
+	, TextureCoordinate(InTextureCoordinate)
+	, Color(InColor)
+	, HitProxyIdColor(InHitProxyColor)
+{
+}
+
+FSimpleElementVertex::FSimpleElementVertex(const FVector3f& InPosition, const FVector2D& InTextureCoordinate, const FLinearColor& InColor, const FColor& InHitProxyColor)
+	: RelativePosition(InPosition)
+	, TilePosition(ForceInitToZero)
+	, TextureCoordinate(FVector2f(InTextureCoordinate))
+	, Color(InColor)
+	, HitProxyIdColor(InHitProxyColor)
+{
+}
+
+FSimpleElementVertex::FSimpleElementVertex(const FVector4d& InPosition, const FVector2D& InTextureCoordinate, const FLinearColor& InColor, const FColor& InHitProxyColor)
+	: TextureCoordinate(InTextureCoordinate)
+	, Color(InColor)
+	, HitProxyIdColor(InHitProxyColor)
+{
+	const FLargeWorldRenderPosition AbsolutePosition(InPosition);
+	RelativePosition = FVector4f(AbsolutePosition.GetOffset(), (float)InPosition.W); // Don't bother with LWC W-component
+	TilePosition = FVector4f(AbsolutePosition.GetTile(), 0.0f);
+}
+
+FSimpleElementVertex::FSimpleElementVertex(const FVector3d& InPosition, const FVector2D& InTextureCoordinate, const FLinearColor& InColor, const FColor& InHitProxyColor)
+	: TextureCoordinate(InTextureCoordinate)
+	, Color(InColor)
+	, HitProxyIdColor(InHitProxyColor)
+{
+	const FLargeWorldRenderPosition AbsolutePosition(InPosition);
+	RelativePosition = FVector4f(AbsolutePosition.GetOffset(), 1.0f);
+	TilePosition = FVector4f(AbsolutePosition.GetTile(), 0.0f);
+}
+
+FSimpleElementVertex::FSimpleElementVertex(const FVector4f& InPosition, const FVector2f& InTextureCoordinate, const FLinearColor& InColor, FHitProxyId InHitProxyId)
+	: FSimpleElementVertex(InPosition, InTextureCoordinate, InColor, InHitProxyId.GetColor())
+{
+}
+
+FSimpleElementVertex::FSimpleElementVertex(const FVector4f& InPosition, const FVector2D& InTextureCoordinate, const FLinearColor& InColor, FHitProxyId InHitProxyId)
+	: FSimpleElementVertex(InPosition, InTextureCoordinate, InColor, InHitProxyId.GetColor())
+{
+}
+
+FSimpleElementVertex::FSimpleElementVertex(const FVector3f& InPosition, const FVector2D& InTextureCoordinate, const FLinearColor& InColor, FHitProxyId InHitProxyId)
+	: FSimpleElementVertex(InPosition, InTextureCoordinate, InColor, InHitProxyId.GetColor())
+{
+}
+
+FSimpleElementVertex::FSimpleElementVertex(const FVector4d& InPosition, const FVector2D& InTextureCoordinate, const FLinearColor& InColor, FHitProxyId InHitProxyId)
+	: FSimpleElementVertex(InPosition, InTextureCoordinate, InColor, InHitProxyId.GetColor())
+{
+}
+
+FSimpleElementVertex::FSimpleElementVertex(const FVector3d& InPosition, const FVector2D& InTextureCoordinate, const FLinearColor& InColor, FHitProxyId InHitProxyId)
+	: FSimpleElementVertex(InPosition, InTextureCoordinate, InColor, InHitProxyId.GetColor())
+{
+}
+
+FSimpleElementVertexDeclaration::FSimpleElementVertexDeclaration() = default;
+FSimpleElementVertexDeclaration::FSimpleElementVertexDeclaration(FSimpleElementVertexDeclaration&&) = default;
+FSimpleElementVertexDeclaration::~FSimpleElementVertexDeclaration() = default;
+
+void FSimpleElementVertexDeclaration::InitRHI()
+{
+	FVertexDeclarationElementList Elements;
+	uint16 Stride = sizeof(FSimpleElementVertex);
+	Elements.Add(FVertexElement(0, STRUCT_OFFSET(FSimpleElementVertex, RelativePosition), VET_Float4, 0, Stride));
+	Elements.Add(FVertexElement(0, STRUCT_OFFSET(FSimpleElementVertex, TilePosition), VET_Float4, 1, Stride));
+	Elements.Add(FVertexElement(0, STRUCT_OFFSET(FSimpleElementVertex, TextureCoordinate), VET_Float2, 2, Stride));
+	Elements.Add(FVertexElement(0, STRUCT_OFFSET(FSimpleElementVertex, Color), VET_Float4, 3, Stride));
+	Elements.Add(FVertexElement(0, STRUCT_OFFSET(FSimpleElementVertex, HitProxyIdColor), VET_Color, 4, Stride));
+	VertexDeclarationRHI = PipelineStateCache::GetOrCreateVertexDeclaration(Elements);
+}
+
+void FSimpleElementVertexDeclaration::ReleaseRHI()
+{
+	VertexDeclarationRHI.SafeRelease();
+}
+
 /** The simple element vertex declaration. */
 TGlobalResource<FSimpleElementVertexDeclaration> GSimpleElementVertexDeclaration;
 
@@ -26,6 +124,15 @@ EBlendModeFilter::Type GetBlendModeFilter(ESimpleElementBlendMode BlendMode)
 	}
 }
 
+FBatchedElements::FBatchedElements()
+	: WireTriVerts(/*InNeedsCPUAccess*/true) // Keep vertices on buffer creation
+	, MaxMeshIndicesAllowed(GDrawUPIndexCheckCount / sizeof(int32))
+	// the index buffer is 2 bytes, so make sure we only address 0xFFFF vertices in the index buffer
+	, MaxMeshVerticesAllowed(FMath::Min<uint32>(0xFFFF, GDrawUPVertexCheckCount / sizeof(FSimpleElementVertex)))
+	, bEnableHDREncoding(true)
+{
+}
+
 void FBatchedElements::AddLine(const FVector& Start, const FVector& End, const FLinearColor& Color, FHitProxyId HitProxyId, float Thickness, float DepthBias, bool bScreenSpace)
 {
 	// Ensure the line isn't masked out.  Some legacy code relies on Color.A being ignored.
@@ -56,7 +163,7 @@ void FBatchedElements::AddLine(const FVector& Start, const FVector& End, const F
 		ThickLine->End = End;
 		ThickLine->Thickness = Thickness;
 		ThickLine->Color = OpaqueColor;
-		ThickLine->HitProxyId = HitProxyId;
+		ThickLine->HitProxyColor = HitProxyId.GetColor();
 		ThickLine->DepthBias = DepthBias;
 		ThickLine->bScreenSpace = bScreenSpace;
 	}
@@ -88,7 +195,7 @@ void FBatchedElements::AddTranslucentLine(const FVector& Start, const FVector& E
 		ThickLine->End = End;
 		ThickLine->Thickness = Thickness;
 		ThickLine->Color = Color;
-		ThickLine->HitProxyId = HitProxyId;
+		ThickLine->HitProxyColor = HitProxyId.GetColor();
 		ThickLine->DepthBias = DepthBias;
 		ThickLine->bScreenSpace = bScreenSpace;
 
@@ -105,7 +212,7 @@ void FBatchedElements::AddPoint(const FVector& Position,float Size,const FLinear
 	Point->Position = Position;
 	Point->Size = Size;
 	Point->Color = OpaqueColor.ToFColor(true);
-	Point->HitProxyId = HitProxyId;
+	Point->HitProxyColor = HitProxyId.GetColor();
 }
 
 int32 FBatchedElements::AddVertex(const FVector4& InPosition, const FVector2D& InTextureCoordinate, const FLinearColor& InColor, FHitProxyId HitProxyId)
@@ -332,7 +439,7 @@ void FBatchedElements::AddSprite(
 	Sprite->SizeY = SizeY;
 	Sprite->Texture = Texture;
 	Sprite->Color = Color;
-	Sprite->HitProxyId = HitProxyId;
+	Sprite->HitProxyColor = HitProxyId.GetColor();
 	Sprite->U = U;
 	Sprite->UL = UL == 0.f ? Texture->GetSizeX() : UL;
 	Sprite->V = V;
@@ -746,12 +853,12 @@ void FBatchedElements::DrawPointElements(FRHICommandList& RHICmdList, const FMat
 			const FVector WorldPointX = CameraX * Point.Size / ViewportMajorAxis * TransformedPosition.W;
 			const FVector WorldPointY = CameraY * -Point.Size / ViewportMajorAxis * TransformedPosition.W;
 					
-			PointVertices[VertIdx + 0] = FSimpleElementVertex(Point.Position + WorldPointX - WorldPointY,FVector2D(1,0),Point.Color,Point.HitProxyId);
-			PointVertices[VertIdx + 1] = FSimpleElementVertex(Point.Position + WorldPointX + WorldPointY,FVector2D(1,1),Point.Color,Point.HitProxyId);
-			PointVertices[VertIdx + 2] = FSimpleElementVertex(Point.Position - WorldPointX - WorldPointY,FVector2D(0,0),Point.Color,Point.HitProxyId);
-			PointVertices[VertIdx + 3] = FSimpleElementVertex(Point.Position + WorldPointX + WorldPointY,FVector2D(1,1),Point.Color,Point.HitProxyId);
-			PointVertices[VertIdx + 4] = FSimpleElementVertex(Point.Position - WorldPointX - WorldPointY,FVector2D(0,0),Point.Color,Point.HitProxyId);
-			PointVertices[VertIdx + 5] = FSimpleElementVertex(Point.Position - WorldPointX + WorldPointY,FVector2D(0,1),Point.Color,Point.HitProxyId);
+			PointVertices[VertIdx + 0] = FSimpleElementVertex(Point.Position + WorldPointX - WorldPointY,FVector2D(1,0),Point.Color,Point.HitProxyColor);
+			PointVertices[VertIdx + 1] = FSimpleElementVertex(Point.Position + WorldPointX + WorldPointY,FVector2D(1,1),Point.Color,Point.HitProxyColor);
+			PointVertices[VertIdx + 2] = FSimpleElementVertex(Point.Position - WorldPointX - WorldPointY,FVector2D(0,0),Point.Color,Point.HitProxyColor);
+			PointVertices[VertIdx + 3] = FSimpleElementVertex(Point.Position + WorldPointX + WorldPointY,FVector2D(1,1),Point.Color,Point.HitProxyColor);
+			PointVertices[VertIdx + 4] = FSimpleElementVertex(Point.Position - WorldPointX - WorldPointY,FVector2D(0,0),Point.Color,Point.HitProxyColor);
+			PointVertices[VertIdx + 5] = FSimpleElementVertex(Point.Position - WorldPointX + WorldPointY,FVector2D(0,1),Point.Color,Point.HitProxyColor);
 
 			VertIdx += 6;
 		}
@@ -921,40 +1028,40 @@ bool FBatchedElements::Draw(FRHICommandList& RHICmdList, const FMeshPassProcesso
 						const FVector WorldPointY = CameraY * Thickness * StartW / ViewportSizeX;
 
 						// Begin point
-						ThickVertices[0] = FSimpleElementVertex(Line.Start + WorldPointXS - WorldPointYS,FVector2D(1,0),Line.Color,Line.HitProxyId); // 0S
-						ThickVertices[1] = FSimpleElementVertex(Line.Start + WorldPointXS + WorldPointYS,FVector2D(1,1),Line.Color,Line.HitProxyId); // 1S
-						ThickVertices[2] = FSimpleElementVertex(Line.Start - WorldPointXS - WorldPointYS,FVector2D(0,0),Line.Color,Line.HitProxyId); // 2S
+						ThickVertices[0] = FSimpleElementVertex(Line.Start + WorldPointXS - WorldPointYS,FVector2D(1,0),Line.Color,Line.HitProxyColor); // 0S
+						ThickVertices[1] = FSimpleElementVertex(Line.Start + WorldPointXS + WorldPointYS,FVector2D(1,1),Line.Color,Line.HitProxyColor); // 1S
+						ThickVertices[2] = FSimpleElementVertex(Line.Start - WorldPointXS - WorldPointYS,FVector2D(0,0),Line.Color,Line.HitProxyColor); // 2S
 					
-						ThickVertices[3] = FSimpleElementVertex(Line.Start + WorldPointXS + WorldPointYS,FVector2D(1,1),Line.Color,Line.HitProxyId); // 1S
-						ThickVertices[4] = FSimpleElementVertex(Line.Start - WorldPointXS - WorldPointYS,FVector2D(0,0),Line.Color,Line.HitProxyId); // 2S
-						ThickVertices[5] = FSimpleElementVertex(Line.Start - WorldPointXS + WorldPointYS,FVector2D(0,1),Line.Color,Line.HitProxyId); // 3S
+						ThickVertices[3] = FSimpleElementVertex(Line.Start + WorldPointXS + WorldPointYS,FVector2D(1,1),Line.Color,Line.HitProxyColor); // 1S
+						ThickVertices[4] = FSimpleElementVertex(Line.Start - WorldPointXS - WorldPointYS,FVector2D(0,0),Line.Color,Line.HitProxyColor); // 2S
+						ThickVertices[5] = FSimpleElementVertex(Line.Start - WorldPointXS + WorldPointYS,FVector2D(0,1),Line.Color,Line.HitProxyColor); // 3S
 
 						// Ending point
-						ThickVertices[0+ 6] = FSimpleElementVertex(Line.End + WorldPointXE - WorldPointYE,FVector2D(1,0),Line.Color,Line.HitProxyId); // 0E
-						ThickVertices[1+ 6] = FSimpleElementVertex(Line.End + WorldPointXE + WorldPointYE,FVector2D(1,1),Line.Color,Line.HitProxyId); // 1E
-						ThickVertices[2+ 6] = FSimpleElementVertex(Line.End - WorldPointXE - WorldPointYE,FVector2D(0,0),Line.Color,Line.HitProxyId); // 2E
+						ThickVertices[0+ 6] = FSimpleElementVertex(Line.End + WorldPointXE - WorldPointYE,FVector2D(1,0),Line.Color,Line.HitProxyColor); // 0E
+						ThickVertices[1+ 6] = FSimpleElementVertex(Line.End + WorldPointXE + WorldPointYE,FVector2D(1,1),Line.Color,Line.HitProxyColor); // 1E
+						ThickVertices[2+ 6] = FSimpleElementVertex(Line.End - WorldPointXE - WorldPointYE,FVector2D(0,0),Line.Color,Line.HitProxyColor); // 2E
 																																							  
-						ThickVertices[3+ 6] = FSimpleElementVertex(Line.End + WorldPointXE + WorldPointYE,FVector2D(1,1),Line.Color,Line.HitProxyId); // 1E
-						ThickVertices[4+ 6] = FSimpleElementVertex(Line.End - WorldPointXE - WorldPointYE,FVector2D(0,0),Line.Color,Line.HitProxyId); // 2E
-						ThickVertices[5+ 6] = FSimpleElementVertex(Line.End - WorldPointXE + WorldPointYE,FVector2D(0,1),Line.Color,Line.HitProxyId); // 3E
+						ThickVertices[3+ 6] = FSimpleElementVertex(Line.End + WorldPointXE + WorldPointYE,FVector2D(1,1),Line.Color,Line.HitProxyColor); // 1E
+						ThickVertices[4+ 6] = FSimpleElementVertex(Line.End - WorldPointXE - WorldPointYE,FVector2D(0,0),Line.Color,Line.HitProxyColor); // 2E
+						ThickVertices[5+ 6] = FSimpleElementVertex(Line.End - WorldPointXE + WorldPointYE,FVector2D(0,1),Line.Color,Line.HitProxyColor); // 3E
 
 						// First part of line
-						ThickVertices[0+12] = FSimpleElementVertex(Line.Start - WorldPointXS - WorldPointYS,FVector2D(0,0),Line.Color,Line.HitProxyId); // 2S
-						ThickVertices[1+12] = FSimpleElementVertex(Line.Start + WorldPointXS + WorldPointYS,FVector2D(1,1),Line.Color,Line.HitProxyId); // 1S
-						ThickVertices[2+12] = FSimpleElementVertex(Line.End   - WorldPointXE - WorldPointYE,FVector2D(0,0),Line.Color,Line.HitProxyId); // 2E
+						ThickVertices[0+12] = FSimpleElementVertex(Line.Start - WorldPointXS - WorldPointYS,FVector2D(0,0),Line.Color,Line.HitProxyColor); // 2S
+						ThickVertices[1+12] = FSimpleElementVertex(Line.Start + WorldPointXS + WorldPointYS,FVector2D(1,1),Line.Color,Line.HitProxyColor); // 1S
+						ThickVertices[2+12] = FSimpleElementVertex(Line.End   - WorldPointXE - WorldPointYE,FVector2D(0,0),Line.Color,Line.HitProxyColor); // 2E
 
-						ThickVertices[3+12] = FSimpleElementVertex(Line.Start + WorldPointXS + WorldPointYS,FVector2D(1,1),Line.Color,Line.HitProxyId); // 1S
-						ThickVertices[4+12] = FSimpleElementVertex(Line.End   + WorldPointXE + WorldPointYE,FVector2D(1,1),Line.Color,Line.HitProxyId); // 1E
-						ThickVertices[5+12] = FSimpleElementVertex(Line.End   - WorldPointXE - WorldPointYE,FVector2D(0,0),Line.Color,Line.HitProxyId); // 2E
+						ThickVertices[3+12] = FSimpleElementVertex(Line.Start + WorldPointXS + WorldPointYS,FVector2D(1,1),Line.Color,Line.HitProxyColor); // 1S
+						ThickVertices[4+12] = FSimpleElementVertex(Line.End   + WorldPointXE + WorldPointYE,FVector2D(1,1),Line.Color,Line.HitProxyColor); // 1E
+						ThickVertices[5+12] = FSimpleElementVertex(Line.End   - WorldPointXE - WorldPointYE,FVector2D(0,0),Line.Color,Line.HitProxyColor); // 2E
 
 						// Second part of line
-						ThickVertices[0+18] = FSimpleElementVertex(Line.Start - WorldPointXS + WorldPointYS,FVector2D(0,1),Line.Color,Line.HitProxyId); // 3S
-						ThickVertices[1+18] = FSimpleElementVertex(Line.Start + WorldPointXS - WorldPointYS,FVector2D(1,0),Line.Color,Line.HitProxyId); // 0S
-						ThickVertices[2+18] = FSimpleElementVertex(Line.End   - WorldPointXE + WorldPointYE,FVector2D(0,1),Line.Color,Line.HitProxyId); // 3E
+						ThickVertices[0+18] = FSimpleElementVertex(Line.Start - WorldPointXS + WorldPointYS,FVector2D(0,1),Line.Color,Line.HitProxyColor); // 3S
+						ThickVertices[1+18] = FSimpleElementVertex(Line.Start + WorldPointXS - WorldPointYS,FVector2D(1,0),Line.Color,Line.HitProxyColor); // 0S
+						ThickVertices[2+18] = FSimpleElementVertex(Line.End   - WorldPointXE + WorldPointYE,FVector2D(0,1),Line.Color,Line.HitProxyColor); // 3E
 
-						ThickVertices[3+18] = FSimpleElementVertex(Line.Start + WorldPointXS - WorldPointYS,FVector2D(1,0),Line.Color,Line.HitProxyId); // 0S
-						ThickVertices[4+18] = FSimpleElementVertex(Line.End   + WorldPointXE - WorldPointYE,FVector2D(1,0),Line.Color,Line.HitProxyId); // 0E
-						ThickVertices[5+18] = FSimpleElementVertex(Line.End   - WorldPointXE + WorldPointYE,FVector2D(0,1),Line.Color,Line.HitProxyId); // 3E
+						ThickVertices[3+18] = FSimpleElementVertex(Line.Start + WorldPointXS - WorldPointYS,FVector2D(1,0),Line.Color,Line.HitProxyColor); // 0S
+						ThickVertices[4+18] = FSimpleElementVertex(Line.End   + WorldPointXE - WorldPointYE,FVector2D(1,0),Line.Color,Line.HitProxyColor); // 0E
+						ThickVertices[5+18] = FSimpleElementVertex(Line.End   - WorldPointXE + WorldPointYE,FVector2D(0,1),Line.Color,Line.HitProxyColor); // 3E
 
 						ThickVertices += 24;
 					}
@@ -1077,13 +1184,13 @@ bool FBatchedElements::Draw(FRHICommandList& RHICmdList, const FMeshPassProcesso
 					const float VStart = Sprite.V / Sprite.Texture->GetSizeY();
 					const float VEnd = (Sprite.V + Sprite.VL) / Sprite.Texture->GetSizeY();
 
-					Vertex[0] = FSimpleElementVertex(Sprite.Position + WorldSpriteX - WorldSpriteY, FVector2D(UEnd, VStart), Sprite.Color, Sprite.HitProxyId);
-					Vertex[1] = FSimpleElementVertex(Sprite.Position + WorldSpriteX + WorldSpriteY, FVector2D(UEnd, VEnd), Sprite.Color, Sprite.HitProxyId);
-					Vertex[2] = FSimpleElementVertex(Sprite.Position - WorldSpriteX - WorldSpriteY, FVector2D(UStart, VStart), Sprite.Color, Sprite.HitProxyId);
+					Vertex[0] = FSimpleElementVertex(Sprite.Position + WorldSpriteX - WorldSpriteY, FVector2D(UEnd, VStart), Sprite.Color, Sprite.HitProxyColor);
+					Vertex[1] = FSimpleElementVertex(Sprite.Position + WorldSpriteX + WorldSpriteY, FVector2D(UEnd, VEnd), Sprite.Color, Sprite.HitProxyColor);
+					Vertex[2] = FSimpleElementVertex(Sprite.Position - WorldSpriteX - WorldSpriteY, FVector2D(UStart, VStart), Sprite.Color, Sprite.HitProxyColor);
 
-					Vertex[3] = FSimpleElementVertex(Sprite.Position + WorldSpriteX + WorldSpriteY, FVector2D(UEnd, VEnd), Sprite.Color, Sprite.HitProxyId);
-					Vertex[4] = FSimpleElementVertex(Sprite.Position - WorldSpriteX - WorldSpriteY, FVector2D(UStart, VStart), Sprite.Color, Sprite.HitProxyId);
-					Vertex[5] = FSimpleElementVertex(Sprite.Position - WorldSpriteX + WorldSpriteY, FVector2D(UStart, VEnd), Sprite.Color, Sprite.HitProxyId);
+					Vertex[3] = FSimpleElementVertex(Sprite.Position + WorldSpriteX + WorldSpriteY, FVector2D(UEnd, VEnd), Sprite.Color, Sprite.HitProxyColor);
+					Vertex[4] = FSimpleElementVertex(Sprite.Position - WorldSpriteX - WorldSpriteY, FVector2D(UStart, VStart), Sprite.Color, Sprite.HitProxyColor);
+					Vertex[5] = FSimpleElementVertex(Sprite.Position - WorldSpriteX + WorldSpriteY, FVector2D(UStart, VEnd), Sprite.Color, Sprite.HitProxyColor);
 				}
 				RHICmdList.UnlockBuffer(VertexBufferRHI);
 
diff --git a/Engine/Source/Runtime/Engine/Private/EngineUtils.cpp b/Engine/Source/Runtime/Engine/Private/EngineUtils.cpp
--- a/Engine/Source/Runtime/Engine/Private/EngineUtils.cpp
+++ b/Engine/Source/Runtime/Engine/Private/EngineUtils.cpp
@@ -14,6 +14,7 @@
 #include "Misc/EngineVersion.h"
 #include "GameFramework/PlayerController.h"
 #include "GameFramework/LightWeightInstanceSubsystem.h"
+#include "GenericPlatform/ICursor.h"
 #include "Elements/Framework/EngineElementsLibrary.h"
 #include "Components/PrimitiveComponent.h"
 #include "Engine/Engine.h"
@@ -38,6 +39,17 @@ IMPLEMENT_HIT_PROXY(HTranslucentActor,HActor)
 
 #define LOCTEXT_NAMESPACE "EngineUtils"
 
+void HActor::AddReferencedObjects(FReferenceCollector& Collector)
+{
+	Collector.AddReferencedObject(Actor);
+	Collector.AddReferencedObject(PrimComponent);
+}
+
+EMouseCursor::Type HActor::GetMouseCursor()
+{
+	return EMouseCursor::Crosshairs;
+}
+
 FTypedElementHandle HActor::GetElementHandle() const
 {
 #if WITH_EDITOR
@@ -62,6 +74,16 @@ bool HActor::AlwaysAllowsTranslucentPrimitives() const
 	return false;
 }
 
+EMouseCursor::Type HTranslucentActor::GetMouseCursor()
+{
+	return EMouseCursor::Crosshairs;
+}
+
+bool HTranslucentActor::AlwaysAllowsTranslucentPrimitives() const
+{
+	return true;
+}
+
 #if !UE_BUILD_SHIPPING
 FContentComparisonHelper::FContentComparisonHelper()
 {
diff --git a/Engine/Source/Runtime/Engine/Private/HModel.cpp b/Engine/Source/Runtime/Engine/Private/HModel.cpp
--- a/Engine/Source/Runtime/Engine/Private/HModel.cpp
+++ b/Engine/Source/Runtime/Engine/Private/HModel.cpp
@@ -143,3 +143,8 @@ void HModel::AddReferencedObjects(FReferenceCollector& Collector)
 	Collector.AddReferencedObject(Component);
 	Collector.AddReferencedObject(Model);
 }
+
+EMouseCursor::Type HModel::GetMouseCursor()
+{
+	return EMouseCursor::Crosshairs;
+}
diff --git a/Engine/Source/Runtime/Engine/Private/HitProxies.cpp b/Engine/Source/Runtime/Engine/Private/HitProxies.cpp
--- a/Engine/Source/Runtime/Engine/Private/HitProxies.cpp
+++ b/Engine/Source/Runtime/Engine/Private/HitProxies.cpp
@@ -6,6 +6,7 @@
 
 #include "HitProxies.h"
 #include "Elements/Framework/TypedElementHandle.h"
+#include "GenericPlatform/ICursor.h"
 
 /// @cond DOXYGEN_WARNINGS
 
@@ -126,6 +127,11 @@ bool HHitProxy::IsA(HHitProxyType* TestType) const
 	return bIsInstance;
 }
 
+EMouseCursor::Type HHitProxy::GetMouseCursor()
+{
+	return EMouseCursor::Default;
+}
+
 FTypedElementHandle HHitProxy::GetElementHandle() const
 {
 	return FTypedElementHandle();
diff --git a/Engine/Source/Runtime/Engine/Private/InstancedStaticMesh.cpp b/Engine/Source/Runtime/Engine/Private/InstancedStaticMesh.cpp
--- a/Engine/Source/Runtime/Engine/Private/InstancedStaticMesh.cpp
+++ b/Engine/Source/Runtime/Engine/Private/InstancedStaticMesh.cpp
@@ -236,6 +236,11 @@ FTypedElementHandle HInstancedStaticMeshInstance::GetElementHandle() const
 	return FTypedElementHandle();
 }
 
+EMouseCursor::Type HInstancedStaticMeshInstance::GetMouseCursor()
+{
+	return EMouseCursor::Crosshairs;
+}
+
 FInstanceUpdateCmdBuffer::FInstanceUpdateCmdBuffer()
 	: NumAdds(0),
 	  NumUpdates(0),
diff --git a/Engine/Source/Runtime/Engine/Private/SceneManagement.cpp b/Engine/Source/Runtime/Engine/Private/SceneManagement.cpp
--- a/Engine/Source/Runtime/Engine/Private/SceneManagement.cpp
+++ b/Engine/Source/Runtime/Engine/Private/SceneManagement.cpp
@@ -8,6 +8,7 @@
 #include "RHIStaticStates.h"
 #include "SceneRendering.h"
 #include "SceneCore.h"
+#include "SceneView.h"
 #include "Async/ParallelFor.h"
 #include "LightMap.h"
 #include "ShadowMap.h"
@@ -106,6 +107,27 @@ void FTemporalLODState::UpdateTemporalLODTransition(const FViewInfo& View, float
 	}
 }
 
+FFrozenSceneViewMatricesGuard::FFrozenSceneViewMatricesGuard(FSceneView& SV)
+	: SceneView(SV)
+{
+#if !(UE_BUILD_SHIPPING || UE_BUILD_TEST)
+	if (SceneView.State)
+	{
+		SceneView.State->ActivateFrozenViewMatrices(SceneView);
+	}
+#endif
+}
+
+FFrozenSceneViewMatricesGuard::~FFrozenSceneViewMatricesGuard()
+{
+#if !(UE_BUILD_SHIPPING || UE_BUILD_TEST)
+	if (SceneView.State)
+	{
+		SceneView.State->RestoreUnfrozenViewMatrices(SceneView);
+	}
+#endif
+}
+
 
 IMPLEMENT_STATIC_UNIFORM_BUFFER_SLOT(WorkingColorSpace);
 IMPLEMENT_STATIC_UNIFORM_BUFFER_STRUCT(FWorkingColorSpaceShaderParameters, "WorkingColorSpace", WorkingColorSpace);
diff --git a/Engine/Source/Runtime/Engine/Public/BatchedElements.h b/Engine/Source/Runtime/Engine/Public/BatchedElements.h
--- a/Engine/Source/Runtime/Engine/Public/BatchedElements.h
+++ b/Engine/Source/Runtime/Engine/Public/BatchedElements.h
@@ -6,22 +6,33 @@
 
 #pragma once
 
-#include "CoreMinimal.h"
 #include "Engine/EngineTypes.h"
 #include "Templates/RefCounting.h"
+#include "RenderResource.h"
+#include "Containers/DynamicRHIResourceArray.h"
+
+#if UE_ENABLE_INCLUDE_ORDER_DEPRECATED_IN_5_2
+#include "CoreMinimal.h"
 #include "Misc/LargeWorldRenderPosition.h"
 #include "RHI.h"
-#include "RenderResource.h"
 #include "HitProxies.h"
 #include "SceneView.h"
 #include "StaticBoundShaderState.h"
 #include "PipelineStateCache.h"
-#include "Containers/DynamicRHIResourceArray.h"
+#endif
+
+class FGraphicsPipelineStateInitializer;
+class FHitProxyId;
+class FRHICommandList;
+class FSceneView;
 
 struct FBatchedPoint;
 struct FMeshPassProcessorRenderState;
 struct FRelativeViewMatrices;
 
+enum EBlendMode : int;
+enum ESimpleElementBlendMode : int;
+
 namespace EBlendModeFilter
 {
 	enum Type
@@ -34,7 +45,7 @@ namespace EBlendModeFilter
 };
 
 /** The type used to store batched line vertices. */
-struct FSimpleElementVertex
+struct ENGINE_API FSimpleElementVertex
 {
 	// Store LWC-scale positions per-vertex
 	// Could potentially optimize this by storing a global batch offset, along with relative position per-vertex, but this would be more complicated
@@ -45,81 +56,35 @@ struct FSimpleElementVertex
 	FLinearColor Color;
 	FColor HitProxyIdColor;
 
-	FSimpleElementVertex() {}
-
-	FSimpleElementVertex(const FVector4f& InPosition, const FVector2f& InTextureCoordinate, const FLinearColor& InColor, FHitProxyId InHitProxyId) :
-		RelativePosition(InPosition),
-		TilePosition(ForceInitToZero),
-		TextureCoordinate(InTextureCoordinate),
-		Color(InColor),
-		HitProxyIdColor(InHitProxyId.GetColor())
-	{}
-
-	FSimpleElementVertex(const FVector4f& InPosition, const FVector2D& InTextureCoordinate, const FLinearColor& InColor, FHitProxyId InHitProxyId) :
-		RelativePosition(InPosition),
-		TilePosition(ForceInitToZero),
-		TextureCoordinate(InTextureCoordinate),
-		Color(InColor),
-		HitProxyIdColor(InHitProxyId.GetColor())
-	{}
-
-	FSimpleElementVertex(const FVector3f& InPosition, const FVector2D& InTextureCoordinate, const FLinearColor& InColor, FHitProxyId InHitProxyId) :
-		RelativePosition(InPosition),
-		TilePosition(ForceInitToZero),
-		TextureCoordinate(FVector2f(InTextureCoordinate)),
-		Color(InColor),
-		HitProxyIdColor(InHitProxyId.GetColor())
-	{}
-
-	FSimpleElementVertex(const FVector4d& InPosition, const FVector2D& InTextureCoordinate, const FLinearColor& InColor, FHitProxyId InHitProxyId) :
-		TextureCoordinate(InTextureCoordinate),
-		Color(InColor),
-		HitProxyIdColor(InHitProxyId.GetColor())
-	{
-		const FLargeWorldRenderPosition AbsolutePosition(InPosition);
-		RelativePosition = FVector4f(AbsolutePosition.GetOffset(), (float)InPosition.W); // Don't bother with LWC W-component
-		TilePosition = FVector4f(AbsolutePosition.GetTile(), 0.0f);
-	}
+	FSimpleElementVertex();
 
-	FSimpleElementVertex(const FVector3d& InPosition, const FVector2D& InTextureCoordinate, const FLinearColor& InColor, FHitProxyId InHitProxyId) :
-		TextureCoordinate(InTextureCoordinate),
-		Color(InColor),
-		HitProxyIdColor(InHitProxyId.GetColor())
-	{
-		const FLargeWorldRenderPosition AbsolutePosition(InPosition);
-		RelativePosition = FVector4f(AbsolutePosition.GetOffset(), 1.0f);
-		TilePosition = FVector4f(AbsolutePosition.GetTile(), 0.0f);
-	}
+	FSimpleElementVertex(const FVector4f& InPosition, const FVector2f& InTextureCoordinate, const FLinearColor& InColor, const FColor& InHitProxyColor);
+	FSimpleElementVertex(const FVector4f& InPosition, const FVector2D& InTextureCoordinate, const FLinearColor& InColor, const FColor& InHitProxyColor);
+	FSimpleElementVertex(const FVector3f& InPosition, const FVector2D& InTextureCoordinate, const FLinearColor& InColor, const FColor& InHitProxyColor);
+	FSimpleElementVertex(const FVector4d& InPosition, const FVector2D& InTextureCoordinate, const FLinearColor& InColor, const FColor& InHitProxyColor);
+	FSimpleElementVertex(const FVector3d& InPosition, const FVector2D& InTextureCoordinate, const FLinearColor& InColor, const FColor& InHitProxyColor);
+
+	FSimpleElementVertex(const FVector4f& InPosition, const FVector2f& InTextureCoordinate, const FLinearColor& InColor, FHitProxyId InHitProxyId);
+	FSimpleElementVertex(const FVector4f& InPosition, const FVector2D& InTextureCoordinate, const FLinearColor& InColor, FHitProxyId InHitProxyId);
+	FSimpleElementVertex(const FVector3f& InPosition, const FVector2D& InTextureCoordinate, const FLinearColor& InColor, FHitProxyId InHitProxyId);
+	FSimpleElementVertex(const FVector4d& InPosition, const FVector2D& InTextureCoordinate, const FLinearColor& InColor, FHitProxyId InHitProxyId);
+	FSimpleElementVertex(const FVector3d& InPosition, const FVector2D& InTextureCoordinate, const FLinearColor& InColor, FHitProxyId InHitProxyId);
 };
 
 /**
 * The simple element vertex declaration resource type.
 */
-class FSimpleElementVertexDeclaration : public FRenderResource
+class ENGINE_API FSimpleElementVertexDeclaration : public FRenderResource
 {
 public:
-
 	FVertexDeclarationRHIRef VertexDeclarationRHI;
 
-	// Destructor.
-	virtual ~FSimpleElementVertexDeclaration() {}
+	FSimpleElementVertexDeclaration();
+	FSimpleElementVertexDeclaration(FSimpleElementVertexDeclaration&&);
+	virtual ~FSimpleElementVertexDeclaration();
 
-	virtual void InitRHI() override
-	{
-		FVertexDeclarationElementList Elements;
-		uint16 Stride = sizeof(FSimpleElementVertex);
-		Elements.Add(FVertexElement(0,STRUCT_OFFSET(FSimpleElementVertex, RelativePosition),VET_Float4,0,Stride));
-		Elements.Add(FVertexElement(0,STRUCT_OFFSET(FSimpleElementVertex, TilePosition),VET_Float4,1,Stride));
-		Elements.Add(FVertexElement(0,STRUCT_OFFSET(FSimpleElementVertex,TextureCoordinate),VET_Float2,2,Stride));
-		Elements.Add(FVertexElement(0,STRUCT_OFFSET(FSimpleElementVertex,Color),VET_Float4,3,Stride));
-		Elements.Add(FVertexElement(0,STRUCT_OFFSET(FSimpleElementVertex,HitProxyIdColor),VET_Color,4,Stride));
-		VertexDeclarationRHI = PipelineStateCache::GetOrCreateVertexDeclaration(Elements);
-	}
-
-	virtual void ReleaseRHI() override
-	{
-		VertexDeclarationRHI.SafeRelease();
-	}
+	virtual void InitRHI() override;
+	virtual void ReleaseRHI() override;
 };
 
 /** The simple element vertex declaration. */
@@ -149,14 +114,7 @@ public:
 	/**
 	 * Constructor 
 	 */
-	FBatchedElements()
-		: WireTriVerts(/*InNeedsCPUAccess*/true) // Keep vertices on buffer creation
-		, MaxMeshIndicesAllowed(GDrawUPIndexCheckCount / sizeof(int32))
-		  // the index buffer is 2 bytes, so make sure we only address 0xFFFF vertices in the index buffer
-		, MaxMeshVerticesAllowed(FMath::Min<uint32>(0xFFFF, GDrawUPVertexCheckCount / sizeof(FSimpleElementVertex)))
-		, bEnableHDREncoding(true)
-	{
-	}
+	FBatchedElements();
 
 	/** Adds a line to the batch. Note only SE_BLEND_Opaque will be used for batched line rendering. */
 	void AddLine(const FVector& Start,const FVector& End,const FLinearColor& Color,FHitProxyId HitProxyId, float Thickness = 0.0f, float DepthBias = 0.0f, bool bScreenSpace = false);
@@ -238,8 +196,8 @@ public:
 		float UL,
 		float V,
 		float VL,
-		uint8 BlendMode = SE_BLEND_Masked,
-		float OpacityMaskRefVal = .5f
+		uint8 BlendMode,
+		float OpacityMaskRefVal
 		);
 
 	/**
@@ -304,7 +262,7 @@ private:
 		FVector Position;
 		float Size;
 		FColor Color;
-		FHitProxyId HitProxyId;
+		FColor HitProxyColor;
 	};
 	TArray<FBatchedPoint> Points;
 
@@ -322,7 +280,7 @@ private:
 		FVector End;
 		float Thickness;
 		FLinearColor Color;
-		FHitProxyId HitProxyId;
+		FColor HitProxyColor;
 		float DepthBias;
 		uint32 bScreenSpace;
 	};
@@ -335,7 +293,7 @@ private:
 		float SizeY;
 		const FTexture* Texture;
 		FLinearColor Color;
-		FHitProxyId HitProxyId;
+		FColor HitProxyColor;
 		float U;
 		float UL;
 		float V;
@@ -371,27 +329,6 @@ private:
 	TArray<FBatchedMeshElement,TInlineAllocator<2> > MeshElements;
 	TArray<FSimpleElementVertex,TInlineAllocator<4> > MeshVertices;
 
-	/** bound shader state for the fast path */
-	class FSimpleElementBSSContainer
-	{
-		static const uint32 NumBSS = (uint32)SE_BLEND_RGBA_MASK_START;
-		FGlobalBoundShaderState UnencodedBSS;
-		FGlobalBoundShaderState EncodedBSS[NumBSS];
-	public:
-		FGlobalBoundShaderState& GetBSS(bool bEncoded, ESimpleElementBlendMode BlendMode)
-		{
-			if (bEncoded)
-			{
-				check((uint32)BlendMode < NumBSS);
-				return EncodedBSS[BlendMode];
-			}
-			else
-			{
-				return UnencodedBSS;
-			}
-		}
-	};
-
 	/**
 	 * Sets the appropriate vertex and pixel shader.
 	 */
diff --git a/Engine/Source/Runtime/Engine/Public/EngineUtils.h b/Engine/Source/Runtime/Engine/Public/EngineUtils.h
--- a/Engine/Source/Runtime/Engine/Public/EngineUtils.h
+++ b/Engine/Source/Runtime/Engine/Public/EngineUtils.h
@@ -70,19 +70,9 @@ struct HActor : public HHitProxy
 		, MaterialIndex(InMaterialIndex)
 		{}
 
-	virtual void AddReferencedObjects( FReferenceCollector& Collector ) override
-	{
-		Collector.AddReferencedObject( Actor );
-		Collector.AddReferencedObject( PrimComponent );
-	}
-
-	virtual EMouseCursor::Type GetMouseCursor() override
-	{
-		return EMouseCursor::Crosshairs;
-	}
-
+	ENGINE_API virtual void AddReferencedObjects(FReferenceCollector& Collector) override;
+	ENGINE_API virtual EMouseCursor::Type GetMouseCursor() override;
 	ENGINE_API virtual FTypedElementHandle GetElementHandle() const override;
-
 	ENGINE_API bool AlwaysAllowsTranslucentPrimitives() const override;
 };
 
@@ -141,15 +131,8 @@ struct HTranslucentActor : public HActor
 		: HActor(InActor, InPrimComponent, InPriority)
 		{}
 
-	virtual EMouseCursor::Type GetMouseCursor() override
-	{
-		return EMouseCursor::Crosshairs;
-	}
-
-	virtual bool AlwaysAllowsTranslucentPrimitives() const override
-	{
-		return true;
-	}
+	ENGINE_API virtual EMouseCursor::Type GetMouseCursor() override;
+	ENGINE_API virtual bool AlwaysAllowsTranslucentPrimitives() const override;
 };
 
 
diff --git a/Engine/Source/Runtime/Engine/Public/HModel.h b/Engine/Source/Runtime/Engine/Public/HModel.h
--- a/Engine/Source/Runtime/Engine/Public/HModel.h
+++ b/Engine/Source/Runtime/Engine/Public/HModel.h
@@ -36,10 +36,7 @@ public:
 
 	// HHitProxy interface.
 	virtual void AddReferencedObjects(FReferenceCollector& Collector) override;
-	virtual EMouseCursor::Type GetMouseCursor() override
-	{
-		return EMouseCursor::Crosshairs;
-	}
+	virtual EMouseCursor::Type GetMouseCursor() override;
 
 	// Accessors.
 	UModelComponent* GetModelComponent() const { return Component; }
diff --git a/Engine/Source/Runtime/Engine/Public/HitProxies.h b/Engine/Source/Runtime/Engine/Public/HitProxies.h
--- a/Engine/Source/Runtime/Engine/Public/HitProxies.h
+++ b/Engine/Source/Runtime/Engine/Public/HitProxies.h
@@ -6,15 +6,21 @@
 
 #pragma once
 
-#include "CoreMinimal.h"
 #include "UObject/ObjectMacros.h"
 #include "Templates/RefCounting.h"
+
+#if UE_ENABLE_INCLUDE_ORDER_DEPRECATED_IN_5_2
+#include "CoreMinimal.h"
 #include "GenericPlatform/ICursor.h"
 #include "Elements/Framework/TypedElementHandle.h"
+#endif
 
 #include "HitProxies.generated.h"
 
 class FReferenceCollector;
+struct FTypedElementHandle;
+
+namespace EMouseCursor { enum Type : int; }
 
 /**
  * The priority a hit proxy has when choosing between several hit proxies near the point the user clicked.
@@ -165,10 +171,7 @@ public:
 	/**
 		Override to change the mouse based on what it is hovering over.
 	*/
-	virtual EMouseCursor::Type GetMouseCursor()
-	{
-		return EMouseCursor::Default;
-	}
+	ENGINE_API virtual EMouseCursor::Type GetMouseCursor();
 
 	/**
 	 * Method that specifies whether the hit proxy *always* allows translucent primitives to be associated with it or not,
diff --git a/Engine/Source/Runtime/Engine/Public/SceneManagement.h b/Engine/Source/Runtime/Engine/Public/SceneManagement.h
--- a/Engine/Source/Runtime/Engine/Public/SceneManagement.h
+++ b/Engine/Source/Runtime/Engine/Public/SceneManagement.h
@@ -32,6 +32,7 @@
 #include "LightmapUniformShaderParameters.h"
 #include "DynamicBufferAllocator.h"
 #include "Rendering/SkyAtmosphereCommonData.h"
+#include "Math/SHMath.h"
 
 class FCanvas;
 class FGlobalDynamicIndexBuffer;
@@ -43,6 +44,7 @@ class FLightSceneInfo;
 class FLightSceneProxy;
 class FPrimitiveSceneInfo;
 class FPrimitiveSceneProxy;
+class FScene;
 class FSceneViewState;
 class FShadowMap;
 class FStaticMeshRenderData;
@@ -62,6 +64,9 @@ class ULightMapVirtualTexture2D;
 class FGPUScenePrimitiveCollector;
 class FVirtualShadowMapArrayCacheManager;
 class FRayTracingGeometry;
+struct FViewMatrices;
+struct FEngineShowFlags;
+class FViewport;
 
 namespace UE { namespace Color { class FColorSpace; } }
 
@@ -310,29 +315,11 @@ private:
 	friend class FScene;
 };
 
-class FFrozenSceneViewMatricesGuard
+class ENGINE_API FFrozenSceneViewMatricesGuard
 {
 public:
-	FFrozenSceneViewMatricesGuard(FSceneView& SV)
-		: SceneView(SV)
-	{
-#if !(UE_BUILD_SHIPPING || UE_BUILD_TEST)
-		if (SceneView.State)
-		{
-			SceneView.State->ActivateFrozenViewMatrices(SceneView);
-		}
-#endif
-	}
-
-	~FFrozenSceneViewMatricesGuard()
-	{
-#if !(UE_BUILD_SHIPPING || UE_BUILD_TEST)
-		if (SceneView.State)
-		{
-			SceneView.State->RestoreUnfrozenViewMatrices(SceneView);
-		}
-#endif
-	}
+	FFrozenSceneViewMatricesGuard(FSceneView& SV);
+	~FFrozenSceneViewMatricesGuard();
 
 private:
 	FSceneView& SceneView;
diff --git a/Engine/Source/Runtime/Landscape/Private/LandscapeSplineProxies.cpp b/Engine/Source/Runtime/Landscape/Private/LandscapeSplineProxies.cpp
--- a/Engine/Source/Runtime/Landscape/Private/LandscapeSplineProxies.cpp
+++ b/Engine/Source/Runtime/Landscape/Private/LandscapeSplineProxies.cpp
@@ -1,9 +1,62 @@
 // Copyright Epic Games, Inc. All Rights Reserved.
 
 #include "LandscapeSplineProxies.h"
+
+#include "GenericPlatform/ICursor.h"
+#include "LandscapeSplineControlPoint.h"
+#include "LandscapeSplinesComponent.h"
 #include "LandscapeSplineSegment.h"
 
+HLandscapeSplineProxy::HLandscapeSplineProxy(ULandscapeSplinesComponent* SplineComponent, EHitProxyPriority InPriority)
+	: HActor(SplineComponent->GetOwner(), SplineComponent, InPriority)
+{
+}
+
+EMouseCursor::Type HLandscapeSplineProxy::GetMouseCursor()
+{
+	return EMouseCursor::Crosshairs;
+}
+
+HLandscapeSplineProxy_Segment::HLandscapeSplineProxy_Segment(class ULandscapeSplineSegment* InSplineSegment)
+	: HLandscapeSplineProxy(InSplineSegment->GetOuterULandscapeSplinesComponent())
+	, SplineSegment(InSplineSegment)
+{
+}
+
+void HLandscapeSplineProxy_Segment::AddReferencedObjects(FReferenceCollector& Collector)
+{
+	Collector.AddReferencedObject(SplineSegment);
+}
+
+HLandscapeSplineProxy_ControlPoint::HLandscapeSplineProxy_ControlPoint(class ULandscapeSplineControlPoint* InControlPoint)
+	: HLandscapeSplineProxy(InControlPoint->GetOuterULandscapeSplinesComponent(), HPP_Foreground)
+	, ControlPoint(InControlPoint)
+{
+}
+
+void HLandscapeSplineProxy_ControlPoint::AddReferencedObjects(FReferenceCollector& Collector)
+{
+	Collector.AddReferencedObject(ControlPoint);
+}
+
+HLandscapeSplineProxy_Tangent::HLandscapeSplineProxy_Tangent(class ULandscapeSplineSegment* InSplineSegment, bool InEnd)
+	: HLandscapeSplineProxy(InSplineSegment->GetOuterULandscapeSplinesComponent(), HPP_UI)
+	, SplineSegment(InSplineSegment)
+	, End(InEnd)
+{
+}
+
 LANDSCAPE_API void HLandscapeSplineProxy_Tangent::Serialize(FArchive& Ar)
 {
 	Ar << SplineSegment;
 }
+
+EMouseCursor::Type HLandscapeSplineProxy_Tangent::GetMouseCursor()
+{
+	return EMouseCursor::CardinalCross;
+}
+
+void HLandscapeSplineProxy_Tangent::AddReferencedObjects(FReferenceCollector& Collector)
+{
+	Collector.AddReferencedObject(SplineSegment);
+}
diff --git a/Engine/Source/Runtime/Landscape/Public/LandscapeSplineProxies.h b/Engine/Source/Runtime/Landscape/Public/LandscapeSplineProxies.h
--- a/Engine/Source/Runtime/Landscape/Public/LandscapeSplineProxies.h
+++ b/Engine/Source/Runtime/Landscape/Public/LandscapeSplineProxies.h
@@ -2,13 +2,18 @@
 
 #pragma once
 
+#include "EngineUtils.h"
+
+#if UE_ENABLE_INCLUDE_ORDER_DEPRECATED_IN_5_2
 #include "CoreMinimal.h"
 #include "UObject/UObjectGlobals.h"
-#include "EngineUtils.h"
 #include "LandscapeSplinesComponent.h"
 #include "LandscapeSplineControlPoint.h"
 #include "LandscapeSplineSegment.h"
+#endif
 
+class ULandscapeSplineControlPoint;
+class ULandscapeSplinesComponent;
 class ULandscapeSplineSegment;
 
 //////////////////////////////////////////////////////////////////////////
@@ -18,48 +23,28 @@ struct HLandscapeSplineProxy : public HActor
 {
 	DECLARE_HIT_PROXY( LANDSCAPE_API );
 
-	HLandscapeSplineProxy(ULandscapeSplinesComponent* SplineComponent, EHitProxyPriority InPriority = HPP_Wireframe) :
-		HActor(SplineComponent->GetOwner(), SplineComponent, InPriority)
-	{
-	}
-	virtual EMouseCursor::Type GetMouseCursor() override
-	{
-		return EMouseCursor::Crosshairs;
-	}
+	HLandscapeSplineProxy(ULandscapeSplinesComponent* SplineComponent, EHitProxyPriority InPriority = HPP_Wireframe);
+	virtual EMouseCursor::Type GetMouseCursor() override;
 };
 
 struct HLandscapeSplineProxy_Segment : public HLandscapeSplineProxy
 {
 	DECLARE_HIT_PROXY( LANDSCAPE_API );
 
-	class ULandscapeSplineSegment* SplineSegment;
-
-	HLandscapeSplineProxy_Segment(class ULandscapeSplineSegment* InSplineSegment) :
-		HLandscapeSplineProxy(InSplineSegment->GetOuterULandscapeSplinesComponent()),
-		SplineSegment(InSplineSegment)
-	{
-	}
-	virtual void AddReferencedObjects( FReferenceCollector& Collector ) override
-	{
-		Collector.AddReferencedObject( SplineSegment );
-	}
+	ULandscapeSplineSegment* SplineSegment;
+
+	HLandscapeSplineProxy_Segment(ULandscapeSplineSegment* InSplineSegment);
+	virtual void AddReferencedObjects(FReferenceCollector& Collector) override;
 };
 
 struct HLandscapeSplineProxy_ControlPoint : public HLandscapeSplineProxy
 {
 	DECLARE_HIT_PROXY( LANDSCAPE_API );
 
-	class ULandscapeSplineControlPoint* ControlPoint;
-
-	HLandscapeSplineProxy_ControlPoint(class ULandscapeSplineControlPoint* InControlPoint) :
-		HLandscapeSplineProxy(InControlPoint->GetOuterULandscapeSplinesComponent(), HPP_Foreground),
-		ControlPoint(InControlPoint)
-	{
-	}
-	virtual void AddReferencedObjects( FReferenceCollector& Collector ) override
-	{
-		Collector.AddReferencedObject( ControlPoint );
-	}
+	ULandscapeSplineControlPoint* ControlPoint;
+
+	HLandscapeSplineProxy_ControlPoint(ULandscapeSplineControlPoint* InControlPoint);
+	virtual void AddReferencedObjects(FReferenceCollector& Collector) override;
 };
 
 struct HLandscapeSplineProxy_Tangent : public HLandscapeSplineProxy
@@ -69,21 +54,9 @@ struct HLandscapeSplineProxy_Tangent : public HLandscapeSplineProxy
 	ULandscapeSplineSegment* SplineSegment;
 	uint32 End:1;
 
-	HLandscapeSplineProxy_Tangent(class ULandscapeSplineSegment* InSplineSegment, bool InEnd) :
-		HLandscapeSplineProxy(InSplineSegment->GetOuterULandscapeSplinesComponent(), HPP_UI),
-		SplineSegment(InSplineSegment),
-		End(InEnd)
-	{
-	}
+	LANDSCAPE_API HLandscapeSplineProxy_Tangent(ULandscapeSplineSegment* InSplineSegment, bool InEnd);
 	LANDSCAPE_API virtual void Serialize(FArchive& Ar);
 
-	virtual EMouseCursor::Type GetMouseCursor() override
-	{
-		return EMouseCursor::CardinalCross;
-	}
-
-	virtual void AddReferencedObjects( FReferenceCollector& Collector ) override
-	{
-		Collector.AddReferencedObject( SplineSegment );
-	}
+	virtual EMouseCursor::Type GetMouseCursor() override;
+	virtual void AddReferencedObjects(FReferenceCollector& Collector) override;
 };
diff --git a/Engine/Source/Runtime/MovieSceneTracks/Private/MovieSceneHitProxy.cpp b/Engine/Source/Runtime/MovieSceneTracks/Private/MovieSceneHitProxy.cpp
--- a/Engine/Source/Runtime/MovieSceneTracks/Private/MovieSceneHitProxy.cpp
+++ b/Engine/Source/Runtime/MovieSceneTracks/Private/MovieSceneHitProxy.cpp
@@ -1,5 +1,17 @@
 // Copyright Epic Games, Inc. All Rights Reserved.
 
 #include "MovieSceneHitProxy.h"
+#include "GenericPlatform/ICursor.h"
 
 IMPLEMENT_HIT_PROXY(HMovieSceneKeyProxy, HHitProxy);
+
+HMovieSceneKeyProxy::HMovieSceneKeyProxy(class UMovieSceneTrack* InTrack, const FTrajectoryKey& InKey)
+	: HHitProxy(HPP_UI)
+	, MovieSceneTrack(InTrack)
+	, Key(InKey)
+{}
+
+EMouseCursor::Type HMovieSceneKeyProxy::GetMouseCursor()
+{
+	return EMouseCursor::Crosshairs;
+}
diff --git a/Engine/Source/Runtime/MovieSceneTracks/Private/MovieSceneSkeletalAnimationRootHitProxy.cpp b/Engine/Source/Runtime/MovieSceneTracks/Private/MovieSceneSkeletalAnimationRootHitProxy.cpp
--- a/Engine/Source/Runtime/MovieSceneTracks/Private/MovieSceneSkeletalAnimationRootHitProxy.cpp
+++ b/Engine/Source/Runtime/MovieSceneTracks/Private/MovieSceneSkeletalAnimationRootHitProxy.cpp
@@ -1,6 +1,20 @@
 // Copyright Epic Games, Inc. All Rights Reserved.
 
 #include "MovieSceneSkeletalAnimationRootHitProxy.h"
+#include "Components/SkeletalMeshComponent.h"
+#include "GenericPlatform/ICursor.h"
+#include "Sections/MovieSceneSkeletalAnimationSection.h"
 
 IMPLEMENT_HIT_PROXY(HMovieSceneSkeletalAnimationRootHitProxy, HHitProxy);
 
+HMovieSceneSkeletalAnimationRootHitProxy::HMovieSceneSkeletalAnimationRootHitProxy(UMovieSceneSkeletalAnimationSection* InSection, USkeletalMeshComponent* InSkelMeshComp)
+	: HHitProxy(HPP_UI)
+	, AnimSection(InSection)
+	, SkelMeshComp(InSkelMeshComp)
+{
+}
+
+EMouseCursor::Type HMovieSceneSkeletalAnimationRootHitProxy::GetMouseCursor()
+{
+	return EMouseCursor::CardinalCross;
+}
diff --git a/Engine/Source/Runtime/MovieSceneTracks/Public/MovieSceneHitProxy.h b/Engine/Source/Runtime/MovieSceneTracks/Public/MovieSceneHitProxy.h
--- a/Engine/Source/Runtime/MovieSceneTracks/Public/MovieSceneHitProxy.h
+++ b/Engine/Source/Runtime/MovieSceneTracks/Public/MovieSceneHitProxy.h
@@ -2,11 +2,17 @@
 
 #pragma once
 
-#include "CoreMinimal.h"
 #include "HitProxies.h"
-#include "Sections/MovieScene3DTransformSection.h"
+#include "UObject/WeakObjectPtr.h"
 #include "Tracks/MovieScene3DTransformTrack.h"
+
+#if UE_ENABLE_INCLUDE_ORDER_DEPRECATED_IN_5_2
+#include "CoreMinimal.h"
+#include "Sections/MovieScene3DTransformSection.h"
 #include "MovieSceneTrack.h"
+#endif
+
+class UMovieSceneTrack;
 
 struct HMovieSceneKeyProxy : public HHitProxy
 {
@@ -14,18 +20,12 @@ struct HMovieSceneKeyProxy : public HHitProxy
 
 	/** The track that contains the section */
 	TWeakObjectPtr<UMovieSceneTrack> MovieSceneTrack;
+
 	/** The trajectory key data */
 	FTrajectoryKey Key;
 
-	HMovieSceneKeyProxy(class UMovieSceneTrack* InTrack, const FTrajectoryKey& InKey) :
-		HHitProxy(HPP_UI),
-		MovieSceneTrack(InTrack),
-		Key(InKey)
-	{}
+	MOVIESCENETRACKS_API HMovieSceneKeyProxy(class UMovieSceneTrack* InTrack, const FTrajectoryKey& InKey);
 
-	virtual EMouseCursor::Type GetMouseCursor() override
-	{
-		return EMouseCursor::Crosshairs;
-	}
+	virtual EMouseCursor::Type GetMouseCursor() override;
 };
 
diff --git a/Engine/Source/Runtime/MovieSceneTracks/Public/MovieSceneSkeletalAnimationRootHitProxy.h b/Engine/Source/Runtime/MovieSceneTracks/Public/MovieSceneSkeletalAnimationRootHitProxy.h
--- a/Engine/Source/Runtime/MovieSceneTracks/Public/MovieSceneSkeletalAnimationRootHitProxy.h
+++ b/Engine/Source/Runtime/MovieSceneTracks/Public/MovieSceneSkeletalAnimationRootHitProxy.h
@@ -2,11 +2,15 @@
 
 #pragma once
 
-#include "CoreMinimal.h"
 #include "HitProxies.h"
+#include "UObject/WeakObjectPtr.h"
+
+#if UE_ENABLE_INCLUDE_ORDER_DEPRECATED_IN_5_2
+#include "CoreMinimal.h"
 #include "MovieSceneTrack.h"
 #include "Components/SkeletalMeshComponent.h"
 #include "Sections/MovieSceneSkeletalAnimationSection.h"
+#endif
 
 class UMovieSceneSkeletalAnimationSection;
 class  USkeletalMeshComponent;
@@ -16,18 +20,10 @@ struct HMovieSceneSkeletalAnimationRootHitProxy : public HHitProxy
 	DECLARE_HIT_PROXY( MOVIESCENETRACKS_API );
 
 	TWeakObjectPtr<UMovieSceneSkeletalAnimationSection> AnimSection;
-	
 	TWeakObjectPtr<USkeletalMeshComponent> SkelMeshComp;
 
-	HMovieSceneSkeletalAnimationRootHitProxy(UMovieSceneSkeletalAnimationSection* InSection, USkeletalMeshComponent* InSkelMeshComp) :
-		HHitProxy(HPP_UI),
-		AnimSection(InSection),
-		SkelMeshComp(InSkelMeshComp)
-	{}
+	MOVIESCENETRACKS_API HMovieSceneSkeletalAnimationRootHitProxy(UMovieSceneSkeletalAnimationSection* InSection, USkeletalMeshComponent* InSkelMeshComp);
 
-	virtual EMouseCursor::Type GetMouseCursor() override
-	{
-		return EMouseCursor::CardinalCross;
-	}
+	virtual EMouseCursor::Type GetMouseCursor() override;
 };
 
