Restore some changes clobbered by a robomerge.
#rnx

diff --git a/Engine/Source/Runtime/Engine/Private/Analytics/StudioAnalytics.cpp b/Engine/Source/Runtime/Engine/Private/Analytics/StudioAnalytics.cpp
--- a/Engine/Source/Runtime/Engine/Private/Analytics/StudioAnalytics.cpp
+++ b/Engine/Source/Runtime/Engine/Private/Analytics/StudioAnalytics.cpp
@@ -1,24 +1,10 @@
 // Copyright Epic Games, Inc. All Rights Reserved.
 
 #include "StudioAnalytics.h"
-#include "Misc/Guid.h"
-#include "Stats/Stats.h"
-#include "Misc/ConfigCacheIni.h"
-#include "EngineGlobals.h"
-#include "Engine/Engine.h"
-#include "Misc/EngineBuildSettings.h"
-#include "AnalyticsBuildType.h"
-#include "AnalyticsEventAttribute.h"
+#include "HAL/PlatformTime.h"
 #include "IAnalyticsProviderET.h"
-#include "Analytics.h"
-#include "AnalyticsET.h"
-#include "GeneralProjectSettings.h"
-#include "Misc/EngineVersion.h"
-#include "RHI.h"
-#include "GenericPlatform/GenericPlatformCrashContext.h"
-#include "Interfaces/IAnalyticsProvider.h"
-#include "Templates/SharedPointer.h"
 #include "HAL/PlatformProcess.h"
+#include "HAL/Thread.h"
 #include "DerivedDataCacheInterface.h"
 #include "DerivedDataCacheUsageStats.h"
 #include "Virtualization/VirtualizationSystem.h"
