Fix vertex color selection when sampling with a GPU pass through vertex factory.

diff --git a/Engine/Source/Runtime/Engine/Private/GPUSkinVertexFactory.cpp b/Engine/Source/Runtime/Engine/Private/GPUSkinVertexFactory.cpp
--- a/Engine/Source/Runtime/Engine/Private/GPUSkinVertexFactory.cpp
+++ b/Engine/Source/Runtime/Engine/Private/GPUSkinVertexFactory.cpp
@@ -1318,6 +1318,9 @@ void FGPUSkinPassthroughVertexFactory::OverrideAttributeData()
 		Data.ColorComponent.Type = VET_Color;
 		Data.ColorComponent.VertexStreamUsage = EVertexStreamUsage::Overridden | EVertexStreamUsage::ManualFetch;
 		Data.ColorComponent.Stride = sizeof(uint32);
+
+		// Set mask to allow full vertex indexing in vertex shader.
+		Data.ColorIndexMask = ~0u;
 	}
 
 	// TexCoord vertex attributes are written in pairs.
