Fix the component hierarchy display in the Actor details tab to remain static/deterministic in terms of tree node order after its Blueprint class is compiled and/or Actor reconstruction occurs.

diff --git a/Engine/Source/Editor/SubobjectDataInterface/Private/SubobjectDataSubsystem.cpp b/Engine/Source/Editor/SubobjectDataInterface/Private/SubobjectDataSubsystem.cpp
--- a/Engine/Source/Editor/SubobjectDataInterface/Private/SubobjectDataSubsystem.cpp
+++ b/Engine/Source/Editor/SubobjectDataInterface/Private/SubobjectDataSubsystem.cpp
@@ -403,7 +403,7 @@ void USubobjectDataSubsystem::GatherSubobjectData(UObject* Context, TArray<FSubo
 		}
 
 		// Sort components by type (always put scene components first in the tree)
-		ComponentsToAdd.Sort([](const UActorComponent& A, const UActorComponent& /* B */)
+		ComponentsToAdd.StableSort([](const UActorComponent& A, const UActorComponent& /* B */)
 		{
 			return A.IsA<USceneComponent>();
 		});
diff --git a/Engine/Source/Runtime/Engine/Private/ActorConstruction.cpp b/Engine/Source/Runtime/Engine/Private/ActorConstruction.cpp
--- a/Engine/Source/Runtime/Engine/Private/ActorConstruction.cpp
+++ b/Engine/Source/Runtime/Engine/Private/ActorConstruction.cpp
@@ -201,6 +201,13 @@ void AActor::DestroyConstructedComponents()
 	{
 		GetPackage()->MarkPackageDirty();
 	}
+
+	// When a constructed component is destroyed, it is removed from this set. We compact the set to ensure any newly-constructed components
+	// will get added back into the set contiguously, so that the iteration order remains stable after re-running actor construction scripts.
+	if (PreviouslyAttachedComponents.Num() > OwnedComponents.Num())
+	{
+		OwnedComponents.CompactStable();
+	}
 }
 
 
