Texture.h make Blueprint query functions non-pure so they are cached

diff --git a/Engine/Source/Runtime/Engine/Classes/Engine/Texture.h b/Engine/Source/Runtime/Engine/Classes/Engine/Texture.h
--- a/Engine/Source/Runtime/Engine/Classes/Engine/Texture.h
+++ b/Engine/Source/Runtime/Engine/Classes/Engine/Texture.h
@@ -1361,7 +1361,7 @@ public:
 	 * Does not cause texture source to be loaded, queries cached values.
 	 * Returns zero for error.
 	 */
-	UFUNCTION(BlueprintCallable, meta = (DisplayName = "GetTextureSourceDiskAndMemorySize"), Category = "Rendering|Texture")
+	UFUNCTION(BlueprintCallable, BlueprintPure=false, meta = (DisplayName = "GetTextureSourceDiskAndMemorySize"), Category = "Rendering|Texture")
 	void Blueprint_GetTextureSourceDiskAndMemorySize(int64 & OutDiskSize,int64 & OutMemorySize) const;
 
 	/**
@@ -1370,7 +1370,7 @@ public:
 	 * Causes texture source data to be loaded, is computed by scanning pixels when called.
 	 * Will set Min=Max=zero and return false on failure
 	 */
-	UFUNCTION(BlueprintCallable, Category = "Rendering|Texture")
+	UFUNCTION(BlueprintCallable, BlueprintPure=false, Category = "Rendering|Texture")
 	bool ComputeTextureSourceChannelMinMax(FLinearColor & OutColorMin, FLinearColor & OutColorMax) const;
 
 private:
