World Partition: hold a reference to newly added actors when initially saving them so they avoid getting unloaded.

diff --git a/Engine/Source/Runtime/Engine/Private/WorldPartition/WorldPartition.cpp b/Engine/Source/Runtime/Engine/Private/WorldPartition/WorldPartition.cpp
--- a/Engine/Source/Runtime/Engine/Private/WorldPartition/WorldPartition.cpp
+++ b/Engine/Source/Runtime/Engine/Private/WorldPartition/WorldPartition.cpp
@@ -1118,6 +1118,11 @@ void UWorldPartition::OnActorDescAdded(FWorldPartitionActorDesc* NewActorDesc)
 
 	HashActorDesc(NewActorDesc);
 
+	if (AActor* NewActor = NewActorDesc->GetActor())
+	{
+		DirtyActors.Add(FWorldPartitionReference(NewActorDesc->GetContainer(), NewActorDesc->GetGuid()), NewActor);
+	}
+
 	if (ForceLoadedActors)
 	{
 		ForceLoadedActors->AddActors({ NewActorDesc->GetGuid() });
