Bug fix: When moving actors with welded actors:
- All welded bodies were not excluded from the sweep
- The hit results were not sorted, even though that is assumes by code that follows
- TestLevel:
QAGame/WeldedSweepTest

#[fyi] Michael.Lentine

diff --git a/Engine/Source/Runtime/Engine/Private/Collision/WorldCollision.cpp b/Engine/Source/Runtime/Engine/Private/Collision/WorldCollision.cpp
--- a/Engine/Source/Runtime/Engine/Private/Collision/WorldCollision.cpp
+++ b/Engine/Source/Runtime/Engine/Private/Collision/WorldCollision.cpp
@@ -11,6 +11,7 @@
 #include "Collision.h"
 #include "PhysicsEngine/PhysicsObjectExternalInterface.h"
 #include "Chaos/ImplicitObject.h"
+#include "Collision\CollisionConversions.h"
 
 using namespace PhysicsInterfaceTypes;
 
@@ -462,7 +463,26 @@ bool UWorld::ComponentSweepMultiByChannel(TArray<struct FHitResult>& OutHits, cl
 	TArray<struct FOverlapResult> TempOverlaps;
 	for (Chaos::FPhysicsObjectHandle Object : PhysicsObjects)
 	{
+		FComponentQueryParams ParamsCopy{ Params };
+		TSet<uint32> ActorsToExclude;
+		for (uint32 ActorID : ParamsCopy.GetIgnoredActors())
+		{
+			ActorsToExclude.Add(ActorID);
+		}		
+		ParamsCopy.ClearIgnoredActors(); // This will be populated a bit later
+		// All actors pointed to by shapes should be ignored (This deals with welded Actors)
 		TArray<Chaos::FPerShapeData*> Shapes = Interface->GetAllShapes({ &Object, 1 });
+		for (Chaos::FPerShapeData* Shape : Shapes)
+		{
+			FCollisionFilterData ShapeFilter = ChaosInterface::GetQueryFilterData(*Shape);
+			uint32 ShapeActorID = ShapeFilter.Word0; // Unique Actor ID is saved in word 0
+			ActorsToExclude.Add(ShapeActorID);
+		}
+
+		for (uint32 ActorID : ActorsToExclude)
+		{
+			ParamsCopy.AddIgnoredActor(ActorID);
+		}
 
 		for (Chaos::FPerShapeData* Shape : Shapes)
 		{
@@ -479,11 +499,12 @@ bool UWorld::ComponentSweepMultiByChannel(TArray<struct FHitResult>& OutHits, cl
 
 			FPhysicsGeometryCollection GeomCollection = FPhysicsInterface::GetGeometryCollection(ShapeHandle);
 			TArray<FHitResult> TmpHits;
-			if (FPhysicsInterface::GeomSweepMulti(this, GeomCollection, Rot, TmpHits, Start, End, TraceChannel, Params, FCollisionResponseParams(PrimComp->GetCollisionResponseToChannels())))
+			if (FPhysicsInterface::GeomSweepMulti(this, GeomCollection, Rot, TmpHits, Start, End, TraceChannel, ParamsCopy, FCollisionResponseParams(PrimComp->GetCollisionResponseToChannels())))
 			{
 				bHaveBlockingHit = true;
 			}
 			OutHits.Append(TmpHits);	//todo: should these be made unique?
+			OutHits.Sort(FCompareFHitResultTime());
 		}
 	}
 
