World Partition Streaming Validation: only do parent actor validation on the first validation pass to avoid infinite validation passes (and backout 23544379).

diff --git a/Engine/Source/Runtime/Engine/Private/WorldPartition/WorldPartitionStreamingGeneration.cpp b/Engine/Source/Runtime/Engine/Private/WorldPartition/WorldPartitionStreamingGeneration.cpp
--- a/Engine/Source/Runtime/Engine/Private/WorldPartition/WorldPartitionStreamingGeneration.cpp
+++ b/Engine/Source/Runtime/Engine/Private/WorldPartition/WorldPartitionStreamingGeneration.cpp
@@ -559,10 +559,10 @@ class FWorldPartitionStreamingGenerator
 
 		// Perform various adjustements based on validations and report errors
 		//
-		// The first validation pass is used to report errors, subsequent passes are used to make corrections to the FWorldPartitionActorDescView
+		// The first validation pass is used to report errors, subsequent passes are used to make corrections to the actor descriptor views.
 		// Since the references can form cycles/long chains in the data fixes might need to be propagated in multiple passes.
 		// 
-		// This works because fixes are deterministic and always apply the same way to both Actors being modified, so there's no ordering issues possible
+		// This works because fixes are deterministic and always apply the same way to both Actors being modified, so there's no ordering issues possible.
 		int32 NbErrorsDetected = INDEX_NONE;
 		for(uint32 NbValidationPasses = 0; NbErrorsDetected; NbValidationPasses++)
 		{
@@ -649,37 +649,35 @@ PRAGMA_ENABLE_DEPRECATION_WARNINGS
 				}
 
 				// Add attach reference for the topmost parent, this reference is inverted since we consider the top most existing 
-				// parent to be refering to us, not the child to be referering the parent
-				FGuid ParentGuid = ActorDescView.GetParentActor();
-				FWorldPartitionActorDescView* TopParentDescView = nullptr;
-
-				while (ParentGuid.IsValid())
+				// parent to be refering to us, not the child to be referering the parent.
+				if (!NbValidationPasses)
 				{
-					FWorldPartitionActorDescView* ParentDescView = ContainerDescriptor.ActorDescViewMap.FindByGuid(ParentGuid);
-					
-					if (ParentDescView)
+					FGuid ParentGuid = ActorDescView.GetParentActor();
+					FWorldPartitionActorDescView* TopParentDescView = nullptr;
+
+					while (ParentGuid.IsValid())
 					{
-						TopParentDescView = ParentDescView;
-						ParentGuid = ParentDescView->GetParentActor();
+						FWorldPartitionActorDescView* ParentDescView = ContainerDescriptor.ActorDescViewMap.FindByGuid(ParentGuid);
+					
+						if (ParentDescView)
+						{
+							TopParentDescView = ParentDescView;
+							ParentGuid = ParentDescView->GetParentActor();
+						}
+						else
+						{
+							// We had a guid but parent cannot be found, this will be a missing reference
+							References.Emplace(FActorReferenceInfo { ActorDescView.GetGuid(), &ActorDescView, ParentGuid, nullptr });
+							break; 
+						}
 					}
-					else
+
+					if (TopParentDescView)
 					{
-						// we had a guid but parent cannot be found, this will be a missing reference
-						break; 
+						References.Emplace(FActorReferenceInfo { TopParentDescView->GetGuid(), TopParentDescView, ActorDescView.GetGuid(), &ActorDescView });
 					}
 				}
 
-				if (TopParentDescView)
-				{
-					References.Emplace(FActorReferenceInfo { TopParentDescView->GetGuid(), TopParentDescView, ActorDescView.GetGuid(), &ActorDescView });
-				}
-
-				if (ParentGuid.IsValid())
-				{
-					// In case of missing parent add a missing reference 
-					References.Emplace(FActorReferenceInfo { ActorDescView.GetGuid(), &ActorDescView, ParentGuid, nullptr });
-				}
-
 				TArray<FGuid> RuntimeReferences;
 				if (NbValidationPasses)
 				{
@@ -778,11 +776,6 @@ PRAGMA_ENABLE_DEPRECATION_WARNINGS
 					}
 				}
 			});		
-			
-			if (!ensureMsgf(NbValidationPasses < 100000, TEXT("Infinite loop in validation passes for %s"), ContainerDescriptor.Container ? *ContainerDescriptor.Container->GetContainerPackage().ToString() : TEXT("[unknown world]")))
-			{
-				break;
-			}
 		}
 	}
 
