[shaders] avoid extra alloc+copy when pulling shader jobs from DDC; store shared buffers in our in-memory cache directly

diff --git a/Engine/Source/Runtime/Engine/Private/ShaderCompiler/ShaderCompiler.cpp b/Engine/Source/Runtime/Engine/Private/ShaderCompiler/ShaderCompiler.cpp
--- a/Engine/Source/Runtime/Engine/Private/ShaderCompiler/ShaderCompiler.cpp
+++ b/Engine/Source/Runtime/Engine/Private/ShaderCompiler/ShaderCompiler.cpp
@@ -629,10 +629,10 @@ void FShaderCompileJobCollection::SubmitJobs(const TArray<FShaderCommonCompileJo
 					const bool bCheckDDC = !(Job->bIsDefaultMaterial || Job->bIsGlobalShader);
 
 					// see if we can find the job in the cache first
-					if (TArray<uint8>* ExistingOutput = CompletedJobsCache.Find(InputHash, bCheckDDC))
+					if (FSharedBuffer* ExistingOutput = CompletedJobsCache.Find(InputHash, bCheckDDC))
 					{
 						UE_LOG(LogShaderCompilers, UE_SHADERCACHE_LOG_LEVEL, TEXT("There is already a cached job with the ihash %s, processing the new one immediately."), *LexToString(InputHash));
-						FMemoryReader MemReader(*ExistingOutput);
+						FMemoryReaderView MemReader(*ExistingOutput);
 						Job->SerializeOutput(MemReader);
 
 						// finish the job instantly
@@ -765,6 +765,7 @@ void FShaderCompileJobCollection::AddToCacheAndProcessPending(FShaderCommonCompi
 	FMemoryWriter Writer(Output);
 	FinishedJob->SerializeOutput(Writer);
 
+	FSharedBuffer Buffer = MakeSharedBufferFromArray(MoveTemp(Output));
 	// TODO: reduce the scope - e.g. SerializeOutput and processing finished jobs can be moved out of it
 	FWriteScopeLock JobLocker(Lock);
 
@@ -777,7 +778,7 @@ void FShaderCompileJobCollection::AddToCacheAndProcessPending(FShaderCommonCompi
 		{
 			checkf(CurHead != FinishedJob, TEXT("Job that is being added to cache was also on a waiting list! Error in bookkeeping."));
 
-			FMemoryReader MemReader(Output);
+			FMemoryReaderView MemReader(Buffer);
 			CurHead->SerializeOutput(MemReader);
 			checkf(CurHead->bSucceeded == FinishedJob->bSucceeded, TEXT("Different success status for the job with the same ihash"));
 
@@ -801,7 +802,7 @@ void FShaderCompileJobCollection::AddToCacheAndProcessPending(FShaderCommonCompi
 	{
 		const bool bAddToDDC = !(FinishedJob->bIsDefaultMaterial || FinishedJob->bIsGlobalShader);
 		// we only cache jobs that succeded
-		CompletedJobsCache.Add(InputHash, Output, NumOutstandingJobsWithSameHash, bAddToDDC);
+		CompletedJobsCache.Add(InputHash, Buffer, NumOutstandingJobsWithSameHash, bAddToDDC);
 	}
 
 	// remove ourselves from the jobs in flight, if we were there (if this job is a cloned job it might not have been)
@@ -8666,7 +8667,7 @@ public:
 	int32 NumHits = 0;
 
 	/** Canned output */
-	TArray<uint8> JobOutput;
+	FSharedBuffer JobOutput;
 
 	/** Similar to FRefCountBase AddRef, but not atomic */
 	int32 AddRef()
@@ -8695,7 +8696,7 @@ public:
 
 	uint64 GetAllocatedSize() const
 	{
-		return static_cast<uint64>(JobOutput.GetAllocatedSize() + sizeof(*this));
+		return static_cast<uint64>(JobOutput.GetSize() + sizeof(*this));
 	}
 };
 
@@ -8771,16 +8772,12 @@ FShaderJobCache::FJobCachedOutput* FShaderJobCache::Find(const FJobInputHash& Ha
 				{
 					// Create a new entry to store in the FShaderJobCache
 					FStoredOutput* NewStoredOutput = new FStoredOutput();
-					NewStoredOutput->JobOutput.Reserve(Results.GetSize());
-					NewStoredOutput->JobOutput.SetNum(Results.GetSize());
+					NewStoredOutput->JobOutput = Results;
 
 					TRACE_COUNTER_ADD(Shaders_JobCacheDDCBytesReceived, Results.GetSize());
 
-					check(Results.GetSize() <= NewStoredOutput->JobOutput.GetAllocatedSize());
-					FMemory::Memcpy(NewStoredOutput->JobOutput.GetData(),  Results.GetData(), Results.GetSize());
-
 					// Generate an output hash and cache the result in the FShaderJobCache
-					FJobOutputHash NewOutputHash = FBlake3::HashBuffer(NewStoredOutput->JobOutput.GetData(), NewStoredOutput->JobOutput.Num());
+					FJobOutputHash NewOutputHash = FBlake3::HashBuffer(NewStoredOutput->JobOutput.GetData(), NewStoredOutput->JobOutput.GetSize());
 
 					NewStoredOutput->AddRef();
 					Outputs.Add(NewOutputHash, NewStoredOutput);
@@ -8835,7 +8832,7 @@ void FShaderJobCache::Add(const FJobInputHash& Hash, const FJobCachedOutput& Con
 		return;
 	}
 
-	FJobOutputHash OutputHash = FBlake3::HashBuffer(Contents.GetData(), Contents.Num());
+	FJobOutputHash OutputHash = FBlake3::HashBuffer(Contents.GetData(), Contents.GetSize());
 
 	// add the record
 	const uint64 InputHashToOutputOriginalSize = InputHashToOutput.GetAllocatedSize();
@@ -8873,9 +8870,9 @@ void FShaderJobCache::Add(const FJobInputHash& Hash, const FJobCachedOutput& Con
 			Key.Hash = Hash;
 			UE::DerivedData::FCacheRecordBuilder RecordBuilder(Key);
 
-			RecordBuilder.AddValue(ShaderJobCacheId, FSharedBuffer::MakeView(MakeMemoryView(Contents)));
+			RecordBuilder.AddValue(ShaderJobCacheId, FSharedBuffer::MakeView(Contents));
 
-			TRACE_COUNTER_ADD(Shaders_JobCacheDDCBytesSent, Contents.GetAllocatedSize());
+			TRACE_COUNTER_ADD(Shaders_JobCacheDDCBytesSent, Contents.GetSize());
 
 			UE::DerivedData::FRequestOwner RequestOwner(UE::DerivedData::EPriority::Normal);
 			UE::DerivedData::FRequestBarrier RequestBarrier(RequestOwner);
diff --git a/Engine/Source/Runtime/Engine/Public/ShaderCompiler.h b/Engine/Source/Runtime/Engine/Public/ShaderCompiler.h
--- a/Engine/Source/Runtime/Engine/Public/ShaderCompiler.h
+++ b/Engine/Source/Runtime/Engine/Public/ShaderCompiler.h
@@ -313,7 +313,7 @@ public:
 	~FShaderJobCache();
 
 	using FJobInputHash = FShaderCommonCompileJob::FInputHash;
-	using FJobCachedOutput = TArray<uint8>;
+	using FJobCachedOutput = FSharedBuffer;
 
 	/** Looks for the job in the cache, returns null if not found */
 	FJobCachedOutput* Find(const FJobInputHash& Hash, const bool bCheckDDC);
