Convert manual MD5Hash to Guid conversions to use the MD5HashToGuid function.

#rnx

diff --git a/Engine/Source/Runtime/Engine/Private/ExternalPackageHelper.cpp b/Engine/Source/Runtime/Engine/Private/ExternalPackageHelper.cpp
--- a/Engine/Source/Runtime/Engine/Private/ExternalPackageHelper.cpp
+++ b/Engine/Source/Runtime/Engine/Private/ExternalPackageHelper.cpp
@@ -101,9 +101,7 @@ FString FExternalPackageHelper::GetExternalPackageName(const FString& InOuterPac
 	FMD5Hash MD5Hash;
 	ArMD5.GetHash(MD5Hash);
 
-	FGuid PackageGuid;
-	check(MD5Hash.GetSize() == sizeof(FGuid));
-	FMemory::Memcpy(&PackageGuid, MD5Hash.GetBytes(), sizeof(FGuid));
+	FGuid PackageGuid = MD5HashToGuid(MD5Hash);
 	check(PackageGuid.IsValid());
 
 	FString GuidBase36 = PackageGuid.ToString(EGuidFormats::Base36Encoded);
diff --git a/Engine/Source/Runtime/Engine/Private/Level.cpp b/Engine/Source/Runtime/Engine/Private/Level.cpp
--- a/Engine/Source/Runtime/Engine/Private/Level.cpp
+++ b/Engine/Source/Runtime/Engine/Private/Level.cpp
@@ -3446,9 +3446,7 @@ FString ULevel::GetActorPackageName(const FString& InBaseDir, EActorPackagingSch
 	FMD5Hash MD5Hash;
 	ArMD5.GetHash(MD5Hash);
 
-	FGuid PackageGuid;
-	check(MD5Hash.GetSize() == sizeof(FGuid));
-	FMemory::Memcpy(&PackageGuid, MD5Hash.GetBytes(), sizeof(FGuid));
+	FGuid PackageGuid = MD5HashToGuid(MD5Hash);
 	check(PackageGuid.IsValid());
 
 	FString GuidBase36 = PackageGuid.ToString(EGuidFormats::Base36Encoded);
diff --git a/Engine/Source/Runtime/Engine/Private/WorldPartition/HLOD/HLODActor.cpp b/Engine/Source/Runtime/Engine/Private/WorldPartition/HLOD/HLODActor.cpp
--- a/Engine/Source/Runtime/Engine/Private/WorldPartition/HLOD/HLODActor.cpp
+++ b/Engine/Source/Runtime/Engine/Private/WorldPartition/HLOD/HLODActor.cpp
@@ -193,8 +193,7 @@ void AWorldPartitionHLOD::PostLoad()
 		FMD5Hash MD5Hash;
 		ArMD5.GetHash(MD5Hash);
 
-		check(MD5Hash.GetSize() == sizeof(FGuid));
-		FMemory::Memcpy(&SourceCellGuid, MD5Hash.GetBytes(), sizeof(FGuid));
+		SourceCellGuid = MD5HashToGuid(MD5Hash);
 		check(SourceCellGuid.IsValid());
 	}
 
