World Partition Streaming Partitioned Worlds Support:
- Make sure to give cell levels unique names when generating them.
- Reduced cell levels names when they don't have data layers.
- Added asserts to detect duplicated cell level names.

diff --git a/Engine/Source/Runtime/Engine/Private/UnrealEngine.cpp b/Engine/Source/Runtime/Engine/Private/UnrealEngine.cpp
--- a/Engine/Source/Runtime/Engine/Private/UnrealEngine.cpp
+++ b/Engine/Source/Runtime/Engine/Private/UnrealEngine.cpp
@@ -17467,21 +17467,7 @@ int32 UEngine::RenderStatLevels(UWorld* World, FViewport* Viewport, FCanvas* Can
 		FColor	Color = ULevelStreaming::GetLevelStreamingStatusColor(LevelStatus.StreamingStatus);
 		const TCHAR* StatusDisplayName = ULevelStreaming::GetLevelStreamingStatusDisplayName(LevelStatus.StreamingStatus);
 		FString LevelPackageName = LevelStatus.PackageName.ToString();
-		FString DisplayName = LevelPackageName;
-
-		if (World->IsPartitionedWorld())
-		{
-			int32 CutPos = LevelPackageName.Find(*MapShortName);
-			if (CutPos != INDEX_NONE)
-			{
-				CutPos += MapShortName.Len();
-				while (LevelPackageName[CutPos] == TEXT('_'))
-				{
-					CutPos++;
-				}
-				DisplayName = LevelPackageName.Right(LevelPackageName.Len() - CutPos);
-			}
-		}
+		FString DisplayName = UWorld::RemovePIEPrefix(FPackageName::GetShortName(LevelPackageName));
 
 		if (LevelStatus.LODIndex != INDEX_NONE)
 		{
diff --git a/Engine/Source/Runtime/Engine/Private/WorldPartition/RuntimeSpatialHash/RuntimeSpatialHashHLOD.cpp b/Engine/Source/Runtime/Engine/Private/WorldPartition/RuntimeSpatialHash/RuntimeSpatialHashHLOD.cpp
--- a/Engine/Source/Runtime/Engine/Private/WorldPartition/RuntimeSpatialHash/RuntimeSpatialHashHLOD.cpp
+++ b/Engine/Source/Runtime/Engine/Private/WorldPartition/RuntimeSpatialHash/RuntimeSpatialHashHLOD.cpp
@@ -133,6 +133,8 @@ static TArray<FGuid> GenerateHLODsForGrid(UWorldPartition* WorldPartition, const
 		return bShouldGenerateHLODs;
 	};
 
+	UWorld* OuterWorld = WorldPartition->GetTypedOuter<UWorld>();
+
 	// Quick pass to compute the number of cells we'll have to process, to provide a meaningful progress display
 	int32 NbCellsToProcess = 0;
 	PartitionedActors.ForEachCells([&](const FSquare2DGridHelper::FGridLevel::FGridCell& GridCell)
@@ -174,8 +176,8 @@ static TArray<FGuid> GenerateHLODsForGrid(UWorldPartition* WorldPartition, const
 				SlowTask.EnterProgressFrame(1);
 				// todo_ow : Generating hlods is not yet supported for content bundles. Pass an invalid UID to specified the cell is not coming from ContentBundles.
 				FGuid ContentBundleInvalidUID;
-				FGuid CellGuid = UWorldPartitionRuntimeSpatialHash::GetCellGuid(RuntimeGrid.GridName, CellGlobalCoord, GridCellDataChunk.GetDataLayersID(), ContentBundleInvalidUID);
-				FString CellName = UWorldPartitionRuntimeSpatialHash::GetCellNameString(RuntimeGrid.GridName, CellGlobalCoord, GridCellDataChunk.GetDataLayersID(), ContentBundleInvalidUID);
+				const FGuid CellGuid = UWorldPartitionRuntimeSpatialHash::GetCellGuid(RuntimeGrid.GridName, CellGlobalCoord, GridCellDataChunk.GetDataLayersID(), ContentBundleInvalidUID);
+				const FString CellName = UWorldPartitionRuntimeSpatialHash::GetCellNameString(OuterWorld->GetPackage()->GetFName(), RuntimeGrid.GridName, CellGlobalCoord, GridCellDataChunk.GetDataLayersID(), ContentBundleInvalidUID);
 
 				UE_LOG(LogWorldPartitionRuntimeSpatialHashHLOD, Verbose, TEXT("Creating HLOD for cell %s at %s..."), *CellName, *CellCoord.ToString());
 
diff --git a/Engine/Source/Runtime/Engine/Private/WorldPartition/WorldPartitionLevelHelper.cpp b/Engine/Source/Runtime/Engine/Private/WorldPartition/WorldPartitionLevelHelper.cpp
--- a/Engine/Source/Runtime/Engine/Private/WorldPartition/WorldPartitionLevelHelper.cpp
+++ b/Engine/Source/Runtime/Engine/Private/WorldPartition/WorldPartitionLevelHelper.cpp
@@ -271,6 +271,7 @@ ULevel* FWorldPartitionLevelHelper::CreateEmptyLevelForRuntimeCell(const UWorldP
 	// Create World & Persistent Level
 	UWorld::InitializationValues IVS = FWorldPartitionLevelHelper::GetWorldInitializationValues();
 	const FName WorldName = FName(FPackageName::ObjectPathToObjectName(InWorldAssetName));
+	check(!FindObject<UWorld>(CellPackage, *WorldName.ToString()));
 	UWorld* NewWorld = UWorld::CreateWorld(InWorld->WorldType, /*bInformEngineOfWorld*/false, WorldName, CellPackage, /*bAddToRoot*/false, InWorld->FeatureLevel, &IVS, /*bInSkipInitWorld*/true);
 	check(NewWorld);
 	NewWorld->SetFlags(RF_Public | RF_Standalone);
diff --git a/Engine/Source/Runtime/Engine/Private/WorldPartition/WorldPartitionRuntimeSpatialHash.cpp b/Engine/Source/Runtime/Engine/Private/WorldPartition/WorldPartitionRuntimeSpatialHash.cpp
--- a/Engine/Source/Runtime/Engine/Private/WorldPartition/WorldPartitionRuntimeSpatialHash.cpp
+++ b/Engine/Source/Runtime/Engine/Private/WorldPartition/WorldPartitionRuntimeSpatialHash.cpp
@@ -1274,9 +1274,15 @@ void UWorldPartitionRuntimeSpatialHash::DumpStateLog(FHierarchicalLogArchive& Ar
 	});
 }
 
-FString UWorldPartitionRuntimeSpatialHash::GetCellNameString(FName InGridName, const FGridCellCoord& InCellGlobalCoord, const FDataLayersID& InDataLayerID, const FGuid& InContentBundleID)
+FString UWorldPartitionRuntimeSpatialHash::GetCellNameString(FName InWorldPackageName, FName InGridName, const FGridCellCoord& InCellGlobalCoord, const FDataLayersID& InDataLayerID, const FGuid& InContentBundleID)
 {
-	FString CellName = FString::Printf(TEXT("%s_%s_DL%X"), *InGridName.ToString(), *GetCellCoordString(InCellGlobalCoord), InDataLayerID.GetHash());
+	const FString ShortPackageName = UWorld::RemovePIEPrefix(FPackageName::GetShortName(InWorldPackageName));
+	FString CellName = FString::Printf(TEXT("%s_%s_%s"), *ShortPackageName, *InGridName.ToString(), *GetCellCoordString(InCellGlobalCoord));
+
+	if (InDataLayerID.GetHash())
+	{
+		CellName += FString::Printf(TEXT("_DL%X"), InDataLayerID.GetHash());
+	}
 
 	if (InContentBundleID.IsValid())
 	{
@@ -1299,9 +1305,7 @@ FGuid UWorldPartitionRuntimeSpatialHash::GetCellGuid(FName InGridName, const FGr
 	FMD5Hash MD5Hash;
 	ArMD5.GetHash(MD5Hash);
 
-	FGuid CellGuid;
-	check(MD5Hash.GetSize() == sizeof(FGuid));
-	FMemory::Memcpy(&CellGuid, MD5Hash.GetBytes(), sizeof(FGuid));
+	FGuid CellGuid = MD5HashToGuid(MD5Hash);
 	check(CellGuid.IsValid());
 
 	return CellGuid;
@@ -1330,7 +1334,8 @@ bool UWorldPartitionRuntimeSpatialHash::CreateStreamingGrid(const FSpatialHashRu
 
 	UWorldPartition* WorldPartition = GetOuterUWorldPartition();
 	UWorld* World = WorldPartition->GetWorld();
-	const bool bIsMainWorldPartition = (World == WorldPartition->GetTypedOuter<UWorld>());
+	UWorld* OuterWorld = GetTypedOuter<UWorld>();
+	const bool bIsMainWorldPartition = (World == OuterWorld);
 
 	FSpatialHashStreamingGrid& CurrentStreamingGrid = StreamingGrids.AddDefaulted_GetRef();
 	CurrentStreamingGrid.GridName = RuntimeGrid.GridName;
@@ -1394,7 +1399,7 @@ bool UWorldPartitionRuntimeSpatialHash::CreateStreamingGrid(const FSpatialHashRu
 				
 				FGridCellCoord CellGlobalCoords;
 				verify(PartionedActors.GetCellGlobalCoords(FGridCellCoord(CellCoordX, CellCoordY, Level), CellGlobalCoords));
-				const FString CellName = GetCellNameString(CurrentStreamingGrid.GridName, CellGlobalCoords, GridCellDataChunk.GetDataLayersID(), GridCellDataChunk.GetContentBundleID());
+				const FString CellName = GetCellNameString(OuterWorld->GetPackage()->GetFName(), CurrentStreamingGrid.GridName, CellGlobalCoords, GridCellDataChunk.GetDataLayersID(), GridCellDataChunk.GetContentBundleID());
 				const FGuid CellGuid = GetCellGuid(CurrentStreamingGrid.GridName, CellGlobalCoords, GridCellDataChunk.GetDataLayersID(), GridCellDataChunk.GetContentBundleID());
 
 				UWorldPartitionRuntimeSpatialHashCell* StreamingCell = NewObject<UWorldPartitionRuntimeSpatialHashCell>(this, StreamingPolicy->GetRuntimeCellClass(), FName(CellName));
diff --git a/Engine/Source/Runtime/Engine/Public/WorldPartition/WorldPartitionRuntimeSpatialHash.h b/Engine/Source/Runtime/Engine/Public/WorldPartition/WorldPartitionRuntimeSpatialHash.h
--- a/Engine/Source/Runtime/Engine/Public/WorldPartition/WorldPartitionRuntimeSpatialHash.h
+++ b/Engine/Source/Runtime/Engine/Public/WorldPartition/WorldPartitionRuntimeSpatialHash.h
@@ -281,7 +281,7 @@ public:
 
 	virtual URuntimeHashExternalStreamingObjectBase* StoreToExternalStreamingObject(UObject* StreamingObjectOuter, FName StreamingObjectName) override;
 
-	static FString GetCellNameString(FName InGridName, const FGridCellCoord& InCellGlobalCoord, const FDataLayersID& InDataLayerID, const FGuid& InContentBundleID);
+	static FString GetCellNameString(FName InWorldPackageName, FName InGridName, const FGridCellCoord& InCellGlobalCoord, const FDataLayersID& InDataLayerID, const FGuid& InContentBundleID);
 	static FGuid GetCellGuid(FName InGridName, const FGridCellCoord& InCellGlobalCoord, const FDataLayersID& InDataLayerID, const FGuid& InContentBundleID);
 
 	bool GetPreviewGrids() const;
