Fix typo in Force Feedback property name

#rnx

diff --git a/Engine/Source/Runtime/Engine/Classes/GameFramework/ForceFeedbackEffect.h b/Engine/Source/Runtime/Engine/Classes/GameFramework/ForceFeedbackEffect.h
--- a/Engine/Source/Runtime/Engine/Classes/GameFramework/ForceFeedbackEffect.h
+++ b/Engine/Source/Runtime/Engine/Classes/GameFramework/ForceFeedbackEffect.h
@@ -120,7 +120,7 @@ class UForceFeedbackEffect : public UObject
 
 	/** A map of platform name -> ForceFeedback channel details */
 	UPROPERTY(EditAnywhere, Category = "ForceFeedbackEffect", Meta = (GetOptions = "Engine.InputPlatformSettings.GetAllHardwareDeviceNames"))
-	TMap<FName, FForceFeedbackEffectOverridenChannelDetails> PerDeviceOverides;
+	TMap<FName, FForceFeedbackEffectOverridenChannelDetails> PerDeviceOverrides;
 
 	/** A map of input device properties that we want to set while this effect is playing */
 	UPROPERTY(EditAnywhere, Instanced, Category = "ForceFeedbackEffect")
diff --git a/Engine/Source/Runtime/Engine/Private/GameFramework/ForceFeedbackEffect.cpp b/Engine/Source/Runtime/Engine/Private/GameFramework/ForceFeedbackEffect.cpp
--- a/Engine/Source/Runtime/Engine/Private/GameFramework/ForceFeedbackEffect.cpp
+++ b/Engine/Source/Runtime/Engine/Private/GameFramework/ForceFeedbackEffect.cpp
@@ -109,7 +109,7 @@ const TArray<FForceFeedbackChannelDetails>& UForceFeedbackEffect::GetCurrentChan
 	{
 		FHardwareDeviceIdentifier Hardware = SubSystem->GetMostRecentlyUsedHardwareDevice(PlatformUser);
 		// Check if there are any per-input device overrides available
-		if (const FForceFeedbackEffectOverridenChannelDetails* Details = PerDeviceOverides.Find(Hardware.HardwareDeviceIdentifier))
+		if (const FForceFeedbackEffectOverridenChannelDetails* Details = PerDeviceOverrides.Find(Hardware.HardwareDeviceIdentifier))
 		{
 			return Details->ChannelDetails;
 		}
