Fix the category on that Input Platform Settings to be correct, instead of saying "Enhanced Input"

#rnx
#pf skip

diff --git a/Engine/Source/Runtime/Engine/Classes/GameFramework/InputSettings.h b/Engine/Source/Runtime/Engine/Classes/GameFramework/InputSettings.h
--- a/Engine/Source/Runtime/Engine/Classes/GameFramework/InputSettings.h
+++ b/Engine/Source/Runtime/Engine/Classes/GameFramework/InputSettings.h
@@ -32,7 +32,7 @@ class ENGINE_API UInputSettings
 	 * Platform specific settings for Input.
 	 * @see UInputPlatformSettings
 	 */
-	UPROPERTY(EditAnywhere, Category = "Enhanced Input")
+	UPROPERTY(EditAnywhere, Category = "Platforms")
 	FPerPlatformSettings PlatformSettings;
 
 	UPROPERTY(config, EditAnywhere, Category="Bindings", AdvancedDisplay)
