Fix crash when importing a custom LOD asynchronously (async import and build).

#rnx

diff --git a/Engine/Source/Runtime/Engine/Private/SkeletalMesh.cpp b/Engine/Source/Runtime/Engine/Private/SkeletalMesh.cpp
--- a/Engine/Source/Runtime/Engine/Private/SkeletalMesh.cpp
+++ b/Engine/Source/Runtime/Engine/Private/SkeletalMesh.cpp
@@ -294,12 +294,20 @@ FScopedSkeletalMeshPostEditChange::FScopedSkeletalMeshPostEditChange(USkeletalMe
 
 FScopedSkeletalMeshPostEditChange::~FScopedSkeletalMeshPostEditChange()
 {
-	//If decrementing the post edit change stack counter return 0 it mean we are the first scope call instance, so we have to call posteditchange and re register component
-	if (SkeletalMesh != nullptr && SkeletalMesh->UnStackPostEditChange() == 0)
+	if (SkeletalMesh)
 	{
-		if (bCallPostEditChange)
+		//If decrementing the post edit change stack counter return 0 it mean we are the first scope call instance, so we have to call posteditchange and re register component
+		if (SkeletalMesh->UnStackPostEditChange() == 0)
+		{
+			if (bCallPostEditChange)
+			{
+				SkeletalMesh->PostEditChange();
+			}
+		}
+		if (bReregisterComponents && SkeletalMesh->IsCompiling())
 		{
-			SkeletalMesh->PostEditChange();
+			//wait until the compilation is done
+			FSkinnedAssetCompilingManager::Get().FinishCompilation({ SkeletalMesh });
 		}
 	}
 	//If there is some re register data it will be delete when the destructor go out of scope. This will re register
