World Partition Unit Tests: add an actor with a child actor component to test recent fix when unloading this case.

diff --git a/Engine/Source/Runtime/Engine/Private/Tests/WorldPartition/WorldPartitionHandleTests.cpp b/Engine/Source/Runtime/Engine/Private/Tests/WorldPartition/WorldPartitionHandleTests.cpp
--- a/Engine/Source/Runtime/Engine/Private/Tests/WorldPartition/WorldPartitionHandleTests.cpp
+++ b/Engine/Source/Runtime/Engine/Private/Tests/WorldPartition/WorldPartitionHandleTests.cpp
@@ -28,6 +28,11 @@ namespace WorldPartitionTests
 
 		FWorldPartitionHandle Handle = FWorldPartitionHandle(ActorDescContainer, FGuid(TEXT("5D9F93BA407A811AFDDDAAB4F1CECC6A")));
 
+		{
+			FWorldPartitionHandle ActorWithChildActorComponentHandle = FWorldPartitionHandle(ActorDescContainer, FGuid(TEXT("538856174ECD465948488F9441AE9251")));
+			FWorldPartitionReference ActorWithChildActorComponentReference = ActorWithChildActorComponentHandle.ToReference();
+		}		
+
 		FWorldPartitionReference Reference;
 		{
 			LoadingContextType LoadingContext;
