Restoring code that got lost in a move.

diff --git a/Engine/Source/Runtime/Engine/Public/MeshCardBuild.h b/Engine/Source/Runtime/Engine/Public/MeshCardBuild.h
--- a/Engine/Source/Runtime/Engine/Public/MeshCardBuild.h
+++ b/Engine/Source/Runtime/Engine/Public/MeshCardBuild.h
@@ -9,13 +9,15 @@
 #include "MeshCardRepresentation.h"
 #include "RenderDeferredCleanup.h"
 
+#include <atomic>
+
 class UStaticMesh;
 class FSignedDistanceFieldBuildMaterialData;
 
 class FLumenCardBuildData
 {
 public:
-	FLumenCardOBB OBB;
+	FLumenCardOBBf OBB;
 	uint8 AxisAlignedDirectionIndex;
 
 	friend FArchive& operator<<(FArchive& Ar, FLumenCardBuildData& Data)
@@ -131,9 +133,8 @@ public:
 	FCardRepresentationData()
 	{
 		// 0 means invalid id.
-		static uint32 NextCardRepresentationId = 0;
-		++NextCardRepresentationId;
-		CardRepresentationDataId.Value = NextCardRepresentationId;
+		static std::atomic<uint32> NextCardRepresentationId { 1 };
+		CardRepresentationDataId.Value = NextCardRepresentationId++;
 	}
 
 	void GetResourceSizeEx(FResourceSizeEx& CumulativeResourceSize) const
@@ -272,6 +273,9 @@ private:
 	/** Cancel or finish any work for the given task. */
 	void CancelAndDeleteTask(const TSet<FAsyncCardRepresentationTask*>& Tasks);
 
+	/** Handle generic finish compilation */
+	void FinishCompilationForObjects(TArrayView<UObject* const> InObjects) override;
+
 	/** Return whether the task has become invalid and should be canceled (i.e. reference unreachable objects) */
 	bool IsTaskInvalid(FAsyncCardRepresentationTask* Task) const;
 
@@ -303,4 +307,3 @@ extern ENGINE_API FString BuildCardRepresentationDerivedDataKey(const FString& I
 
 extern ENGINE_API void BeginCacheMeshCardRepresentation(const ITargetPlatform* TargetPlatform, UStaticMesh* StaticMeshAsset, class FStaticMeshRenderData& RenderData, const FString& DistanceFieldKey, FSourceMeshDataForDerivedDataTask* OptionalSourceMeshData);
 
-
