Test bone index validity to avoid runtime failure in rare conditions
Converted assert into log warning

diff --git a/Engine/Source/Runtime/Engine/Private/Animation/PoseAsset.cpp b/Engine/Source/Runtime/Engine/Private/Animation/PoseAsset.cpp
--- a/Engine/Source/Runtime/Engine/Private/Animation/PoseAsset.cpp
+++ b/Engine/Source/Runtime/Engine/Private/Animation/PoseAsset.cpp
@@ -609,17 +609,20 @@ bool UPoseAsset::GetAnimationPose(struct FAnimationPoseData& OutAnimationPoseDat
 					// Per-track (bone) transform
 					for (int32 TrackIndex = 0; TrackIndex < TrackNum; ++TrackIndex)
 					{
-						const int32 SkeletonBoneIndex = PoseContainer.TrackBoneIndices[TrackIndex];
-						const FCompactPoseBoneIndex CompactIndex = RequiredBones.GetCompactPoseIndexFromSkeletonIndex(SkeletonBoneIndex);
+						const FSkeletonPoseBoneIndex SkeletonBoneIndex = FSkeletonPoseBoneIndex(PoseContainer.TrackBoneIndices[TrackIndex]);
+
+						UE_CLOG(!RequiredBones.IsSkeletonPoseIndexValid(SkeletonBoneIndex), LogAnimation, Warning, TEXT("PoseAsset [%s] with skeleton [%s] has bones not present in the evaluation container. Bone index(%d) not found."), *GetPathName(), MySkeleton ? *MySkeleton->GetPathName() : TEXT("<Skeleton Not Found>"), SkeletonBoneIndex.GetInt());
+
+						const FCompactPoseBoneIndex CompactIndex = RequiredBones.GetCompactPoseIndexFromSkeletonPoseIndex(SkeletonBoneIndex);
 						
 						// If bone index is invalid, or not required for the pose - skip
-						if (CompactIndex == INDEX_NONE || !ExtractionContext.IsBoneRequired(CompactIndex.GetInt()))
+						if (!CompactIndex.IsValid() || !ExtractionContext.IsBoneRequired(CompactIndex.GetInt()))
 						{
 							continue;
 						}
 					
 						// Check if this track is part of the pose
-						const TArray<FPoseAssetInfluence>& PoseInfluences = PoseContainer.TrackPoseInfluenceIndices[TrackIndex].Influences;						
+						const TArray<FPoseAssetInfluence>& PoseInfluences = PoseContainer.TrackPoseInfluenceIndices[TrackIndex].Influences;
 						const int32 InfluenceIndex = PoseInfluences.IndexOfByPredicate([PoseIndex](const FPoseAssetInfluence& Influence) -> bool
 						{
 							return Influence.PoseIndex == PoseIndex;
@@ -650,7 +653,7 @@ bool UPoseAsset::GetAnimationPose(struct FAnimationPoseData& OutAnimationPoseDat
 							}
 
 							// Retarget the bone transform
-							FAnimationRuntime::RetargetBoneTransform(MySkeleton, GetRetargetTransformsSourceName(), GetRetargetTransforms(), OutBoneTransform, SkeletonBoneIndex, CompactIndex, RequiredBones, bAdditivePose);
+							FAnimationRuntime::RetargetBoneTransform(MySkeleton, GetRetargetTransformsSourceName(), GetRetargetTransforms(), OutBoneTransform, SkeletonBoneIndex.GetInt(), CompactIndex, RequiredBones, bAdditivePose);
 							OutBoneTransform.NormalizeRotation();
 						}
 					}					
@@ -670,11 +673,15 @@ bool UPoseAsset::GetAnimationPose(struct FAnimationPoseData& OutAnimationPoseDat
 
 			for(int32 TrackIndex = 0; TrackIndex < TrackNum; ++TrackIndex)
 			{
-				const int32 SkeletonBoneIndex = PoseContainer.TrackBoneIndices[TrackIndex];
-				const FCompactPoseBoneIndex PoseBoneIndex = RequiredBones.GetCompactPoseIndexFromSkeletonIndex(SkeletonBoneIndex);
+				const FSkeletonPoseBoneIndex SkeletonBoneIndex = FSkeletonPoseBoneIndex(PoseContainer.TrackBoneIndices[TrackIndex]);
+
+				UE_CLOG(!RequiredBones.IsSkeletonPoseIndexValid(SkeletonBoneIndex), LogAnimation, Warning, TEXT("PoseAsset [%s] with skeleton [%s] has bones not present in the evaluation container. Bone index(%d) not found."), *GetPathName(), MySkeleton ? *MySkeleton->GetPathName() : TEXT("<Skeleton Not Found>"), SkeletonBoneIndex.GetInt());
+
+				const FCompactPoseBoneIndex CompactIndex = RequiredBones.GetCompactPoseIndexFromSkeletonPoseIndex(SkeletonBoneIndex);
+
 				// we add even if it's invalid because we want it to match with track index
-				BoneIndices[TrackIndex].SkeletonBoneIndex = SkeletonBoneIndex;
-				BoneIndices[TrackIndex].CompactBoneIndex = PoseBoneIndex;
+				BoneIndices[TrackIndex].SkeletonBoneIndex = SkeletonBoneIndex.GetInt();
+				BoneIndices[TrackIndex].CompactBoneIndex = CompactIndex;
 			}
 
 			// you could only have morphtargets
diff --git a/Engine/Source/Runtime/Engine/Public/BoneContainer.h b/Engine/Source/Runtime/Engine/Public/BoneContainer.h
--- a/Engine/Source/Runtime/Engine/Public/BoneContainer.h
+++ b/Engine/Source/Runtime/Engine/Public/BoneContainer.h
@@ -647,6 +647,16 @@ public:
 		return FCompactPoseBoneIndex(INDEX_NONE);
 	}
 
+	/**
+	 * Returns whether or not the skeleton index is contained in the mapping used to build this container.
+	 * Note that even if the skeleton index is valid, it might not contain a valid compact pose index if
+	 * that bone isn't used due to LOD or other reasons.
+	 */
+	bool IsSkeletonPoseIndexValid(const FSkeletonPoseBoneIndex& SkeletonIndex) const
+	{
+		return SkeletonToCompactPose.IsValidIndex(SkeletonIndex.GetInt());
+	}
+
 	FMeshPoseBoneIndex MakeMeshPoseIndex(const FCompactPoseBoneIndex& BoneIndex) const
 	{
 		return FMeshPoseBoneIndex(GetBoneIndicesArray()[BoneIndex.GetInt()]);
