World Partition: better error output when reporting errors for invalid actor references.

diff --git a/Engine/Source/Runtime/Engine/Private/WorldPartition/ErrorHandling/WorldPartitionStreamingGenerationMapCheckErrorHandler.cpp b/Engine/Source/Runtime/Engine/Private/WorldPartition/ErrorHandling/WorldPartitionStreamingGenerationMapCheckErrorHandler.cpp
--- a/Engine/Source/Runtime/Engine/Private/WorldPartition/ErrorHandling/WorldPartitionStreamingGenerationMapCheckErrorHandler.cpp
+++ b/Engine/Source/Runtime/Engine/Private/WorldPartition/ErrorHandling/WorldPartitionStreamingGenerationMapCheckErrorHandler.cpp
@@ -3,20 +3,10 @@
 #if WITH_EDITOR
 #include "WorldPartition/ErrorHandling/WorldPartitionStreamingGenerationMapCheckErrorHandler.h"
 #include "Logging/MessageLog.h"
-#include "Misc/MapErrors.h"
-
-#define LOCTEXT_NAMESPACE "WorldPartition"
-
-void FStreamingGenerationMapCheckErrorHandler::AddAdditionalNameToken(TSharedRef<FTokenizedMessage>& InMessage, const FName& InErrorName)
-{
-	InMessage->AddToken(FMapErrorToken::Create(InErrorName));
-}
 
 void FStreamingGenerationMapCheckErrorHandler::HandleTokenizedMessage(TSharedRef<FTokenizedMessage>&& ErrorMessage)
 {
 	FMessageLog("MapCheck").AddMessage(ErrorMessage);
 }
 
-#undef LOCTEXT_NAMESPACE
-
 #endif
diff --git a/Engine/Source/Runtime/Engine/Private/WorldPartition/ErrorHandling/WorldPartitionStreamingGenerationTokenizedMessageErrorHandler.cpp b/Engine/Source/Runtime/Engine/Private/WorldPartition/ErrorHandling/WorldPartitionStreamingGenerationTokenizedMessageErrorHandler.cpp
--- a/Engine/Source/Runtime/Engine/Private/WorldPartition/ErrorHandling/WorldPartitionStreamingGenerationTokenizedMessageErrorHandler.cpp
+++ b/Engine/Source/Runtime/Engine/Private/WorldPartition/ErrorHandling/WorldPartitionStreamingGenerationTokenizedMessageErrorHandler.cpp
@@ -5,6 +5,7 @@
 
 #include "Misc/PackageName.h"
 #include "Misc/UObjectToken.h"
+#include "Misc/MapErrors.h"
 #include "WorldPartition/DataLayer/DataLayerInstanceWithAsset.h"
 #include "WorldPartition/WorldPartitionActorDescView.h"
 
@@ -14,11 +15,10 @@ void ITokenizedMessageErrorHandler::OnInvalidRuntimeGrid(const FWorldPartitionAc
 {
 	TSharedRef<FTokenizedMessage> Message = FTokenizedMessage::Create(EMessageSeverity::Error);
 	Message->AddToken(FTextToken::Create(LOCTEXT("TokenMessage_WorldPartition_Actor", "Actor")))
-		->AddToken(FActorToken::Create(ActorDescView.GetActorSoftPath().ToString(), ActorDescView.GetGuid(), FText::FromString(GetFullActorName(ActorDescView))))
+		->AddToken(FActorToken::Create(ActorDescView.GetActorSoftPath().ToString(), ActorDescView.GetGuid(), FText::FromString(GetActorName(ActorDescView))))
 		->AddToken(FTextToken::Create(LOCTEXT("TokenMessage_WorldPartition_HaveInvalidRuntimeGrid", "has an invalid runtime grid")))
-		->AddToken(FTextToken::Create(FText::FromName(GridName)));
-
-	AddAdditionalNameToken(Message, FName(TEXT("WorldPartition_InvalidRuntimeGrid_CheckForErrors")));
+		->AddToken(FTextToken::Create(FText::FromName(GridName)))
+		->AddToken(FMapErrorToken::Create(TEXT("WorldPartition_InvalidRuntimeGrid_CheckForErrors")));
 
 	HandleTokenizedMessage(MoveTemp(Message));
 }
@@ -27,11 +27,19 @@ void ITokenizedMessageErrorHandler::OnInvalidReference(const FWorldPartitionActo
 {
 	TSharedRef<FTokenizedMessage> Message = FTokenizedMessage::Create(EMessageSeverity::Error);
 	Message->AddToken(FTextToken::Create(LOCTEXT("TokenMessage_WorldPartition_Actor", "Actor")))
-		->AddToken(FActorToken::Create(ActorDescView.GetActorSoftPath().ToString(), ActorDescView.GetGuid(), FText::FromString(GetFullActorName(ActorDescView))))
-		->AddToken(FTextToken::Create(LOCTEXT("TokenMessage_WorldPartition_HaveMissingRefsTo", "has an invalid reference to")))
-		->AddToken(FTextToken::Create(ReferenceActorDescView ? FText::FromString(GetFullActorName(*ReferenceActorDescView)) : FText::FromString(ReferenceGuid.ToString())));
+		->AddToken(FActorToken::Create(ActorDescView.GetActorSoftPath().ToString(), ActorDescView.GetGuid(), FText::FromString(GetActorName(ActorDescView))))
+		->AddToken(FTextToken::Create(LOCTEXT("TokenMessage_WorldPartition_HaveMissingRefsTo", "has an invalid reference to")));
+
+	if (ReferenceActorDescView)
+	{
+		Message->AddToken(FActorToken::Create(ReferenceActorDescView->GetActorSoftPath().ToString(), ReferenceActorDescView->GetGuid(), FText::FromString(GetActorName(*ReferenceActorDescView))));
+	}
+	else
+	{
+		Message->AddToken(FTextToken::Create(FText::FromString(ReferenceGuid.ToString())));
+	}
 
-	AddAdditionalNameToken(Message, FName(TEXT("WorldPartition_MissingActorReference_CheckForErrors")));
+	Message->AddToken(FMapErrorToken::Create(TEXT("WorldPartition_MissingActorReference_CheckForErrors")));
 
 	HandleTokenizedMessage(MoveTemp(Message));
 }
@@ -44,12 +52,11 @@ void ITokenizedMessageErrorHandler::OnInvalidReferenceGridPlacement(const FWorld
 
 	TSharedRef<FTokenizedMessage> Message = FTokenizedMessage::Create(MessageSeverity);
 	Message->AddToken(FTextToken::Create(ActorDescView.GetIsSpatiallyLoaded() ? SpatiallyLoadedActor : NonSpatiallyLoadedActor))
-		->AddToken(FActorToken::Create(ActorDescView.GetActorSoftPath().ToString(), ActorDescView.GetGuid(), FText::FromString(GetFullActorName(ActorDescView))))
+		->AddToken(FActorToken::Create(ActorDescView.GetActorSoftPath().ToString(), ActorDescView.GetGuid(), FText::FromString(GetActorName(ActorDescView))))
 		->AddToken(FTextToken::Create(LOCTEXT("TokenMessage_WorldPartition_References", "references")))
 		->AddToken(FTextToken::Create(ReferenceActorDescView.GetIsSpatiallyLoaded() ? SpatiallyLoadedActor : NonSpatiallyLoadedActor))
-		->AddToken(FActorToken::Create(ReferenceActorDescView.GetActorSoftPath().ToString(), ReferenceActorDescView.GetGuid(), FText::FromName(ReferenceActorDescView.GetActorLabelOrName())));
-
-	AddAdditionalNameToken(Message, FName(TEXT("WorldPartition_StreamedActorReferenceAlwaysLoadedActor_CheckForErrors")));
+		->AddToken(FActorToken::Create(ReferenceActorDescView.GetActorSoftPath().ToString(), ReferenceActorDescView.GetGuid(), FText::FromString(GetActorName(ReferenceActorDescView))))
+		->AddToken(FMapErrorToken::Create(TEXT("WorldPartition_StreamedActorReferenceAlwaysLoadedActor_CheckForErrors")));
 
 	HandleTokenizedMessage(MoveTemp(Message));
 }
@@ -58,11 +65,10 @@ void ITokenizedMessageErrorHandler::OnInvalidReferenceDataLayers(const FWorldPar
 {
 	TSharedRef<FTokenizedMessage> Message = FTokenizedMessage::Create(EMessageSeverity::Error);
 	Message->AddToken(FTextToken::Create(LOCTEXT("TokenMessage_WorldPartition_Actor", "Actor")))
-		->AddToken(FActorToken::Create(ActorDescView.GetActorSoftPath().ToString(), ActorDescView.GetGuid(),  FText::FromString(GetFullActorName(ActorDescView))))
+		->AddToken(FActorToken::Create(ActorDescView.GetActorSoftPath().ToString(), ActorDescView.GetGuid(),  FText::FromString(GetActorName(ActorDescView))))
 		->AddToken(FTextToken::Create(LOCTEXT("TokenMessage_WorldPartition_ReferenceActorInOtherDataLayers", "references an actor in a different set of runtime data layers")))
-		->AddToken(FActorToken::Create(ReferenceActorDescView.GetActorSoftPath().ToString(), ReferenceActorDescView.GetGuid(), FText::FromName(ReferenceActorDescView.GetActorLabelOrName())));
-
-	AddAdditionalNameToken(Message, FName(TEXT("WorldPartition_ActorReferenceActorInAnotherDataLayer_CheckForErrors")));
+		->AddToken(FActorToken::Create(ReferenceActorDescView.GetActorSoftPath().ToString(), ReferenceActorDescView.GetGuid(), FText::FromString(GetActorName(ReferenceActorDescView))))
+		->AddToken(FMapErrorToken::Create(TEXT("WorldPartition_ActorReferenceActorInAnotherDataLayer_CheckForErrors")));
 
 	HandleTokenizedMessage(MoveTemp(Message));
 }
@@ -71,11 +77,10 @@ void ITokenizedMessageErrorHandler::OnInvalidReferenceRuntimeGrid(const FWorldPa
 {
 	TSharedRef<FTokenizedMessage> Message = FTokenizedMessage::Create(EMessageSeverity::Error);
 	Message->AddToken(FTextToken::Create(LOCTEXT("TokenMessage_WorldPartition_Actor", "Actor")))
-		->AddToken(FActorToken::Create(ActorDescView.GetActorSoftPath().ToString(), ActorDescView.GetGuid(), FText::FromString(GetFullActorName(ActorDescView))))
+		->AddToken(FActorToken::Create(ActorDescView.GetActorSoftPath().ToString(), ActorDescView.GetGuid(), FText::FromString(GetActorName(ActorDescView))))
 		->AddToken(FTextToken::Create(LOCTEXT("TokenMessage_WorldPartition_ReferenceActorInOtherRuntimeGrid", "references an actor in a different runtime grid")))
-		->AddToken(FActorToken::Create(ReferenceActorDescView.GetActorSoftPath().ToString(), ReferenceActorDescView.GetGuid(), FText::FromName(ReferenceActorDescView.GetActorLabelOrName())));
-
-	AddAdditionalNameToken(Message, FName(TEXT("WorldPartition_ActorReferenceActorInAnotherRuntimeGrid_CheckForErrors")));
+		->AddToken(FActorToken::Create(ReferenceActorDescView.GetActorSoftPath().ToString(), ReferenceActorDescView.GetGuid(), FText::FromString(GetActorName(ReferenceActorDescView))))
+		->AddToken(FMapErrorToken::Create(TEXT("WorldPartition_ActorReferenceActorInAnotherRuntimeGrid_CheckForErrors")));
 
 	HandleTokenizedMessage(MoveTemp(Message));
 }
@@ -84,9 +89,8 @@ void ITokenizedMessageErrorHandler::OnInvalidReferenceLevelScriptStreamed(const
 {
 	TSharedRef<FTokenizedMessage> Message = FTokenizedMessage::Create(EMessageSeverity::Error);
 	Message->AddToken(FTextToken::Create(LOCTEXT("TokenMessage_WorldPartition_LevelScriptBlueprintStreamedActorReference", "Level Script Blueprint references streamed actor")))
-		->AddToken(FActorToken::Create(ActorDescView.GetActorSoftPath().ToString(), ActorDescView.GetGuid(), FText::FromString(GetFullActorName(ActorDescView))));
-
-	AddAdditionalNameToken(Message, FName(TEXT("WorldPartition_LevelScriptBlueprintRefefenceStreamed_CheckForErrors")));
+		->AddToken(FActorToken::Create(ActorDescView.GetActorSoftPath().ToString(), ActorDescView.GetGuid(), FText::FromString(GetActorName(ActorDescView))))
+		->AddToken(FMapErrorToken::Create(TEXT("WorldPartition_LevelScriptBlueprintRefefenceStreamed_CheckForErrors")));
 
 	HandleTokenizedMessage(MoveTemp(Message));
 	
@@ -96,10 +100,9 @@ void ITokenizedMessageErrorHandler::OnInvalidReferenceLevelScriptDataLayers(cons
 {
 	TSharedRef<FTokenizedMessage> Message = FTokenizedMessage::Create(EMessageSeverity::Error);
 	Message->AddToken(FTextToken::Create(LOCTEXT("TokenMessage_WorldPartition_LevelScriptBlueprintActorReference", "Level Script Blueprint references actor")))
-		->AddToken(FActorToken::Create(ActorDescView.GetActorSoftPath().ToString(), ActorDescView.GetGuid(), FText::FromString(GetFullActorName(ActorDescView))))
-		->AddToken(FTextToken::Create(LOCTEXT("TokenMessage_WorldPartition_LevelScriptBlueprintDataLayerReference", "with a non empty set of data layers")));
-
-	AddAdditionalNameToken(Message, FName(TEXT("WorldPartition_LevelScriptBlueprintRefefenceDataLayer_CheckForErrors")));
+		->AddToken(FActorToken::Create(ActorDescView.GetActorSoftPath().ToString(), ActorDescView.GetGuid(), FText::FromString(GetActorName(ActorDescView))))
+		->AddToken(FTextToken::Create(LOCTEXT("TokenMessage_WorldPartition_LevelScriptBlueprintDataLayerReference", "with a non empty set of data layers")))
+		->AddToken(FMapErrorToken::Create(TEXT("WorldPartition_LevelScriptBlueprintRefefenceDataLayer_CheckForErrors")));
 
 	HandleTokenizedMessage(MoveTemp(Message));
 }
@@ -109,9 +112,8 @@ void ITokenizedMessageErrorHandler::OnInvalidReferenceDataLayerAsset(const UData
 	TSharedRef<FTokenizedMessage> Message = FTokenizedMessage::Create(EMessageSeverity::Error);
 	Message->AddToken(FTextToken::Create(LOCTEXT("TokenMessage_DataLayers_DataLayer", "Data layer")))
 		->AddToken(FUObjectToken::Create(DataLayerInstance))
-		->AddToken(FTextToken::Create(LOCTEXT("TokenMessage_DataLayers_NullAsset", "Does not have Data Layer Asset")));
-
-	AddAdditionalNameToken(Message, FName(TEXT("DataLayers_InvalidAsset_CheckForErrors")));
+		->AddToken(FTextToken::Create(LOCTEXT("TokenMessage_DataLayers_NullAsset", "Does not have Data Layer Asset")))
+		->AddToken(FMapErrorToken::Create(TEXT("DataLayers_InvalidAsset_CheckForErrors")));
 
 	HandleTokenizedMessage(MoveTemp(Message));
 }
@@ -123,9 +125,8 @@ void ITokenizedMessageErrorHandler::OnDataLayerHierarchyTypeMismatch(const UData
 		->AddToken(FTextToken::Create(FText::FromString(DataLayerInstance->GetDataLayerShortName())))
 		->AddToken(FTextToken::Create(LOCTEXT("TokenMessage_DataLayers_IsRntime", "is Runtime but its parent data layer")))
 		->AddToken(FTextToken::Create(FText::FromString(Parent->GetDataLayerShortName())))
-		->AddToken(FTextToken::Create(LOCTEXT("TokenMessage_DataLayers_IsNot", "is not")));
-
-	AddAdditionalNameToken(Message, FName(TEXT("DataLayers_HierarchyTypeMisMatch_CheckForErrors")));
+		->AddToken(FTextToken::Create(LOCTEXT("TokenMessage_DataLayers_IsNot", "is not")))
+		->AddToken(FMapErrorToken::Create(TEXT("DataLayers_HierarchyTypeMisMatch_CheckForErrors")));
 
 	HandleTokenizedMessage(MoveTemp(Message));
 }
@@ -138,9 +139,8 @@ void ITokenizedMessageErrorHandler::OnDataLayerAssetConflict(const UDataLayerIns
 		->AddToken(FTextToken::Create(LOCTEXT("TokenMessage_DataLayers_AndDataLayerInstance", "and Data Layer Instance")))
 		->AddToken(FUObjectToken::Create(ConflictingDataLayerInstance))
 		->AddToken(FTextToken::Create(LOCTEXT("TokenMessage_DataLayers_BothReferencing", "are both referencing Data Layer Asset")))
-		->AddToken(FUObjectToken::Create(DataLayerInstance->GetAsset()));
-
-	AddAdditionalNameToken(Message, FName(TEXT("DataLayers_AssetConflict_CheckForErrors")));
+		->AddToken(FUObjectToken::Create(DataLayerInstance->GetAsset()))
+		->AddToken(FMapErrorToken::Create(TEXT("DataLayers_AssetConflict_CheckForErrors")));
 
 	HandleTokenizedMessage(MoveTemp(Message));
 }
@@ -149,9 +149,8 @@ void ITokenizedMessageErrorHandler::OnActorNeedsResave(const FWorldPartitionActo
 {
 	TSharedRef<FTokenizedMessage> Message = FTokenizedMessage::Create(EMessageSeverity::Warning);
 	Message->AddToken(FTextToken::Create(LOCTEXT("TokenMessage_WorldPartition_ActorNeedsResave", "Actor needs resave")))
-		->AddToken(FActorToken::Create(ActorDescView.GetActorSoftPath().ToString(), ActorDescView.GetGuid(), FText::FromString(GetFullActorName(ActorDescView))));
-
-	AddAdditionalNameToken(Message, FName(TEXT("WorldPartition_ActorNeedsResave_CheckForErrors")));
+		->AddToken(FActorToken::Create(ActorDescView.GetActorSoftPath().ToString(), ActorDescView.GetGuid(), FText::FromString(GetActorName(ActorDescView))))
+		->AddToken(FMapErrorToken::Create(TEXT("WorldPartition_ActorNeedsResave_CheckForErrors")));
 
 	HandleTokenizedMessage(MoveTemp(Message));
 }
@@ -187,11 +186,10 @@ void ITokenizedMessageErrorHandler::OnLevelInstanceInvalidWorldAsset(const FWorl
 
 	TSharedRef<FTokenizedMessage> Message = FTokenizedMessage::Create(MessageSeverity);
 	Message->AddToken(FTextToken::Create(LOCTEXT("TokenMessage_WorldPartition_LevelInstance", "Level Instance")))
-		->AddToken(FTextToken::Create(FText::FromString(GetFullActorName(ActorDescView))))
+		->AddToken(FTextToken::Create(FText::FromString(GetActorName(ActorDescView))))
 		->AddToken(FTextToken::Create(ReasonText))
-		->AddToken(FAssetNameToken::Create(WorldAssetPath));
-
-	AddAdditionalNameToken(Message, FName(TEXT("WorldPartition_LevelInstanceInvalidWorldAsset_CheckForErrors")));
+		->AddToken(FAssetNameToken::Create(WorldAssetPath))
+		->AddToken(FMapErrorToken::Create(TEXT("WorldPartition_LevelInstanceInvalidWorldAsset_CheckForErrors")));
 
 	HandleTokenizedMessage(MoveTemp(Message));
 }
diff --git a/Engine/Source/Runtime/Engine/Public/WorldPartition/ErrorHandling/WorldPartitionStreamingGenerationMapCheckErrorHandler.h b/Engine/Source/Runtime/Engine/Public/WorldPartition/ErrorHandling/WorldPartitionStreamingGenerationMapCheckErrorHandler.h
--- a/Engine/Source/Runtime/Engine/Public/WorldPartition/ErrorHandling/WorldPartitionStreamingGenerationMapCheckErrorHandler.h
+++ b/Engine/Source/Runtime/Engine/Public/WorldPartition/ErrorHandling/WorldPartitionStreamingGenerationMapCheckErrorHandler.h
@@ -7,8 +7,6 @@
 class ENGINE_API FStreamingGenerationMapCheckErrorHandler : public ITokenizedMessageErrorHandler
 {
 protected:
-	virtual void AddAdditionalNameToken(TSharedRef<FTokenizedMessage>& InMessage, const FName& InErrorName) override;
-
 	virtual void HandleTokenizedMessage(TSharedRef<FTokenizedMessage>&& ErrorMessage) override;
 };
 #endif
diff --git a/Engine/Source/Runtime/Engine/Public/WorldPartition/ErrorHandling/WorldPartitionStreamingGenerationTokenizedMessageErrorHandler.h b/Engine/Source/Runtime/Engine/Public/WorldPartition/ErrorHandling/WorldPartitionStreamingGenerationTokenizedMessageErrorHandler.h
--- a/Engine/Source/Runtime/Engine/Public/WorldPartition/ErrorHandling/WorldPartitionStreamingGenerationTokenizedMessageErrorHandler.h
+++ b/Engine/Source/Runtime/Engine/Public/WorldPartition/ErrorHandling/WorldPartitionStreamingGenerationTokenizedMessageErrorHandler.h
@@ -22,8 +22,6 @@ class ENGINE_API ITokenizedMessageErrorHandler : public IStreamingGenerationErro
 	virtual void OnLevelInstanceInvalidWorldAsset(const FWorldPartitionActorDescView& ActorDescView, FName WorldAsset, ELevelInstanceInvalidReason Reason) override;
 
 protected:
-	virtual void AddAdditionalNameToken(TSharedRef<FTokenizedMessage>& InMessage, const FName& InErrorName) {}
-
 	virtual void HandleTokenizedMessage(TSharedRef<FTokenizedMessage>&& ErrorMessage) = 0;
 };
 
