Fix for Anim Notifies not getting fired in the server side.

diff --git a/Engine/Source/Runtime/Engine/Private/Components/CharacterMovementComponent.cpp b/Engine/Source/Runtime/Engine/Private/Components/CharacterMovementComponent.cpp
--- a/Engine/Source/Runtime/Engine/Private/Components/CharacterMovementComponent.cpp
+++ b/Engine/Source/Runtime/Engine/Private/Components/CharacterMovementComponent.cpp
@@ -137,7 +137,7 @@ namespace CharacterMovementCVars
 		TEXT("0: Disable, 1: Enable"),
 		ECVF_Default);
 
-	static int32 EnableQueuedAnimEventsOnServer = 0;
+	static int32 EnableQueuedAnimEventsOnServer = 1;
 	FAutoConsoleVariableRef CVarEnableQueuedAnimEventsOnServer(
 		TEXT("a.EnableQueuedAnimEventsOnServer"),
 		EnableQueuedAnimEventsOnServer,
