Path Tracer: Don't display lighting must be rebuilt message when using the path tracer

diff --git a/Engine/Source/Runtime/Engine/Private/UnrealEngine.cpp b/Engine/Source/Runtime/Engine/Private/UnrealEngine.cpp
--- a/Engine/Source/Runtime/Engine/Private/UnrealEngine.cpp
+++ b/Engine/Source/Runtime/Engine/Private/UnrealEngine.cpp
@@ -11621,9 +11621,14 @@ float DrawMapWarnings(UWorld* World, FViewport* Viewport, FCanvas* Canvas, UCanv
 		Canvas->DrawItem(TextItem);
 	}
 
+	bool bIsPathTracing = false;
+	if (Viewport && Viewport->GetClient() && Viewport->GetClient()->GetEngineShowFlags())
+	{
+		bIsPathTracing = Viewport->GetClient()->GetEngineShowFlags()->PathTracing != 0;
+	}
 
 	// Put the messages over fairly far to stay in the safe zone on consoles
-	if (World->NumLightingUnbuiltObjects > 0)
+	if (World->NumLightingUnbuiltObjects > 0 && !bIsPathTracing)
 	{
 		SmallTextItem.SetColor(FLinearColor::White);
 		// Color unbuilt lighting red if encountered within the last second
@@ -11669,7 +11674,7 @@ float DrawMapWarnings(UWorld* World, FViewport* Viewport, FCanvas* Canvas, UCanv
 		}
 	}
 
-	if (World->NumUnbuiltReflectionCaptures > 0)
+	if (World->NumUnbuiltReflectionCaptures > 0 && !bIsPathTracing)
 	{
 		int32 NumLightingScenariosEnabled = 0;
 
