Prepared code for removing Threading.h include from ParticleHandleFwd.h..

diff --git a/Engine/Source/Runtime/Engine/Private/ActorConstruction.cpp b/Engine/Source/Runtime/Engine/Private/ActorConstruction.cpp
--- a/Engine/Source/Runtime/Engine/Private/ActorConstruction.cpp
+++ b/Engine/Source/Runtime/Engine/Private/ActorConstruction.cpp
@@ -2,6 +2,8 @@
 
 #include "EngineLogs.h"
 #include "Math/RandomStream.h"
+#include "Misc/ScopeExit.h"
+#include "ProfilingDebugging/CsvProfiler.h"
 #include "UObject/UObjectThreadContext.h"
 #include "Serialization/ObjectReader.h"
 #include "Engine/Blueprint.h"
diff --git a/Engine/Source/Runtime/Engine/Private/Animation/AnimTrace.cpp b/Engine/Source/Runtime/Engine/Private/Animation/AnimTrace.cpp
--- a/Engine/Source/Runtime/Engine/Private/Animation/AnimTrace.cpp
+++ b/Engine/Source/Runtime/Engine/Private/Animation/AnimTrace.cpp
@@ -1,8 +1,6 @@
 // Copyright Epic Games, Inc. All Rights Reserved.
 
 #include "Animation/AnimTrace.h"
-#include "Trace/Detail/Important/ImportantLogScope.h"
-#include "Trace/Detail/Important/SharedBuffer.h"
 
 #if ANIM_TRACE_ENABLED
 
@@ -14,7 +12,7 @@
 #include "Animation/AnimNode_SequencePlayer.h"
 #include "Animation/AnimNotifies/AnimNotify.h"
 #include "Animation/AnimNotifies/AnimNotifyState.h"
-#include "Trace/Detail/Atomic.h"
+#include "Trace/Trace.inl"
 #include "UObject/UObjectAnnotation.h"
 
 UE_TRACE_CHANNEL_DEFINE(AnimationChannel);
diff --git a/Engine/Source/Runtime/Engine/Private/Collision/WorldCollisionAsync.cpp b/Engine/Source/Runtime/Engine/Private/Collision/WorldCollisionAsync.cpp
--- a/Engine/Source/Runtime/Engine/Private/Collision/WorldCollisionAsync.cpp
+++ b/Engine/Source/Runtime/Engine/Private/Collision/WorldCollisionAsync.cpp
@@ -5,9 +5,12 @@
 =============================================================================*/
 
 #include "Engine/World.h"
+#include "Async/TaskGraphInterfaces.h"
+#include "HAL/IConsoleManager.h"
 #include "Misc/App.h"
 #include "Misc/Fork.h"
 #include "Physics/Experimental/PhysInterface_Chaos.h"
+#include "ProfilingDebugging/CsvProfiler.h"
 
 CSV_DEFINE_CATEGORY(WorldCollision, true);
 
diff --git a/Engine/Source/Runtime/Engine/Private/Components/ProjectileMovementComponent.cpp b/Engine/Source/Runtime/Engine/Private/Components/ProjectileMovementComponent.cpp
--- a/Engine/Source/Runtime/Engine/Private/Components/ProjectileMovementComponent.cpp
+++ b/Engine/Source/Runtime/Engine/Private/Components/ProjectileMovementComponent.cpp
@@ -5,6 +5,7 @@
 #include "Engine/World.h"
 #include "Components/PrimitiveComponent.h"
 #include "GameFramework/WorldSettings.h"
+#include "ProfilingDebugging/CsvProfiler.h"
 
 #include UE_INLINE_GENERATED_CPP_BY_NAME(ProjectileMovementComponent)
 
diff --git a/Engine/Source/Runtime/Engine/Private/DataChannel.cpp b/Engine/Source/Runtime/Engine/Private/DataChannel.cpp
--- a/Engine/Source/Runtime/Engine/Private/DataChannel.cpp
+++ b/Engine/Source/Runtime/Engine/Private/DataChannel.cpp
@@ -11,6 +11,8 @@
 #include "Engine/Level.h"
 #include "Engine/World.h"
 #include "GameFramework/WorldSettings.h"
+#include "Misc/MemStack.h"
+#include "Misc/ScopeExit.h"
 #include "Net/Core/Trace/Private/NetTraceInternal.h"
 #include "UObject/UObjectIterator.h"
 #include "EngineStats.h"
diff --git a/Engine/Source/Runtime/Engine/Private/DataReplication.cpp b/Engine/Source/Runtime/Engine/Private/DataReplication.cpp
--- a/Engine/Source/Runtime/Engine/Private/DataReplication.cpp
+++ b/Engine/Source/Runtime/Engine/Private/DataReplication.cpp
@@ -8,6 +8,8 @@
 #include "Containers/StaticBitArray.h"
 #include "EngineStats.h"
 #include "Engine/World.h"
+#include "Misc/MemStack.h"
+#include "Misc/ScopeExit.h"
 #include "Net/Core/Trace/Private/NetTraceInternal.h"
 #include "Net/NetworkProfiler.h"
 #include "Engine/PackageMapClient.h"
diff --git a/Engine/Source/Runtime/Engine/Private/DemoNetDriver.cpp b/Engine/Source/Runtime/Engine/Private/DemoNetDriver.cpp
--- a/Engine/Source/Runtime/Engine/Private/DemoNetDriver.cpp
+++ b/Engine/Source/Runtime/Engine/Private/DemoNetDriver.cpp
@@ -10,6 +10,7 @@
 
 
 #include "Engine/DemoNetDriver.h"
+#include "Async/TaskGraphInterfaces.h"
 #include "Engine/DemoNetConnection.h"
 #include "Engine/GameInstance.h"
 #include "Engine/LevelStreaming.h"
diff --git a/Engine/Source/Runtime/Engine/Private/Engine/ReplicatedState.cpp b/Engine/Source/Runtime/Engine/Private/Engine/ReplicatedState.cpp
--- a/Engine/Source/Runtime/Engine/Private/Engine/ReplicatedState.cpp
+++ b/Engine/Source/Runtime/Engine/Private/Engine/ReplicatedState.cpp
@@ -2,6 +2,7 @@
 
 #include "Engine/ReplicatedState.h"
 #include "Engine/World.h"
+#include "HAL/IConsoleManager.h"
 
 #include UE_INLINE_GENERATED_CPP_BY_NAME(ReplicatedState)
 
diff --git a/Engine/Source/Runtime/Engine/Private/InGamePerformanceTracker.cpp b/Engine/Source/Runtime/Engine/Private/InGamePerformanceTracker.cpp
--- a/Engine/Source/Runtime/Engine/Private/InGamePerformanceTracker.cpp
+++ b/Engine/Source/Runtime/Engine/Private/InGamePerformanceTracker.cpp
@@ -2,6 +2,7 @@
 
 #include "InGamePerformanceTracker.h"
 #include "Engine/World.h"
+#include "HAL/IConsoleManager.h"
 
 DEFINE_LOG_CATEGORY_STATIC(LogInGamePerformanceTracker, Log, All);
 
diff --git a/Engine/Source/Runtime/Engine/Private/LevelTick.cpp b/Engine/Source/Runtime/Engine/Private/LevelTick.cpp
--- a/Engine/Source/Runtime/Engine/Private/LevelTick.cpp
+++ b/Engine/Source/Runtime/Engine/Private/LevelTick.cpp
@@ -5,6 +5,7 @@
 =============================================================================*/
 
 #include "Engine/Level.h"
+#include "Async/ParallelFor.h"
 #include "Misc/TimeGuard.h"
 #include "GameFramework/WorldSettings.h"
 #include "UObject/UObjectStats.h"
diff --git a/Engine/Source/Runtime/Engine/Private/NetworkProfiler.cpp b/Engine/Source/Runtime/Engine/Private/NetworkProfiler.cpp
--- a/Engine/Source/Runtime/Engine/Private/NetworkProfiler.cpp
+++ b/Engine/Source/Runtime/Engine/Private/NetworkProfiler.cpp
@@ -6,6 +6,7 @@
 
 #include "Net/NetworkProfiler.h"
 #include "HAL/FileManager.h"
+#include "HAL/IConsoleManager.h"
 #include "Misc/Paths.h"
 #include "Misc/App.h"
 #include "Engine/World.h"
diff --git a/Engine/Source/Runtime/Engine/Private/PhysicsEngine/PhysicsCollisionHandler.cpp b/Engine/Source/Runtime/Engine/Private/PhysicsEngine/PhysicsCollisionHandler.cpp
--- a/Engine/Source/Runtime/Engine/Private/PhysicsEngine/PhysicsCollisionHandler.cpp
+++ b/Engine/Source/Runtime/Engine/Private/PhysicsEngine/PhysicsCollisionHandler.cpp
@@ -1,6 +1,7 @@
 // Copyright Epic Games, Inc. All Rights Reserved.
 
 #include "PhysicsEngine/PhysicsCollisionHandler.h"
+#include "ChaosLog.h"
 #include "Kismet/GameplayStatics.h"
 #include "Components/PrimitiveComponent.h"
 
diff --git a/Engine/Source/Runtime/Engine/Private/TickTaskManager.cpp b/Engine/Source/Runtime/Engine/Private/TickTaskManager.cpp
--- a/Engine/Source/Runtime/Engine/Private/TickTaskManager.cpp
+++ b/Engine/Source/Runtime/Engine/Private/TickTaskManager.cpp
@@ -6,6 +6,7 @@
 
 #include "Engine/Level.h"
 #include "Engine/World.h"
+#include "ProfilingDebugging/CsvProfiler.h"
 #include "Stats/StatsTrace.h"
 #include "TickTaskManagerInterface.h"
 #include "Async/ParallelFor.h"
diff --git a/Engine/Source/Runtime/Engine/Private/Tickable.cpp b/Engine/Source/Runtime/Engine/Private/Tickable.cpp
--- a/Engine/Source/Runtime/Engine/Private/Tickable.cpp
+++ b/Engine/Source/Runtime/Engine/Private/Tickable.cpp
@@ -2,6 +2,7 @@
 
 #include "Tickable.h"
 #include "Engine/World.h"
+#include "ProfilingDebugging/CsvProfiler.h"
 
 DECLARE_CYCLE_STAT(TEXT("TickableGameObjects Time"), STAT_TickableGameObjectsTime, STATGROUP_Game);
 
diff --git a/Engine/Source/Runtime/Engine/Private/Timeline.cpp b/Engine/Source/Runtime/Engine/Private/Timeline.cpp
--- a/Engine/Source/Runtime/Engine/Private/Timeline.cpp
+++ b/Engine/Source/Runtime/Engine/Private/Timeline.cpp
@@ -6,6 +6,7 @@
 
 #include "Curves/CurveLinearColor.h"
 #include "Curves/CurveVector.h"
+#include "ProfilingDebugging/CsvProfiler.h"
 #include "UObject/EnumProperty.h"
 #include "UObject/Package.h"
 #include "GameFramework/WorldSettings.h"
diff --git a/Engine/Source/Runtime/Engine/Private/TimerManager.cpp b/Engine/Source/Runtime/Engine/Private/TimerManager.cpp
--- a/Engine/Source/Runtime/Engine/Private/TimerManager.cpp
+++ b/Engine/Source/Runtime/Engine/Private/TimerManager.cpp
@@ -9,6 +9,7 @@
 #include "Engine/GameInstance.h"
 #include "Logging/LogScopedCategoryAndVerbosityOverride.h"
 #include "Misc/CoreDelegates.h"
+#include "ProfilingDebugging/CsvProfiler.h"
 #include "Stats/StatsTrace.h"
 #include "UnrealEngine.h"
 #include "Misc/TimeGuard.h"
diff --git a/Engine/Source/Runtime/Engine/Private/VisualLogger/VisualLoggerTraceDevice.cpp b/Engine/Source/Runtime/Engine/Private/VisualLogger/VisualLoggerTraceDevice.cpp
--- a/Engine/Source/Runtime/Engine/Private/VisualLogger/VisualLoggerTraceDevice.cpp
+++ b/Engine/Source/Runtime/Engine/Private/VisualLogger/VisualLoggerTraceDevice.cpp
@@ -1,12 +1,13 @@
 // Copyright Epic Games, Inc. All Rights Reserved.
 #include "VisualLogger/VisualLoggerTraceDevice.h"
+
+#if ENABLE_VISUAL_LOG
+
 #include "Serialization/BufferArchive.h"
 #include "ObjectTrace.h"
-
+#include "Trace/Trace.inl"
 #include "VisualLogger/VisualLoggerCustomVersion.h"
 
-#if ENABLE_VISUAL_LOG
-
 UE_TRACE_CHANNEL_DEFINE(VisualLoggerChannel);
 
 UE_TRACE_EVENT_BEGIN(VisualLogger, VisualLogEntry)
diff --git a/Engine/Source/Runtime/Engine/Private/WorldPartition/DataLayer/WorldDataLayers.cpp b/Engine/Source/Runtime/Engine/Private/WorldPartition/DataLayer/WorldDataLayers.cpp
--- a/Engine/Source/Runtime/Engine/Private/WorldPartition/DataLayer/WorldDataLayers.cpp
+++ b/Engine/Source/Runtime/Engine/Private/WorldPartition/DataLayer/WorldDataLayers.cpp
@@ -6,6 +6,7 @@
 
 #include "WorldPartition/DataLayer/WorldDataLayers.h"
 #include "Engine/Level.h"
+#include "ProfilingDebugging/CsvProfiler.h"
 #include "WorldPartition/DataLayer/DataLayerSubsystem.h"
 #include "WorldPartition/DataLayer/DataLayer.h"
 #include "WorldPartition/DataLayer/DataLayerInstanceWithAsset.h"
