//UE5/Main - UE.EditorAutomation(RunTest=Editor) Mac - Several tests exited prematurely due to critical failure
#fix Added new override point OnResetSkeleton, which allows UAnimSequence to wait for any in-flight compression tasks to finish before swapping out the skeleton

diff --git a/Engine/Source/Editor/AnimGraph/Private/AnimGraphNode_BlendSpaceGraphBase.cpp b/Engine/Source/Editor/AnimGraph/Private/AnimGraphNode_BlendSpaceGraphBase.cpp
--- a/Engine/Source/Editor/AnimGraph/Private/AnimGraphNode_BlendSpaceGraphBase.cpp
+++ b/Engine/Source/Editor/AnimGraph/Private/AnimGraphNode_BlendSpaceGraphBase.cpp
@@ -165,7 +165,7 @@ void UAnimGraphNode_BlendSpaceGraphBase::SetupFromAsset(const FAssetData& InAsse
 			BlendSpaceGraph->BlendSpace = BlendSpace = DuplicateObject(Asset, BlendSpaceGraph);
 			BlendSpaceGraph->BlendSpace->ClearFlags(RF_Public | RF_Standalone | RF_Transient);
 			BlendSpaceGraph->BlendSpace->SetFlags(RF_Transactional);
-			BlendSpaceGraph->BlendSpace->ResetSkeleton(nullptr);
+			BlendSpaceGraph->BlendSpace->SetSkeleton(nullptr);
 			BlendSpaceGraph->BlendSpace->EmptyMetaData();
 			BlendSpaceGraph->BlendSpace->RemoveUserDataOfClass(UAssetUserData::StaticClass());
 			BlendSpaceGraph->bAllowDeletion = false;
diff --git a/Engine/Source/Runtime/Engine/Classes/Animation/AnimSequence.h b/Engine/Source/Runtime/Engine/Classes/Animation/AnimSequence.h
--- a/Engine/Source/Runtime/Engine/Classes/Animation/AnimSequence.h
+++ b/Engine/Source/Runtime/Engine/Classes/Animation/AnimSequence.h
@@ -404,6 +404,7 @@ public:
 #if WITH_EDITOR
 	virtual bool GetAllAnimationSequencesReferred(TArray<UAnimationAsset*>& AnimationAssets, bool bRecursive = true) override;
 	virtual void ReplaceReferredAnimations(const TMap<UAnimationAsset*, UAnimationAsset*>& ReplacementMap) override;
+	virtual void OnSetSkeleton(USkeleton* NewSkeleton) override;
 #endif
 	//~ End UAnimationAsset Interface
 
diff --git a/Engine/Source/Runtime/Engine/Classes/Animation/AnimationAsset.h b/Engine/Source/Runtime/Engine/Classes/Animation/AnimationAsset.h
--- a/Engine/Source/Runtime/Engine/Classes/Animation/AnimationAsset.h
+++ b/Engine/Source/Runtime/Engine/Classes/Animation/AnimationAsset.h
@@ -1053,6 +1053,7 @@ public:
 	virtual float GetPlayLength() const { return 0.f; };
 
 	void SetSkeleton(USkeleton* NewSkeleton);
+	UE_DEPRECATED(5.2, "ResetSkeleton has been deprecated, use ReplaceSkeleton or SetSkeleton instead")
 	void ResetSkeleton(USkeleton* NewSkeleton);
 	virtual void PostLoad() override;
 
@@ -1098,6 +1099,8 @@ public:
 	 */
 	bool ReplaceSkeleton(USkeleton* NewSkeleton, bool bConvertSpaces=false);
 
+	virtual void OnSetSkeleton(USkeleton* NewSkeleton) {}
+
 	// Helper function for GetAllAnimationSequencesReferred, it adds itself first and call GetAllAnimationSEquencesReferred
 	void HandleAnimReferenceCollection(TArray<UAnimationAsset*>& AnimationAssets, bool bRecursive);
 
diff --git a/Engine/Source/Runtime/Engine/Private/Animation/AnimSequence.cpp b/Engine/Source/Runtime/Engine/Private/Animation/AnimSequence.cpp
--- a/Engine/Source/Runtime/Engine/Private/Animation/AnimSequence.cpp
+++ b/Engine/Source/Runtime/Engine/Private/Animation/AnimSequence.cpp
@@ -2406,7 +2406,7 @@ void UAnimSequence::RemapTracksToNewSkeleton( USkeleton* NewSkeleton, bool bConv
 		}
 		// @todo end rig testing
 		// @IMPORTANT: now otherwise this will try to do bone to bone mapping
-		else if(OldSkeleton)
+		else if(OldSkeleton && OldSkeleton != NewSkeleton)
 		{
 			// Validate animation tracks against the new skeleton, any tracks linked to bones that do not exist in the new hierarchy are removed
 			Controller->RemoveBoneTracksMissingFromSkeleton(NewSkeleton);
@@ -3746,6 +3746,11 @@ void UAnimSequence::EvaluateAttributes(FAnimationPoseData& OutAnimationPoseData,
 }
 
 #if WITH_EDITOR
+void UAnimSequence::OnSetSkeleton(USkeleton* NewSkeleton)
+{
+	UE::Anim::FAnimSequenceCompilingManager::Get().FinishCompilation({this});
+}
+
 void UAnimSequence::RemoveCustomAttribute(const FName& BoneName, const FName& AttributeName)
 {
 	ValidateModel();
diff --git a/Engine/Source/Runtime/Engine/Private/AnimationAsset.cpp b/Engine/Source/Runtime/Engine/Private/AnimationAsset.cpp
--- a/Engine/Source/Runtime/Engine/Private/AnimationAsset.cpp
+++ b/Engine/Source/Runtime/Engine/Private/AnimationAsset.cpp
@@ -301,9 +301,7 @@ void UAnimationAsset::PostLoad()
 
 void UAnimationAsset::ResetSkeleton(USkeleton* NewSkeleton)
 {
-// @TODO LH, I'd like this to work outside of editor, but that requires unlocking track names data in game
 #if WITH_EDITOR
-	Skeleton = NULL;
 	ReplaceSkeleton(NewSkeleton);
 #endif
 }
@@ -360,7 +358,16 @@ void UAnimationAsset::RemoveMetaData(TArrayView<UAnimMetaData*> MetaDataInstance
 
 void UAnimationAsset::SetSkeleton(USkeleton* NewSkeleton)
 {
-	if (NewSkeleton && NewSkeleton != Skeleton)
+#if WITH_EDITOR
+	OnSetSkeleton(NewSkeleton);
+#endif // WITH_EDITOR
+
+	if (NewSkeleton == nullptr)
+	{
+		Skeleton = nullptr;
+		SkeletonGuid.Invalidate();
+	}
+	else if (NewSkeleton != Skeleton)
 	{
 		Skeleton = NewSkeleton;
 		SkeletonGuid = NewSkeleton->GetGuid();
@@ -417,8 +424,10 @@ void UAnimationAsset::RemapTracksToNewSkeleton(USkeleton* NewSkeleton, bool bCon
 bool UAnimationAsset::ReplaceSkeleton(USkeleton* NewSkeleton, bool bConvertSpaces/*=false*/)
 {
 	// if it's not same 
-	if (NewSkeleton != Skeleton)
+	if (NewSkeleton && (NewSkeleton != Skeleton || NewSkeleton->GetGuid() != SkeletonGuid))
 	{
+		OnSetSkeleton(NewSkeleton);
+
 		// get all sequences that need to change
 		TArray<UAnimationAsset*> AnimAssetsToReplace;
 
@@ -641,9 +650,13 @@ void UAnimationAsset::ValidateSkeleton()
 {
 	if (Skeleton && Skeleton->GetGuid() != SkeletonGuid)
 	{
+#if WITH_EDITOR
 		// reset Skeleton
-		ResetSkeleton(Skeleton);
+		ReplaceSkeleton(Skeleton);
 		UE_LOG(LogAnimation, Verbose, TEXT("Needed to reset skeleton. Resave this asset to speed up load time: %s"), *GetPathNameSafe(this));
+#else
+		UE_LOG(LogAnimation, Warning, TEXT("Skeleton GUID is out-of-date, this should have been updated during cook. %s"), *GetPathNameSafe(this));
+#endif
 	}
 }
 
