PR #9866: SendAllEndOfFrameUpdates is called less aggressively by Reflections (Contributed by [REDACTED_PERSON])

diff --git a/Engine/Source/Runtime/Engine/Private/Components/ReflectionCaptureComponent.cpp b/Engine/Source/Runtime/Engine/Private/Components/ReflectionCaptureComponent.cpp
--- a/Engine/Source/Runtime/Engine/Private/Components/ReflectionCaptureComponent.cpp
+++ b/Engine/Source/Runtime/Engine/Private/Components/ReflectionCaptureComponent.cpp
@@ -1056,6 +1056,20 @@ void UReflectionCaptureComponent::UpdateReflectionCaptureContents(UWorld* WorldT
 #endif
 		)
 	{
+		// If no captures to update or load, on the condition that SendAllEndOfFrameUpdates() does not change their status, 
+		// WorldCombinedCaptures is empty and AllocateReflectionCaptures() simply returns. The logic here becomes simply calling
+		// an unnecessary SendAllEndOfFrameUpdates() and return. So we could have this early out.
+		
+		const bool bHasCapturesToUpdate = ReflectionCapturesToUpdate.Num() > 0;
+		const bool bHasCapturesToLoad = ReflectionCapturesToUpdateForLoad.Num() > 0;
+		const bool bAlwaysUpdatesCaptures = CVarReflectionCaptureUpdateEveryFrame.GetValueOnGameThread() != 0;
+		const bool bNeedsAnyUpdate = bHasCapturesToUpdate || bHasCapturesToLoad || bAlwaysUpdatesCaptures;
+		
+		if (!bNeedsAnyUpdate)
+		{
+			return;
+		}
+
 		//guarantee that all render proxies are up to date before kicking off this render
 		WorldToUpdate->SendAllEndOfFrameUpdates();
 
