Fix static analyzer warning.

diff --git a/Engine/Source/Runtime/Engine/Private/Tests/WorldPartition/WorldPartitionStreamingGenerationTests.cpp b/Engine/Source/Runtime/Engine/Private/Tests/WorldPartition/WorldPartitionStreamingGenerationTests.cpp
--- a/Engine/Source/Runtime/Engine/Private/Tests/WorldPartition/WorldPartitionStreamingGenerationTests.cpp
+++ b/Engine/Source/Runtime/Engine/Private/Tests/WorldPartition/WorldPartitionStreamingGenerationTests.cpp
@@ -37,21 +37,21 @@ namespace WorldPartitionTests
 		);
 		
 		UWorld* World = ScopedEditorWorld.GetWorld();
-		if (!TestTrue(TEXT("Missing World Object"), !!World))
+		if (!TestNotNull(TEXT("Missing World Object"), World))
 		{
 			return false;
 		}
 
 		check(World);
 		UWorldPartition* WorldPartition = World->GetWorldPartition();
-		if (!TestTrue(TEXT("Missing World Partition Object"), !!WorldPartition))
+		if (!TestNotNull(TEXT("Missing World Partition Object"), WorldPartition))
 		{
 			return false;
 		}
 		
 		check(WorldPartition);
 		UActorDescContainer* ActorDescMainContainer = WorldPartition->GetActorDescContainer();
-		if (!TestTrue(TEXT("Missing World Partition Container"), !!ActorDescMainContainer))
+		if (!TestNotNull(TEXT("Missing World Partition Container"), ActorDescMainContainer))
 		{
 			return false;
 		}
@@ -68,7 +68,7 @@ namespace WorldPartitionTests
 		}
 
 		AActor* Actor = ActorRef->GetActor();
-		if (!TestTrue(TEXT("Missing Actor"), !!Actor))
+		if (!TestNotNull(TEXT("Missing Actor"), Actor))
 		{
 			return false;
 		}
@@ -79,13 +79,13 @@ namespace WorldPartitionTests
 			return false;
 		}
 
-		if (!TestTrue(TEXT("Actor Handle Not Loaded"), !ActorHandle->GetActor()))
+		if (!TestNull(TEXT("Actor Handle Not Loaded"), ActorHandle->GetActor()))
 		{
 			return false;
 		}
 
 		UObject* ResolvedObject = ActorHandle->GetActorSoftPath().TryLoad();
-		if (!TestTrue(TEXT("Actor Handle Loaded"), !!ActorHandle->GetActor()))
+		if (!TestNotNull(TEXT("Actor Handle Loaded"), ActorHandle->GetActor()))
 		{
 			return false;
 		}
@@ -147,7 +147,7 @@ namespace WorldPartitionTests
 
 		// Test path conversions to a newly spawned actor
 		AActor* NewActor = World->SpawnActor(AActor::StaticClass());
-		if (!TestTrue(TEXT("Sapwning Actor Failed"), !!NewActor))
+		if (!TestNotNull(TEXT("Sapwning Actor Failed"), NewActor))
 		{
 			return false;
 		}
