AnimInstance display debug fix to avoid a crash when a montage ends while blending physical animation (Montage pointer inside MontageInstance is null)

diff --git a/Engine/Source/Runtime/Engine/Private/Animation/AnimInstance.cpp b/Engine/Source/Runtime/Engine/Private/Animation/AnimInstance.cpp
--- a/Engine/Source/Runtime/Engine/Private/Animation/AnimInstance.cpp
+++ b/Engine/Source/Runtime/Engine/Private/Animation/AnimInstance.cpp
@@ -1165,11 +1165,14 @@ void UAnimInstance::DisplayDebug(class UCanvas* Canvas, const FDebugDisplayInfo&
 			FIndenter PlayerIndent(Indent);
 
 			FAnimMontageInstance* MontageInstance = MontageInstances[MontageIndex];
+			
+			if (MontageInstance != nullptr && MontageInstance->Montage != nullptr)
+			{
+				DisplayDebugManager.SetLinearDrawColor((MontageInstance->IsActive()) ? ActiveColor : TextWhite);
 
-			DisplayDebugManager.SetLinearDrawColor((MontageInstance->IsActive()) ? ActiveColor : TextWhite);
-
-			FString MontageEntry = FString::Printf(TEXT("%i) %s CurrSec: %s NextSec: %s W:%.2f DW:%.2f"), MontageIndex, *MontageInstance->Montage->GetName(), *MontageInstance->GetCurrentSection().ToString(), *MontageInstance->GetNextSection().ToString(), MontageInstance->GetWeight(), MontageInstance->GetDesiredWeight());
-			DisplayDebugManager.DrawString(MontageEntry, Indent);
+				FString MontageEntry = FString::Printf(TEXT("%i) %s CurrSec: %s NextSec: %s W:%.2f DW:%.2f"), MontageIndex, *MontageInstance->Montage->GetName(), *MontageInstance->GetCurrentSection().ToString(), *MontageInstance->GetNextSection().ToString(), MontageInstance->GetWeight(), MontageInstance->GetDesiredWeight());
+				DisplayDebugManager.DrawString(MontageEntry, Indent);
+			}
 		}
 	}
 
