Adding missing includes.

diff --git a/Engine/Source/Runtime/Engine/Private/Components/SkinnedMeshComponent.cpp b/Engine/Source/Runtime/Engine/Private/Components/SkinnedMeshComponent.cpp
--- a/Engine/Source/Runtime/Engine/Private/Components/SkinnedMeshComponent.cpp
+++ b/Engine/Source/Runtime/Engine/Private/Components/SkinnedMeshComponent.cpp
@@ -6,6 +6,7 @@
 
 #include "Components/SkinnedMeshComponent.h"
 #include "Engine/SkinnedAsset.h"
+#include "HAL/LowLevelMemStats.h"
 #include "RenderingThread.h"
 #include "GameFramework/PlayerController.h"
 #include "ContentStreaming.h"
diff --git a/Engine/Source/Runtime/Engine/Private/TextureRenderTarget2D.cpp b/Engine/Source/Runtime/Engine/Private/TextureRenderTarget2D.cpp
--- a/Engine/Source/Runtime/Engine/Private/TextureRenderTarget2D.cpp
+++ b/Engine/Source/Runtime/Engine/Private/TextureRenderTarget2D.cpp
@@ -5,6 +5,7 @@
 =============================================================================*/
 
 #include "Engine/TextureRenderTarget2D.h"
+#include "HAL/LowLevelMemStats.h"
 #include "Misc/MessageDialog.h"
 #include "RenderingThread.h"
 #include "TextureResource.h"
diff --git a/Engine/Source/Runtime/Engine/Private/TextureRenderTarget2DArray.cpp b/Engine/Source/Runtime/Engine/Private/TextureRenderTarget2DArray.cpp
--- a/Engine/Source/Runtime/Engine/Private/TextureRenderTarget2DArray.cpp
+++ b/Engine/Source/Runtime/Engine/Private/TextureRenderTarget2DArray.cpp
@@ -5,6 +5,7 @@
 =============================================================================*/
 
 #include "Engine/TextureRenderTarget2DArray.h"
+#include "HAL/LowLevelMemStats.h"
 #include "RenderingThread.h"
 #include "TextureRenderTarget2DArrayResource.h"
 #include "UnrealEngine.h"
diff --git a/Engine/Source/Runtime/Engine/Private/TextureRenderTargetCube.cpp b/Engine/Source/Runtime/Engine/Private/TextureRenderTargetCube.cpp
--- a/Engine/Source/Runtime/Engine/Private/TextureRenderTargetCube.cpp
+++ b/Engine/Source/Runtime/Engine/Private/TextureRenderTargetCube.cpp
@@ -5,6 +5,7 @@
 =============================================================================*/
 
 #include "Engine/TextureRenderTargetCube.h"
+#include "HAL/LowLevelMemStats.h"
 #include "RenderingThread.h"
 #include "TextureResource.h"
 #include "UObject/UnrealType.h"
diff --git a/Engine/Source/Runtime/Engine/Private/TextureRenderTargetVolume.cpp b/Engine/Source/Runtime/Engine/Private/TextureRenderTargetVolume.cpp
--- a/Engine/Source/Runtime/Engine/Private/TextureRenderTargetVolume.cpp
+++ b/Engine/Source/Runtime/Engine/Private/TextureRenderTargetVolume.cpp
@@ -11,6 +11,7 @@
 #include "DeviceProfiles/DeviceProfile.h"
 #include "DeviceProfiles/DeviceProfileManager.h"
 #include "Engine/VolumeTexture.h"
+#include "HAL/LowLevelMemStats.h"
 
 #include UE_INLINE_GENERATED_CPP_BY_NAME(TextureRenderTargetVolume)
 
