Prepared code for removing Threading.h include from ParticleHandleFwd.h..

diff --git a/Engine/Source/Runtime/Engine/Private/Streaming/StreamingManagerTexture.cpp b/Engine/Source/Runtime/Engine/Private/Streaming/StreamingManagerTexture.cpp
--- a/Engine/Source/Runtime/Engine/Private/Streaming/StreamingManagerTexture.cpp
+++ b/Engine/Source/Runtime/Engine/Private/Streaming/StreamingManagerTexture.cpp
@@ -25,6 +25,10 @@
 #include "Misc/CoreDelegates.h"
 #include "TextureResource.h"
 
+#if !STATS
+#include "Async/ParallelFor.h"
+#endif
+
 CSV_DECLARE_CATEGORY_MODULE_EXTERN(CORE_API, Basic);
 
 CSV_DEFINE_CATEGORY(TextureStreaming, true);
