Fix PostSerialize to support async loading thread

#rnx

diff --git a/Engine/Source/Runtime/Engine/Private/PhysicsEngine/ConstraintInstance.cpp b/Engine/Source/Runtime/Engine/Private/PhysicsEngine/ConstraintInstance.cpp
--- a/Engine/Source/Runtime/Engine/Private/PhysicsEngine/ConstraintInstance.cpp
+++ b/Engine/Source/Runtime/Engine/Private/PhysicsEngine/ConstraintInstance.cpp
@@ -1146,12 +1146,12 @@ void FConstraintInstance::PostSerialize(const FArchive& Ar)
 {
 	if (Ar.IsLoading() && Ar.UEVer() < VER_UE4_FIXUP_STIFFNESS_AND_DAMPING_SCALE)
 	{
-		LinearLimitStiffness_DEPRECATED		/= CVarConstraintAngularStiffnessScale.GetValueOnGameThread();
-		SwingLimitStiffness_DEPRECATED		/= CVarConstraintAngularStiffnessScale.GetValueOnGameThread();
-		TwistLimitStiffness_DEPRECATED		/= CVarConstraintAngularStiffnessScale.GetValueOnGameThread();
-		LinearLimitDamping_DEPRECATED		/=  CVarConstraintAngularDampingScale.GetValueOnGameThread();
-		SwingLimitDamping_DEPRECATED		/=  CVarConstraintAngularDampingScale.GetValueOnGameThread();
-		TwistLimitDamping_DEPRECATED		/=  CVarConstraintAngularDampingScale.GetValueOnGameThread();
+		LinearLimitStiffness_DEPRECATED		/= CVarConstraintAngularStiffnessScale.GetValueOnAnyThread();
+		SwingLimitStiffness_DEPRECATED		/= CVarConstraintAngularStiffnessScale.GetValueOnAnyThread();
+		TwistLimitStiffness_DEPRECATED		/= CVarConstraintAngularStiffnessScale.GetValueOnAnyThread();
+		LinearLimitDamping_DEPRECATED		/=  CVarConstraintAngularDampingScale.GetValueOnAnyThread();
+		SwingLimitDamping_DEPRECATED		/=  CVarConstraintAngularDampingScale.GetValueOnAnyThread();
+		TwistLimitDamping_DEPRECATED		/=  CVarConstraintAngularDampingScale.GetValueOnAnyThread();
 	}
 
 	if (Ar.IsLoading() && Ar.UEVer() < VER_UE4_FIXUP_MOTOR_UNITS)
