Cook waves on PostLoad
- This fixes a DDC issue where USoundWaves were having their compressed data calculated before the USoundWave finished loading.
- Any calls which triggered audio derived data for USOundWave to be created are moved from USoundWave::Serialize to USoundWave::PostLoad.
- Retaining and Priming audio is done in USoundWave::PostLoad.
- RetainOnLoad will load the first chunk synchronously whether-or-not RetainOnLoad is set directly on the USoundWave or if it is set via a USoundClass override.

diff --git a/Engine/Source/Runtime/Engine/Private/SoundWave.cpp b/Engine/Source/Runtime/Engine/Private/SoundWave.cpp
--- a/Engine/Source/Runtime/Engine/Private/SoundWave.cpp
+++ b/Engine/Source/Runtime/Engine/Private/SoundWave.cpp
@@ -52,7 +52,7 @@ FAutoConsoleVariableRef CVarDisableRetaining(
 	TEXT("au.streamcache.DisableRetaining"),
 	DisableRetainingCVar,
 	TEXT("When set to 1, USoundWaves will not retain chunks of their own audio.\n")
-	TEXT("0: Don't disable retaining, 1: retaining."),
+	TEXT("0: Don't disable retaining, 1: disable retaining."),
 	ECVF_Default);
 
 static int32 BlockOnChunkLoadCompletionCVar = 0;
@@ -1053,8 +1053,6 @@ void USoundWave::Serialize( FArchive& Ar )
 	}
 #endif //WITH_EDITORONLY_DATA	
 	
-	bool bBuiltStreamedAudio = false;
-
 	if (bShouldStreamSound)
 	{
 		if (bCooked)
@@ -1065,55 +1063,12 @@ void USoundWave::Serialize( FArchive& Ar )
 				SerializeCookedPlatformData(Ar);
 			}
 		}
-
-#if WITH_EDITORONLY_DATA
-		if (Ar.IsLoading() && !Ar.IsTransacting() && !bCooked && !GetOutermost()->HasAnyPackageFlags(PKG_ReloadingForCooker) && FApp::CanEverRenderAudio())
-		{
-			// Prepare or fetch chunks asynchronously so that they are available from local cache when needed.
-			CachePlatformData(true /* bAsyncCache */);
-			// If async streamed sound compilation is disabled, this will revert to the old behavior.
-			bBuiltStreamedAudio = !IsCompiling();
-		}
-#endif // #if WITH_EDITORONLY_DATA
 	}
 
-	if (!(IsTemplate() || IsRunningDedicatedServer()) &&  Ar.IsLoading())
+	if (!(IsTemplate() || IsRunningDedicatedServer()) && Ar.IsLoading())
 	{
 		// For non-editor builds, we can immediately cache the sample rate.
 		SampleRate = GetSampleRateForCurrentPlatform();
-
-		const bool bShouldLoadChunks = bCooked || bBuiltStreamedAudio;
-
-		// If stream caching is enabled, here we determine if we should retain or prime this wave on load.
-		if (bShouldStreamSound && bShouldLoadChunks && FPlatformCompressionUtilities::IsCurrentPlatformUsingStreamCaching() && FApp::CanEverRenderAudio())
-		{
-			ESoundWaveLoadingBehavior CurrentLoadingBehavior = GetLoadingBehavior(false);
-
-			const bool bHasFirstChunk = GetNumChunks() > 1;
-			
-			if (!bHasFirstChunk)
-			{
-				return;
-			}
-
-			if (!GIsEditor && CurrentLoadingBehavior == ESoundWaveLoadingBehavior::RetainOnLoad)
-			{
-				RetainCompressedAudio(true);
-			}
-			else if ((CurrentLoadingBehavior == ESoundWaveLoadingBehavior::PrimeOnLoad)
-				|| (GIsEditor && CurrentLoadingBehavior == ESoundWaveLoadingBehavior::RetainOnLoad))
-			{
-				// Prime first chunk of audio.
-				SoundWaveDataPtr->LoadingBehavior = ESoundWaveLoadingBehavior::PrimeOnLoad;
-				SoundWaveDataPtr->bLoadingBehaviorOverridden = true;
-
-				if (!Proxy.IsValid())
-				{
-					Proxy = CreateSoundWaveProxy();
-				}
-				IStreamingManager::Get().GetAudioStreamingManager().RequestChunk(Proxy, 1, [](EAudioChunkLoadResult) {});
-			}
-		}
 	}
 }
 
@@ -1694,47 +1649,54 @@ void USoundWave::PostLoad()
 
 	CacheInheritedLoadingBehavior();
 	
-	// If our loading behavior is defined by a sound class, we need to update whether this sound wave actually needs to retain its audio data or not.
-	ESoundWaveLoadingBehavior ActualLoadingBehavior = GetLoadingBehavior();
-
-	// no need for a proxy if we can't render audio!
-	if (!Proxy.IsValid() && ActualLoadingBehavior != ESoundWaveLoadingBehavior::ForceInline && FApp::CanEverRenderAudio())
-	{
-		Proxy = CreateSoundWaveProxy();
-	}
-
-	if (FApp::CanEverRenderAudio() && ShouldUseStreamCaching() && ActualLoadingBehavior != GetLoadingBehavior(false))
+	if (FApp::CanEverRenderAudio())
 	{
-		if (!GIsEditor && !DisableRetainingCVar && ActualLoadingBehavior == ESoundWaveLoadingBehavior::RetainOnLoad)
+#if WITH_EDITORONLY_DATA
+		if (IsStreaming(nullptr) && FPlatformProperties::SupportsAudioStreaming())
 		{
-			UE_LOG(LogAudio, Display, TEXT("Sound Wave '%s' will have to load its compressed audio data async."), *GetName());
-			RetainCompressedAudio(false);
+			CachePlatformData(true /* bAsyncCache */);
 		}
-		else
-		{
-			// if a sound class defined our loading behavior as something other than Retain and our cvar default is to retain, we need to release our handle.
-			ReleaseCompressedAudio();
+#endif // #if WITH_EDITORONLY_DATA
 
-			const bool bHasMultipleChunks = GetNumChunks() > 1;
-			bool bShouldPrime = (ActualLoadingBehavior == ESoundWaveLoadingBehavior::PrimeOnLoad);
-			bShouldPrime |= (GIsEditor && (ActualLoadingBehavior == ESoundWaveLoadingBehavior::RetainOnLoad)); // treat this scenario like PrimeOnLoad
-			
-			if (bShouldPrime && bHasMultipleChunks)
-			{
-				IStreamingManager::Get().GetAudioStreamingManager().RequestChunk(Proxy, 1, [](EAudioChunkLoadResult) {});
-			}
-		}
+		ESoundWaveLoadingBehavior ActualLoadingBehavior = GetLoadingBehavior();
 
-		// If DisableRetainingCVar was set after this USoundWave was loaded by the ALT, release its compressed audio here.
-		if (DisableRetainingCVar)
+		if (!Proxy.IsValid() && ActualLoadingBehavior != ESoundWaveLoadingBehavior::ForceInline)
 		{
-			ReleaseCompressedAudio();
+			Proxy = CreateSoundWaveProxy();
 		}
 
-		if (!GIsEditor)
+		if (ShouldUseStreamCaching())
 		{
-			// In case any code accesses bStreaming directly, we fix up bStreaming based on the current platform's cook overrides.
-			bStreaming = IsStreaming(nullptr);
+			if (!GIsEditor && !DisableRetainingCVar && ActualLoadingBehavior == ESoundWaveLoadingBehavior::RetainOnLoad)
+			{
+				RetainCompressedAudio(true);
+			}
+			else
+			{
+				// if a sound class defined our loading behavior as something other than Retain and our cvar default is to retain, we need to release our handle.
+				ReleaseCompressedAudio();
+
+				const bool bHasMultipleChunks = GetNumChunks() > 1;
+				bool bShouldPrime = (ActualLoadingBehavior == ESoundWaveLoadingBehavior::PrimeOnLoad);
+				bShouldPrime |= (GIsEditor && (ActualLoadingBehavior == ESoundWaveLoadingBehavior::RetainOnLoad)); // treat this scenario like PrimeOnLoad
+
+				if (bShouldPrime && bHasMultipleChunks)
+				{
+					IStreamingManager::Get().GetAudioStreamingManager().RequestChunk(Proxy, 1, [](EAudioChunkLoadResult) {});
+				}
+			}
+
+			// If DisableRetainingCVar was set after this USoundWave was loaded by the ALT, release its compressed audio here.
+			if (DisableRetainingCVar)
+			{
+				ReleaseCompressedAudio();
+			}
+
+			if (!GIsEditor)
+			{
+				// In case any code accesses bStreaming directly, we fix up bStreaming based on the current platform's cook overrides.
+				bStreaming = IsStreaming(nullptr);
+			}
 		}
 	}
 
@@ -1747,7 +1709,7 @@ void USoundWave::PostLoad()
 	const bool bShouldLoadCompressedData = true;
 #endif
 	// Don't anticipate use of compressed data for streamed sound. The CachePlatformData call will request it on demand instead.
-	if (TPM && bShouldLoadCompressedData && !bStreaming)
+	if (TPM && bShouldLoadCompressedData && !IsStreaming(nullptr))
 	{
 		const TArray<ITargetPlatform*>& Platforms = TPM->GetActiveTargetPlatforms();
 
@@ -1781,11 +1743,6 @@ void USoundWave::PostLoad()
 	// Only add this streaming sound if the platform supports streaming
 	if (FApp::CanEverRenderAudio() && IsStreaming(nullptr) && FPlatformProperties::SupportsAudioStreaming())
 	{
-#if WITH_EDITORONLY_DATA
-		// Replaced the use of FinishCachePlatformData here because we just want to make sure streaming
-		// data will be ready in case CachePlatformData has never been called.
-		CachePlatformData(true /* bAsyncCache */);
-#endif // #if WITH_EDITORONLY_DATA
 		if (!ShouldUseStreamCaching())
 		{
 			IStreamingManager::Get().GetAudioStreamingManager().AddStreamingSoundWave(Proxy);
