Support for playing Metasounds from SoundCues.

Bulk of work is to support parameter forwarding from the SoundCue's param transmitter to any number of playing wave instances that also can create their own param transmitters.

#rnx

diff --git a/Engine/Plugins/Runtime/Metasound/Source/MetasoundGenerator/Private/MetasoundGenerator.cpp b/Engine/Plugins/Runtime/Metasound/Source/MetasoundGenerator/Private/MetasoundGenerator.cpp
--- a/Engine/Plugins/Runtime/Metasound/Source/MetasoundGenerator/Private/MetasoundGenerator.cpp
+++ b/Engine/Plugins/Runtime/Metasound/Source/MetasoundGenerator/Private/MetasoundGenerator.cpp
@@ -3,13 +3,11 @@
 
 #include "Analysis/MetasoundFrontendVertexAnalyzer.h"
 #include "AudioParameter.h"
-#include "DSP/Dsp.h"
 #include "Interfaces/MetasoundFrontendSourceInterface.h"
-#include "MetasoundGraph.h"
-#include "MetasoundInputNode.h"
 #include "MetasoundOperatorBuilder.h"
 #include "MetasoundOperatorInterface.h"
 #include "MetasoundOutputNode.h"
+#include "MetasoundParameterTransmitter.h"
 #include "MetasoundRouter.h"
 #include "MetasoundTrace.h"
 #include "MetasoundTrigger.h"
@@ -151,14 +149,14 @@ namespace Metasound
 				{
 					InputData.BindVertex(ParamName, *DataReference);
 				}
-				else
+				else if(MetaSoundParameterEnableWarningOnIgnoredParameterCVar)
 				{
-					UE_LOG(LogMetaSound, Error, TEXT("Failed to create initial input data reference from parameter %s of type %s on graph in MetaSoundSource [%s]"), *ParamName.ToString(), *InputVertex->DataTypeName.ToString(), *InitParams.MetaSoundName);
+					UE_LOG(LogMetaSound, Warning, TEXT("Failed to create initial input data reference from parameter %s of type %s on graph in MetaSoundSource [%s]"), *ParamName.ToString(), *InputVertex->DataTypeName.ToString(), *InitParams.MetaSoundName);
 				}
 			}
-			else
+			else if(MetaSoundParameterEnableWarningOnIgnoredParameterCVar)
 			{
-				UE_LOG(LogMetaSound, Error, TEXT("Failed to set initial input parameter %s on graph in MetaSoundSource [%s]"), *ParamName.ToString(), *InitParams.MetaSoundName);
+				UE_LOG(LogMetaSound, Warning, TEXT("Failed to set initial input parameter %s on graph in MetaSoundSource [%s]"), *ParamName.ToString(), *InitParams.MetaSoundName);
 			}
 
 		}
diff --git a/Engine/Source/Runtime/Engine/Classes/Sound/SoundCue.h b/Engine/Source/Runtime/Engine/Classes/Sound/SoundCue.h
--- a/Engine/Source/Runtime/Engine/Classes/Sound/SoundCue.h
+++ b/Engine/Source/Runtime/Engine/Classes/Sound/SoundCue.h
@@ -336,3 +336,22 @@ private:
 	FCriticalSection EditorOnlyCs;
 #endif // WITH_EDITOR
 };
+
+class ENGINE_API FSoundCueParameterTransmitter : public Audio::FParameterTransmitterBase
+{
+public:
+	FSoundCueParameterTransmitter(Audio::FParameterTransmitterInitParams&& InParams)
+		: Audio::FParameterTransmitterBase(MoveTemp(InParams.DefaultParams))
+	{
+	}
+
+	virtual ~FSoundCueParameterTransmitter() = default;
+
+	TArray<UObject*> GetReferencedObjects() const override;
+
+	virtual bool SetParameters(TArray<FAudioParameter>&& InParameters) override;
+
+	TArray<FAudioParameter> ParamsToSet;
+
+	TMap<UPTRINT, TSharedPtr<Audio::IParameterTransmitter>> Transmitters;
+};
\ No newline at end of file
diff --git a/Engine/Source/Runtime/Engine/Private/SoundCue.cpp b/Engine/Source/Runtime/Engine/Private/SoundCue.cpp
--- a/Engine/Source/Runtime/Engine/Private/SoundCue.cpp
+++ b/Engine/Source/Runtime/Engine/Private/SoundCue.cpp
@@ -621,6 +621,30 @@ void USoundCue::Parse(FAudioDevice* AudioDevice, const UPTRINT NodeWaveInstanceH
 	{
 		FirstNode->ParseNodes(AudioDevice, (UPTRINT)FirstNode, ActiveSound, ParseParams, WaveInstances);
 	}
+
+	if (FSoundCueParameterTransmitter* Transmitter = static_cast<FSoundCueParameterTransmitter*>(ActiveSound.GetTransmitter()))
+	{
+		if (Transmitter->ParamsToSet.IsEmpty())
+		{
+			return;
+		}
+
+		for (const FWaveInstance* Instance : WaveInstances)
+		{
+			if (TSharedPtr<Audio::IParameterTransmitter>* ChildTransmitterPtr = Transmitter->Transmitters.Find(Instance->WaveInstanceHash))
+			{
+				TSharedPtr<Audio::IParameterTransmitter>& ChildTransmitter = *ChildTransmitterPtr;
+				
+				if (ChildTransmitter.IsValid())
+				{
+					TArray<FAudioParameter> Params = Transmitter->ParamsToSet;
+					ChildTransmitter->SetParameters(MoveTemp(Params));
+				}
+			}
+		}
+		
+		Transmitter->ParamsToSet.Reset();
+	}
 }
 
 float USoundCue::GetVolumeMultiplier()
@@ -701,39 +725,6 @@ bool USoundCue::HasCookedAmplitudeEnvelopeData() const
 
 TSharedPtr<Audio::IParameterTransmitter> USoundCue::CreateParameterTransmitter(Audio::FParameterTransmitterInitParams&& InParams) const
 {
-	class FSoundCueParameterTransmitter : public Audio::FParameterTransmitterBase
-	{
-	public:
-		FSoundCueParameterTransmitter(Audio::FParameterTransmitterInitParams&& InParams)
-			: Audio::FParameterTransmitterBase(MoveTemp(InParams.DefaultParams))
-		{
-		}
-
-		virtual ~FSoundCueParameterTransmitter() = default;
-
-		TArray<UObject*> GetReferencedObjects() const override
-		{
-			TArray<UObject*> Objects;
-			for (const FAudioParameter& Param : AudioParameters)
-			{
-				if (Param.ObjectParam)
-				{
-					Objects.Add(Param.ObjectParam);
-				}
-
-				for (UObject* Object : Param.ArrayObjectParam)
-				{
-					if (Object)
-					{
-						Objects.Add(Object);
-					}
-				}
-			}
-
-			return Objects;
-		}
-	};
-
 	return MakeShared<FSoundCueParameterTransmitter>(MoveTemp(InParams));
 }
 
@@ -810,3 +801,41 @@ TSharedPtr<ISoundCueAudioEditor> USoundCue::GetSoundCueAudioEditor()
 }
 #endif // WITH_EDITOR
 
+TArray<UObject*> FSoundCueParameterTransmitter::GetReferencedObjects() const
+{
+	TArray<UObject*> Objects;
+	for (const FAudioParameter& Param : AudioParameters)
+	{
+		if (Param.ObjectParam)
+		{
+			Objects.Add(Param.ObjectParam);
+		}
+
+		for (UObject* Object : Param.ArrayObjectParam)
+		{
+			if (Object)
+			{
+				Objects.Add(Object);
+			}
+		}
+	}
+
+	return Objects;
+}
+
+bool FSoundCueParameterTransmitter::SetParameters(TArray<FAudioParameter>&& InParameters)
+{
+	TArray<FAudioParameter> TempParams = InParameters;
+	FAudioParameter::Merge(MoveTemp(TempParams), ParamsToSet);
+
+	// todo: filter trigger types from being saved in the arrray of persistent params. triggers are not detectable at this point
+	// currently, triggers will be set on sounds as they start if they had been executed at any point in the past
+	// solvable if triggers were their own enum type, but are currently just bool types
+	//
+	// InParameters.FilterByPredicate([](const FAudioParameter& Param)
+	// {
+	// 	return Param.ParamType != EAudioParameterType::Trigger;
+	// });
+		
+	return Audio::FParameterTransmitterBase::SetParameters(MoveTemp(InParameters));
+}
\ No newline at end of file
diff --git a/Engine/Source/Runtime/Engine/Private/SoundNodeWavePlayer.cpp b/Engine/Source/Runtime/Engine/Private/SoundNodeWavePlayer.cpp
--- a/Engine/Source/Runtime/Engine/Private/SoundNodeWavePlayer.cpp
+++ b/Engine/Source/Runtime/Engine/Private/SoundNodeWavePlayer.cpp
@@ -2,6 +2,9 @@
 
 #include "Sound/SoundNodeWavePlayer.h"
 #include "ActiveSound.h"
+#include "AudioDevice.h"
+#include "Sound/SoundCue.h"
+#include "Sound/SoundWave.h"
 #include "UObject/FrameworkObjectVersion.h"
 #include "Async/Async.h"
 
@@ -58,6 +61,7 @@ void USoundNodeWavePlayer::LoadAsset(bool bAddToRoot)
 		if (SoundWave)
 		{
 			SoundWave->AddToCluster(this, true);
+			SoundWave->InitResources();
 		}
 	}
 	else
@@ -76,6 +80,7 @@ void USoundNodeWavePlayer::LoadAsset(bool bAddToRoot)
 				SoundWave->AddToRoot();
 			}
 			SoundWave->AddToCluster(this);
+			SoundWave->InitResources();
 		}
 	}
 }
@@ -132,57 +137,65 @@ void USoundNodeWavePlayer::ParseNodes( FAudioDevice* AudioDevice, const UPTRINT
 
 	if (SoundWave)
 	{
+		RETRIEVE_SOUNDNODE_PAYLOAD(sizeof(int32));
+		DECLARE_SOUNDNODE_ELEMENT(int32, bPlayFailed);
+
+		if (*RequiresInitialization)
+		{
+			bPlayFailed = 0;
+			
+			Audio::FParameterTransmitterInitParams Params;
+			Params.DefaultParams = ActiveSound.GetTransmitter()->GetParameters();
+			Params.InstanceID = ActiveSound.GetAudioComponentID();
+			Params.SampleRate = AudioDevice->GetSampleRate();
+
+			if (FSoundCueParameterTransmitter* SoundCueTransmitter = static_cast<FSoundCueParameterTransmitter*>(ActiveSound.GetTransmitter()))
+			{
+				const TSharedPtr<Audio::IParameterTransmitter> SoundWaveTransmitter = SoundWave->CreateParameterTransmitter(MoveTemp(Params));
+				
+				if (SoundWaveTransmitter.IsValid())
+				{
+					SoundCueTransmitter->Transmitters.Add(NodeWaveInstanceHash, SoundWaveTransmitter);
+				}
+			}
+			
+			*RequiresInitialization = 0;
+		}
+
 		// The SoundWave's bLooping is only for if it is directly referenced, so clear it
 		// in the case that it is being played from a player
-		bool bWaveIsLooping = SoundWave->bLooping;
+		const bool bWaveIsLooping = SoundWave->bLooping;
 		SoundWave->bLooping = false;
 
-		if (bLooping)
+		if (bLooping || (SoundWave->bProcedural && !SoundWave->IsOneShot()))
 		{
 			FSoundParseParameters UpdatedParams = ParseParams;
 			UpdatedParams.bLooping = true;
 			SoundWave->Parse(AudioDevice, NodeWaveInstanceHash, ActiveSound, UpdatedParams, WaveInstances);
 		}
+		else if (ParseParams.bEnableRetrigger)
+		{
+			// Don't play non-looping sounds again if this sound has been revived to// avoid re-triggering one shots played adjacent with looping wave instances in cue
+			SoundWave->Parse(AudioDevice, NodeWaveInstanceHash, ActiveSound, ParseParams, WaveInstances);
+		}
 		else
 		{
-			// Don't play non-looping sounds again if this sound has been revived to
-			// avoid re-triggering one shots played adjacent with looping wave instances in cue
-			if (ParseParams.bEnableRetrigger)
-			{
-				SoundWave->Parse(AudioDevice, NodeWaveInstanceHash, ActiveSound, ParseParams, WaveInstances);
-			}
-			else
+			// If sound has been virtualized, don't try to revive one-shots
+			if (!ActiveSound.bHasVirtualized)
 			{
-				// If sound has been virtualized, don't try to revive one-shots
-				if (!ActiveSound.bHasVirtualized)
+				// Guard against continual parsing if wave instance was created but not added to transient
+				// wave instance list to avoid inaudible sounds popping back in.
+				if (bPlayFailed == 0)
 				{
-					RETRIEVE_SOUNDNODE_PAYLOAD(sizeof(int32));
-					DECLARE_SOUNDNODE_ELEMENT(int32, bPlayFailed);
-					if (*RequiresInitialization)
-					{
-						bPlayFailed = 0;
-					}
-
 					const int32 InitActiveSoundWaveInstanceNum = ActiveSound.GetWaveInstances().Num();
 					const int32 InitWaveInstancesNum = WaveInstances.Num();
-
-					// Guard against continual parsing if wave instance was created but not added to transient
-					// wave instance list to avoid inaudible sounds popping back in.
-					if (!bPlayFailed)
+					
+					SoundWave->Parse(AudioDevice, NodeWaveInstanceHash, ActiveSound, ParseParams, WaveInstances);
+					
+					const bool bFailed = ActiveSound.GetWaveInstances().Num() == InitActiveSoundWaveInstanceNum && WaveInstances.Num() == InitWaveInstancesNum;
+					if (bFailed)
 					{
-						SoundWave->Parse(AudioDevice, NodeWaveInstanceHash, ActiveSound, ParseParams, WaveInstances);
-					}
-
-					if (*RequiresInitialization != 0)
-					{
-						if (ActiveSound.GetWaveInstances().Num() == InitActiveSoundWaveInstanceNum)
-						{
-							if (WaveInstances.Num() == InitWaveInstancesNum)
-							{
-								bPlayFailed = 1;
-							}
-						}
-						*RequiresInitialization = 0;
+						bPlayFailed = 1;
 					}
 				}
 			}
@@ -278,5 +291,4 @@ int32 USoundNodeWavePlayer::GetMaxChildNodes() const
 }
 
 
-#undef LOCTEXT_NAMESPACE
-
+#undef LOCTEXT_NAMESPACE
\ No newline at end of file
