Fix for some textures reporting no mip error

diff --git a/Engine/Source/Runtime/Engine/Private/TextureDerivedDataTask.cpp b/Engine/Source/Runtime/Engine/Private/TextureDerivedDataTask.cpp
--- a/Engine/Source/Runtime/Engine/Private/TextureDerivedDataTask.cpp
+++ b/Engine/Source/Runtime/Engine/Private/TextureDerivedDataTask.cpp
@@ -1962,7 +1962,7 @@ void FTextureCacheDerivedDataWorker::Finalize()
 	// --	if the texture compiler cvar disallows async texture compilation
 	if (!bSucceeded)
 	{
-		if (!TextureData.HasPayload())
+		if (!TextureData.HasPayload() && !Texture.Source.HasPayloadData())
 		{
 			UE_LOG(LogTexture, Warning, TEXT("Unable to build texture source data, no available payload for %s. This may happen if it was duplicated from cooked data."), *TexturePathName);
 			return;
