World Partition: basic support for streaming partitioned worlds with their own grids, through Blueprint (no Level Instance support).

diff --git a/Engine/Source/Runtime/Engine/Private/Level.cpp b/Engine/Source/Runtime/Engine/Private/Level.cpp
--- a/Engine/Source/Runtime/Engine/Private/Level.cpp
+++ b/Engine/Source/Runtime/Engine/Private/Level.cpp
@@ -2986,7 +2986,7 @@ void ULevel::OnLevelLoaded()
 			const bool bIsValidLevelInstance = IsValidLevelInstanceWorldPartition(WorldPartition);
 			const bool bIsMainWorldLevel = OwningWorld->PersistentLevel == this;
 			const bool bInitializeForEditor = !bIsOwningWorldGameWorld && bIsValidLevelInstance;
-			const bool bInitializeForGame = bIsOwningWorldGameWorld && !bIsOwningWorldPartitioned;
+			const bool bInitializeForGame = bIsOwningWorldGameWorld;
 
 			UE_LOG(LogWorldPartition, Log, TEXT("ULevel::OnLevelLoaded(%s)(bIsOwningWorldGameWorld=%d, bIsOwningWorldPartitioned=%d, bIsValidLevelInstance=%d, InitializeForMainWorld=%d, InitializeForEditor=%d, InitializeForGame=%d)"), 
 				*GetTypedOuter<UWorld>()->GetName(), bIsOwningWorldGameWorld ? 1 : 0, bIsOwningWorldPartitioned ? 1 : 0, bIsValidLevelInstance ? 1 : 0, bIsMainWorldLevel ? 1 : 0, bInitializeForEditor ? 1 : 0, bInitializeForGame ? 1 : 0);
diff --git a/Engine/Source/Runtime/Engine/Private/WorldPartition/WorldPartition.cpp b/Engine/Source/Runtime/Engine/Private/WorldPartition/WorldPartition.cpp
--- a/Engine/Source/Runtime/Engine/Private/WorldPartition/WorldPartition.cpp
+++ b/Engine/Source/Runtime/Engine/Private/WorldPartition/WorldPartition.cpp
@@ -539,7 +539,7 @@ void UWorldPartition::Initialize(UWorld* InWorld, const FTransform& InTransform)
 
 					InstancingContext.AddPackageMapping(*LongActorPackageName, *InstancedName);
 
-					ActorDescIterator->TransformInstance(SourceWorldPath, RemappedWorldPath, GetInstanceTransform());
+					ActorDescIterator->TransformInstance(SourceWorldPath, RemappedWorldPath);
 				}
 
 				ActorDescIterator->bIsForcedNonSpatiallyLoaded = !IsStreamingEnabled();
diff --git a/Engine/Source/Runtime/Engine/Private/WorldPartition/WorldPartitionActorDesc.cpp b/Engine/Source/Runtime/Engine/Private/WorldPartition/WorldPartitionActorDesc.cpp
--- a/Engine/Source/Runtime/Engine/Private/WorldPartition/WorldPartitionActorDesc.cpp
+++ b/Engine/Source/Runtime/Engine/Private/WorldPartition/WorldPartitionActorDesc.cpp
@@ -222,21 +222,10 @@ void FWorldPartitionActorDesc::RegisterActorDescDeprecator(TSubclassOf<AActor> A
 	Deprecators.Add(ActorClass, Deprecator);
 }
 
-void FWorldPartitionActorDesc::TransformInstance(const FString& From, const FString& To, const FTransform& InstanceTransform)
+void FWorldPartitionActorDesc::TransformInstance(const FString& From, const FString& To)
 {
 	check(!HardRefCount);
-
 	ActorPath = *ActorPath.ToString().Replace(*From, *To);
-
-	// Transform BoundsLocation and BoundsExtent if necessary
-	if (!InstanceTransform.Equals(FTransform::Identity))
-	{
-		//@todo_ow: This will result in a new BoundsExtent that is larger than it should. To fix this, we would need the Object Oriented BoundingBox of the actor (the BV of the actor without rotation)
-		const FVector BoundsMin = BoundsLocation - BoundsExtent;
-		const FVector BoundsMax = BoundsLocation + BoundsExtent;
-		const FBox NewBounds = FBox(BoundsMin, BoundsMax).TransformBy(InstanceTransform);
-		NewBounds.GetCenterAndExtents(BoundsLocation, BoundsExtent);
-	}
 }
 
 FString FWorldPartitionActorDesc::ToString(EToStringMode Mode) const
diff --git a/Engine/Source/Runtime/Engine/Public/WorldPartition/WorldPartitionActorDesc.h b/Engine/Source/Runtime/Engine/Public/WorldPartition/WorldPartitionActorDesc.h
--- a/Engine/Source/Runtime/Engine/Public/WorldPartition/WorldPartitionActorDesc.h
+++ b/Engine/Source/Runtime/Engine/Public/WorldPartition/WorldPartitionActorDesc.h
@@ -230,7 +230,7 @@ public:
 
 	void SerializeTo(TArray<uint8>& OutData);
 
-	void TransformInstance(const FString& From, const FString& To, const FTransform& InstanceTransform);
+	void TransformInstance(const FString& From, const FString& To);
 
 	using FActorDescDeprecator = TFunction<void(FArchive&, FWorldPartitionActorDesc*)>;
 	static void RegisterActorDescDeprecator(TSubclassOf<AActor> ActorClass, const FActorDescDeprecator& Deprecator);
