Add budget project setting for staticmesh

diff --git a/Engine/Plugins/Importers/AlembicImporter/Source/AlembicLibrary/Private/AbcImporter.cpp b/Engine/Plugins/Importers/AlembicImporter/Source/AlembicLibrary/Private/AbcImporter.cpp
--- a/Engine/Plugins/Importers/AlembicImporter/Source/AlembicLibrary/Private/AbcImporter.cpp
+++ b/Engine/Plugins/Importers/AlembicImporter/Source/AlembicLibrary/Private/AbcImporter.cpp
@@ -58,6 +58,10 @@ THIRD_PARTY_INCLUDES_END
 #include "UObject/MetaData.h"
 #include "UObject/Package.h"
 
+#if WITH_EDITOR
+#include "MeshBudgetProjectSettings.h"
+#endif
+
 #define LOCTEXT_NAMESPACE "AbcImporter"
 
 DEFINE_LOG_CATEGORY_STATIC(LogAbcImporter, Verbose, All);
@@ -251,6 +255,9 @@ UStaticMesh* FAbcImporter::CreateStaticMeshFromSample(UObject* InParent, const F
 		//Set the Imported version before calling the build
 		StaticMesh->ImportVersion = EImportStaticMeshVersion::LastVersion;
 
+#if WITH_EDITOR
+		FMeshBudgetProjectSettingsUtils::SetLodGroupForStaticMesh(StaticMesh);
+#endif
 		// Build the static mesh (using the build setting etc.) this generates correct tangents using the extracting smoothing group along with the imported Normals data
 		StaticMesh->Build(false);
 
diff --git a/Engine/Plugins/Importers/USDImporter/Source/USDSchemas/Private/USDGeomMeshTranslator.cpp b/Engine/Plugins/Importers/USDImporter/Source/USDSchemas/Private/USDGeomMeshTranslator.cpp
--- a/Engine/Plugins/Importers/USDImporter/Source/USDSchemas/Private/USDGeomMeshTranslator.cpp
+++ b/Engine/Plugins/Importers/USDImporter/Source/USDSchemas/Private/USDGeomMeshTranslator.cpp
@@ -51,6 +51,7 @@
 #include "GeometryCacheTrackUSD.h"
 #include "GeometryCacheUSDComponent.h"
 #include "IMeshBuilderModule.h"
+#include "MeshBudgetProjectSettings.h"
 #endif // WITH_EDITOR
 
 #include "USDIncludesStart.h"
@@ -754,6 +755,9 @@ namespace UsdGeomMeshTranslatorImpl
 				check(StaticMeshDescription);
 				*StaticMeshDescription = MoveTemp(MeshDescription);
 			}
+
+			FMeshBudgetProjectSettingsUtils::SetLodGroupForStaticMesh(StaticMesh);
+
 #endif // WITH_EDITOR
 
 			StaticMesh->SetLightingGuid();
diff --git a/Engine/Plugins/Interchange/Runtime/Source/Import/Private/Mesh/InterchangeStaticMeshFactory.cpp b/Engine/Plugins/Interchange/Runtime/Source/Import/Private/Mesh/InterchangeStaticMeshFactory.cpp
--- a/Engine/Plugins/Interchange/Runtime/Source/Import/Private/Mesh/InterchangeStaticMeshFactory.cpp
+++ b/Engine/Plugins/Interchange/Runtime/Source/Import/Private/Mesh/InterchangeStaticMeshFactory.cpp
@@ -25,6 +25,7 @@
 #include "Materials/MaterialInstance.h"
 #include "Materials/MaterialInterface.h"
 #include "Mesh/InterchangeStaticMeshPayloadInterface.h"
+#include "MeshBudgetProjectSettings.h"
 #include "Model.h"
 #include "Nodes/InterchangeBaseNode.h"
 #include "Nodes/InterchangeBaseNodeContainer.h"
@@ -449,6 +450,10 @@ UObject* UInterchangeStaticMeshFactory::CreateAsset(const FCreateAssetParams& Ar
 	// Getting the file Hash will cache it into the source data
 	Arguments.SourceData->GetFileContentHash();
 
+#if WITH_EDITOR
+	FMeshBudgetProjectSettingsUtils::SetLodGroupForStaticMesh(StaticMesh);
+#endif
+
 	return StaticMeshObject;
 }
 
diff --git a/Engine/Source/Editor/UnrealEd/Private/Fbx/FbxStaticMeshImport.cpp b/Engine/Source/Editor/UnrealEd/Private/Fbx/FbxStaticMeshImport.cpp
--- a/Engine/Source/Editor/UnrealEd/Private/Fbx/FbxStaticMeshImport.cpp
+++ b/Engine/Source/Editor/UnrealEd/Private/Fbx/FbxStaticMeshImport.cpp
@@ -22,6 +22,7 @@
 #include "Engine/Polys.h"
 #include "Engine/StaticMeshSocket.h"
 #include "Editor.h"
+#include "MeshBudgetProjectSettings.h"
 #include "Modules/ModuleManager.h"
 #include "Subsystems/AssetEditorSubsystem.h"
 
@@ -2382,6 +2383,10 @@ void UnFbx::FFbxImporter::PostImportStaticMesh(UStaticMesh* StaticMesh, TArray<F
 	{
 		ReorderMaterialAfterImport(StaticMesh, MeshNodeArray, ImportOptions->bReorderMaterialToFbxOrder);
 	}
+
+#if WITH_EDITOR
+	FMeshBudgetProjectSettingsUtils::SetLodGroupForStaticMesh(StaticMesh);
+#endif
 }
 
 void UnFbx::FFbxImporter::UpdateStaticMeshImportData(UStaticMesh *StaticMesh, UFbxStaticMeshImportData* StaticMeshImportData)
diff --git a/Engine/Source/Runtime/Engine/Classes/Engine/StaticMesh.h b/Engine/Source/Runtime/Engine/Classes/Engine/StaticMesh.h
--- a/Engine/Source/Runtime/Engine/Classes/Engine/StaticMesh.h
+++ b/Engine/Source/Runtime/Engine/Classes/Engine/StaticMesh.h
@@ -1537,7 +1537,7 @@ public:
 	ENGINE_API virtual void ClearAllCachedCookedPlatformData() override;
 	ENGINE_API virtual void BeginCacheForCookedPlatformData(const ITargetPlatform* TargetPlatform) override;
 	ENGINE_API virtual bool IsCachedCookedPlatformDataLoaded(const ITargetPlatform* TargetPlatform) override;
-	ENGINE_API void SetLODGroup(FName NewGroup, bool bRebuildImmediately = true);
+	ENGINE_API void SetLODGroup(FName NewGroup, bool bRebuildImmediately = true, bool bAllowModify = true);
 	ENGINE_API void BroadcastNavCollisionChange();
 
 	FOnExtendedBoundsChanged& GetOnExtendedBoundsChanged() { return OnExtendedBoundsChanged; }
diff --git a/Engine/Source/Runtime/Engine/Private/MeshBudgetProjectSettings.cpp b/Engine/Source/Runtime/Engine/Private/MeshBudgetProjectSettings.cpp
new file mode 100644
--- /dev/null
+++ b/Engine/Source/Runtime/Engine/Private/MeshBudgetProjectSettings.cpp
@@ -0,0 +1,78 @@
+// Copyright Epic Games, Inc. All Rights Reserved.
+
+#include "MeshBudgetProjectSettings.h"
+
+#if WITH_EDITOR
+#include "Engine/StaticMesh.h"
+#include "Interfaces/ITargetPlatform.h"
+#include "Interfaces/ITargetPlatformManagerModule.h"
+#include "Misc/CoreMisc.h"
+#include "StaticMeshResources.h"
+
+#define LOCTEXT_NAMESPACE "MeshBudget"
+
+DEFINE_LOG_CATEGORY(LogMeshBudget);
+
+void FMeshBudgetProjectSettingsUtils::SetLodGroupForStaticMesh(UStaticMesh* StaticMesh)
+{
+	//Log only once the project settings error
+	static bool bProjectSettingsValidated = false;
+
+	const UMeshBudgetProjectSettings* MeshBudgetProjectSettings = GetDefault<UMeshBudgetProjectSettings>();
+	if (!MeshBudgetProjectSettings || !MeshBudgetProjectSettings->bEnableStaticMeshBudget)
+	{
+		return;
+	}
+
+	const ITargetPlatform* Platform = GetTargetPlatformManagerRef().GetRunningTargetPlatform();
+	check(Platform);
+	TArray<FName> LodGroupNames;
+	Platform->GetStaticMeshLODSettings().GetLODGroupNames(LodGroupNames);
+
+	if (StaticMesh->LODGroup != NAME_None)
+	{
+		if (LodGroupNames.Contains(StaticMesh->LODGroup))
+		{
+			//We already have a valid LODGroup name
+			return;
+		}
+	}
+
+	//We must assign a budget LOD group, we will assign the maximum budget that fit the extent of the mesh
+	const double AssetExtent = StaticMesh->GetBoundingBox().GetExtent().Size();
+	FName ValidBudgetName = NAME_None;
+	double MaxMinimumExtentAssign = -1.0;
+	for (const FStaticMeshBudgetInfo& MeshBudgetInfo : MeshBudgetProjectSettings->StaticMeshBudgetInfos)
+	{
+		if (!LodGroupNames.Contains(MeshBudgetInfo.LodGroupName))
+		{
+			if (!bProjectSettingsValidated)
+			{
+				UE_LOG(LogMeshBudget, Warning, TEXT("MeshBudgetProjectSettings MeshBudgetInfo [%s] LOD group name do not exist add it to the ini file. Skipping this budget... "), *MeshBudgetInfo.LodGroupName.ToString());
+			}
+			//Skip invalid budget lod group name
+			continue;
+		}
+
+		if (MaxMinimumExtentAssign < MeshBudgetInfo.MinimumExtent && AssetExtent >= MeshBudgetInfo.MinimumExtent)
+		{
+			ValidBudgetName = MeshBudgetInfo.LodGroupName;
+			MaxMinimumExtentAssign = MeshBudgetInfo.MinimumExtent;
+		}
+	}
+
+	if (ValidBudgetName != NAME_None)
+	{
+		StaticMesh->SetLODGroup(ValidBudgetName, false, false);
+	}
+	else if (!bProjectSettingsValidated)
+	{
+		UE_LOG(LogMeshBudget, Error, TEXT("MeshBudgetProjectSettings Cannot assign any lod group to static mesh, the budget data is not configure properly in the project settings. "));
+	}
+
+	bProjectSettingsValidated = true;
+}
+
+#undef LOCTEXT_NAMESPACE
+
+#endif
diff --git a/Engine/Source/Runtime/Engine/Private/StaticMesh.cpp b/Engine/Source/Runtime/Engine/Private/StaticMesh.cpp
--- a/Engine/Source/Runtime/Engine/Private/StaticMesh.cpp
+++ b/Engine/Source/Runtime/Engine/Private/StaticMesh.cpp
@@ -44,6 +44,7 @@
 
 #if WITH_EDITOR
 #include "RawMesh.h"
+#include "MeshBudgetProjectSettings.h"
 #include "NaniteBuilder.h"
 #include "MeshUtilitiesCommon.h"
 #include "DerivedDataCacheInterface.h"
@@ -4011,11 +4012,11 @@ void UStaticMesh::PostEditUndo()
 	Super::PostEditUndo();
 }
 
-void UStaticMesh::SetLODGroup(FName NewGroup, bool bRebuildImmediately)
+void UStaticMesh::SetLODGroup(FName NewGroup, bool bRebuildImmediately, bool bAllowModify)
 {
 #if WITH_EDITORONLY_DATA
 	const bool bBeforeDerivedDataCached = (GetRenderData() == nullptr);
-	if (!bBeforeDerivedDataCached)
+	if (!bBeforeDerivedDataCached && bAllowModify)
 	{
 		Modify();
 	}
@@ -5886,6 +5887,8 @@ void UStaticMesh::BeginPostLoadInternal(FStaticMeshPostLoadContext& Context)
 	// This scope allows us to use any locked properties without causing stalls
 	FStaticMeshAsyncBuildScope AsyncBuildScope(this);
 
+	FMeshBudgetProjectSettingsUtils::SetLodGroupForStaticMesh(this);
+
 	if (GetNumSourceModels() > 0)
 	{
 		static const auto CVar = IConsoleManager::Get().FindTConsoleVariableDataInt(TEXT("r.GenerateMeshDistanceFields"));
diff --git a/Engine/Source/Runtime/Engine/Public/MeshBudgetProjectSettings.h b/Engine/Source/Runtime/Engine/Public/MeshBudgetProjectSettings.h
new file mode 100644
--- /dev/null
+++ b/Engine/Source/Runtime/Engine/Public/MeshBudgetProjectSettings.h
@@ -0,0 +1,59 @@
+// Copyright Epic Games, Inc. All Rights Reserved.
+
+#pragma once
+
+#include "Containers/Array.h"
+#include "Engine/DeveloperSettings.h"
+#include "Logging/LogMacros.h"
+#include "Misc/Optional.h"
+#include "UObject/NameTypes.h"
+
+#include "MeshBudgetProjectSettings.generated.h"
+
+class UStaticMesh;
+
+DECLARE_LOG_CATEGORY_EXTERN(LogMeshBudget, Log, All)
+
+USTRUCT()
+struct FStaticMeshBudgetInfo
+{
+	GENERATED_BODY()
+	/** The name of the LOD group we will use for this budget.*/
+	UPROPERTY(EditAnywhere, Category = "StaticMesh")
+	FName LodGroupName = NAME_None;
+
+	/** The minimum volume extent to assign this budget info. We will compare the mesh bounding box extent to this value. */
+	UPROPERTY(EditAnywhere, Category = "StaticMesh")
+	double MinimumExtent = 0.0;
+};
+
+UCLASS(config=Engine, meta=(DisplayName="Mesh Budget"), MinimalAPI)
+class UMeshBudgetProjectSettings : public UDeveloperSettings
+{
+	GENERATED_BODY()
+
+public:
+	/**
+	 * Enable/disable the static mesh budget.
+	 * Static mesh budget will auto assign a lod group to any static mesh when loading or importing the asset in the editor.
+	 * The auto budget will not override a static mesh lod group, user can control the lod group for a specific asset.
+	 * @note: When the static mesh budget is enable and properly configure, there will be no static mesh without a lod group in the editor.
+	 *        
+	 */
+	UPROPERTY(EditAnywhere, config, Category = "StaticMesh")
+	bool bEnableStaticMeshBudget = false;
+
+	/**
+	 * The static mesh budgets array.
+	 */
+	UPROPERTY(EditAnywhere, config, Category = "StaticMesh")
+	TArray<FStaticMeshBudgetInfo> StaticMeshBudgetInfos;
+};
+
+class ENGINE_API FMeshBudgetProjectSettingsUtils
+{
+public:
+#if WITH_EDITOR
+	static void SetLodGroupForStaticMesh(UStaticMesh* StaticMesh);
+#endif
+};
\ No newline at end of file
