[Engine]
- CoreMinimal.h include removal all the way to PlayerController.h not transitively including it

diff --git a/Engine/Plugins/Runtime/ComputeFramework/Source/ComputeFramework/Private/ComputeFramework/ComputeSystem.h b/Engine/Plugins/Runtime/ComputeFramework/Source/ComputeFramework/Private/ComputeFramework/ComputeSystem.h
--- a/Engine/Plugins/Runtime/ComputeFramework/Source/ComputeFramework/Private/ComputeFramework/ComputeSystem.h
+++ b/Engine/Plugins/Runtime/ComputeFramework/Source/ComputeFramework/Private/ComputeFramework/ComputeSystem.h
@@ -3,6 +3,7 @@
 #pragma once
 
 #include "ComputeSystemInterface.h"
+#include "Containers/Map.h"
 
 class FComputeGraphTaskWorker;
 
diff --git a/Engine/Source/Runtime/CoreOnline/Public/Online/CoreOnline.h b/Engine/Source/Runtime/CoreOnline/Public/Online/CoreOnline.h
--- a/Engine/Source/Runtime/CoreOnline/Public/Online/CoreOnline.h
+++ b/Engine/Source/Runtime/CoreOnline/Public/Online/CoreOnline.h
@@ -6,7 +6,6 @@
 #include "Containers/Map.h"
 #include "Containers/Set.h"
 #include "Containers/UnrealString.h"
-#include "CoreMinimal.h"
 #include "HAL/PlatformCrt.h"
 #include "HAL/PlatformMath.h"
 #include "HAL/UnrealMemory.h"
@@ -765,3 +764,7 @@ struct FUniqueNetIdKeyFuncs : public DefaultKeyFuncs<FUniqueNetIdRef>
 };
 
 using FUniqueNetIdSet = TSet<FUniqueNetIdRef, FUniqueNetIdKeyFuncs>;
+
+#if UE_ENABLE_INCLUDE_ORDER_DEPRECATED_IN_5_2
+#include "CoreMinimal.h"
+#endif
diff --git a/Engine/Source/Runtime/Engine/Classes/AI/Navigation/NavigationTypes.h b/Engine/Source/Runtime/Engine/Classes/AI/Navigation/NavigationTypes.h
--- a/Engine/Source/Runtime/Engine/Classes/AI/Navigation/NavigationTypes.h
+++ b/Engine/Source/Runtime/Engine/Classes/AI/Navigation/NavigationTypes.h
@@ -2,7 +2,7 @@
 
 #pragma once
 
-#ifdef UE_ENABLE_INCLUDE_ORDER_DEPRECATED_IN_5_1
+#if UE_ENABLE_INCLUDE_ORDER_DEPRECATED_IN_5_1
 #include "CoreMinimal.h"
 #include "Stats/Stats.h"
 #include "UObject/UObjectGlobals.h"
diff --git a/Engine/Source/Runtime/Engine/Classes/Camera/CameraTypes.h b/Engine/Source/Runtime/Engine/Classes/Camera/CameraTypes.h
--- a/Engine/Source/Runtime/Engine/Classes/Camera/CameraTypes.h
+++ b/Engine/Source/Runtime/Engine/Classes/Camera/CameraTypes.h
@@ -2,7 +2,6 @@
 
 #pragma once
 
-#include "CoreMinimal.h"
 #include "UObject/ObjectMacros.h"
 #include "Engine/EngineTypes.h"
 #include "EngineDefines.h"
@@ -157,3 +156,7 @@ struct FMinimalViewInfo
 		return PerspectiveNearClipPlane > 0.0f ? PerspectiveNearClipPlane : GNearClippingPlane;
 	}
 };
+
+#if UE_ENABLE_INCLUDE_ORDER_DEPRECATED_IN_5_2
+#include "CoreMinimal.h"
+#endif
diff --git a/Engine/Source/Runtime/Engine/Classes/Camera/PlayerCameraManager.h b/Engine/Source/Runtime/Engine/Classes/Camera/PlayerCameraManager.h
--- a/Engine/Source/Runtime/Engine/Classes/Camera/PlayerCameraManager.h
+++ b/Engine/Source/Runtime/Engine/Classes/Camera/PlayerCameraManager.h
@@ -2,7 +2,6 @@
 
 #pragma once
 
-#include "CoreMinimal.h"
 #include "UObject/ObjectMacros.h"
 #include "Templates/SubclassOf.h"
 #include "Engine/EngineTypes.h"
@@ -912,3 +911,7 @@ public:
 	/** Returns TransformComponent subobject */
 	class USceneComponent* GetTransformComponent() const { return TransformComponent; }
 };
+
+#if UE_ENABLE_INCLUDE_ORDER_DEPRECATED_IN_5_2
+#include "CoreMinimal.h"
+#endif
diff --git a/Engine/Source/Runtime/Engine/Classes/Engine/BlendableInterface.h b/Engine/Source/Runtime/Engine/Classes/Engine/BlendableInterface.h
--- a/Engine/Source/Runtime/Engine/Classes/Engine/BlendableInterface.h
+++ b/Engine/Source/Runtime/Engine/Classes/Engine/BlendableInterface.h
@@ -2,7 +2,6 @@
 
 #pragma once
 
-#include "CoreMinimal.h"
 #include "UObject/ObjectMacros.h"
 #include "UObject/Interface.h"
 #include "BlendableInterface.generated.h"
@@ -123,4 +122,6 @@ private:
 	bool bIsBlendable;
 };
 
-
+#if UE_ENABLE_INCLUDE_ORDER_DEPRECATED_IN_5_2
+#include "CoreMinimal.h"
+#endif
diff --git a/Engine/Source/Runtime/Engine/Classes/Engine/LatentActionManager.h b/Engine/Source/Runtime/Engine/Classes/Engine/LatentActionManager.h
--- a/Engine/Source/Runtime/Engine/Classes/Engine/LatentActionManager.h
+++ b/Engine/Source/Runtime/Engine/Classes/Engine/LatentActionManager.h
@@ -2,7 +2,6 @@
 
 #pragma once
 
-#include "CoreMinimal.h"
 #include "UObject/ObjectMacros.h"
 #include "UObject/WeakObjectPtr.h"
 #include "LatentActionManager.generated.h"
@@ -226,3 +225,6 @@ protected:
 	static FOnLatentActionsChanged LatentActionsChangedDelegate;
 };
 
+#if UE_ENABLE_INCLUDE_ORDER_DEPRECATED_IN_5_2
+#include "CoreMinimal.h"
+#endif
diff --git a/Engine/Source/Runtime/Engine/Classes/GameFramework/Controller.h b/Engine/Source/Runtime/Engine/Classes/GameFramework/Controller.h
--- a/Engine/Source/Runtime/Engine/Classes/GameFramework/Controller.h
+++ b/Engine/Source/Runtime/Engine/Classes/GameFramework/Controller.h
@@ -2,7 +2,6 @@
 
 #pragma once
 
-#include "CoreMinimal.h"
 #include "UObject/ObjectMacros.h"
 #include "UObject/CoreNet.h"
 #include "GameFramework/Actor.h"
@@ -430,3 +429,7 @@ FORCEINLINE_DEBUGGABLE bool AController::IsLocalPlayerController() const
 {
 	return IsPlayerController() && IsLocalController();
 }
+
+#if UE_ENABLE_INCLUDE_ORDER_DEPRECATED_IN_5_2
+#include "CoreMinimal.h"
+#endif
diff --git a/Engine/Source/Runtime/Engine/Classes/GameFramework/PlayerController.h b/Engine/Source/Runtime/Engine/Classes/GameFramework/PlayerController.h
--- a/Engine/Source/Runtime/Engine/Classes/GameFramework/PlayerController.h
+++ b/Engine/Source/Runtime/Engine/Classes/GameFramework/PlayerController.h
@@ -2,7 +2,6 @@
 
 #pragma once
 
-#include "CoreMinimal.h"
 #include "SlateFwd.h"
 #include "UObject/ObjectMacros.h"
 #include "Containers/SortedMap.h"
@@ -2356,3 +2355,7 @@ public:
 	int32 GetLocalToServerAsyncPhysicsTickOffset() const { return LocalToServerAsyncPhysicsTickOffset; }
 	
 };
+
+#if UE_ENABLE_INCLUDE_ORDER_DEPRECATED_IN_5_2
+#include "CoreMinimal.h"
+#endif
diff --git a/Engine/Source/Runtime/Engine/Classes/GameFramework/PlayerMuteList.h b/Engine/Source/Runtime/Engine/Classes/GameFramework/PlayerMuteList.h
--- a/Engine/Source/Runtime/Engine/Classes/GameFramework/PlayerMuteList.h
+++ b/Engine/Source/Runtime/Engine/Classes/GameFramework/PlayerMuteList.h
@@ -2,7 +2,6 @@
 
 #pragma once
 
-#include "CoreMinimal.h"
 #include "UObject/ObjectMacros.h"
 #include "Online/CoreOnline.h"
 #include "Misc/EnumRange.h"
@@ -137,3 +136,7 @@ public:
 
 /** Dump out information about all player controller mute state */
 ENGINE_API FString DumpMutelistState(UWorld* World);
+
+#if UE_ENABLE_INCLUDE_ORDER_DEPRECATED_IN_5_2
+#include "CoreMinimal.h"
+#endif
diff --git a/Engine/Source/Runtime/Engine/Classes/GameFramework/UpdateLevelVisibilityLevelInfo.h b/Engine/Source/Runtime/Engine/Classes/GameFramework/UpdateLevelVisibilityLevelInfo.h
--- a/Engine/Source/Runtime/Engine/Classes/GameFramework/UpdateLevelVisibilityLevelInfo.h
+++ b/Engine/Source/Runtime/Engine/Classes/GameFramework/UpdateLevelVisibilityLevelInfo.h
@@ -2,8 +2,6 @@
 
 #pragma once
 
-#include "CoreMinimal.h"
-#include "UObject/ObjectMacros.h"
 #include "UpdateLevelVisibilityLevelInfo.generated.h"
 
 /** This structure is used to to identify NetLevelVisibility transactions between server and client */
@@ -127,3 +125,7 @@ inline uint32 FNetLevelVisibilityTransactionId::IncrementTransactionIndex()
 	SetTransactionIndex(NewIndex);
 	return NewIndex;
 }
+
+#if UE_ENABLE_INCLUDE_ORDER_DEPRECATED_IN_5_2
+#include "CoreMinimal.h"
+#endif
diff --git a/Engine/Source/Runtime/Engine/Public/AI/NavigationModifier.h b/Engine/Source/Runtime/Engine/Public/AI/NavigationModifier.h
--- a/Engine/Source/Runtime/Engine/Public/AI/NavigationModifier.h
+++ b/Engine/Source/Runtime/Engine/Public/AI/NavigationModifier.h
@@ -2,7 +2,6 @@
 
 #pragma once 
 
-#include "CoreMinimal.h"
 #include "CollisionShape.h"
 #include "Templates/SubclassOf.h"
 #include "UObject/WeakObjectPtrTemplates.h"
@@ -368,3 +367,7 @@ private:
 	uint32 bMaskFillCollisionUnderneathForNavmesh : 1;
 	ENavigationDataResolution NavMeshResolution;
 };
+
+#if UE_ENABLE_INCLUDE_ORDER_DEPRECATED_IN_5_2
+#include "CoreMinimal.h"
+#endif
diff --git a/Engine/Source/Runtime/Engine/Public/EngineDefines.h b/Engine/Source/Runtime/Engine/Public/EngineDefines.h
--- a/Engine/Source/Runtime/Engine/Public/EngineDefines.h
+++ b/Engine/Source/Runtime/Engine/Public/EngineDefines.h
@@ -4,7 +4,7 @@
 
 #pragma once
 
-#include "CoreMinimal.h"
+#include "HAL/Platform.h"
 
 /*-----------------------------------------------------------------------------
 	Configuration defines
@@ -57,3 +57,7 @@
 
 /** bits needed to store DPG value */
 #define SDPG_NumBits 3
+
+#if UE_ENABLE_INCLUDE_ORDER_DEPRECATED_IN_5_2
+#include "CoreMinimal.h"
+#endif
diff --git a/Engine/Source/Runtime/Engine/Public/Physics/AsyncPhysicsData.h b/Engine/Source/Runtime/Engine/Public/Physics/AsyncPhysicsData.h
--- a/Engine/Source/Runtime/Engine/Public/Physics/AsyncPhysicsData.h
+++ b/Engine/Source/Runtime/Engine/Public/Physics/AsyncPhysicsData.h
@@ -2,7 +2,6 @@
 
 #pragma once
 
-#include "CoreMinimal.h"
 #include "AsyncPhysicsData.generated.h"
 
 /**
@@ -28,4 +27,8 @@ protected:
 	int32 ReplicationRedundancy = 4;
 
 	friend class UAsyncPhysicsInputComponent;
-};
\ No newline at end of file
+};
+
+#if UE_ENABLE_INCLUDE_ORDER_DEPRECATED_IN_5_2
+#include "CoreMinimal.h"
+#endif
diff --git a/Engine/Source/Runtime/Engine/Public/WorldPartition/WorldPartitionStreamingSource.h b/Engine/Source/Runtime/Engine/Public/WorldPartition/WorldPartitionStreamingSource.h
--- a/Engine/Source/Runtime/Engine/Public/WorldPartition/WorldPartitionStreamingSource.h
+++ b/Engine/Source/Runtime/Engine/Public/WorldPartition/WorldPartitionStreamingSource.h
@@ -1,7 +1,6 @@
 // Copyright Epic Games, Inc. All Rights Reserved.
 #pragma once
 
-#include "CoreMinimal.h"
 #include "Math/Color.h"
 #include "Math/RandomStream.h"
 #include "WorldPartitionStreamingSource.generated.h"
@@ -557,4 +556,8 @@ struct ENGINE_API IWorldPartitionStreamingSourceProvider
 	}
 
 	virtual const UObject* GetStreamingSourceOwner() const { return nullptr; }
-};
\ No newline at end of file
+};
+
+#if UE_ENABLE_INCLUDE_ORDER_DEPRECATED_IN_5_2
+#include "CoreMinimal.h"
+#endif
diff --git a/Engine/Source/Runtime/PhysicsCore/Public/CollisionShape.h b/Engine/Source/Runtime/PhysicsCore/Public/CollisionShape.h
--- a/Engine/Source/Runtime/PhysicsCore/Public/CollisionShape.h
+++ b/Engine/Source/Runtime/PhysicsCore/Public/CollisionShape.h
@@ -2,7 +2,7 @@
 
 #pragma once 
 
-#include "CoreMinimal.h"
+#include "Math/Vector.h"
 
 /** Types of Collision Shapes that are used by Trace **/
 namespace ECollisionShape
@@ -229,3 +229,7 @@ struct FCollisionShape
 		return CapsuleShape;
 	}
 };
+
+#if UE_ENABLE_INCLUDE_ORDER_DEPRECATED_IN_5_2
+#include "CoreMinimal.h"
+#endif
diff --git a/Engine/Source/Runtime/RHI/Public/DataDrivenShaderPlatformInfo.h b/Engine/Source/Runtime/RHI/Public/DataDrivenShaderPlatformInfo.h
--- a/Engine/Source/Runtime/RHI/Public/DataDrivenShaderPlatformInfo.h
+++ b/Engine/Source/Runtime/RHI/Public/DataDrivenShaderPlatformInfo.h
@@ -2,6 +2,7 @@
 
 #pragma once
 
+#include "Internationalization/Text.h"
 #include "RHIDefinitions.h"
 #include "RHIShaderPlatform.h"
 #include "RHIFeatureLevel.h"
diff --git a/Engine/Source/Runtime/RHI/Public/RHIBreadcrumbs.h b/Engine/Source/Runtime/RHI/Public/RHIBreadcrumbs.h
--- a/Engine/Source/Runtime/RHI/Public/RHIBreadcrumbs.h
+++ b/Engine/Source/Runtime/RHI/Public/RHIBreadcrumbs.h
@@ -2,6 +2,7 @@
 
 #pragma once
 
+#include "Containers/Array.h"
 #include "RHIDefinitions.h"
 
 #if !defined(RHI_WANT_BREADCRUMB_EVENTS)
diff --git a/Engine/Source/Runtime/Renderer/Private/RayTracing/RayTracingInstanceCulling.h b/Engine/Source/Runtime/Renderer/Private/RayTracing/RayTracingInstanceCulling.h
--- a/Engine/Source/Runtime/Renderer/Private/RayTracing/RayTracingInstanceCulling.h
+++ b/Engine/Source/Runtime/Renderer/Private/RayTracing/RayTracingInstanceCulling.h
@@ -6,6 +6,7 @@
 
 #if RHI_RAYTRACING
 
+#include "Containers/ArrayView.h"
 #include "Math/MathFwd.h"
 
 class FPrimitiveSceneInfo;
