World Partition: made actor descriptors ref counts protected and exposed friend accessor struct to unit tests.

diff --git a/Engine/Source/Runtime/Engine/Private/Tests/WorldPartition/WorldPartitionHandleTests.cpp b/Engine/Source/Runtime/Engine/Private/Tests/WorldPartition/WorldPartitionHandleTests.cpp
--- a/Engine/Source/Runtime/Engine/Private/Tests/WorldPartition/WorldPartitionHandleTests.cpp
+++ b/Engine/Source/Runtime/Engine/Private/Tests/WorldPartition/WorldPartitionHandleTests.cpp
@@ -12,6 +12,21 @@
 
 #define TEST_NAME_ROOT "System.Engine.WorldPartition"
 
+#if WITH_EDITOR
+struct FWorldPartitionActorDescUnitTestAcccessor
+{
+	inline static uint32 GetSoftRefCount(const FWorldPartitionActorDesc* ActorDesc)
+	{
+		return ActorDesc->GetSoftRefCount();
+	}
+
+	inline static uint32 GetHardRefCount(const FWorldPartitionActorDesc* ActorDesc)
+	{
+		return ActorDesc->GetHardRefCount();
+	}
+};
+#endif
+
 namespace WorldPartitionTests
 {
 	IMPLEMENT_SIMPLE_AUTOMATION_TEST(FWorldPartitionSoftRefTest, TEST_NAME_ROOT ".Handle", EAutomationTestFlags::ApplicationContextMask | EAutomationTestFlags::EngineFilter)
@@ -42,65 +57,65 @@ namespace WorldPartitionTests
 		TestTrue(TEXT("Handle container"), Handle->GetContainer() == ActorDescContainer);
 		TestTrue(TEXT("Reference container"), Reference->GetContainer() == ActorDescContainer);
 
-		TestTrue(TEXT("Handle soft refcount"), Handle->GetSoftRefCount() == 1);
-		TestTrue(TEXT("Handle hard refcount"), Handle->GetHardRefCount() == 0);
-		TestTrue(TEXT("Reference soft refcount"), Reference->GetSoftRefCount() == 0);
-		TestTrue(TEXT("Reference hard refcount"), Reference->GetHardRefCount() == 1);
+		TestTrue(TEXT("Handle soft refcount"), FWorldPartitionActorDescUnitTestAcccessor::GetSoftRefCount(Handle.Get()) == 1);
+		TestTrue(TEXT("Handle hard refcount"), FWorldPartitionActorDescUnitTestAcccessor::GetHardRefCount(Handle.Get()) == 0);
+		TestTrue(TEXT("Reference soft refcount"), FWorldPartitionActorDescUnitTestAcccessor::GetSoftRefCount(Reference.Get()) == 0);
+		TestTrue(TEXT("Reference hard refcount"), FWorldPartitionActorDescUnitTestAcccessor::GetHardRefCount(Reference.Get()) == 1);
 
 		// Pin handle scope test
 		{
 			FWorldPartitionHandlePinRefScope PinRefScopeHandle(Handle);
-			TestTrue(TEXT("Pin to Handle soft refcount"), Handle->GetSoftRefCount() == 1);
-			TestTrue(TEXT("Pin to Handle hard refcount"), Handle->GetHardRefCount() == 0);
+			TestTrue(TEXT("Pin to Handle soft refcount"), FWorldPartitionActorDescUnitTestAcccessor::GetSoftRefCount(Handle.Get()) == 1);
+			TestTrue(TEXT("Pin to Handle hard refcount"), FWorldPartitionActorDescUnitTestAcccessor::GetHardRefCount(Handle.Get()) == 0);
 		}
 
 		// Pin reference scope test
 		{
 			FWorldPartitionHandlePinRefScope PinRefScopeReference(Reference);
-			TestTrue(TEXT("Pin to Reference soft refcount"), Reference->GetSoftRefCount() == 0);
-			TestTrue(TEXT("Pin to Reference hard refcount"), Reference->GetHardRefCount() == 2);
+			TestTrue(TEXT("Pin to Reference soft refcount"), FWorldPartitionActorDescUnitTestAcccessor::GetSoftRefCount(Reference.Get()) == 0);
+			TestTrue(TEXT("Pin to Reference hard refcount"), FWorldPartitionActorDescUnitTestAcccessor::GetHardRefCount(Reference.Get()) == 2);
 		}
 
 		TestTrue(TEXT("Handle/Reference equality"), Handle != Reference);
 		TestTrue(TEXT("Reference/Handle equality"), Reference != Handle);
 		
-		TestTrue(TEXT("Handle soft refcount"), Handle->GetSoftRefCount() == 1);
-		TestTrue(TEXT("Handle hard refcount"), Handle->GetHardRefCount() == 0);
-		TestTrue(TEXT("Reference soft refcount"), Reference->GetSoftRefCount() == 0);
-		TestTrue(TEXT("Reference hard refcount"), Reference->GetHardRefCount() == 1);
+		TestTrue(TEXT("Handle soft refcount"), FWorldPartitionActorDescUnitTestAcccessor::GetSoftRefCount(Handle.Get()) == 1);
+		TestTrue(TEXT("Handle hard refcount"), FWorldPartitionActorDescUnitTestAcccessor::GetHardRefCount(Handle.Get()) == 0);
+		TestTrue(TEXT("Reference soft refcount"), FWorldPartitionActorDescUnitTestAcccessor::GetSoftRefCount(Reference.Get()) == 0);
+		TestTrue(TEXT("Reference hard refcount"), FWorldPartitionActorDescUnitTestAcccessor::GetHardRefCount(Reference.Get()) == 1);
 
 		// Conversions
 		{
 			FWorldPartitionHandle HandleToReference = Reference.ToHandle();
 			TestTrue(TEXT("Handle/Reference equality"), HandleToReference == Reference);
 			TestTrue(TEXT("Reference/Handle equality"), Reference == HandleToReference);
-			TestTrue(TEXT("Reference soft refcount"), Reference->GetSoftRefCount() == 1);
-			TestTrue(TEXT("Reference hard refcount"), Reference->GetHardRefCount() == 1);
+			TestTrue(TEXT("Reference soft refcount"), FWorldPartitionActorDescUnitTestAcccessor::GetSoftRefCount(Reference.Get()) == 1);
+			TestTrue(TEXT("Reference hard refcount"), FWorldPartitionActorDescUnitTestAcccessor::GetHardRefCount(Reference.Get()) == 1);
 
 			FWorldPartitionReference ReferenceToHandle = Handle.ToReference();
 			TestTrue(TEXT("Handle/Reference equality"), ReferenceToHandle == Handle);
 			TestTrue(TEXT("Handle/Reference equality"), Handle == ReferenceToHandle);
-			TestTrue(TEXT("Reference soft refcount"), Reference->GetSoftRefCount() == 1);
-			TestTrue(TEXT("Reference hard refcount"), Reference->GetHardRefCount() == 1);
+			TestTrue(TEXT("Reference soft refcount"), FWorldPartitionActorDescUnitTestAcccessor::GetSoftRefCount(Reference.Get()) == 1);
+			TestTrue(TEXT("Reference hard refcount"), FWorldPartitionActorDescUnitTestAcccessor::GetHardRefCount(Reference.Get()) == 1);
 		}
 
-		TestTrue(TEXT("Handle soft refcount"), Handle->GetSoftRefCount() == 1);
-		TestTrue(TEXT("Handle hard refcount"), Handle->GetHardRefCount() == 0);
-		TestTrue(TEXT("Reference soft refcount"), Reference->GetSoftRefCount() == 0);
-		TestTrue(TEXT("Reference hard refcount"), Reference->GetHardRefCount() == 1);
+		TestTrue(TEXT("Handle soft refcount"), FWorldPartitionActorDescUnitTestAcccessor::GetSoftRefCount(Handle.Get()) == 1);
+		TestTrue(TEXT("Handle hard refcount"), FWorldPartitionActorDescUnitTestAcccessor::GetHardRefCount(Handle.Get()) == 0);
+		TestTrue(TEXT("Reference soft refcount"), FWorldPartitionActorDescUnitTestAcccessor::GetSoftRefCount(Reference.Get()) == 0);
+		TestTrue(TEXT("Reference hard refcount"), FWorldPartitionActorDescUnitTestAcccessor::GetHardRefCount(Reference.Get()) == 1);
 
 		// inplace new test
 		{
 			uint8 Buffer[sizeof(FWorldPartitionHandle)];
 			FWorldPartitionHandle* HandlePtr = new (Buffer) FWorldPartitionHandle(Reference.ToHandle());
 
-			TestTrue(TEXT("Handle array soft refcount"), Reference->GetSoftRefCount() == 1);
-			TestTrue(TEXT("Handle array hard refcount"), Reference->GetHardRefCount() == 1);
+			TestTrue(TEXT("Handle array soft refcount"), FWorldPartitionActorDescUnitTestAcccessor::GetSoftRefCount(Reference.Get()) == 1);
+			TestTrue(TEXT("Handle array hard refcount"), FWorldPartitionActorDescUnitTestAcccessor::GetHardRefCount(Reference.Get()) == 1);
 
 			HandlePtr->~FWorldPartitionHandle();
 
-			TestTrue(TEXT("Handle array soft refcount"), Reference->GetSoftRefCount() == 0);
-			TestTrue(TEXT("Handle array hard refcount"), Reference->GetHardRefCount() == 1);
+			TestTrue(TEXT("Handle array soft refcount"), FWorldPartitionActorDescUnitTestAcccessor::GetSoftRefCount(Reference.Get()) == 0);
+			TestTrue(TEXT("Handle array hard refcount"), FWorldPartitionActorDescUnitTestAcccessor::GetHardRefCount(Reference.Get()) == 1);
 		}
 
 		// TArray test
@@ -108,14 +123,14 @@ namespace WorldPartitionTests
 			TArray<FWorldPartitionHandle> HandleList;
 			HandleList.Add(Handle);
 
-			TestTrue(TEXT("Handle array soft refcount"), Handle->GetSoftRefCount() == 2);
-			TestTrue(TEXT("Handle array hard refcount"), Handle->GetHardRefCount() == 0);
+			TestTrue(TEXT("Handle array soft refcount"), FWorldPartitionActorDescUnitTestAcccessor::GetSoftRefCount(Handle.Get()) == 2);
+			TestTrue(TEXT("Handle array hard refcount"), FWorldPartitionActorDescUnitTestAcccessor::GetHardRefCount(Handle.Get()) == 0);
 
 			FWorldPartitionReference ReferenceToHandle = Handle.ToReference();
 			TestTrue(TEXT("Handle/Reference equality"), ReferenceToHandle == Handle);
 			TestTrue(TEXT("Handle/Reference equality"), Handle == ReferenceToHandle);
-			TestTrue(TEXT("Handle soft refcount"), Handle->GetSoftRefCount() == 2);
-			TestTrue(TEXT("Handle hard refcount"), Handle->GetHardRefCount() == 1);
+			TestTrue(TEXT("Handle soft refcount"), FWorldPartitionActorDescUnitTestAcccessor::GetSoftRefCount(Handle.Get()) == 2);
+			TestTrue(TEXT("Handle hard refcount"), FWorldPartitionActorDescUnitTestAcccessor::GetHardRefCount(Handle.Get()) == 1);
 
 			TestTrue(TEXT("Handle array contains handle"), HandleList.Contains(Handle));
 			TestTrue(TEXT("Handle array contains reference"), HandleList.Contains(ReferenceToHandle));
@@ -123,36 +138,36 @@ namespace WorldPartitionTests
 			HandleList.Add(Reference.ToHandle());
 			
 			TestTrue(TEXT("Handle array contains reference"), HandleList.Contains(Reference));
-			TestTrue(TEXT("Handle array soft refcount"), Reference->GetSoftRefCount() == 1);
-			TestTrue(TEXT("Handle array hard refcount"), Reference->GetHardRefCount() == 1);
+			TestTrue(TEXT("Handle array soft refcount"), FWorldPartitionActorDescUnitTestAcccessor::GetSoftRefCount(Reference.Get()) == 1);
+			TestTrue(TEXT("Handle array hard refcount"), FWorldPartitionActorDescUnitTestAcccessor::GetHardRefCount(Reference.Get()) == 1);
 			
 			HandleList.Remove(Handle);
-			TestTrue(TEXT("Handle array soft refcount"), Handle->GetSoftRefCount() == 1);
-			TestTrue(TEXT("Handle array hard refcount"), Handle->GetHardRefCount() == 1);
+			TestTrue(TEXT("Handle array soft refcount"), FWorldPartitionActorDescUnitTestAcccessor::GetSoftRefCount(Handle.Get()) == 1);
+			TestTrue(TEXT("Handle array hard refcount"), FWorldPartitionActorDescUnitTestAcccessor::GetHardRefCount(Handle.Get()) == 1);
 
 			HandleList.Remove(Reference.ToHandle());
-			TestTrue(TEXT("Handle array soft refcount"), Reference->GetSoftRefCount() == 0);
-			TestTrue(TEXT("Handle array hard refcount"), Reference->GetHardRefCount() == 1);
+			TestTrue(TEXT("Handle array soft refcount"), FWorldPartitionActorDescUnitTestAcccessor::GetSoftRefCount(Reference.Get()) == 0);
+			TestTrue(TEXT("Handle array hard refcount"), FWorldPartitionActorDescUnitTestAcccessor::GetHardRefCount(Reference.Get()) == 1);
 		}
 
-		TestTrue(TEXT("Handle soft refcount"), Handle->GetSoftRefCount() == 1);
-		TestTrue(TEXT("Handle hard refcount"), Handle->GetHardRefCount() == 0);
-		TestTrue(TEXT("Reference soft refcount"), Reference->GetSoftRefCount() == 0);
-		TestTrue(TEXT("Reference hard refcount"), Reference->GetHardRefCount() == 1);
+		TestTrue(TEXT("Handle soft refcount"), FWorldPartitionActorDescUnitTestAcccessor::GetSoftRefCount(Handle.Get()) == 1);
+		TestTrue(TEXT("Handle hard refcount"), FWorldPartitionActorDescUnitTestAcccessor::GetHardRefCount(Handle.Get()) == 0);
+		TestTrue(TEXT("Reference soft refcount"), FWorldPartitionActorDescUnitTestAcccessor::GetSoftRefCount(Reference.Get()) == 0);
+		TestTrue(TEXT("Reference hard refcount"), FWorldPartitionActorDescUnitTestAcccessor::GetHardRefCount(Reference.Get()) == 1);
 
 		// TSet test
 		{
 			TSet<FWorldPartitionHandle> HandleSet;
 			HandleSet.Add(Handle);
 
-			TestTrue(TEXT("Handle set soft refcount"), Handle->GetSoftRefCount() == 2);
-			TestTrue(TEXT("Handle set hard refcount"), Handle->GetHardRefCount() == 0);
+			TestTrue(TEXT("Handle set soft refcount"), FWorldPartitionActorDescUnitTestAcccessor::GetSoftRefCount(Handle.Get()) == 2);
+			TestTrue(TEXT("Handle set hard refcount"), FWorldPartitionActorDescUnitTestAcccessor::GetHardRefCount(Handle.Get()) == 0);
 
 			FWorldPartitionReference ReferenceToHandle = Handle.ToReference();
 			TestTrue(TEXT("Handle/Reference equality"), ReferenceToHandle == Handle);
 			TestTrue(TEXT("Handle/Reference equality"), Handle == ReferenceToHandle);
-			TestTrue(TEXT("Handle soft refcount"), Handle->GetSoftRefCount() == 2);
-			TestTrue(TEXT("Handle hard refcount"), Handle->GetHardRefCount() == 1);
+			TestTrue(TEXT("Handle soft refcount"), FWorldPartitionActorDescUnitTestAcccessor::GetSoftRefCount(Handle.Get()) == 2);
+			TestTrue(TEXT("Handle hard refcount"), FWorldPartitionActorDescUnitTestAcccessor::GetHardRefCount(Handle.Get()) == 1);
 
 			TestTrue(TEXT("Handle set contains handle"), HandleSet.Contains(Handle));
 			TestTrue(TEXT("Handle set contains reference"), HandleSet.Contains(ReferenceToHandle.ToHandle()));
@@ -160,49 +175,49 @@ namespace WorldPartitionTests
 			HandleSet.Add(Reference.ToHandle());
 			TestTrue(TEXT("Handle set contains reference"), HandleSet.Contains(Reference.ToHandle()));
 
-			TestTrue(TEXT("Reference soft refcount"), Reference->GetSoftRefCount() == 1);
-			TestTrue(TEXT("Reference hard refcount"), Reference->GetHardRefCount() == 1);
+			TestTrue(TEXT("Reference soft refcount"), FWorldPartitionActorDescUnitTestAcccessor::GetSoftRefCount(Reference.Get()) == 1);
+			TestTrue(TEXT("Reference hard refcount"), FWorldPartitionActorDescUnitTestAcccessor::GetHardRefCount(Reference.Get()) == 1);
 		}
 
-		TestTrue(TEXT("Handle soft refcount"), Handle->GetSoftRefCount() == 1);
-		TestTrue(TEXT("Handle hard refcount"), Handle->GetHardRefCount() == 0);
-		TestTrue(TEXT("Reference soft refcount"), Reference->GetSoftRefCount() == 0);
-		TestTrue(TEXT("Reference hard refcount"), Reference->GetHardRefCount() == 1);
+		TestTrue(TEXT("Handle soft refcount"), FWorldPartitionActorDescUnitTestAcccessor::GetSoftRefCount(Handle.Get()) == 1);
+		TestTrue(TEXT("Handle hard refcount"), FWorldPartitionActorDescUnitTestAcccessor::GetHardRefCount(Handle.Get()) == 0);
+		TestTrue(TEXT("Reference soft refcount"), FWorldPartitionActorDescUnitTestAcccessor::GetSoftRefCount(Reference.Get()) == 0);
+		TestTrue(TEXT("Reference hard refcount"), FWorldPartitionActorDescUnitTestAcccessor::GetHardRefCount(Reference.Get()) == 1);
 
 		// Move tests
 		{
 			// Handle move
-			TestTrue(TEXT("Handle soft refcount"), Handle->GetSoftRefCount() == 1);
-			TestTrue(TEXT("Handle hard refcount"), Handle->GetHardRefCount() == 0);
+			TestTrue(TEXT("Handle soft refcount"), FWorldPartitionActorDescUnitTestAcccessor::GetSoftRefCount(Handle.Get()) == 1);
+			TestTrue(TEXT("Handle hard refcount"), FWorldPartitionActorDescUnitTestAcccessor::GetHardRefCount(Handle.Get()) == 0);
 			{
 				FWorldPartitionHandle HandleCopy(MoveTemp(Handle));
 				TestTrue(TEXT("Handle move src not valid"), !Handle.IsValid());
 				TestTrue(TEXT("Handle move dst valid"), HandleCopy.IsValid());
-				TestTrue(TEXT("Handle soft refcount"), HandleCopy->GetSoftRefCount() == 1);
-				TestTrue(TEXT("Handle hard refcount"), HandleCopy->GetHardRefCount() == 0);
+				TestTrue(TEXT("Handle soft refcount"), FWorldPartitionActorDescUnitTestAcccessor::GetSoftRefCount(HandleCopy.Get()) == 1);
+				TestTrue(TEXT("Handle hard refcount"), FWorldPartitionActorDescUnitTestAcccessor::GetHardRefCount(HandleCopy.Get()) == 0);
 
 				Handle = MoveTemp(HandleCopy);
 				TestTrue(TEXT("Handle move src not valid"), !HandleCopy.IsValid());
 				TestTrue(TEXT("Handle move dst valid"), Handle.IsValid());
-				TestTrue(TEXT("Handle soft refcount"), Handle->GetSoftRefCount() == 1);
-				TestTrue(TEXT("Handle hard refcount"), Handle->GetHardRefCount() == 0);
+				TestTrue(TEXT("Handle soft refcount"), FWorldPartitionActorDescUnitTestAcccessor::GetSoftRefCount(Handle.Get()) == 1);
+				TestTrue(TEXT("Handle hard refcount"), FWorldPartitionActorDescUnitTestAcccessor::GetHardRefCount(Handle.Get()) == 0);
 			}
 
 			// Reference move
-			TestTrue(TEXT("Reference soft refcount"), Reference->GetSoftRefCount() == 0);
-			TestTrue(TEXT("Reference hard refcount"), Reference->GetHardRefCount() == 1);
+			TestTrue(TEXT("Reference soft refcount"), FWorldPartitionActorDescUnitTestAcccessor::GetSoftRefCount(Reference.Get()) == 0);
+			TestTrue(TEXT("Reference hard refcount"), FWorldPartitionActorDescUnitTestAcccessor::GetHardRefCount(Reference.Get()) == 1);
 			{
 				FWorldPartitionReference ReferenceCopy(MoveTemp(Reference));
 				TestTrue(TEXT("Reference move src not valid"), !Reference.IsValid());
 				TestTrue(TEXT("Reference move dst valid"), ReferenceCopy.IsValid());
-				TestTrue(TEXT("Reference soft refcount"), ReferenceCopy->GetSoftRefCount() == 0);
-				TestTrue(TEXT("Reference hard refcount"), ReferenceCopy->GetHardRefCount() == 1);
+				TestTrue(TEXT("Reference soft refcount"), FWorldPartitionActorDescUnitTestAcccessor::GetSoftRefCount(ReferenceCopy.Get()) == 0);
+				TestTrue(TEXT("Reference hard refcount"), FWorldPartitionActorDescUnitTestAcccessor::GetHardRefCount(ReferenceCopy.Get()) == 1);
 
 				Reference = MoveTemp(ReferenceCopy);
 				TestTrue(TEXT("Reference move src not valid"), !ReferenceCopy.IsValid());
 				TestTrue(TEXT("Reference move dst valid"), Reference.IsValid());
-				TestTrue(TEXT("Reference soft refcount"), Reference->GetSoftRefCount() == 0);
-				TestTrue(TEXT("Reference hard refcount"), Reference->GetHardRefCount() == 1);
+				TestTrue(TEXT("Reference soft refcount"), FWorldPartitionActorDescUnitTestAcccessor::GetSoftRefCount(Reference.Get()) == 0);
+				TestTrue(TEXT("Reference hard refcount"), FWorldPartitionActorDescUnitTestAcccessor::GetHardRefCount(Reference.Get()) == 1);
 			}
 
 			// Handle reference move
@@ -210,14 +225,14 @@ namespace WorldPartitionTests
 				FWorldPartitionHandle HandleFromReference(MoveTemp(Reference));
 				TestTrue(TEXT("Handle move src not valid"), !Reference.IsValid());
 				TestTrue(TEXT("Handle move dst valid"), HandleFromReference.IsValid());
-				TestTrue(TEXT("Handle soft refcount"), HandleFromReference->GetSoftRefCount() == 1);
-				TestTrue(TEXT("Handle hard refcount"), HandleFromReference->GetHardRefCount() == 0);
+				TestTrue(TEXT("Handle soft refcount"), FWorldPartitionActorDescUnitTestAcccessor::GetSoftRefCount(HandleFromReference.Get()) == 1);
+				TestTrue(TEXT("Handle hard refcount"), FWorldPartitionActorDescUnitTestAcccessor::GetHardRefCount(HandleFromReference.Get()) == 0);
 
 				Reference = MoveTemp(HandleFromReference);
 				TestTrue(TEXT("Handle move src not valid"), !HandleFromReference.IsValid());
 				TestTrue(TEXT("Reference move dst valid"), Reference.IsValid());
-				TestTrue(TEXT("Reference soft refcount"), Reference->GetSoftRefCount() == 0);
-				TestTrue(TEXT("Reference hard refcount"), Reference->GetHardRefCount() == 1);
+				TestTrue(TEXT("Reference soft refcount"), FWorldPartitionActorDescUnitTestAcccessor::GetSoftRefCount(Reference.Get()) == 0);
+				TestTrue(TEXT("Reference hard refcount"), FWorldPartitionActorDescUnitTestAcccessor::GetHardRefCount(Reference.Get()) == 1);
 			}
 
 			// Reference handle move
@@ -225,27 +240,27 @@ namespace WorldPartitionTests
 				FWorldPartitionReference ReferenceFromHandle(MoveTemp(Handle));
 				TestTrue(TEXT("Handle move src not valid"), !Handle.IsValid());
 				TestTrue(TEXT("Reference move dst valid"), ReferenceFromHandle.IsValid());
-				TestTrue(TEXT("Reference soft refcount"), ReferenceFromHandle->GetSoftRefCount() == 0);
-				TestTrue(TEXT("Reference hard refcount"), ReferenceFromHandle->GetHardRefCount() == 1);
+				TestTrue(TEXT("Reference soft refcount"), FWorldPartitionActorDescUnitTestAcccessor::GetSoftRefCount(ReferenceFromHandle.Get()) == 0);
+				TestTrue(TEXT("Reference hard refcount"), FWorldPartitionActorDescUnitTestAcccessor::GetHardRefCount(ReferenceFromHandle.Get()) == 1);
 
 				Handle = MoveTemp(ReferenceFromHandle);
 				TestTrue(TEXT("Reference move src not valid"), !ReferenceFromHandle.IsValid());
 				TestTrue(TEXT("Handle move dst valid"), Handle.IsValid());
-				TestTrue(TEXT("Handle soft refcount"), Handle->GetSoftRefCount() == 1);
-				TestTrue(TEXT("Handle hard refcount"), Handle->GetHardRefCount() == 0);
+				TestTrue(TEXT("Handle soft refcount"), FWorldPartitionActorDescUnitTestAcccessor::GetSoftRefCount(Handle.Get()) == 1);
+				TestTrue(TEXT("Handle hard refcount"), FWorldPartitionActorDescUnitTestAcccessor::GetHardRefCount(Handle.Get()) == 0);
 			}
 
 			// Reset
 			{
 				FWorldPartitionReference ReferenceFromHandle(Handle.ToReference());
 				TestTrue(TEXT("Reference valid"), ReferenceFromHandle.IsValid());
-				TestTrue(TEXT("Reference soft refcount"), ReferenceFromHandle->GetSoftRefCount() == 1);
-				TestTrue(TEXT("Reference hard refcount"), ReferenceFromHandle->GetHardRefCount() == 1);
+				TestTrue(TEXT("Reference soft refcount"), FWorldPartitionActorDescUnitTestAcccessor::GetSoftRefCount(ReferenceFromHandle.Get()) == 1);
+				TestTrue(TEXT("Reference hard refcount"), FWorldPartitionActorDescUnitTestAcccessor::GetHardRefCount(ReferenceFromHandle.Get()) == 1);
 
 				ReferenceFromHandle.Reset();
 				TestTrue(TEXT("Reference not valid"), !ReferenceFromHandle.IsValid());
-				TestTrue(TEXT("Handle soft refcount"), Handle->GetSoftRefCount() == 1);
-				TestTrue(TEXT("Handle hard refcount"), Handle->GetHardRefCount() == 0);
+				TestTrue(TEXT("Handle soft refcount"), FWorldPartitionActorDescUnitTestAcccessor::GetSoftRefCount(Handle.Get()) == 1);
+				TestTrue(TEXT("Handle hard refcount"), FWorldPartitionActorDescUnitTestAcccessor::GetHardRefCount(Handle.Get()) == 0);
 			}
 		}
 
diff --git a/Engine/Source/Runtime/Engine/Public/WorldPartition/WorldPartitionActorDesc.h b/Engine/Source/Runtime/Engine/Public/WorldPartition/WorldPartitionActorDesc.h
--- a/Engine/Source/Runtime/Engine/Public/WorldPartition/WorldPartitionActorDesc.h
+++ b/Engine/Source/Runtime/Engine/Public/WorldPartition/WorldPartitionActorDesc.h
@@ -69,6 +69,7 @@ class ENGINE_API FWorldPartitionActorDesc
 	friend struct FWorldPartitionHandleImpl;
 	friend struct FWorldPartitionReferenceImpl;
 	friend struct FWorldPartitionActorDescUtils;
+	friend struct FWorldPartitionActorDescUnitTestAcccessor;
 
 public:
 	virtual ~FWorldPartitionActorDesc() {}
@@ -179,7 +180,6 @@ protected:
 		return --HardRefCount;
 	}
 
-public:
 	inline uint32 GetSoftRefCount() const
 	{
 		return SoftRefCount;
@@ -190,6 +190,7 @@ public:
 		return HardRefCount;
 	}
 
+public:
 	const TArray<FGuid>& GetReferences() const
 	{
 		return References;
