Added check that reports errors if a parameter is added in a Material bound to different values. This should not be allowed.

Code now checks for this occurrences and reprorts an error message instead of siliently trying to compile the material anyway.

This was causing a crash when two material functions containing the same parameter bound to different values (e.g. different textures) are used in the same material.

diff --git a/Engine/Plugins/Enterprise/GLTFExporter/Source/GLTFMaterialBaking/Private/GLTFExportMaterialProxy.h b/Engine/Plugins/Enterprise/GLTFExporter/Source/GLTFMaterialBaking/Private/GLTFExportMaterialProxy.h
--- a/Engine/Plugins/Enterprise/GLTFExporter/Source/GLTFMaterialBaking/Private/GLTFExportMaterialProxy.h
+++ b/Engine/Plugins/Enterprise/GLTFExporter/Source/GLTFMaterialBaking/Private/GLTFExportMaterialProxy.h
@@ -548,6 +548,11 @@ public:
 		}
 	}
 
+	virtual bool CheckInValidStateForCompilation(class FMaterialCompiler* Compiler) const override
+	{
+		return Material && Material->CheckInValidStateForCompilation(Compiler);
+	}
+
 private:
 	int32 CompileInputForCustomOutput(FMaterialCompiler* Compiler, int32 InputIndex, uint32 ForceCastFlags) const
 	{
diff --git a/Engine/Source/Developer/MaterialBaking/Private/ExportMaterialProxy.h b/Engine/Source/Developer/MaterialBaking/Private/ExportMaterialProxy.h
--- a/Engine/Source/Developer/MaterialBaking/Private/ExportMaterialProxy.h
+++ b/Engine/Source/Developer/MaterialBaking/Private/ExportMaterialProxy.h
@@ -562,6 +562,11 @@ public:
 		}
 	}
 
+	virtual bool CheckInValidStateForCompilation(FMaterialCompiler* Compiler) const override
+	{
+		return Material && Material->CheckInValidStateForCompilation(Compiler);
+	}
+
 private:
 	int32 CompileInputForCustomOutput(FMaterialCompiler* Compiler, int32 InputIndex, uint32 ForceCastFlags) const
 	{
diff --git a/Engine/Source/Developer/MaterialUtilities/Private/MaterialUtilities.cpp b/Engine/Source/Developer/MaterialUtilities/Private/MaterialUtilities.cpp
--- a/Engine/Source/Developer/MaterialUtilities/Private/MaterialUtilities.cpp
+++ b/Engine/Source/Developer/MaterialUtilities/Private/MaterialUtilities.cpp
@@ -765,6 +765,10 @@ public:
 		}
 	}
 
+	virtual bool CheckInValidStateForCompilation(class FMaterialCompiler* Compiler) const override
+	{
+		return Material && Material->CheckInValidStateForCompilation(Compiler);
+	}
 private:
 	/** The material interface for this proxy */
 	UMaterialInterface* MaterialInterface;
diff --git a/Engine/Source/Editor/MaterialEditor/Private/MaterialEditor.cpp b/Engine/Source/Editor/MaterialEditor/Private/MaterialEditor.cpp
--- a/Engine/Source/Editor/MaterialEditor/Private/MaterialEditor.cpp
+++ b/Engine/Source/Editor/MaterialEditor/Private/MaterialEditor.cpp
@@ -393,6 +393,11 @@ bool FMatExpressionPreview::IsUsingNewHLSLGenerator() const
 	return false;
 }
 
+bool FMatExpressionPreview::CheckInValidStateForCompilation(FMaterialCompiler* Compiler) const
+{
+	return Expression.IsValid() && Expression->Material && Expression->Material->CheckInValidStateForCompilation(Compiler);
+}
+
 /////////////////////
 // FMaterialEditor //
 /////////////////////
diff --git a/Engine/Source/Editor/MaterialEditor/Private/MaterialEditor.h b/Engine/Source/Editor/MaterialEditor/Private/MaterialEditor.h
--- a/Engine/Source/Editor/MaterialEditor/Private/MaterialEditor.h
+++ b/Engine/Source/Editor/MaterialEditor/Private/MaterialEditor.h
@@ -166,6 +166,12 @@ public:
 		}
 	}
 
+	/**
+	 * Checks that no pre-compilation errors have been detected and if so it reports them using specified compiler.
+	 * @return whether no errors occurred.
+	 */
+	 bool CheckInValidStateForCompilation(class FMaterialCompiler* Compiler) const override;
+
 	float UnrelatedNodesOpacity;
 
 private:
diff --git a/Engine/Source/Editor/UnrealEd/Private/Lightmass/LightmassRender.cpp b/Engine/Source/Editor/UnrealEd/Private/Lightmass/LightmassRender.cpp
--- a/Engine/Source/Editor/UnrealEd/Private/Lightmass/LightmassRender.cpp
+++ b/Engine/Source/Editor/UnrealEd/Private/Lightmass/LightmassRender.cpp
@@ -980,6 +980,11 @@ public:
 		}
 	}
 
+	virtual bool CheckInValidStateForCompilation(class FMaterialCompiler* Compiler) const override
+	{
+		return Material && Material->CheckInValidStateForCompilation(Compiler);
+	}
+
 private:
 	/** The material interface for this proxy */
 	UMaterialInterface* MaterialInterface;
diff --git a/Engine/Source/Runtime/Engine/Classes/Materials/Material.h b/Engine/Source/Runtime/Engine/Classes/Materials/Material.h
--- a/Engine/Source/Runtime/Engine/Classes/Materials/Material.h
+++ b/Engine/Source/Runtime/Engine/Classes/Materials/Material.h
@@ -1238,6 +1238,12 @@ public:
 	 */
 	ENGINE_API virtual void GetLightingGuidChain(bool bIncludeTextures, TArray<FGuid>& OutGuids) const override;
 
+	/**
+	 * Checks that no pre-compilation errors have been detected and if so it reports them using specified compiler.
+	 * @return whether no errors occurred.
+	 */
+	ENGINE_API bool CheckInValidStateForCompilation(class FMaterialCompiler* Compiler) const;
+	
 private:
 	void BackwardsCompatibilityInputConversion();
 	void BackwardsCompatibilityVirtualTextureOutputConversion();
diff --git a/Engine/Source/Runtime/Engine/Private/Materials/HLSLMaterialTranslator.cpp b/Engine/Source/Runtime/Engine/Private/Materials/HLSLMaterialTranslator.cpp
--- a/Engine/Source/Runtime/Engine/Private/Materials/HLSLMaterialTranslator.cpp
+++ b/Engine/Source/Runtime/Engine/Private/Materials/HLSLMaterialTranslator.cpp
@@ -757,7 +757,6 @@ bool FHLSLMaterialTranslator::Translate()
 	STAT(double HLSLTranslateTime = 0);
 	{
 		SCOPE_SECONDS_COUNTER(HLSLTranslateTime);
-		bSuccess = true;
 
 		check(ScopeStack.Num() == 0);
 
@@ -767,8 +766,10 @@ bool FHLSLMaterialTranslator::Translate()
 		Material->CompileErrors.Empty();
 		Material->ErrorExpressions.Empty();
 
-		bEnableExecutionFlow = Material->IsUsingControlFlow();
+		// Verify that no pre-compilation errors have occured.
+		bSuccess = Material->CheckInValidStateForCompilation(this);
 
+		bEnableExecutionFlow = Material->IsUsingControlFlow();
 		bCompileForComputeShader = Material->IsLightFunction();
 
 		const bool bStrataEnabled = Engine_IsStrataEnabled();
@@ -7021,7 +7022,13 @@ int32 FHLSLMaterialTranslator::TextureParameter(FName ParameterName, UTexture* I
 
 	EMaterialValueType ShaderType = DefaultValue->GetMaterialType();
 	TextureReferenceIndex = Material->GetReferencedTextures().Find(DefaultValue);
-	checkf(TextureReferenceIndex != INDEX_NONE, TEXT("Material expression called Compiler->TextureParameter() without implementing UMaterialExpression::GetReferencedTexture properly"));
+	if (TextureReferenceIndex == INDEX_NONE)
+	{
+		FString TextureName;
+		DefaultValue->GetName(TextureName);
+		Errorf(TEXT("Could not resolve referenced texture '%s'."), *TextureName);
+		return INDEX_NONE;
+	}
 
 	FMaterialParameterInfo ParameterInfo = GetParameterAssociationInfo();
 	ParameterInfo.Name = ParameterName;
diff --git a/Engine/Source/Runtime/Engine/Private/Materials/Material.cpp b/Engine/Source/Runtime/Engine/Private/Materials/Material.cpp
--- a/Engine/Source/Runtime/Engine/Private/Materials/Material.cpp
+++ b/Engine/Source/Runtime/Engine/Private/Materials/Material.cpp
@@ -6242,6 +6242,25 @@ void UMaterial::ForceRecompileForRendering()
 	CacheResourceShadersForRendering( false );
 }
 
+bool UMaterial::CheckInValidStateForCompilation(class FMaterialCompiler* Compiler) const
+{
+	bool bSuccess = true;
+	
+	// Check and report errors due to duplicate parameters set to distinct values.
+	if (!CachedExpressionData->DuplicateParameterErrors.IsEmpty())
+	{
+		for (const TPair<TObjectPtr<UMaterialExpression>, FName>& Error : CachedExpressionData->DuplicateParameterErrors)
+		{
+			FString ErrorMsg = FString::Format(TEXT("Parameter '{0}' is set multiple times to different values. Make sure each parameter is set once or always to the same value (e.g. same texture)."), { *Error.Get<1>().ToString() });
+			Compiler->AppendExpressionError(Error.Get<0>(), *ErrorMsg);
+		}
+
+		bSuccess = false;
+	}
+
+	return bSuccess;
+}
+
 UMaterial::FMaterialCompilationFinished UMaterial::MaterialCompilationFinishedEvent;
 UMaterial::FMaterialCompilationFinished& UMaterial::OnMaterialCompilationFinished()
 {
diff --git a/Engine/Source/Runtime/Engine/Private/Materials/MaterialCachedData.cpp b/Engine/Source/Runtime/Engine/Private/Materials/MaterialCachedData.cpp
--- a/Engine/Source/Runtime/Engine/Private/Materials/MaterialCachedData.cpp
+++ b/Engine/Source/Runtime/Engine/Private/Materials/MaterialCachedData.cpp
@@ -87,31 +87,32 @@ void FMaterialCachedExpressionData::AppendReferencedParameterCollectionIdsTo(TAr
 }
 
 #if WITH_EDITOR
-static int32 TryAddParameter(FMaterialCachedExpressionData& CachedData,
+static bool TryAddParameter(FMaterialCachedExpressionData& CachedData,
 	EMaterialParameterType Type,
 	const FMaterialParameterInfo& ParameterInfo,
-	const FMaterialCachedParameterEditorInfo& InEditorInfo)
+	const FMaterialCachedParameterEditorInfo& InEditorInfo,
+	int32& OutIndex)
 {
 	check(CachedData.EditorOnlyData);
 	FMaterialCachedParameterEntry& Entry = CachedData.GetParameterTypeEntry(Type);
 	FMaterialCachedParameterEditorEntry& EditorEntry = CachedData.EditorOnlyData->EditorEntries[(int32)Type];
 
 	FSetElementId ElementId = Entry.ParameterInfoSet.FindId(ParameterInfo);
-	int32 Index = INDEX_NONE;
+	OutIndex = INDEX_NONE;
 	if (!ElementId.IsValidId())
 	{
 		ElementId = Entry.ParameterInfoSet.Add(ParameterInfo);
-		Index = ElementId.AsInteger();
-		EditorEntry.EditorInfo.Insert(InEditorInfo, Index);
+		OutIndex = ElementId.AsInteger();
+		EditorEntry.EditorInfo.Insert(InEditorInfo, OutIndex);
 		// should be valid as long as we don't ever remove elements from ParameterInfoSet
 		check(Entry.ParameterInfoSet.Num() == EditorEntry.EditorInfo.Num());
-		return Index;
+		return true;
 	}
 
 	// Update any editor values that haven't been set yet
 	// TODO still need to do this??
-	Index = ElementId.AsInteger();
-	FMaterialCachedParameterEditorInfo& EditorInfo = EditorEntry.EditorInfo[Index];
+	OutIndex = ElementId.AsInteger();
+	FMaterialCachedParameterEditorInfo& EditorInfo = EditorEntry.EditorInfo[OutIndex];
 	if (!EditorInfo.ExpressionGuid.IsValid())
 	{
 		EditorInfo.ExpressionGuid = InEditorInfo.ExpressionGuid;
@@ -126,11 +127,11 @@ static int32 TryAddParameter(FMaterialCachedExpressionData& CachedData,
 		EditorInfo.SortPriority = InEditorInfo.SortPriority;
 	}
 	
-	// Still return INDEX_NONE, to signify this parameter was already added (don't want to add it again)
-	return INDEX_NONE;
+	// Still return false, to signify this parameter was already added (don't want to add it again)
+	return false;
 }
 
-void FMaterialCachedExpressionData::AddParameter(const FMaterialParameterInfo& ParameterInfo, const FMaterialParameterMetadata& ParameterMeta, UObject*& OutReferencedTexture)
+bool FMaterialCachedExpressionData::AddParameter(const FMaterialParameterInfo& ParameterInfo, const FMaterialParameterMetadata& ParameterMeta, UObject*& OutReferencedTexture)
 {
 	check(EditorOnlyData);
 	int32 AssetIndex = INDEX_NONE;
@@ -140,8 +141,8 @@ void FMaterialCachedExpressionData::AddParameter(const FMaterialParameterInfo& P
 	}
 
 	const FMaterialCachedParameterEditorInfo EditorInfo(ParameterMeta.ExpressionGuid, ParameterMeta.Description, ParameterMeta.Group, ParameterMeta.SortPriority, AssetIndex);
-	const int32 Index = TryAddParameter(*this, ParameterMeta.Value.Type, ParameterInfo, EditorInfo);
-	if (Index != INDEX_NONE)
+	int32 Index = INDEX_NONE;
+	if (TryAddParameter(*this, ParameterMeta.Value.Type, ParameterInfo, EditorInfo, Index))
 	{
 		switch (ParameterMeta.Value.Type)
 		{
@@ -161,20 +162,24 @@ void FMaterialCachedExpressionData::AddParameter(const FMaterialParameterInfo& P
 				EditorOnlyData->ScalarCurveAtlasValues.Insert(nullptr, Index);
 			}
 			break;
+
 		case EMaterialParameterType::Vector:
 			VectorValues.Insert(ParameterMeta.Value.AsLinearColor(), Index);
 			EditorOnlyData->VectorChannelNameValues.Insert(ParameterMeta.ChannelNames, Index);
 			EditorOnlyData->VectorUsedAsChannelMaskValues.Insert(ParameterMeta.bUsedAsChannelMask, Index);
 			VectorPrimitiveDataIndexValues.Insert(ParameterMeta.PrimitiveDataIndex, Index);
 			break;
+
 		case EMaterialParameterType::DoubleVector:
 			DoubleVectorValues.Insert(ParameterMeta.Value.AsVector4d(), Index);
 			break;
+
 		case EMaterialParameterType::Texture:
 			TextureValues.Insert(ParameterMeta.Value.Texture, Index);
 			EditorOnlyData->TextureChannelNameValues.Insert(ParameterMeta.ChannelNames, Index);
 			OutReferencedTexture = ParameterMeta.Value.Texture;
 			break;
+
 		case EMaterialParameterType::Font:
 			FontValues.Insert(ParameterMeta.Value.Font.Value, Index);
 			FontPageValues.Insert(ParameterMeta.Value.Font.Page, Index);
@@ -183,25 +188,93 @@ void FMaterialCachedExpressionData::AddParameter(const FMaterialParameterInfo& P
 				OutReferencedTexture = ParameterMeta.Value.Font.Value->Textures[ParameterMeta.Value.Font.Page];
 			}
 			break;
+
 		case EMaterialParameterType::RuntimeVirtualTexture:
 			RuntimeVirtualTextureValues.Insert(ParameterMeta.Value.RuntimeVirtualTexture, Index);
 			OutReferencedTexture = ParameterMeta.Value.RuntimeVirtualTexture;
 			break;
+
 		case EMaterialParameterType::SparseVolumeTexture:
 			SparseVolumeTextureValues.Insert(ParameterMeta.Value.SparseVolumeTexture, Index);
 			OutReferencedTexture = ParameterMeta.Value.SparseVolumeTexture;
 			break;
+
 		case EMaterialParameterType::StaticSwitch:
 			EditorOnlyData->StaticSwitchValues.Insert(ParameterMeta.Value.AsStaticSwitch(), Index);
 			break;
+
 		case EMaterialParameterType::StaticComponentMask:
 			EditorOnlyData->StaticComponentMaskValues.Insert(ParameterMeta.Value.AsStaticComponentMask(), Index);
 			break;
+
 		default:
 			checkNoEntry();
 			break;
 		}
 	}
+	else
+	{
+		bool bSameValue = true;
+		switch (ParameterMeta.Value.Type)
+		{
+		case EMaterialParameterType::Scalar:
+			bSameValue =
+				ScalarValues[Index] == ParameterMeta.Value.AsScalar()
+				&& ScalarPrimitiveDataIndexValues[Index] == ParameterMeta.PrimitiveDataIndex;
+			if (ParameterMeta.bUsedAsAtlasPosition)
+			{
+				bSameValue &=
+					EditorOnlyData->ScalarCurveValues[Index] == ParameterMeta.ScalarCurve.Get()
+					&& EditorOnlyData->ScalarCurveAtlasValues[Index] == ParameterMeta.ScalarAtlas.Get();
+			}
+			break;
+
+		case EMaterialParameterType::Vector:
+			bSameValue = VectorValues[Index] == ParameterMeta.Value.AsLinearColor()
+				&& EditorOnlyData->VectorChannelNameValues[Index] == ParameterMeta.ChannelNames
+				&& EditorOnlyData->VectorUsedAsChannelMaskValues[Index] == ParameterMeta.bUsedAsChannelMask
+				&& VectorPrimitiveDataIndexValues[Index] == ParameterMeta.PrimitiveDataIndex;
+			break;
+
+		case EMaterialParameterType::DoubleVector:
+			bSameValue = DoubleVectorValues[Index] == ParameterMeta.Value.AsVector4d();
+			break;
+
+		case EMaterialParameterType::Texture:
+			bSameValue = TextureValues[Index] == ParameterMeta.Value.Texture
+				&& EditorOnlyData->TextureChannelNameValues[Index] == ParameterMeta.ChannelNames;
+			break;
+
+		case EMaterialParameterType::Font:
+			bSameValue = FontValues[Index] == ParameterMeta.Value.Font.Value
+				&& FontPageValues[Index] == ParameterMeta.Value.Font.Page;
+			break;
+
+		case EMaterialParameterType::RuntimeVirtualTexture:
+			bSameValue = RuntimeVirtualTextureValues[Index] == ParameterMeta.Value.RuntimeVirtualTexture;
+			break;
+
+		case EMaterialParameterType::SparseVolumeTexture:
+			bSameValue = SparseVolumeTextureValues[Index] == ParameterMeta.Value.SparseVolumeTexture;
+			break;
+
+		case EMaterialParameterType::StaticSwitch:
+			bSameValue = EditorOnlyData->StaticSwitchValues[Index] == ParameterMeta.Value.AsStaticSwitch();
+			break;
+
+		case EMaterialParameterType::StaticComponentMask:
+			bSameValue = EditorOnlyData->StaticComponentMaskValues[Index] == ParameterMeta.Value.AsStaticComponentMask();
+			break;
+
+		default:
+			checkNoEntry();
+			break;
+		}
+
+		return bSameValue;
+	}
+
+	return true;
 }
 
 void FMaterialCachedExpressionData::UpdateForFunction(const FMaterialCachedExpressionContext& Context, UMaterialFunctionInterface* Function, EMaterialParameterAssociation Association, int32 ParameterIndex)
@@ -280,7 +353,12 @@ void FMaterialCachedExpressionData::UpdateForExpressions(const FMaterialCachedEx
 			}
 
 			const FMaterialParameterInfo ParameterInfo(ParameterName, Association, ParameterIndex);
-			AddParameter(ParameterInfo, ParameterMeta, ReferencedTexture);
+
+			// Try add the parameter. If this fails, the parameter is being added twice with different values. Report it as error.
+			if (!AddParameter(ParameterInfo, ParameterMeta, ReferencedTexture))
+			{
+				DuplicateParameterErrors.AddUnique({ Expression, ParameterName });
+			}
 		}
 
 		// We first try to extract the referenced texture from the parameter value, that way we'll also get the proper texture in case value is overriden by a function instance
diff --git a/Engine/Source/Runtime/Engine/Private/Materials/MaterialHLSLTree.cpp b/Engine/Source/Runtime/Engine/Private/Materials/MaterialHLSLTree.cpp
--- a/Engine/Source/Runtime/Engine/Private/Materials/MaterialHLSLTree.cpp
+++ b/Engine/Source/Runtime/Engine/Private/Materials/MaterialHLSLTree.cpp
@@ -458,12 +458,13 @@ bool FExpressionParameter::PrepareValue(FEmitContext& Context, FEmitScope& Scope
 	{
 		if (!ParameterInfo.Name.IsNone())
 		{
-			EmitData.CachedExpressionData->AddParameter(ParameterInfo, ParameterMeta);
+			UObject* UnusedReferencedTexture;
+			EmitData.CachedExpressionData->AddParameter(ParameterInfo, ParameterMeta, UnusedReferencedTexture);
 		}
 
 		UObject* ReferencedTexture = ParameterMeta.Value.AsTextureObject();
 		if (ReferencedTexture)
-		{
+		{	
 			EmitData.CachedExpressionData->ReferencedTextures.AddUnique(ReferencedTexture);
 		}
 	}
diff --git a/Engine/Source/Runtime/Engine/Private/Materials/MaterialShared.cpp b/Engine/Source/Runtime/Engine/Private/Materials/MaterialShared.cpp
--- a/Engine/Source/Runtime/Engine/Private/Materials/MaterialShared.cpp
+++ b/Engine/Source/Runtime/Engine/Private/Materials/MaterialShared.cpp
@@ -1434,6 +1434,11 @@ bool FMaterialResource::GetAllowDevelopmentShaderCompile()const
 	return Material->bAllowDevelopmentShaderCompile;
 }
 
+bool FMaterialResource::CheckInValidStateForCompilation(class FMaterialCompiler* Compiler) const
+{
+	return Material && Material->CheckInValidStateForCompilation(Compiler);
+}
+
 void FMaterial::ReleaseShaderMap()
 {
 	UE_CLOG(IsOwnerBeginDestroyed(), LogMaterial, Error, TEXT("ReleaseShaderMap called on FMaterial %s, owner is BeginDestroyed"), *GetDebugName());
@@ -4367,7 +4372,8 @@ bool UMaterialInterface::IsPropertyActive(EMaterialProperty InProperty)const
 }
 
 #if WITH_EDITOR
-int32 UMaterialInterface::CompilePropertyEx( class FMaterialCompiler* Compiler, const FGuid& AttributeID )
+
+int32 UMaterialInterface::CompilePropertyEx(FMaterialCompiler* Compiler, const FGuid& AttributeID )
 {
 	return INDEX_NONE;
 }
diff --git a/Engine/Source/Runtime/Engine/Public/MaterialCachedData.h b/Engine/Source/Runtime/Engine/Public/MaterialCachedData.h
--- a/Engine/Source/Runtime/Engine/Public/MaterialCachedData.h
+++ b/Engine/Source/Runtime/Engine/Public/MaterialCachedData.h
@@ -190,12 +190,9 @@ struct FMaterialCachedExpressionData
 	ENGINE_API void UpdateForCachedHLSLTree(const FMaterialCachedHLSLTree& CachedTree, const FStaticParameterSet* StaticParameters);
 	void Validate();
 
-	void AddParameter(const FMaterialParameterInfo& ParameterInfo, const FMaterialParameterMetadata& ParameterMeta, UObject*& OutReferencedTexture);
-	inline void AddParameter(const FMaterialParameterInfo& ParameterInfo, const FMaterialParameterMetadata& ParameterMeta)
-	{
-		UObject* UnusedReferencedTexture = nullptr;
-		AddParameter(ParameterInfo, ParameterMeta, UnusedReferencedTexture);
-	}
+	/** Adds a parameter. If this returns false, a parameter with identical name has already been added but it was set to a different value. */
+	bool AddParameter(const FMaterialParameterInfo& ParameterInfo, const FMaterialParameterMetadata& ParameterMeta, UObject*& OutReferencedTexture);
+	
 #endif // WITH_EDITOR
 
 #if WITH_EDITORONLY_DATA
@@ -338,5 +335,10 @@ struct FMaterialCachedExpressionData
 	/** Each bit corresponds to EMaterialProperty connection status. */
 	UPROPERTY()
 	uint32 PropertyConnectedBitmask = 0;
+
+#if WITH_EDITOR
+	/** Array of errors reporting a parameter being set multiple times to distinct values. */
+	TArray<TPair<TObjectPtr<class UMaterialExpression>, FName>> DuplicateParameterErrors;
+#endif
 };
 
diff --git a/Engine/Source/Runtime/Engine/Public/MaterialShared.h b/Engine/Source/Runtime/Engine/Public/MaterialShared.h
--- a/Engine/Source/Runtime/Engine/Public/MaterialShared.h
+++ b/Engine/Source/Runtime/Engine/Public/MaterialShared.h
@@ -2209,6 +2209,12 @@ protected:
 	uint32 GetGameThreadCompilingShaderMapId() const { return GameThreadCompilingShaderMapId; }
 #endif // WITH_EDITOR
 
+	/**
+	 * Checks that no pre-compilation errors have been detected and if so it reports them using specified compiler.
+	 * @return whether no errors occurred.
+	 */
+	ENGINE_API virtual bool CheckInValidStateForCompilation(class FMaterialCompiler* Compiler) const = 0;
+
 	/**
 	 * Entry point for compiling a specific material property.  This must call SetMaterialProperty. 
 	 * @param OverrideShaderFrequency SF_NumFrequencies to not override
@@ -2635,6 +2641,8 @@ public:
 	inline void SetMaterial(UMaterial* InMaterial) { Material = InMaterial; }
 	inline void SetMaterialInstance(UMaterialInstance* InMaterialInstance) { MaterialInstance = InMaterialInstance; }
 
+	ENGINE_API virtual bool CheckInValidStateForCompilation(class FMaterialCompiler* Compiler) const override;
+
 protected:
 	UMaterial* Material;
 	UMaterialInstance* MaterialInstance;
diff --git a/Engine/Source/Runtime/Engine/Public/MaterialTypes.h b/Engine/Source/Runtime/Engine/Public/MaterialTypes.h
--- a/Engine/Source/Runtime/Engine/Public/MaterialTypes.h
+++ b/Engine/Source/Runtime/Engine/Public/MaterialTypes.h
@@ -274,6 +274,11 @@ struct FParameterChannelNames
 
 	UPROPERTY(EditAnywhere, Category = MaterialExpressionVectorParameter)
 	FText A;
+
+	bool operator==(const FParameterChannelNames& Other) const
+	{
+		return R.EqualTo(Other.R) && G.EqualTo(Other.G) && B.EqualTo(Other.B) && A.EqualTo(Other.A);
+	}
 };
 
 USTRUCT()
@@ -295,6 +300,11 @@ struct FStaticComponentMaskValue
 
 	UPROPERTY()
 	bool A = false;
+
+	bool operator==(const FStaticComponentMaskValue& Other) const
+	{
+		return R == Other.R && G == Other.G && B == Other.B && A == Other.A;
+	}
 };
 
 struct FMaterialTextureValue
