Fix crash when deleting an animation blueprint

Removed check() and replaced with an if(). It is legitimate when deleting an anim BP for the cached properties to be invalid and the constants to no longer be found.

diff --git a/Engine/Source/Runtime/Engine/Private/Animation/AnimBlueprintGeneratedClass.cpp b/Engine/Source/Runtime/Engine/Private/Animation/AnimBlueprintGeneratedClass.cpp
--- a/Engine/Source/Runtime/Engine/Private/Animation/AnimBlueprintGeneratedClass.cpp
+++ b/Engine/Source/Runtime/Engine/Private/Animation/AnimBlueprintGeneratedClass.cpp
@@ -415,49 +415,24 @@ void UAnimBlueprintGeneratedClass::ForEachSubsystem(TFunctionRef<EAnimSubsystemE
 
 void UAnimBlueprintGeneratedClass::ForEachSubsystem(UObject* InObject, TFunctionRef<EAnimSubsystemEnumeration(const FAnimSubsystemInstanceContext&)> InFunction) const
 {
-#if DO_CHECK && WITH_EDITOR
-	// Tracking for UE-152551
-	if(ConstantSubsystemProperties.Num() != MutableSubsystemProperties.Num())
+	if(ConstantSubsystemProperties.Num() == MutableSubsystemProperties.Num())
 	{
-		UE_LOG(LogAnimation, Error, TEXT("Mismatched constant/mutable subsystems in class %s (ConstantSubsystemProperties == %d, MutableSubsystemProperties == %d)"), *GetName(), ConstantSubsystemProperties.Num(), MutableSubsystemProperties.Num());
-		UE_LOG(LogAnimation, Error, TEXT("Class Info:\n"));
-		UE_LOG(LogAnimation, Error, TEXT("  Parent: %s (%s)\n"), *GetSuperClass()->GetName(), GetSuperClass()->HasAnyClassFlags(CLASS_Native) ? TEXT("Native") : TEXT("Blueprint"));
-		UE_LOG(LogAnimation, Error, TEXT("  Sparse class data struct (%s):"), (GetSparseClassDataStruct() != nullptr) ? *GetSparseClassDataStruct()->GetName() : TEXT("None"));
-		if(GetSparseClassDataStruct() != nullptr)
+		for(int32 SubsystemIndex = 0; SubsystemIndex < ConstantSubsystemProperties.Num(); ++SubsystemIndex)
 		{
-			for(TFieldIterator<FProperty> It(GetSparseClassDataStruct()); It; ++It)
+			FStructProperty* ConstantSubsystemProperty = ConstantSubsystemProperties[SubsystemIndex];
+			FStructProperty* MutableSubsystemProperty = MutableSubsystemProperties[SubsystemIndex];
+			
+			checkSlow(ConstantSubsystemProperty->GetOwner<UStruct>() && GetSparseClassDataStruct()->IsChildOf(ConstantSubsystemProperty->GetOwner<UStruct>()));
+			const FAnimSubsystem& Subsystem = *ConstantSubsystemProperty->ContainerPtrToValuePtr<FAnimSubsystem>(GetConstantNodeData());
+
+			checkSlow(MutableSubsystemProperty->GetOwner<UClass>() && InObject->GetClass()->IsChildOf(MutableSubsystemProperty->GetOwner<UClass>()));
+			FAnimSubsystemInstance& SubsystemInstance = *MutableSubsystemProperty->ContainerPtrToValuePtr<FAnimSubsystemInstance>(InObject);
+			
+			if(InFunction(FAnimSubsystemInstanceContext(Subsystem, ConstantSubsystemProperty->Struct, SubsystemInstance, MutableSubsystemProperty->Struct)) == EAnimSubsystemEnumeration::Stop)
 			{
-				const FProperty* Property = *It;
-				UE_LOG(LogAnimation, Error, TEXT("    %s (%s)"), *Property->GetName(), *Property->GetClass()->GetName()); 
+				break;
 			}
 		}
-
-		UE_LOG(LogAnimation, Error, TEXT("  Properties:"));
-		for (TFieldIterator<FProperty> It(this); It; ++It)
-		{
-			const FProperty* Property = *It;
-			UE_LOG(LogAnimation, Error, TEXT("    %s (%s)"), *Property->GetName(), *Property->GetClass()->GetName()); 
-		}
-	}
-#endif
-	
-	check(ConstantSubsystemProperties.Num() == MutableSubsystemProperties.Num());
-	
-	for(int32 SubsystemIndex = 0; SubsystemIndex < ConstantSubsystemProperties.Num(); ++SubsystemIndex)
-	{
-		FStructProperty* ConstantSubsystemProperty = ConstantSubsystemProperties[SubsystemIndex];
-		FStructProperty* MutableSubsystemProperty = MutableSubsystemProperties[SubsystemIndex];
-		
-		checkSlow(ConstantSubsystemProperty->GetOwner<UStruct>() && GetSparseClassDataStruct()->IsChildOf(ConstantSubsystemProperty->GetOwner<UStruct>()));
-		const FAnimSubsystem& Subsystem = *ConstantSubsystemProperty->ContainerPtrToValuePtr<FAnimSubsystem>(GetConstantNodeData());
-
-		checkSlow(MutableSubsystemProperty->GetOwner<UClass>() && InObject->GetClass()->IsChildOf(MutableSubsystemProperty->GetOwner<UClass>()));
-		FAnimSubsystemInstance& SubsystemInstance = *MutableSubsystemProperty->ContainerPtrToValuePtr<FAnimSubsystemInstance>(InObject);
-		
-		if(InFunction(FAnimSubsystemInstanceContext(Subsystem, ConstantSubsystemProperty->Struct, SubsystemInstance, MutableSubsystemProperty->Struct)) == EAnimSubsystemEnumeration::Stop)
-		{
-			break;
-		}
 	}
 }
 
