Do not release skeletal mesh rendering resources when there is pending streaming operations

diff --git a/Engine/Source/Runtime/Engine/Private/SkeletalMesh.cpp b/Engine/Source/Runtime/Engine/Private/SkeletalMesh.cpp
--- a/Engine/Source/Runtime/Engine/Private/SkeletalMesh.cpp
+++ b/Engine/Source/Runtime/Engine/Private/SkeletalMesh.cpp
@@ -1130,7 +1130,14 @@ bool USkeletalMesh::DoesMipDataExist(const int32 MipIndex) const
 
 bool USkeletalMesh::HasPendingRenderResourceInitialization() const
 {
-	return GetSkeletalMeshRenderData() && !GetSkeletalMeshRenderData()->bReadyForStreaming;
+	// Verify we're not compiling before accessing the renderdata to avoid forcing the compilation
+	// to finish during garbage collection. If we're still compiling, the render data has not
+	// yet been created, hence it is not possible we're actively streaming anything from it...
+
+	// Only check !bReadyForStreaming if the render data is initialized from FSkeletalMeshRenderData::InitResources(), 
+	// otherwise no render commands are pending and the state will never resolve.
+	// Note that bReadyForStreaming is set on the renderthread.
+	return !IsCompiling() && GetSkeletalMeshRenderData() && GetSkeletalMeshRenderData()->IsInitialized() && !GetSkeletalMeshRenderData()->bReadyForStreaming;
 }
 
 bool USkeletalMesh::StreamOut(int32 NewMipCount)
@@ -1500,6 +1507,12 @@ bool USkeletalMesh::IsReadyForFinishDestroy()
 		return false;
 	}
 
+	// Match BeginDestroy() by checking for HasPendingInitOrStreaming().
+	if (HasPendingInitOrStreaming())
+	{
+		return false;
+	}
+
 	ReleaseResources();
 
 	// see if we have hit the resource flush fence
