Blendspace fixes for:
* Showing a single invalid triangle
* Incorrect grid-based interpolation when outside the convex hull of samples
Also adds some code (disabled) to help debugging problems with triangulation if they come up again

diff --git a/Engine/Source/Editor/Persona/Private/SAnimationBlendSpaceGridWidget.cpp b/Engine/Source/Editor/Persona/Private/SAnimationBlendSpaceGridWidget.cpp
--- a/Engine/Source/Editor/Persona/Private/SAnimationBlendSpaceGridWidget.cpp
+++ b/Engine/Source/Editor/Persona/Private/SAnimationBlendSpaceGridWidget.cpp
@@ -507,6 +507,12 @@ int32 SBlendSpaceGridWidget::OnPaint(const FPaintArgs& Args, const FGeometry& Al
 	
 	PaintBackgroundAndGrid(AllottedGeometry, MyCullingRect, OutDrawElements, LayerId);
 
+#if 0
+	// Showing the sample-weights on the grid points is not useful to end users, but can be helpful when debugging
+	// the grid-based interpolation.
+	PaintGridSampleWeights(AllottedGeometry, MyCullingRect, OutDrawElements, LayerId);
+#endif
+
 	PaintTriangulation(AllottedGeometry, MyCullingRect, OutDrawElements, LayerId);
 	PaintSampleKeys(AllottedGeometry, MyCullingRect, OutDrawElements, LayerId);
 
@@ -820,6 +826,67 @@ void SBlendSpaceGridWidget::PaintAxisText(
 	DrawLayerId += 1;
 }
 
+//======================================================================================================================
+void SBlendSpaceGridWidget::PaintGridSampleWeights(
+	const FGeometry& AllottedGeometry,
+	const FSlateRect& MyCullingRect,
+	FSlateWindowElementList& OutDrawElements,
+	int32& DrawLayerId) const
+{
+	const UBlendSpace* BlendSpace = BlendSpaceBase.Get();
+	if (!BlendSpace)
+	{
+		return;
+	}
+	if (!BlendSpace->bInterpolateUsingGrid)
+	{
+		return;
+	}
+
+	const TSharedRef< FSlateFontMeasure > FontMeasure = FSlateApplication::Get().GetRenderer()->GetFontMeasureService();
+
+	const TArray<FEditorElement>& GridSamples = BlendSpace->GetGridSamples();
+	int32 NumGridSamples = GridSamples.Num();
+	const TArray<FBlendSample>& Samples = BlendSpace->GetBlendSamples();
+	for (int32 GridSampleIndex = 0 ; GridSampleIndex != NumGridSamples ; ++GridSampleIndex)
+	{
+		const FEditorElement& EditorElement = GridSamples[GridSampleIndex];
+
+		int32 TextOffset = 0;
+		for (int32 ElementIndex = 0 ; ElementIndex != 3 ; ++ElementIndex)
+		{
+			float SampleWeight = EditorElement.Weights[ElementIndex];
+			int32 SampleIndex = EditorElement.Indices[ElementIndex];
+			if (SampleWeight <= 0 || SampleIndex < 0)
+			{
+				continue;
+			}
+
+			const FBlendSample& Sample = Samples[SampleIndex];
+
+			const FText Name = FText::Format(LOCTEXT("SampleNameFormatWeight", "{0} ({1}) {2}"), 
+				GetSampleName(Sample, SampleIndex), FText::AsNumber(SampleIndex), SampleWeight);
+			const FVector2D TextSize = FontMeasure->Measure(Name, FontInfo);
+			const FVector2D Padding = FVector2D(12.0f, 4.0f);
+
+			FVector GridSamplePosition = BlendSpace->GetGridPosition(GridSampleIndex);
+
+			// Show the sample name/index/weight, going progressively up (because sample labels are
+			// below the grid points)
+			FVector2D GridPosition = SampleValueToScreenPosition(GridSamplePosition);
+			GridPosition += FVector2D(-TextSize.X / 2, -2 * KeySize.Y);
+			GridPosition.Y -= TextSize.Y * TextOffset++;
+			GridPosition.X -= Padding.X / 2;
+			GridPosition.Y += Padding.Y / 2;
+
+			FSlateDrawElement::MakeText(
+				OutDrawElements, DrawLayerId + 2, AllottedGeometry.MakeChild(GridPosition + Padding / 2,
+					FVector2D(1.0f, 1.0f)).ToPaintGeometry(), Name, FontInfo, ESlateDrawEffect::None, FLinearColor::White);
+		}
+	}
+}
+
+
 //======================================================================================================================
 void SBlendSpaceGridWidget::PaintTriangulation(
 	const FGeometry&         AllottedGeometry, 
@@ -869,8 +936,8 @@ void SBlendSpaceGridWidget::PaintTriangulation(
 				SampleValueToScreenPosition(Samples[Triangle.SampleIndices[2]].SampleValue),
 			};
 
-			// If we just have one triangle it's OK for it to be degenerate
-			if (BlendSpaceData.Triangles.Num() > 1)
+			// Show invalid triangles even if there's only one triangle, because that probably
+			// happened when somebody failed to place the sample points in a proper line.
 			{
 				FVector2D NormalizedPositions[3] = {
 					SampleValueToNormalizedPosition(Samples[Triangle.SampleIndices[0]].SampleValue),
diff --git a/Engine/Source/Editor/Persona/Private/SAnimationBlendSpaceGridWidget.h b/Engine/Source/Editor/Persona/Private/SAnimationBlendSpaceGridWidget.h
--- a/Engine/Source/Editor/Persona/Private/SAnimationBlendSpaceGridWidget.h
+++ b/Engine/Source/Editor/Persona/Private/SAnimationBlendSpaceGridWidget.h
@@ -136,6 +136,7 @@ protected:
 	void PaintAxisText(const FGeometry& AllottedGeometry, const FSlateRect& MyCullingRect, FSlateWindowElementList& OutDrawElements, int32& DrawLayerId) const;
 	void PaintTriangulation(const FGeometry& AllottedGeometry, const FSlateRect& MyCullingRect, FSlateWindowElementList& OutDrawElements, int32& DrawLayerId) const;
 	void PaintAnimationNames(const FGeometry& AllottedGeometry, const FSlateRect& MyCullingRect, FSlateWindowElementList& OutDrawElements, int32& DrawLayerId) const;
+	void PaintGridSampleWeights(const FGeometry& AllottedGeometry, const FSlateRect& MyCullingRect, FSlateWindowElementList& OutDrawElements, int32& DrawLayerId) const;
 
 	/** Validation for drag and drop operation, will populate InvalidOperationText and return false in case it is invalid */
 	const bool IsValidDragDropOperation(const FDragDropEvent& DragDropEvent, FText& InvalidOperationText);
diff --git a/Engine/Source/Runtime/Engine/Private/Animation/BlendSpaceHelpers.cpp b/Engine/Source/Runtime/Engine/Private/Animation/BlendSpaceHelpers.cpp
--- a/Engine/Source/Runtime/Engine/Private/Animation/BlendSpaceHelpers.cpp
+++ b/Engine/Source/Runtime/Engine/Private/Animation/BlendSpaceHelpers.cpp
@@ -6,11 +6,12 @@
 
 #define LOCTEXT_NAMESPACE "AnimationBlendSpaceHelpers"
 
-struct FSortByDistance
+// Stores an edge or triangle index, and distance to a point, in anticipation of sorting.
+struct FIndexAndDistance
 {
 	int32 Index;
 	double Distance;
-	FSortByDistance(int32 InIndex, double InDistance) 
+	FIndexAndDistance(int32 InIndex, double InDistance)
 		: Index(InIndex), Distance(InDistance) {}
 };
 
@@ -188,11 +189,11 @@ void FDelaunayTriangleGenerator::AdjustEdgeDirections(EPreferredTriangulationDir
 
 				Triangle0->Vertices[EdgeIndexNext0] = Triangle1->Vertices[EdgeIndexPrev1];
 				Triangle1->Vertices[EdgeIndexNext1] = Triangle0->Vertices[EdgeIndexPrev0];
+				Triangle0->UpdateCenter();
+				Triangle1->UpdateCenter();
 			}
 		}
-
 	}
-
 }
 
 void FDelaunayTriangleGenerator::SortSamples()
@@ -592,17 +593,17 @@ static FVector GetBaryCentric2D(const FVector2D& Point, const FVector2D& A, cons
 bool FBlendSpaceGrid::FindTriangleThisPointBelongsTo(const FVector2D& TestPoint, FVector& OutBarycentricCoords, FTriangle*& OutTriangle, const TArray<FTriangle*>& TriangleList) const
 {
 	// Calculate distance from point to triangle and sort the triangle list accordingly
-	TArray<FSortByDistance> SortedTriangles;
+	TArray<FIndexAndDistance> SortedTriangles;
 	SortedTriangles.AddUninitialized(TriangleList.Num());
 	for (int32 TriangleIndex=0; TriangleIndex<TriangleList.Num(); ++TriangleIndex)
 	{
 		SortedTriangles[TriangleIndex].Index = TriangleIndex;
 		SortedTriangles[TriangleIndex].Distance = TriangleList[TriangleIndex]->GetDistance(TestPoint);
 	}
-	SortedTriangles.Sort([](const FSortByDistance &A, const FSortByDistance &B) { return A.Distance < B.Distance; });
+	SortedTriangles.Sort([](const FIndexAndDistance &A, const FIndexAndDistance &B) { return A.Distance < B.Distance; });
 
 	// Now loop over the sorted triangles and test the barycentric coordinates with the point
-	for (const FSortByDistance& SortedTriangle : SortedTriangles)
+	for (const FIndexAndDistance& SortedTriangle : SortedTriangles)
 	{
 		FTriangle* Triangle = TriangleList[SortedTriangle.Index];
 
@@ -626,7 +627,7 @@ bool FBlendSpaceGrid::FindTriangleThisPointBelongsTo(const FVector2D& TestPoint,
 	return false;
 }
 
-static FVector2D ClosestPointOnSegment2D(const FVector2D& Point, const FVector2D& StartPoint, const FVector2D& EndPoint)
+static FVector2D ClosestPointOnSegment2D(const FVector2D& Point, const FVector2D& StartPoint, const FVector2D& EndPoint, double& T)
 {
 	const FVector2D Segment = EndPoint - StartPoint;
 	const FVector2D VectToPoint = Point - StartPoint;
@@ -635,6 +636,7 @@ static FVector2D ClosestPointOnSegment2D(const FVector2D& Point, const FVector2D
 	const double Dot1 = VectToPoint | Segment;
 	if (Dot1 <= 0)
 	{
+		T = 0.0;
 		return StartPoint;
 	}
 
@@ -642,11 +644,13 @@ static FVector2D ClosestPointOnSegment2D(const FVector2D& Point, const FVector2D
 	const double Dot2 = Segment | Segment;
 	if (Dot2 <= Dot1)
 	{
+		T = 1.0;
 		return EndPoint;
 	}
 
 	// Closest Point is within segment
-	return StartPoint + Segment * (Dot1 / Dot2);
+	T = Dot1 / Dot2;
+	return StartPoint + Segment * T;
 }
 
 
@@ -692,82 +696,42 @@ void FBlendSpaceGrid::GenerateGridElements(const TArray<FVertex>& SamplePoints,
 			}
 			else
 			{
-				TArray<FSortByDistance> SortedTriangles;
-				for (int32 TriangleIndex = 0; TriangleIndex < TriangleList.Num(); ++TriangleIndex )
+				// Work through all the edges and find the one with a point closest to this grid position.
+				int32 ClosestTriangleIndex = -1;
+				int32 ClosestEdgeIndex = -1;
+				double ClosestDistance = UE_DOUBLE_BIG_NUMBER;
+				double ClosestT = 0.0;
+
+				// Just walk through all the edges from all the triangles and find the closest ones
+				for (int32 TriangleIndex = 0; TriangleIndex < TriangleList.Num(); ++TriangleIndex)
 				{
-					// Check if points are collinear
 					const FTriangle* Triangle = TriangleList[TriangleIndex];
-					const FVector2D EdgeA = Triangle->Vertices[1]->Position - Triangle->Vertices[0]->Position;
-					const FVector2D EdgeB = Triangle->Vertices[2]->Position - Triangle->Vertices[0]->Position;
-					const double Result = EdgeA.X * EdgeB.Y - EdgeA.Y * EdgeB.X;
-					// Only add valid triangles
-					if (Result > 0.0f)
-					{ 
-						SortedTriangles.Add(FSortByDistance(TriangleIndex, Triangle->GetDistance(GridPointPosition)));
-					}					
-				}
-
-				if (SortedTriangles.Num())
-				{
-					SortedTriangles.Sort([](const FSortByDistance &A, const FSortByDistance &B) { return A.Distance < B.Distance; });
-					FTriangle* ClosestTriangle = TriangleList[SortedTriangles[0].Index];
-
-					// For the closest triangle, determine which of its edges is closest to the grid point
-					TArray<FSortByDistance> Edges;
-					TArray<FVector2D> PointsOnEdges;
-					for (int32 EdgeIndex = 0; EdgeIndex < 3; ++EdgeIndex)
+					for (int32 EdgeIndex = 0 ; EdgeIndex != 3 ; ++EdgeIndex)
 					{
+						double T;
 						const FVector2D ClosestPoint = ClosestPointOnSegment2D(
 							GridPointPosition,
-							ClosestTriangle->Edges[EdgeIndex].Vertices[0]->Position,
-							ClosestTriangle->Edges[EdgeIndex].Vertices[1]->Position);
-						Edges.Add(FSortByDistance(EdgeIndex, (ClosestPoint - GridPointPosition).SizeSquared()));
-						PointsOnEdges.Add(ClosestPoint);
-					}
-					Edges.Sort([](const FSortByDistance &A, const FSortByDistance &B) { return A.Distance < B.Distance; });
-					
-					// Calculate weighting using the closest edge points and the clamped grid position on the line
-					const FVector GridWeights = GetBaryCentric2D(PointsOnEdges[Edges[0].Index], ClosestTriangle->Vertices[0]->Position, ClosestTriangle->Vertices[1]->Position, ClosestTriangle->Vertices[2]->Position);
-					
-					for (int32 Index = 0; Index < 3; ++Index)
-					{
-						GridPoint.Weights[Index] = GridWeights[Index];
-						GridPoint.Indices[Index] = ClosestTriangle->Vertices[Index]->SampleIndex;
-					}
-				}
-				else
-				{
-					// This means that there is either one point, two points or collinear triangles on the grid
-					if (SamplePoints.Num() == 1)
-					{
-						// Just one, fill all grid points to the single sample
-						GridPoint.Weights[0] = 1.0f;
-						GridPoint.Indices[0] = 0;
-					}
-					else
-					{
-						// Two points or co-linear triangles, first find the two closest samples
-						TArray<FSortByDistance> SampleDistances;
-						for (int32 PointIndex = 0; PointIndex < SamplePoints.Num(); ++PointIndex)
-						{
-							const double DistanceFromSampleToPoint= (SamplePoints[PointIndex].Position - GridPointPosition).SizeSquared();
-							SampleDistances.Add(FSortByDistance(PointIndex, DistanceFromSampleToPoint));
-						}
-						SampleDistances.Sort([](const FSortByDistance &A, const FSortByDistance &B) { return A.Distance < B.Distance; });
-
-						// Find closest point on line between the two samples (clamping the grid position to the line, just like clamping to the triangle edges)
-						const FVertex Samples[2] = { SamplePoints[SampleDistances[0].Index], SamplePoints[SampleDistances[1].Index] };
-						const FVector2D ClosestPoint = ClosestPointOnSegment2D(GridPointPosition, Samples[0].Position, Samples[1].Position);
-						const double LineLength = (Samples[0].Position - Samples[1].Position).SizeSquared();
+							Triangle->Edges[EdgeIndex].Vertices[0]->Position,
+							Triangle->Edges[EdgeIndex].Vertices[1]->Position,
+							T);
+						const double Distance = (ClosestPoint - GridPointPosition).SizeSquared();
 
-						// Weight the samples according to the distance from the grid point on the line to the samples
-						for (int32 Index = 0; Index < 2; ++Index)
+						if (Distance < ClosestDistance)
 						{
-							GridPoint.Weights[Index] = (LineLength - (Samples[Index].Position - ClosestPoint).SizeSquared()) / LineLength;
-							GridPoint.Indices[Index] = Samples[Index].SampleIndex;
+							ClosestTriangleIndex = TriangleIndex;
+							ClosestEdgeIndex = EdgeIndex;
+							ClosestDistance = Distance;
+							ClosestT = T;
 						}
 					}
 				}
+
+				const FTriangle* Triangle = TriangleList[ClosestTriangleIndex];
+
+				GridPoint.Weights[0] = 1.0 - ClosestT;
+				GridPoint.Indices[0] = Triangle->Edges[ClosestEdgeIndex].Vertices[0]->SampleIndex;
+				GridPoint.Weights[1] = ClosestT;
+				GridPoint.Indices[1] = Triangle->Edges[ClosestEdgeIndex].Vertices[1]->SampleIndex;
 			}
 		}
 	}
diff --git a/Engine/Source/Runtime/Engine/Private/Animation/BlendSpaceHelpers.h b/Engine/Source/Runtime/Engine/Private/Animation/BlendSpaceHelpers.h
--- a/Engine/Source/Runtime/Engine/Private/Animation/BlendSpaceHelpers.h
+++ b/Engine/Source/Runtime/Engine/Private/Animation/BlendSpaceHelpers.h
@@ -93,6 +93,11 @@ struct FTriangle
 	// FEdges
 	FHalfEdge Edges[3];
 
+	void UpdateCenter() 
+	{
+		Center = (Vertices[0]->Position + Vertices[1]->Position + Vertices[2]->Position) / 3.0;
+	}
+
 	bool operator==( const FTriangle& Other ) const 
 	{
 		// if same position, it's same point
@@ -103,7 +108,7 @@ struct FTriangle
 	{
 		FMemory::Memcpy(Vertices, Copy.Vertices);
 		FMemory::Memcpy(Edges, Copy.Edges);
-		Center = Copy.Center;
+		UpdateCenter();
 
 		Vertices[0]->AddTriangle(this);
 		Vertices[1]->AddTriangle(this);
@@ -115,7 +120,7 @@ struct FTriangle
 		Vertices[0] = A;
 		Vertices[1] = B;
 		Vertices[2] = C;
-		Center = (A->Position + B->Position + C->Position)/3.0f;
+		UpdateCenter();
 
 		Vertices[0]->AddTriangle(this);
 		Vertices[1]->AddTriangle(this);
@@ -134,7 +139,7 @@ struct FTriangle
 		Vertices[0] = A;
 		Vertices[1] = A;
 		Vertices[2] = A;
-		Center = A->Position;
+		UpdateCenter();
 
 		Vertices[0]->AddTriangle(this);
 		Vertices[1]->AddTriangle(this);
@@ -151,7 +156,7 @@ struct FTriangle
 		Vertices[0] = A;
 		Vertices[1] = B;
 		Vertices[2] = B;
-		Center = (A->Position + B->Position) / 2.0f;
+		UpdateCenter();
 
 		Vertices[0]->AddTriangle(this);
 		Vertices[1]->AddTriangle(this);
