Fix a compilation error

diff --git a/Engine/Source/Runtime/Engine/Classes/Engine/SkinnedAsset.h b/Engine/Source/Runtime/Engine/Classes/Engine/SkinnedAsset.h
--- a/Engine/Source/Runtime/Engine/Classes/Engine/SkinnedAsset.h
+++ b/Engine/Source/Runtime/Engine/Classes/Engine/SkinnedAsset.h
@@ -242,11 +242,15 @@ public:
 	/** Helper function for resource tracking, construct a string using the skinned asset's path name and LOD index . */
 	static FString GetLODPathName(const USkinnedAsset* Mesh, int32 LODIndex);
 
-#if WITH_EDITOR
 	/** IInterface_AsyncCompilation begin*/
+#if WITH_EDITOR
 	virtual bool IsCompiling() const override;
+#else
+	FORCEINLINE bool IsCompiling() const { return false; }
+#endif
 	/** IInterface_AsyncCompilation end*/
 
+#if WITH_EDITOR
 	virtual FString BuildDerivedDataKey(const ITargetPlatform* TargetPlatform)
 	PURE_VIRTUAL(USkinnedAsset::BuildDerivedDataKey, return TEXT(""););
 
