Fixed out-of-bounds access when using retargeting

If a skeleton is edited (e.g. bones added) and the various animations that use it do not get their RetargetSourceAssetReferencePose updated as well, then the buffer can potentially be too small until the animation is updated. Adding a simple guard in the retargeting code to skip generating cached transforms for out-of-bounds bone indices. They are not needed anyways until the animation is updated/re-saved.

Additionally fixed a minor issue where the preview scene does not correctly respect the skeleton compatibility when choosing a preview mesh. It looks at the current preview mesh and only shows meshes that are compatible with that, rather than the 'current asset' (e.g. anim sequence, anim BP).

diff --git a/Engine/Source/Editor/Persona/Private/PreviewSceneCustomizations.cpp b/Engine/Source/Editor/Persona/Private/PreviewSceneCustomizations.cpp
--- a/Engine/Source/Editor/Persona/Private/PreviewSceneCustomizations.cpp
+++ b/Engine/Source/Editor/Persona/Private/PreviewSceneCustomizations.cpp
@@ -478,13 +478,12 @@ bool FPreviewSceneDescriptionCustomization::HandleShouldFilterAsset(const FAsset
 		return false;
 	}
 
-	const UPersonaPreviewSceneDescription* PersonaPreviewSceneDescription = PreviewScene.Pin()->GetPreviewSceneDescription();
-	if(!PersonaPreviewSceneDescription->PreviewMesh.IsValid())
+	if(!PersonaToolkit.IsValid())
 	{
 		return false;
 	}
-
-	const USkeleton* Skeleton = PersonaPreviewSceneDescription->PreviewMesh->GetSkeleton();
+	
+	const USkeleton* Skeleton = PersonaToolkit.Pin()->GetSkeleton();
 	const FString SkeletonTag = InAssetData.GetTagValueRef<FString>(InTag);
 	if (Skeleton && Skeleton->IsCompatibleForEditor(SkeletonTag))
 	{
diff --git a/Engine/Source/Runtime/Engine/Private/BoneContainer.cpp b/Engine/Source/Runtime/Engine/Private/BoneContainer.cpp
--- a/Engine/Source/Runtime/Engine/Private/BoneContainer.cpp
+++ b/Engine/Source/Runtime/Engine/Private/BoneContainer.cpp
@@ -458,40 +458,43 @@ const FRetargetSourceCachedData& FBoneContainer::GetRetargetSourceCachedData(con
 
 			if (SourceSkeletonBoneIndex != INDEX_NONE && AssetSkeleton->GetBoneTranslationRetargetingMode(SourceSkeletonBoneIndex) == EBoneTranslationRetargetingMode::OrientAndScale)
 			{
-				const FVector SourceSkelTrans = AuthoredOnRefSkeleton[SourceSkeletonBoneIndex].GetTranslation();
-				FVector TargetSkelTrans;
-				if (RefPoseOverride.IsValid())
+				if(AuthoredOnRefSkeleton.IsValidIndex(SourceSkeletonBoneIndex))
 				{
-					TargetSkelTrans = RefPoseOverride->RefBonePoses[BoneIndicesArray[CompactBoneIndex]].GetTranslation();
-				}
-				else
-				{
-					TargetSkelTrans = PlayingOnRefSkeleton[BoneIndicesArray[CompactBoneIndex]].GetTranslation();
-				}
-
-				// If translations are identical, we don't need to do any retargeting
-				if (!SourceSkelTrans.Equals(TargetSkelTrans, BONE_TRANS_RT_ORIENT_AND_SCALE_PRECISION))
-				{
-					const float SourceSkelTransLength = SourceSkelTrans.Size();
-					const float TargetSkelTransLength = TargetSkelTrans.Size();
-
-					// this only works on non zero vectors.
-					if (!FMath::IsNearlyZero(SourceSkelTransLength * TargetSkelTransLength))
+					const FVector SourceSkelTrans = AuthoredOnRefSkeleton[SourceSkeletonBoneIndex].GetTranslation();
+					FVector TargetSkelTrans;
+					if (RefPoseOverride.IsValid())
+					{
+						TargetSkelTrans = RefPoseOverride->RefBonePoses[BoneIndicesArray[CompactBoneIndex]].GetTranslation();
+					}
+					else
 					{
-						const FVector SourceSkelTransDir = SourceSkelTrans / SourceSkelTransLength;
-						const FVector TargetSkelTransDir = TargetSkelTrans / TargetSkelTransLength;
+						TargetSkelTrans = PlayingOnRefSkeleton[BoneIndicesArray[CompactBoneIndex]].GetTranslation();
+					}
 
-						const FQuat DeltaRotation = FQuat::FindBetweenNormals(SourceSkelTransDir, TargetSkelTransDir);
-						const float Scale = TargetSkelTransLength / SourceSkelTransLength;
-						const int32 OrientAndScaleIndex = RetargetSourceCachedData->OrientAndScaleData.Add(FOrientAndScaleRetargetingCachedData(DeltaRotation, Scale, SourceSkelTrans, TargetSkelTrans));
+					// If translations are identical, we don't need to do any retargeting
+					if (!SourceSkelTrans.Equals(TargetSkelTrans, BONE_TRANS_RT_ORIENT_AND_SCALE_PRECISION))
+					{
+						const float SourceSkelTransLength = SourceSkelTrans.Size();
+						const float TargetSkelTransLength = TargetSkelTrans.Size();
 
-						// initialize CompactPoseBoneIndex to OrientAndScale Index LUT on demand
-						if (RetargetSourceCachedData->CompactPoseIndexToOrientAndScaleIndex.Num() == 0)
+						// this only works on non zero vectors.
+						if (!FMath::IsNearlyZero(SourceSkelTransLength * TargetSkelTransLength))
 						{
-							RetargetSourceCachedData->CompactPoseIndexToOrientAndScaleIndex.Init(INDEX_NONE, CompactPoseNumBones);
-						}
+							const FVector SourceSkelTransDir = SourceSkelTrans / SourceSkelTransLength;
+							const FVector TargetSkelTransDir = TargetSkelTrans / TargetSkelTransLength;
+
+							const FQuat DeltaRotation = FQuat::FindBetweenNormals(SourceSkelTransDir, TargetSkelTransDir);
+							const float Scale = TargetSkelTransLength / SourceSkelTransLength;
+							const int32 OrientAndScaleIndex = RetargetSourceCachedData->OrientAndScaleData.Add(FOrientAndScaleRetargetingCachedData(DeltaRotation, Scale, SourceSkelTrans, TargetSkelTrans));
 
-						RetargetSourceCachedData->CompactPoseIndexToOrientAndScaleIndex[CompactBoneIndex] = OrientAndScaleIndex;
+							// initialize CompactPoseBoneIndex to OrientAndScale Index LUT on demand
+							if (RetargetSourceCachedData->CompactPoseIndexToOrientAndScaleIndex.Num() == 0)
+							{
+								RetargetSourceCachedData->CompactPoseIndexToOrientAndScaleIndex.Init(INDEX_NONE, CompactPoseNumBones);
+							}
+
+							RetargetSourceCachedData->CompactPoseIndexToOrientAndScaleIndex[CompactBoneIndex] = OrientAndScaleIndex;
+						}
 					}
 				}
 			}
