Modulation Destination IsActive() state switched logic
#rnx

diff --git a/Engine/Source/Runtime/Engine/Private/SoundModulationDestination.cpp b/Engine/Source/Runtime/Engine/Private/SoundModulationDestination.cpp
--- a/Engine/Source/Runtime/Engine/Private/SoundModulationDestination.cpp
+++ b/Engine/Source/Runtime/Engine/Private/SoundModulationDestination.cpp
@@ -2,7 +2,7 @@
 
 #include "Sound/SoundModulationDestination.h"
 
-#include "Algo/NoneOf.h"
+#include "Algo/AnyOf.h"
 #include "AudioDevice.h"
 #include "DSP/FloatArrayMath.h"
 
@@ -140,7 +140,7 @@ namespace Audio
 	bool FModulationDestination::IsActive()
 	{
 		FScopeLock Lock(&HandleCritSection);
-		return Algo::NoneOf(Handles, [](const FModulatorHandle& Handle) { return Handle.IsValid(); });
+		return Algo::AnyOf(Handles, [](const FModulatorHandle& Handle) { return Handle.IsValid(); });
 	}
 
 	bool FModulationDestination::ProcessControl(float InValueUnitBase, int32 InNumSamples)
