Allow Make nodes for TopLevelAssetPath to convert gracefully to the new Native Make function (assuming they were using a literal parameter)

diff --git a/Engine/Config/BaseEngine.ini b/Engine/Config/BaseEngine.ini
--- a/Engine/Config/BaseEngine.ini
+++ b/Engine/Config/BaseEngine.ini
@@ -1109,6 +1109,8 @@ DefaultWorldPartitionHLODLayer=/Engine/Maps/Templates/HLODs/HLODLayer_Instanced.
 +PropertyRedirects=(OldName="KismetMathLibrary.Conv_DoubleToLinearColor.InFloat", NewName="Conv_DoubleToLinearColor.InDouble")
 +PropertyRedirects=(OldName="KismetMathLibrary.Conv_DoubleToVector.InFloat", NewName="Conv_DoubleToVector.InDouble")
 
++PropertyRedirects=(OldName="KismetSystemLibrary.MakeTopLevelAssetPath.FullPathOrPackageName", NewName="KismetSystemLibrary.MakeTopLevelAssetPath.PackageName")
+
 [CoreUObject.Metadata]
 ; Note: UnrealHeaderTool should be rerun after making changes to MetadataRedirects to catch any keys specified in class headers
 +MetadataRedirects=(OldKey="K2Protected", NewKey="BlueprintProtected")
diff --git a/Engine/Source/Runtime/Engine/Classes/Kismet/KismetSystemLibrary.h b/Engine/Source/Runtime/Engine/Classes/Kismet/KismetSystemLibrary.h
--- a/Engine/Source/Runtime/Engine/Classes/Kismet/KismetSystemLibrary.h
+++ b/Engine/Source/Runtime/Engine/Classes/Kismet/KismetSystemLibrary.h
@@ -268,7 +268,7 @@ class ENGINE_API UKismetSystemLibrary : public UBlueprintFunctionLibrary
 
 	/** Builds a TopLevelAssetPath struct from single Path string or from PackageName and AssetName string. */
 	UFUNCTION(BlueprintPure, Category = "Utilities", meta = (Keywords = "construct build", NativeMakeFunc, BlueprintThreadSafe, BlueprintAutocast))
-	static FORCENOINLINE FTopLevelAssetPath MakeTopLevelAssetPath(const FString& FullPathOrPackageName, const FString& AssetName);
+	static FORCENOINLINE FTopLevelAssetPath MakeTopLevelAssetPath(UPARAM(DisplayName="FullPathOrPackageName") const FString& PackageName, const FString& AssetName);
 
 	/** Gets the path string out of a TopLevelAssetPath */
 	UFUNCTION(BlueprintPure, Category = "Utilities", meta = (NativeBreakFunc, BlueprintThreadSafe, BlueprintAutocast))
