ShaderGeneration: Fixes issues with ShaderModelID being incorrectly determined because of rounding rules in EncodeGBufferToMRT.

diff --git a/Engine/Source/Runtime/Engine/Private/ShaderCompiler/ShaderGenerationUtil.cpp b/Engine/Source/Runtime/Engine/Private/ShaderCompiler/ShaderGenerationUtil.cpp
--- a/Engine/Source/Runtime/Engine/Private/ShaderCompiler/ShaderGenerationUtil.cpp
+++ b/Engine/Source/Runtime/Engine/Private/ShaderCompiler/ShaderGenerationUtil.cpp
@@ -1055,8 +1055,7 @@ static FString CreateGBufferEncodeFunction(const FGBufferInfo& BufferInfo)
 						FString TargetName = FString::Printf(TEXT("MrtUint%d"), I);
 						int32 ChanBits = GetBufferNumBits(BufferInfo.Targets[I].TargetType, Chan);
 
-						// the 0.5 is to put it in the center of the int range, avoid roundoff error weirdness
-						CurrLine += FString::Printf(TEXT("(float(%s.%s) + .5f) / %u.0f"),
+						CurrLine += FString::Printf(TEXT("float(%s.%s) / %u.0f"),
 							TargetName.GetCharArray().GetData(),
 							Swizzles[Chan].GetCharArray().GetData(),
 							(1 << ChanBits) - 1);
