Deprecate CreateCompressedInfo path and GetRuntimeFormat in favour of IAudioInfoFactoryRegistry

diff --git a/Engine/Source/Editor/AudioEditor/Private/Factories/SoundFactory.cpp b/Engine/Source/Editor/AudioEditor/Private/Factories/SoundFactory.cpp
--- a/Engine/Source/Editor/AudioEditor/Private/Factories/SoundFactory.cpp
+++ b/Engine/Source/Editor/AudioEditor/Private/Factories/SoundFactory.cpp
@@ -236,11 +236,8 @@ UObject* USoundFactory::CreateObject
 			// to be auditioned in the editor properly.
 			if (!ExistingSound->GetResourceData())
 			{
-				if (FAudioDeviceHandle AudioDevice = GEngine->GetMainAudioDevice())
-				{
-					FName RuntimeFormat = AudioDevice->GetRuntimeFormat(ExistingSound);
-					ExistingSound->InitAudioResource(RuntimeFormat);
-				}
+				FName RuntimeFormat = ExistingSound->GetRuntimeFormat();
+				ExistingSound->InitAudioResource(RuntimeFormat);
 			}
 
 			if (ComponentsToRestart.Num() > 0)
diff --git a/Engine/Source/Runtime/AudioCodecEngine/Private/DecoderInputBackCompat.cpp b/Engine/Source/Runtime/AudioCodecEngine/Private/DecoderInputBackCompat.cpp
--- a/Engine/Source/Runtime/AudioCodecEngine/Private/DecoderInputBackCompat.cpp
+++ b/Engine/Source/Runtime/AudioCodecEngine/Private/DecoderInputBackCompat.cpp
@@ -14,39 +14,34 @@ namespace Audio
 		FFormatDescriptorSection* OutDescriptor /*= nullptr*/) const
 	{
 		if (!OldInfoObject.IsValid())
-		{
-			FAudioDeviceHandle Handle = FAudioDeviceManager::Get()->GetActiveAudioDevice();
+		{			
+			OldInfoObject.Reset(IAudioInfoFactoryRegistry::Get().Create(Wave->GetRuntimeFormat()));
+			audio_ensure(OldInfoObject.IsValid());
 
-			if (ensure(Handle))
+			FSoundQualityInfo Info;
+			if (Wave->IsStreaming())
 			{
-				OldInfoObject.Reset(Handle->CreateCompressedAudioInfo(Wave));
-				audio_ensure(OldInfoObject.IsValid());
-
-				FSoundQualityInfo Info;
-				if (Wave->IsStreaming())
+				if (!OldInfoObject->StreamCompressedInfo(Wave, &Info))
 				{
-					if (!OldInfoObject->StreamCompressedInfo(Wave, &Info))
-					{
-						return nullptr;
-					}
+					return nullptr;
 				}
-				else
+			}
+			else
+			{
+				if (!OldInfoObject->ReadCompressedInfo(Wave->GetResourceData(), Wave->GetResourceSize(), &Info))
 				{
-					if (!OldInfoObject->ReadCompressedInfo(Wave->GetResourceData(), Wave->GetResourceSize(), &Info))
-					{
-						return nullptr;
-					}
+					return nullptr;
 				}
+			}
 
-				Desc.NumChannels		= Info.NumChannels;
-				Desc.NumFramesPerSec	= Info.SampleRate;
-				Desc.NumFrames			= (uint32)((float)Info.Duration * Info.SampleRate);
-				Desc.NumBytesPerPacket	= ~0;
+			Desc.NumChannels		= Info.NumChannels;
+			Desc.NumFramesPerSec	= Info.SampleRate;
+			Desc.NumFrames			= (uint32)((float)Info.Duration * Info.SampleRate);
+			Desc.NumBytesPerPacket	= ~0;
 
-				Desc.CodecName			= FBackCompatCodec::GetDetailsStatic().Name;
-				Desc.CodecFamilyName	= FBackCompatCodec::GetDetailsStatic().FamilyName;
-				Desc.CodecVersion		= FBackCompatCodec::GetDetailsStatic().Version;		
-			}
+			Desc.CodecName			= FBackCompatCodec::GetDetailsStatic().Name;
+			Desc.CodecFamilyName	= FBackCompatCodec::GetDetailsStatic().FamilyName;
+			Desc.CodecVersion		= FBackCompatCodec::GetDetailsStatic().Version;		
 		}
 
 		if( OutDescriptor )
diff --git a/Engine/Source/Runtime/AudioMixer/Private/AudioMixerBuffer.cpp b/Engine/Source/Runtime/AudioMixer/Private/AudioMixerBuffer.cpp
--- a/Engine/Source/Runtime/AudioMixer/Private/AudioMixerBuffer.cpp
+++ b/Engine/Source/Runtime/AudioMixer/Private/AudioMixerBuffer.cpp
@@ -336,7 +336,8 @@ namespace Audio
 		FMixerBuffer* Buffer = new FMixerBuffer(AudioDevice, InWave, EBufferType::Streaming);
 
 		FSoundQualityInfo QualityInfo = { 0 };
-		Buffer->DecompressionState = AudioDevice->CreateCompressedAudioInfo(InWave);
+		
+		Buffer->DecompressionState = IAudioInfoFactoryRegistry::Get().Create(InWave->GetRuntimeFormat());
 
 		// Get the header information of our compressed format
 		if (Buffer->DecompressionState && Buffer->DecompressionState->StreamCompressedInfo(InWave, &QualityInfo))
@@ -384,12 +385,14 @@ namespace Audio
 		// Create a new buffer for real-time sounds
 		FMixerBuffer* Buffer = new FMixerBuffer(AudioDevice, InWave, EBufferType::PCMRealTime);
 
+		FName FormatName = InWave->GetRuntimeFormat();
 		if (InWave->GetResourceData() == nullptr)
 		{
-			InWave->InitAudioResource(AudioDevice->GetRuntimeFormat(InWave));
+			InWave->InitAudioResource(FormatName);
 		}
 
-		Buffer->DecompressionState = AudioDevice->CreateCompressedAudioInfo(InWave);
+		Buffer->DecompressionState = IAudioInfoFactoryRegistry::Get().Create(FormatName);
+		check(Buffer->DecompressionState);
 
 		if (Buffer->DecompressionState)
 		{
diff --git a/Engine/Source/Runtime/AudioMixerCore/Public/AudioMixer.h b/Engine/Source/Runtime/AudioMixerCore/Public/AudioMixer.h
--- a/Engine/Source/Runtime/AudioMixerCore/Public/AudioMixer.h
+++ b/Engine/Source/Runtime/AudioMixerCore/Public/AudioMixer.h
@@ -519,9 +519,11 @@ namespace Audio
 		virtual bool IsNonRealtime() const { return false; }
 
 		/** Returns the FName version of a given sound wave's format. Override to provide a runtime format (codec) that is platform specific. */
+		UE_DEPRECATED(5.2, "GetRuntimeFormat is now deprecated. Please use USoundWave::GetRuntimeFormat()")
 		virtual FName GetRuntimeFormat(const USoundWave* InSoundWave) const { return FName(); }
 
 		/** Creates a Compressed audio info class suitable for decompressing this SoundWave. OVerride to create platform-specific decoders. */
+		UE_DEPRECATED(5.2, "CreateCompressedAudioInfo is now deprecated. Please use IAudioInfoFactory::Create()")
 		virtual ICompressedAudioInfo* CreateCompressedAudioInfo(const FName& InRuntimeFormat) const { return nullptr; };
 
 		/** Return any optional device name defined in platform configuratio. */
diff --git a/Engine/Source/Runtime/Engine/Private/AudioDecompress.cpp b/Engine/Source/Runtime/Engine/Private/AudioDecompress.cpp
--- a/Engine/Source/Runtime/Engine/Private/AudioDecompress.cpp
+++ b/Engine/Source/Runtime/Engine/Private/AudioDecompress.cpp
@@ -702,15 +702,9 @@ FAsyncAudioDecompressWorker::FAsyncAudioDecompressWorker(USoundWave* InWave, int
 	, NumPrecacheFrames(InPrecacheBufferNumFrames)
 {
 	check(NumPrecacheFrames > 0);
-	if (InAudioDevice)
-	{
-		AudioInfo = InAudioDevice->CreateCompressedAudioInfo(Wave);
-	}
-	else if (GEngine && GEngine->GetMainAudioDevice())
-	{
-		// alternatively, we could expose InWave::InternalProxy via public getter to avoid the proxy creation
-		AudioInfo = GEngine->GetMainAudioDevice()->CreateCompressedAudioInfo(InWave->CreateSoundWaveProxy());
-	}
+
+
+	AudioInfo = IAudioInfoFactoryRegistry::Get().Create(InWave->GetRuntimeFormat());
 }
 
 void FAsyncAudioDecompressWorker::DoWork()
diff --git a/Engine/Source/Runtime/Engine/Private/AudioDevice.cpp b/Engine/Source/Runtime/Engine/Private/AudioDevice.cpp
--- a/Engine/Source/Runtime/Engine/Private/AudioDevice.cpp
+++ b/Engine/Source/Runtime/Engine/Private/AudioDevice.cpp
@@ -1250,7 +1250,7 @@ void FAudioDevice::GetSoundClassInfo(TMap<FName, FAudioClassInfo>& AudioClassInf
 		}
 
 #if !WITH_EDITOR
-		AudioClassInfo->SizeResident += SoundWave->GetCompressedDataSize(GetRuntimeFormat(SoundWave));
+		AudioClassInfo->SizeResident += SoundWave->GetCompressedDataSize(SoundWave->GetRuntimeFormat());
 		AudioClassInfo->NumResident++;
 #else
 		switch(SoundWave->DecompressionType)
@@ -1262,14 +1262,14 @@ void FAudioDevice::GetSoundClassInfo(TMap<FName, FAudioClassInfo>& AudioClassInf
 			break;
 
 		case DTYPE_RealTime:
-			AudioClassInfo->SizeRealTime += SoundWave->GetCompressedDataSize(GetRuntimeFormat(SoundWave));
+			AudioClassInfo->SizeRealTime += SoundWave->GetCompressedDataSize(SoundWave->GetRuntimeFormat());
 			AudioClassInfo->NumRealTime++;
 			break;
 
 		case DTYPE_Streaming:
 			// Add these to real time count for now - eventually compressed data won't be loaded &
 			// might have a class info entry of their own
-			AudioClassInfo->SizeRealTime += SoundWave->GetCompressedDataSize(GetRuntimeFormat(SoundWave));
+			AudioClassInfo->SizeRealTime += SoundWave->GetCompressedDataSize(SoundWave->GetRuntimeFormat());
 			AudioClassInfo->NumRealTime++;
 			break;
 
@@ -6626,7 +6626,7 @@ void FAudioDevice::Precache(USoundWave* SoundWave, bool bSynchronous, bool bTrac
 		float CompressedDurationThreshold = GetCompressionDurationThreshold(SoundGroup);
 
 		static FName NAME_OGG(TEXT("OGG"));
-		SoundWave->bDecompressedFromOgg = GetRuntimeFormat(SoundWave) == NAME_OGG;
+		SoundWave->bDecompressedFromOgg = SoundWave->GetRuntimeFormat() == NAME_OGG;
 
 		// handle audio decompression
 		if (FPlatformProperties::SupportsAudioStreaming() && SoundWave->IsStreaming(nullptr))
@@ -6655,7 +6655,7 @@ void FAudioDevice::Precache(USoundWave* SoundWave, bool bSynchronous, bool bTrac
 		}
 
 		// Grab the compressed audio data
-		SoundWave->InitAudioResource(GetRuntimeFormat(SoundWave));
+		SoundWave->InitAudioResource(SoundWave->GetRuntimeFormat());
 
 		if (SoundWave->AudioDecompressor == nullptr && (SoundWave->DecompressionType == DTYPE_Native || SoundWave->DecompressionType == DTYPE_RealTime))
 		{
diff --git a/Engine/Source/Runtime/Engine/Private/SoundWave.cpp b/Engine/Source/Runtime/Engine/Private/SoundWave.cpp
--- a/Engine/Source/Runtime/Engine/Private/SoundWave.cpp
+++ b/Engine/Source/Runtime/Engine/Private/SoundWave.cpp
@@ -782,7 +782,7 @@ void USoundWave::GetResourceSizeEx(FResourceSizeEx& CumulativeResourceSize)
 
 			if (!FPlatformProperties::SupportsAudioStreaming() || !IsStreaming(nullptr))
 			{
-				CumulativeResourceSize.AddDedicatedSystemMemoryBytes(GetCompressedDataSize(LocalAudioDevice->GetRuntimeFormat(this)));
+				CumulativeResourceSize.AddDedicatedSystemMemoryBytes(GetCompressedDataSize(GetRuntimeFormat()));
 			}
 		}
 	}
@@ -3854,7 +3854,7 @@ void USoundWave::UpdateAsset(bool bMarkDirty)
 			FAudioDevice* LocalAudioDevice = GEngine->GetMainAudioDeviceRaw();
 			if (LocalAudioDevice)
 			{
-				FName RuntimeFormat = LocalAudioDevice->GetRuntimeFormat(this);
+				FName RuntimeFormat = GetRuntimeFormat();
 
 				if (LoadingBehavior == ESoundWaveLoadingBehavior::ForceInline && !GetResourceData())
 				{
diff --git a/Engine/Source/Runtime/Engine/Public/AudioDecompress.h b/Engine/Source/Runtime/Engine/Public/AudioDecompress.h
--- a/Engine/Source/Runtime/Engine/Public/AudioDecompress.h
+++ b/Engine/Source/Runtime/Engine/Public/AudioDecompress.h
@@ -342,6 +342,16 @@ struct ENGINE_API IAudioInfoFactoryRegistry
 	virtual void Register(IAudioInfoFactory*, FName) = 0;
 	virtual void Unregister(IAudioInfoFactory*, FName) = 0;
 	virtual IAudioInfoFactory* Find(FName InName) const = 0;
+	
+	// Convenience helper.
+	ICompressedAudioInfo* Create(FName InName)
+	{
+		if (IAudioInfoFactory* Factory = Find(InName))
+		{
+			return Factory->Create();
+		}
+		return nullptr;
+	}
 };
 
 class FSimpleAudioInfoFactory : public IAudioInfoFactory
diff --git a/Engine/Source/Runtime/Engine/Public/AudioDevice.h b/Engine/Source/Runtime/Engine/Public/AudioDevice.h
--- a/Engine/Source/Runtime/Engine/Public/AudioDevice.h
+++ b/Engine/Source/Runtime/Engine/Public/AudioDevice.h
@@ -1146,6 +1146,7 @@ public:
 	 */
 	void DeactivateReverbEffect(FName TagName);
 
+	UE_DEPRECATED(5.2, "GetRuntimeFormat is now deprecated. Please call SoundWave::GetRuntimeFormat() instead")
 	virtual FName GetRuntimeFormat(const USoundWave* SoundWave) const = 0;
 
 	/** Whether this SoundWave has an associated info class to decompress it */
@@ -1164,7 +1165,10 @@ public:
 	}
 
 	/** Creates a Compressed audio info class suitable for decompressing this SoundWave */
+	UE_DEPRECATED(5.2, "CreateCompressedAudioInfo is now deprecated. Please use IAudioInfoFactory::Create() instead")
 	virtual ICompressedAudioInfo* CreateCompressedAudioInfo(const USoundWave* SoundWave) const { return nullptr; }
+	
+	UE_DEPRECATED(5.2, "CreateCompressedAudioInfo is now deprecated. Please use IAudioInfoFactory::Create() instead")
 	virtual ICompressedAudioInfo* CreateCompressedAudioInfo(const FSoundWaveProxyPtr& SoundWave) const { return nullptr; }
 
 	/**
