Save more time in PropagateMaterialChangesToPrimitives by marking primitives dirty instead of removing/readding all of them
  - Saves another 4.5s of stuttering at begin of PIE in CitySample

#rnx

diff --git a/Engine/Source/Runtime/Engine/Private/ShaderCompiler/ShaderCompiler.cpp b/Engine/Source/Runtime/Engine/Private/ShaderCompiler/ShaderCompiler.cpp
--- a/Engine/Source/Runtime/Engine/Private/ShaderCompiler/ShaderCompiler.cpp
+++ b/Engine/Source/Runtime/Engine/Private/ShaderCompiler/ShaderCompiler.cpp
@@ -4990,7 +4990,6 @@ void FShaderCompilingManager::PropagateMaterialChangesToPrimitives(TMap<TRefCoun
 
 	TSet<FSceneInterface*> ScenesToUpdate;
 	FObjectCacheContextScope ObjectCacheScope;
-	TIndirectArray<FComponentRecreateRenderStateContext> ComponentContexts;
 	{
 		TRACE_CPUPROFILER_EVENT_SCOPE(FShaderCompilingManager::PropagateMaterialChangesToPrimitives);
 
@@ -5003,10 +5002,7 @@ void FShaderCompilingManager::PropagateMaterialChangesToPrimitives(TMap<TRefCoun
 
 		for (UPrimitiveComponent* PrimitiveComponent : ObjectCacheScope.GetContext().GetPrimitivesAffectedByMaterials(UpdatedMaterials))
 		{
-			if (PrimitiveComponent->IsRenderStateCreated() && PrimitiveComponent->SceneProxy)
-			{
-				ComponentContexts.Add(new FComponentRecreateRenderStateContext(PrimitiveComponent, &ScenesToUpdate));
-			}
+			PrimitiveComponent->MarkRenderStateDirty();
 		}
 	}
 
@@ -5016,10 +5012,6 @@ void FShaderCompilingManager::PropagateMaterialChangesToPrimitives(TMap<TRefCoun
 	// a refcount of 2, so the FMaterial destructor will trigger a check failure because the refcount doesn't reach 0. Empty this map before
 	// recreating the render state to allow resources to be deleted cleanly.
 	MaterialsToUpdate.Empty();
-
-	UpdateAllPrimitiveSceneInfosForScenes(ScenesToUpdate);
-	ComponentContexts.Empty();
-	UpdateAllPrimitiveSceneInfosForScenes(ScenesToUpdate);
 }
 
 
