Don't rerun construction scripts when it's not required.

diff --git a/Engine/Source/Runtime/Engine/Private/Level.cpp b/Engine/Source/Runtime/Engine/Private/Level.cpp
--- a/Engine/Source/Runtime/Engine/Private/Level.cpp
+++ b/Engine/Source/Runtime/Engine/Private/Level.cpp
@@ -1756,6 +1756,7 @@ bool ULevel::IncrementalRunConstructionScripts(bool bProcessAllActors)
 	{
 		CurrentActorIndexForIncrementalUpdate = 0;
 		bCachedHasStaticMeshCompilationPending.Reset();
+		bHasRerunConstructionScripts = true;
 		return true;
 	}
 	return false;
