Trace.Screenshot can now be called using only the ShowUI parameter.

diff --git a/Engine/Source/Runtime/Engine/Private/ProfilingDebugging/TraceScreenshot.cpp b/Engine/Source/Runtime/Engine/Private/ProfilingDebugging/TraceScreenshot.cpp
--- a/Engine/Source/Runtime/Engine/Private/ProfilingDebugging/TraceScreenshot.cpp
+++ b/Engine/Source/Runtime/Engine/Private/ProfilingDebugging/TraceScreenshot.cpp
@@ -14,15 +14,44 @@ static void TraceScreenshotCommandCallback(const TArray<FString>& Args)
 	FString Name;
 	bool bShowUI = false;
 
-	if (Args.Num() > 0)
+	if (Args.Num() == 1)
 	{
-		Name = Args[0];
+		const FCoreTexts& CoreTexts = FCoreTexts::Get();
+
+		// Try to see if the parameter is a name or a bool value for bShowUI.
+		if (FCString::Stricmp(*Args[0], TEXT("true")) == 0 ||
+		    FCString::Stricmp(*Args[0], TEXT("on")) == 0 ||
+		    FCString::Stricmp(*Args[0], TEXT("yes")) == 0 ||
+			FCString::Stricmp(*Args[0], TEXT("1")) == 0 ||
+			FCString::Stricmp(*Args[0], *CoreTexts.True.ToString()) == 0 ||
+			FCString::Stricmp(*Args[0], *CoreTexts.Yes.ToString()) == 0)
+		{
+			bShowUI = true;
+		}
+		else if (FCString::Stricmp(*Args[0], TEXT("false")) == 0 ||
+				 FCString::Stricmp(*Args[0], TEXT("off")) == 0 ||
+				 FCString::Stricmp(*Args[0], TEXT("no")) == 0 ||
+				 FCString::Stricmp(*Args[0], TEXT("0")) == 0 ||
+				 FCString::Stricmp(*Args[0], *CoreTexts.False.ToString()) == 0 ||
+				 FCString::Stricmp(*Args[0], *CoreTexts.No.ToString()) == 0)
+		{
+			bShowUI = false;
+		}
+		else
+		{
+			Name = Args[0];
+		}
 	}
-
-	if (Args.Num() > 1)
+	else if (Args.Num() == 2)
 	{
+		Name = Args[0];
 		bShowUI = FCString::ToBool(*Args[1]);
 	}
+	else if (Args.Num() > 2)
+	{
+		UE_LOG(LogConsoleResponse, Warning, TEXT("Invalid arguments. Usage: Trace.Screenshot [Name] [bShowUI]"));
+		return;
+	}
 
 	FTraceScreenshot::RequestScreenshot(Name, bShowUI, LogConsoleResponse);
 }
