ShaderGen: Bump the GBufferGeneratorVersion to ensure the DDC key changes and shaders are forced to recompile.

diff --git a/Engine/Source/Runtime/Engine/Private/ShaderCompiler/ShaderGenerationUtil.cpp b/Engine/Source/Runtime/Engine/Private/ShaderCompiler/ShaderGenerationUtil.cpp
--- a/Engine/Source/Runtime/Engine/Private/ShaderCompiler/ShaderGenerationUtil.cpp
+++ b/Engine/Source/Runtime/Engine/Private/ShaderCompiler/ShaderGenerationUtil.cpp
@@ -263,7 +263,7 @@ void FShaderCompileUtilities::ApplyFetchEnvironment(FShaderCompilerDefines& SrcD
 }
 
 // if we change the logic, increment this number to force a DDC key change
-static const int32 GBufferGeneratorVersion = 4;
+static const int32 GBufferGeneratorVersion = 5;
 
 static FShaderGlobalDefines FetchShaderGlobalDefines(EShaderPlatform TargetPlatform, EGBufferLayout GBufferLayout)
 {
