ComputeFramework: Add support for fallback delegate to be called whenever EnqueueWork fails.
Failure can happen immediately or later in the pipe (for example when we try to fetch shaders).
Optimus: Use fallback for mesh deformers. Fallback is to simply reset the gpu passthrough vertex factory so that we see the bind pose. This is the same behavior as before, but now can happen at top or bottom of the pipe.

diff --git a/Engine/Plugins/Animation/DeformerGraph/Source/OptimusCore/Private/OptimusDeformerInstance.cpp b/Engine/Plugins/Animation/DeformerGraph/Source/OptimusCore/Private/OptimusDeformerInstance.cpp
--- a/Engine/Plugins/Animation/DeformerGraph/Source/OptimusCore/Private/OptimusDeformerInstance.cpp
+++ b/Engine/Plugins/Animation/DeformerGraph/Source/OptimusCore/Private/OptimusDeformerInstance.cpp
@@ -392,15 +392,12 @@ void UOptimusDeformerInstance::SetupFromDeformer(UOptimusDeformer* InDeformer)
 void UOptimusDeformerInstance::SetCanBeActive(bool bInCanBeActive)
 {
 	bCanBeActive = bInCanBeActive;
-
 }
 
 void UOptimusDeformerInstance::AllocateResources()
 {
-	
 }
 
-
 void UOptimusDeformerInstance::ReleaseResources()
 {
 	if (BufferPool)
@@ -413,29 +410,11 @@ void UOptimusDeformerInstance::ReleaseResources()
 	}
 }
 
-
-bool UOptimusDeformerInstance::IsActive() const
-{
-	if (!bCanBeActive)
-	{
-		return false;
-	}
-		
-	for (const FOptimusDeformerInstanceExecInfo& Info: ComputeGraphExecInfos)
-	{
-		if (!Info.ComputeGraphInstance.ValidateDataProviders(Info.ComputeGraph))
-		{
-			return false;
-		}
-	}
-	return !ComputeGraphExecInfos.IsEmpty();
-}
-
-void UOptimusDeformerInstance::EnqueueWork(FSceneInterface* InScene, EWorkLoad InWorkLoadType, EExectutionGroup InExecutionGroup, FName InOwnerName)
+void UOptimusDeformerInstance::EnqueueWork(FEnqueueWorkDesc const& InDesc)
 {
 	// Convert execution group enum to ComputeTaskExecutionGroup name.
 	FName ExecutionGroupName;
-	switch (InExecutionGroup)
+	switch (InDesc.ExecutionGroup)
 	{
 	case UMeshDeformerInstance::ExecutionGroup_Immediate:
 		ExecutionGroupName = ComputeTaskExecutionGroup::Immediate;
@@ -458,19 +437,30 @@ void UOptimusDeformerInstance::EnqueueWork(FSceneInterface* InScene, EWorkLoad I
 	
 	// Enqueue work.
 	bool bIsWorkEnqueued = false;
-	for (FOptimusDeformerInstanceExecInfo& Info: ComputeGraphExecInfos)
+	if (bCanBeActive)
 	{
-		if (Info.GraphType == EOptimusNodeGraphType::Update || GraphsToRun.Contains(Info.GraphName))
+		for (FOptimusDeformerInstanceExecInfo& Info: ComputeGraphExecInfos)
 		{
-			Info.ComputeGraphInstance.EnqueueWork(Info.ComputeGraph, InScene, ExecutionGroupName, InOwnerName);
-			bIsWorkEnqueued = true;
+			if (Info.GraphType == EOptimusNodeGraphType::Update || GraphsToRun.Contains(Info.GraphName))
+			{
+				bIsWorkEnqueued |= Info.ComputeGraphInstance.EnqueueWork(Info.ComputeGraph, InDesc.Scene, ExecutionGroupName, InDesc.OwnerName, InDesc.FallbackDelegate);
+			}
 		}
 	}
 
-	// If this was enqueued to the Immediate group then flush all work on that group now.
-	if (bIsWorkEnqueued && InExecutionGroup == UMeshDeformerInstance::ExecutionGroup_Immediate)
+	if (!bIsWorkEnqueued)
+	{
+		// If we failed to enqueue work then enqueue the fallback.
+		// todo: This might need enqueuing for EndOfFrame instead of immediate execution?
+		ENQUEUE_RENDER_COMMAND(ComputeFrameworkEnqueueFallback)([FallbackDelegate = InDesc.FallbackDelegate](FRHICommandListImmediate& RHICmdList)
+		{ 
+			FallbackDelegate.ExecuteIfBound();
+		});
+	}
+	else if (InDesc.ExecutionGroup == UMeshDeformerInstance::ExecutionGroup_Immediate)
 	{
-		ComputeFramework::FlushWork(InScene, ExecutionGroupName);
+		// If we succesfully enqueued to the Immediate group then flush all work on that group now.
+		ComputeFramework::FlushWork(InDesc.Scene, ExecutionGroupName);
 	}
 }
 
diff --git a/Engine/Plugins/Animation/DeformerGraph/Source/OptimusCore/Public/OptimusDeformerInstance.h b/Engine/Plugins/Animation/DeformerGraph/Source/OptimusCore/Public/OptimusDeformerInstance.h
--- a/Engine/Plugins/Animation/DeformerGraph/Source/OptimusCore/Public/OptimusDeformerInstance.h
+++ b/Engine/Plugins/Animation/DeformerGraph/Source/OptimusCore/Public/OptimusDeformerInstance.h
@@ -204,8 +204,7 @@ protected:
 	/** Implementation of UMeshDeformerInstance. */
 	void AllocateResources() override;
 	void ReleaseResources() override;
-	bool IsActive() const override;
-	void EnqueueWork(FSceneInterface* InScene, EWorkLoad InWorkLoadType, EExectutionGroup InExecutionGroup, FName InOwnerName) override;
+	void EnqueueWork(FEnqueueWorkDesc const& InDesc) override;
 
 private:
 	/** The Mesh Component that owns this Mesh Deformer Instance. */
diff --git a/Engine/Plugins/Runtime/ComputeFramework/Source/ComputeFramework/Private/ComputeFramework/ComputeGraphWorker.h b/Engine/Plugins/Runtime/ComputeFramework/Source/ComputeFramework/Private/ComputeFramework/ComputeGraphWorker.h
--- a/Engine/Plugins/Runtime/ComputeFramework/Source/ComputeFramework/Private/ComputeFramework/ComputeGraphWorker.h
+++ b/Engine/Plugins/Runtime/ComputeFramework/Source/ComputeFramework/Private/ComputeFramework/ComputeGraphWorker.h
@@ -24,7 +24,8 @@ public:
 		FName InExecutionGroupName,
 		FName InOwnerName,
 		FComputeGraphRenderProxy const* InGraphRenderProxy, 
-		TArray<FComputeDataProviderRenderProxy*> InDataProviderRenderProxies );
+		TArray<FComputeDataProviderRenderProxy*> InDataProviderRenderProxies,
+		FSimpleDelegate InFallbackDelegate);
 
 	/** Submit enqueued compute graph work. */
 	void SubmitWork(
@@ -42,6 +43,8 @@ private:
 		FComputeGraphRenderProxy const* GraphRenderProxy = nullptr;
 		/** Data provider render proxies. */
 		TArray<FComputeDataProviderRenderProxy*> DataProviderRenderProxies;
+		/** Render thread fallback logic for invocations that are invalid. */
+		FSimpleDelegate FallbackDelegate;
 		
 		~FGraphInvocation();
 	};
diff --git a/Engine/Plugins/Runtime/ComputeFramework/Source/ComputeFramework/Private/ComputeGraph.cpp b/Engine/Plugins/Runtime/ComputeFramework/Source/ComputeFramework/Private/ComputeGraph.cpp
--- a/Engine/Plugins/Runtime/ComputeFramework/Source/ComputeFramework/Private/ComputeGraph.cpp
+++ b/Engine/Plugins/Runtime/ComputeFramework/Source/ComputeFramework/Private/ComputeGraph.cpp
@@ -157,10 +157,6 @@ bool UComputeGraph::ValidateProviders(TArray< TObjectPtr<UComputeDataProvider> >
 		{
 			return false;
 		}
-		if (DataProviders[Index] != nullptr && !DataProviders[Index]->IsValid())
-		{
-			return false;
-		}
 	}
 	return true;
 }
@@ -169,7 +165,7 @@ void UComputeGraph::CreateDataProviders(int32 InBindingIndex, TObjectPtr<UObject
 {
 	InOutDataProviders.SetNumZeroed(DataInterfaces.Num());
 
-	if (ensure(Bindings.IsValidIndex(InBindingIndex)) && InBindingObject && InBindingObject.IsA(Bindings[InBindingIndex]))
+	if (ensure(Bindings.IsValidIndex(InBindingIndex)) && (!InBindingObject || InBindingObject.IsA(Bindings[InBindingIndex])))
 	{
 		for (int32 DataInterfaceIndex = 0; DataInterfaceIndex < DataInterfaces.Num(); ++DataInterfaceIndex)
 		{
diff --git a/Engine/Plugins/Runtime/ComputeFramework/Source/ComputeFramework/Private/ComputeGraphComponent.cpp b/Engine/Plugins/Runtime/ComputeFramework/Source/ComputeFramework/Private/ComputeGraphComponent.cpp
--- a/Engine/Plugins/Runtime/ComputeFramework/Source/ComputeFramework/Private/ComputeGraphComponent.cpp
+++ b/Engine/Plugins/Runtime/ComputeFramework/Source/ComputeFramework/Private/ComputeGraphComponent.cpp
@@ -24,33 +24,19 @@ UComputeGraphComponent::~UComputeGraphComponent()
 void UComputeGraphComponent::CreateDataProviders(int32 InBindingIndex, UObject* InBindingObject)
 {
  	ComputeGraphInstance.CreateDataProviders(ComputeGraph, InBindingIndex, InBindingObject);
-
-	// We only want to queue work after validating the new providers.
-	bValidProviders = false;
 }
 
 void UComputeGraphComponent::DestroyDataProviders()
 {
 	ComputeGraphInstance.DestroyDataProviders();
-	bValidProviders = false;
 }
 
 void UComputeGraphComponent::QueueExecute()
 {
-	if (ComputeGraph == nullptr)
+	if (ComputeGraph != nullptr)
 	{
-		// todo[CF]: We should have a default fallback for all cases where we can't submit work.
-		return;
+		MarkRenderDynamicDataDirty();
 	}
-
-	// Don't submit work if we don't have all of the expected bindings.
-	bValidProviders = bValidProviders || ComputeGraphInstance.ValidateDataProviders(ComputeGraph);
-	if (!bValidProviders)
-	{
-		return;
-	}
-
-	MarkRenderDynamicDataDirty();
 }
 
 void UComputeGraphComponent::TickComponent(float DeltaTime, enum ELevelTick TickType, FActorComponentTickFunction* ThisTickFunction)
@@ -64,5 +50,5 @@ void UComputeGraphComponent::SendRenderDynamicData_Concurrent()
 {
 	Super::SendRenderDynamicData_Concurrent();
 	
-	ComputeGraphInstance.EnqueueWork(ComputeGraph, GetScene(), ComputeTaskExecutionGroup::EndOfFrameUpdate, GetOwner()->GetFName());
+	ComputeGraphInstance.EnqueueWork(ComputeGraph, GetScene(), ComputeTaskExecutionGroup::EndOfFrameUpdate, GetOwner()->GetFName(), FSimpleDelegate());
 }
diff --git a/Engine/Plugins/Runtime/ComputeFramework/Source/ComputeFramework/Private/ComputeGraphInstance.cpp b/Engine/Plugins/Runtime/ComputeFramework/Source/ComputeFramework/Private/ComputeGraphInstance.cpp
--- a/Engine/Plugins/Runtime/ComputeFramework/Source/ComputeFramework/Private/ComputeGraphInstance.cpp
+++ b/Engine/Plugins/Runtime/ComputeFramework/Source/ComputeFramework/Private/ComputeGraphInstance.cpp
@@ -31,11 +31,10 @@ bool FComputeGraphInstance::ValidateDataProviders(UComputeGraph* InComputeGraph)
 	return InComputeGraph != nullptr && InComputeGraph->IsCompiled() && InComputeGraph->ValidateGraph() && InComputeGraph->ValidateProviders(DataProviders);
 }
 
-bool FComputeGraphInstance::EnqueueWork(UComputeGraph* InComputeGraph, FSceneInterface const* InScene, FName InExecutionGroupName, FName InOwnerName)
+bool FComputeGraphInstance::EnqueueWork(UComputeGraph* InComputeGraph, FSceneInterface const* InScene, FName InExecutionGroupName, FName InOwnerName, FSimpleDelegate InFallbackDelegate)
 {
 	if (InComputeGraph == nullptr || InScene == nullptr)
 	{
-		// todo[CF]: We should have a default fallback for all cases where we can't submit work.
 		return false;
 	}
 
@@ -44,18 +43,15 @@ bool FComputeGraphInstance::EnqueueWork(UComputeGraph* InComputeGraph, FSceneInt
 		return false;
 	}
 
-	// Lookup the compute worker associated with this scene.
-	FComputeGraphTaskWorker* ComputeGraphWorker = FComputeFrameworkModule::GetComputeSystem()->GetComputeWorker(InScene);
-	if (!ensure(ComputeGraphWorker))
+	// Don't submit work if we don't have all of the expected bindings.
+	if (!InComputeGraph->ValidateProviders(DataProviders))
 	{
 		return false;
 	}
 
-	// Don't submit work if we don't have all of the expected bindings.
-	// If we hit the ensure then something invalidated providers without calling CreateDataProviders().
-	// Those paths DO need fixing. We can remove the ensure() if we ever feel safe enough!
-	const bool bValidProviders = InComputeGraph->ValidateProviders(DataProviders);
-	if (!ensure(bValidProviders))
+	// Lookup the compute worker associated with this scene.
+	FComputeGraphTaskWorker* ComputeGraphWorker = FComputeFrameworkModule::GetComputeSystem()->GetComputeWorker(InScene);
+	if (!ensure(ComputeGraphWorker))
 	{
 		return false;
 	}
@@ -76,10 +72,10 @@ bool FComputeGraphInstance::EnqueueWork(UComputeGraph* InComputeGraph, FSceneInt
 	}
 
 	ENQUEUE_RENDER_COMMAND(ComputeFrameworkEnqueueExecutionCommand)(
-		[ComputeGraphWorker, InExecutionGroupName, InOwnerName, GraphRenderProxy, MovedDataProviderRenderProxies = MoveTemp(DataProviderRenderProxies)](FRHICommandListImmediate& RHICmdList)
+		[ComputeGraphWorker, InExecutionGroupName, InOwnerName, GraphRenderProxy, MovedDataProviderRenderProxies = MoveTemp(DataProviderRenderProxies), InFallbackDelegate](FRHICommandListImmediate& RHICmdList)
 		{
 			// Compute graph scheduler will take ownership of the provider proxies.
-			ComputeGraphWorker->Enqueue(InExecutionGroupName, InOwnerName, GraphRenderProxy, MovedDataProviderRenderProxies);
+			ComputeGraphWorker->Enqueue(InExecutionGroupName, InOwnerName, GraphRenderProxy, MovedDataProviderRenderProxies, InFallbackDelegate);
 		});
 
 	return true;
diff --git a/Engine/Plugins/Runtime/ComputeFramework/Source/ComputeFramework/Private/ComputeGraphWorker.cpp b/Engine/Plugins/Runtime/ComputeFramework/Source/ComputeFramework/Private/ComputeGraphWorker.cpp
--- a/Engine/Plugins/Runtime/ComputeFramework/Source/ComputeFramework/Private/ComputeGraphWorker.cpp
+++ b/Engine/Plugins/Runtime/ComputeFramework/Source/ComputeFramework/Private/ComputeGraphWorker.cpp
@@ -14,12 +14,13 @@
 
 DECLARE_GPU_STAT_NAMED(ComputeFramework_ExecuteBatches, TEXT("ComputeFramework::ExecuteBatches"));
 
-void FComputeGraphTaskWorker::Enqueue(FName InExecutionGroupName, FName InOwnerName, FComputeGraphRenderProxy const* InGraphRenderProxy, TArray<FComputeDataProviderRenderProxy*> InDataProviderRenderProxies)
+void FComputeGraphTaskWorker::Enqueue(FName InExecutionGroupName, FName InOwnerName, FComputeGraphRenderProxy const* InGraphRenderProxy, TArray<FComputeDataProviderRenderProxy*> InDataProviderRenderProxies, FSimpleDelegate InFallbackDelegate)
 {
 	FGraphInvocation& GraphInvocation = GraphInvocationsPerGroup.FindOrAdd(InExecutionGroupName).AddDefaulted_GetRef();
 	GraphInvocation.OwnerName = InOwnerName;
 	GraphInvocation.GraphRenderProxy = InGraphRenderProxy;
 	GraphInvocation.DataProviderRenderProxies = MoveTemp(InDataProviderRenderProxies);
+	GraphInvocation.FallbackDelegate = InFallbackDelegate;
 }
 
 void FComputeGraphTaskWorker::SubmitWork(FRDGBuilder& GraphBuilder, FName InExecutionGroupName, ERHIFeatureLevel::Type FeatureLevel)
@@ -90,9 +91,10 @@ void FComputeGraphTaskWorker::SubmitWork(FRDGBuilder& GraphBuilder, FName InExec
 				}
 			}
 
-			// If we can't run the graph for any reason, back out now.
+			// If we can't run the graph for any reason, back out now and apply fallback logic.
 			if (!bIsValid)
 			{
+				GraphInvocation.FallbackDelegate.ExecuteIfBound();
 				continue;
 			}
 
diff --git a/Engine/Plugins/Runtime/ComputeFramework/Source/ComputeFramework/Public/ComputeFramework/ComputeGraphComponent.h b/Engine/Plugins/Runtime/ComputeFramework/Source/ComputeFramework/Public/ComputeFramework/ComputeGraphComponent.h
--- a/Engine/Plugins/Runtime/ComputeFramework/Source/ComputeFramework/Public/ComputeFramework/ComputeGraphComponent.h
+++ b/Engine/Plugins/Runtime/ComputeFramework/Source/ComputeFramework/Public/ComputeFramework/ComputeGraphComponent.h
@@ -47,6 +47,4 @@ protected:
 private:
 	UPROPERTY()
 	FComputeGraphInstance ComputeGraphInstance;
-
-	bool bValidProviders = false;
 };
diff --git a/Engine/Plugins/Runtime/ComputeFramework/Source/ComputeFramework/Public/ComputeFramework/ComputeGraphInstance.h b/Engine/Plugins/Runtime/ComputeFramework/Source/ComputeFramework/Public/ComputeFramework/ComputeGraphInstance.h
--- a/Engine/Plugins/Runtime/ComputeFramework/Source/ComputeFramework/Public/ComputeFramework/ComputeGraphInstance.h
+++ b/Engine/Plugins/Runtime/ComputeFramework/Source/ComputeFramework/Public/ComputeFramework/ComputeGraphInstance.h
@@ -35,7 +35,7 @@ public:
 	TArray< TObjectPtr<UComputeDataProvider> >& GetDataProviders() { return DataProviders; }
 
 	/** Enqueue the ComputeGraph work. */
-	bool EnqueueWork(UComputeGraph* InComputeGraph, FSceneInterface const* InScene, FName InExecutionGroupName, FName InOwnerName);
+	bool EnqueueWork(UComputeGraph* InComputeGraph, FSceneInterface const* InScene, FName InExecutionGroupName, FName InOwnerName, FSimpleDelegate InFallbackDelegate);
 
 private:
 	/** The currently bound Data Provider objects. */
diff --git a/Engine/Plugins/Runtime/HairStrands/Source/HairStrandsCore/Private/GroomComponent.cpp b/Engine/Plugins/Runtime/HairStrands/Source/HairStrandsCore/Private/GroomComponent.cpp
--- a/Engine/Plugins/Runtime/HairStrands/Source/HairStrandsCore/Private/GroomComponent.cpp
+++ b/Engine/Plugins/Runtime/HairStrands/Source/HairStrandsCore/Private/GroomComponent.cpp
@@ -3577,14 +3577,11 @@ void UGroomComponent::SendRenderDynamicData_Concurrent()
 
 	if (MeshDeformerInstance != nullptr && GroomAsset != nullptr)
 	{
-		if (MeshDeformerInstance->IsActive())
-		{
-			MeshDeformerInstance->EnqueueWork(GetScene(), UMeshDeformerInstance::WorkLoad_Update, GroomAsset->GetFName());
-		}
-		else
-		{
-			// TODO: Reset so groom appears rest pose
-		}
+		UMeshDeformerInstance::FEnqueueWorkDesc Desc;
+		Desc.Scene = GetScene();
+		Desc.OwnerName = GroomAsset->GetFName();
+		// TODO: Provide a FEnqueueWorkDesc::FallbackDelegate so that groom appears rest pose if the Enqueue fails.
+		MeshDeformerInstance->EnqueueWork(Desc);
 	}
 }
 
diff --git a/Engine/Source/Runtime/Engine/Classes/Animation/MeshDeformerInstance.h b/Engine/Source/Runtime/Engine/Classes/Animation/MeshDeformerInstance.h
--- a/Engine/Source/Runtime/Engine/Classes/Animation/MeshDeformerInstance.h
+++ b/Engine/Source/Runtime/Engine/Classes/Animation/MeshDeformerInstance.h
@@ -29,7 +29,13 @@ class ENGINE_API UMeshDeformerInstance : public UObject
 	GENERATED_BODY()
 
 public:
-	/** Enumeration for workloads to enqueue. */
+	/** Called to allocate any persistent render resources */
+	virtual void AllocateResources() PURE_VIRTUAL(, );
+
+	/** Called when persistent render resources should be released */
+	virtual void ReleaseResources() PURE_VIRTUAL(, );
+
+	/** Enumeration for workloads to EnqueueWork. */
 	enum EWorkLoad
 	{
 		WorkLoad_Setup,
@@ -37,7 +43,7 @@ public:
 		WorkLoad_Update,
 	};
 
-	/** Enumeration for exectution groups to enqueue on. */
+	/** Enumeration for execution groups to EnqueueWork on. */
 	enum EExectutionGroup
 	{
 		ExecutionGroup_Default,
@@ -45,21 +51,18 @@ public:
 		ExecutionGroup_EndOfFrameUodate,
 	};
 
-	/** Called to allocate any persistent render resources */
-	virtual void AllocateResources() PURE_VIRTUAL(, );
-
-	/** Called when persistent render resources should be released */ 
-	virtual void ReleaseResources() PURE_VIRTUAL(, );
-
-	/** Get if mesh deformer is active (compiled and valid). */
-	virtual bool IsActive() const PURE_VIRTUAL(, return false;);
-	
-	/** Enqueue the mesh deformer workload on a scene. InOwnerName is used for debugging. */
-	virtual void EnqueueWork(FSceneInterface* InScene, EWorkLoad InWorkLoadType, EExectutionGroup InExecutionGroup, FName InOwnerName) PURE_VIRTUAL(, );
-
-	/** For backwards compatibility. To deprecate. */
-	void EnqueueWork(FSceneInterface* InScene, EWorkLoad InWorkLoadType, FName InOwnerName)
+	/** Structure of inputs to EnqueueWork. */
+	struct FEnqueueWorkDesc
 	{
-		EnqueueWork(InScene, InWorkLoadType, ExecutionGroup_Default, InOwnerName);
-	}
+		FSceneInterface* Scene = nullptr;
+		EWorkLoad WorkLoadType = WorkLoad_Update;
+		EExectutionGroup ExecutionGroup = ExecutionGroup_Default;
+		/** Name used for debugging and profiling markers. */
+		FName OwnerName;
+		/** Render thread delegate that will be executed if Enqueue fails at any stage. */
+		FSimpleDelegate FallbackDelegate;
+	};
+
+	/** Enqueue the mesh deformer workload on a scene. */
+	virtual void EnqueueWork(FEnqueueWorkDesc const& InDesc) PURE_VIRTUAL(, );
 };
diff --git a/Engine/Source/Runtime/Engine/Private/Components/SkinnedMeshComponent.cpp b/Engine/Source/Runtime/Engine/Private/Components/SkinnedMeshComponent.cpp
--- a/Engine/Source/Runtime/Engine/Private/Components/SkinnedMeshComponent.cpp
+++ b/Engine/Source/Runtime/Engine/Private/Components/SkinnedMeshComponent.cpp
@@ -967,11 +967,13 @@ void USkinnedMeshComponent::CreateRenderState_Concurrent(FRegisterComponentConte
  	if (MeshDeformerInstance)
  	{
 		MeshDeformerInstance->AllocateResources();
-		if (MeshDeformerInstance->IsActive())
-		{
-			// Enqueue immediate execution of work here to ensure that we have some deformer outputs written for the next frame.
-			MeshDeformerInstance->EnqueueWork(GetScene(), UMeshDeformerInstance::WorkLoad_Update, UMeshDeformerInstance::ExecutionGroup_Immediate, GetSkinnedAsset()->GetFName());
-		}
+		
+		// Enqueue immediate execution of work here to ensure that we have some deformer outputs written for the next frame.
+		UMeshDeformerInstance::FEnqueueWorkDesc Desc;
+		Desc.Scene = GetScene();
+		Desc.ExecutionGroup = UMeshDeformerInstance::ExecutionGroup_Immediate;
+		Desc.OwnerName = GetSkinnedAsset() != nullptr ? GetSkinnedAsset()->GetFName() : GetFName();
+		MeshDeformerInstance->EnqueueWork(Desc);
 	}
 }
 
@@ -1086,16 +1088,13 @@ void USkinnedMeshComponent::SendRenderDynamicData_Concurrent()
 
 		if (MeshDeformerInstance != nullptr && UseLOD <= GetMeshDeformerMaxLOD())
 		{
-			if (MeshDeformerInstance->IsActive())
-			{
-				MeshDeformerInstance->EnqueueWork(GetScene(), UMeshDeformerInstance::WorkLoad_Update, UMeshDeformerInstance::ExecutionGroup_Default, GetSkinnedAsset()->GetFName());
-			}
-			else
-			{
-				// Reset so mesh appears in bind pose.
-				// We could use a fallback mesh deformer instead?
-				FSkeletalMeshDeformerHelpers::ResetVertexFactoryBufferOverrides_GameThread(MeshObject, UseLOD);
-			}
+			UMeshDeformerInstance::FEnqueueWorkDesc Desc;
+			Desc.Scene = GetScene();
+			Desc.OwnerName = GetSkinnedAsset() != nullptr ? GetSkinnedAsset()->GetFName() : GetFName();
+			// Fallback is to reset the passthrough vertex factory if the deformer fails to run.
+			Desc.FallbackDelegate.BindLambda([MeshObjectPtr = MeshObject, UseLOD]() { FSkeletalMeshDeformerHelpers::ResetVertexFactoryBufferOverrides(MeshObjectPtr, UseLOD); });
+
+			MeshDeformerInstance->EnqueueWork(Desc);
 		}
 	}
 }
diff --git a/Engine/Source/Runtime/Engine/Private/SkeletalMeshDeformerHelpers.cpp b/Engine/Source/Runtime/Engine/Private/SkeletalMeshDeformerHelpers.cpp
--- a/Engine/Source/Runtime/Engine/Private/SkeletalMeshDeformerHelpers.cpp
+++ b/Engine/Source/Runtime/Engine/Private/SkeletalMeshDeformerHelpers.cpp
@@ -252,50 +252,49 @@ void FSkeletalMeshDeformerHelpers::UpdateVertexFactoryBufferOverrides(FSkeletalM
 	}
 }
 
-void FSkeletalMeshDeformerHelpers::ResetVertexFactoryBufferOverrides_GameThread(FSkeletalMeshObject* InMeshObject, int32 LODIndex)
+void FSkeletalMeshDeformerHelpers::ResetVertexFactoryBufferOverrides(FSkeletalMeshObject* InMeshObject, int32 LODIndex)
 {
+	check(IsInRenderingThread());
+
 	if (InMeshObject->IsCPUSkinned())
 	{
 		return;
 	}
 
-	ENQUEUE_RENDER_COMMAND(ResetSkinPassthroughVertexFactory)([InMeshObject, LODIndex](FRHICommandList& CmdList)
-	{
-		FSkeletalMeshObjectGPUSkin* MeshObjectGPU = static_cast<FSkeletalMeshObjectGPUSkin*>(InMeshObject);
-		FMeshDeformerGeometry& DeformerGeometry = MeshObjectGPU->GetDeformerGeometry(LODIndex);
+	FSkeletalMeshObjectGPUSkin* MeshObjectGPU = static_cast<FSkeletalMeshObjectGPUSkin*>(InMeshObject);
+	FMeshDeformerGeometry& DeformerGeometry = MeshObjectGPU->GetDeformerGeometry(LODIndex);
 
-		// This can be called per frame even when already reset. So early out if we don't need to do anything.
-		const bool bIsReset = DeformerGeometry.PositionUpdatedFrame == 0 && DeformerGeometry.TangentUpdatedFrame == 0 && DeformerGeometry.ColorUpdatedFrame == 0;
-		if (bIsReset)
-		{
-			return;
-		}
+	// This can be called per frame even when already reset. So early out if we don't need to do anything.
+	const bool bIsReset = DeformerGeometry.PositionUpdatedFrame == 0 && DeformerGeometry.TangentUpdatedFrame == 0 && DeformerGeometry.ColorUpdatedFrame == 0;
+	if (bIsReset)
+	{
+		return;
+	}
 
-		// Reset stored buffers.
-		DeformerGeometry.Reset();
+	// Reset stored buffers.
+	DeformerGeometry.Reset();
 
-		// Reset vertex factories.
-		const FSkeletalMeshObjectGPUSkin::FSkeletalMeshObjectLOD& LOD = MeshObjectGPU->LODs[LODIndex];
-		const int32 NumSections = InMeshObject->GetRenderSections(LODIndex).Num();
-		for (int32 SectionIndex = 0; SectionIndex < NumSections; ++SectionIndex)
-		{
-			FGPUBaseSkinVertexFactory const* BaseVertexFactory = MeshObjectGPU->GetBaseSkinVertexFactory(LODIndex, SectionIndex);
-			FGPUSkinPassthroughVertexFactory* TargetVertexFactory = LOD.GPUSkinVertexFactories.PassthroughVertexFactories[SectionIndex].Get();
-			TargetVertexFactory->ResetVertexAttributes();
-			FGPUSkinPassthroughVertexFactory::FDataType Data;
-			BaseVertexFactory->CopyDataTypeForLocalVertexFactory(Data);
-			TargetVertexFactory->SetData(Data);
-		}
+	// Reset vertex factories.
+	const FSkeletalMeshObjectGPUSkin::FSkeletalMeshObjectLOD& LOD = MeshObjectGPU->LODs[LODIndex];
+	const int32 NumSections = InMeshObject->GetRenderSections(LODIndex).Num();
+	for (int32 SectionIndex = 0; SectionIndex < NumSections; ++SectionIndex)
+	{
+		FGPUBaseSkinVertexFactory const* BaseVertexFactory = MeshObjectGPU->GetBaseSkinVertexFactory(LODIndex, SectionIndex);
+		FGPUSkinPassthroughVertexFactory* TargetVertexFactory = LOD.GPUSkinVertexFactories.PassthroughVertexFactories[SectionIndex].Get();
+		TargetVertexFactory->ResetVertexAttributes();
+		FGPUSkinPassthroughVertexFactory::FDataType Data;
+		BaseVertexFactory->CopyDataTypeForLocalVertexFactory(Data);
+		TargetVertexFactory->SetData(Data);
+	}
 
 #if RHI_RAYTRACING
-		// Reset ray tracing geometry.
-		FSkeletalMeshRenderData& SkelMeshRenderData = MeshObjectGPU->GetSkeletalMeshRenderData();
-		FSkeletalMeshLODRenderData& LODModel = SkelMeshRenderData.LODRenderData[LODIndex];
-		FBufferRHIRef VertexBuffer = LODModel.StaticVertexBuffers.PositionVertexBuffer.VertexBufferRHI;
-
-		TArray<FBufferRHIRef> VertexBuffers;
-		VertexBuffers.Init(VertexBuffer, NumSections);
-		MeshObjectGPU->UpdateRayTracingGeometry(LODModel, LODIndex, VertexBuffers);
+	// Reset ray tracing geometry.
+	FSkeletalMeshRenderData& SkelMeshRenderData = MeshObjectGPU->GetSkeletalMeshRenderData();
+	FSkeletalMeshLODRenderData& LODModel = SkelMeshRenderData.LODRenderData[LODIndex];
+	FBufferRHIRef VertexBuffer = LODModel.StaticVertexBuffers.PositionVertexBuffer.VertexBufferRHI;
+
+	TArray<FBufferRHIRef> VertexBuffers;
+	VertexBuffers.Init(VertexBuffer, NumSections);
+	MeshObjectGPU->UpdateRayTracingGeometry(LODModel, LODIndex, VertexBuffers);
 #endif // RHI_RAYTRACING
-	});
 }
diff --git a/Engine/Source/Runtime/Engine/Public/SkeletalMeshDeformerHelpers.h b/Engine/Source/Runtime/Engine/Public/SkeletalMeshDeformerHelpers.h
--- a/Engine/Source/Runtime/Engine/Public/SkeletalMeshDeformerHelpers.h
+++ b/Engine/Source/Runtime/Engine/Public/SkeletalMeshDeformerHelpers.h
@@ -82,9 +82,8 @@ public:
 	/** 
 	 * Release all of the the buffers that have been allocated through the AllocateVertexFactory*() functions.
 	 * Reset all of the MeshObject's passthrough vertex factories.  
-	 * Usually called on the game thread when we know that the mesh deformer is not able to update to force bind pose rendering.
 	 */
-	ENGINE_API static void ResetVertexFactoryBufferOverrides_GameThread(FSkeletalMeshObject* InMeshObject, int32 InLodIndex);
+	ENGINE_API static void ResetVertexFactoryBufferOverrides(FSkeletalMeshObject* InMeshObject, int32 InLodIndex);
 
 #pragma endregion SetInternals
 };
