Improvements to Modulation Routing BP API

diff --git a/Engine/Plugins/Runtime/AudioModulation/Source/AudioModulation/Private/AudioModulationStatics.cpp b/Engine/Plugins/Runtime/AudioModulation/Source/AudioModulation/Private/AudioModulationStatics.cpp
--- a/Engine/Plugins/Runtime/AudioModulation/Source/AudioModulation/Private/AudioModulationStatics.cpp
+++ b/Engine/Plugins/Runtime/AudioModulation/Source/AudioModulation/Private/AudioModulationStatics.cpp
@@ -482,6 +482,17 @@ float UAudioModulationStatics::GetModulatorValue(const UObject* WorldContextObje
 	return 1.0f;
 }
 
+TSet<USoundModulatorBase*> UAudioModulationStatics::GetModulatorsFromDestination(const FSoundModulationDestinationSettings& Destination)
+{
+	TSet<USoundModulatorBase*> Modulators;
+	for (const TObjectPtr<USoundModulatorBase>& Modulator : Destination.Modulators)
+	{
+		Modulators.Add(Modulator.Get());
+	}
+
+	return Modulators;
+}
+
 
 #undef LOCTEXT_NAMESPACE
 
diff --git a/Engine/Plugins/Runtime/AudioModulation/Source/AudioModulation/Public/AudioModulationStatics.h b/Engine/Plugins/Runtime/AudioModulation/Source/AudioModulation/Public/AudioModulationStatics.h
--- a/Engine/Plugins/Runtime/AudioModulation/Source/AudioModulation/Public/AudioModulationStatics.h
+++ b/Engine/Plugins/Runtime/AudioModulation/Source/AudioModulation/Public/AudioModulationStatics.h
@@ -320,4 +320,14 @@ public:
 		Keywords = "modulation generator bus")
 	)
 	static UPARAM(DisplayName = "Value") float GetModulatorValue(const UObject* WorldContextObject, USoundModulatorBase* Modulator);
+
+	/** Gets the list of modulators currently applied to a Modulation Destination.
+	* @param Destination - The Modulation Destination.
+	* @return Modulators - The set of Modulators.
+	*/
+	UFUNCTION(BlueprintPure, Category = "Audio|Modulation", DisplayName = "Get Modulators From Destination", meta = (
+		WorldContext = "WorldContextObject",
+		Keywords = "modulation generator bus")
+	)
+	static UPARAM(DisplayName = "Modulators") TSet<USoundModulatorBase*> GetModulatorsFromDestination(const FSoundModulationDestinationSettings& Destination);
 };
diff --git a/Engine/Source/Runtime/AudioMixer/Private/AudioMixerSource.cpp b/Engine/Source/Runtime/AudioMixer/Private/AudioMixerSource.cpp
--- a/Engine/Source/Runtime/AudioMixer/Private/AudioMixerSource.cpp
+++ b/Engine/Source/Runtime/AudioMixer/Private/AudioMixerSource.cpp
@@ -1016,6 +1016,8 @@ namespace Audio
 			AudioLink->OnUpdateWorldState(Params);
 		}
 
+		UpdateModulation();
+
 		UpdatePitch();
 
 		UpdateVolume();
@@ -1793,6 +1795,21 @@ namespace Audio
 		MixerSourceVoice->SetSourceBufferListener(WaveInstance->SourceBufferListener, WaveInstance->bShouldSourceBufferListenerZeroBuffer);
 	}
 
+	void FMixerSource::UpdateModulation()
+	{
+		check(WaveInstance);
+
+		FActiveSound* ActiveSound = WaveInstance->ActiveSound;
+		check(ActiveSound);
+
+		if (ActiveSound->bModulationRoutingUpdated)
+		{
+			MixerSourceVoice->SetModulationRouting(ActiveSound->ModulationRouting);
+		}
+
+		ActiveSound->bModulationRoutingUpdated = false;
+	}
+
 	void FMixerSource::UpdateSourceBusSends()
 	{
 		// 1) loop through all bus sends
diff --git a/Engine/Source/Runtime/AudioMixer/Private/AudioMixerSource.h b/Engine/Source/Runtime/AudioMixer/Private/AudioMixerSource.h
--- a/Engine/Source/Runtime/AudioMixer/Private/AudioMixerSource.h
+++ b/Engine/Source/Runtime/AudioMixer/Private/AudioMixerSource.h
@@ -79,6 +79,9 @@ namespace Audio
 		/** Updates and source effect on this voice. */
 		void UpdateEffects();
 
+		/** Updates the Modulation Routing settings on this voice. */
+		void UpdateModulation();
+
 		/** Updates source bus send levels based on game data. */
 		void UpdateSourceBusSends();
 
diff --git a/Engine/Source/Runtime/AudioMixer/Private/AudioMixerSourceManager.cpp b/Engine/Source/Runtime/AudioMixer/Private/AudioMixerSourceManager.cpp
--- a/Engine/Source/Runtime/AudioMixer/Private/AudioMixerSourceManager.cpp
+++ b/Engine/Source/Runtime/AudioMixer/Private/AudioMixerSourceManager.cpp
@@ -1646,6 +1646,45 @@ namespace Audio
 		}, AUDIO_MIXER_THREAD_COMMAND_STRING("SetModHPFFrequency()"));
 	}
 
+	void FMixerSourceManager::SetModulationRouting(const int32 SourceId, FSoundModulationDefaultSettings& ModulationSettings)
+	{
+		FSourceInfo& SourceInfo = SourceInfos[SourceId];
+
+		FModulationDestination VolumeMod;
+		VolumeMod.Init(MixerDevice->DeviceID, FName("Volume"), false /* bInIsBuffered */, true /* bInValueLinear */);
+		VolumeMod.UpdateModulators(ModulationSettings.VolumeModulationDestination.Modulators);
+
+		FModulationDestination PitchMod;
+		PitchMod.Init(MixerDevice->DeviceID, FName("Pitch"), false /* bInIsBuffered */);
+		PitchMod.UpdateModulators(ModulationSettings.PitchModulationDestination.Modulators);
+
+		FModulationDestination HighpassMod;
+		HighpassMod.Init(MixerDevice->DeviceID, FName("HPFCutoffFrequency"), false /* bInIsBuffered */);
+		HighpassMod.UpdateModulators(ModulationSettings.HighpassModulationDestination.Modulators);
+
+		FModulationDestination LowpassMod;
+		LowpassMod.Init(MixerDevice->DeviceID, FName("LPFCutoffFrequency"), false /* bInIsBuffered */);
+		LowpassMod.UpdateModulators(ModulationSettings.LowpassModulationDestination.Modulators);
+
+
+		AudioMixerThreadCommand([
+			this,
+				SourceId,
+				VolumeModulation = MoveTemp(VolumeMod),
+				HighpassModulation = MoveTemp(HighpassMod),
+				LowpassModulation = MoveTemp(LowpassMod),
+				PitchModulation = MoveTemp(PitchMod)
+		]() mutable
+			{
+				FSourceInfo& SourceInfo = SourceInfos[SourceId];
+
+				SourceInfo.VolumeModulation = MoveTemp(VolumeModulation);
+				SourceInfo.PitchModulation = MoveTemp(PitchModulation);
+				SourceInfo.LowpassModulation = MoveTemp(LowpassModulation);
+				SourceInfo.HighpassModulation = MoveTemp(HighpassModulation);
+			}, AUDIO_MIXER_THREAD_COMMAND_STRING("SetModulationRouting"));
+	}
+
 	void FMixerSourceManager::SetSourceBufferListener(const int32 SourceId, FSharedISourceBufferListenerPtr& InSourceBufferListener, bool InShouldSourceBufferListenerZeroBuffer)
 	{
 		AUDIO_MIXER_THREAD_DEBUG_UPDATE()
diff --git a/Engine/Source/Runtime/AudioMixer/Private/AudioMixerSourceManager.h b/Engine/Source/Runtime/AudioMixer/Private/AudioMixerSourceManager.h
--- a/Engine/Source/Runtime/AudioMixer/Private/AudioMixerSourceManager.h
+++ b/Engine/Source/Runtime/AudioMixer/Private/AudioMixerSourceManager.h
@@ -232,6 +232,8 @@ namespace Audio
 		void SetModVolume(const int32 SourceId, const float InModVolume);
 		void SetModLPFFrequency(const int32 SourceId, const float InModFrequency);
 		void SetModHPFFrequency(const int32 SourceId, const float InModFrequency);
+		
+		void SetModulationRouting(const int32 SourceId, FSoundModulationDefaultSettings& ModulationSettings);
 
 		void SetSourceBufferListener(const int32 SourceId, FSharedISourceBufferListenerPtr& InSourceBufferListener, bool InShouldSourceBufferListenerZeroBuffer);
 
diff --git a/Engine/Source/Runtime/AudioMixer/Private/AudioMixerSourceVoice.cpp b/Engine/Source/Runtime/AudioMixer/Private/AudioMixerSourceVoice.cpp
--- a/Engine/Source/Runtime/AudioMixer/Private/AudioMixerSourceVoice.cpp
+++ b/Engine/Source/Runtime/AudioMixer/Private/AudioMixerSourceVoice.cpp
@@ -196,6 +196,13 @@ namespace Audio
 		}
 	}
 
+	void FMixerSourceVoice::SetModulationRouting(FSoundModulationDefaultRoutingSettings& RoutingSettings)
+	{
+		AUDIO_MIXER_CHECK_GAME_THREAD(MixerDevice);
+
+		SourceManager->SetModulationRouting(SourceId, RoutingSettings);
+	}
+
 	void FMixerSourceVoice::SetSourceBufferListener(FSharedISourceBufferListenerPtr& InSourceBufferListener, bool InShouldSourceBufferListenerZeroBuffer)
 	{
 		AUDIO_MIXER_CHECK_GAME_THREAD(MixerDevice);
diff --git a/Engine/Source/Runtime/AudioMixer/Private/AudioMixerSourceVoice.h b/Engine/Source/Runtime/AudioMixer/Private/AudioMixerSourceVoice.h
--- a/Engine/Source/Runtime/AudioMixer/Private/AudioMixerSourceVoice.h
+++ b/Engine/Source/Runtime/AudioMixer/Private/AudioMixerSourceVoice.h
@@ -59,6 +59,8 @@ namespace Audio
 		// Sets the source voice's HPF filter modulation base frequency.
 		void SetModHPFFrequency(const float InFrequency);
 
+		void SetModulationRouting(FSoundModulationDefaultRoutingSettings& RoutingSettings);
+
 		// Set the source voice's SourceBufferListener and associated boolean.
 		void SetSourceBufferListener(FSharedISourceBufferListenerPtr& InSourceBufferListener, bool InShouldSourceBufferListenerZeroBuffer);
 
diff --git a/Engine/Source/Runtime/Engine/Classes/Components/AudioComponent.h b/Engine/Source/Runtime/Engine/Classes/Components/AudioComponent.h
--- a/Engine/Source/Runtime/Engine/Classes/Components/AudioComponent.h
+++ b/Engine/Source/Runtime/Engine/Classes/Components/AudioComponent.h
@@ -127,6 +127,24 @@ enum class EAudioFaderCurve : uint8
 	Count UMETA(Hidden)
 };
 
+UENUM(BlueprintType)
+enum class EModulationDestination : uint8
+{
+	/* Volume modulation */
+	Volume,
+
+	/* Pitch modulation */
+	Pitch,
+
+	/* Cutoff Frequency of a lowpass filter */
+	Lowpass,
+
+	/* Cutoff Frequency of a highpass filter */
+	Highpass,
+
+	
+	COUNT
+};
 
 /**
  * Legacy struct used for storing named parameter for a given AudioComponent.
@@ -713,6 +731,24 @@ public:
 	UFUNCTION(BlueprintCallable, Category = "Audio|Components|Audio", DisplayName="Get Cooked Amplitude Envelope Data For All Playing Sounds")
 	bool GetCookedEnvelopeDataForAllPlayingSounds(TArray<FSoundWaveEnvelopeDataPerSound>& OutEnvelopeData);
 
+	/**
+	* Sets the routing for one of the given Audio component's Modulation Destinations.
+	* The changes do not apply to any currently active sounds, but will apply to future sounds.
+	* @param Modulators The set of modulators to apply to the given destination on the component.
+	* @param Destination The destination to assign the modulators to.
+	* @param RoutingMethod The routing method to use for the given volume modulator.
+	*/
+	UFUNCTION(BlueprintCallable, Category = "Audio|Components|Audio", DisplayName = "Set Modulation Routing")
+	void SetModulationRouting(const TSet<USoundModulatorBase*>& Modulators, const EModulationDestination Destination, const EModulationRouting RoutingMethod = EModulationRouting::Inherit);
+
+	/**
+	* Gets the set of currently active modulators for a given Modulation Destination.
+	* @param Destination The Destination to retrieve the Modulators from.
+	* @return The set of of Modulators applied to this component for the given Destination.
+	*/
+	UFUNCTION(BlueprintPure, Category = "Audio|Components|Audio", DisplayName = "Get Modulators")
+	UPARAM(DisplayName = "Modulators") TSet<USoundModulatorBase*> GetModulators(const EModulationDestination Destination);
+
 	static void PlaybackCompleted(uint64 AudioComponentID, bool bFailedToStart);
 
 	bool GetDisableParameterUpdatesWhilePlaying() const override { return static_cast<bool>(bDisableParameterUpdatesWhilePlaying); }
diff --git a/Engine/Source/Runtime/Engine/Private/ActiveSound.cpp b/Engine/Source/Runtime/Engine/Private/ActiveSound.cpp
--- a/Engine/Source/Runtime/Engine/Private/ActiveSound.cpp
+++ b/Engine/Source/Runtime/Engine/Private/ActiveSound.cpp
@@ -431,6 +431,12 @@ void FActiveSound::ResetNewBusSends()
 	bHasNewBusSends = false;
 }
 
+void FActiveSound::SetNewModulationRouting(const FSoundModulationDefaultRoutingSettings& NewRouting)
+{
+	ModulationRouting = NewRouting;
+	bModulationRoutingUpdated = true;
+}
+
 void FActiveSound::Stop()
 {
 	if (AudioDevice)
diff --git a/Engine/Source/Runtime/Engine/Private/Components/AudioComponent.cpp b/Engine/Source/Runtime/Engine/Private/Components/AudioComponent.cpp
--- a/Engine/Source/Runtime/Engine/Private/Components/AudioComponent.cpp
+++ b/Engine/Source/Runtime/Engine/Private/Components/AudioComponent.cpp
@@ -1730,6 +1730,91 @@ bool UAudioComponent::GetCookedEnvelopeDataForAllPlayingSounds(TArray<FSoundWave
 	return bHadData;
 }
 
+void UAudioComponent::SetModulationRouting(const TSet<USoundModulatorBase*>& Modulators, const EModulationDestination Destination, const EModulationRouting RoutingMethod)
+{
+	FAudioDevice* AudioDevice = GetAudioDevice();
+	if (!AudioDevice)
+	{
+		return;
+	}
+
+	switch (Destination)
+	{
+	case EModulationDestination::Volume:
+		ModulationRouting.VolumeRouting = RoutingMethod;
+		ModulationRouting.VolumeModulationDestination.Modulators = Modulators;
+		break;
+	case EModulationDestination::Pitch:
+		ModulationRouting.PitchRouting = RoutingMethod;
+		ModulationRouting.PitchModulationDestination.Modulators = Modulators;
+		break;
+	case EModulationDestination::Lowpass:
+		ModulationRouting.LowpassRouting = RoutingMethod;
+		ModulationRouting.LowpassModulationDestination.Modulators = Modulators;
+		break;
+	case EModulationDestination::Highpass:
+		ModulationRouting.HighpassRouting = RoutingMethod;
+		ModulationRouting.HighpassModulationDestination.Modulators = Modulators;
+		break;
+	default:
+	{
+		static_assert(static_cast<int32>(EModulationDestination::COUNT) == 4, "Possible missing ELiteralType case coverage.");
+		ensureMsgf(false, TEXT("Failed to set input node default: Literal type not supported"));
+		return;
+	}
+	}
+
+	// Tell the active sounds on the component to use the new Modulation Routing
+	AudioDevice->SendCommandToActiveSounds(AudioComponentID, [NewRouting = ModulationRouting](FActiveSound& ActiveSound)
+		{
+			ActiveSound.SetNewModulationRouting(NewRouting);
+		});
+
+}
+
+TSet<USoundModulatorBase*> UAudioComponent::GetModulators(const EModulationDestination Destination)
+{
+	FAudioDevice* AudioDevice = GetAudioDevice();
+	if (!AudioDevice)
+	{
+		return TSet<USoundModulatorBase*>();
+	}
+	
+	const TSet<TObjectPtr<USoundModulatorBase>>* ModulatorSet = nullptr;
+
+	switch (Destination)
+	{
+	case EModulationDestination::Volume:
+		ModulatorSet = &ModulationRouting.VolumeModulationDestination.Modulators;
+		break;
+	case EModulationDestination::Pitch:
+		ModulatorSet = &ModulationRouting.PitchModulationDestination.Modulators;
+		break;
+	case EModulationDestination::Lowpass:
+		ModulatorSet = &ModulationRouting.LowpassModulationDestination.Modulators;
+		break;
+	case EModulationDestination::Highpass:
+		ModulatorSet = &ModulationRouting.HighpassModulationDestination.Modulators;
+		break;
+	default:
+	{
+		static_assert(static_cast<int32>(EModulationDestination::COUNT) == 4, "Possible missing ELiteralType case coverage.");
+		ensureMsgf(false, TEXT("Failed to set input node default: Literal type not supported"));
+		return TSet<USoundModulatorBase*>();
+	}
+	}
+
+	check(ModulatorSet);
+
+	TSet<USoundModulatorBase*> Modulators;
+	for (const TObjectPtr<USoundModulatorBase>& Modulator : *ModulatorSet)
+	{
+		Modulators.Add(Modulator.Get());
+	}
+
+	return Modulators;
+}
+
 void UAudioComponent::SetSourceEffectChain(USoundEffectSourcePresetChain* InSourceEffectChain)
 {
 	SourceEffectChain = InSourceEffectChain;
diff --git a/Engine/Source/Runtime/Engine/Public/ActiveSound.h b/Engine/Source/Runtime/Engine/Public/ActiveSound.h
--- a/Engine/Source/Runtime/Engine/Public/ActiveSound.h
+++ b/Engine/Source/Runtime/Engine/Public/ActiveSound.h
@@ -517,6 +517,8 @@ public:
 	uint8 bIsFirstAttenuationUpdate : 1;
 	uint8 bStartedWithinNonBinauralRadius : 1;
 
+	uint8 bModulationRoutingUpdated : 1;
+
 	uint8 UserIndex;
 
 	/** Type of fade out currently being applied */
@@ -692,6 +694,9 @@ public:
 	/** Resets internal data of new Source Bus Sends */
 	void ResetNewBusSends();
 
+	/* Gives new Modulation Routing settings to the Active Sound. */
+	void SetNewModulationRouting(const FSoundModulationDefaultRoutingSettings& NewRouting);
+
 	/* Determines which of the provided listeners is the closest to the sound */
 	int32 FindClosestListener( const TArray<struct FListener>& InListeners ) const;
 
