Restoring code lost after moving to a different file.

diff --git a/Engine/Source/Runtime/Engine/Private/LightSceneProxy.cpp b/Engine/Source/Runtime/Engine/Private/LightSceneProxy.cpp
--- a/Engine/Source/Runtime/Engine/Private/LightSceneProxy.cpp
+++ b/Engine/Source/Runtime/Engine/Private/LightSceneProxy.cpp
@@ -66,6 +66,7 @@ FLightSceneProxy::FLightSceneProxy(const ULightComponent* InLightComponent)
 	, ShadowAmount(1.0f)
 	, SamplesPerPixel(1)
 	, DeepShadowLayerDistribution(InLightComponent->DeepShadowLayerDistribution)
+	, IESAtlasId(~0u)
 #if ACTOR_HAS_LABELS
 	, OwnerNameOrLabel(InLightComponent->GetOwner() ? InLightComponent->GetOwner()->GetActorNameOrLabel() : InLightComponent->GetName())
 #endif
diff --git a/Engine/Source/Runtime/Engine/Public/LightSceneProxy.h b/Engine/Source/Runtime/Engine/Public/LightSceneProxy.h
--- a/Engine/Source/Runtime/Engine/Public/LightSceneProxy.h
+++ b/Engine/Source/Runtime/Engine/Public/LightSceneProxy.h
@@ -464,6 +464,9 @@ protected:
 	/** Deep shadow layer distribution. */
 	float DeepShadowLayerDistribution;
 
+	/** IES texture atlas id. */
+	uint32 IESAtlasId;
+
 	/**
 	 * Updates the light proxy's cached transforms.
 	 * @param InLightToWorld - The new light-to-world transform.
