Fixing SceneInterface compile errors after header dependencies changed.

diff --git a/Engine/Source/Runtime/Engine/Private/SceneInterface.cpp b/Engine/Source/Runtime/Engine/Private/SceneInterface.cpp
--- a/Engine/Source/Runtime/Engine/Private/SceneInterface.cpp
+++ b/Engine/Source/Runtime/Engine/Private/SceneInterface.cpp
@@ -22,6 +22,10 @@ void FSceneInterface::UpdateAllPrimitiveSceneInfos(FRHICommandListImmediate& RHI
 	GraphBuilder.Execute();
 }
 
+void FSceneInterface::ProcessAndRenderIlluminanceMeter(FRDGBuilder& GraphBuilder, TArrayView<FViewInfo> Views, FRDGTextureRef SceneColorTexture)
+{
+}
+
 EShaderPlatform FSceneInterface::GetShaderPlatform() const
 {
 	return GShaderPlatformForFeatureLevel[GetFeatureLevel()];
diff --git a/Engine/Source/Runtime/Engine/Public/SceneInterface.h b/Engine/Source/Runtime/Engine/Public/SceneInterface.h
--- a/Engine/Source/Runtime/Engine/Public/SceneInterface.h
+++ b/Engine/Source/Runtime/Engine/Public/SceneInterface.h
@@ -20,6 +20,7 @@ class AWorldSettings;
 class FArchive;
 class FFloat16Color;
 class FInstanceCullingManager;
+class FLightSceneProxy;
 class FMaterial;
 class FMaterialShaderMap;
 class FName;
@@ -27,6 +28,7 @@ class FOutputDevice;
 class FPrimitiveComponentId;
 class FPrimitiveSceneInfo;
 class FRDGBuilder;
+class FRDGExternalAccessQueue;
 class FRectLightSceneProxy;
 class FRenderResource;
 class FRenderTarget;
@@ -37,6 +39,7 @@ class FSceneViewStateInterface;
 class FSkyAtmosphereRenderSceneInfo;
 class FSkyAtmosphereSceneProxy;
 class FSkyLightSceneProxy;
+class FSparseVolumeTextureViewerSceneProxy;
 class FTexture;
 class FVertexFactory;
 class FViewInfo;
@@ -222,7 +225,7 @@ public:
 	virtual void AllocateAndCaptureFrameSkyEnvMap(FRDGBuilder& GraphBuilder, FSceneRenderer& SceneRenderer, FViewInfo& MainView, bool bShouldRenderSkyAtmosphere, bool bShouldRenderVolumetricCloud, FInstanceCullingManager& InstanceCullingManager, FRDGExternalAccessQueue& ExternalAccessQueue) {}
 	virtual void ValidateSkyLightRealTimeCapture(FRDGBuilder& GraphBuilder, const FViewInfo& View, FRDGTextureRef SceneColorTexture) {}
 
-	virtual void ProcessAndRenderIlluminanceMeter(FRDGBuilder& GraphBuilder, TArrayView<FViewInfo> Views, FRDGTextureRef SceneColorTexture) {}
+	ENGINE_API virtual void ProcessAndRenderIlluminanceMeter(FRDGBuilder& GraphBuilder, TArrayView<FViewInfo> Views, FRDGTextureRef SceneColorTexture);
 
 	virtual void AddPlanarReflection(class UPlanarReflectionComponent* Component) {}
 	virtual void RemovePlanarReflection(class UPlanarReflectionComponent* Component) {}
