World Partition: allow non-player-controller streaming sources when simulating in PIE.

diff --git a/Engine/Source/Runtime/Engine/Private/WorldPartition/WorldPartitionStreamingPolicy.cpp b/Engine/Source/Runtime/Engine/Private/WorldPartition/WorldPartitionStreamingPolicy.cpp
--- a/Engine/Source/Runtime/Engine/Private/WorldPartition/WorldPartitionStreamingPolicy.cpp
+++ b/Engine/Source/Runtime/Engine/Private/WorldPartition/WorldPartitionStreamingPolicy.cpp
@@ -130,6 +130,8 @@ void UWorldPartitionStreamingPolicy::UpdateStreamingSources()
 
 	if (!AWorldPartitionReplay::IsPlaybackEnabled(World) || !WorldPartition->Replay->GetReplayStreamingSources(StreamingSources))
 	{
+		bool bAllowPlayerControllerStreamingSources = true;
+
 #if WITH_EDITOR
 		if (UWorldPartition::IsSimulating())
 		{
@@ -138,8 +140,8 @@ void UWorldPartitionStreamingPolicy::UpdateStreamingSources()
 			const FRotator ViewRotation = WorldToLocal.TransformRotation(GCurrentLevelEditingViewportClient->GetViewRotation().Quaternion()).Rotator();
 			static const FName NAME_SIE(TEXT("SIE"));
 			StreamingSources.Add(FWorldPartitionStreamingSource(NAME_SIE, ViewLocation, ViewRotation, EStreamingSourceTargetState::Activated, /*bBlockOnSlowLoading=*/false, EStreamingSourcePriority::Default, false));
+			bAllowPlayerControllerStreamingSources = false;
 		}
-		else
 #endif
 		if (!bIsServer || bIsServerStreamingEnabled || AWorldPartitionReplay::IsRecordingEnabled(World))
 		{
@@ -149,15 +151,18 @@ void UWorldPartitionStreamingPolicy::UpdateStreamingSources()
 			TArray<FWorldPartitionStreamingSource> ProviderStreamingSources;
 			for (IWorldPartitionStreamingSourceProvider* StreamingSourceProvider : WorldPartitionSubsystem->GetStreamingSourceProviders())
 			{
-				ProviderStreamingSources.Reset();
-				if (StreamingSourceProvider->GetStreamingSources(ProviderStreamingSources))
+				if (bAllowPlayerControllerStreamingSources || !Cast<APlayerController>(StreamingSourceProvider->GetStreamingSourceOwner()))
 				{
-					for (FWorldPartitionStreamingSource& ProviderStreamingSource : ProviderStreamingSources)
+					ProviderStreamingSources.Reset();
+					if (StreamingSourceProvider->GetStreamingSources(ProviderStreamingSources))
 					{
-						// Transform to Local
-						ProviderStreamingSource.Location = WorldToLocal.TransformPosition(ProviderStreamingSource.Location);
-						ProviderStreamingSource.Rotation = WorldToLocal.TransformRotation(ProviderStreamingSource.Rotation.Quaternion()).Rotator();
-						StreamingSources.Add(MoveTemp(ProviderStreamingSource));
+						for (FWorldPartitionStreamingSource& ProviderStreamingSource : ProviderStreamingSources)
+						{
+							// Transform to Local
+							ProviderStreamingSource.Location = WorldToLocal.TransformPosition(ProviderStreamingSource.Location);
+							ProviderStreamingSource.Rotation = WorldToLocal.TransformRotation(ProviderStreamingSource.Rotation.Quaternion()).Rotator();
+							StreamingSources.Add(MoveTemp(ProviderStreamingSource));
+						}
 					}
 				}
 			}
