Added stateless handshake level NMT_Upgrade message.

diff --git a/Engine/Plugins/Online/OnlineSubsystemUtils/Source/OnlineSubsystemUtils/Private/OnlineBeaconClient.cpp b/Engine/Plugins/Online/OnlineSubsystemUtils/Source/OnlineSubsystemUtils/Private/OnlineBeaconClient.cpp
--- a/Engine/Plugins/Online/OnlineSubsystemUtils/Source/OnlineSubsystemUtils/Private/OnlineBeaconClient.cpp
+++ b/Engine/Plugins/Online/OnlineSubsystemUtils/Source/OnlineSubsystemUtils/Private/OnlineBeaconClient.cpp
@@ -502,22 +502,7 @@ void AOnlineBeaconClient::NotifyControlMessage(UNetConnection* Connection, uint8
 
 				if (FNetControlMessage<NMT_Upgrade>::Receive(Bunch, RemoteNetworkVersion, RemoteNetworkFeatures))
 				{
-					TStringBuilder<128> RemoteFeaturesDescription;
-					FNetworkVersion::DescribeNetworkRuntimeFeaturesBitset(RemoteNetworkFeatures, RemoteFeaturesDescription);
-
-					TStringBuilder<128> LocalFeaturesDescription;
-					FNetworkVersion::DescribeNetworkRuntimeFeaturesBitset(NetDriver->GetNetworkRuntimeFeatures(), LocalFeaturesDescription);
-
-					UE_LOG(LogBeacon, Error, TEXT("Beacon is incompatible with the local version of the game: RemoteNetworkVersion=%u, RemoteNetworkFeatures=%s vs LocalNetworkVersion=%u, LocalNetworkFeatures=%s"), 
-						RemoteNetworkVersion, RemoteFeaturesDescription.ToString(),
-						FNetworkVersion::GetLocalNetworkVersion(), LocalFeaturesDescription.ToString()
-					);
-				
-					// Upgrade
-					const FString ConnectionError = NSLOCTEXT("Engine", "ClientOutdated",
-						"The match you are trying to join is running an incompatible version of the game.  Please try upgrading your game version.").ToString();
-
-					GEngine->BroadcastNetworkFailure(GetWorld(), NetDriver, ENetworkFailure::OutdatedClient, ConnectionError);
+					Connection->HandleReceiveNetUpgrade(RemoteNetworkVersion, RemoteNetworkFeatures);
 				}
 
 				break;
diff --git a/Engine/Source/Runtime/Engine/Classes/Engine/NetConnection.h b/Engine/Source/Runtime/Engine/Classes/Engine/NetConnection.h
--- a/Engine/Source/Runtime/Engine/Classes/Engine/NetConnection.h
+++ b/Engine/Source/Runtime/Engine/Classes/Engine/NetConnection.h
@@ -54,6 +54,7 @@ class ULevelStreaming;
 struct FEncryptionKeyResponse;
 class PacketHandler;
 class FRPCDoSDetection;
+enum class EEngineNetworkRuntimeFeatures : uint16;
 
 namespace UE::Net
 {
@@ -96,6 +97,16 @@ enum EConnectionState
 };
 ENGINE_API const TCHAR* LexToString(const EConnectionState Value);
 
+namespace UE::Net
+{
+	/** The source of a net upgrade message */
+	enum class ENetUpgradeSource : uint8
+	{
+		ControlChannel,			// The upgrade message came from the control channel
+		StatelessHandshake		// The upgrade message came from the stateless handshake
+	};
+}
+
 
 /** If this connection is from a client, this is the current login state of this connection/login attempt */
 namespace EClientLoginState
@@ -1760,6 +1771,16 @@ public:
 	 */
 	ENGINE_API void HandleReceiveCloseReason(const FString& CloseReasonList);
 
+	/**
+	 * Handles receiving NMT_Upgrade messages (including at stateless handshake level)
+	 *
+	 * @param RemoteNetworkVersion		The net version of the remote side
+	 * @param RemoteNetworkFeatures		The net runtime features of the remote side
+	 * @param NetUpgradeSource			The source of the net upgrade message
+	 */
+	ENGINE_API void HandleReceiveNetUpgrade(uint32 RemoteNetworkVersion, EEngineNetworkRuntimeFeatures RemoteNetworkFeatures,
+											UE::Net::ENetUpgradeSource NetUpgradeSource=UE::Net::ENetUpgradeSource::ControlChannel);
+
 private:
 	/**
 	 * Attempts to recover from a NetConnection error, and closes the connection if that fails
diff --git a/Engine/Source/Runtime/Engine/Private/NetConnection.cpp b/Engine/Source/Runtime/Engine/Private/NetConnection.cpp
--- a/Engine/Source/Runtime/Engine/Private/NetConnection.cpp
+++ b/Engine/Source/Runtime/Engine/Private/NetConnection.cpp
@@ -604,6 +604,8 @@ void UNetConnection::InitConnection(UNetDriver* InDriver, EConnectionState InSta
 
 void UNetConnection::InitHandler()
 {
+	using namespace UE::Net;
+
 	LLM_SCOPE_BYTAG(NetConnection);
 
 	check(!Handler.IsValid());
@@ -643,7 +645,18 @@ void UNetConnection::InitHandler()
 
 			if (StatelessConnectComponent.IsValid())
 			{
-				StatelessConnectComponent.Pin()->SetDriver(Driver);
+				StatelessConnectHandlerComponent* CurComponent = StatelessConnectComponent.Pin().Get();
+				
+				CurComponent->SetDriver(Driver);
+
+				CurComponent->SetHandshakeFailureCallback([this](FStatelessHandshakeFailureInfo HandshakeFailureInfo)
+					{
+						if (HandshakeFailureInfo.FailureReason == EHandshakeFailureReason::WrongVersion)
+						{
+							this->HandleReceiveNetUpgrade(HandshakeFailureInfo.RemoteNetworkVersion, HandshakeFailureInfo.RemoteNetworkFeatures,
+															ENetUpgradeSource::StatelessHandshake);
+						}
+					});
 			}
 
 
@@ -1119,6 +1132,36 @@ void UNetConnection::HandleReceiveCloseReason(const FString& CloseReasonList)
 	}
 }
 
+void UNetConnection::HandleReceiveNetUpgrade(uint32 RemoteNetworkVersion, EEngineNetworkRuntimeFeatures RemoteNetworkFeatures,
+												UE::Net::ENetUpgradeSource NetUpgradeSource/*=UE::Net::ENetUpgradeSource::ControlChannel*/)
+{
+	TStringBuilder<128> RemoteFeaturesDescription;
+	TStringBuilder<128> LocalFeaturesDescription;
+
+	FNetworkVersion::DescribeNetworkRuntimeFeaturesBitset(RemoteNetworkFeatures, RemoteFeaturesDescription);
+
+	if (Driver != nullptr)
+	{
+		FNetworkVersion::DescribeNetworkRuntimeFeaturesBitset(Driver->GetNetworkRuntimeFeatures(), LocalFeaturesDescription);
+	}
+
+	UE_LOG(LogNet, Error, TEXT("Server is incompatible with the local version of the game: RemoteNetworkVersion=%u, ")
+			TEXT("RemoteNetworkFeatures=%s vs LocalNetworkVersion=%u, LocalNetworkFeatures=%s"), 
+			RemoteNetworkVersion, RemoteFeaturesDescription.ToString(), FNetworkVersion::GetLocalNetworkVersion(),
+			LocalFeaturesDescription.ToString());
+
+
+	const FString ConnectionError = NSLOCTEXT("Engine", "ClientOutdated",
+		"The match you are trying to join is running an incompatible version of the game.  Please try upgrading your game version.").ToString();
+
+	GEngine->BroadcastNetworkFailure(GetWorld(), Driver, ENetworkFailure::OutdatedClient, ConnectionError);
+
+	if (NetUpgradeSource == UE::Net::ENetUpgradeSource::StatelessHandshake)
+	{
+		Close(ENetCloseResult::OutdatedClient);
+	}
+}
+
 FString UNetConnection::Describe()
 {
 	return FString::Printf( TEXT( "[UNetConnection] RemoteAddr: %s, Name: %s, Driver: %s, IsServer: %s, PC: %s, Owner: %s, UniqueId: %s" ),
diff --git a/Engine/Source/Runtime/Engine/Private/PacketHandlers/StatelessConnectHandlerComponent.cpp b/Engine/Source/Runtime/Engine/Private/PacketHandlers/StatelessConnectHandlerComponent.cpp
--- a/Engine/Source/Runtime/Engine/Private/PacketHandlers/StatelessConnectHandlerComponent.cpp
+++ b/Engine/Source/Runtime/Engine/Private/PacketHandlers/StatelessConnectHandlerComponent.cpp
@@ -10,6 +10,7 @@
 #include "Stats/StatsTrace.h"
 #include <limits> // IWYU pragma: keep
 
+
 DEFINE_LOG_CATEGORY(LogHandshake);
 
 
@@ -41,19 +42,19 @@ DEFINE_LOG_CATEGORY(LogHandshake);
  *
  * [?:MagicHeader][2:SessionID][3:ClientID][HandshakeBit][RestartHandshakeBit]
  * [8:MinVersion][8:CurVersion][8:HandshakePacketType][8:SentPacketCount][32:NetworkVersion]
- * [SecretIdBit] [28:PacketSizeFiller][AlignPad][?:RandomData]
+ * [16:NetworkFeatures][SecretIdBit][28:PacketSizeFiller][AlignPad][?:RandomData]
  *													--->
  *															Server - Stateless Handshake Challenge:
  *
  *															[?:MagicHeader][2:SessionID][3:ClientID][HandshakeBit][RestartHandshakeBit]
  *															[8:MinVersion][8:CurVersion][8:HandshakePacketType][8:SentPacketCount][32:NetworkVersion]
- *															[SecretIdBit][8:Timestamp][20:Cookie][AlignPad][?:RandomData]
+ *															[16:NetworkFeatures][SecretIdBit][8:Timestamp][20:Cookie][AlignPad][?:RandomData]
  *													<---
  * Client - Stateless Challenge Response:
  *
  * [?:MagicHeader][2:SessionID][3:ClientID][HandshakeBit][RestartHandshakeBit]
  * [8:MinVersion][8:CurVersion][8:HandshakePacketType][8:SentPacketCount][32:NetworkVersion]
- * [SecretIdBit][8:Timestamp][20:Cookie][AlignPad][?:RandomData]
+ * [16:NetworkFeatures][SecretIdBit][8:Timestamp][20:Cookie][AlignPad][?:RandomData]
  *													--->
  *															Server:
  *															Ignore, or create UNetConnection.
@@ -62,7 +63,7 @@ DEFINE_LOG_CATEGORY(LogHandshake);
  *
  *															[?:MagicHeader][2:SessionID][3:ClientID][HandshakeBit][RestartHandshakeBit]
  *															[8:MinVersion][8:CurVersion][8:HandshakePacketType][8:SentPacketCount][32:NetworkVersion]
- *															[SecretIdBit][8:Timestamp][20:Cookie][AlignPad][?:RandomData]
+ *															[16:NetworkFeatures][SecretIdBit][8:Timestamp][20:Cookie][AlignPad][?:RandomData]
  *													<---
  * Client:
  * Handshake Complete.
@@ -76,7 +77,8 @@ DEFINE_LOG_CATEGORY(LogHandshake);
  *															Server - Restart Handshake Request:
  *
  *															[?:MagicHeader][2:SessionID][3:ClientID][HandshakeBit][RestartHandshakeBit]
- *															[8:HandshakePacketType][8:SentPacketCount][32:NetworkVersion][AlignPad][?:RandomData]
+ *															[8:HandshakePacketType][8:SentPacketCount][32:NetworkVersion][16:NetworkFeatures]
+ *															[AlignPad][?:RandomData]
  *													<--
  * Client -  Initial Connect (as above)
  *													-->
@@ -86,7 +88,8 @@ DEFINE_LOG_CATEGORY(LogHandshake);
  *
  * [?:MagicHeader][2:SessionID][3:ClientID][HandshakeBit][RestartHandshakeBit]
  * [8:MinVersion][8:CurVersion][8:HandshakePacketType][8:SentPacketCount][32:NetworkVersion]
- * [SecretIdBit][8:Timestamp][20:Cookie][20:OriginalCookie][AlignPad][?:RandomData]
+ * [16:NetworkFeatures][SecretIdBit][8:Timestamp][20:Cookie][20:OriginalCookie]
+ * [AlignPad][?:RandomData]
  *													-->
  *															Server:
  *															Ignore, or restore UNetConnection.
@@ -109,6 +112,7 @@ DEFINE_LOG_CATEGORY(LogHandshake);
  *	- HandshakePacketType:	Number indicating the type of handshake packet, based on EHandshakePacketType
  *	- SentPacketCount:		The number of handshake packets sent - for packet-analysis/debugging purposes
  *	- NetworkVersion:		The Network CL version, according to FNetworkVersion::GetLocalNetworkVersion
+ *	- NetworkFeatures		The runtime Network Features, according to UNetDriver::GetNetworkRuntimeFeatures
  *	- Timestamp:			Server timestamp, from the moment the handshake challenge was sent.
  *	- Cookie:				Cookie generated by the server, which the client must reply with.
  *	- AlignPad:				Handshake packets and PacketHandler's in general, require complex padding of packets. See ParseHandshakePacket.
@@ -194,7 +198,6 @@ DEFINE_LOG_CATEGORY(LogHandshake);
  */
 
 
-
 /**
  * Debug Defines
  */
@@ -204,15 +207,15 @@ DEFINE_LOG_CATEGORY(LogHandshake);
 #define PACKETLOSS_TEST 0
 
 
-
 /**
  * Defines
  */
 
-#define HANDSHAKE_PACKET_SIZE_BITS				291
-#define RESTART_HANDSHAKE_PACKET_SIZE_BITS		66
-#define RESTART_RESPONSE_SIZE_BITS				451
-
+#define BASE_PACKET_SIZE_BITS					82
+#define HANDSHAKE_PACKET_SIZE_BITS				(BASE_PACKET_SIZE_BITS + 225)
+#define RESTART_HANDSHAKE_PACKET_SIZE_BITS		BASE_PACKET_SIZE_BITS
+#define RESTART_RESPONSE_SIZE_BITS				(BASE_PACKET_SIZE_BITS + 385)
+#define VERSION_UPGRADE_SIZE_BITS				BASE_PACKET_SIZE_BITS
 
 // The number of seconds between secret value updates, and the random variance applied to this
 #define SECRET_UPDATE_TIME			15.f
@@ -278,8 +281,7 @@ namespace UE::Net
 	static FAutoConsoleVariableRef CVarNetHandshakeEnforceNetworkCLVersion(
 		TEXT("net.HandshakeEnforceNetworkCLVersion"),
 		GHandshakeEnforceNetworkCLVersion,
-		TEXT("Whether or not the stateless handshake should enforce the Network CL version, instead of the higher level netcode. ")
-		TEXT("(does NOT result in a graceful rejection, packet is dropped without response - higher level netcode would have sent upgrade message)"));
+		TEXT("Whether or not the stateless handshake should enforce the Network CL version, instead of the higher level netcode."));
 
 	static int32 GVerifyNetSessionID = 1;
 
@@ -312,7 +314,6 @@ namespace UE::Net
 	/** The amount by which the length of random data should randomly vary */
 	static constexpr int32 RandomDataLengthVarianceBytes	= 8;
 
-
 	/** HANDSHAKE_PACKET_SIZE_BITS for EHandshakeVersion::Original */
 	static constexpr int32 OriginalHandshakePacketSizeBits = 227;
 
@@ -322,7 +323,6 @@ namespace UE::Net
 	/** RESTART_RESPONSE_SIZE_BITS for EHandshakeVersion::Original */
 	static constexpr int32 OriginalRestartResponseSizeBits = 387;
 
-
 	/** HANDSHAKE_PACKET_SIZE_BITS for EHandshakeVersion::Randomized */
 	static constexpr int32 VerRandomizedHandshakePacketSizeBits = 259;
 
@@ -331,14 +331,56 @@ namespace UE::Net
 
 	/** RESTART_RESPONSE_SIZE_BITS for EHandshakeVersion::Randomized */
 	static constexpr int32 VerRandomizedRestartResponseSizeBits = 419;
-}
 
 
+	const TCHAR* LexToString(EHandshakePacketType PacketType)
+	{
+		static_assert(EHandshakePacketType::Last == EHandshakePacketType::VersionUpgrade &&
+						static_cast<uint8>(EHandshakePacketType::VersionUpgrade) == 6, "Add new EHandshakePacketType entries to LexToString.");
+
+		switch (PacketType)
+		{
+		case EHandshakePacketType::InitialPacket:		return TEXT("InitialPacket");
+		case EHandshakePacketType::Challenge:			return TEXT("Challenge");
+		case EHandshakePacketType::Response:			return TEXT("Response");
+		case EHandshakePacketType::Ack:					return TEXT("Ack");
+		case EHandshakePacketType::RestartHandshake:	return TEXT("RestartHandshake");
+		case EHandshakePacketType::RestartResponse:		return TEXT("RestartResponse");
+		case EHandshakePacketType::VersionUpgrade:		return TEXT("VersionUpgrade");
+		}
+
+		return TEXT("Unknown");
+	}
+
+
+	/**
+	 * FStatelessHandshakeFailureInfo
+	 */
+
+	FStatelessHandshakeFailureInfo::FStatelessHandshakeFailureInfo()
+		: RemoteNetworkFeatures(EEngineNetworkRuntimeFeatures::None)
+	{
+	}
+}
+
 
 /**
  * StatelessConnectHandlerComponent
  */
 
+StatelessConnectHandlerComponent::FCommonSendToClientParams::FCommonSendToClientParams(const TSharedPtr<const FInternetAddr>& InClientAddress,
+																						EHandshakeVersion InHandshakeVersion, uint32 InClientID)
+	: ClientAddress(InClientAddress)
+	, HandshakeVersion(InHandshakeVersion)
+	, ClientID(InClientID)
+{
+}
+
+StatelessConnectHandlerComponent::FParsedHandshakeData::FParsedHandshakeData()
+	: RemoteNetworkFeatures(EEngineNetworkRuntimeFeatures::None)
+{
+}
+
 StatelessConnectHandlerComponent::StatelessConnectHandlerComponent()
 	: HandlerComponent(FName(TEXT("StatelessConnectHandlerComponent")))
 	, Driver(nullptr)
@@ -432,83 +474,19 @@ void StatelessConnectHandlerComponent::SendInitialPacket(EHandshakeVersion Hands
 		{
 			const int32 AdjustedSize = GetAdjustedSizeBits(HANDSHAKE_PACKET_SIZE_BITS, HandshakeVersion);
 			FBitWriter InitialPacket(AdjustedSize + (BaseRandomDataLengthBytes * 8) + 1 /* Termination bit */);
-			uint8 bHandshakePacket = 1;
-
-			if (MagicHeader.Num() > 0)
-			{
-				InitialPacket.SerializeBits(MagicHeader.GetData(), MagicHeader.Num());
-			}
-
-			if (HandshakeVersion >= EHandshakeVersion::SessionClientId)
-			{
-				InitialPacket.SerializeBits(&CachedGlobalNetTravelCount, SessionIDSizeBits);
-				InitialPacket.SerializeBits(&CachedClientID, ClientIDSizeBits);
-			}
-
-			InitialPacket.WriteBit(bHandshakePacket);
 
+			BeginHandshakePacket(InitialPacket, EHandshakePacketType::InitialPacket, HandshakeVersion, SentHandshakePacketCount, CachedClientID,
+									(bRestartedHandshake ? EHandshakePacketModifier::RestartHandshake : EHandshakePacketModifier::None));
 
-			// In order to prevent DRDoS reflection amplification attacks, clients must pad the packet to match server packet size
-			uint8 bRestartHandshake = bRestartedHandshake ? 1 : 0;
 			uint8 SecretIdPad = 0;
 			uint8 PacketSizeFiller[28];
 
-			InitialPacket.WriteBit(bRestartHandshake);
-
-			if (HandshakeVersion >= EHandshakeVersion::Randomized)
-			{
-				uint8 MinVersion = MinSupportedHandshakeVersion;
-				uint8 CurVersion = static_cast<uint8>(HandshakeVersion);
-				uint8 HandshakePacketType = static_cast<uint8>(EHandshakePacketType::InitialPacket);
-
-				InitialPacket << MinVersion;
-				InitialPacket << CurVersion;
-				InitialPacket << HandshakePacketType;
-				InitialPacket << SentHandshakePacketCount;
-			}
-
-			if (HandshakeVersion >= EHandshakeVersion::NetCLVersion)
-			{
-				uint32 LocalNetworkVersion = FNetworkVersion::GetLocalNetworkVersion();
-
-				InitialPacket << LocalNetworkVersion;
-			}
-
 			InitialPacket.WriteBit(SecretIdPad);
 
 			FMemory::Memzero(PacketSizeFiller, UE_ARRAY_COUNT(PacketSizeFiller));
 			InitialPacket.Serialize(PacketSizeFiller, UE_ARRAY_COUNT(PacketSizeFiller));
 
-
-
-			CapHandshakePacket(InitialPacket, HandshakeVersion);
-
-
-			// Disable PacketHandler parsing, and send the raw packet
-			Handler->SetRawSend(true);
-
-#if !UE_BUILD_SHIPPING && PACKETLOSS_TEST
-			bool bRandFail = FMath::RandBool();
-
-			if (bRandFail)
-			{
-				UE_LOG(LogHandshake, Log, TEXT("Triggering random initial connect packet fail."));
-			}
-
-			if (!bRandFail)
-#endif
-			{
-				if (ServerConn->Driver->IsNetResourceValid())
-				{
-					FOutPacketTraits Traits;
-
-					ServerConn->LowLevelSend(InitialPacket.GetData(), InitialPacket.GetNumBits(), Traits);
-				}
-			}
-
-			Handler->SetRawSend(false);
-
-			LastClientSendTimestamp = FPlatformTime::Seconds();
+			SendToServer(HandshakeVersion, EHandshakePacketType::InitialPacket, InitialPacket);
 		}
 		else
 		{
@@ -517,58 +495,27 @@ void StatelessConnectHandlerComponent::SendInitialPacket(EHandshakeVersion Hands
 	}
 }
 
-void StatelessConnectHandlerComponent::SendConnectChallenge(TSharedPtr<const FInternetAddr> ClientAddress, EHandshakeVersion HandshakeVersion,
-															uint8 ClientSentHandshakePacketCount, uint32 InClientID)
+void StatelessConnectHandlerComponent::SendConnectChallenge(FCommonSendToClientParams CommonParams, uint8 ClientSentHandshakePacketCount)
 {
 	using namespace UE::Net;
 
 	if (Driver != nullptr)
 	{
-		const int32 AdjustedSize = GetAdjustedSizeBits(HANDSHAKE_PACKET_SIZE_BITS, HandshakeVersion);
+		const int32 AdjustedSize = GetAdjustedSizeBits(HANDSHAKE_PACKET_SIZE_BITS, CommonParams.HandshakeVersion);
 		FBitWriter ChallengePacket(AdjustedSize + (BaseRandomDataLengthBytes * 8) + 1 /* Termination bit */);
-		uint8 bHandshakePacket = 1;
-		uint8 bRestartHandshake = 0; // Ignored clientside
-		double Timestamp = Driver->GetElapsedTime();
-		uint8 Cookie[COOKIE_BYTE_SIZE];
-
-		GenerateCookie(ClientAddress, ActiveSecret, Timestamp, Cookie);
-
-		if (MagicHeader.Num() > 0)
-		{
-			ChallengePacket.SerializeBits(MagicHeader.GetData(), MagicHeader.Num());
-		}
 
-		if (HandshakeVersion >= EHandshakeVersion::SessionClientId)
-		{
-			ChallengePacket.SerializeBits(&CachedGlobalNetTravelCount, SessionIDSizeBits);
-			ChallengePacket.SerializeBits(&InClientID, ClientIDSizeBits);
-		}
+		BeginHandshakePacket(ChallengePacket, EHandshakePacketType::Challenge, CommonParams.HandshakeVersion, ClientSentHandshakePacketCount,
+								CommonParams.ClientID);
 
-		ChallengePacket.WriteBit(bHandshakePacket);
-		ChallengePacket.WriteBit(bRestartHandshake);
-
-		if (HandshakeVersion >= EHandshakeVersion::Randomized)
-		{
-			uint8 MinVersion = MinSupportedHandshakeVersion;
-			uint8 CurVersion = static_cast<uint8>(HandshakeVersion);
-			uint8 HandshakePacketType = static_cast<uint8>(EHandshakePacketType::Challenge);
-
-			ChallengePacket << MinVersion;
-			ChallengePacket << CurVersion;
-			ChallengePacket << HandshakePacketType;
-			ChallengePacket << ClientSentHandshakePacketCount;
-		}
-
-		if (HandshakeVersion >= EHandshakeVersion::NetCLVersion)
-		{
-			uint32 LocalNetworkVersion = FNetworkVersion::GetLocalNetworkVersion();
+		double Timestamp = Driver->GetElapsedTime();
+		uint8 Cookie[COOKIE_BYTE_SIZE];
 
-			ChallengePacket << LocalNetworkVersion;
-		}
+		GenerateCookie(CommonParams.ClientAddress, ActiveSecret, Timestamp, Cookie);
 
 		ChallengePacket.WriteBit(ActiveSecret);
 
 		ChallengePacket << Timestamp;
+
 		ChallengePacket.Serialize(Cookie, UE_ARRAY_COUNT(Cookie));
 
 #if !UE_BUILD_SHIPPING
@@ -578,47 +525,7 @@ void StatelessConnectHandlerComponent::SendConnectChallenge(TSharedPtr<const FIn
 				TEXT("SendConnectChallenge. Timestamp: %f, Cookie: %s" ), Timestamp, *FString::FromBlob(Cookie, UE_ARRAY_COUNT(Cookie)));
 #endif
 
-		CapHandshakePacket(ChallengePacket, HandshakeVersion);
-
-		
-		// Disable PacketHandler parsing, and send the raw packet
-		PacketHandler* ConnectionlessHandler = Driver->ConnectionlessHandler.Get();
-
-		if (ConnectionlessHandler != nullptr)
-		{
-			ConnectionlessHandler->SetRawSend(true);
-		}
-
-#if !UE_BUILD_SHIPPING && PACKETLOSS_TEST
-		bool bRandFail = FMath::RandBool();
-
-		if (bRandFail)
-		{
-			UE_LOG(LogHandshake, Log, TEXT("Triggering random connect challenge packet fail."));
-		}
-
-		if (!bRandFail)
-#endif
-		{
-			if (Driver->IsNetResourceValid())
-			{
-				FOutPacketTraits Traits;
-
-				Driver->LowLevelSend(ClientAddress, ChallengePacket.GetData(), ChallengePacket.GetNumBits(), Traits);
-			}
-		}
-
-
-		if (ConnectionlessHandler != nullptr)
-		{
-			ConnectionlessHandler->SetRawSend(false);
-		}
-	}
-	else
-	{
-#if !UE_BUILD_SHIPPING
-		UE_LOG(LogHandshake, Error, TEXT("Tried to send handshake challenge packet without a net driver."));
-#endif
+		SendToClient(CommonParams, EHandshakePacketType::Challenge, ChallengePacket);
 	}
 }
 
@@ -634,43 +541,10 @@ void StatelessConnectHandlerComponent::SendChallengeResponse(EHandshakeVersion H
 		const int32 AdjustedSize = GetAdjustedSizeBits((bRestartedHandshake ? RESTART_RESPONSE_SIZE_BITS : HANDSHAKE_PACKET_SIZE_BITS),
 														HandshakeVersion);
 		FBitWriter ResponsePacket(AdjustedSize + (BaseRandomDataLengthBytes * 8) + 1 /* Termination bit */);
-		uint8 bHandshakePacket = 1;
-		uint8 bRestartHandshake = (bRestartedHandshake ? 1 : 0);
-
-		if (MagicHeader.Num() > 0)
-		{
-			ResponsePacket.SerializeBits(MagicHeader.GetData(), MagicHeader.Num());
-		}
-
-		if (HandshakeVersion >= EHandshakeVersion::SessionClientId)
-		{
-			ResponsePacket.SerializeBits(&CachedGlobalNetTravelCount, SessionIDSizeBits);
-			ResponsePacket.SerializeBits(&CachedClientID, ClientIDSizeBits);
-		}
+		EHandshakePacketType HandshakePacketType = bRestartedHandshake ? EHandshakePacketType::RestartResponse : EHandshakePacketType::Response;
 
-		ResponsePacket.WriteBit(bHandshakePacket);
-		ResponsePacket.WriteBit(bRestartHandshake);
-
-		if (HandshakeVersion >= EHandshakeVersion::Randomized)
-		{
-			uint8 MinVersion = MinSupportedHandshakeVersion;
-			uint8 CurVersion = static_cast<uint8>(HandshakeVersion);
-			uint8 HandshakePacketType = static_cast<uint8>(bRestartedHandshake ? EHandshakePacketType::RestartResponse :
-											EHandshakePacketType::Response);
-
-			ResponsePacket << MinVersion;
-			ResponsePacket << CurVersion;
-
-			ResponsePacket << HandshakePacketType;
-			ResponsePacket << SentHandshakePacketCount;
-		}
-
-		if (HandshakeVersion >= EHandshakeVersion::NetCLVersion)
-		{
-			uint32 LocalNetworkVersion = FNetworkVersion::GetLocalNetworkVersion();
-
-			ResponsePacket << LocalNetworkVersion;
-		}
+		BeginHandshakePacket(ResponsePacket, HandshakePacketType, HandshakeVersion, SentHandshakePacketCount, CachedClientID,
+								(bRestartedHandshake ? EHandshakePacketModifier::RestartHandshake : EHandshakePacketModifier::None));
 
 		ResponsePacket.WriteBit(InSecretId);
 
@@ -683,39 +557,15 @@ void StatelessConnectHandlerComponent::SendChallengeResponse(EHandshakeVersion H
 		}
 
 #if !UE_BUILD_SHIPPING
-		UE_LOG( LogHandshake, Log, TEXT( "SendChallengeResponse. Timestamp: %f, Cookie: %s" ), InTimestamp, *FString::FromBlob( InCookie, COOKIE_BYTE_SIZE ) );
+		UE_LOG(LogHandshake, Log, TEXT("SendChallengeResponse. Timestamp: %f, Cookie: %s"), InTimestamp,
+				*FString::FromBlob(InCookie, COOKIE_BYTE_SIZE));
 #endif
 
-		CapHandshakePacket(ResponsePacket, HandshakeVersion);
-
-
-		// Disable PacketHandler parsing, and send the raw packet
-		Handler->SetRawSend(true);
-
-#if !UE_BUILD_SHIPPING && PACKETLOSS_TEST
-		bool bRandFail = FMath::RandBool();
-
-		if (bRandFail)
-		{
-			UE_LOG(LogHandshake, Log, TEXT("Triggering random challenge response packet fail."));
-		}
-
-		if (!bRandFail)
-#endif
-		{
-			if (ServerConn->Driver->IsNetResourceValid())
-			{
-				FOutPacketTraits Traits;
-
-				ServerConn->LowLevelSend(ResponsePacket.GetData(), ResponsePacket.GetNumBits(), Traits);
-			}
-		}
+		SendToServer(HandshakeVersion, HandshakePacketType, ResponsePacket);
 
-		Handler->SetRawSend(false);
 
 		int16* CurSequence = (int16*)InCookie;
 
-		LastClientSendTimestamp = FPlatformTime::Seconds();
 		LastSecretId = InSecretId;
 		LastTimestamp = InTimestamp;
 		LastServerSequence = *CurSequence & (MAX_PACKETID - 1);
@@ -730,54 +580,24 @@ void StatelessConnectHandlerComponent::SendChallengeResponse(EHandshakeVersion H
 	}
 }
 
-void StatelessConnectHandlerComponent::SendChallengeAck(TSharedPtr<const FInternetAddr> ClientAddress, EHandshakeVersion HandshakeVersion,
-														uint8 ClientSentHandshakePacketCount, uint32 InClientID, uint8 InCookie[COOKIE_BYTE_SIZE])
+void StatelessConnectHandlerComponent::SendChallengeAck(FCommonSendToClientParams CommonParams, uint8 ClientSentHandshakePacketCount,
+														uint8 InCookie[COOKIE_BYTE_SIZE])
 {
 	using namespace UE::Net;
 
 	if (Driver != nullptr)
 	{
-		const int32 AdjustedSize = GetAdjustedSizeBits(HANDSHAKE_PACKET_SIZE_BITS, HandshakeVersion);
+		const int32 AdjustedSize = GetAdjustedSizeBits(HANDSHAKE_PACKET_SIZE_BITS, CommonParams.HandshakeVersion);
 		FBitWriter AckPacket(AdjustedSize + (BaseRandomDataLengthBytes * 8) + 1 /* Termination bit */);
-		uint8 bHandshakePacket = 1;
-		uint8 bRestartHandshake = 0; // Ignored clientside
-		double Timestamp  = -1.0;
 
-		if (MagicHeader.Num() > 0)
-		{
-			AckPacket.SerializeBits(MagicHeader.GetData(), MagicHeader.Num());
-		}
-
-		if (HandshakeVersion >= EHandshakeVersion::SessionClientId)
-		{
-			AckPacket.SerializeBits(&CachedGlobalNetTravelCount, SessionIDSizeBits);
-			AckPacket.SerializeBits(&InClientID, ClientIDSizeBits);
-		}
+		BeginHandshakePacket(AckPacket, EHandshakePacketType::Ack, CommonParams.HandshakeVersion, ClientSentHandshakePacketCount,
+								CommonParams.ClientID);
 
-		AckPacket.WriteBit(bHandshakePacket);
-		AckPacket.WriteBit(bRestartHandshake);
 
-		if (HandshakeVersion >= EHandshakeVersion::Randomized)
-		{
-			uint8 MinVersion = MinSupportedHandshakeVersion;
-			uint8 CurVersion = static_cast<uint8>(HandshakeVersion);
-			uint8 HandshakePacketType = static_cast<uint8>(EHandshakePacketType::Ack);
-
-			AckPacket << MinVersion;
-			AckPacket << CurVersion;
-
-			AckPacket << HandshakePacketType;
-			AckPacket << ClientSentHandshakePacketCount;
-		}
-
-		if (HandshakeVersion >= EHandshakeVersion::NetCLVersion)
-		{
-			uint32 LocalNetworkVersion = FNetworkVersion::GetLocalNetworkVersion();
-
-			AckPacket << LocalNetworkVersion;
-		}
+		double Timestamp  = -1.0;
+		uint8 ActiveSecret_Unused = 1;
 
-		AckPacket.WriteBit(bHandshakePacket);	// ActiveSecret
+		AckPacket.WriteBit(ActiveSecret_Unused);
 
 		AckPacket << Timestamp;
 		AckPacket.Serialize(InCookie, COOKIE_BYTE_SIZE);
@@ -786,103 +606,137 @@ void StatelessConnectHandlerComponent::SendChallengeAck(TSharedPtr<const FIntern
 		UE_LOG(LogHandshake, Log, TEXT("SendChallengeAck. InCookie: %s" ), *FString::FromBlob(InCookie, COOKIE_BYTE_SIZE));
 #endif
 
-		CapHandshakePacket(AckPacket, HandshakeVersion);
+		SendToClient(CommonParams, EHandshakePacketType::Ack, AckPacket);
+	}
+}
 
-		
-		// Disable PacketHandler parsing, and send the raw packet
-		PacketHandler* ConnectionlessHandler = Driver->ConnectionlessHandler.Get();
+void StatelessConnectHandlerComponent::SendRestartHandshakeRequest(FCommonSendToClientParams CommonParams)
+{
+	using namespace UE::Net;
 
-		if (ConnectionlessHandler != nullptr)
-		{
-			ConnectionlessHandler->SetRawSend(true);
-		}
+	if (Driver != nullptr)
+	{
+		const int32 AdjustedSize = GetAdjustedSizeBits(RESTART_HANDSHAKE_PACKET_SIZE_BITS, CommonParams.HandshakeVersion);
+		FBitWriter RestartPacket(AdjustedSize + (BaseRandomDataLengthBytes * 8) + 1 /* Termination bit */);
 
-#if !UE_BUILD_SHIPPING && PACKETLOSS_TEST
-		bool bRandFail = FMath::RandBool();
+		BeginHandshakePacket(RestartPacket, EHandshakePacketType::RestartHandshake, CommonParams.HandshakeVersion, SentHandshakePacketCount,
+								CommonParams.ClientID, EHandshakePacketModifier::RestartHandshake);
 
-		if (bRandFail)
-		{
-			UE_LOG(LogHandshake, Log, TEXT("Triggering random challenge ack packet fail."));
-		}
+#if !UE_BUILD_SHIPPING
+		FDDoSDetection* DDoS = Handler->GetDDoS();
 
-		if (!bRandFail)
+		UE_CLOG((DDoS == nullptr || !DDoS->CheckLogRestrictions()), LogHandshake, Verbose, TEXT("SendRestartHandshakeRequest."));
 #endif
-		{
-			if (Driver->IsNetResourceValid())
-			{
-				FOutPacketTraits Traits;
 
-				Driver->LowLevelSend(ClientAddress, AckPacket.GetData(), AckPacket.GetNumBits(), Traits);
-			}
-		}
+		SendToClient(CommonParams, EHandshakePacketType::RestartHandshake, RestartPacket);
+	}
+}
 
+void StatelessConnectHandlerComponent::SendVersionUpgradeMessage(FCommonSendToClientParams CommonParams)
+{
+	using namespace UE::Net;
 
-		if (ConnectionlessHandler != nullptr)
-		{
-			ConnectionlessHandler->SetRawSend(false);
-		}
-	}
-	else
+	if (Driver != nullptr)
 	{
-#if !UE_BUILD_SHIPPING
-		UE_LOG(LogHandshake, Error, TEXT("Tried to send handshake challenge ack packet without a net driver."));
-#endif
+		const int32 AdjustedSize = GetAdjustedSizeBits(VERSION_UPGRADE_SIZE_BITS, CommonParams.HandshakeVersion);
+		FBitWriter UpgradePacket(AdjustedSize + (BaseRandomDataLengthBytes * 8) + 1 /* Termination bit */);
+
+		BeginHandshakePacket(UpgradePacket, EHandshakePacketType::VersionUpgrade, CommonParams.HandshakeVersion, SentHandshakePacketCount,
+								CommonParams.ClientID);
+
+		SendToClient(CommonParams, EHandshakePacketType::VersionUpgrade, UpgradePacket);
 	}
 }
 
-void StatelessConnectHandlerComponent::SendRestartHandshakeRequest(const TSharedPtr<const FInternetAddr> ClientAddress,
-																	EHandshakeVersion HandshakeVersion, uint32 InClientID)
+void StatelessConnectHandlerComponent::BeginHandshakePacket(FBitWriter& HandshakePacket, EHandshakePacketType HandshakePacketType,
+															EHandshakeVersion HandshakeVersion, uint8 SentHandshakePacketCount_LocalOrRemote,
+															uint32 ClientID,
+															EHandshakePacketModifier HandshakePacketModifier/*=EHandshakePacketModifier::None*/)
 {
 	using namespace UE::Net;
 
-	if (Driver != nullptr)
+	uint8 bHandshakePacket = 1;
+	uint8 bRestartHandshake = EnumHasAnyFlags(HandshakePacketModifier, EHandshakePacketModifier::RestartHandshake) ? 1 : 0;
+
+	if (MagicHeader.Num() > 0)
 	{
-		const int32 AdjustedSize = GetAdjustedSizeBits(RESTART_HANDSHAKE_PACKET_SIZE_BITS, HandshakeVersion);
-		FBitWriter RestartPacket(AdjustedSize + (BaseRandomDataLengthBytes * 8) + 1 /* Termination bit */);
-		uint8 bHandshakePacket = 1;
-		uint8 bRestartHandshake = 1;
+		HandshakePacket.SerializeBits(MagicHeader.GetData(), MagicHeader.Num());
+	}
 
-		if (MagicHeader.Num() > 0)
-		{
-			RestartPacket.SerializeBits(MagicHeader.GetData(), MagicHeader.Num());
-		}
+	if (HandshakeVersion >= EHandshakeVersion::SessionClientId)
+	{
+		HandshakePacket.SerializeBits(&CachedGlobalNetTravelCount, SessionIDSizeBits);
+		HandshakePacket.SerializeBits(&ClientID, ClientIDSizeBits);
+	}
 
-		if (HandshakeVersion >= EHandshakeVersion::SessionClientId)
-		{
-			RestartPacket.SerializeBits(&CachedGlobalNetTravelCount, SessionIDSizeBits);
-			RestartPacket.SerializeBits(&InClientID, ClientIDSizeBits);
-		}
+	HandshakePacket.WriteBit(bHandshakePacket);
+	HandshakePacket.WriteBit(bRestartHandshake);
+
+	if (HandshakeVersion >= EHandshakeVersion::Randomized)
+	{
+		uint8 MinVersion = MinSupportedHandshakeVersion;
+		uint8 CurVersion = static_cast<uint8>(HandshakeVersion);
+		uint8 HandshakePacketTypeUint = static_cast<uint8>(HandshakePacketType);
+
+		HandshakePacket << MinVersion;
+		HandshakePacket << CurVersion;
+		HandshakePacket << HandshakePacketTypeUint;
+		HandshakePacket << SentHandshakePacketCount_LocalOrRemote;
+	}
 
-		RestartPacket.WriteBit(bHandshakePacket);
-		RestartPacket.WriteBit(bRestartHandshake);
+	if (HandshakeVersion >= EHandshakeVersion::NetCLVersion)
+	{
+		uint32 LocalNetworkVersion = FNetworkVersion::GetLocalNetworkVersion();
+		EEngineNetworkRuntimeFeatures LocalNetworkFeatures = (Driver != nullptr ? Driver->GetNetworkRuntimeFeatures() :
+																EEngineNetworkRuntimeFeatures::None);
 
-		if (HandshakeVersion >= EHandshakeVersion::Randomized)
-		{
-			uint8 MinVersion = MinSupportedHandshakeVersion;
-			uint8 CurVersion = static_cast<uint8>(HandshakeVersion);
-			uint8 HandshakePacketType = static_cast<uint8>(EHandshakePacketType::RestartHandshake);
+		static_assert(sizeof(EEngineNetworkRuntimeFeatures) == 2, "If EEngineNetworkRuntimeFeatures size changes, adjust BASE_PACKET_SIZE_BITS");
 
-			RestartPacket << MinVersion;
-			RestartPacket << CurVersion;
+		HandshakePacket << LocalNetworkVersion;
+		HandshakePacket << LocalNetworkFeatures;
+	}
+}
 
-			RestartPacket << HandshakePacketType;
-			RestartPacket << SentHandshakePacketCount;
+void StatelessConnectHandlerComponent::SendToServer(EHandshakeVersion HandshakeVersion, EHandshakePacketType PacketType, FBitWriter& Packet)
+{
+	if (UNetConnection* ServerConn = (Driver != nullptr ? Driver->ServerConnection : nullptr))
+	{
+		CapHandshakePacket(Packet, HandshakeVersion);
+
+
+		// Disable PacketHandler parsing, and send the raw packet
+		Handler->SetRawSend(true);
+
+#if !UE_BUILD_SHIPPING && PACKETLOSS_TEST
+		bool bRandFail = FMath::RandBool();
+
+		if (bRandFail)
+		{
+			UE_LOG(LogHandshake, Log, TEXT("Triggering random '%s' packet fail."), ToCStr(LexToString(PacketType)));
 		}
 
-		if (HandshakeVersion >= EHandshakeVersion::NetCLVersion)
+		if (!bRandFail)
+#endif
 		{
-			uint32 LocalNetworkVersion = FNetworkVersion::GetLocalNetworkVersion();
+			if (Driver->IsNetResourceValid())
+			{
+				FOutPacketTraits Traits;
 
-			RestartPacket << LocalNetworkVersion;
+				Driver->ServerConnection->LowLevelSend(Packet.GetData(), Packet.GetNumBits(), Traits);
+			}
 		}
 
-#if !UE_BUILD_SHIPPING
-		FDDoSDetection* DDoS = Handler->GetDDoS();
+		Handler->SetRawSend(false);
 
-		UE_CLOG((DDoS == nullptr || !DDoS->CheckLogRestrictions()), LogHandshake, Verbose, TEXT("SendRestartHandshakeRequest."));
-#endif
+		LastClientSendTimestamp = FPlatformTime::Seconds();
+	}
+}
 
-		CapHandshakePacket(RestartPacket, HandshakeVersion);
+void StatelessConnectHandlerComponent::SendToClient(FCommonSendToClientParams CommonParams, EHandshakePacketType PacketType, FBitWriter& Packet)
+{
+	if (Driver != nullptr)
+	{
+		CapHandshakePacket(Packet, CommonParams.HandshakeVersion);
 
 		
 		// Disable PacketHandler parsing, and send the raw packet
@@ -898,7 +752,7 @@ void StatelessConnectHandlerComponent::SendRestartHandshakeRequest(const TShared
 
 		if (bRandFail)
 		{
-			UE_LOG(LogHandshake, Log, TEXT("Triggering random restart handshake packet fail."));
+			UE_LOG(LogHandshake, Log, TEXT("Triggering random '%s' packet fail."), ToCStr(LexToString(PacketType)));
 		}
 
 		if (!bRandFail)
@@ -908,7 +762,7 @@ void StatelessConnectHandlerComponent::SendRestartHandshakeRequest(const TShared
 			{
 				FOutPacketTraits Traits;
 
-				Driver->LowLevelSend(ClientAddress, RestartPacket.GetData(), RestartPacket.GetNumBits(), Traits);
+				Driver->LowLevelSend(CommonParams.ClientAddress, Packet.GetData(), Packet.GetNumBits(), Traits);
 			}
 		}
 
@@ -918,12 +772,6 @@ void StatelessConnectHandlerComponent::SendRestartHandshakeRequest(const TShared
 			ConnectionlessHandler->SetRawSend(false);
 		}
 	}
-	else
-	{
-#if !UE_BUILD_SHIPPING
-		UE_LOG(LogHandshake, Error, TEXT("Tried to send restart handshake packet without a net driver."));
-#endif
-	}
 }
 
 void StatelessConnectHandlerComponent::CapHandshakePacket(FBitWriter& HandshakePacket, EHandshakeVersion HandshakeVersion)
@@ -945,7 +793,8 @@ void StatelessConnectHandlerComponent::CapHandshakePacket(FBitWriter& HandshakeP
 	}
 	else
 	{
-		check(NumBits == HANDSHAKE_PACKET_SIZE_BITS || NumBits == RESTART_HANDSHAKE_PACKET_SIZE_BITS || NumBits == RESTART_RESPONSE_SIZE_BITS);
+		check(NumBits == HANDSHAKE_PACKET_SIZE_BITS || NumBits == RESTART_HANDSHAKE_PACKET_SIZE_BITS || NumBits == RESTART_RESPONSE_SIZE_BITS ||
+				NumBits == VERSION_UPGRADE_SIZE_BITS);
 	}
 #endif
 
@@ -1026,7 +875,7 @@ void StatelessConnectHandlerComponent::SetDriver(UNetDriver* InDriver)
 
 		CachedClientID = CachedClientID & ((1 << ClientIDSizeBits) - 1);
 
-		UE_LOG(LogHandshake, Log, TEXT("Stateless Handshake: NetDriverDefinition '%s' CachedClientID: %i"), ToCStr(NetDriverDef),
+		UE_LOG(LogHandshake, Log, TEXT("Stateless Handshake: NetDriverDefinition '%s' CachedClientID: %u"), ToCStr(NetDriverDef),
 				CachedClientID);
 	}
 }
@@ -1047,7 +896,7 @@ void StatelessConnectHandlerComponent::Initialize()
 		if (CurHandshakeVersion > EHandshakeVersion::Latest)
 		{
 			UE_LOG(LogHandshake, Error, TEXT("net.CurrentHandshakeVersion value '%i' is invalid. Maximum value: %i"), CurrentHandshakeVersion,
-					static_cast<EHandshakeVersion>(EHandshakeVersion::Latest));
+					static_cast<int32>(EHandshakeVersion::Latest));
 		}
 		else
 		{
@@ -1070,7 +919,7 @@ void StatelessConnectHandlerComponent::Initialize()
 			if (MinHandshakeVersion < CompatibleRangeStart || MinHandshakeVersion > CompatibleRangeEnd)
 			{
 				UE_LOG(LogHandshake, Error, TEXT("net.MinHandshakeVersion value '%i' is invalid relative to net.CurrentHandshakeVersion value '%i'. ")
-						TEXT("Minimum value: %i, Maximum value: %i"), MinSupportedHandshakeVersion, CurrentHandshakeVersion,
+						TEXT("Minimum value: %u, Maximum value: %u"), MinSupportedHandshakeVersion, CurrentHandshakeVersion,
 						static_cast<uint8>(CompatibleRangeStart), static_cast<uint8>(CompatibleRangeEnd));
 			}
 		}
@@ -1090,6 +939,11 @@ void StatelessConnectHandlerComponent::InitFromConnectionless(StatelessConnectHa
 	LastInitTimestamp = (Driver != nullptr ? Driver->GetElapsedTime() : 0.0);
 }
 
+void StatelessConnectHandlerComponent::SetHandshakeFailureCallback(FHandshakeFailureFunc&& InHandshakeFailureFunc)
+{
+	HandshakeFailureCallback = MoveTemp(InHandshakeFailureFunc);
+}
+
 void StatelessConnectHandlerComponent::Incoming(FBitReader& Packet)
 {
 	using namespace UE::Net;
@@ -1127,8 +981,9 @@ void StatelessConnectHandlerComponent::Incoming(FBitReader& Packet)
 		{
 			const bool bIsChallengePacket = HandshakeData.HandshakePacketType == EHandshakePacketType::Challenge && HandshakeData.Timestamp > 0.0;
 			const bool bIsInitialChallengePacket = bIsChallengePacket && State != Handler::Component::State::Initialized;
+			const bool bIsUpgradePacket = HandshakeData.HandshakePacketType == EHandshakePacketType::VersionUpgrade;
 
-			if (Handler->Mode == Handler::Mode::Client && bHasValidClientID && (bHasValidSessionID || bIsInitialChallengePacket))
+			if (Handler->Mode == Handler::Mode::Client && bHasValidClientID && (bHasValidSessionID || bIsInitialChallengePacket || bIsUpgradePacket))
 			{
 				if (State == Handler::Component::State::UnInitialized || State == Handler::Component::State::InitializedOnLocal)
 				{
@@ -1142,7 +997,7 @@ void StatelessConnectHandlerComponent::Incoming(FBitReader& Packet)
 					else if (bIsChallengePacket)
 					{
 #if !UE_BUILD_SHIPPING
-						UE_LOG(LogHandshake, Log, TEXT("Cached server SessionID: %i"), SessionID);
+						UE_LOG(LogHandshake, Log, TEXT("Cached server SessionID: %u"), SessionID);
 #endif
 
 						CachedGlobalNetTravelCount = SessionID;
@@ -1185,6 +1040,32 @@ void StatelessConnectHandlerComponent::Incoming(FBitReader& Packet)
 						// Reset packet count clientside, due to how it affects protocol version fallback selection
 						SentHandshakePacketCount = 0;
 					}
+					else if (bIsUpgradePacket)
+					{
+						const uint32 LocalNetworkVersion = FNetworkVersion::GetLocalNetworkVersion();
+						TStringBuilder<128> LocalNetFeaturesDescription;
+						TStringBuilder<128> RemoteNetFeaturesDescription;
+						EEngineNetworkRuntimeFeatures LocalNetworkFeatures = (Driver != nullptr ? Driver->GetNetworkRuntimeFeatures() :
+																				EEngineNetworkRuntimeFeatures::None);
+						FNetworkVersion::DescribeNetworkRuntimeFeaturesBitset(LocalNetworkFeatures, LocalNetFeaturesDescription);
+						FNetworkVersion::DescribeNetworkRuntimeFeaturesBitset(HandshakeData.RemoteNetworkFeatures, RemoteNetFeaturesDescription);
+
+						UE_LOG(LogHandshake, Log, TEXT("Server is running an incompatible version of the game, and has rejected the connection. ")
+								TEXT("Server version '%u', Local version '%u', Server features '%s', Local features '%s'."),
+								HandshakeData.RemoteNetworkVersion, LocalNetworkVersion, ToCStr(RemoteNetFeaturesDescription.ToString()),
+								ToCStr(LocalNetFeaturesDescription.ToString()));
+
+						if (HandshakeFailureCallback)
+						{
+							FStatelessHandshakeFailureInfo FailureInfo;
+
+							FailureInfo.FailureReason = EHandshakeFailureReason::WrongVersion;
+							FailureInfo.RemoteNetworkVersion = HandshakeData.RemoteNetworkVersion;
+							FailureInfo.RemoteNetworkFeatures = HandshakeData.RemoteNetworkFeatures;
+
+							HandshakeFailureCallback(FailureInfo);
+						}
+					}
 				}
 				else if (HandshakeData.bRestartHandshake)
 				{
@@ -1297,14 +1178,15 @@ void StatelessConnectHandlerComponent::Incoming(FBitReader& Packet)
 							*LastChallengeSuccessAddress->ToString(true), *FString::FromBlob(AuthorisedCookie, COOKIE_BYTE_SIZE));
 #endif
 
-					SendChallengeAck(LastChallengeSuccessAddress, LastRemoteHandshakeVersion, 0, CachedClientID, AuthorisedCookie);
+					SendChallengeAck(FCommonSendToClientParams(LastChallengeSuccessAddress, LastRemoteHandshakeVersion, CachedClientID), 0,
+										AuthorisedCookie);
 				}
 			}
 			else if (!bHasValidSessionID || !bHasValidClientID)
 			{
 #if !UE_BUILD_SHIPPING
 				UE_CLOG(TrackValidationLogs(), LogHandshake, Log,
-						TEXT("Incoming: Rejecting handshake packet with invalid session id (%i vs %i) or connection id (%i vs %i)."),
+						TEXT("Incoming: Rejecting handshake packet with invalid session id (%u vs %u) or connection id (%u vs %u)."),
 						SessionID, CachedGlobalNetTravelCount, ClientID, CachedClientID);
 #endif
 
@@ -1331,7 +1213,7 @@ void StatelessConnectHandlerComponent::Incoming(FBitReader& Packet)
 	{
 #if !UE_BUILD_SHIPPING
 		UE_CLOG(TrackValidationLogs(), LogHandshake, Log,
-				TEXT("Incoming: Rejecting game packet with invalid session id (%i vs %i) or connection id (%i vs %i)."),
+				TEXT("Incoming: Rejecting game packet with invalid session id (%u vs %u) or connection id (%u vs %u)."),
 				SessionID, CachedGlobalNetTravelCount, ClientID, CachedClientID);
 #endif
 
@@ -1435,7 +1317,7 @@ void StatelessConnectHandlerComponent::IncomingConnectionless(FIncomingPacketRef
 		if (bHandshakePacket)
 		{
 			EHandshakeVersion TargetVersion = EHandshakeVersion::Latest;
-			bool bValidVersion = CheckVersion(HandshakeData, TargetVersion);
+			const bool bValidVersion = CheckVersion(HandshakeData, TargetVersion);
 			const bool bInitialConnect = HandshakeData.HandshakePacketType == EHandshakePacketType::InitialPacket &&
 												HandshakeData.Timestamp == 0.0;
 
@@ -1443,7 +1325,7 @@ void StatelessConnectHandlerComponent::IncomingConnectionless(FIncomingPacketRef
 			{
 				if (bInitialConnect)
 				{
-					SendConnectChallenge(Address, TargetVersion, HandshakeData.RemoteSentHandshakePacketCount, ClientID);
+					SendConnectChallenge(FCommonSendToClientParams(Address, TargetVersion, ClientID), HandshakeData.RemoteSentHandshakePacketCount);
 				}
 				// Challenge response
 				else if (Driver != nullptr)
@@ -1493,7 +1375,8 @@ void StatelessConnectHandlerComponent::IncomingConnectionless(FIncomingPacketRef
 
 
 							// Now ack the challenge response - the cookie is stored in AuthorisedCookie, to enable retries
-							SendChallengeAck(Address, TargetVersion, HandshakeData.RemoteSentHandshakePacketCount, ClientID, AuthorisedCookie);
+							SendChallengeAck(FCommonSendToClientParams(Address, TargetVersion, ClientID),
+												HandshakeData.RemoteSentHandshakePacketCount, AuthorisedCookie);
 						}
 					}
 				}
@@ -1501,12 +1384,35 @@ void StatelessConnectHandlerComponent::IncomingConnectionless(FIncomingPacketRef
 			else if (Handler->Mode == Handler::Mode::Server && bValidVersion && !bHasValidSessionID)
 			{
 #if !UE_BUILD_SHIPPING
-				UE_LOG(LogHandshake, Log, TEXT("IncomingConnectionless: Rejecting packet with invalid session id: %i vs %i."),
+				UE_LOG(LogHandshake, Log, TEXT("IncomingConnectionless: Rejecting packet with invalid session id: %u vs %u."),
 						SessionID, CachedGlobalNetTravelCount);
 #endif
 
 				Packet.SetError();
 			}
+			else if (Handler->Mode == Handler::Mode::Server && !bValidVersion && bInitialConnect &&
+						HandshakeData.RemoteCurVersion >= EHandshakeVersion::NetCLUpgradeMessage)
+			{
+				// Limit of 512 upgrade message packets, over 5 minutes
+				const uint32 NumUpgradeMessagesPerPeriod = 512;
+				const double UpgradeMessagePeriod = 300;
+				const double ElapsedTime = Driver->GetElapsedTime();
+
+				if (ElapsedTime - LastUpgradeMessagePeriodStart >= UpgradeMessagePeriod)
+				{
+					LastUpgradeMessagePeriodStart = ElapsedTime;
+					UpgradeMessageCounter = 0;
+				}
+				else
+				{
+					UpgradeMessageCounter++;
+				}
+
+				if (UpgradeMessageCounter < NumUpgradeMessagesPerPeriod)
+				{
+					SendVersionUpgradeMessage(FCommonSendToClientParams(Address, TargetVersion, ClientID));
+				}
+			}
 		}
 		else
 		{
@@ -1538,7 +1444,7 @@ void StatelessConnectHandlerComponent::IncomingConnectionless(FIncomingPacketRef
 		{
 			// The packet was fine but not a handshake packet - an existing client might suddenly be communicating on a different address.
 			// If we get them to resend their cookie, we can update the connection's info with their new address.
-			SendRestartHandshakeRequest(Address, static_cast<EHandshakeVersion>(MinSupportedHandshakeVersion), ClientID);
+			SendRestartHandshakeRequest(FCommonSendToClientParams(Address, static_cast<EHandshakeVersion>(MinSupportedHandshakeVersion), ClientID));
 		}
 	}
 }
@@ -1557,27 +1463,24 @@ bool StatelessConnectHandlerComponent::ParseHandshakePacket(FBitReader& Packet,
 		static constexpr int32 MaxRestartHandshakePacketVariance = RESTART_HANDSHAKE_PACKET_SIZE_BITS + MaxRandomBits;
 		static constexpr int32 MinRestartResponsePacketVariance = RESTART_RESPONSE_SIZE_BITS + MinRandomBits;
 		static constexpr int32 MaxRestartResponsePacketVariance = RESTART_RESPONSE_SIZE_BITS + MaxRandomBits;
+		static constexpr int32 MinVersionUpgradePacketVariance = VERSION_UPGRADE_SIZE_BITS + MinRandomBits;
+		static constexpr int32 MaxVersionUpgradePacketVariance = VERSION_UPGRADE_SIZE_BITS + MaxRandomBits;
 
 		static_assert(OriginalHandshakePacketSizeBits < MinHandshakePacketVariance || OriginalHandshakePacketSizeBits > MaxHandshakePacketVariance);
 		static_assert(OriginalHandshakePacketSizeBits < MinRestartHandshakePacketVariance ||
 						OriginalHandshakePacketSizeBits > MaxRestartHandshakePacketVariance);
 		static_assert(OriginalHandshakePacketSizeBits < MinRestartResponsePacketVariance ||
 						OriginalHandshakePacketSizeBits > MaxRestartResponsePacketVariance);
+		static_assert(OriginalHandshakePacketSizeBits < MinVersionUpgradePacketVariance ||
+						OriginalHandshakePacketSizeBits > MaxVersionUpgradePacketVariance);
 		static_assert(OriginalRestartHandshakePacketSizeBits < MinHandshakePacketVariance ||
 						OriginalRestartHandshakePacketSizeBits > MaxHandshakePacketVariance);
 		static_assert(OriginalRestartHandshakePacketSizeBits < MinRestartHandshakePacketVariance ||
 						OriginalRestartHandshakePacketSizeBits > MaxRestartHandshakePacketVariance);
 		static_assert(OriginalRestartHandshakePacketSizeBits < MinRestartResponsePacketVariance ||
 						OriginalRestartHandshakePacketSizeBits > MaxRestartResponsePacketVariance);
-
-		// Disable these static asserts, as the new protocol has been adjusted to never hit OriginalRestartResponseSizeBits
-#if 0
-		static_assert(OriginalRestartResponseSizeBits < MinHandshakePacketVariance || OriginalRestartResponseSizeBits > MaxHandshakePacketVariance);
-		static_assert(OriginalRestartResponseSizeBits < MinRestartHandshakePacketVariance ||
-						OriginalRestartResponseSizeBits > MaxRestartHandshakePacketVariance);
-		static_assert(OriginalRestartResponseSizeBits < MinRestartResponsePacketVariance ||
-						OriginalRestartResponseSizeBits > MaxRestartResponsePacketVariance);
-#endif
+		static_assert(OriginalRestartHandshakePacketSizeBits < MinVersionUpgradePacketVariance ||
+						OriginalRestartHandshakePacketSizeBits > MaxVersionUpgradePacketVariance);
 	};
 
 	bool bValidPacket = false;
@@ -1595,18 +1498,23 @@ bool StatelessConnectHandlerComponent::ParseHandshakePacket(FBitReader& Packet,
 	}
 
 
-	// Remaining bits, excluding packet sizes from different protocol versions
+	// Remaining bits, excluding packet sizes from different protocol versions (NOTE: Current code assumes packet size defines only increase)
 	const int32 MinBitsLeftExclHandshake = BitsLeft - (HANDSHAKE_PACKET_SIZE_BITS - 1);
 	const int32 MaxBitsLeftExclHandshake = BitsLeft - (VerRandomizedHandshakePacketSizeBits - 1);
 	const int32 MinBitsLeftExclRestartHandshake = BitsLeft - (RESTART_HANDSHAKE_PACKET_SIZE_BITS - 1);
 	const int32 MaxBitsLeftExclRestartHandshake = BitsLeft - (VerRandomizedRestartHandshakePacketSizeBits - 1);
 	const int32 MinBitsLeftExclRestartResponse = BitsLeft - (RESTART_RESPONSE_SIZE_BITS - 1);
 	const int32 MaxBitsLeftExclRestartResponse = BitsLeft - (VerRandomizedRestartResponseSizeBits - 1);
+	const int32 MinBitsLeftExclVersionUpgrade = BitsLeft - (VERSION_UPGRADE_SIZE_BITS - 1);
+	const int32 MaxBitsLeftExclVersionUpgrade = BitsLeft - (VERSION_UPGRADE_SIZE_BITS - 1);	// To be updated if size or BASE_PACKET_SIZE_BITS changes
 	const int32 MinRandomBits = (BaseRandomDataLengthBytes - RandomDataLengthVarianceBytes) * 8;
 	const int32 MaxRandomBits = BaseRandomDataLengthBytes * 8;
 	const bool bMaybeHandshakePacket = MaxBitsLeftExclHandshake >= MinRandomBits && MinBitsLeftExclHandshake <= MaxRandomBits;
 	const bool bMaybeRestartHandshakePacket = MaxBitsLeftExclRestartHandshake >= MinRandomBits && MinBitsLeftExclRestartHandshake <= MaxRandomBits;
 	const bool bMaybeRestartResponsePacket = MaxBitsLeftExclRestartResponse >= MinRandomBits && MinBitsLeftExclRestartResponse <= MaxRandomBits;
+	const bool bMaybeVersionUpgradePacket = MaxBitsLeftExclVersionUpgrade >= MinRandomBits && MinBitsLeftExclVersionUpgrade <= MaxRandomBits;
+
+	static_assert(BASE_PACKET_SIZE_BITS == 82 && VERSION_UPGRADE_SIZE_BITS == 82, "MaxBitsLeftExclVersionUpgrade needs to be updated.");
 
 	OutResult.bRestartHandshake = !!Packet.ReadBit();
 
@@ -1627,6 +1535,7 @@ bool StatelessConnectHandlerComponent::ParseHandshakePacket(FBitReader& Packet,
 	if (OutResult.RemoteCurVersion >= EHandshakeVersion::NetCLVersion)
 	{
 		Packet << OutResult.RemoteNetworkVersion;
+		Packet << OutResult.RemoteNetworkFeatures;
 	}
 
 	// Only accept handshake packets of roughly the right size
@@ -1636,6 +1545,7 @@ bool StatelessConnectHandlerComponent::ParseHandshakePacket(FBitReader& Packet,
 
 	const bool bRestartHandshakePacket = bMaybeRestartHandshakePacket && HandshakePacketTypeEnum == EHandshakePacketType::RestartHandshake;
 	const bool bRestartResponsePacket = bMaybeRestartResponsePacket && HandshakePacketTypeEnum == EHandshakePacketType::RestartResponse;
+	const bool bVersionUpgradePacket = bMaybeVersionUpgradePacket && HandshakePacketTypeEnum == EHandshakePacketType::VersionUpgrade;
 
 	// Only accept handshake packets of precisely the right size
 	if (bHandshakePacket || bRestartResponsePacket)
@@ -1657,6 +1567,10 @@ bool StatelessConnectHandlerComponent::ParseHandshakePacket(FBitReader& Packet,
 	{
 		bValidPacket = !Packet.IsError() && OutResult.bRestartHandshake && Handler->Mode == Handler::Mode::Client;
 	}
+	else if (bVersionUpgradePacket)
+	{
+		bValidPacket = !Packet.IsError() && !OutResult.bRestartHandshake && Handler->Mode == Handler::Mode::Client;
+	}
 
 	if (bValidPacket)
 	{
@@ -1762,12 +1676,20 @@ bool StatelessConnectHandlerComponent::CheckVersion(const FParsedHandshakeData&
 
 	const bool bCheckNetVersion = !UE_BUILD_SHIPPING || GHandshakeEnforceNetworkCLVersion;
 	bool bValidNetVersion = true;
+	EEngineNetworkRuntimeFeatures LocalNetworkFeatures = EEngineNetworkRuntimeFeatures::None;
 
 	if (bCheckNetVersion && HandshakeData.RemoteCurVersion >= EHandshakeVersion::NetCLVersion)
 	{
+		if (Driver != nullptr)
+		{
+			LocalNetworkFeatures = Driver->GetNetworkRuntimeFeatures();
+		}
+
 		const uint32 LocalNetworkVersion = FNetworkVersion::GetLocalNetworkVersion();
+		const bool bIsCompatible = FNetworkVersion::IsNetworkCompatible(LocalNetworkVersion, HandshakeData.RemoteNetworkVersion) &&
+									FNetworkVersion::AreNetworkRuntimeFeaturesCompatible(LocalNetworkFeatures, HandshakeData.RemoteNetworkFeatures);
 
-		if (!FNetworkVersion::IsNetworkCompatible(LocalNetworkVersion, HandshakeData.RemoteNetworkVersion))
+		if (!bIsCompatible)
 		{
 			bValidNetVersion = false;
 		}
@@ -1778,14 +1700,21 @@ bool StatelessConnectHandlerComponent::CheckVersion(const FParsedHandshakeData&
 	{
 		FDDoSDetection* DDoS = Handler->GetDDoS();
 		const uint32 LocalNetworkVersion = FNetworkVersion::GetLocalNetworkVersion();
+		TStringBuilder<128> LocalNetFeaturesDescription;
+		TStringBuilder<128> RemoteNetFeaturesDescription;
+
+		FNetworkVersion::DescribeNetworkRuntimeFeaturesBitset(LocalNetworkFeatures, LocalNetFeaturesDescription);
+		FNetworkVersion::DescribeNetworkRuntimeFeaturesBitset(HandshakeData.RemoteNetworkFeatures, RemoteNetFeaturesDescription);
 
 		UE_CLOG((DDoS == nullptr || !DDoS->CheckLogRestrictions()), LogHandshake, Log,
 				TEXT("CheckVersion: Incompatible version. bValidHandshakeVersion: %i, bValidNetVersion: %i, ")
-				TEXT("GHandshakeEnforceNetworkCLVersion: %i, RemoteMinVersion: %i, RemoteCurVersion: %i, MinSupportedHandshakeVersion: %i, ")
-				TEXT("CurrentHandshakeVersion: %i, RemoteNetworkVersion: %i, LocalNetworkVersion: %i"),
+				TEXT("GHandshakeEnforceNetworkCLVersion: %i, RemoteMinVersion: %u, RemoteCurVersion: %u, MinSupportedHandshakeVersion: %i, ")
+				TEXT("CurrentHandshakeVersion: %i, RemoteNetworkVersion: %u, LocalNetworkVersion: %u, RemoteNetworkFeatures: %s, ")
+				TEXT("LocalNetworkFeatures: %s"),
 				(int32)bValidHandshakeVersion, (int32)bValidNetVersion, GHandshakeEnforceNetworkCLVersion, RemoteMinVersionUint8,
 				static_cast<uint8>(HandshakeData.RemoteCurVersion), MinSupportedHandshakeVersion, CurrentHandshakeVersion,
-				HandshakeData.RemoteNetworkVersion, LocalNetworkVersion);
+				HandshakeData.RemoteNetworkVersion, LocalNetworkVersion, ToCStr(RemoteNetFeaturesDescription.ToString()),
+				ToCStr(LocalNetFeaturesDescription.ToString()));
 	}
 #endif
 
@@ -1794,7 +1723,7 @@ bool StatelessConnectHandlerComponent::CheckVersion(const FParsedHandshakeData&
 	return bValidHandshakeVersion && bPassedNetVersionConditions;
 }
 
-void StatelessConnectHandlerComponent::GenerateCookie(TSharedPtr<const FInternetAddr> ClientAddress, uint8 SecretId, double Timestamp,
+void StatelessConnectHandlerComponent::GenerateCookie(const TSharedPtr<const FInternetAddr>& ClientAddress, uint8 SecretId, double Timestamp,
 														uint8 (&OutCookie)[20]) const
 {
 	TArray<uint8> CookieData;
diff --git a/Engine/Source/Runtime/Engine/Private/PendingNetGame.cpp b/Engine/Source/Runtime/Engine/Private/PendingNetGame.cpp
--- a/Engine/Source/Runtime/Engine/Private/PendingNetGame.cpp
+++ b/Engine/Source/Runtime/Engine/Private/PendingNetGame.cpp
@@ -224,20 +224,10 @@ void UPendingNetGame::NotifyControlMessage(UNetConnection* Connection, uint8 Mes
 
 			if (FNetControlMessage<NMT_Upgrade>::Receive(Bunch, RemoteNetworkVersion, RemoteNetworkFeatures))
 			{
-				TStringBuilder<128> RemoteFeaturesDescription;
-				FNetworkVersion::DescribeNetworkRuntimeFeaturesBitset(RemoteNetworkFeatures, RemoteFeaturesDescription);
-
-				TStringBuilder<128> LocalFeaturesDescription;
-				FNetworkVersion::DescribeNetworkRuntimeFeaturesBitset(NetDriver->GetNetworkRuntimeFeatures(), LocalFeaturesDescription);
-
-				UE_LOG(LogNet, Error, TEXT("Server is incompatible with the local version of the game: RemoteNetworkVersion=%u, RemoteNetworkFeatures=%s vs LocalNetworkVersion=%u, LocalNetworkFeatures=%s"), 
-					RemoteNetworkVersion, RemoteFeaturesDescription.ToString(),
-					FNetworkVersion::GetLocalNetworkVersion(), LocalFeaturesDescription.ToString()
-				);
-
 				// Upgrade
 				ConnectionError = NSLOCTEXT("Engine", "ClientOutdated", "The match you are trying to join is running an incompatible version of the game.  Please try upgrading your game version.").ToString();
-				GEngine->BroadcastNetworkFailure(NULL, NetDriver, ENetworkFailure::OutdatedClient, ConnectionError);
+
+				Connection->HandleReceiveNetUpgrade(RemoteNetworkVersion, RemoteNetworkFeatures);
 			}
 
 			break;
diff --git a/Engine/Source/Runtime/Engine/Public/PacketHandlers/StatelessConnectHandlerComponent.h b/Engine/Source/Runtime/Engine/Public/PacketHandlers/StatelessConnectHandlerComponent.h
--- a/Engine/Source/Runtime/Engine/Public/PacketHandlers/StatelessConnectHandlerComponent.h
+++ b/Engine/Source/Runtime/Engine/Public/PacketHandlers/StatelessConnectHandlerComponent.h
@@ -9,18 +9,13 @@
 DECLARE_LOG_CATEGORY_EXTERN(LogHandshake, Log, All);
 
 
-/**
- * Forward Declarations
- */
-
+// Forward Declarations
 class UNetConnection;
 class UNetDriver;
+enum class EEngineNetworkRuntimeFeatures : uint16;
 
 
-/**
- * Defines
- */
-
+// Defines
 #define SECRET_BYTE_SIZE 64
 #define SECRET_COUNT 2
 #define COOKIE_BYTE_SIZE 20
@@ -33,15 +28,16 @@ namespace UE::Net
 	 */
 	enum class EHandshakeVersion : uint8
 	{
-		Original		= 0,	// The original/unversioned handshake protocol
-		Randomized		= 1,	// The version of the handshake protocol with randomization, versioning and debug/diagnostic tweaks
-		NetCLVersion	= 2,	// Added Network CL version, for optional extra-early client rejection
+		Original				= 0,	// The original/unversioned handshake protocol
+		Randomized				= 1,	// The version of the handshake protocol with randomization, versioning and debug/diagnostic tweaks
+		NetCLVersion			= 2,	// Added Network CL version, for optional extra-early client rejection
 
 		/** Net compatibility break */
 
-		SessionClientId	= 3,	// Added SessionId (index of server game session) and ClientId (clientside index of NetConnection)
+		SessionClientId			= 3,	// Added SessionId (index of server game session) and ClientId (clientside index of NetConnection)
+		NetCLUpgradeMessage		= 4,	// Added stateless handshake level NMT_Upgrade message, to allow Network CL checks to be permanently enabled
 
-		Latest			= SessionClientId
+		Latest					= NetCLUpgradeMessage
 	};
 
 	/** List of EHandshakeVersion's which break net compatibility (sorted from earliest to latest) */
@@ -58,8 +54,59 @@ namespace UE::Net
 		Response			= 2,
 		Ack					= 3,
 		RestartHandshake	= 4,
-		RestartResponse		= 5
+		RestartResponse		= 5,
+		VersionUpgrade		= 6,
+
+		Last = VersionUpgrade
+	};
+
+	/**
+	 * Flags/modifiers affecting how handshake packets are handled (e.g. restart handshake packets)
+	 */
+	enum class EHandshakePacketModifier : uint8
+	{
+		None				= 0x00,
+		RestartHandshake	= 0x01		// Restart handshake packet
+	};
+
+	ENUM_CLASS_FLAGS(EHandshakePacketModifier);
+
+
+	/**
+	 * Reasons for handshake failures
+	 */
+	enum class EHandshakeFailureReason
+	{
+		None,
+		WrongVersion		// Client/Server have incompatible versions
+	};
+
+	/**
+	 * Contains the information necessary for processing a failed stateless handshake, on the Game Thread.
+	 * Designed to be sent from the Receive Thread, to the Game Thread.
+	 */
+	struct FStatelessHandshakeFailureInfo
+	{
+		/** The reason the handshake failed */
+		EHandshakeFailureReason FailureReason = EHandshakeFailureReason::None;
+
+		/** The Network CL version of the remote side */
+		uint32 RemoteNetworkVersion = 0;
+
+		/** The net runtime features of the remote side */
+		EEngineNetworkRuntimeFeatures RemoteNetworkFeatures;
+
+
+	public:
+		FStatelessHandshakeFailureInfo();
 	};
+
+	/**
+	 * Callback function for handling a failed stateless handshake. Triggered when processing incoming packets.
+	 *
+	 * @param HandshakeFailureInfo		Information about the handshake failure.
+	 */
+	using FHandshakeFailureFunc = TUniqueFunction<void(FStatelessHandshakeFailureInfo HandshakeFailureInfo)>;
 }
 
 
@@ -71,6 +118,27 @@ namespace UE::Net
 class ENGINE_API StatelessConnectHandlerComponent : public HandlerComponent
 {
 	using EHandshakeVersion = UE::Net::EHandshakeVersion;
+	using EHandshakePacketType = UE::Net::EHandshakePacketType;
+	using EHandshakePacketModifier = UE::Net::EHandshakePacketModifier;
+	using FHandshakeFailureFunc = UE::Net::FHandshakeFailureFunc;
+
+private:
+	/** Common parameters for all 'Send*' functions which send packets to clients */
+	struct FCommonSendToClientParams
+	{
+		/** The address of the client to send the packet to */
+		const TSharedPtr<const FInternetAddr>& ClientAddress;
+		
+		/** The handshake format version to use when sending */
+		EHandshakeVersion HandshakeVersion;
+
+		/** The client-specified connection id */
+		uint32 ClientID;
+
+
+	public:
+		FCommonSendToClientParams(const TSharedPtr<const FInternetAddr>& InClientAddress, EHandshakeVersion InHandshakeVersion, uint32 InClientID);
+	};
 
 public:
 	/**
@@ -92,6 +160,13 @@ public:
 	 */
 	void InitFromConnectionless(StatelessConnectHandlerComponent* InConnectionlessHandler);
 
+	/**
+	 * Set a callback for notifying of handshake failure (clientside only).
+	 *
+	 * @param InHandshakeFailureFunc	The callback to use for notification.
+	 */
+	void SetHandshakeFailureCallback(FHandshakeFailureFunc&& InHandshakeFailureFunc);
+
 private:
 	/**
 	 * Constructs and sends the initial handshake packet, from the client to the server
@@ -103,13 +178,10 @@ private:
 	/**
 	 * Constructs and sends the server response to the initial connect packet, from the server to the client.
 	 *
-	 * @param ClientAddress						The address of the client to send the challenge to.
-	 * @param HandshakeVersion					The handshake format version to use when sending
+	 * @param CommonParams						Common parameters for 'Send*' functions to clients
 	 * @param ClientSentHandshakePacketCount	The number of handshake packets the client has sent (for debugging/packet-analysis)
-	 * @param InClientID						The client-specified connection id
 	 */
-	void SendConnectChallenge(TSharedPtr<const FInternetAddr> ClientAddress, EHandshakeVersion HandshakeVersion,
-								uint8 ClientSentHandshakePacketCount, uint32 InClientID);
+	void SendConnectChallenge(FCommonSendToClientParams CommonParams, uint8 ClientSentHandshakePacketCount);
 
 	/**
 	 * Constructs and sends the handshake challenge response packet, from the client to the server
@@ -124,63 +196,70 @@ private:
 	/**
 	 * Constructs and sends the server ack to a successful challenge response, from the server to the client.
 	 *
-	 * @param ClientAddress						The address of the client to send the ack to.
-	 * @param HandshakeVersion					The handshake format version to use when sending
+	 * @param CommonParams						Common parameters for 'Send*' functions to clients
 	 * @param ClientSentHandshakePacketCount	The number of handshake packets the client has sent (for debugging/packet-analysis)
-	 * @param InClientID						The client-specified connection id
 	 * @param InCookie							The cookie value to send
 	 */
-	void SendChallengeAck(TSharedPtr<const FInternetAddr> ClientAddress, EHandshakeVersion HandshakeVersion, uint8 ClientSentHandshakePacketCount,
-							uint32 InClientID, uint8 InCookie[COOKIE_BYTE_SIZE]);
+	void SendChallengeAck(FCommonSendToClientParams CommonParams, uint8 ClientSentHandshakePacketCount, uint8 InCookie[COOKIE_BYTE_SIZE]);
 
 	/**
 	 * Constructs and sends a request to resend the cookie, from the server to the client.
 	 *
-	 * @param ClientAddress		The address of the client to send the request to.
-	 * @param HandshakeVersion	The handshake format version to use when sending
-	 * @param InClientID		The client-specified connection id
+	 * @param CommonParams	Common parameters for 'Send*' functions to clients
+	 */
+	void SendRestartHandshakeRequest(FCommonSendToClientParams CommonParams);
+
+	/**
+	 * Constructs and sends a stateless handshake level NMT_Upgrade message
+	 *
+	 * @param CommonParams	Common parameters for 'Send*' functions to clients
 	 */
-	void SendRestartHandshakeRequest(const TSharedPtr<const FInternetAddr> ClientAddress, EHandshakeVersion HandshakeVersion,
-										uint32 InClientID);
+	void SendVersionUpgradeMessage(FCommonSendToClientParams CommonParams);
 
 
 	/**
-	 * Pads the handshake packet, to match the PacketBitAlignment of the PacketHandler, so that it will parse correctly.
+	 * Writes out the base/common parameters for all handshake packets, both to server and client
 	 *
-	 * @param HandshakePacket	The handshake packet to be aligned.
-	 * @param HandshakeVersion	The handshake format version to use when sending
+	 * @param HandshakePacket							The handshake packet being written.
+	 * @param HandshakePacketType						The type of handshake packet.
+	 * @param HandshakeVersion							The handshake format version to use when sending.
+	 * @param SentHandshakePacketCount_LocalOrRemote	The number of handshake packets sent (for debugging/packet-analysis).
+	 * @param ClientID									The client-specified connection id.
+	 * @param HandshakePacketModifier					Flags/modifiers affecting how the packet will be handled.
 	 */
-	void CapHandshakePacket(FBitWriter& HandshakePacket, EHandshakeVersion HandshakeVersion);
+	void BeginHandshakePacket(FBitWriter& HandshakePacket, EHandshakePacketType HandshakePacketType, EHandshakeVersion HandshakeVersion,
+								uint8 SentHandshakePacketCount_LocalOrRemote, uint32 ClientID,
+								EHandshakePacketModifier HandshakePacketModifier=EHandshakePacketModifier::None);
+
 
-public:
 	/**
-	 * Whether or not the specified connection address, has just passed the connection handshake challenge.
+	 * Sends a packet from the client to the server.
 	 *
-	 * @param Address	The address (including port, for UIpNetDriver) being checked
+	 * @param HandshakeVersion		The handshake format version to use when sending.
+	 * @param PacketType			The type of handshake packet.
+	 * @param Packet				The handshake packet being written.
 	 */
-	UE_DEPRECATED(4.22, "HasPassedChallenge must return bOutRestartedHandshake, in order to correctly process the challenge. Use the new version")
-	FORCEINLINE bool HasPassedChallenge(const FString& Address) const
-	{
-		bool bDud = false;
-		PRAGMA_DISABLE_DEPRECATION_WARNINGS
-		return HasPassedChallenge(Address, bDud);
-		PRAGMA_ENABLE_DEPRECATION_WARNINGS
-	}
+	void SendToServer(EHandshakeVersion HandshakeVersion, EHandshakePacketType PacketType, FBitWriter& Packet);
 
 	/**
-	 * Whether or not the specified connection address, has just passed the connection handshake challenge.
+	 * Sends a packet from the server to a client.
 	 *
-	 * @param Address					The address (including port, for UIpNetDriver) being checked
-	 * @param bOutRestartedHandshake	Whether or not the passed challenge, was a restarted handshake for an existing NetConnection
+	 * @param CommonParams			Common parameters for 'Send*' functions to clients.
+	 * @param PacketType			The type of handshake packet.
+	 * @param Packet				The handshake packet being written.
 	 */
-	UE_DEPRECATED(4.23, "HasPassedChallenge should pass a FInternetAddr pointer for faster address checking")
-	FORCEINLINE bool HasPassedChallenge(const FString& Address, bool& bOutRestartedHandshake) const
-	{
-		bOutRestartedHandshake = bRestartedHandshake;
+	void SendToClient(FCommonSendToClientParams CommonParams, EHandshakePacketType PacketType, FBitWriter& Packet);
 
-		return (LastChallengeSuccessAddress.IsValid()) ? LastChallengeSuccessAddress->ToString(true) == Address : false;
-	}
 
+	/**
+	 * Pads the handshake packet, to match the PacketBitAlignment of the PacketHandler, so that it will parse correctly.
+	 *
+	 * @param HandshakePacket	The handshake packet to be aligned.
+	 * @param HandshakeVersion	The handshake format version to use when sending
+	 */
+	void CapHandshakePacket(FBitWriter& HandshakePacket, EHandshakeVersion HandshakeVersion);
+
+public:
 	/**
 	 * Whether or not the specified connection address, has just passed the connection handshake challenge.
 	 *
@@ -273,6 +352,9 @@ private:
 		/** The Network CL version of the remote side */
 		uint32 RemoteNetworkVersion = 0;
 
+		/** The net runtime features of the remote side */
+		EEngineNetworkRuntimeFeatures RemoteNetworkFeatures;
+
 		/** The type of handshake packet */
 		UE::Net::EHandshakePacketType HandshakePacketType = UE::Net::EHandshakePacketType::InitialPacket;
 
@@ -293,6 +375,10 @@ private:
 
 		/** If this is a restart handshake challenge response, this is the original handshake's cookie */
 		uint8 OrigCookie[COOKIE_BYTE_SIZE] = {};
+
+
+	public:
+		FParsedHandshakeData();
 	};
 
 
@@ -324,7 +410,7 @@ private:
 	 * @param TimeStamp			The serverside timestamp
 	 * @param OutCookie			Outputs the generated cookie value.
 	 */
-	void GenerateCookie(TSharedPtr<const FInternetAddr> ClientAddress, uint8 SecretId, double Timestamp, uint8 (&OutCookie)[COOKIE_BYTE_SIZE]) const;
+	void GenerateCookie(const TSharedPtr<const FInternetAddr>& ClientAddress, uint8 SecretId, double Timestamp, uint8 (&OutCookie)[COOKIE_BYTE_SIZE]) const;
 
 	/**
 	 * Generates a new HandshakeSecret value
@@ -385,6 +471,15 @@ private:
 	uint32 ValidationLogCounter = 0;
 #endif
 
+	/** The last time an 'upgrade message' send period began */
+	double LastUpgradeMessagePeriodStart = 0.0;
+
+	/** The number of upgrade messages sent in the current send period (for rate limiting) */
+	uint32 UpgradeMessageCounter = 0;
+
+	/** Callback for processing handshake failures */
+	FHandshakeFailureFunc HandshakeFailureCallback;
+
 
 	/** Clientside variables */
 
