Fix for static analysis warning

diff --git a/Engine/Source/Runtime/Engine/Private/Animation/BlendSpaceHelpers.cpp b/Engine/Source/Runtime/Engine/Private/Animation/BlendSpaceHelpers.cpp
--- a/Engine/Source/Runtime/Engine/Private/Animation/BlendSpaceHelpers.cpp
+++ b/Engine/Source/Runtime/Engine/Private/Animation/BlendSpaceHelpers.cpp
@@ -697,8 +697,8 @@ void FBlendSpaceGrid::GenerateGridElements(const TArray<FVertex>& SamplePoints,
 			else
 			{
 				// Work through all the edges and find the one with a point closest to this grid position.
-				int32 ClosestTriangleIndex = -1;
-				int32 ClosestEdgeIndex = -1;
+				int32 ClosestTriangleIndex = 0;
+				int32 ClosestEdgeIndex = 0;
 				double ClosestDistance = UE_DOUBLE_BIG_NUMBER;
 				double ClosestT = 0.0;
 
