Fix ensures in dormant refcount code caused when an invalid replicator in a channel's ReplicationMap gets removed.
This was caused by a safeguard were we tried to remove any dormant replicator of the pointer reused from a now invalid sub object. This safeguard was useless for two reasons:
   a) Dormant replicators are always validated before being returned. No need to manually discard them.
   b) A dormant replicator cannot be stored in both the ReplicationMap and the connection's dormant list. It's one or the other not both.

diff --git a/Engine/Source/Runtime/Engine/Classes/Engine/ActorChannel.h b/Engine/Source/Runtime/Engine/Classes/Engine/ActorChannel.h
--- a/Engine/Source/Runtime/Engine/Classes/Engine/ActorChannel.h
+++ b/Engine/Source/Runtime/Engine/Classes/Engine/ActorChannel.h
@@ -359,15 +359,24 @@ public:
 
 protected:
 
+	/** Attempts to find a valid, non-dormant replicator for the given object. */
+	TSharedRef<FObjectReplicator>* FindReplicator(UObject* Obj);
+
 	/**
-	 * Attempts to find a valid, non-dormant replicator for the given object.
-	 *
-	 * @param Obj				The object whose replicator to find.
-	 * @param bOutFoundInvalid	Indicates we found a replicator, but it was invalid.
-	 *
-	 * @return A replicator, if one was found.
-	 */
-	TSharedRef<FObjectReplicator>* FindReplicator(UObject* Obj, bool* bOutFoundInvalid=nullptr);
+	* Attempts to find a valid, non-dormant replicator for the given object.
+	*
+	* @param Obj				The object whose replicator to find.
+	* @param bOutFoundInvalid	Indicates we found a replicator, but it was invalid.
+	*
+	* @return A replicator, if one was found.
+	*/
+	UE_DEPRECATED(5.2, "This function has been deprecated in favor of the one with a single parameter")
+	TSharedRef<FObjectReplicator>* FindReplicator(UObject* Obj, bool* bOutFoundInvalid);
+
+	/** 
+	* Creates a new object replicator or reuses a replicator if it was stored for dormancy in the Connection.
+	*/
+	 TSharedRef<FObjectReplicator>& CreateReplicator(UObject* Obj);
 
 	/**
 	 * Creates a new object replicator.
@@ -382,6 +391,7 @@ protected:
 	 *
 	 * @return The newly created replicator.
 	 */
+	UE_DEPRECATED(5.2, "This function has been deprecated in favor of the one with a single parameter")
 	TSharedRef<FObjectReplicator>& CreateReplicator(UObject* Obj, bool bCheckDormantReplicators);
 
 	/**
diff --git a/Engine/Source/Runtime/Engine/Private/DataChannel.cpp b/Engine/Source/Runtime/Engine/Private/DataChannel.cpp
--- a/Engine/Source/Runtime/Engine/Private/DataChannel.cpp
+++ b/Engine/Source/Runtime/Engine/Private/DataChannel.cpp
@@ -4943,6 +4943,13 @@ FObjectReplicator & UActorChannel::GetActorReplicationData()
 	return *ActorReplicator;
 }
 
+TSharedRef<FObjectReplicator>* UActorChannel::FindReplicator(UObject* Obj)
+{
+	PRAGMA_DISABLE_DEPRECATION_WARNINGS
+	return FindReplicator(Obj, nullptr);
+	PRAGMA_ENABLE_DEPRECATION_WARNINGS
+}
+
 TSharedRef<FObjectReplicator>* UActorChannel::FindReplicator(UObject* Obj, bool* bOutFoundInvalid)
 {
 	CONDITIONAL_SCOPE_CYCLE_COUNTER(Stat_ActorChanFindOrCreateRep, CVarNetEnableDetailedScopeCounters.GetValueOnAnyThread() > 0);
@@ -4970,6 +4977,14 @@ TSharedRef<FObjectReplicator>* UActorChannel::FindReplicator(UObject* Obj, bool*
 	return ReplicatorRefPtr;
 }
 
+TSharedRef<FObjectReplicator>& UActorChannel::CreateReplicator(UObject* Obj)
+{
+	PRAGMA_DISABLE_DEPRECATION_WARNINGS
+	// We always want to reuse the dormant replicator if it was stored in the connection previously.
+	return CreateReplicator(Obj, true);
+	PRAGMA_ENABLE_DEPRECATION_WARNINGS
+}
+
 TSharedRef<FObjectReplicator>& UActorChannel::CreateReplicator(UObject* Obj, bool bCheckDormantReplicators)
 {
 	CONDITIONAL_SCOPE_CYCLE_COUNTER(Stat_ActorChanFindOrCreateRep, CVarNetEnableDetailedScopeCounters.GetValueOnAnyThread() > 0);
@@ -4994,6 +5009,7 @@ TSharedRef<FObjectReplicator>& UActorChannel::CreateReplicator(UObject* Obj, boo
 		}
 #endif
 	}
+	// Failsafe thats makes sure to discard any dormant replicator until we completely delete the option to ignore them.
 	else
 	{
 		// Discard dormant replicator if there was one that we're not using
@@ -5027,8 +5043,6 @@ TSharedRef<FObjectReplicator>& UActorChannel::CreateReplicator(UObject* Obj, boo
 	// Add to the replication map
 	TSharedRef<FObjectReplicator>& NewRef = ReplicationMap.Add(Obj, NewReplicator.ToSharedRef());
 
-
-
 	// Start replicating with this replicator
 	NewRef->StartReplicating(this);
 	return NewRef;
@@ -5036,8 +5050,7 @@ TSharedRef<FObjectReplicator>& UActorChannel::CreateReplicator(UObject* Obj, boo
 
 TSharedRef<FObjectReplicator>& UActorChannel::FindOrCreateReplicator(UObject* Obj, bool* bOutCreated)
 {
-	bool bFoundInvalidReplicator = false;
-	TSharedRef<FObjectReplicator>* ReplicatorRefPtr = FindReplicator(Obj, &bFoundInvalidReplicator);
+	TSharedRef<FObjectReplicator>* ReplicatorRefPtr = FindReplicator(Obj);
 
 	// This should only be false if we found the replicator in the ReplicationMap
 	// If we pickup the replicator from the DormantReplicatorMap we treat it as it has been created.
@@ -5048,7 +5061,7 @@ TSharedRef<FObjectReplicator>& UActorChannel::FindOrCreateReplicator(UObject* Ob
 
 	if (!ReplicatorRefPtr)
 	{
-		ReplicatorRefPtr = &CreateReplicator(Obj, !bFoundInvalidReplicator);
+		ReplicatorRefPtr = &CreateReplicator(Obj);
 	}
 
 	return *ReplicatorRefPtr;
@@ -5143,8 +5156,7 @@ bool UActorChannel::WriteSubObjectInBunch(UObject* Obj, FOutBunch& Bunch, FRepli
 		return ReplicateSubobjectCustom(Obj, Bunch, RepFlags);
 	}
 
-	bool bFoundInvalidReplicator = false;
-	TSharedRef<FObjectReplicator>* FoundReplicator = FindReplicator(Obj, &bFoundInvalidReplicator);
+	TSharedRef<FObjectReplicator>* FoundReplicator = FindReplicator(Obj);
 	const bool bFoundReplicator = (FoundReplicator != nullptr);
 
 	// Special case for replay checkpoints where the replicator could already exist but we still need to consider this a new object even if it has an empty layout
@@ -5172,7 +5184,7 @@ bool UActorChannel::WriteSubObjectInBunch(UObject* Obj, FOutBunch& Bunch, FRepli
 	bool NewSubobject = false;
 	
 	FReplicationFlags ObjRepFlags = RepFlags;
-	TSharedRef<FObjectReplicator>& ObjectReplicator = !bFoundReplicator ? CreateReplicator(Obj, !bFoundInvalidReplicator) : *FoundReplicator;
+	TSharedRef<FObjectReplicator>& ObjectReplicator = !bFoundReplicator ? CreateReplicator(Obj) : *FoundReplicator;
 
 	if (!bFoundReplicator || bNewToReplay)
 	{
diff --git a/Engine/Source/Runtime/Engine/Private/Net/NetDormantHolder.cpp b/Engine/Source/Runtime/Engine/Private/Net/NetDormantHolder.cpp
--- a/Engine/Source/Runtime/Engine/Private/Net/NetDormantHolder.cpp
+++ b/Engine/Source/Runtime/Engine/Private/Net/NetDormantHolder.cpp
@@ -98,10 +98,11 @@ void FDormantReplicatorHolder::StoreReplicator(AActor* DormantActor, UObject* Re
 	ActorReplicators.DormantReplicators.Add(FDormantObjectReplicator(ReplicatedObject, ObjectReplicator));
 }
 
-void FDormantReplicatorHolder::RemoveStoredReplicator(AActor* DormantActor, FObjectKey ReplicatedObjectKey)
+bool FDormantReplicatorHolder::RemoveStoredReplicator(AActor* DormantActor, FObjectKey ReplicatedObjectKey)
 {
 	FSetElementId Index = ActorReplicatorSet.FindId(DormantActor);
-	if (Index.IsValidId())
+	const bool bIsFound = Index.IsValidId();
+	if (bIsFound)
 	{
 		ActorReplicatorSet[Index].DormantReplicators.Remove(ReplicatedObjectKey);
 
@@ -111,6 +112,8 @@ void FDormantReplicatorHolder::RemoveStoredReplicator(AActor* DormantActor, FObj
 			ActorReplicatorSet.Remove(Index);
 		}
 	}
+
+	return bIsFound;
 }
 
 void FDormantReplicatorHolder::CleanupAllReplicatorsOfActor(AActor* DormantActor)
diff --git a/Engine/Source/Runtime/Engine/Private/NetConnection.cpp b/Engine/Source/Runtime/Engine/Private/NetConnection.cpp
--- a/Engine/Source/Runtime/Engine/Private/NetConnection.cpp
+++ b/Engine/Source/Runtime/Engine/Private/NetConnection.cpp
@@ -5104,11 +5104,11 @@ void UNetConnection::RemoveDormantReplicator(AActor* Actor, UObject* Object)
 {
 	const FObjectKey ObjectKey = Object;
 
-	DormantReplicatorSet.RemoveStoredReplicator(Actor, ObjectKey);
+	const bool bWasStored = DormantReplicatorSet.RemoveStoredReplicator(Actor, ObjectKey);
 
 #if UE_REPLICATED_OBJECT_REFCOUNTING
 	// If this is a subobject (and not the main actor) remove the reference
-	if (Object != Actor)
+	if (bWasStored && Object != Actor)
 	{
 		if (Driver)
 		{
diff --git a/Engine/Source/Runtime/Engine/Public/Net/NetDormantHolder.h b/Engine/Source/Runtime/Engine/Public/Net/NetDormantHolder.h
--- a/Engine/Source/Runtime/Engine/Public/Net/NetDormantHolder.h
+++ b/Engine/Source/Runtime/Engine/Public/Net/NetDormantHolder.h
@@ -128,8 +128,9 @@ struct FDormantReplicatorHolder
 	* 
 	* @param DormantActor The dormant actor that owns the replicated object
 	* @param ReplicatedObjectKey The object key to a replicated object that may have stored it's object replicator here
+	* @return Return true if we did find the replicated object and removed it
 	*/
-	void RemoveStoredReplicator(AActor* DormantActor, FObjectKey ReplicatedObjectKey);
+	bool RemoveStoredReplicator(AActor* DormantActor, FObjectKey ReplicatedObjectKey);
 
 	/**
 	* Remove the references to all the object replicators tied to the given actor
