[shaders] free global shadermaps on shutdown; this needs to happen prior to shader code library shutdown (as the latter will check and fire warnings if any shadermaps still have references).

diff --git a/Engine/Source/Runtime/Engine/Private/ShaderCompiler/ShaderCompiler.cpp b/Engine/Source/Runtime/Engine/Private/ShaderCompiler/ShaderCompiler.cpp
--- a/Engine/Source/Runtime/Engine/Private/ShaderCompiler/ShaderCompiler.cpp
+++ b/Engine/Source/Runtime/Engine/Private/ShaderCompiler/ShaderCompiler.cpp
@@ -50,6 +50,7 @@
 #include "UObject/DevObjectVersion.h"
 #include "UObject/UObjectIterator.h"
 #include "Math/UnitConversion.h"
+#include "UnrealEngine.h"
 
 #if WITH_EDITOR
 #include "UObject/ArchiveCookContext.h"
@@ -7822,6 +7823,20 @@ void CompileGlobalShaderMap(bool bRefreshShaderMap)
 	CompileGlobalShaderMap(GMaxRHIFeatureLevel, bRefreshShaderMap);
 }
 
+void ShutdownGlobalShaderMap()
+{
+	// at the point this function is called (during the shutdown process) we do not expect any outstanding work that could potentially be still referencing
+	// global shaders, so we are not deferring the deletion (via GGlobalShaderMap_DeferredDeleteCopy) like we do during the shader recompilation.
+	EShaderPlatform Platform = GShaderPlatformForFeatureLevel[GMaxRHIFeatureLevel];
+	if (GGlobalShaderMap[Platform] != nullptr)
+	{
+		GGlobalShaderMap[Platform]->ReleaseAllSections();
+
+		delete GGlobalShaderMap[Platform];
+		GGlobalShaderMap[Platform] = nullptr;
+	}
+}
+
 void ReloadGlobalShaders()
 {
 	UE_LOG(LogShaders, Display, TEXT("Reloading global shaders..."));
diff --git a/Engine/Source/Runtime/Engine/Public/ShaderCompiler.h b/Engine/Source/Runtime/Engine/Public/ShaderCompiler.h
--- a/Engine/Source/Runtime/Engine/Public/ShaderCompiler.h
+++ b/Engine/Source/Runtime/Engine/Public/ShaderCompiler.h
@@ -1396,6 +1396,7 @@ extern ENGINE_API void CompileGlobalShaderMap(bool bRefreshShaderMap=false);
 extern ENGINE_API void CompileGlobalShaderMap(ERHIFeatureLevel::Type InFeatureLevel, bool bRefreshShaderMap=false);
 extern ENGINE_API void CompileGlobalShaderMap(EShaderPlatform Platform, bool bRefreshShaderMap = false);
 extern ENGINE_API void CompileGlobalShaderMap(EShaderPlatform Platform, const ITargetPlatform* TargetPlatform, bool bRefreshShaderMap);
+extern ENGINE_API void ShutdownGlobalShaderMap();
 
 extern ENGINE_API const FString& GetGlobalShaderMapDDCKey();
 
diff --git a/Engine/Source/Runtime/Launch/Private/LaunchEngineLoop.cpp b/Engine/Source/Runtime/Launch/Private/LaunchEngineLoop.cpp
--- a/Engine/Source/Runtime/Launch/Private/LaunchEngineLoop.cpp
+++ b/Engine/Source/Runtime/Launch/Private/LaunchEngineLoop.cpp
@@ -4979,6 +4979,10 @@ void FEngineLoop::Exit()
 	// Disable the PSO cache
 	FShaderPipelineCache::Shutdown();
 
+	// Free the global shader map, needs to happen before FShaderCodeLibrary::Shutdown to avoid warnings
+	// about leaking RHI references.
+	ShutdownGlobalShaderMap();
+
 	// Close shader code map, if any
 	FShaderCodeLibrary::Shutdown();
 
