Fixed the Typed Element Data Storage not correctly managing actors during undo/redo.

When an undo is triggered on an actor addition the Typed Element Data Storage would not remove the actor from the database or add it back in when an undo was called. As a result the actor counter was incorrect, but in the future this would cause crashes when the actor was used during processing of the Typed Element Data Storage.

diff --git a/Engine/Source/Runtime/Engine/Private/ActorEditor.cpp b/Engine/Source/Runtime/Engine/Private/ActorEditor.cpp
--- a/Engine/Source/Runtime/Engine/Private/ActorEditor.cpp
+++ b/Engine/Source/Runtime/Engine/Private/ActorEditor.cpp
@@ -24,6 +24,7 @@
 
 #include "Editor.h"
 #include "ActorTransactionAnnotation.h"
+#include "Elements/Framework/EngineElementsLibrary.h"
 #include "WorldPartition/DataLayer/IDataLayerEditorModule.h"
 #include "LevelInstance/LevelInstanceSubsystem.h"
 #include "ActorFolder.h"
@@ -695,10 +696,12 @@ bool AActor::InternalPostEditUndo()
 
 		// notify navigation system
 		FNavigationSystem::UpdateActorAndComponentData(*this);
+		UEngineElementsLibrary::RegisterActorElement(this);
 	}
 	else
 	{
 		FNavigationSystem::RemoveActorData(*this);
+		UEngineElementsLibrary::UnregisterActorElement(this);
 	}
 
 	// This is a normal undo, so call super
diff --git a/Engine/Source/Runtime/Engine/Private/Elements/Framework/EngineElementsLibrary.cpp b/Engine/Source/Runtime/Engine/Private/Elements/Framework/EngineElementsLibrary.cpp
--- a/Engine/Source/Runtime/Engine/Private/Elements/Framework/EngineElementsLibrary.cpp
+++ b/Engine/Source/Runtime/Engine/Private/Elements/Framework/EngineElementsLibrary.cpp
@@ -360,12 +360,7 @@ TTypedElementOwner<FActorElementData> UEngineElementsLibrary::CreateActorElement
 		ActorElement.GetDataChecked().Actor = const_cast<AActor*>(InActor);
 	}
 	
-	UTypedElementRegistry* Registry = UTypedElementRegistry::GetInstance();
-	ITypedElementDataStorageCompatibilityInterface* Storage = Registry->GetMutableDataStorageCompatibility();
-	if (Storage)
-	{
-		Storage->AddCompatibleObject(const_cast<AActor*>(InActor));
-	}
+	RegisterActorElement(InActor);
 
 	return ActorElement;
 }
@@ -382,6 +377,16 @@ void UEngineElementsLibrary::DestroyActorElement(const AActor* InActor, TTypedEl
 	}
 }
 
+void UEngineElementsLibrary::RegisterActorElement(const AActor* InActor)
+{
+	UTypedElementRegistry* Registry = UTypedElementRegistry::GetInstance();
+	ITypedElementDataStorageCompatibilityInterface* Storage = Registry->GetMutableDataStorageCompatibility();
+	if (Storage)
+	{
+		Storage->AddCompatibleObject(const_cast<AActor*>(InActor));
+	}
+}
+
 void UEngineElementsLibrary::UnregisterActorElement(const AActor* InActor)
 {
 	UTypedElementRegistry* Registry = UTypedElementRegistry::GetInstance();
diff --git a/Engine/Source/Runtime/Engine/Public/Elements/Framework/EngineElementsLibrary.h b/Engine/Source/Runtime/Engine/Public/Elements/Framework/EngineElementsLibrary.h
--- a/Engine/Source/Runtime/Engine/Public/Elements/Framework/EngineElementsLibrary.h
+++ b/Engine/Source/Runtime/Engine/Public/Elements/Framework/EngineElementsLibrary.h
@@ -42,6 +42,7 @@ public:
 
 	static TTypedElementOwner<FActorElementData> CreateActorElement(const AActor* InActor);
 	static void DestroyActorElement(const AActor* InActor, TTypedElementOwner<FActorElementData>& InOutActorElement);
+	static void RegisterActorElement(const AActor* InActor);
 	static void UnregisterActorElement(const AActor* InActor);
 #if WITH_EDITOR
 	static void CreateEditorActorElement(const AActor* Actor);
