Update the help text for Affect Indirect Lighting on primitives (was refering to Light Propagation Volumes which no longer exist, but it now works with Lumen and the Path Tracer)

Also allow this field to be set independently from CastShadows as the two don't seem to be linked in any way anymore

diff --git a/Engine/Source/Runtime/Engine/Classes/Components/PrimitiveComponent.h b/Engine/Source/Runtime/Engine/Classes/Components/PrimitiveComponent.h
--- a/Engine/Source/Runtime/Engine/Classes/Components/PrimitiveComponent.h
+++ b/Engine/Source/Runtime/Engine/Classes/Components/PrimitiveComponent.h
@@ -482,8 +482,8 @@ public:
 	UPROPERTY(EditAnywhere, BlueprintReadOnly, Category=Lighting, AdvancedDisplay)
 	uint8 bEmissiveLightSource:1;
 
-	/** Controls whether the primitive should inject light into the Light Propagation Volume.  This flag is only used if CastShadow is true. */
-	UPROPERTY(EditAnywhere, BlueprintReadOnly, Category=Lighting, AdvancedDisplay, meta=(EditCondition="CastShadow"))
+	/** Controls whether the primitive should influence indirect lighting. */
+	UPROPERTY(EditAnywhere, BlueprintReadOnly, Category=Lighting, AdvancedDisplay)
 	uint8 bAffectDynamicIndirectLighting:1;
 
 	/** Controls whether the primitive should affect indirect lighting when hidden. This flag is only used if bAffectDynamicIndirectLighting is true. */
