Static analysis (PVS) fix.

diff --git a/Engine/Source/Runtime/Engine/Private/PacketHandlers/StatelessConnectHandlerComponent.cpp b/Engine/Source/Runtime/Engine/Private/PacketHandlers/StatelessConnectHandlerComponent.cpp
--- a/Engine/Source/Runtime/Engine/Private/PacketHandlers/StatelessConnectHandlerComponent.cpp
+++ b/Engine/Source/Runtime/Engine/Private/PacketHandlers/StatelessConnectHandlerComponent.cpp
@@ -1514,7 +1514,7 @@ bool StatelessConnectHandlerComponent::ParseHandshakePacket(FBitReader& Packet,
 	const bool bMaybeRestartResponsePacket = MaxBitsLeftExclRestartResponse >= MinRandomBits && MinBitsLeftExclRestartResponse <= MaxRandomBits;
 	const bool bMaybeVersionUpgradePacket = MaxBitsLeftExclVersionUpgrade >= MinRandomBits && MinBitsLeftExclVersionUpgrade <= MaxRandomBits;
 
-	static_assert(BASE_PACKET_SIZE_BITS == 82 && VERSION_UPGRADE_SIZE_BITS == 82, "MaxBitsLeftExclVersionUpgrade needs to be updated.");
+	static_assert(BASE_PACKET_SIZE_BITS == 82 && VERSION_UPGRADE_SIZE_BITS == 82, "MaxBitsLeftExclVersionUpgrade needs to be updated."); // -V501
 
 	OutResult.bRestartHandshake = !!Packet.ReadBit();
 
