Disabled 'resync to sync group' for sequence and blend space evaluators as they generally want to be driven with explicit time

diff --git a/Engine/Source/Runtime/AnimGraphRuntime/Private/AnimNodes/AnimNode_RandomPlayer.cpp b/Engine/Source/Runtime/AnimGraphRuntime/Private/AnimNodes/AnimNode_RandomPlayer.cpp
--- a/Engine/Source/Runtime/AnimGraphRuntime/Private/AnimNodes/AnimNode_RandomPlayer.cpp
+++ b/Engine/Source/Runtime/AnimGraphRuntime/Private/AnimNodes/AnimNode_RandomPlayer.cpp
@@ -236,7 +236,7 @@ void FAnimNode_RandomPlayer::Update_AnyThread(const FAnimationUpdateContext& Con
 		NextData = &GetPlayData(ERandomDataIndexType::Next);
 	}
 
-	FAnimTickRecord TickRecord(CurrentData->Entry->Sequence, true, CurrentData->PlayRate, CurrentData->BlendWeight, CurrentData->CurrentPlayTime, CurrentData->MarkerTickRecord);
+	FAnimTickRecord TickRecord(CurrentData->Entry->Sequence, true, CurrentData->PlayRate, false, CurrentData->BlendWeight, CurrentData->CurrentPlayTime, CurrentData->MarkerTickRecord);
 	TickRecord.DeltaTimeRecord = &CurrentData->DeltaTimeRecord;
 	TickRecord.GatherContextData(Context);
 
@@ -247,7 +247,7 @@ void FAnimNode_RandomPlayer::Update_AnyThread(const FAnimationUpdateContext& Con
 
 	if (FAnimationRuntime::HasWeight(NextData->BlendWeight))
 	{
-		FAnimTickRecord NextTickRecord(NextData->Entry->Sequence, true, NextData->PlayRate, NextData->BlendWeight, NextData->CurrentPlayTime, NextData->MarkerTickRecord);
+		FAnimTickRecord NextTickRecord(NextData->Entry->Sequence, true, NextData->PlayRate, false, NextData->BlendWeight, NextData->CurrentPlayTime, NextData->MarkerTickRecord);
 		NextTickRecord.DeltaTimeRecord = &NextData->DeltaTimeRecord;
 		NextTickRecord.GatherContextData(Context);
 
diff --git a/Engine/Source/Runtime/AnimGraphRuntime/Private/AnimNodes/AnimNode_SequenceEvaluator.cpp b/Engine/Source/Runtime/AnimGraphRuntime/Private/AnimNodes/AnimNode_SequenceEvaluator.cpp
--- a/Engine/Source/Runtime/AnimGraphRuntime/Private/AnimNodes/AnimNode_SequenceEvaluator.cpp
+++ b/Engine/Source/Runtime/AnimGraphRuntime/Private/AnimNodes/AnimNode_SequenceEvaluator.cpp
@@ -67,12 +67,12 @@ void FAnimNode_SequenceEvaluatorBase::UpdateAssetPlayer(const FAnimationUpdateCo
 			const float DeltaTime = Context.GetDeltaTime();
 			const float RateScale = CurrentSequence->RateScale;
 			const float PlayRate = FMath::IsNearlyZero(DeltaTime) || FMath::IsNearlyZero(RateScale) ? 0.f : (TimeJump / (DeltaTime * RateScale));
-			CreateTickRecordForNode(Context, CurrentSequence, GetShouldLoop(), PlayRate);
+			CreateTickRecordForNode(Context, CurrentSequence, GetShouldLoop(), PlayRate, true);
 		}
 		else
 		{
 			InternalTimeAccumulator = CurrentExplicitTime;
-			CreateTickRecordForNode(Context, CurrentSequence, GetShouldLoop(), 0);
+			CreateTickRecordForNode(Context, CurrentSequence, GetShouldLoop(), 0, true);
 		}
 	}
 
diff --git a/Engine/Source/Runtime/Engine/Classes/Animation/AnimNode_AssetPlayerBase.h b/Engine/Source/Runtime/Engine/Classes/Animation/AnimNode_AssetPlayerBase.h
--- a/Engine/Source/Runtime/Engine/Classes/Animation/AnimNode_AssetPlayerBase.h
+++ b/Engine/Source/Runtime/Engine/Classes/Animation/AnimNode_AssetPlayerBase.h
@@ -30,7 +30,7 @@ struct ENGINE_API FAnimNode_AssetPlayerBase : public FAnimNode_AssetPlayerReleva
 	virtual void UpdateAssetPlayer(const FAnimationUpdateContext& Context) {};
 
 	// Create a tick record for this node
-	void CreateTickRecordForNode(const FAnimationUpdateContext& Context, UAnimSequenceBase* Sequence, bool bLooping, float PlayRate);
+	void CreateTickRecordForNode(const FAnimationUpdateContext& Context, UAnimSequenceBase* Sequence, bool bLooping, float PlayRate, bool bIsEvaluator);
 
 	// Get the sync group name we are using
 	virtual FName GetGroupName() const { return NAME_None; }
diff --git a/Engine/Source/Runtime/Engine/Classes/Animation/AnimationAsset.h b/Engine/Source/Runtime/Engine/Classes/Animation/AnimationAsset.h
--- a/Engine/Source/Runtime/Engine/Classes/Animation/AnimationAsset.h
+++ b/Engine/Source/Runtime/Engine/Classes/Animation/AnimationAsset.h
@@ -372,6 +372,7 @@ struct FAnimTickRecord
 	float RootMotionWeightModifier = 1.0f;
 
 	bool bLooping = false;
+	bool bIsEvaluator = false;
 	const UMirrorDataTable* MirrorDataTable = nullptr;
 
 	TSharedPtr<TArray<TUniquePtr<const UE::Anim::IAnimNotifyEventContextDataInterface>>> ContextData;
@@ -386,7 +387,6 @@ struct FAnimTickRecord
 			float  BlendSpacePositionX;
 			float  BlendSpacePositionY;
 			bool   bTeleportToTime;
-			bool   bIsEvaluator;
 		} BlendSpace;
 
 		struct
@@ -412,27 +412,19 @@ private:
 	void AllocateContextDataContainer();
 
 public:
-	FAnimTickRecord()
-		: SourceAsset(nullptr)
-		, TimeAccumulator(nullptr)
-		, PlayRateMultiplier(1.f)
-		, EffectiveBlendWeight(0.f)
-		, RootMotionWeightModifier(1.f)
-		, bLooping(false)
-		, DeltaTimeRecord(nullptr)
-		, MarkerTickRecord(nullptr)
-		, bCanUseMarkerSync(false)
-		, LeaderScore(0.f)
-	{
-	}
+	FAnimTickRecord() = default;
 
 	// Create a tick record for an anim sequence
+	UE_DEPRECATED(5.2, "Please use the anim sequence FAnimTickRecord constructor which adds bInIsEvaluator (defaulted to false)")
 	ENGINE_API FAnimTickRecord(UAnimSequenceBase* InSequence, bool bInLooping, float InPlayRate, float InFinalBlendWeight, float& InCurrentTime, FMarkerTickRecord& InMarkerTickRecord);
 
+	// Create a tick record for an anim sequence
+	ENGINE_API FAnimTickRecord(UAnimSequenceBase* InSequence, bool bInLooping, float InPlayRate, bool bInIsEvaluator, float InFinalBlendWeight, float& InCurrentTime, FMarkerTickRecord& InMarkerTickRecord);
+
 	// Create a tick record for a blendspace
 	ENGINE_API FAnimTickRecord(
 		UBlendSpace* InBlendSpace, const FVector& InBlendInput, TArray<FBlendSampleData>& InBlendSampleDataCache, FBlendFilter& InBlendFilter, bool bInLooping, 
-		float InPlayRate, bool bShouldTeleportToTime, bool bIsEvaluator, float InFinalBlendWeight, float& InCurrentTime, FMarkerTickRecord& InMarkerTickRecord);
+		float InPlayRate, bool bShouldTeleportToTime, bool bInIsEvaluator, float InFinalBlendWeight, float& InCurrentTime, FMarkerTickRecord& InMarkerTickRecord);
 
 	// Create a tick record for a montage
 	UE_DEPRECATED(5.0, "Please use the montage FAnimTickRecord constructor which removes InPreviousPosition and InMoveDelta")
diff --git a/Engine/Source/Runtime/Engine/Private/Animation/AnimInstanceProxy.cpp b/Engine/Source/Runtime/Engine/Private/Animation/AnimInstanceProxy.cpp
--- a/Engine/Source/Runtime/Engine/Private/Animation/AnimInstanceProxy.cpp
+++ b/Engine/Source/Runtime/Engine/Private/Animation/AnimInstanceProxy.cpp
@@ -761,6 +761,7 @@ void FAnimInstanceProxy::MakeSequenceTickRecord(FAnimTickRecord& TickRecord, cla
 	TickRecord.PlayRateMultiplier = PlayRate;
 	TickRecord.EffectiveBlendWeight = FinalBlendWeight;
 	TickRecord.bLooping = bLooping;
+	TickRecord.bIsEvaluator = false;
 }
 
 void FAnimInstanceProxy::MakeBlendSpaceTickRecord(
@@ -772,7 +773,6 @@ void FAnimInstanceProxy::MakeBlendSpaceTickRecord(
 	TickRecord.BlendSpace.BlendSpacePositionY = BlendInput.Y;
 	// This way of making a tick record is deprecated, so just set to defaults here rather than changing the API
 	TickRecord.BlendSpace.bTeleportToTime = false; 
-	TickRecord.BlendSpace.bIsEvaluator = false;
 	TickRecord.BlendSpace.BlendSampleDataCache = &BlendSampleDataCache;
 	TickRecord.BlendSpace.BlendFilter = &BlendFilter;
 	TickRecord.TimeAccumulator = &CurrentTime;
@@ -780,6 +780,7 @@ void FAnimInstanceProxy::MakeBlendSpaceTickRecord(
 	TickRecord.PlayRateMultiplier = PlayRate;
 	TickRecord.EffectiveBlendWeight = FinalBlendWeight;
 	TickRecord.bLooping = bLooping;
+	TickRecord.bIsEvaluator = false;
 }
 
 void FAnimInstanceProxy::MakePoseAssetTickRecord(FAnimTickRecord& TickRecord, class UPoseAsset* PoseAsset, float FinalBlendWeight) const
@@ -3242,8 +3243,39 @@ void FAnimInstanceProxy::RegisterWatchedPose(const FCompactPose& Pose, int32 Lin
 						PoseWatch.Object = GetAnimInstanceObject();
 						PoseWatch.PoseWatch->SetIsNodeEnabled(true);
 
-						const TArray<FTransform, FAnimStackAllocator>& BoneTransforms = Pose.GetBones();
-						const TArray<FBoneIndexType>& TmpRequiredBones = Pose.GetBoneContainer().GetBoneIndicesArray();
+						/*
+						for (FCompactPoseBoneIndex BoneIndex : Pose.ForEachBoneIndex())
+						{
+							FMeshPoseBoneIndex MeshBoneIndex = Pose.GetBoneContainer().MakeMeshPoseIndex(BoneIndex);
+
+							int32 ParentIndex = Pose.GetBoneContainer().GetParentBoneIndex(MeshBoneIndex.GetInt());
+
+							if (ParentIndex == INDEX_NONE)
+							{
+								WorldTransforms[MeshBoneIndex.GetInt()] = Pose[BoneIndex] * MeshComponent->GetComponentTransform();
+							}
+							else
+							{
+								WorldTransforms[MeshBoneIndex.GetInt()] = Pose[BoneIndex] * WorldTransforms[ParentIndex];
+							}
+							BoneColors[MeshBoneIndex.GetInt()] = BoneColor;
+						}
+						*/
+
+						TArray<FTransform> BoneTransforms;
+						BoneTransforms.AddUninitialized(Pose.GetBoneContainer().GetNumBones());
+
+						TArray<FBoneIndexType> TmpRequiredBones;
+						TmpRequiredBones.Reserve(Pose.GetBoneContainer().GetNumBones());
+
+						for (FCompactPoseBoneIndex BoneIndex : Pose.ForEachBoneIndex())
+						{
+							FMeshPoseBoneIndex MeshBoneIndex = Pose.GetBoneContainer().MakeMeshPoseIndex(BoneIndex);
+
+							BoneTransforms[MeshBoneIndex.GetInt()] = Pose[BoneIndex];
+							TmpRequiredBones.Add(MeshBoneIndex.GetInt());
+						}
+
 						PoseWatch.SetPose(TmpRequiredBones, BoneTransforms);
 						PoseWatch.SetWorldTransform(SkelMeshComponent->GetComponentTransform());
 
diff --git a/Engine/Source/Runtime/Engine/Private/Animation/AnimNode_AssetPlayerBase.cpp b/Engine/Source/Runtime/Engine/Private/Animation/AnimNode_AssetPlayerBase.cpp
--- a/Engine/Source/Runtime/Engine/Private/Animation/AnimNode_AssetPlayerBase.cpp
+++ b/Engine/Source/Runtime/Engine/Private/Animation/AnimNode_AssetPlayerBase.cpp
@@ -23,7 +23,7 @@ void FAnimNode_AssetPlayerBase::Update_AnyThread(const FAnimationUpdateContext&
 	UpdateAssetPlayer(Context);
 }
 
-void FAnimNode_AssetPlayerBase::CreateTickRecordForNode(const FAnimationUpdateContext& Context, UAnimSequenceBase* Sequence, bool bLooping, float PlayRate)
+void FAnimNode_AssetPlayerBase::CreateTickRecordForNode(const FAnimationUpdateContext& Context, UAnimSequenceBase* Sequence, bool bLooping, float PlayRate, bool bIsEvaluator)
 {
 	// Create a tick record and push into the closest scope
 	const float FinalBlendWeight = Context.GetFinalBlendWeight();
@@ -41,7 +41,7 @@ void FAnimNode_AssetPlayerBase::CreateTickRecordForNode(const FAnimationUpdateCo
 	}
 
 	const UE::Anim::FAnimSyncParams SyncParams(GroupNameToUse, SyncGroupRole, MethodToUse);
-	FAnimTickRecord TickRecord(Sequence, bLooping, PlayRate, FinalBlendWeight, /*inout*/ InternalTimeAccumulator, MarkerTickRecord);
+	FAnimTickRecord TickRecord(Sequence, bLooping, PlayRate, bIsEvaluator, FinalBlendWeight, /*inout*/ InternalTimeAccumulator, MarkerTickRecord);
 	TickRecord.GatherContextData(Context);
 
 	TickRecord.RootMotionWeightModifier = Context.GetRootMotionWeightModifier();
diff --git a/Engine/Source/Runtime/Engine/Private/Animation/AnimNode_SequencePlayer.cpp b/Engine/Source/Runtime/Engine/Private/Animation/AnimNode_SequencePlayer.cpp
--- a/Engine/Source/Runtime/Engine/Private/Animation/AnimNode_SequencePlayer.cpp
+++ b/Engine/Source/Runtime/Engine/Private/Animation/AnimNode_SequencePlayer.cpp
@@ -97,7 +97,7 @@ void FAnimNode_SequencePlayerBase::UpdateAssetPlayer(const FAnimationUpdateConte
 
 		InternalTimeAccumulator = FMath::Clamp(InternalTimeAccumulator, 0.f, CurrentSequence->GetPlayLength());
 		const float AdjustedPlayRate = PlayRateScaleBiasClampState.ApplyTo(GetPlayRateScaleBiasClampConstants(), FMath::IsNearlyZero(CurrentPlayRateBasis) ? 0.f : (CurrentPlayRate / CurrentPlayRateBasis), Context.GetDeltaTime());
-		CreateTickRecordForNode(Context, CurrentSequence, GetLoopAnimation(), AdjustedPlayRate);
+		CreateTickRecordForNode(Context, CurrentSequence, GetLoopAnimation(), AdjustedPlayRate, false);
 	}
 
 #if WITH_EDITORONLY_DATA
diff --git a/Engine/Source/Runtime/Engine/Private/Animation/AnimSequenceBase.cpp b/Engine/Source/Runtime/Engine/Private/Animation/AnimSequenceBase.cpp
--- a/Engine/Source/Runtime/Engine/Private/Animation/AnimSequenceBase.cpp
+++ b/Engine/Source/Runtime/Engine/Private/Animation/AnimSequenceBase.cpp
@@ -537,7 +537,7 @@ void UAnimSequenceBase::RemapTracksToNewSkeleton(USkeleton* NewSkeleton, bool bC
 void UAnimSequenceBase::TickAssetPlayer(FAnimTickRecord& Instance, struct FAnimNotifyQueue& NotifyQueue, FAnimAssetTickContext& Context) const
 {
 	float CurrentTime = *(Instance.TimeAccumulator);
-	if (Context.ShouldResyncToSyncGroup())
+	if (Context.ShouldResyncToSyncGroup() && !Instance.bIsEvaluator)
 	{
 		// Synchronize the asset player time to the other sync group members when (re)joining the group
 		CurrentTime = Context.GetAnimationPositionRatio() * GetPlayLength();
diff --git a/Engine/Source/Runtime/Engine/Private/Animation/AnimSingleNodeInstanceProxy.cpp b/Engine/Source/Runtime/Engine/Private/Animation/AnimSingleNodeInstanceProxy.cpp
--- a/Engine/Source/Runtime/Engine/Private/Animation/AnimSingleNodeInstanceProxy.cpp
+++ b/Engine/Source/Runtime/Engine/Private/Animation/AnimSingleNodeInstanceProxy.cpp
@@ -587,7 +587,7 @@ void FAnimNode_SingleNode::Update_AnyThread(const FAnimationUpdateContext& Conte
 		}
 		else if (UAnimSequence* Sequence = Cast<UAnimSequence>(Proxy->CurrentAsset))
 		{
-			FAnimTickRecord TickRecord(Sequence, Proxy->bLooping, NewPlayRate, 1.f, /*inout*/ Proxy->CurrentTime, Proxy->MarkerTickRecord);
+			FAnimTickRecord TickRecord(Sequence, Proxy->bLooping, NewPlayRate, false, 1.f, /*inout*/ Proxy->CurrentTime, Proxy->MarkerTickRecord);
 			TickRecord.DeltaTimeRecord = &(Proxy->DeltaTimeRecord);
 			
 			SyncScope.AddTickRecord(TickRecord);
@@ -606,7 +606,7 @@ void FAnimNode_SingleNode::Update_AnyThread(const FAnimationUpdateContext& Conte
 		}
 		else if (UAnimStreamable* Streamable = Cast<UAnimStreamable>(Proxy->CurrentAsset))
 		{
-			FAnimTickRecord TickRecord(Streamable, Proxy->bLooping, NewPlayRate, 1.f, /*inout*/ Proxy->CurrentTime, Proxy->MarkerTickRecord);
+			FAnimTickRecord TickRecord(Streamable, Proxy->bLooping, NewPlayRate, false, 1.f, /*inout*/ Proxy->CurrentTime, Proxy->MarkerTickRecord);
 			TickRecord.DeltaTimeRecord = &(Proxy->DeltaTimeRecord);
 			
 			SyncScope.AddTickRecord(TickRecord);
@@ -625,7 +625,7 @@ void FAnimNode_SingleNode::Update_AnyThread(const FAnimationUpdateContext& Conte
 		}
 		else if(UAnimComposite* Composite = Cast<UAnimComposite>(Proxy->CurrentAsset))
 		{
-			FAnimTickRecord TickRecord(Composite, Proxy->bLooping, NewPlayRate, 1.f, /*inout*/ Proxy->CurrentTime, Proxy->MarkerTickRecord);
+			FAnimTickRecord TickRecord(Composite, Proxy->bLooping, NewPlayRate, false, 1.f, /*inout*/ Proxy->CurrentTime, Proxy->MarkerTickRecord);
 			TickRecord.DeltaTimeRecord = &(Proxy->DeltaTimeRecord);
 			
 			SyncScope.AddTickRecord(TickRecord);
diff --git a/Engine/Source/Runtime/Engine/Private/Animation/BlendSpace.cpp b/Engine/Source/Runtime/Engine/Private/Animation/BlendSpace.cpp
--- a/Engine/Source/Runtime/Engine/Private/Animation/BlendSpace.cpp
+++ b/Engine/Source/Runtime/Engine/Private/Animation/BlendSpace.cpp
@@ -421,7 +421,7 @@ void UBlendSpace::TickAssetPlayer(FAnimTickRecord& Instance, struct FAnimNotifyQ
 			}
 
 			float& NormalizedCurrentTime = *(Instance.TimeAccumulator);
-			if (Context.ShouldResyncToSyncGroup())
+			if (Context.ShouldResyncToSyncGroup() && !Instance.bIsEvaluator)
 			{
 				// Synchronize the asset player time to the other sync group members when (re)joining the group
 				NormalizedCurrentTime = Context.GetAnimationPositionRatio();
@@ -545,7 +545,7 @@ void UBlendSpace::TickAssetPlayer(FAnimTickRecord& Instance, struct FAnimNotifyQ
 				float ClampedNormalizedPreviousTime = FMath::Clamp<float>(NormalizedPreviousTime, 0.f, 1.f);
 				float ClampedNormalizedCurrentTime = FMath::Clamp<float>(NormalizedCurrentTime, 0.f, 1.f);
 
-				if (Instance.BlendSpace.bIsEvaluator && !Instance.BlendSpace.bTeleportToTime)
+				if (Instance.bIsEvaluator && !Instance.BlendSpace.bTeleportToTime)
 				{
 					// When running under an evaluator the time is being set explicitly and we want to add on the deltas.
 					ClampedNormalizedPreviousTime -= ExtraNormalizedDeltaTime;
diff --git a/Engine/Source/Runtime/Engine/Private/AnimationAsset.cpp b/Engine/Source/Runtime/Engine/Private/AnimationAsset.cpp
--- a/Engine/Source/Runtime/Engine/Private/AnimationAsset.cpp
+++ b/Engine/Source/Runtime/Engine/Private/AnimationAsset.cpp
@@ -191,6 +191,11 @@ void FAnimTickRecord::AllocateContextDataContainer()
 }
 
 FAnimTickRecord::FAnimTickRecord(UAnimSequenceBase* InSequence, bool bInLooping, float InPlayRate, float InFinalBlendWeight, float& InCurrentTime, FMarkerTickRecord& InMarkerTickRecord)
+	: FAnimTickRecord(InSequence, bInLooping, InPlayRate, /*bInIsEvaluator*/ false, InFinalBlendWeight, InCurrentTime, InMarkerTickRecord)
+{
+}
+
+FAnimTickRecord::FAnimTickRecord(UAnimSequenceBase* InSequence, bool bInLooping, float InPlayRate, bool bInIsEvaluator, float InFinalBlendWeight, float& InCurrentTime, FMarkerTickRecord& InMarkerTickRecord)
 {
 	SourceAsset = InSequence;
 	TimeAccumulator = &InCurrentTime;
@@ -198,11 +203,12 @@ FAnimTickRecord::FAnimTickRecord(UAnimSequenceBase* InSequence, bool bInLooping,
 	PlayRateMultiplier = InPlayRate;
 	EffectiveBlendWeight = InFinalBlendWeight;
 	bLooping = bInLooping;
+	bIsEvaluator = bInIsEvaluator;
 }
 
 FAnimTickRecord::FAnimTickRecord(
 	UBlendSpace* InBlendSpace, const FVector& InBlendInput, TArray<FBlendSampleData>& InBlendSampleDataCache, FBlendFilter& InBlendFilter, bool bInLooping, 
-	float InPlayRate, bool bTeleportToTime, bool bIsEvaluator, float InFinalBlendWeight, float& InCurrentTime, FMarkerTickRecord& InMarkerTickRecord)
+	float InPlayRate, bool bTeleportToTime, bool bInIsEvaluator, float InFinalBlendWeight, float& InCurrentTime, FMarkerTickRecord& InMarkerTickRecord)
 {
 	SourceAsset = InBlendSpace;
 	BlendSpace.BlendSpacePositionX = InBlendInput.X;
@@ -210,12 +216,12 @@ FAnimTickRecord::FAnimTickRecord(
 	BlendSpace.BlendSampleDataCache = &InBlendSampleDataCache;
 	BlendSpace.BlendFilter = &InBlendFilter;
 	BlendSpace.bTeleportToTime = bTeleportToTime;
-	BlendSpace.bIsEvaluator = bIsEvaluator;
 	TimeAccumulator = &InCurrentTime;
 	MarkerTickRecord = &InMarkerTickRecord;
 	PlayRateMultiplier = InPlayRate;
 	EffectiveBlendWeight = InFinalBlendWeight;
 	bLooping = bInLooping;
+	bIsEvaluator = bInIsEvaluator;
 }
 
 FAnimTickRecord::FAnimTickRecord(UAnimMontage* InMontage, float InCurrentPosition, float, float, float InWeight, TArray<FPassedMarker>& InMarkersPassedThisTick, FMarkerTickRecord& InMarkerTickRecord)
