Add some LLM memory scopes for functions used when cooking.
#rnx

diff --git a/Engine/Plugins/FX/Niagara/Source/Niagara/Private/NiagaraSystem.cpp b/Engine/Plugins/FX/Niagara/Source/Niagara/Private/NiagaraSystem.cpp
--- a/Engine/Plugins/FX/Niagara/Source/Niagara/Private/NiagaraSystem.cpp
+++ b/Engine/Plugins/FX/Niagara/Source/Niagara/Private/NiagaraSystem.cpp
@@ -2242,6 +2242,7 @@ void UNiagaraSystem::ForceGraphToRecompileOnNextCheck()
 
 void UNiagaraSystem::WaitForCompilationComplete(bool bIncludingGPUShaders, bool bShowProgress)
 {
+	LLM_SCOPE(ELLMTag::Niagara);
 	TRACE_CPUPROFILER_EVENT_SCOPE(WaitForNiagaraCompilation);
 	TRACE_CPUPROFILER_EVENT_SCOPE_TEXT_ON_CHANNEL(*GetPathName(), NiagaraChannel);
 
diff --git a/Engine/Source/Runtime/Engine/Private/Components/StaticMeshComponent.cpp b/Engine/Source/Runtime/Engine/Private/Components/StaticMeshComponent.cpp
--- a/Engine/Source/Runtime/Engine/Private/Components/StaticMeshComponent.cpp
+++ b/Engine/Source/Runtime/Engine/Private/Components/StaticMeshComponent.cpp
@@ -2045,6 +2045,7 @@ void UStaticMeshComponent::UpdateCollisionFromStaticMesh()
 
 void UStaticMeshComponent::PostLoad()
 {
+	LLM_SCOPE(ELLMTag::StaticMesh);
 	NotifyIfStaticMeshChanged();
 
 	// need to postload the StaticMesh because super initializes variables based on GetStaticLightingType() which we override and use from the StaticMesh
