[shaders] exclude material, material function, and parameter collection guids when generating unique names (debug folder names) to avoid material changes changing the debug folder name, and add in the base material path as well (as opposed to the material/material instance path which is unstable due to shadermap deduplication) to distinguish same-named materials in different folders.

diff --git a/Engine/Source/Runtime/Engine/Private/Materials/MaterialShader.cpp b/Engine/Source/Runtime/Engine/Private/Materials/MaterialShader.cpp
--- a/Engine/Source/Runtime/Engine/Private/Materials/MaterialShader.cpp
+++ b/Engine/Source/Runtime/Engine/Private/Materials/MaterialShader.cpp
@@ -1064,11 +1064,14 @@ void FMaterialShaderMapId::AppendStaticParametersString(FString& ParamsString) c
 	}
 }
 
-void FMaterialShaderMapId::AppendKeyString(FString& KeyString, bool bIncludeSourceHashes) const
+void FMaterialShaderMapId::AppendKeyString(FString& KeyString, bool bIncludeSourceAndMaterialState) const
 {
 	check(IsContentValid());
-	BaseMaterialId.AppendString(KeyString);
-	KeyString.AppendChar('_');
+	if (bIncludeSourceAndMaterialState)
+	{
+		BaseMaterialId.AppendString(KeyString);
+		KeyString.AppendChar('_');
+	}
 
 	GetMaterialQualityLevelFName(QualityLevel).AppendString(KeyString);
 	KeyString.AppendChar('_');
@@ -1097,10 +1100,13 @@ void FMaterialShaderMapId::AppendKeyString(FString& KeyString, bool bIncludeSour
 		KeyString.AppendChar('_');
 	}
 
-	// Add any referenced functions to the key so that we will recompile when they are changed
-	for (int32 FunctionIndex = 0; FunctionIndex < ReferencedFunctions.Num(); FunctionIndex++)
+	if (bIncludeSourceAndMaterialState)
 	{
-		ReferencedFunctions[FunctionIndex].AppendString(KeyString);
+		// Add any referenced functions to the key so that we will recompile when they are changed
+		for (int32 FunctionIndex = 0; FunctionIndex < ReferencedFunctions.Num(); FunctionIndex++)
+		{
+			ReferencedFunctions[FunctionIndex].AppendString(KeyString);
+		}
 	}
 
 	{
@@ -1112,9 +1118,12 @@ void FMaterialShaderMapId::AppendKeyString(FString& KeyString, bool bIncludeSour
 
 	KeyString.AppendChar('_');
 
-	for (int32 CollectionIndex = 0; CollectionIndex < ReferencedParameterCollections.Num(); CollectionIndex++)
+	if (bIncludeSourceAndMaterialState)
 	{
-		ReferencedParameterCollections[CollectionIndex].AppendString(KeyString);
+		for (int32 CollectionIndex = 0; CollectionIndex < ReferencedParameterCollections.Num(); CollectionIndex++)
+		{
+			ReferencedParameterCollections[CollectionIndex].AppendString(KeyString);
+		}
 	}
 
 	// Add the inputs for any shaders that are stored inline in the shader map
@@ -1124,7 +1133,7 @@ void FMaterialShaderMapId::AppendKeyString(FString& KeyString, bool bIncludeSour
 		MakeArrayView(VertexFactoryTypeDependencies),
 		LayoutParams, 
 		KeyString,
-		bIncludeSourceHashes);
+		bIncludeSourceAndMaterialState);
 
 	BytesToHex(&TextureReferencesHash.Hash[0], sizeof(TextureReferencesHash.Hash), KeyString);
 
diff --git a/Engine/Source/Runtime/Engine/Private/Materials/MaterialShared.cpp b/Engine/Source/Runtime/Engine/Private/Materials/MaterialShared.cpp
--- a/Engine/Source/Runtime/Engine/Private/Materials/MaterialShared.cpp
+++ b/Engine/Source/Runtime/Engine/Private/Materials/MaterialShared.cpp
@@ -2747,8 +2747,13 @@ FString FMaterial::GetUniqueAssetName(EShaderPlatform Platform, const FMaterialS
 {
 	FString IdKeyStr;
 	// append the portion of the DDC key string associated with the shadermapid, but exclude the source code
-	// hashes, such that this name remains stable when source code edits are applied (including version bumps)
-	ShaderMapId.AppendKeyString(IdKeyStr, /* bIncludeSourceHashes */ false);
+	// hashes and material function/parameter collection guids, such that this name remains stable when edits
+	// to this data are applied (including source version bumps)
+	ShaderMapId.AppendKeyString(IdKeyStr, /* bIncludeSourceAndMaterialState */ false);
+	// append the base material path as well to differentiate materials with the same name and different paths.
+	// note we explicitly _do not_ use the path of the asset itself as if the asset is a material instance we 
+	// want it to properly deduplicate against other instances which might end up pointing to the same shadermap.
+	IdKeyStr.Append(GetBaseMaterialPathName());
 	uint64 Hash = CityHash64((const char*)*IdKeyStr, IdKeyStr.Len() * sizeof(TCHAR));
 	return FString::Printf(TEXT("%s_%llx"), *GetFriendlyName(), Hash);
 }
diff --git a/Engine/Source/Runtime/Engine/Public/MaterialShared.h b/Engine/Source/Runtime/Engine/Public/MaterialShared.h
--- a/Engine/Source/Runtime/Engine/Public/MaterialShared.h
+++ b/Engine/Source/Runtime/Engine/Public/MaterialShared.h
@@ -983,7 +983,7 @@ public:
 	void UpdateFromParameterSet(const FStaticParameterSet& StaticParameters);
 
 	/** Appends string representations of this Id to a key string. */
-	void AppendKeyString(FString& KeyString, bool bIncludeSourceHashes = true) const;
+	void AppendKeyString(FString& KeyString, bool bIncludeSourceAndMaterialState = true) const;
 	void AppendStaticParametersString(FString& ParamsString) const;
 
 	const TArray<FStaticSwitchParameter> &GetStaticSwitchParameters() const 					{ return StaticSwitchParameters; }
