Removed the callback system from core Analytics system. Fixed up analytics gathering in DerivedDataCache, Virtualization and HttpCacheStore

diff --git a/Engine/Source/Developer/DerivedDataCache/DerivedDataCache.Build.cs b/Engine/Source/Developer/DerivedDataCache/DerivedDataCache.Build.cs
--- a/Engine/Source/Developer/DerivedDataCache/DerivedDataCache.Build.cs
+++ b/Engine/Source/Developer/DerivedDataCache/DerivedDataCache.Build.cs
@@ -9,10 +9,10 @@ public class DerivedDataCache : ModuleRules
 	{
 		PrivateDependencyModuleNames.Add("Core");
 		PrivateDependencyModuleNames.Add("Zen");
-
+		
 		// Dependencies for "S3" and "HTTP" backends
 		PrivateDependencyModuleNames.AddRange(new string[] { "SSL", "Json", "Zen" });
-		PrivateIncludePathModuleNames.AddRange(new string[] { "DesktopPlatform", "Zen"});
+		PrivateIncludePathModuleNames.AddRange(new string[] { "DesktopPlatform", "Zen", "Analytics" });
 		AddEngineThirdPartyPrivateStaticDependencies(Target, "libcurl");
 		AddEngineThirdPartyPrivateStaticDependencies(Target, "OpenSSL");
 
diff --git a/Engine/Source/Developer/DerivedDataCache/Private/DerivedDataCache.cpp b/Engine/Source/Developer/DerivedDataCache/Private/DerivedDataCache.cpp
--- a/Engine/Source/Developer/DerivedDataCache/Private/DerivedDataCache.cpp
+++ b/Engine/Source/Developer/DerivedDataCache/Private/DerivedDataCache.cpp
@@ -4,6 +4,7 @@
 #include "DerivedDataCacheInterface.h"
 
 #include "Algo/AllOf.h"
+#include "AnalyticsEventAttribute.h"
 #include "Async/AsyncWork.h"
 #include "Async/InheritedContext.h"
 #include "Async/TaskGraphInterfaces.h"
@@ -944,6 +945,86 @@ public:
 		GatherDerivedDataCacheSummaryStats(DDCSummaryStats);
 	}
 
+	virtual void GatherAnalytics(TArray<FAnalyticsEventAttribute>& Attributes) const override
+	{
+#if ENABLE_COOK_STATS
+
+		// Gather the latest resource stats
+		TArray<FDerivedDataCacheResourceStat> ResourceStats;
+
+		GatherDerivedDataCacheResourceStats(ResourceStats);
+
+		FDerivedDataCacheResourceStat ResourceStatsTotal(TEXT("Total"));
+
+		// Accumulate Totals
+		for (const FDerivedDataCacheResourceStat& Stat : ResourceStats)
+		{
+			ResourceStatsTotal += Stat;
+		}
+
+		ResourceStats.Emplace(ResourceStatsTotal);
+
+		// Append to the attributes
+		for (const FDerivedDataCacheResourceStat& Stat : ResourceStats)
+		{
+			FString BaseName = TEXT("DDC.Resource.") + Stat.AssetType;
+
+			BaseName = BaseName.Replace(TEXT("("), TEXT("")).Replace(TEXT(")"), TEXT(""));
+
+			{
+				FString AttrName = BaseName + TEXT(".BuildCount");
+				Attributes.Emplace(MoveTemp(AttrName), Stat.BuildCount);
+			}
+
+			{
+				FString AttrName = BaseName + TEXT(".BuildTimeSec");
+				Attributes.Emplace(MoveTemp(AttrName), Stat.BuildTimeSec);
+			}
+
+			{
+				FString AttrName = BaseName + TEXT(".BuildSizeMB");
+				Attributes.Emplace(MoveTemp(AttrName), Stat.BuildSizeMB);
+			}
+
+			{
+				FString AttrName = BaseName + TEXT(".LoadCount");
+				Attributes.Emplace(MoveTemp(AttrName), Stat.LoadCount);
+			}
+
+			{
+				FString AttrName = BaseName + TEXT(".LoadTimeSecLoadTimeSec");
+				Attributes.Emplace(MoveTemp(AttrName), Stat.LoadTimeSec);
+			}
+
+			{
+				FString AttrName = BaseName + TEXT(".LoadSizeMB");
+				Attributes.Emplace(MoveTemp(AttrName), Stat.LoadSizeMB);
+			}
+		}
+
+		// Gather the summary stats
+		FDerivedDataCacheSummaryStats SummaryStats;
+
+		GatherDerivedDataCacheSummaryStats(SummaryStats);
+
+		// Append to the attributes
+		for (const FDerivedDataCacheSummaryStat& Stat : SummaryStats.Stats)
+		{
+			FString FormattedAttrName = "DDC.Summary." + Stat.Key;
+			Attributes.Emplace(FormattedAttrName, Stat.Value);
+		}
+
+		TSharedRef<FDerivedDataCacheStatsNode> RootNode = Backend->GatherUsageStats();
+		RootNode->ForEachDescendant([&Attributes](TSharedRef<const FDerivedDataCacheStatsNode> Node)
+			{
+				for (const FCookStatsManager::StringKeyValue& Stat : Node->CustomStats)
+				{
+					Attributes.Emplace(Stat.Key, Stat.Value);
+				}
+			});
+#endif
+	}
+
 	/** Get event delegate for data cache notifications */
 	virtual FOnDDCNotification& GetDDCNotificationEvent()
 	{
diff --git a/Engine/Source/Developer/DerivedDataCache/Private/HttpCacheStore.cpp b/Engine/Source/Developer/DerivedDataCache/Private/HttpCacheStore.cpp
--- a/Engine/Source/Developer/DerivedDataCache/Private/HttpCacheStore.cpp
+++ b/Engine/Source/Developer/DerivedDataCache/Private/HttpCacheStore.cpp
@@ -34,6 +34,7 @@
 #include "Misc/ScopeLock.h"
 #include "Misc/ScopeRWLock.h"
 #include "Misc/StringBuilder.h"
+#include "ProfilingDebugging/CookStats.h"
 #include "ProfilingDebugging/CountersTrace.h"
 #include "ProfilingDebugging/CpuProfilerTrace.h"
 #include "Serialization/CompactBinary.h"
@@ -332,7 +333,7 @@ private:
 	bool EndIsServiceReady(THttpUniquePtr<IHttpResponse>& Response, TArray64<uint8>& Body);
 	bool AcquireAccessToken(IHttpClient* Client = nullptr);
 	void SetAccessToken(FStringView Token, double RefreshDelay = 0.0);
-
+	
 	enum class EOperationCategory
 	{
 		Get,
@@ -2313,6 +2314,30 @@ void FHttpCacheStore::LegacyStats(FDerivedDataCacheStatsNode& OutNode)
 {
 	OutNode = {TEXT("Unreal Cloud DDC"), FString::Printf(TEXT("%s (%s)"), *Domain, *Namespace), /*bIsLocal*/ false};
 	OutNode.UsageStats.Add(TEXT(""), UsageStats);
+
+#if ENABLE_COOK_STATS
+	const int64 GetHits = UsageStats.GetStats.GetAccumulatedValueAnyThread(FCookStats::CallStats::EHitOrMiss::Hit, FCookStats::CallStats::EStatType::Counter);
+	const int64 GetMisses = UsageStats.GetStats.GetAccumulatedValueAnyThread(FCookStats::CallStats::EHitOrMiss::Miss, FCookStats::CallStats::EStatType::Counter);
+	const int64 PutHits = UsageStats.PutStats.GetAccumulatedValueAnyThread(FCookStats::CallStats::EHitOrMiss::Hit, FCookStats::CallStats::EStatType::Counter);
+	const int64 PutMisses = UsageStats.PutStats.GetAccumulatedValueAnyThread(FCookStats::CallStats::EHitOrMiss::Miss, FCookStats::CallStats::EStatType::Counter);
+	const int64 TotalGets = GetHits + GetMisses;
+	const int64 TotalPuts = PutHits + PutMisses;
+
+	const FString BaseName(TEXTVIEW("CloudDDC."));
+
+	OutNode.CustomStats = FCookStatsManager::CreateKeyValueArray(	FString(WriteToString<64>(BaseName, TEXTVIEW("Domain"))), Domain,
+																	FString(WriteToString<64>(BaseName, TEXTVIEW("EffectiveDomain"))), *EffectiveDomain,
+																	FString(WriteToString<64>(BaseName, TEXTVIEW("Namespace"))), Namespace,
+																	FString(WriteToString<64>(BaseName, TEXTVIEW("LoginAttempts"))), LoginAttempts,
+																	FString(WriteToString<64>(BaseName, TEXTVIEW("InteractiveLoginAttempts"))), InteractiveLoginAttempts,
+																	FString(WriteToString<64>(BaseName, TEXTVIEW("FailedLoginAttempts"))), FailedLoginAttempts,
+																	FString(WriteToString<64>(BaseName, TEXTVIEW("GetHits"))), GetHits,
+																	FString(WriteToString<64>(BaseName, TEXTVIEW("GetMisses"))), GetMisses,
+																	FString(WriteToString<64>(BaseName, TEXTVIEW("TotalGets"))), TotalGets,
+																	FString(WriteToString<64>(BaseName, TEXTVIEW("PutHits"))), PutHits,
+																	FString(WriteToString<64>(BaseName, TEXTVIEW("PutMisses"))), PutMisses,
+																	FString(WriteToString<64>(BaseName, TEXTVIEW("TotalPuts"))), TotalPuts);
+#endif
 }
 
 void FHttpCacheStore::Put(
diff --git a/Engine/Source/Developer/DerivedDataCache/Public/DerivedDataCacheInterface.h b/Engine/Source/Developer/DerivedDataCache/Public/DerivedDataCacheInterface.h
--- a/Engine/Source/Developer/DerivedDataCache/Public/DerivedDataCacheInterface.h
+++ b/Engine/Source/Developer/DerivedDataCache/Public/DerivedDataCacheInterface.h
@@ -14,6 +14,7 @@ class FDerivedDataCacheStatsNode;
 class IDDCCleanup;
 struct FDerivedDataCacheResourceStat;
 struct FDerivedDataCacheSummaryStats;
+struct FAnalyticsEventAttribute;
 
 /** 
  * Interface for the derived data cache
@@ -307,6 +308,8 @@ public:
 
 	virtual void GatherSummaryStats(FDerivedDataCacheSummaryStats& DDCSummaryStats) const = 0;
 
+	virtual void GatherAnalytics(TArray<FAnalyticsEventAttribute>& Attributes) const = 0;
+
 	//-----------------------
 	// Notification Interface
 	//-----------------------
diff --git a/Engine/Source/Developer/DerivedDataCache/Public/DerivedDataCacheUsageStats.h b/Engine/Source/Developer/DerivedDataCache/Public/DerivedDataCacheUsageStats.h
--- a/Engine/Source/Developer/DerivedDataCache/Public/DerivedDataCacheUsageStats.h
+++ b/Engine/Source/Developer/DerivedDataCache/Public/DerivedDataCacheUsageStats.h
@@ -127,6 +127,13 @@ enum class EDerivedDataCacheStatus
 	Deactivation
 };
 
+COOK_STAT(using FDerivedDataCacheSummaryStat = FCookStatsManager::StringKeyValue);
+
+struct FDerivedDataCacheSummaryStats
+{
+	COOK_STAT(TArray<FDerivedDataCacheSummaryStat> Stats);
+};
+
 /**
  *  Hierarchical usage stats for the DDC nodes.
  */
@@ -198,6 +205,7 @@ public:
 
 	TMap<FString, FDerivedDataCacheUsageStats> UsageStats;
 	FDerivedDataCacheSpeedStats SpeedStats;
+	COOK_STAT(TArray<FCookStatsManager::StringKeyValue> CustomStats);
 
 	TArray<TSharedRef<FDerivedDataCacheStatsNode>> Children;
 
@@ -283,11 +291,4 @@ struct FDerivedDataCacheResourceStatKeyFuncs : BaseKeyFuncs<FDerivedDataCacheRes
 	static const FString& GetSetKey(const FDerivedDataCacheResourceStat& Element) { return Element.AssetType; }
 	static bool Matches(const FString& A, const FString& B) { return A == B; }
 	static uint32 GetKeyHash(const FString& Key) { return GetTypeHash(Key); }
-};
-
-COOK_STAT(using FDerivedDataCacheSummaryStat = FCookStatsManager::StringKeyValue);
-
-struct FDerivedDataCacheSummaryStats
-{
-	COOK_STAT(TArray<FDerivedDataCacheSummaryStat> Stats);
-};
+};
\ No newline at end of file
diff --git a/Engine/Source/Developer/Virtualization/Private/VirtualizationManager.cpp b/Engine/Source/Developer/Virtualization/Private/VirtualizationManager.cpp
--- a/Engine/Source/Developer/Virtualization/Private/VirtualizationManager.cpp
+++ b/Engine/Source/Developer/Virtualization/Private/VirtualizationManager.cpp
@@ -18,6 +18,7 @@
 #include "PackageVirtualizationProcess.h"
 #include "ProfilingDebugging/CookStats.h"
 #include "VirtualizationFilterSettings.h"
+#include "AnalyticsEventAttribute.h"
 
 #define LOCTEXT_NAMESPACE "Virtualization"
 
@@ -2022,6 +2023,71 @@ void FVirtualizationManager::BroadcastEvent(TConstArrayView<FPullRequest> Reques
 	}
 }
 
+void FVirtualizationManager::GatherAnalytics(TArray<FAnalyticsEventAttribute>& Attributes) const
+{
+	using namespace UE::Virtualization;
+
+	// Grab the Virtualization stats
+	if (IVirtualizationSystem::IsInitialized())
+	{
+		IVirtualizationSystem& System = IVirtualizationSystem::Get();
+
+		FPayloadActivityInfo PayloadActivityInfo = System.GetAccumualtedPayloadActivityInfo();
+
+		const FString BaseName = TEXT("Virtualization");
+
+		{
+			FString AttrName = BaseName + TEXT(".Enabled");
+			Attributes.Emplace(MoveTemp(AttrName), System.IsEnabled());
+		}
+
+		{
+			FString AttrName = BaseName + TEXT(".Cache.TimeSpent");
+			Attributes.Emplace(MoveTemp(AttrName), (double)PayloadActivityInfo.Cache.CyclesSpent * FPlatformTime::GetSecondsPerCycle());
+		}
+
+		{
+			FString AttrName = BaseName + TEXT(".Cache.PayloadCount");
+			Attributes.Emplace(MoveTemp(AttrName), (double)PayloadActivityInfo.Cache.PayloadCount);
+		}
+
+		{
+			FString AttrName = BaseName + TEXT(".Cache.TotalBytes");
+			Attributes.Emplace(MoveTemp(AttrName), (double)PayloadActivityInfo.Cache.TotalBytes);
+		}
+
+		{
+			FString AttrName = BaseName + TEXT(".Push.TimeSpent");
+			Attributes.Emplace(MoveTemp(AttrName), (double)PayloadActivityInfo.Push.CyclesSpent * FPlatformTime::GetSecondsPerCycle());
+		}
+
+		{
+			FString AttrName = BaseName + TEXT(".Push.PayloadCount");
+			Attributes.Emplace(MoveTemp(AttrName), (double)PayloadActivityInfo.Push.PayloadCount);
+		}
+
+		{
+			FString AttrName = BaseName + TEXT(".Push.TotalBytes");
+			Attributes.Emplace(MoveTemp(AttrName), (double)PayloadActivityInfo.Push.TotalBytes);
+		}
+
+		{
+			FString AttrName = BaseName + TEXT(".Pull.TimeSpent");
+			Attributes.Emplace(MoveTemp(AttrName), (double)PayloadActivityInfo.Pull.CyclesSpent * FPlatformTime::GetSecondsPerCycle());
+		}
+
+		{
+			FString AttrName = BaseName + TEXT(".Pull.PayloadCount");
+			Attributes.Emplace(MoveTemp(AttrName), (double)PayloadActivityInfo.Pull.PayloadCount);
+		}
+
+		{
+			FString AttrName = BaseName + TEXT(".Pull.TotalBytes");
+			Attributes.Emplace(MoveTemp(AttrName), (double)PayloadActivityInfo.Pull.TotalBytes);
+		}
+	}
+}
+
 } // namespace UE::Virtualization
 
 #undef UE_INLINE_ALLOCATION_COUNT
diff --git a/Engine/Source/Developer/Virtualization/Private/VirtualizationManager.h b/Engine/Source/Developer/Virtualization/Private/VirtualizationManager.h
--- a/Engine/Source/Developer/Virtualization/Private/VirtualizationManager.h
+++ b/Engine/Source/Developer/Virtualization/Private/VirtualizationManager.h
@@ -7,10 +7,12 @@
 #include "Logging/LogMacros.h"
 #include "Templates/UniquePtr.h"
 
+
 #include "Virtualization/VirtualizationSystem.h"
 
 class IConsoleObject;
 class FOutputDevice;
+struct FAnalyticsEventAttribute;
 
 /**
  * Configuring the backend hierarchy
@@ -136,6 +138,8 @@ private:
 
 	virtual void GetPayloadActivityInfo( GetPayloadActivityInfoFuncRef ) const override;
 
+	virtual void GatherAnalytics(TArray<FAnalyticsEventAttribute>& Attributes) const override;
+
 	virtual FOnNotification& GetNotificationEvent() override
 	{
 		return NotificationEvent;
@@ -154,7 +158,6 @@ private:
 	void OnUpdateDebugMissBackendsFromConsole(const TArray<FString>& Args, FOutputDevice& OutputDevice);
 	void OnUpdateDebugMissChanceFromConsole(const TArray<FString>& Args, FOutputDevice& OutputDevice);
 	void OnUpdateDebugMissCountFromConsole(const TArray<FString>& Args, FOutputDevice& OutputDevice);
-
 	void UpdateBackendDebugState();
 
 	bool ShouldDebugDisablePulling(FStringView BackendConfigName) const;	
@@ -175,7 +178,7 @@ private:
 
 	void PullDataFromAllBackends(TArrayView<FPullRequest> Requests);
 	void PullDataFromBackend(IVirtualizationBackend& Backend, TArrayView<FPullRequest> Requests);
-
+	
 	bool ShouldVirtualizeAsset(const UObject* Owner) const;
 
 	/** 
diff --git a/Engine/Source/Developer/Virtualization/Virtualization.Build.cs b/Engine/Source/Developer/Virtualization/Virtualization.Build.cs
--- a/Engine/Source/Developer/Virtualization/Virtualization.Build.cs
+++ b/Engine/Source/Developer/Virtualization/Virtualization.Build.cs
@@ -7,6 +7,8 @@ public class Virtualization : ModuleRules
 {
 	public Virtualization(ReadOnlyTargetRules Target) : base(Target)
 	{
+		PrivateIncludePathModuleNames.Add("Analytics");
+
 		PrivateDependencyModuleNames.AddRange(
 			new string[]
 			{
diff --git a/Engine/Source/Developer/Zen/Private/ZenServerInterface.cpp b/Engine/Source/Developer/Zen/Private/ZenServerInterface.cpp
--- a/Engine/Source/Developer/Zen/Private/ZenServerInterface.cpp
+++ b/Engine/Source/Developer/Zen/Private/ZenServerInterface.cpp
@@ -6,6 +6,7 @@
 #include "ZenSerialization.h"
 #include "ZenServerHttp.h"
 
+#include "AnalyticsEventAttribute.h"
 #include "Async/Async.h"
 #include "Dom/JsonValue.h"
 #include "HAL/FileManager.h"
@@ -1697,6 +1698,94 @@ FZenServiceInstance::RequestGC(const bool* OverrideCollectSmallObjects, const ui
 	return false;
 }
 
+bool 
+FZenServiceInstance::GatherAnalytics(TArray<FAnalyticsEventAttribute>& Attributes)
+{
+	FZenStats ZenStats;
+
+	if (GetStats(ZenStats) == false)
+		return false;
+
+	const FString BaseName = TEXT("Zen");
+
+	{
+		FString AttrName = BaseName + TEXT(".Enabled");
+		Attributes.Emplace(MoveTemp(AttrName), ZenStats.IsValid);
+	}
+
+	{
+		FString AttrName = BaseName + TEXT(".Cache.HitRatio");
+		Attributes.Emplace(MoveTemp(AttrName), ZenStats.CacheStats.HitRatio);
+	}
+
+	{
+		FString AttrName = BaseName + TEXT(".Cache.Hits");
+		Attributes.Emplace(MoveTemp(AttrName), ZenStats.CacheStats.Hits);
+	}
+
+	{
+		FString AttrName = BaseName + TEXT(".Cache.Misses");
+		Attributes.Emplace(MoveTemp(AttrName), ZenStats.CacheStats.Misses);
+	}
+
+	{
+		FString AttrName = BaseName + TEXT(".Cache.Size.Disk");
+		Attributes.Emplace(MoveTemp(AttrName), ZenStats.CacheStats.Size.Disk);
+	}
+
+	{
+		FString AttrName = BaseName + TEXT(".Cache.Size.Memory");
+		Attributes.Emplace(MoveTemp(AttrName), ZenStats.CacheStats.Size.Memory);
+	}
+
+	{
+		FString AttrName = BaseName + TEXT(".Cache.UpstreamHits");
+		Attributes.Emplace(MoveTemp(AttrName), ZenStats.CacheStats.UpstreamHits);
+	}
+
+	{
+		FString AttrName = BaseName + TEXT(".Cache.UpstreamRatio");
+		Attributes.Emplace(MoveTemp(AttrName), ZenStats.CacheStats.UpstreamRatio);
+	}
+
+	{
+		FString AttrName = BaseName + TEXT(".Cache.TotalUploadedMB");
+		Attributes.Emplace(MoveTemp(AttrName), ZenStats.UpstreamStats.TotalUploadedMB);
+	}
+
+	{
+		FString AttrName = BaseName + TEXT(".Upstream.TotalDownloadedMB");
+		Attributes.Emplace(MoveTemp(AttrName), ZenStats.UpstreamStats.TotalDownloadedMB);
+	}
+
+	{
+		FString AttrName = BaseName + TEXT(".Upstream.TotalUploadedMB");
+		Attributes.Emplace(MoveTemp(AttrName), ZenStats.UpstreamStats.TotalUploadedMB);
+	}
+
+	{
+		FString AttrName = BaseName + TEXT(".Cas.Size.Large");
+		Attributes.Emplace(MoveTemp(AttrName), ZenStats.CASStats.Size.Large);
+	}
+
+	{
+		FString AttrName = BaseName + TEXT(".Cas.Size.Small");
+		Attributes.Emplace(MoveTemp(AttrName), ZenStats.CASStats.Size.Small);
+	}
+
+	{
+		FString AttrName = BaseName + TEXT(".Cas.Size.Tiny");
+		Attributes.Emplace(MoveTemp(AttrName), ZenStats.CASStats.Size.Tiny);
+	}
+
+	{
+		FString AttrName = BaseName + TEXT(".Cas.Size.Total");
+		Attributes.Emplace(MoveTemp(AttrName), ZenStats.CASStats.Size.Total);
+	}
+
+	return true;
+}
+
 #endif // UE_WITH_ZEN
 
 }
diff --git a/Engine/Source/Developer/Zen/Public/Experimental/ZenServerInterface.h b/Engine/Source/Developer/Zen/Public/Experimental/ZenServerInterface.h
--- a/Engine/Source/Developer/Zen/Public/Experimental/ZenServerInterface.h
+++ b/Engine/Source/Developer/Zen/Public/Experimental/ZenServerInterface.h
@@ -23,6 +23,8 @@
 
 #define UE_API ZEN_API
 
+struct FAnalyticsEventAttribute;
+
 namespace UE::Zen
 {
 
@@ -183,6 +185,8 @@ public:
 	UE_API bool TryRecovery();
 
 	UE_API bool GetStats(FZenStats& Stats);
+	UE_API bool GatherAnalytics(TArray<FAnalyticsEventAttribute>& Attributes);
+
 	UE_API bool GetGCStatus(FGCStatus& Status);
 	UE_API bool RequestGC(const bool* OverrideCollectSmallObjects = nullptr, const uint32* OverrideMaxCacheDuration = nullptr);
 
@@ -194,7 +198,7 @@ private:
 	void PromptUserToStopRunningServerInstance(const FString& ServerFilePath);
 	FString ConditionalUpdateLocalInstall();
 	static bool AutoLaunch(const FServiceAutoLaunchSettings& InSettings, FString&& ExecutablePath, FString& OutHostName, uint16& OutPort);
-
+	
 	mutable TPimplPtr<class FZenHttpRequest> GCStatusHttpRequest;
 	mutable TFuture<TOptional<FGCStatus>> GCStatusRequest;
 	mutable TOptional<FGCStatus> LastGCStatus;
diff --git a/Engine/Source/Developer/Zen/Zen.Build.cs b/Engine/Source/Developer/Zen/Zen.Build.cs
--- a/Engine/Source/Developer/Zen/Zen.Build.cs
+++ b/Engine/Source/Developer/Zen/Zen.Build.cs
@@ -9,7 +9,8 @@ public class Zen : ModuleRules
 	{
 		PrivateDependencyModuleNames.AddRange(new string[] { "Core", "Sockets", "SSL", "Json" });
 		PrivateIncludePathModuleNames.Add("DesktopPlatform");
-
+		PrivateIncludePathModuleNames.Add("Analytics");
+		
 		AddEngineThirdPartyPrivateStaticDependencies(Target, "libcurl");
 		AddEngineThirdPartyPrivateStaticDependencies(Target, "OpenSSL");
 
diff --git a/Engine/Source/Editor/UnrealEd/Private/Commandlets/CookCommandlet.cpp b/Engine/Source/Editor/UnrealEd/Private/Commandlets/CookCommandlet.cpp
--- a/Engine/Source/Editor/UnrealEd/Private/Commandlets/CookCommandlet.cpp
+++ b/Engine/Source/Editor/UnrealEd/Private/Commandlets/CookCommandlet.cpp
@@ -41,7 +41,6 @@
 #include "Settings/ProjectPackagingSettings.h"
 #include "ShaderCompiler.h"
 #include "Stats/StatsMisc.h"
-#include "StudioAnalytics.h"
 #include "UObject/Class.h"
 #include "UObject/MetaData.h"
 #include "UObject/Package.h"
diff --git a/Engine/Source/Editor/UnrealEd/Private/Cooker/CookProfiling.cpp b/Engine/Source/Editor/UnrealEd/Private/Cooker/CookProfiling.cpp
--- a/Engine/Source/Editor/UnrealEd/Private/Cooker/CookProfiling.cpp
+++ b/Engine/Source/Editor/UnrealEd/Private/Cooker/CookProfiling.cpp
@@ -38,7 +38,9 @@
 #include "AnalyticsEventAttribute.h"
 #include "IAnalyticsProviderET.h"
 #include "StudioAnalytics.h"
+#include "DerivedDataCacheInterface.h"
 #include "Virtualization/VirtualizationSystem.h"
+#include "Experimental/ZenServerInterface.h"
 #endif
 
 #if OUTPUT_COOKTIMING
@@ -696,24 +698,37 @@ void LogCookStats(ECookMode::Type CookMode)
 	{
 
 		// convert filtered stats directly to an analytics event
-		TArray<FAnalyticsEventAttribute> StatAttrs;
+		TArray<FAnalyticsEventAttribute> Attributes;
 
 		// Sends each cook stat to the studio analytics system.
-		auto SendCookStatsToAnalytics = [&StatAttrs](const FString& StatName, const TArray<FCookStatsManager::StringKeyValue>& StatAttributes)
+		auto SendCookStatsToAnalytics = [&Attributes](const FString& StatName, const TArray<FCookStatsManager::StringKeyValue>& StatAttributes)
 		{
 			for (const auto& Attr : StatAttributes)
 			{
 				FString FormattedAttrName = StatName + "." + Attr.Key;
-
-				StatAttrs.Emplace(FormattedAttrName, Attr.Value);
+				Attributes.Emplace(FormattedAttrName, Attr.Value);
 			}
 		};
 
 		// Now actually grab the stats 
 		FCookStatsManager::LogCookStats(SendCookStatsToAnalytics);
 
+		// Gather DDC analytics
+		GetDerivedDataCacheRef().GatherAnalytics(Attributes);
+
+		// Gather Virtualization analytics
+		UE::Virtualization::IVirtualizationSystem::Get().GatherAnalytics(Attributes);
+
+#if UE_WITH_ZEN
+		// Gather Zen analytics
+		if (UE::Zen::IsDefaultServicePresent())
+		{
+			UE::Zen::GetDefaultServiceInstance().GatherAnalytics(Attributes);
+		}
+#endif
+
 		// Record them all under cooking event
-		FStudioAnalytics::GetProvider().RecordEvent(TEXT("Core.Cooking"), StatAttrs);
+		FStudioAnalytics::GetProvider().RecordEvent(TEXT("Core.Cooking"), Attributes);
 
 		FStudioAnalytics::GetProvider().BlockUntilFlushed(60.0f);
 	}
diff --git a/Engine/Source/Editor/UnrealEd/Private/FileHelpers.cpp b/Engine/Source/Editor/UnrealEd/Private/FileHelpers.cpp
--- a/Engine/Source/Editor/UnrealEd/Private/FileHelpers.cpp
+++ b/Engine/Source/Editor/UnrealEd/Private/FileHelpers.cpp
@@ -70,7 +70,6 @@
 #include "AutoSaveUtils.h"
 #include "AssetRegistry/AssetRegistryModule.h"
 #include "Misc/NamePermissionList.h"
-#include "EngineAnalytics.h"
 #include "StudioAnalytics.h"
 #include "AnalyticsEventAttribute.h"
 #include "HierarchicalLOD.h"
diff --git a/Engine/Source/Editor/UnrealEd/Private/Subsystems/AssetEditorSubsystem.cpp b/Engine/Source/Editor/UnrealEd/Private/Subsystems/AssetEditorSubsystem.cpp
--- a/Engine/Source/Editor/UnrealEd/Private/Subsystems/AssetEditorSubsystem.cpp
+++ b/Engine/Source/Editor/UnrealEd/Private/Subsystems/AssetEditorSubsystem.cpp
@@ -3,6 +3,7 @@
 #include "Subsystems/AssetEditorSubsystem.h"
 #include "AssetEditorMessages.h"
 #include "MessageEndpoint.h"
+#include "StudioAnalytics.h"
 #include "EngineAnalytics.h"
 #include "AnalyticsEventAttribute.h"
 #include "UObject/Package.h"
diff --git a/Engine/Source/Runtime/Core/Private/Virtualization/VirtualizationSystem.cpp b/Engine/Source/Runtime/Core/Private/Virtualization/VirtualizationSystem.cpp
--- a/Engine/Source/Runtime/Core/Private/Virtualization/VirtualizationSystem.cpp
+++ b/Engine/Source/Runtime/Core/Private/Virtualization/VirtualizationSystem.cpp
@@ -150,6 +150,11 @@ public:
 		return FPayloadActivityInfo();
 	}
 
+	virtual void GatherAnalytics(TArray<FAnalyticsEventAttribute>& Attributes) const override
+	{
+		// The null implementation has analytics to capture
+	}
+
 	virtual FOnNotification& GetNotificationEvent() override
 	{
 		return NotificationEvent;
diff --git a/Engine/Source/Runtime/Core/Public/Virtualization/VirtualizationSystem.h b/Engine/Source/Runtime/Core/Public/Virtualization/VirtualizationSystem.h
--- a/Engine/Source/Runtime/Core/Public/Virtualization/VirtualizationSystem.h
+++ b/Engine/Source/Runtime/Core/Public/Virtualization/VirtualizationSystem.h
@@ -23,6 +23,7 @@
 class FPackagePath;
 class FText;
 class UObject;
+struct FAnalyticsEventAttribute;
 
 namespace UE::Virtualization
 {
@@ -782,6 +783,9 @@ public:
 	/** Access profiling info relating to accumulated payload activity. Stats will only be collected if ENABLE_COOK_STATS is enabled.*/
 	virtual FPayloadActivityInfo GetAccumualtedPayloadActivityInfo() const = 0;
 
+	/** Gather analytics data. Stats will only be collected if ENABLE_COOK_STATS is enabled.*/
+	virtual void GatherAnalytics(TArray<FAnalyticsEventAttribute>& Attributes) const =0;
+
 	//* Notification messages
 	enum ENotification
 	{
diff --git a/Engine/Source/Runtime/Engine/Private/Analytics/StudioAnalytics.cpp b/Engine/Source/Runtime/Engine/Private/Analytics/StudioAnalytics.cpp
--- a/Engine/Source/Runtime/Engine/Private/Analytics/StudioAnalytics.cpp
+++ b/Engine/Source/Runtime/Engine/Private/Analytics/StudioAnalytics.cpp
@@ -135,8 +135,6 @@ void FStudioAnalytics::RecordEvent(const FString& EventName, const TArray<FAnaly
 
 void FStudioAnalytics::FireEvent_Loading(const FString& LoadingName, double SecondsSpentLoading, const TArray<FAnalyticsEventAttribute>& InAttributes)
 {
-	using namespace UE::Virtualization;
-
 	// Ignore anything less than a 1/4th a second.
 	if (SecondsSpentLoading < 0.250)
 	{
@@ -161,213 +159,17 @@ void FStudioAnalytics::FireEvent_Loading(const FString& LoadingName, double Seco
 
 #if ENABLE_COOK_STATS
 
-		TArray<FDerivedDataCacheResourceStat> DDCResourceStats;
-
-		// Grab the latest resource stats
-		GetDerivedDataCacheRef().GatherResourceStats(DDCResourceStats);
+		// Gather DDC analytics
+		GetDerivedDataCacheRef().GatherAnalytics(Attributes);
 
-		FDerivedDataCacheResourceStat DDCResourceStatsTotal(TEXT("Total"));
-
-		// Accumulate Totals
-		for (const FDerivedDataCacheResourceStat& Stat : DDCResourceStats)
-		{
-			DDCResourceStatsTotal += Stat;
-		}
-
-		DDCResourceStats.Emplace(DDCResourceStatsTotal);
-
-		for (const FDerivedDataCacheResourceStat& Stat : DDCResourceStats)
-		{
-			FString BaseName = TEXT("DDC.Resource.") + Stat.AssetType;
-
-			BaseName = BaseName.Replace(TEXT("("), TEXT("")).Replace(TEXT(")"), TEXT(""));
-			
-			{
-				FString AttrName = BaseName + TEXT(".BuildCount");
-				Attributes.Emplace(MoveTemp(AttrName), Stat.BuildCount);
-			}
-			
-			{
-				FString AttrName = BaseName + TEXT(".BuildTimeSec");
-				Attributes.Emplace(MoveTemp(AttrName), Stat.BuildTimeSec);
-			}
-
-			{
-				FString AttrName = BaseName + TEXT(".BuildSizeMB");
-				Attributes.Emplace(MoveTemp(AttrName), Stat.BuildSizeMB);
-			}
-
-			{
-				FString AttrName = BaseName + TEXT(".LoadCount");
-				Attributes.Emplace(MoveTemp(AttrName), Stat.LoadCount);
-			}
-
-			{
-				FString AttrName = BaseName + TEXT(".LoadTimeSec");
-				Attributes.Emplace(MoveTemp(AttrName), Stat.LoadTimeSec);
-			}
-
-			{
-				FString AttrName = BaseName + TEXT(".LoadSizeMB");
-				Attributes.Emplace(MoveTemp(AttrName), Stat.LoadSizeMB);
-			}
-		}
-
-		// Grab the DDC summary stats
-		{
-			FDerivedDataCacheSummaryStats DDCSummaryStats;
-			GetDerivedDataCacheRef().GatherSummaryStats(DDCSummaryStats);
-
-			for (const FDerivedDataCacheSummaryStat& Stat : DDCSummaryStats.Stats)
-			{
-				FString FormattedAttrName = "DDC.Summary." + Stat.Key;
-				Attributes.Emplace(FormattedAttrName, Stat.Value);
-			}
-
-			// Grab the Virtualization stats
-			if (IVirtualizationSystem::IsInitialized())
-			{
-				IVirtualizationSystem& System = IVirtualizationSystem::Get();
-
-				FPayloadActivityInfo PayloadActivityInfo = System.GetAccumualtedPayloadActivityInfo();
-
-				const FString BaseName = TEXT("Virtualization");
-
-				{
-					FString AttrName = BaseName + TEXT(".Enabled");
-					Attributes.Emplace(MoveTemp(AttrName), System.IsEnabled());
-				}
-
-				{
-					FString AttrName = BaseName + TEXT(".Cache.TimeSpent");
-					Attributes.Emplace(MoveTemp(AttrName), (double)PayloadActivityInfo.Cache.CyclesSpent * FPlatformTime::GetSecondsPerCycle());
-				}
-
-				{
-					FString AttrName = BaseName + TEXT(".Cache.PayloadCount");
-					Attributes.Emplace(MoveTemp(AttrName), (double)PayloadActivityInfo.Cache.PayloadCount);
-				}
-
-				{
-					FString AttrName = BaseName + TEXT(".Cache.TotalBytes");
-					Attributes.Emplace(MoveTemp(AttrName), (double)PayloadActivityInfo.Cache.TotalBytes);
-				}
-
-				{
-					FString AttrName = BaseName + TEXT(".Push.TimeSpent");
-					Attributes.Emplace(MoveTemp(AttrName), (double)PayloadActivityInfo.Push.CyclesSpent * FPlatformTime::GetSecondsPerCycle());
-				}
-
-				{
-					FString AttrName = BaseName + TEXT(".Push.PayloadCount");
-					Attributes.Emplace(MoveTemp(AttrName), (double)PayloadActivityInfo.Push.PayloadCount);
-				}
-
-				{
-					FString AttrName = BaseName + TEXT(".Push.TotalBytes");
-					Attributes.Emplace(MoveTemp(AttrName), (double)PayloadActivityInfo.Push.TotalBytes);
-				}
-
-				{
-					FString AttrName = BaseName + TEXT(".Pull.TimeSpent");
-					Attributes.Emplace(MoveTemp(AttrName), (double)PayloadActivityInfo.Pull.CyclesSpent * FPlatformTime::GetSecondsPerCycle());
-				}
-
-				{
-					FString AttrName = BaseName + TEXT(".Pull.PayloadCount");
-					Attributes.Emplace(MoveTemp(AttrName), (double)PayloadActivityInfo.Pull.PayloadCount);
-				}
-
-				{
-					FString AttrName = BaseName + TEXT(".Pull.TotalBytes");
-					Attributes.Emplace(MoveTemp(AttrName), (double)PayloadActivityInfo.Pull.TotalBytes);
-				}
-			}
-		}
+		// Gather Virtualization analytics
+		UE::Virtualization::IVirtualizationSystem::Get().GatherAnalytics(Attributes);
 
 #if UE_WITH_ZEN
+		// Gather Zen analytics
 		if (UE::Zen::IsDefaultServicePresent())
 		{
-			// Grab the Zen summary stats
-			UE::Zen::FZenStats ZenStats;
-			UE::Zen::GetDefaultServiceInstance().GetStats(ZenStats);
-
-			const FString BaseName = TEXT("Zen");
-
-			{
-				FString AttrName = BaseName + TEXT(".Enabled");
-				Attributes.Emplace(MoveTemp(AttrName), ZenStats.IsValid);
-			}
-
-			{
-				FString AttrName = BaseName + TEXT(".Cache.HitRatio");
-				Attributes.Emplace(MoveTemp(AttrName), ZenStats.CacheStats.HitRatio);
-			}
-
-			{
-				FString AttrName = BaseName + TEXT(".Cache.Hits");
-				Attributes.Emplace(MoveTemp(AttrName), ZenStats.CacheStats.Hits);
-			}
-
-			{
-				FString AttrName = BaseName + TEXT(".Cache.Misses");
-				Attributes.Emplace(MoveTemp(AttrName), ZenStats.CacheStats.Misses);
-			}
-
-			{
-				FString AttrName = BaseName + TEXT(".Cache.Size.Disk");
-				Attributes.Emplace(MoveTemp(AttrName), ZenStats.CacheStats.Size.Disk);
-			}
-
-			{
-				FString AttrName = BaseName + TEXT(".Cache.Size.Memory");
-				Attributes.Emplace(MoveTemp(AttrName), ZenStats.CacheStats.Size.Memory);
-			}
-
-			{
-				FString AttrName = BaseName + TEXT(".Cache.UpstreamHits");
-				Attributes.Emplace(MoveTemp(AttrName), ZenStats.CacheStats.UpstreamHits);
-			}
-
-			{
-				FString AttrName = BaseName + TEXT(".Cache.UpstreamRatio");
-				Attributes.Emplace(MoveTemp(AttrName), ZenStats.CacheStats.UpstreamRatio);
-			}
-
-			{
-				FString AttrName = BaseName + TEXT(".Cache.TotalUploadedMB");
-				Attributes.Emplace(MoveTemp(AttrName), ZenStats.UpstreamStats.TotalUploadedMB);
-			}
-
-			{
-				FString AttrName = BaseName + TEXT(".Upstream.TotalDownloadedMB");
-				Attributes.Emplace(MoveTemp(AttrName), ZenStats.UpstreamStats.TotalDownloadedMB);
-			}
-
-			{
-				FString AttrName = BaseName + TEXT(".Upstream.TotalUploadedMB");
-				Attributes.Emplace(MoveTemp(AttrName), ZenStats.UpstreamStats.TotalUploadedMB);
-			}
-
-			{
-				FString AttrName = BaseName + TEXT(".Cas.Size.Large");
-				Attributes.Emplace(MoveTemp(AttrName), ZenStats.CASStats.Size.Large);
-			}
-
-			{
-				FString AttrName = BaseName + TEXT(".Cas.Size.Small");
-				Attributes.Emplace(MoveTemp(AttrName), ZenStats.CASStats.Size.Small);
-			}
-
-			{
-				FString AttrName = BaseName + TEXT(".Cas.Size.Tiny");
-				Attributes.Emplace(MoveTemp(AttrName), ZenStats.CASStats.Size.Tiny);
-			}
-
-			{
-				FString AttrName = BaseName + TEXT(".Cas.Size.Total");
-				Attributes.Emplace(MoveTemp(AttrName), ZenStats.CASStats.Size.Total);
-			}
+			UE::Zen::GetDefaultServiceInstance().GatherAnalytics(Attributes);
 		}
 #endif
 
@@ -375,4 +177,4 @@ void FStudioAnalytics::FireEvent_Loading(const FString& LoadingName, double Seco
 		FStudioAnalytics::GetProvider().RecordEvent(TEXT("Core.Loading"), Attributes);
 #endif
 	}
-}
+}
\ No newline at end of file
