Crash of commandlet when building HLODs on Tilted Towers POI template

* Prevent an array OOB access when an empty indices array is found

diff --git a/Engine/Source/Runtime/Engine/Private/StaticMesh.cpp b/Engine/Source/Runtime/Engine/Private/StaticMesh.cpp
--- a/Engine/Source/Runtime/Engine/Private/StaticMesh.cpp
+++ b/Engine/Source/Runtime/Engine/Private/StaticMesh.cpp
@@ -3781,34 +3781,29 @@ void UStaticMesh::UpdateUVChannelData(bool bRebuildAll)
 #if WITH_EDITORONLY_DATA
 static void AccumulateBounds(FBox& Bounds, const FStaticMeshLODResources& LODModel, const FStaticMeshSection& SectionInfo, const FTransform& Transform)
 {
-	const int32 SectionIndexCount = SectionInfo.NumTriangles * 3;
-	FIndexArrayView IndexBuffer = LODModel.IndexBuffer.GetArrayView();
+	const int32 FirstIndex = SectionInfo.FirstIndex;
+	const int32 LastIndex = FirstIndex + SectionInfo.NumTriangles * 3;
+	const int32 NumIndices = LODModel.IndexBuffer.GetNumIndices();
 
-	FBox TransformedBox(ForceInit);
-	for (uint32 TriangleIndex = 0; TriangleIndex < SectionInfo.NumTriangles; ++TriangleIndex)
+	if (LastIndex < NumIndices)
 	{
-		const int32 Index0 = IndexBuffer[SectionInfo.FirstIndex + TriangleIndex * 3 + 0];
-		const int32 Index1 = IndexBuffer[SectionInfo.FirstIndex + TriangleIndex * 3 + 1];
-		const int32 Index2 = IndexBuffer[SectionInfo.FirstIndex + TriangleIndex * 3 + 2];
-
-		FVector Pos1 = Transform.TransformPosition(FVector(LODModel.VertexBuffers.PositionVertexBuffer.VertexPosition(Index1)));
-		FVector Pos2 = Transform.TransformPosition(FVector(LODModel.VertexBuffers.PositionVertexBuffer.VertexPosition(Index2)));
-		FVector Pos0 = Transform.TransformPosition(FVector(LODModel.VertexBuffers.PositionVertexBuffer.VertexPosition(Index0)));
-
-		Bounds += Pos0;
-		Bounds += Pos1;
-		Bounds += Pos2;
+		const FIndexArrayView IndexBuffer = LODModel.IndexBuffer.GetArrayView();
+		for (int32 Index = FirstIndex; Index < LastIndex; ++Index)
+		{
+			Bounds += Transform.TransformPosition(FVector(LODModel.VertexBuffers.PositionVertexBuffer.VertexPosition(Index)));
+		}
 	}
 }
 #endif
 
 FBox UStaticMesh::GetMaterialBox(int32 MaterialIndex, const FTransform& Transform) const
 {
+	FBox MaterialBounds(ForceInit);
+
 #if WITH_EDITORONLY_DATA
 	// Once cooked, the data requires to compute the scales will not be CPU accessible.
 	if (FPlatformProperties::HasEditorOnlyData() && GetRenderData())
 	{
-		FBox MaterialBounds(ForceInit);
 		for (const FStaticMeshLODResources& LODModel : GetRenderData()->LODResources)
 		{
 			for (const FStaticMeshSection& SectionInfo : LODModel.Sections)
@@ -3819,11 +3814,16 @@ FBox UStaticMesh::GetMaterialBox(int32 MaterialIndex, const FTransform& Transfor
 				AccumulateBounds(MaterialBounds, LODModel, SectionInfo, Transform);
 			}
 		}
-		return MaterialBounds;
 	}
 #endif
-	// Fallback back using the full bounds.
-	return GetBoundingBox().TransformBy(Transform);
+
+	if (!MaterialBounds.IsValid)
+	{
+		// Fallback back using the full bounds.
+		MaterialBounds = GetBoundingBox().TransformBy(Transform);
+	}
+
+	return MaterialBounds;
 }
 
 const FMeshUVChannelInfo* UStaticMesh::GetUVChannelData(int32 MaterialIndex) const
