Audio Dashboard Multi-PIE Support Checkpoint P1:
Add DeviceId to dashboard messages (To be used in Dashboard to filter data when running multi-PIE for example)
- Misc Clean-up
#rnx

diff --git a/Engine/Source/Runtime/AudioMixer/Private/AudioMixerSource.cpp b/Engine/Source/Runtime/AudioMixer/Private/AudioMixerSource.cpp
--- a/Engine/Source/Runtime/AudioMixer/Private/AudioMixerSource.cpp
+++ b/Engine/Source/Runtime/AudioMixer/Private/AudioMixerSource.cpp
@@ -23,14 +23,18 @@ CSV_DECLARE_CATEGORY_MODULE_EXTERN(AUDIOMIXERCORE_API, Audio);
 
 #if UE_AUDIO_PROFILERTRACE_ENABLED
 UE_TRACE_EVENT_BEGIN(Audio, MixerSourceStart)
+	UE_TRACE_EVENT_FIELD(uint32, DeviceId)
 	UE_TRACE_EVENT_FIELD(uint64, Timestamp)
+	UE_TRACE_EVENT_FIELD(uint32, PlayOrder)
 	UE_TRACE_EVENT_FIELD(int32, SourceId)
+	UE_TRACE_EVENT_FIELD(uint64, ComponentId)
 	UE_TRACE_EVENT_FIELD(UE::Trace::WideString, Name)
 UE_TRACE_EVENT_END()
 
 UE_TRACE_EVENT_BEGIN(Audio, MixerSourceStop)
+	UE_TRACE_EVENT_FIELD(uint32, DeviceId)
 	UE_TRACE_EVENT_FIELD(uint64, Timestamp)
-	UE_TRACE_EVENT_FIELD(int32, SourceId)
+	UE_TRACE_EVENT_FIELD(uint32, PlayOrder)
 UE_TRACE_EVENT_END()
 #endif // UE_AUDIO_PROFILERTRACE_ENABLED
 
@@ -689,6 +693,7 @@ namespace Audio
 			InitParams.bEnableBusSends = WaveInstance->bEnableBusSends;
 			InitParams.bEnableBaseSubmix = WaveInstance->bEnableBaseSubmix;
 			InitParams.bEnableSubmixSends = WaveInstance->bEnableSubmixSends;
+			InitParams.PlayOrder = WaveInstance->GetPlayOrder();
 			bPreviousBusEnablement = WaveInstance->bEnableBusSends;
 			DynamicBusSendInfos.Reset();
 
@@ -1248,16 +1253,22 @@ namespace Audio
 			const bool bChannelEnabled = UE_TRACE_CHANNELEXPR_IS_ENABLED(AudioMixerChannel);
 			if (bChannelEnabled && WaveInstance)
 			{
-				int32 TraceSourceId = INDEX_NONE;
-				if (MixerSourceVoice)
+				if (const FActiveSound* ActiveSound = WaveInstance->ActiveSound)
 				{
-					TraceSourceId = MixerSourceVoice->GetSourceId();
-				}
+					int32 TraceSourceId = INDEX_NONE;
+					if (MixerSourceVoice)
+					{
+						TraceSourceId = MixerSourceVoice->GetSourceId();
+					}
 
-				UE_TRACE_LOG(Audio, MixerSourceStart, AudioMixerChannel)
-					<< MixerSourceStart.Timestamp(FPlatformTime::Cycles64())
-					<< MixerSourceStart.SourceId(TraceSourceId)
-					<< MixerSourceStart.Name(*WaveInstance->WaveData->GetPathName());
+					UE_TRACE_LOG(Audio, MixerSourceStart, AudioMixerChannel)
+						<< MixerSourceStart.DeviceId(MixerDevice->DeviceID)
+						<< MixerSourceStart.Timestamp(FPlatformTime::Cycles64())
+						<< MixerSourceStart.PlayOrder(WaveInstance->GetPlayOrder())
+						<< MixerSourceStart.SourceId(TraceSourceId)
+						<< MixerSourceStart.ComponentId(ActiveSound->GetAudioComponentID())
+						<< MixerSourceStart.Name(*WaveInstance->WaveData->GetPathName());
+				}
 			}
 #endif // UE_AUDIO_PROFILERTRACE_ENABLED
 		}
@@ -1355,8 +1366,9 @@ namespace Audio
 					}
 
 					UE_TRACE_LOG(Audio, MixerSourceStop, AudioMixerChannel)
+						<< MixerSourceStop.DeviceId(MixerDevice->DeviceID)
 						<< MixerSourceStop.Timestamp(FPlatformTime::Cycles64())
-						<< MixerSourceStop.SourceId(TraceSourceId);
+						<< MixerSourceStop.PlayOrder(WaveInstance->GetPlayOrder());
 				}
 #endif // UE_AUDIO_PROFILERTRACE_ENABLED
 
diff --git a/Engine/Source/Runtime/AudioMixer/Private/AudioMixerSourceManager.cpp b/Engine/Source/Runtime/AudioMixer/Private/AudioMixerSourceManager.cpp
--- a/Engine/Source/Runtime/AudioMixer/Private/AudioMixerSourceManager.cpp
+++ b/Engine/Source/Runtime/AudioMixer/Private/AudioMixerSourceManager.cpp
@@ -171,33 +171,38 @@ DEFINE_STAT(STAT_AudioMixerSourceOutputBuffers);
 
 #if UE_AUDIO_PROFILERTRACE_ENABLED
 UE_TRACE_EVENT_BEGIN(Audio, MixerSourceVolume)
+	UE_TRACE_EVENT_FIELD(uint32, DeviceId)
 	UE_TRACE_EVENT_FIELD(uint64, Timestamp)
-	UE_TRACE_EVENT_FIELD(int32, SourceId)
+	UE_TRACE_EVENT_FIELD(uint32, PlayOrder)
 	UE_TRACE_EVENT_FIELD(float, Volume)
 UE_TRACE_EVENT_END()
 
 UE_TRACE_EVENT_BEGIN(Audio, MixerSourceDistanceAttenuation)
+	UE_TRACE_EVENT_FIELD(uint32, DeviceId)
 	UE_TRACE_EVENT_FIELD(uint64, Timestamp)
-	UE_TRACE_EVENT_FIELD(int32, SourceId)
+	UE_TRACE_EVENT_FIELD(uint32, PlayOrder)
 	UE_TRACE_EVENT_FIELD(float, DistanceAttenuation)
 UE_TRACE_EVENT_END()
 
 UE_TRACE_EVENT_BEGIN(Audio, MixerSourcePitch)
+	UE_TRACE_EVENT_FIELD(uint32, DeviceId)
 	UE_TRACE_EVENT_FIELD(uint64, Timestamp)
-	UE_TRACE_EVENT_FIELD(int32, SourceId)
+	UE_TRACE_EVENT_FIELD(uint32, PlayOrder)
 	UE_TRACE_EVENT_FIELD(float, Pitch)
 UE_TRACE_EVENT_END()
 
 UE_TRACE_EVENT_BEGIN(Audio, MixerSourceFilters)
+	UE_TRACE_EVENT_FIELD(uint32, DeviceId)
 	UE_TRACE_EVENT_FIELD(uint64, Timestamp)
-	UE_TRACE_EVENT_FIELD(int32, SourceId)
+	UE_TRACE_EVENT_FIELD(uint32, PlayOrder)
 	UE_TRACE_EVENT_FIELD(float, LPFFrequency)
 	UE_TRACE_EVENT_FIELD(float, HPFFrequency)
 UE_TRACE_EVENT_END()
 
 UE_TRACE_EVENT_BEGIN(Audio, MixerSourceEnvelope)
+	UE_TRACE_EVENT_FIELD(uint32, DeviceId)
 	UE_TRACE_EVENT_FIELD(uint64, Timestamp)
-	UE_TRACE_EVENT_FIELD(int32, SourceId)
+	UE_TRACE_EVENT_FIELD(uint32, PlayOrder)
 	UE_TRACE_EVENT_FIELD(float, Envelope)
 UE_TRACE_EVENT_END()
 #endif // UE_AUDIO_PROFILERTRACE_ENABLED
@@ -692,6 +697,9 @@ namespace Audio
 		SourceInfo.bHasPreDistanceAttenuationSend = false;
 		SourceInfo.bModFiltersUpdated = false;
 
+		SourceInfo.AudioComponentID = 0;
+		SourceInfo.PlayOrder = INDEX_NONE;
+
 		SourceInfo.QuantizedCommandHandle.Reset();
 
 #if AUDIO_MIXER_ENABLE_DEBUG_MODE
@@ -866,6 +874,7 @@ namespace Audio
 			SourceInfo.bUseHRTFSpatializer = InitParams.bUseHRTFSpatialization;
 			SourceInfo.bIsExternalSend = InitParams.bIsExternalSend;
 			SourceInfo.bIsVorbis = InitParams.bIsVorbis;
+			SourceInfo.PlayOrder = InitParams.PlayOrder;
 			SourceInfo.AudioComponentID = InitParams.AudioComponentID;
 			SourceInfo.bIsSoundfield = InitParams.bIsSoundfield;
 
@@ -2168,8 +2177,9 @@ namespace Audio
 				if (bChannelEnabled)
 				{
 					UE_TRACE_LOG(Audio, MixerSourcePitch, AudioMixerChannel)
+						<< MixerSourcePitch.DeviceId(MixerDevice->DeviceID)
 						<< MixerSourcePitch.Timestamp(FPlatformTime::Cycles64())
-						<< MixerSourcePitch.SourceId(SourceId)
+						<< MixerSourcePitch.PlayOrder(SourceInfo.PlayOrder)
 						<< MixerSourcePitch.Pitch(TargetPitch);
 				}
 #endif // UE_AUDIO_PROFILERTRACE_ENABLED
@@ -2607,13 +2617,15 @@ namespace Audio
 				if (bChannelEnabled)
 				{
 					UE_TRACE_LOG(Audio, MixerSourceVolume, AudioMixerChannel)
+						<< MixerSourceVolume.DeviceId(MixerDevice->DeviceID)
 						<< MixerSourceVolume.Timestamp(FPlatformTime::Cycles64())
-						<< MixerSourceVolume.SourceId(SourceId)
+						<< MixerSourceVolume.PlayOrder(SourceInfo.PlayOrder)
 						<< MixerSourceVolume.Volume(VolumeDestination);
 
 					UE_TRACE_LOG(Audio, MixerSourceDistanceAttenuation, AudioMixerChannel)
+						<< MixerSourceDistanceAttenuation.DeviceId(MixerDevice->DeviceID)
 						<< MixerSourceDistanceAttenuation.Timestamp(FPlatformTime::Cycles64())
-						<< MixerSourceDistanceAttenuation.SourceId(SourceId)
+						<< MixerSourceDistanceAttenuation.PlayOrder(SourceInfo.PlayOrder)
 						<< MixerSourceDistanceAttenuation.DistanceAttenuation(SourceInfo.DistanceAttenuationSourceDestination);
 				}
 #endif // UE_AUDIO_PROFILERTRACE_ENABLED
@@ -2750,13 +2762,15 @@ namespace Audio
 					}
 
 					UE_TRACE_LOG(Audio, MixerSourceFilters, AudioMixerChannel)
+						<< MixerSourceFilters.DeviceId(MixerDevice->DeviceID)
 						<< MixerSourceFilters.Timestamp(FPlatformTime::Cycles64())
-						<< MixerSourceFilters.SourceId(SourceId)
+						<< MixerSourceFilters.PlayOrder(SourceInfo.PlayOrder)
 						<< MixerSourceFilters.HPFFrequency(HPFFrequency)
 						<< MixerSourceFilters.LPFFrequency(LPFFrequency);
 					UE_TRACE_LOG(Audio, MixerSourceEnvelope, AudioMixerChannel)
+						<< MixerSourceEnvelope.DeviceId(MixerDevice->DeviceID)
 						<< MixerSourceEnvelope.Timestamp(FPlatformTime::Cycles64())
-						<< MixerSourceEnvelope.SourceId(SourceId)
+						<< MixerSourceEnvelope.PlayOrder(SourceInfo.PlayOrder)
 						<< MixerSourceEnvelope.Envelope(SourceInfo.SourceEnvelopeValue);
 				}
 #endif // UE_AUDIO_PROFILERTRACE_ENABLED
diff --git a/Engine/Source/Runtime/AudioMixer/Private/AudioMixerSourceManager.h b/Engine/Source/Runtime/AudioMixer/Private/AudioMixerSourceManager.h
--- a/Engine/Source/Runtime/AudioMixer/Private/AudioMixerSourceManager.h
+++ b/Engine/Source/Runtime/AudioMixer/Private/AudioMixerSourceManager.h
@@ -130,7 +130,7 @@ namespace Audio
 
 		FQuartzQuantizedRequestData QuantizedRequestData;
 
-		FSharedISourceBufferListenerPtr SourceBufferListener;		
+		FSharedISourceBufferListenerPtr SourceBufferListener;
 
 		IAudioLinkFactory::FAudioLinkSourcePushedSharedPtr AudioLink;
 
@@ -148,6 +148,8 @@ namespace Audio
 		bool bIsSoundfield = false;
 		bool bIsSeeking = false;
 		bool bShouldSourceBufferListenerZeroBuffer = false;
+
+		uint32 PlayOrder = INDEX_NONE;
 	};
 
 	struct FSourceManagerInitParams
@@ -512,6 +514,8 @@ namespace Audio
 			int32 NumPostEffectChannels;
 			int32 NumInputFrames;
 
+			uint32 PlayOrder;
+
 			// ID for associated Audio Component if there is one, 0 otherwise
 			uint64 AudioComponentID;
 
diff --git a/Engine/Source/Runtime/Engine/Private/Audio.cpp b/Engine/Source/Runtime/Engine/Private/Audio.cpp
--- a/Engine/Source/Runtime/Engine/Private/Audio.cpp
+++ b/Engine/Source/Runtime/Engine/Private/Audio.cpp
@@ -760,8 +760,8 @@ FArchive& operator<<( FArchive& Ar, FNotifyBufferFinishedHooks& NotifyHook )
 	FWaveInstance implementation.
 -----------------------------------------------------------------------------*/
 
-/** Helper to create good unique type hashs for FWaveInstance instances */
-uint32 FWaveInstance::TypeHashCounter = 0;
+/** Helper to create good play order for FWaveInstance instances */
+uint32 FWaveInstance::PlayOrderCounter = 0;
 
 PRAGMA_DISABLE_DEPRECATION_WARNINGS
 FWaveInstance::FWaveInstance(FWaveInstance&&) = default;
@@ -834,11 +834,11 @@ FWaveInstance::FWaveInstance(const UPTRINT InWaveInstanceHash, FActiveSound& InA
 	, PlaybackTime(0.0f)
 	, ReverbSendLevel(0.0f)
 	, ManualReverbSendLevel(0.0f)
-	, TypeHash(0)
+	, PlayOrder(0)
 	, WaveInstanceHash(InWaveInstanceHash)
 	, UserIndex(0)
 {
-	TypeHash = ++TypeHashCounter;
+	PlayOrder = ++PlayOrderCounter;
 }
 
 bool FWaveInstance::IsPlaying() const
diff --git a/Engine/Source/Runtime/Engine/Private/AudioDeviceManager.cpp b/Engine/Source/Runtime/Engine/Private/AudioDeviceManager.cpp
--- a/Engine/Source/Runtime/Engine/Private/AudioDeviceManager.cpp
+++ b/Engine/Source/Runtime/Engine/Private/AudioDeviceManager.cpp
@@ -874,6 +874,20 @@ FAudioDevice* FAudioDeviceManager::GetAudioDeviceRaw(Audio::FDeviceId InDeviceID
 	return AudioDevice;
 }
 
+const FAudioDevice* FAudioDeviceManager::GetAudioDeviceRaw(Audio::FDeviceId InDeviceID) const
+{
+	FScopeLock ScopeLock(&DeviceMapCriticalSection);
+	if (!IsValidAudioDevice(InDeviceID))
+	{
+		return nullptr;
+	}
+
+	const FAudioDevice* AudioDevice = Devices[InDeviceID].Device;
+	check(AudioDevice != nullptr);
+
+	return AudioDevice;
+}
+
 void FAudioDeviceManager::SetAudioDevice(UWorld& InWorld, Audio::FDeviceId InDeviceID)
 {
 	FScopeLock ScopeLock(&DeviceMapCriticalSection);
@@ -1026,6 +1040,22 @@ void FAudioDeviceManager::IterateOverAllDevices(TUniqueFunction<void(Audio::FDev
 	}
 }
 
+void FAudioDeviceManager::IterateOverAllDevices(TUniqueFunction<void(Audio::FDeviceId, const FAudioDevice*)> ForEachDevice) const
+{
+	TArray<Audio::FDeviceId> DeviceIDs;
+	{
+		FScopeLock ScopeLock(&DeviceMapCriticalSection);
+		Devices.GetKeys(DeviceIDs);
+	}
+
+	for (const Audio::FDeviceId DeviceID : DeviceIDs)
+	{
+		if (const FAudioDevice* Device = GetAudioDeviceRaw(DeviceID))
+		{
+			ForEachDevice(DeviceID, Device);
+		}
+	}
+}
 
 void FAudioDeviceManager::AddReferencedObjects(FReferenceCollector& Collector)
 {
diff --git a/Engine/Source/Runtime/Engine/Private/AudioVirtualLoop.cpp b/Engine/Source/Runtime/Engine/Private/AudioVirtualLoop.cpp
--- a/Engine/Source/Runtime/Engine/Private/AudioVirtualLoop.cpp
+++ b/Engine/Source/Runtime/Engine/Private/AudioVirtualLoop.cpp
@@ -44,17 +44,21 @@ FAutoConsoleVariableRef CVarVirtualLoopsUpdateRateMax(
 
 #if UE_AUDIO_PROFILERTRACE_ENABLED
 UE_TRACE_EVENT_BEGIN(Audio, VirtualLoopVirtualize)
+	UE_TRACE_EVENT_FIELD(uint32, DeviceId)
 	UE_TRACE_EVENT_FIELD(uint64, Timestamp)
 	UE_TRACE_EVENT_FIELD(uint32, PlayOrder)
+	UE_TRACE_EVENT_FIELD(uint64, ComponentId)
 	UE_TRACE_EVENT_FIELD(UE::Trace::WideString, Name)
 UE_TRACE_EVENT_END()
 
 UE_TRACE_EVENT_BEGIN(Audio, VirtualLoopRealize)
+	UE_TRACE_EVENT_FIELD(uint32, DeviceId)
 	UE_TRACE_EVENT_FIELD(uint64, Timestamp)
 	UE_TRACE_EVENT_FIELD(uint32, PlayOrder)
 UE_TRACE_EVENT_END()
 
 UE_TRACE_EVENT_BEGIN(Audio, VirtualLoopUpdate)
+	UE_TRACE_EVENT_FIELD(uint32, DeviceId)
 	UE_TRACE_EVENT_FIELD(double, Timestamp)
 	UE_TRACE_EVENT_FIELD(uint32, PlayOrder)
 	UE_TRACE_EVENT_FIELD(float, TimeVirtualized)
@@ -125,10 +129,15 @@ bool FAudioVirtualLoop::Virtualize(const FActiveSound& InActiveSound, FAudioDevi
 	const bool bChannelEnabled = UE_TRACE_CHANNELEXPR_IS_ENABLED(AudioChannel);
 	if (bChannelEnabled)
 	{
-		UE_TRACE_LOG(Audio, VirtualLoopVirtualize, AudioChannel)
-			<< VirtualLoopVirtualize.Timestamp(FPlatformTime::Cycles64())
-			<< VirtualLoopVirtualize.PlayOrder(ActiveSound->GetPlayOrder())
-			<< VirtualLoopVirtualize.Name(Sound ? *Sound->GetPathName() : TEXT("N/A"));
+		if (const FAudioDevice* AudioDevice = ActiveSound->AudioDevice)
+		{
+			UE_TRACE_LOG(Audio, VirtualLoopVirtualize, AudioChannel)
+				<< VirtualLoopVirtualize.DeviceId(static_cast<uint32>(AudioDevice->DeviceID))
+				<< VirtualLoopVirtualize.Timestamp(FPlatformTime::Cycles64())
+				<< VirtualLoopVirtualize.PlayOrder(ActiveSound->GetPlayOrder())
+				<< VirtualLoopVirtualize.ComponentId(ActiveSound->GetAudioComponentID())
+				<< VirtualLoopVirtualize.Name(Sound ? *Sound->GetPathName() : TEXT("N/A"));
+		}
 	}
 #endif // UE_AUDIO_PROFILERTRACE_ENABLED
 	return true;
@@ -266,23 +275,27 @@ bool FAudioVirtualLoop::Update(float DeltaTime, bool bForceUpdate)
 	const bool bChannelEnabled = UE_TRACE_CHANNELEXPR_IS_ENABLED(AudioChannel);
 	if (bChannelEnabled)
 	{
-		const USoundBase* Sound = ActiveSound->GetSound();
-
-		const FTransform& Transform = ActiveSound->Transform;
-		const FVector& Location = Transform.GetLocation();
-		const FRotator Rotator = Transform.GetRotation().Rotator();
-		UE_TRACE_LOG(Audio, VirtualLoopUpdate, AudioChannel)
-			<< VirtualLoopUpdate.Timestamp(FPlatformTime::Cycles64())
-			<< VirtualLoopUpdate.PlayOrder(ActiveSound->GetPlayOrder())
-			<< VirtualLoopUpdate.TimeVirtualized(TimeVirtualized)
-			<< VirtualLoopUpdate.PlaybackTime(ActiveSound->PlaybackTime)
-			<< VirtualLoopUpdate.UpdateInterval(UpdateInterval)
-			<< VirtualLoopUpdate.LocationX(Location.X)
-			<< VirtualLoopUpdate.LocationY(Location.Y)
-			<< VirtualLoopUpdate.LocationZ(Location.Z)
-			<< VirtualLoopUpdate.RotatorPitch(Rotator.Pitch)
-			<< VirtualLoopUpdate.RotatorYaw(Rotator.Yaw)
-			<< VirtualLoopUpdate.RotatorRoll(Rotator.Roll);
+		if (const FAudioDevice* AudioDevice = ActiveSound->AudioDevice)
+		{
+			const USoundBase* Sound = ActiveSound->GetSound();
+
+			const FTransform& Transform = ActiveSound->Transform;
+			const FVector& Location = Transform.GetLocation();
+			const FRotator Rotator = Transform.GetRotation().Rotator();
+			UE_TRACE_LOG(Audio, VirtualLoopUpdate, AudioChannel)
+				<< VirtualLoopUpdate.DeviceId(static_cast<uint32>(AudioDevice->DeviceID))
+				<< VirtualLoopUpdate.Timestamp(FPlatformTime::Cycles64())
+				<< VirtualLoopUpdate.PlayOrder(ActiveSound->GetPlayOrder())
+				<< VirtualLoopUpdate.TimeVirtualized(TimeVirtualized)
+				<< VirtualLoopUpdate.PlaybackTime(ActiveSound->PlaybackTime)
+				<< VirtualLoopUpdate.UpdateInterval(UpdateInterval)
+				<< VirtualLoopUpdate.LocationX(Location.X)
+				<< VirtualLoopUpdate.LocationY(Location.Y)
+				<< VirtualLoopUpdate.LocationZ(Location.Z)
+				<< VirtualLoopUpdate.RotatorPitch(Rotator.Pitch)
+				<< VirtualLoopUpdate.RotatorYaw(Rotator.Yaw)
+				<< VirtualLoopUpdate.RotatorRoll(Rotator.Roll);
+		}
 	}
 #endif // UE_AUDIO_PROFILERTRACE_ENABLED
 
@@ -298,9 +311,13 @@ bool FAudioVirtualLoop::Update(float DeltaTime, bool bForceUpdate)
 #if UE_AUDIO_PROFILERTRACE_ENABLED
 	if (bChannelEnabled)
 	{
-		UE_TRACE_LOG(Audio, VirtualLoopRealize, AudioChannel)
-			<< VirtualLoopRealize.Timestamp(FPlatformTime::Cycles64())
-			<< VirtualLoopRealize.PlayOrder(ActiveSound->GetPlayOrder());
+		if (const FAudioDevice* AudioDevice = ActiveSound->AudioDevice)
+		{
+			UE_TRACE_LOG(Audio, VirtualLoopRealize, AudioChannel)
+				<< VirtualLoopRealize.DeviceId(static_cast<uint32>(AudioDevice->DeviceID))
+				<< VirtualLoopRealize.Timestamp(FPlatformTime::Cycles64())
+				<< VirtualLoopRealize.PlayOrder(ActiveSound->GetPlayOrder());
+		}
 	}
 #endif // UE_AUDIO_PROFILERTRACE_ENABLED
 	return true;
diff --git a/Engine/Source/Runtime/Engine/Public/Audio.h b/Engine/Source/Runtime/Engine/Public/Audio.h
--- a/Engine/Source/Runtime/Engine/Public/Audio.h
+++ b/Engine/Source/Runtime/Engine/Public/Audio.h
@@ -180,9 +180,11 @@ enum class EBusSendType : uint8
  */
 struct ENGINE_API FWaveInstance
 {
+private:
 	/** Static helper to create good unique type hashes */
-	static uint32 TypeHashCounter;
+	static uint32 PlayOrderCounter;
 
+public:
 	/** Wave data */
 	USoundWave* WaveData;
 
@@ -216,7 +218,7 @@ struct ENGINE_API FWaveInstance
 
 	/** AudioLink Opt in */
 	bool bShouldUseAudioLink = true;
-	UAudioLinkSettingsAbstract* AudioLinkSettingsOverride = nullptr;	
+	UAudioLinkSettingsAbstract* AudioLinkSettingsOverride = nullptr;
 
 private:
 
@@ -397,9 +399,11 @@ public:
 	/** The submix send settings to use. */
 	TArray<FAttenuationSubmixSendSettings> SubmixSendSettings;
 
-	/** Cached type hash */
-	uint32 TypeHash;
+private:
+	/** Cached play order */
+	uint32 PlayOrder;
 
+public:
 	/** Hash value for finding the wave instance based on the path through the cue to get to it */
 	UPTRINT WaveInstanceHash;
 
@@ -498,7 +502,9 @@ public:
 	/** Whether spatialization is an external send */
 	bool GetSpatializationIsExternalSend() const {	return bSpatializationIsExternalSend;}
 
-	friend inline uint32 GetTypeHash(FWaveInstance* A) { return A->TypeHash; }
+	uint32 GetPlayOrder() const { return PlayOrder; }
+
+	friend inline uint32 GetTypeHash(FWaveInstance* A) { return A->PlayOrder; }
 };
 
 /*-----------------------------------------------------------------------------
diff --git a/Engine/Source/Runtime/Engine/Public/AudioDeviceHandle.h b/Engine/Source/Runtime/Engine/Public/AudioDeviceHandle.h
--- a/Engine/Source/Runtime/Engine/Public/AudioDeviceHandle.h
+++ b/Engine/Source/Runtime/Engine/Public/AudioDeviceHandle.h
@@ -46,7 +46,7 @@ public:
 	// Returns world set when handle was created
 	TWeakObjectPtr<UWorld> GetWorld() const;
 
-	// Checks whether this points to a valid compressed chunk.
+	// Checks whether this points to a valid AudioDevice.
 	bool IsValid() const;
 
 	// Resets handle to invalid state
diff --git a/Engine/Source/Runtime/Engine/Public/AudioDeviceManager.h b/Engine/Source/Runtime/Engine/Public/AudioDeviceManager.h
--- a/Engine/Source/Runtime/Engine/Public/AudioDeviceManager.h
+++ b/Engine/Source/Runtime/Engine/Public/AudioDeviceManager.h
@@ -140,6 +140,7 @@ public:
 	 * handle is invalid then a NULL device ptr will be returned.
 	 */
 	FAudioDevice* GetAudioDeviceRaw(Audio::FDeviceId InDeviceID);
+	const FAudioDevice* GetAudioDeviceRaw(Audio::FDeviceId InDeviceID) const;
 
 	/**
 	  Sets the device associated with the given world.
@@ -172,7 +173,9 @@ public:
 	/** Updates all active audio devices */
 	void UpdateActiveAudioDevices(bool bGameTicking);
 
+	/** Iterates over all managed audio devices */
 	void IterateOverAllDevices(TUniqueFunction<void(Audio::FDeviceId, FAudioDevice*)> ForEachDevice);
+	void IterateOverAllDevices(TUniqueFunction<void(Audio::FDeviceId, const FAudioDevice*)> ForEachDevice) const;
 
 	/** Tracks objects in the active audio devices. */
 	void AddReferencedObjects(FReferenceCollector& Collector);
