Fixed crash when loading actors with child actor components when AActor::RegisterAllComponents/RerunConstructionScripts would recreate child actors and loading code would process old/destroyed child actors.

diff --git a/Engine/Source/Runtime/Engine/Private/Level.cpp b/Engine/Source/Runtime/Engine/Private/Level.cpp
--- a/Engine/Source/Runtime/Engine/Private/Level.cpp
+++ b/Engine/Source/Runtime/Engine/Private/Level.cpp
@@ -892,6 +892,12 @@ void ULevel::AddLoadedActors(const TArray<AActor*>& ActorList, const FTransform*
 	// Register all components
 	for (AActor* Actor : ActorsQueue)
 	{
+		// RegisterAllComponents can destroy child actors
+		if (!IsValid(Actor))
+		{
+			continue;
+		}
+
 		if (TransformToApply)
 		{
 			FLevelUtils::FApplyLevelTransformParams TransformParams(this, *TransformToApply);
@@ -911,6 +917,12 @@ void ULevel::AddLoadedActors(const TArray<AActor*>& ActorList, const FTransform*
 	// Rerun construction scripts
 	for (AActor* Actor : ActorsQueue)
 	{
+		// RegisterAllComponents/RerunConstructionScripts can destroy child actors
+		if (!IsValid(Actor))
+		{
+			continue;
+		}
+
 		if (bAreComponentsCurrentlyRegistered)
 		{
 			Actor->RerunConstructionScripts();
@@ -922,6 +934,12 @@ void ULevel::AddLoadedActors(const TArray<AActor*>& ActorList, const FTransform*
 	// Finalize actors
 	for (AActor* Actor : ActorsQueue)
 	{
+		// RegisterAllComponents/RerunConstructionScripts can destroy child actors
+		if (!IsValid(Actor))
+		{
+			continue;
+		}
+
 		if (bAreComponentsCurrentlyRegistered)
 		{
 			GetWorld()->UpdateCullDistanceVolumes(Actor);
@@ -985,6 +1003,12 @@ void ULevel::RemoveLoadedActors(const TArray<AActor*>& ActorList, const FTransfo
 
 	for (AActor* Actor : ActorsQueue)
 	{
+		// UnregisterAllComponents can destroy child actors
+		if (!IsValid(Actor))
+		{
+			continue;
+		}
+
 		Actor->UnregisterAllComponents();
 		Actor->RegisterAllActorTickFunctions(false, true);
 
