HLODs don't disappear when the building loads back in
* Retrieve the WorldPartition through the GetCellOwner() call

diff --git a/Engine/Source/Runtime/Engine/Private/WorldPartition/WorldPartitionLevelStreamingDynamic.cpp b/Engine/Source/Runtime/Engine/Private/WorldPartition/WorldPartitionLevelStreamingDynamic.cpp
--- a/Engine/Source/Runtime/Engine/Private/WorldPartition/WorldPartitionLevelStreamingDynamic.cpp
+++ b/Engine/Source/Runtime/Engine/Private/WorldPartition/WorldPartitionLevelStreamingDynamic.cpp
@@ -623,24 +623,33 @@ UWorld* UWorldPartitionLevelStreamingDynamic::GetStreamingWorld() const
 	return OuterWorldPartition->GetTypedOuter<UWorld>();
 }
 
-bool UWorldPartitionLevelStreamingDynamic::CanMakeVisible()
+bool UWorldPartitionLevelStreamingDynamic::CanChangeVisibility(bool bMakeVisible) const
 {
 	const ENetMode NetMode = GetWorld()->GetNetMode();
 	if (NetMode != NM_DedicatedServer)
 	{
-		// @todo_ow ContentBundles do not support Hlods and events below are forwarding
-		// Events to the HLodSubsystem which assumes knowledge of all cells (not true with plugins)
 		if (const UWorldPartitionRuntimeLevelStreamingCell* RuntimeLevelStreamingCell = StreamingCell.Get())
 		{
-			if (RuntimeLevelStreamingCell->GetIsHLOD() && !RuntimeLevelStreamingCell->GetContentBundleID().IsValid())
+			// Source cells that aren't HLOD can always be made visible
+			bool bAlwaysAllowVisibilityChange = bMakeVisible && !RuntimeLevelStreamingCell->GetIsHLOD();
+			
+			// @todo_ow ContentBundles do not support Hlods and events below are forwarding
+			// Events to the HLodSubsystem which assumes knowledge of all cells (not true with plugins)
+			bAlwaysAllowVisibilityChange |= RuntimeLevelStreamingCell->GetContentBundleID().IsValid();
+
+			if (!bAlwaysAllowVisibilityChange)
 			{
-				if (const UWorldPartition* WorldPartition = GetWorld()->GetWorldPartition())
+				if (const UWorldPartition* WorldPartition = RuntimeLevelStreamingCell->GetCellOwner()->GetWorldPartition())
 				{
 					if (UHLODSubsystem* HLODSubsystem = GetWorld()->GetSubsystem<UHLODSubsystem>())
 					{
-						if (!HLODSubsystem->CanMakeVisible(StreamingCell.Get()))
+						if (bMakeVisible)
+						{
+							return HLODSubsystem->CanMakeVisible(RuntimeLevelStreamingCell);
+						}
+						else
 						{
-							return false;
+							return HLODSubsystem->CanMakeInvisible(RuntimeLevelStreamingCell);
 						}
 					}
 				}
@@ -648,35 +657,19 @@ bool UWorldPartitionLevelStreamingDynamic::CanMakeVisible()
 		}
 	}
 
-	return Super::CanMakeVisible();
+	return true;
 }
 
-bool UWorldPartitionLevelStreamingDynamic::CanMakeInvisible()
+bool UWorldPartitionLevelStreamingDynamic::CanMakeVisible()
 {
-	const ENetMode NetMode = GetWorld()->GetNetMode();
-	if (NetMode != NM_DedicatedServer)
-	{
-		// @todo_ow ContentBundles do not support Hlods and events below are forwarding
-		// Events to the HLodSubsystem which assumes knowledge of all cells (not true with plugins)
-		if (const UWorldPartitionRuntimeLevelStreamingCell* RuntimeLevelStreamingCell = StreamingCell.Get())
-		{
-			if (!RuntimeLevelStreamingCell->GetContentBundleID().IsValid())
-			{
-				if (const UWorldPartition* WorldPartition = GetWorld()->GetWorldPartition())
-				{
-					if (UHLODSubsystem* HLODSubsystem = GetWorld()->GetSubsystem<UHLODSubsystem>())
-					{
-						if (!HLODSubsystem->CanMakeInvisible(StreamingCell.Get()))
-						{
-							return false;
-						}
-					}
-				}
-			}
-		}
-	}
+	const bool bCanMakeVisible = CanChangeVisibility(true);
+	return bCanMakeVisible && Super::CanMakeVisible();
+}
 
-	return Super::CanMakeInvisible();
+bool UWorldPartitionLevelStreamingDynamic::CanMakeInvisible()
+{
+	const bool bCanMakeInvisible = CanChangeVisibility(false);
+	return bCanMakeInvisible && Super::CanMakeInvisible();
 }
 
 bool UWorldPartitionLevelStreamingDynamic::ShouldBlockOnUnload() const
diff --git a/Engine/Source/Runtime/Engine/Private/WorldPartition/WorldPartitionRuntimeLevelStreamingCell.cpp b/Engine/Source/Runtime/Engine/Private/WorldPartition/WorldPartitionRuntimeLevelStreamingCell.cpp
--- a/Engine/Source/Runtime/Engine/Private/WorldPartition/WorldPartitionRuntimeLevelStreamingCell.cpp
+++ b/Engine/Source/Runtime/Engine/Private/WorldPartition/WorldPartitionRuntimeLevelStreamingCell.cpp
@@ -478,8 +478,7 @@ void UWorldPartitionRuntimeLevelStreamingCell::Deactivate() const
 
 void UWorldPartitionRuntimeLevelStreamingCell::OnLevelShown()
 {
-	check(LevelStreaming);
-	if (UWorldPartition* WorldPartition = LevelStreaming->GetWorld()->GetWorldPartition())
+	if (UWorldPartition* WorldPartition = GetCellOwner()->GetWorldPartition())
 	{
 		WorldPartition->OnCellShown(this);
 	}
@@ -487,10 +486,8 @@ void UWorldPartitionRuntimeLevelStreamingCell::OnLevelShown()
 
 void UWorldPartitionRuntimeLevelStreamingCell::OnLevelHidden()
 {
-	check(LevelStreaming);
-	if (UWorldPartition* WorldPartition = LevelStreaming->GetWorld()->GetWorldPartition())
+	if (UWorldPartition* WorldPartition = GetCellOwner()->GetWorldPartition())
 	{
 		WorldPartition->OnCellHidden(this);
 	}
 }
-
diff --git a/Engine/Source/Runtime/Engine/Public/WorldPartition/WorldPartitionLevelStreamingDynamic.h b/Engine/Source/Runtime/Engine/Public/WorldPartition/WorldPartitionLevelStreamingDynamic.h
--- a/Engine/Source/Runtime/Engine/Public/WorldPartition/WorldPartitionLevelStreamingDynamic.h
+++ b/Engine/Source/Runtime/Engine/Public/WorldPartition/WorldPartitionLevelStreamingDynamic.h
@@ -84,6 +84,7 @@ private:
 
 private:
 	void UpdateShouldSkipMakingVisibilityTransactionRequest();
+	bool CanChangeVisibility(bool bMakeVisible) const;
 
 	UPROPERTY()
 	bool bShouldBeAlwaysLoaded;
