Remove unused r.Mobile.UseLegacyShadingMode

diff --git a/Engine/Source/Runtime/Engine/Private/ShaderCompiler/ShaderGenerationUtil.cpp b/Engine/Source/Runtime/Engine/Private/ShaderCompiler/ShaderGenerationUtil.cpp
--- a/Engine/Source/Runtime/Engine/Private/ShaderCompiler/ShaderGenerationUtil.cpp
+++ b/Engine/Source/Runtime/Engine/Private/ShaderCompiler/ShaderGenerationUtil.cpp
@@ -87,7 +87,6 @@ void FShaderCompileUtilities::ApplyFetchEnvironment(FShaderGlobalDefines& SrcDef
 	FETCH_COMPILE_BOOL(PROJECT_SUPPORT_SKY_ATMOSPHERE_AFFECTS_HEIGHFOG);
 	FETCH_COMPILE_BOOL(SUPPORT_CLOUD_SHADOW_ON_FORWARD_LIT_TRANSLUCENT);
 	FETCH_COMPILE_BOOL(SUPPORT_CLOUD_SHADOW_ON_SINGLE_LAYER_WATER);
-	FETCH_COMPILE_BOOL(PROJECT_MOBILE_USE_LEGACY_SHADING);
 	FETCH_COMPILE_BOOL(POST_PROCESS_ALPHA);
 	FETCH_COMPILE_BOOL(PLATFORM_SUPPORTS_RENDERTARGET_WRITE_MASK);
 	FETCH_COMPILE_BOOL(PLATFORM_SUPPORTS_PER_PIXEL_DBUFFER_MASK);
@@ -370,11 +369,6 @@ static FShaderGlobalDefines FetchShaderGlobalDefines(EShaderPlatform TargetPlatf
 		Ret.SUPPORT_CLOUD_SHADOW_ON_SINGLE_LAYER_WATER = bSupportCloudShadowOnSingleLayerWater ? 1 : 0;
 	}
 
-	{
-		static IConsoleVariable* CVar = IConsoleManager::Get().FindConsoleVariable(TEXT("r.Mobile.UseLegacyShadingModel"));
-		Ret.PROJECT_MOBILE_USE_LEGACY_SHADING = CVar ? (CVar->GetInt() != 0) : 0;
-	}
-
 	{
 		static IConsoleVariable* CVar = IConsoleManager::Get().FindConsoleVariable(TEXT("r.PostProcessing.PropagateAlpha"));
 		int32 PropagateAlpha = CVar->GetInt();
diff --git a/Engine/Source/Runtime/RenderCore/Public/ShaderMaterial.h b/Engine/Source/Runtime/RenderCore/Public/ShaderMaterial.h
--- a/Engine/Source/Runtime/RenderCore/Public/ShaderMaterial.h
+++ b/Engine/Source/Runtime/RenderCore/Public/ShaderMaterial.h
@@ -34,7 +34,6 @@ struct FShaderGlobalDefines
 	bool PROJECT_SUPPORT_SKY_ATMOSPHERE_AFFECTS_HEIGHFOG;
 	bool SUPPORT_CLOUD_SHADOW_ON_FORWARD_LIT_TRANSLUCENT;
 	bool SUPPORT_CLOUD_SHADOW_ON_SINGLE_LAYER_WATER;
-	bool PROJECT_MOBILE_USE_LEGACY_SHADING;
 	bool POST_PROCESS_ALPHA;
 	bool PLATFORM_SUPPORTS_RENDERTARGET_WRITE_MASK;
 	bool PLATFORM_SUPPORTS_PER_PIXEL_DBUFFER_MASK;
