Fix issue with some stats listeners exiting early.
Some are meant to allow fire and forget behavior but a recent change meant that orphaned listeners would be cleaned up immediately.
Also fixed an issue with stats listeners when running without a rendering thread.

diff --git a/Engine/Plugins/FX/Niagara/Source/Niagara/Classes/NiagaraPerfBaseline.h b/Engine/Plugins/FX/Niagara/Source/Niagara/Classes/NiagaraPerfBaseline.h
--- a/Engine/Plugins/FX/Niagara/Source/Niagara/Classes/NiagaraPerfBaseline.h
+++ b/Engine/Plugins/FX/Niagara/Source/Niagara/Classes/NiagaraPerfBaseline.h
@@ -155,11 +155,6 @@ public:
 	virtual void BeginPlay();
 	virtual void TickActor(float DeltaTime, enum ELevelTick TickType, FActorTickFunction& ThisTickFunction);
 	//AActor Interface END
-
-	FORCEINLINE bool IsBaselineTestFinished()const { return bDone; }
-
-private:
-	bool bDone = false;
 #endif
 };
 
@@ -181,6 +176,8 @@ public:
 	virtual bool NeedsWorldStats()const override { return false; }
 	virtual bool NeedsSystemStats()const override { return true; }
 	virtual bool NeedsComponentStats()const override { return false; }
+	
+	virtual bool AllowOrphaned()const { return true; }
 
 private:
 	/** The baseline actor controlling the test conditions and which we'll send the completed stats to. */
diff --git a/Engine/Source/Runtime/Engine/Classes/Particles/ParticlePerfStatsManager.h b/Engine/Source/Runtime/Engine/Classes/Particles/ParticlePerfStatsManager.h
--- a/Engine/Source/Runtime/Engine/Classes/Particles/ParticlePerfStatsManager.h
+++ b/Engine/Source/Runtime/Engine/Classes/Particles/ParticlePerfStatsManager.h
@@ -199,6 +199,13 @@ public:
 	virtual bool NeedsWorldStats()const = 0;
 	virtual bool NeedsSystemStats()const = 0;
 	virtual bool NeedsComponentStats()const = 0;
+
+	/** 
+	Controls whether this listener should be stooped and cleaned up when it's orphaned, i.e. the manager is the only one with a reference.
+	In some cases we want to signal that a listener should stop by clearing an external reference.
+	In other cases we want to have fire and forget listeners that can signal their own termination via their tick function.
+	*/
+	virtual bool AllowOrphaned()const { return false; }
 };
 
 typedef TSharedPtr<FParticlePerfStatsListener, ESPMode::ThreadSafe> FParticlePerfStatsListenerPtr;
@@ -374,6 +381,9 @@ public:
 
 	virtual void End()override;
 	virtual bool Tick()override;
+
+	virtual bool AllowOrphaned()const { return true; }
+
 private:
 	int32 FramesRemaining;
 };
diff --git a/Engine/Source/Runtime/Engine/Private/Particles/ParticlePerfStats.cpp b/Engine/Source/Runtime/Engine/Private/Particles/ParticlePerfStats.cpp
--- a/Engine/Source/Runtime/Engine/Private/Particles/ParticlePerfStats.cpp
+++ b/Engine/Source/Runtime/Engine/Private/Particles/ParticlePerfStats.cpp
@@ -233,7 +233,7 @@ void FParticlePerfStatsManager::Tick()
 		TArray<FParticlePerfStatsListenerPtr, TInlineAllocator<8>> ToRemove;
 		for (FParticlePerfStatsListenerPtr& Listener : Listeners)
 		{
-			if (Listener.IsValid() == false || Listener.IsUnique() || Listener->Tick() == false)
+			if (Listener.IsValid() == false || (Listener.IsUnique() && Listener->AllowOrphaned() == false) || Listener->Tick() == false)
 			{
 				ToRemove.Add(Listener);
 			}
@@ -720,7 +720,7 @@ void FAccumulatedParticlePerfStats::Tick(FParticlePerfStats& Stats)
 
 void FAccumulatedParticlePerfStats::TickRT(FParticlePerfStats& Stats)
 {
-	check(IsInActualRenderingThread());
+	check(IsInRenderingThread());
 	RenderThreadStats.Tick(Stats);
 	GPUStats.Tick(Stats);
 }
