Add debug view mode for ClumpID and Per-CV AO.

diff --git a/Engine/Plugins/Runtime/HairStrands/Source/HairStrandsCore/Private/GroomComponent.cpp b/Engine/Plugins/Runtime/HairStrands/Source/HairStrandsCore/Private/GroomComponent.cpp
--- a/Engine/Plugins/Runtime/HairStrands/Source/HairStrandsCore/Private/GroomComponent.cpp
+++ b/Engine/Plugins/Runtime/HairStrands/Source/HairStrandsCore/Private/GroomComponent.cpp
@@ -639,21 +639,23 @@ public:
 					bool bNeedDebugMaterial = false;
 					switch(ViewMode)
 					{
-					case EGroomViewMode::SimHairStrands				:
-					case EGroomViewMode::RenderVisCluster			:
-					case EGroomViewMode::RenderVisClusterAABB		:
-					case EGroomViewMode::RenderHairStrands			:
+					case EGroomViewMode::SimHairStrands	:
+					case EGroomViewMode::Cluster		:
+					case EGroomViewMode::ClusterAABB	:
+					case EGroomViewMode::RenderHairStrands:
 						bNeedDebugMaterial = Instances[GroupIt]->GeometryType == EHairGeometryType::Strands; break;
-					case EGroomViewMode::RenderHairRootUV			:
-					case EGroomViewMode::RenderHairUV				:
-					case EGroomViewMode::RenderHairSeed				:
-					case EGroomViewMode::RenderHairDimension		:
-					case EGroomViewMode::RenderHairRadiusVariation	:
-					case EGroomViewMode::RenderHairRootUDIM			:
-					case EGroomViewMode::RenderHairBaseColor		:
-					case EGroomViewMode::RenderHairRoughness		:
-					case EGroomViewMode::RenderHairGroup			:
-					case EGroomViewMode::RenderLODColoration		:
+					case EGroomViewMode::RootUV			:
+					case EGroomViewMode::UV				:
+					case EGroomViewMode::Seed			:
+					case EGroomViewMode::ClumpID		:
+					case EGroomViewMode::Dimension		:
+					case EGroomViewMode::RadiusVariation:
+					case EGroomViewMode::RootUDIM		:
+					case EGroomViewMode::Color			:
+					case EGroomViewMode::Roughness		:
+					case EGroomViewMode::AO				:
+					case EGroomViewMode::Group			:
+					case EGroomViewMode::LODColoration	:
 						bNeedDebugMaterial = true; break;
 					};
 
@@ -662,21 +664,23 @@ public:
 						float DebugModeScalar = 0;
 						switch(ViewMode)
 						{
-						case EGroomViewMode::None						: DebugModeScalar =99.f; break;
-						case EGroomViewMode::SimHairStrands				: DebugModeScalar = 0.f; break;
-						case EGroomViewMode::RenderHairStrands			: DebugModeScalar = 0.f; break;
-						case EGroomViewMode::RenderHairRootUV			: DebugModeScalar = 1.f; break;
-						case EGroomViewMode::RenderHairUV				: DebugModeScalar = 2.f; break;
-						case EGroomViewMode::RenderHairSeed				: DebugModeScalar = 3.f; break;
-						case EGroomViewMode::RenderHairDimension		: DebugModeScalar = 4.f; break;
-						case EGroomViewMode::RenderHairRadiusVariation	: DebugModeScalar = 5.f; break;
-						case EGroomViewMode::RenderHairRootUDIM			: DebugModeScalar = 6.f; break;
-						case EGroomViewMode::RenderHairBaseColor		: DebugModeScalar = 7.f; break;
-						case EGroomViewMode::RenderHairRoughness		: DebugModeScalar = 8.f; break;
-						case EGroomViewMode::RenderVisCluster			: DebugModeScalar = 0.f; break;
-						case EGroomViewMode::RenderVisClusterAABB		: DebugModeScalar = 0.f; break;
-						case EGroomViewMode::RenderHairGroup			: DebugModeScalar = 9.f; break;
-						case EGroomViewMode::RenderLODColoration		: DebugModeScalar = 10.f; break;
+						case EGroomViewMode::None				: DebugModeScalar =99.f; break;
+						case EGroomViewMode::SimHairStrands		: DebugModeScalar = 0.f; break;
+						case EGroomViewMode::RenderHairStrands	: DebugModeScalar = 0.f; break;
+						case EGroomViewMode::RootUV				: DebugModeScalar = 1.f; break;
+						case EGroomViewMode::UV					: DebugModeScalar = 2.f; break;
+						case EGroomViewMode::Seed				: DebugModeScalar = 3.f; break;
+						case EGroomViewMode::Dimension			: DebugModeScalar = 4.f; break;
+						case EGroomViewMode::RadiusVariation	: DebugModeScalar = 5.f; break;
+						case EGroomViewMode::RootUDIM			: DebugModeScalar = 6.f; break;
+						case EGroomViewMode::Color				: DebugModeScalar = 7.f; break;
+						case EGroomViewMode::Roughness			: DebugModeScalar = 8.f; break;
+						case EGroomViewMode::Cluster			: DebugModeScalar = 0.f; break;
+						case EGroomViewMode::ClusterAABB		: DebugModeScalar = 0.f; break;
+						case EGroomViewMode::Group				: DebugModeScalar = 9.f; break;
+						case EGroomViewMode::LODColoration		: DebugModeScalar = 10.f; break;
+						case EGroomViewMode::ClumpID			: DebugModeScalar = 11.f; break;
+						case EGroomViewMode::AO					: DebugModeScalar = 12.f; break;
 						};
 
 						// TODO: fix this as the radius is incorrect. This code run before the interpolation code, which is where HairRadius is updated.
@@ -687,17 +691,17 @@ public:
 						}
 						
 						// Reuse the HairMaxRadius field to send the LOD index instead of adding yet another variable
-						if (ViewMode == EGroomViewMode::RenderLODColoration)
+						if (ViewMode == EGroomViewMode::LODColoration)
 						{
 							HairMaxRadius = Instances[GroupIt]->HairGroupPublicData ? Instances[GroupIt]->HairGroupPublicData->LODIndex : 0;
 						}
 
 						FVector HairColor = FVector::ZeroVector;
-						if (ViewMode == EGroomViewMode::RenderHairGroup)
+						if (ViewMode == EGroomViewMode::Group)
 						{
 							HairColor = FVector(GetHairGroupDebugColor(GroupIt));
 						}
-						else if (ViewMode == EGroomViewMode::RenderLODColoration)
+						else if (ViewMode == EGroomViewMode::LODColoration)
 						{
 							int32 LODIndex = Instances[GroupIt]->HairGroupPublicData ? Instances[GroupIt]->HairGroupPublicData->LODIndex : 0;
 							LODIndex = FMath::Clamp(LODIndex, 0, GEngine->LODColorationColors.Num() - 1);
diff --git a/Engine/Plugins/Runtime/HairStrands/Source/HairStrandsCore/Private/GroomDebug.cpp b/Engine/Plugins/Runtime/HairStrands/Source/HairStrandsCore/Private/GroomDebug.cpp
--- a/Engine/Plugins/Runtime/HairStrands/Source/HairStrandsCore/Private/GroomDebug.cpp
+++ b/Engine/Plugins/Runtime/HairStrands/Source/HairStrandsCore/Private/GroomDebug.cpp
@@ -1165,7 +1165,7 @@ void RunHairStrandsDebug(
 			AddDrawDebugCardsGuidesPass(GraphBuilder, View, ShaderMap, Instance, ShaderPrintData, true /*bDeformed*/, false);
 		}
 
-		if (ViewMode == EGroomViewMode::RenderHairControlPoints)
+		if (ViewMode == EGroomViewMode::ControlPoints)
 		{
 			AddDrawDebugStrandsCVsPass(GraphBuilder, View, ShaderMap, Instance, ShaderPrintData, SceneColorTexture, SceneDepthTexture);
 		}
diff --git a/Engine/Plugins/Runtime/HairStrands/Source/HairStrandsCore/Private/HairStrandsRendering.cpp b/Engine/Plugins/Runtime/HairStrands/Source/HairStrandsCore/Private/HairStrandsRendering.cpp
--- a/Engine/Plugins/Runtime/HairStrands/Source/HairStrandsCore/Private/HairStrandsRendering.cpp
+++ b/Engine/Plugins/Runtime/HairStrands/Source/HairStrandsCore/Private/HairStrandsRendering.cpp
@@ -1681,7 +1681,7 @@ static FHairGroupPublicData::FVertexFactoryInput InternalComputeHairStrandsVerte
 			OutVFInput.Strands.PrevPositionOffset = Instance->Strands.DeformedResource->GetPositionOffset(FHairStrandsDeformedResource::EFrameType::Current);
 		}
 		CONVERT_HAIRSSTRANDS_VF_PARAMETERS(OutVFInput.Strands.TangentBuffer,			Instance->Strands.DeformedResource->TangentBuffer);
-		CONVERT_HAIRSSTRANDS_VF_PARAMETERS(OutVFInput.Strands.AttributeBuffer,			(ViewMode == EGroomViewMode::RenderHairStrands || ViewMode == EGroomViewMode::RenderVisCluster || ViewMode == EGroomViewMode::RenderVisClusterAABB) ? Instance->Strands.DebugAttributeBuffer : Instance->Strands.RestResource->AttributeBuffer);
+		CONVERT_HAIRSSTRANDS_VF_PARAMETERS(OutVFInput.Strands.AttributeBuffer,			(ViewMode == EGroomViewMode::RenderHairStrands || ViewMode == EGroomViewMode::Cluster || ViewMode == EGroomViewMode::ClusterAABB) ? Instance->Strands.DebugAttributeBuffer : Instance->Strands.RestResource->AttributeBuffer);
 		CONVERT_HAIRSSTRANDS_VF_PARAMETERS(OutVFInput.Strands.VertexToCurveBuffer,		Instance->Strands.RestResource->VertexToCurveBuffer);
 		CONVERT_HAIRSSTRANDS_VF_PARAMETERS(OutVFInput.Strands.CurveBuffer,				Instance->Strands.RestResource->CurveBuffer);
 
@@ -1694,7 +1694,7 @@ static FHairGroupPublicData::FVertexFactoryInput InternalComputeHairStrandsVerte
 		CONVERT_HAIRSSTRANDS_VF_PARAMETERS(OutVFInput.Strands.PositionBuffer,			Instance->Strands.RestResource->PositionBuffer);
 		CONVERT_HAIRSSTRANDS_VF_PARAMETERS(OutVFInput.Strands.PrevPositionBuffer,		Instance->Strands.RestResource->PositionBuffer);
 		CONVERT_HAIRSSTRANDS_VF_PARAMETERS(OutVFInput.Strands.TangentBuffer,			Instance->Strands.RestResource->TangentBuffer);
-		CONVERT_HAIRSSTRANDS_VF_PARAMETERS(OutVFInput.Strands.AttributeBuffer,			(ViewMode == EGroomViewMode::RenderVisCluster || ViewMode == EGroomViewMode::RenderVisClusterAABB) ? Instance->Strands.DebugAttributeBuffer : Instance->Strands.RestResource->AttributeBuffer);
+		CONVERT_HAIRSSTRANDS_VF_PARAMETERS(OutVFInput.Strands.AttributeBuffer,			(ViewMode == EGroomViewMode::Cluster || ViewMode == EGroomViewMode::ClusterAABB) ? Instance->Strands.DebugAttributeBuffer : Instance->Strands.RestResource->AttributeBuffer);
 		CONVERT_HAIRSSTRANDS_VF_PARAMETERS(OutVFInput.Strands.VertexToCurveBuffer,		Instance->Strands.RestResource->VertexToCurveBuffer);
 		CONVERT_HAIRSSTRANDS_VF_PARAMETERS(OutVFInput.Strands.CurveBuffer,				Instance->Strands.RestResource->CurveBuffer);
 		CONVERT_HAIRSSTRANDS_VF_PARAMETERS(OutVFInput.Strands.PositionOffsetBuffer,		Instance->Strands.RestResource->PositionOffsetBuffer);
@@ -1811,8 +1811,8 @@ void ComputeHairStrandsInterpolation(
 			// "WITH_EDITOR && PatchMode == EHairPatchAttribute::GuideInflucence" is a special path when visualizing groom within the groom editor, so that we can diplay guides even when there is no simulation or global interpolation enabled
 			EHairPatchAttribute PatchMode = EHairPatchAttribute::None;
 			#if WITH_EDITOR
-			if (ViewMode == EGroomViewMode::RenderHairStrands && bNeedDeformation)									{ PatchMode = EHairPatchAttribute::GuideInflucence; }
-			if (ViewMode == EGroomViewMode::RenderVisCluster || ViewMode == EGroomViewMode::RenderVisClusterAABB)	{ PatchMode = EHairPatchAttribute::ClusterInfluence; }
+			if (ViewMode == EGroomViewMode::RenderHairStrands && bNeedDeformation)				{ PatchMode = EHairPatchAttribute::GuideInflucence; }
+			if (ViewMode == EGroomViewMode::Cluster || ViewMode == EGroomViewMode::ClusterAABB)	{ PatchMode = EHairPatchAttribute::ClusterInfluence; }
 			#endif
 
 			// 2. Deform hair if needed (e.g.: skinning, simulation, RBF) and recompute tangent
diff --git a/Engine/Shaders/Private/HairStrands/HairStrandsAttributeCommon.ush b/Engine/Shaders/Private/HairStrands/HairStrandsAttributeCommon.ush
--- a/Engine/Shaders/Private/HairStrands/HairStrandsAttributeCommon.ush
+++ b/Engine/Shaders/Private/HairStrands/HairStrandsAttributeCommon.ush
@@ -172,7 +172,7 @@ float GetHairStrandsAO(uint HairPrimitiveId, float2 HairPrimitiveUV)
 	}
 	else
 	{
-		return 0;
+		return 1.f;
 	}
 }
 
diff --git a/Engine/Source/Editor/CommonMenuExtensions/Private/GroomVisualizationMenuCommands.cpp b/Engine/Source/Editor/CommonMenuExtensions/Private/GroomVisualizationMenuCommands.cpp
--- a/Engine/Source/Editor/CommonMenuExtensions/Private/GroomVisualizationMenuCommands.cpp
+++ b/Engine/Source/Editor/CommonMenuExtensions/Private/GroomVisualizationMenuCommands.cpp
@@ -67,21 +67,22 @@ void FGroomVisualizationMenuCommands::InternalBuildVisualisationSubMenu(FMenuBui
 			{
 				Commands.AddCommandTypeToMenu(Menu, EGroomViewMode::MacroGroups);
 			}
-			Commands.AddCommandTypeToMenu(Menu, EGroomViewMode::RenderHairGroup);
-			Commands.AddCommandTypeToMenu(Menu, EGroomViewMode::RenderLODColoration);
+			Commands.AddCommandTypeToMenu(Menu, EGroomViewMode::Group);
+			Commands.AddCommandTypeToMenu(Menu, EGroomViewMode::LODColoration);
 			Menu.EndSection();
 		}
 
 		// Attributes
 		{
 			Menu.BeginSection("LevelViewportGroomVisualizationAttributes", LOCTEXT("GroomVisualizationAttributes", "Groom Attributes"));
-			Commands.AddCommandTypeToMenu(Menu, EGroomViewMode::RenderHairSeed);
-			Commands.AddCommandTypeToMenu(Menu, EGroomViewMode::RenderHairRootUV);
-			Commands.AddCommandTypeToMenu(Menu, EGroomViewMode::RenderHairRootUDIM);
-			Commands.AddCommandTypeToMenu(Menu, EGroomViewMode::RenderHairUV);
-			Commands.AddCommandTypeToMenu(Menu, EGroomViewMode::RenderHairDimension);
-			Commands.AddCommandTypeToMenu(Menu, EGroomViewMode::RenderHairRadiusVariation);
-			Commands.AddCommandTypeToMenu(Menu, EGroomViewMode::RenderHairTangent);
+			Commands.AddCommandTypeToMenu(Menu, EGroomViewMode::Seed);
+			Commands.AddCommandTypeToMenu(Menu, EGroomViewMode::RootUV);
+			Commands.AddCommandTypeToMenu(Menu, EGroomViewMode::RootUDIM);
+			Commands.AddCommandTypeToMenu(Menu, EGroomViewMode::UV);
+			Commands.AddCommandTypeToMenu(Menu, EGroomViewMode::Dimension);
+			Commands.AddCommandTypeToMenu(Menu, EGroomViewMode::RadiusVariation);
+			Commands.AddCommandTypeToMenu(Menu, EGroomViewMode::Tangent);
+			Commands.AddCommandTypeToMenu(Menu, EGroomViewMode::ClumpID);
 			Menu.EndSection();
 
 		}
@@ -97,14 +98,15 @@ void FGroomVisualizationMenuCommands::InternalBuildVisualisationSubMenu(FMenuBui
 		// Strands
 		{
 			Menu.BeginSection("LevelViewportGroomVisualizationStrands", LOCTEXT("GroomVisualizationMaterial", "Groom Strands"));
-			Commands.AddCommandTypeToMenu(Menu, EGroomViewMode::RenderHairControlPoints);
+			Commands.AddCommandTypeToMenu(Menu, EGroomViewMode::ControlPoints);
+			Commands.AddCommandTypeToMenu(Menu, EGroomViewMode::Color);
+			Commands.AddCommandTypeToMenu(Menu, EGroomViewMode::Roughness);
+			Commands.AddCommandTypeToMenu(Menu, EGroomViewMode::AO);
 			Commands.AddCommandTypeToMenu(Menu, EGroomViewMode::MaterialDepth);
 			Commands.AddCommandTypeToMenu(Menu, EGroomViewMode::MaterialBaseColor);
 			Commands.AddCommandTypeToMenu(Menu, EGroomViewMode::MaterialRoughness);
 			Commands.AddCommandTypeToMenu(Menu, EGroomViewMode::MaterialSpecular);
 			Commands.AddCommandTypeToMenu(Menu, EGroomViewMode::MaterialTangent);
-			Commands.AddCommandTypeToMenu(Menu, EGroomViewMode::RenderHairBaseColor);
-			Commands.AddCommandTypeToMenu(Menu, EGroomViewMode::RenderHairRoughness);
 			Menu.EndSection();
 		}
 
@@ -112,8 +114,8 @@ void FGroomVisualizationMenuCommands::InternalBuildVisualisationSubMenu(FMenuBui
 		{
 			Menu.BeginSection("LevelViewportGroomVisualizationAdvanced", LOCTEXT("GroomVisualizationAdvanced", "Groom Strands Advanced"));
 			Commands.AddCommandTypeToMenu(Menu, EGroomViewMode::VoxelsDensity);
-			Commands.AddCommandTypeToMenu(Menu, EGroomViewMode::RenderVisCluster);
-			Commands.AddCommandTypeToMenu(Menu, EGroomViewMode::RenderVisClusterAABB);
+			Commands.AddCommandTypeToMenu(Menu, EGroomViewMode::Cluster);
+			Commands.AddCommandTypeToMenu(Menu, EGroomViewMode::ClusterAABB);
 			if (!bIsGroomEditor)
 			{
 				Commands.AddCommandTypeToMenu(Menu, EGroomViewMode::LightBounds);
diff --git a/Engine/Source/Runtime/Engine/Private/GroomVisualizationData.cpp b/Engine/Source/Runtime/Engine/Private/GroomVisualizationData.cpp
--- a/Engine/Source/Runtime/Engine/Private/GroomVisualizationData.cpp
+++ b/Engine/Source/Runtime/Engine/Private/GroomVisualizationData.cpp
@@ -56,21 +56,23 @@ EGroomViewMode GetGroomViewMode(const FSceneView& View)
 		case 20: return EGroomViewMode::None;
 		case 21: return EGroomViewMode::SimHairStrands;
 		case 22: return EGroomViewMode::RenderHairStrands;
-		case 23: return EGroomViewMode::RenderHairRootUV;
-		case 24: return EGroomViewMode::RenderHairRootUDIM;
-		case 25: return EGroomViewMode::RenderHairUV;
-		case 26: return EGroomViewMode::RenderHairSeed;
-		case 27: return EGroomViewMode::RenderHairDimension;
-		case 28: return EGroomViewMode::RenderHairRadiusVariation;
-		case 29: return EGroomViewMode::RenderHairBaseColor;
-		case 30: return EGroomViewMode::RenderHairRoughness;
-		case 31: return EGroomViewMode::RenderVisCluster;
-		case 32: return EGroomViewMode::RenderVisClusterAABB;
-		case 33: return EGroomViewMode::RenderHairTangent;
-		case 34: return EGroomViewMode::RenderHairControlPoints;
-		case 35: return EGroomViewMode::RenderHairGroup;
-		case 36: return EGroomViewMode::RenderLODColoration;
+		case 23: return EGroomViewMode::RootUV;
+		case 24: return EGroomViewMode::RootUDIM;
+		case 25: return EGroomViewMode::UV;
+		case 26: return EGroomViewMode::Seed;
+		case 27: return EGroomViewMode::Dimension;
+		case 28: return EGroomViewMode::RadiusVariation;
+		case 29: return EGroomViewMode::Color;
+		case 30: return EGroomViewMode::Roughness;
+		case 31: return EGroomViewMode::Cluster;
+		case 32: return EGroomViewMode::ClusterAABB;
+		case 33: return EGroomViewMode::Tangent;
+		case 34: return EGroomViewMode::ControlPoints;
+		case 35: return EGroomViewMode::Group;
+		case 36: return EGroomViewMode::LODColoration;
 		case 37: return EGroomViewMode::CardGuides;
+		case 38: return EGroomViewMode::AO;
+		case 39: return EGroomViewMode::ClumpID;
 		default: break;
 		}
 
@@ -81,7 +83,7 @@ EGroomViewMode GetGroomViewMode(const FSceneView& View)
 		}
 		else if (View.Family && View.Family->EngineShowFlags.LODColoration)
 		{
-			Out = EGroomViewMode::RenderLODColoration;
+			Out = EGroomViewMode::LODColoration;
 		}
 	}
 	return Out;
@@ -110,20 +112,22 @@ const TCHAR* GetGroomViewModeName(EGroomViewMode In)
 	case EGroomViewMode::Tile:						return TEXT("Tile");
 	case EGroomViewMode::SimHairStrands:			return TEXT("SimHairStrands");
 	case EGroomViewMode::RenderHairStrands:			return TEXT("RenderHairStrands");
-	case EGroomViewMode::RenderHairRootUV:			return TEXT("RenderHairRootUV");
-	case EGroomViewMode::RenderHairRootUDIM:		return TEXT("RenderHairRootUDIM");
-	case EGroomViewMode::RenderHairUV:				return TEXT("RenderHairUV");
-	case EGroomViewMode::RenderHairSeed:			return TEXT("RenderHairSeed");
-	case EGroomViewMode::RenderHairDimension:		return TEXT("RenderHairDimension");
-	case EGroomViewMode::RenderHairRadiusVariation:	return TEXT("RenderHairRadiusVariation");
-	case EGroomViewMode::RenderHairBaseColor:		return TEXT("RenderHairBaseColor");
-	case EGroomViewMode::RenderHairRoughness:		return TEXT("RenderHairRoughness");
-	case EGroomViewMode::RenderVisCluster:			return TEXT("RenderVisCluster");
-	case EGroomViewMode::RenderVisClusterAABB:		return TEXT("RenderVisClusterAABB");
-	case EGroomViewMode::RenderHairTangent:			return TEXT("RenderHairTangent");
-	case EGroomViewMode::RenderHairControlPoints:	return TEXT("RenderHairControlPoints");
-	case EGroomViewMode::RenderHairGroup:			return TEXT("RenderHairGroup");
-	case EGroomViewMode::RenderLODColoration:		return TEXT("RenderLODColoration");
+	case EGroomViewMode::RootUV:					return TEXT("RootUV");
+	case EGroomViewMode::RootUDIM:					return TEXT("RootUDIM");
+	case EGroomViewMode::UV:						return TEXT("UV");
+	case EGroomViewMode::Seed:						return TEXT("Seed");
+	case EGroomViewMode::Dimension:					return TEXT("Dimension");
+	case EGroomViewMode::RadiusVariation:			return TEXT("RadiusVariation");
+	case EGroomViewMode::Color:						return TEXT("Color");
+	case EGroomViewMode::Roughness:					return TEXT("Roughness");
+	case EGroomViewMode::AO:						return TEXT("AO");
+	case EGroomViewMode::ClumpID:					return TEXT("ClumpID");
+	case EGroomViewMode::Cluster:					return TEXT("Cluster");
+	case EGroomViewMode::ClusterAABB:				return TEXT("ClusterAABB");
+	case EGroomViewMode::Tangent:					return TEXT("Tangent");
+	case EGroomViewMode::ControlPoints:				return TEXT("ControlPoints");
+	case EGroomViewMode::Group:						return TEXT("Group");
+	case EGroomViewMode::LODColoration:				return TEXT("LODColoration");
 	case EGroomViewMode::CardGuides:				return TEXT("CardGuides");
 	}
 	return TEXT("None");
@@ -197,20 +201,22 @@ void FGroomVisualizationData::Initialize()
 		AddVisualizationMode(ModeMap, true, EGroomViewMode::SimHairStrands,				LOCTEXT("SimHairStrands", "Guides"),						LOCTEXT("SimHairStrandsDesc", "Simulation strands"));
 		AddVisualizationMode(ModeMap, true, EGroomViewMode::RenderHairStrands,			LOCTEXT("RenderHairStrands", "Stands Guides Influences"),	LOCTEXT("RenderHairStrandsDesc", "Rendering strands influences"));
 		AddVisualizationMode(ModeMap, true, EGroomViewMode::CardGuides, 				LOCTEXT("CardGuides", "Cards Guides"),						LOCTEXT("CardGuidesDesc", "Cards Guides"));
-		AddVisualizationMode(ModeMap, true, EGroomViewMode::RenderHairRootUV,			LOCTEXT("RenderHairRootUV", "Root UV"),						LOCTEXT("RenderHairRootUVDesc", "Roots UV"));
-		AddVisualizationMode(ModeMap, true, EGroomViewMode::RenderHairRootUDIM,			LOCTEXT("RenderHairRootUDIM", "Root UDIM"),					LOCTEXT("RenderHairRootUDIMDesc", "Roots UV UDIM texture index"));
-		AddVisualizationMode(ModeMap, true, EGroomViewMode::RenderHairUV,				LOCTEXT("RenderHairUV", "UV"),								LOCTEXT("RenderHairUVDesc", "Hair UV"));
-		AddVisualizationMode(ModeMap, true, EGroomViewMode::RenderHairSeed,				LOCTEXT("RenderHairSeed", "Seed"),							LOCTEXT("RenderHairSeedDesc", "Hair seed"));
-		AddVisualizationMode(ModeMap, true, EGroomViewMode::RenderHairDimension,		LOCTEXT("RenderHairDimension", "Dimension"),				LOCTEXT("RenderHairDimensionDesc", "Hair dimensions"));
-		AddVisualizationMode(ModeMap, true, EGroomViewMode::RenderHairRadiusVariation,	LOCTEXT("RenderHairRadiusVariation", "Radius Variation"),	LOCTEXT("RenderHairRadiusVariationDesc", "Hair radius variation"));
-		AddVisualizationMode(ModeMap, true, EGroomViewMode::RenderHairTangent,			LOCTEXT("RenderHairTangent", "Tangent"),					LOCTEXT("RenderHairTangentDesc", "Hair tangent"));
-		AddVisualizationMode(ModeMap, true, EGroomViewMode::RenderHairControlPoints,	LOCTEXT("RenderHairControlPoints", "Control Points"),		LOCTEXT("RenderHairControlPointsDesc", "Hair control points"));
-		AddVisualizationMode(ModeMap, true, EGroomViewMode::RenderHairBaseColor,		LOCTEXT("RenderHairBaseColor", "Per-CV Color"),				LOCTEXT("RenderHairBaseColorDesc", "CV color"));
-		AddVisualizationMode(ModeMap, true, EGroomViewMode::RenderHairRoughness,		LOCTEXT("RenderHairRoughness", "Per-CV Roughness"),			LOCTEXT("RenderHairRoughnessDesc", "CV roughness"));
-		AddVisualizationMode(ModeMap, true, EGroomViewMode::RenderVisCluster,			LOCTEXT("RenderVisCluster", "Clusters"),					LOCTEXT("RenderVisClusterDesc", "Hair visility clusters"));
-		AddVisualizationMode(ModeMap, true, EGroomViewMode::RenderVisClusterAABB,		LOCTEXT("RenderVisClusterAABB", "Clusters Bounds"),			LOCTEXT("RenderVisClusterAABBDesc", "Hair visility clusters AABBs"));
-		AddVisualizationMode(ModeMap, true, EGroomViewMode::RenderHairGroup,			LOCTEXT("RenderHairGroup", "Groups"),						LOCTEXT("RenderHairGroupDesc", "Hair hair groups"));
-		AddVisualizationMode(ModeMap, true, EGroomViewMode::RenderLODColoration,		LOCTEXT("RenderLODColoration", "LOD Color"),				LOCTEXT("RenderLODColorationDesc", "Hair LOD coloring"));
+		AddVisualizationMode(ModeMap, true, EGroomViewMode::RootUV,						LOCTEXT("RootUV", "Root UV"),								LOCTEXT("RootUVDesc", "Roots UV"));
+		AddVisualizationMode(ModeMap, true, EGroomViewMode::RootUDIM,					LOCTEXT("RootUDIM", "Root UDIM"),							LOCTEXT("RootUDIMDesc", "Roots UV UDIM texture index"));
+		AddVisualizationMode(ModeMap, true, EGroomViewMode::UV,							LOCTEXT("UV", "UV"),										LOCTEXT("UVDesc", "Hair UV"));
+		AddVisualizationMode(ModeMap, true, EGroomViewMode::Seed,						LOCTEXT("Seed", "Seed"),									LOCTEXT("SeedDesc", "Hair seed"));
+		AddVisualizationMode(ModeMap, true, EGroomViewMode::Dimension,					LOCTEXT("Dimension", "Dimension"),							LOCTEXT("DimensionDesc", "Hair dimensions"));
+		AddVisualizationMode(ModeMap, true, EGroomViewMode::RadiusVariation,			LOCTEXT("RadiusVariation", "Radius Variation"),				LOCTEXT("RadiusVariationDesc", "Hair radius variation"));
+		AddVisualizationMode(ModeMap, true, EGroomViewMode::Tangent,					LOCTEXT("Tangent", "Tangent"),								LOCTEXT("TangentDesc", "Hair tangent"));
+		AddVisualizationMode(ModeMap, true, EGroomViewMode::ControlPoints,				LOCTEXT("ControlPoints", "Control Points"),					LOCTEXT("ControlPointsDesc", "Hair control points"));
+		AddVisualizationMode(ModeMap, true, EGroomViewMode::Color,						LOCTEXT("Color", "Per-CV Color"),							LOCTEXT("ColorDesc", "Control point color (optional)"));
+		AddVisualizationMode(ModeMap, true, EGroomViewMode::Roughness,					LOCTEXT("Roughness", "Per-CV Roughness"),					LOCTEXT("RoughnessDesc", "Control point roughness  (optional)"));
+		AddVisualizationMode(ModeMap, true, EGroomViewMode::Cluster,					LOCTEXT("Cluster", "Clusters"),								LOCTEXT("ClusterDesc", "Hair visility clusters"));
+		AddVisualizationMode(ModeMap, true, EGroomViewMode::ClusterAABB,				LOCTEXT("ClusterAABB", "Clusters Bounds"),					LOCTEXT("ClusterAABBDesc", "Hair visility clusters AABBs"));
+		AddVisualizationMode(ModeMap, true, EGroomViewMode::Group,						LOCTEXT("Group", "Groups"),									LOCTEXT("GroupDesc", "Hair hair groups"));
+		AddVisualizationMode(ModeMap, true, EGroomViewMode::LODColoration,				LOCTEXT("LODColoration", "LOD Color"),						LOCTEXT("LODColorationDesc", "Hair LOD coloring"));
+		AddVisualizationMode(ModeMap, true, EGroomViewMode::AO,							LOCTEXT("AO", "Per-CV AO"),									LOCTEXT("AODesc", "Control point AO  (optional)"));
+		AddVisualizationMode(ModeMap, true, EGroomViewMode::ClumpID,					LOCTEXT("ClumpID", "Clump ID"),								LOCTEXT("ClumpIDDesc", "Curve clumpID  (optional)"));
 
 		ConsoleDocumentationVisualizationMode = ConfigureConsoleCommand(ModeMap);
 	}
diff --git a/Engine/Source/Runtime/Engine/Public/GroomVisualizationData.h b/Engine/Source/Runtime/Engine/Public/GroomVisualizationData.h
--- a/Engine/Source/Runtime/Engine/Public/GroomVisualizationData.h
+++ b/Engine/Source/Runtime/Engine/Public/GroomVisualizationData.h
@@ -11,20 +11,22 @@ enum class EGroomViewMode : uint8
 	None,
 	SimHairStrands,    // Guide
 	RenderHairStrands, // Guide influence + add cluster
-	RenderHairUV,
-	RenderHairRootUV,
-	RenderHairRootUDIM,
-	RenderHairSeed,
-	RenderHairDimension,
-	RenderHairRadiusVariation,
-	RenderHairTangent,
-	RenderHairBaseColor,
-	RenderHairRoughness,
-	RenderVisCluster,
-	RenderVisClusterAABB,
-	RenderHairGroup,
-	RenderLODColoration,
-	RenderHairControlPoints,	
+	UV,
+	RootUV,
+	RootUDIM,
+	Seed,
+	Dimension,
+	RadiusVariation,
+	Tangent,
+	Color,
+	Roughness,
+	AO,
+	ClumpID,
+	Cluster,
+	ClusterAABB,
+	Group,
+	LODColoration,
+	ControlPoints,
 	MacroGroups,
 	LightBounds,
 	DeepOpacityMaps,
diff --git a/Engine/Source/Runtime/Renderer/Private/HairStrands/HairStrandsClusters.cpp b/Engine/Source/Runtime/Renderer/Private/HairStrands/HairStrandsClusters.cpp
--- a/Engine/Source/Runtime/Renderer/Private/HairStrands/HairStrandsClusters.cpp
+++ b/Engine/Source/Runtime/Renderer/Private/HairStrands/HairStrandsClusters.cpp
@@ -281,8 +281,8 @@ static void AddClusterCullingPass(
 		};
 
 		static IConsoleVariable* CVarClusterDebug = IConsoleManager::Get().FindConsoleVariable(TEXT("r.HairStrands.Cluster.Debug"));
-		bClusterDebugAABBBuffer = GetGroomViewMode(View) == EGroomViewMode::RenderVisClusterAABB;
-		bClusterDebug = GetGroomViewMode(View) == EGroomViewMode::RenderVisCluster;
+		bClusterDebugAABBBuffer = GetGroomViewMode(View) == EGroomViewMode::ClusterAABB;
+		bClusterDebug = GetGroomViewMode(View) == EGroomViewMode::Cluster;
 		ClusterDebugInfoBuffer = GraphBuilder.CreateBuffer(FRDGBufferDesc::CreateStructuredDesc(sizeof(FHairClusterDebugInfo), ClusterData.ClusterCount), TEXT("Hair.CulledCompactedIndexBuffer"));
 	}
 #endif
diff --git a/Engine/Source/Runtime/Renderer/Private/HairStrands/HairStrandsDebug.cpp b/Engine/Source/Runtime/Renderer/Private/HairStrands/HairStrandsDebug.cpp
--- a/Engine/Source/Runtime/Renderer/Private/HairStrands/HairStrandsDebug.cpp
+++ b/Engine/Source/Runtime/Renderer/Private/HairStrands/HairStrandsDebug.cpp
@@ -1110,7 +1110,7 @@ static void AddDrawDebugClusterPass(
 	const FViewInfo& View)
 {
 	const EGroomViewMode ViewMode = GetGroomViewMode(View);
-	const bool bDebugEnable = ViewMode == EGroomViewMode::RenderVisClusterAABB || ViewMode == EGroomViewMode::RenderVisCluster;
+	const bool bDebugEnable = ViewMode == EGroomViewMode::ClusterAABB || ViewMode == EGroomViewMode::Cluster;
 	const bool bCullingEnable = IsHairStrandsClusterCullingEnable();
 	if (!bDebugEnable || !bCullingEnable)
 	{
@@ -1124,7 +1124,7 @@ static void AddDrawDebugClusterPass(
 
 	for (const FHairStrandsMacroGroupData& MacroGroupData : View.HairStrandsViewData.MacroGroupDatas)
 	{
-		const bool bDebugAABB = GetGroomViewMode(View) == EGroomViewMode::RenderVisClusterAABB;
+		const bool bDebugAABB = GetGroomViewMode(View) == EGroomViewMode::ClusterAABB;
 
 		for (const FHairStrandsMacroGroupData::PrimitiveInfo& PrimitiveInfo : MacroGroupData.PrimitivesInfos)
 		{
@@ -1156,7 +1156,7 @@ static void AddDrawDebugClusterPass(
 					Parameters->ClusterCount = HairGroupClusters.ClusterCount;
 					Parameters->TriangleCount = HairGroupClusters.VertexCount * 2; // VertexCount is actually the number of control points
 					Parameters->HairGroupId = DataIndex;
-					Parameters->ClusterDebugMode = ViewMode == EGroomViewMode::RenderVisCluster ? 1 : (ViewMode == EGroomViewMode::RenderVisClusterAABB ? 2 : 0);
+					Parameters->ClusterDebugMode = ViewMode == EGroomViewMode::Cluster ? 1 : (ViewMode == EGroomViewMode::ClusterAABB ? 2 : 0);
 					Parameters->ClusterAABBBuffer = HairGroupClusters.ClusterAABBBuffer->SRV;
 					Parameters->CulledDispatchIndirectParametersClusterCountBuffer = GraphBuilder.CreateSRV(HairGroupClusters.CulledClusterCountBuffer, EPixelFormat::PF_R32_UINT);
 					Parameters->CulledDrawIndirectParameters = DrawIndirectBuffer.SRV;
@@ -1214,7 +1214,7 @@ static void InternalRenderHairStrandsDebugInfo(
 	const EGroomViewMode ViewMode = GetGroomViewMode(View);
 
 	// Display tangent vector for strands/cards/meshes
-	if (ViewMode == EGroomViewMode::RenderHairTangent)
+	if (ViewMode == EGroomViewMode::Tangent)
 	{
 		AddDebugHairTangentPass(GraphBuilder, View, SceneTextures, SceneColorTexture);
 	}
@@ -1369,7 +1369,7 @@ static void InternalRenderHairStrandsDebugInfo(
 		FComputeShaderUtils::AddPass(GraphBuilder, RDG_EVENT_NAME("HairStrands::PPLLDebug"), ComputeShader, PassParameters, FIntVector::DivideAndRoundUp(TextureSize, FIntVector(8, 8, 1)));
 	}
 
-	if (ViewMode == EGroomViewMode::RenderVisCluster || ViewMode == EGroomViewMode::RenderVisClusterAABB)
+	if (ViewMode == EGroomViewMode::Cluster || ViewMode == EGroomViewMode::ClusterAABB)
 	{
 		AddDrawDebugClusterPass(GraphBuilder, HairClusterData, View);
 	}
