Focus on asset types that are required for PIE during world loading
  - 5.8s to 2.1s on RT during heavy scene update at beginning of PIE in CitySample due to drastic reduction in redundant shader invalidation processing

#rnx

diff --git a/Engine/Source/Runtime/Engine/Private/AssetCompilingManager.cpp b/Engine/Source/Runtime/Engine/Private/AssetCompilingManager.cpp
--- a/Engine/Source/Runtime/Engine/Private/AssetCompilingManager.cpp
+++ b/Engine/Source/Runtime/Engine/Private/AssetCompilingManager.cpp
@@ -521,6 +521,19 @@ void FAssetCompilingManager::ProcessAsyncTasks(bool bLimitExecutionTime)
 	UpdateNumRemainingAssets();
 }
 
+void FAssetCompilingManager::ProcessAsyncTasks(const AssetCompilation::FProcessAsyncTaskParams& Params)
+{
+	// Reuse ObjectIterator Caching and Reverse lookups for the duration of all asset updates
+	FObjectCacheContextScope ObjectCacheScope;
+
+	for (IAssetCompilingManager* AssetCompilingManager : AssetCompilingManagers)
+	{
+		AssetCompilingManager->ProcessAsyncTasks(Params);
+	}
+
+	UpdateNumRemainingAssets();
+}
+
 void FAssetCompilingManager::UpdateNumRemainingAssets()
 {
 	const int32 NumRemainingAssets = GetNumRemainingAssets();
diff --git a/Engine/Source/Runtime/Engine/Private/ShaderCompiler/ShaderCompiler.cpp b/Engine/Source/Runtime/Engine/Private/ShaderCompiler/ShaderCompiler.cpp
--- a/Engine/Source/Runtime/Engine/Private/ShaderCompiler/ShaderCompiler.cpp
+++ b/Engine/Source/Runtime/Engine/Private/ShaderCompiler/ShaderCompiler.cpp
@@ -4126,6 +4126,17 @@ void FShaderCompilingManager::ProcessAsyncTasks(bool bLimitExecutionTime)
 	ProcessAsyncResults(bLimitExecutionTime, false);
 }
 
+void FShaderCompilingManager::ProcessAsyncTasks(const AssetCompilation::FProcessAsyncTaskParams& Params)
+{
+	// Shader compilations are not required for PIE to begin.
+	if (Params.bPlayInEditorAssetsOnly)
+	{
+		return;
+	}
+
+	ProcessAsyncResults(Params.bLimitExecutionTime, false);
+}
+
 int32 FShaderCompilingManager::GetNumPendingJobs() const
 {
 	return AllJobs.GetNumPendingJobs();
diff --git a/Engine/Source/Runtime/Engine/Private/World.cpp b/Engine/Source/Runtime/Engine/Private/World.cpp
--- a/Engine/Source/Runtime/Engine/Private/World.cpp
+++ b/Engine/Source/Runtime/Engine/Private/World.cpp
@@ -2927,6 +2927,9 @@ void UWorld::AddToWorld( ULevel* Level, const FTransform& LevelTransform, bool b
 	}
 
 #if WITH_EDITOR
+	AssetCompilation::FProcessAsyncTaskParams ProcessAsyncTasksParams;
+	ProcessAsyncTasksParams.bPlayInEditorAssetsOnly = true;
+
 	if (bExecuteNextStep)
 	{
 		// Wait on any async DDC handles
@@ -2960,7 +2963,7 @@ void UWorld::AddToWorld( ULevel* Level, const FTransform& LevelTransform, bool b
 		}
 
 		// Gives a chance to any assets being used for PIE/game to complete
-		FAssetCompilingManager::Get().ProcessAsyncTasks();
+		FAssetCompilingManager::Get().ProcessAsyncTasks(ProcessAsyncTasksParams);
 	}
 #endif
 
@@ -3010,7 +3013,7 @@ void UWorld::AddToWorld( ULevel* Level, const FTransform& LevelTransform, bool b
 #if WITH_EDITOR
 	// Gives a chance to any assets being used for PIE/game to complete before calling
 	// BeginPlay on all actors
-	FAssetCompilingManager::Get().ProcessAsyncTasks();
+	FAssetCompilingManager::Get().ProcessAsyncTasks(ProcessAsyncTasksParams);
 #endif
 
 	if( IsGameWorld() && AreActorsInitialized() )
@@ -9056,4 +9059,3 @@ FString ENGINE_API ToString(ENetMode NetMode)
 }
 
 #undef LOCTEXT_NAMESPACE 
-
diff --git a/Engine/Source/Runtime/Engine/Public/AssetCompilingManager.h b/Engine/Source/Runtime/Engine/Public/AssetCompilingManager.h
--- a/Engine/Source/Runtime/Engine/Public/AssetCompilingManager.h
+++ b/Engine/Source/Runtime/Engine/Public/AssetCompilingManager.h
@@ -28,6 +28,18 @@ struct FAssetCompileData
 	}
 };
 
+namespace AssetCompilation
+{
+	struct FProcessAsyncTaskParams
+	{
+		/* Limits the execution time instead of processing everything */
+		bool bLimitExecutionTime = false;
+
+		/* Limits processing for assets required for PIE only. */
+		bool bPlayInEditorAssetsOnly = false;
+	};
+}
+
 struct IAssetCompilingManager
 {
 	/**
@@ -75,6 +87,14 @@ protected:
 	 */
 	ENGINE_API virtual void ProcessAsyncTasks(bool bLimitExecutionTime = false) = 0;
 
+	/**
+	 * Called once per frame, fetches completed tasks and applies them to the scene.
+	 */
+	ENGINE_API virtual void ProcessAsyncTasks(const AssetCompilation::FProcessAsyncTaskParams& Params)
+	{
+		/* Forward for backward compatibility */
+		ProcessAsyncTasks(Params.bLimitExecutionTime);
+	}
 
 	ENGINE_API virtual ~IAssetCompilingManager() {}
 };
@@ -131,6 +151,11 @@ public:
 	 */
 	ENGINE_API void ProcessAsyncTasks(bool bLimitExecutionTime = false);
 
+	/**
+	 * Called once per frame, fetches completed tasks and applies them to the scene.
+	 */
+	ENGINE_API void ProcessAsyncTasks(const AssetCompilation::FProcessAsyncTaskParams& Params);
+
 	/**
 	 * Event called after an asset finishes compilation.
 	 */
diff --git a/Engine/Source/Runtime/Engine/Public/ShaderCompiler.h b/Engine/Source/Runtime/Engine/Public/ShaderCompiler.h
--- a/Engine/Source/Runtime/Engine/Public/ShaderCompiler.h
+++ b/Engine/Source/Runtime/Engine/Public/ShaderCompiler.h
@@ -851,6 +851,7 @@ private:
 	TArrayView<FName> GetDependentTypeNames() const override;
 	int32 GetNumRemainingAssets() const override;
 	void ProcessAsyncTasks(bool bLimitExecutionTime = false) override;
+	void ProcessAsyncTasks(const AssetCompilation::FProcessAsyncTaskParams& Params) override;
 
 	//////////////////////////////////////////////////////
 	// Thread shared properties: These variables can only be read from or written to when a lock on CompileQueueSection is obtained, since they are used by both threads.
