Fixed crash when exiting PIE where world level instance subsystem is accessed by uninitializing instanced world partition after world subsystems are deinitialized.

diff --git a/Engine/Source/Runtime/Engine/Private/WorldPartition/WorldPartition.cpp b/Engine/Source/Runtime/Engine/Private/WorldPartition/WorldPartition.cpp
--- a/Engine/Source/Runtime/Engine/Private/WorldPartition/WorldPartition.cpp
+++ b/Engine/Source/Runtime/Engine/Private/WorldPartition/WorldPartition.cpp
@@ -902,7 +902,7 @@ void UWorldPartition::OnLevelRemovedFromWorld(ULevel* InLevel, UWorld* InWorld)
 	if ((World == InWorld) && (InLevel == GetTypedOuter<UWorld>()->PersistentLevel))
 	{
 		check(!CanStream());
-		UpdateStreamingState();
+		Uninitialize();
 	}
 }
 
