Fix anim sequence curve colors all being reset to white

diff --git a/Engine/Source/Runtime/Engine/Private/Animation/AnimSequenceHelpers.cpp b/Engine/Source/Runtime/Engine/Private/Animation/AnimSequenceHelpers.cpp
--- a/Engine/Source/Runtime/Engine/Private/Animation/AnimSequenceHelpers.cpp
+++ b/Engine/Source/Runtime/Engine/Private/Animation/AnimSequenceHelpers.cpp
@@ -130,6 +130,7 @@ void CopyCurveDataToModel(const FRawCurveTracks& CurveData, const USkeleton* Ske
 			Attributes.SetPreExtrapolation(FloatCurve.FloatCurve.PreInfinityExtrap);
 			Attributes.SetPostExtrapolation(FloatCurve.FloatCurve.PostInfinityExtrap);					
 			Controller.SetCurveAttributes(CurveId, Attributes);
+			Controller.SetCurveColor(CurveId, FloatCurve.GetColor());
 			Controller.SetCurveKeys(CurveId, FloatCurve.FloatCurve.GetConstRefOfKeys());
 		}
 	}
