Allow editing of the HLOD layer property for actors inside LI in WorldPartition

diff --git a/Engine/Source/Runtime/Engine/Private/ActorEditor.cpp b/Engine/Source/Runtime/Engine/Private/ActorEditor.cpp
--- a/Engine/Source/Runtime/Engine/Private/ActorEditor.cpp
+++ b/Engine/Source/Runtime/Engine/Private/ActorEditor.cpp
@@ -84,7 +84,7 @@ bool AActor::CanEditChange(const FProperty* PropertyThatWillChange) const
 		{
 			UWorld* World = GetTypedOuter<UWorld>();
 			UWorldPartition* WorldPartition = World ? World->GetWorldPartition() : nullptr;
-			if (!WorldPartition || (!WorldPartition->IsStreamingEnabled() && !bIsDataLayersProperty))
+			if (!WorldPartition || (!WorldPartition->IsStreamingEnabled() && !bIsDataLayersProperty && !bIsHLODLayerProperty))
 			{
 				return false;
 			}
