Build machines should output shader debug info for failed shaders as a build artifact.

- Due to size limits only output 10 items (this was arbitrarly picked so could be more/less).
- We keep track of how many debug infos we've dumped so far.  This value is incremented when we dump a debug info (which happens on a worker thread) and is read on the main thread so the value is atomic.
- Remove old code for dumping the debug info to the log.

diff --git a/Engine/Source/Runtime/Engine/Private/ShaderCompiler/ShaderCompiler.cpp b/Engine/Source/Runtime/Engine/Private/ShaderCompiler/ShaderCompiler.cpp
--- a/Engine/Source/Runtime/Engine/Private/ShaderCompiler/ShaderCompiler.cpp
+++ b/Engine/Source/Runtime/Engine/Private/ShaderCompiler/ShaderCompiler.cpp
@@ -158,13 +158,12 @@ static FAutoConsoleVariableRef CVarShaderCompilerAllowDistributedCompilation(
 	ECVF_Default
 );
 
-/** Maximum number of preprocessed shaders to dump to the log on a crash. Replace with CVar if needed. */
-int32 GMaxNumDumpedShaderSources = 1;
+int32 GMaxNumDumpedShaderSources = 10;
 static FAutoConsoleVariableRef CVarShaderCompilerMaxDumpedShaderSources(
 	TEXT("r.ShaderCompiler.MaxDumpedShaderSources"),
 	GMaxNumDumpedShaderSources,
-	TEXT("Maximum number of preprocessed shader sources to dump to the log on shader compile errors. By default 1."),
-	ECVF_Default
+	TEXT("Maximum number of preprocessed shader sources to dump as a build artifact on shader compile errors. By default 10."),
+	ECVF_ReadOnly
 );
 
 int32 GSShaderCheckLevel = 1;
@@ -294,21 +293,6 @@ void DumpShaderDDCKeyToFile(const EShaderPlatform InPlatform, bool bWithEditor,
 	DumpAr->Serialize(const_cast<TCHAR*>(*DDCKey), DDCKey.Len() * sizeof(TCHAR));
 }
 
-static void DumpShaderSourceToLog(const FString& DumpedSource)
-{
-	// Log dumped shader source line by line as message lengths in UE_LOG are limited.
-	TArray<FString> DumpedSourceLines;
-	DumpedSource.ParseIntoArrayLines(DumpedSourceLines, /*bCullEmpty:*/ false);
-
-	TGuardValue<ELogTimes::Type> DisableLogTimes(GPrintLogTimes, ELogTimes::None);
-	GLog->Serialize(TEXT("\n======================= DUMPED SHADER BEGIN =======================\n"), ELogVerbosity::Log, NAME_None);
-	for (const FString& Line : DumpedSourceLines)
-		{
-		GLog->Serialize(*Line, ELogVerbosity::Log, NAME_None);
-	}
-	GLog->Serialize(TEXT("\n======================= DUMPED SHADER END =======================\n"), ELogVerbosity::Log, NAME_None);
-}
-
 namespace ShaderCompiler
 {
 	bool IsJobCacheEnabled()
@@ -1726,7 +1710,7 @@ bool FShaderCompileUtilities::DoWriteTasks(const TArray<FShaderCommonCompileJobP
 	return DoWriteTasksInner(QueuedJobs, InTransferFile, BuildDistributionController, bUseRelativePaths, bCompressTaskFile);
 }
 
-static void ProcessErrors(const FShaderCompileJob& CurrentJob, TArray<FString>& UniqueErrors, FString& ErrorString, int32& NumDumpedShaderSources)
+static void ProcessErrors(const FShaderCompileJob& CurrentJob, TArray<FString>& UniqueErrors, FString& ErrorString)
 {
 	bool bReportedDebugInfo = false;
 
@@ -1800,20 +1784,7 @@ static void ProcessErrors(const FShaderCompileJob& CurrentJob, TArray<FString>&
 
 			FString UniqueErrorString = UniqueErrorPrefix + CurrentError.GetErrorStringWithLineMarker() + TEXT("\n");
 
-			if (GIsBuildMachine)
-			{
-				// Format everything on one line, and with the correct verbosity, so we can display proper errors in the failure logs.
-				UE_LOG(LogShaderCompilers, Error, TEXT("%s%s"), *UniqueErrorPrefix.Replace(TEXT("\n"), TEXT("")), *CurrentError.GetErrorStringWithLineMarker());
-
-				const FString& DumpedSource = CurrentJob.Output.OptionalPreprocessedShaderSource;
-				if (!DumpedSource.IsEmpty() && NumDumpedShaderSources < GMaxNumDumpedShaderSources)
-				{
-					// Limit number of preprocessed shaders to dump to the log as they are quite large
-					DumpShaderSourceToLog(DumpedSource);
-					++NumDumpedShaderSources;
-				}
-			}
-			else if (FPlatformMisc::IsDebuggerPresent())
+			if (FPlatformMisc::IsDebuggerPresent())
 			{
 				// Using OutputDebugString to avoid any text getting added before the filename,
 				// Which will throw off VS.NET's ability to take you directly to the file and line of the error when double clicking it in the output window.
@@ -1855,6 +1826,10 @@ static bool ReadSingleJob(FShaderCompileJob* CurrentJob, FArchive& OutputFile)
 	{
 		// Build debug info path and create the directory if it doesn't already exist
 		CurrentJob->Input.DumpDebugInfoPath = GShaderCompilingManager->CreateShaderDebugInfoPath(CurrentJob->Input);
+
+		// We failed to compile a shader, we will retry and dump the shader source. so increment the number of shader sources we've dumped so far.
+		GShaderCompilingManager->IncrementNumDumpedShaderSources();
+
 		return true;
 	}
 
@@ -2190,7 +2165,7 @@ void FShaderCompileUtilities::DoReadTaskResults(const TArray<FShaderCommonCompil
 }
 
 #if WITH_EDITOR
-static bool CheckSingleJob(const FShaderCompileJob& SingleJob, TArray<FString>& OutErrors, FString* OutDumpedSource)
+static bool CheckSingleJob(const FShaderCompileJob& SingleJob, TArray<FString>& OutErrors)
 {
 	if (SingleJob.bSucceeded)
 	{
@@ -2231,11 +2206,6 @@ static bool CheckSingleJob(const FShaderCompileJob& SingleJob, TArray<FString>&
 		}
 	}
 
-	if (!bSucceeded && OutDumpedSource != nullptr)
-	{
-		*OutDumpedSource = SingleJob.Output.OptionalPreprocessedShaderSource;
-	}
-
 	return bSucceeded;
 };
 
@@ -3794,7 +3764,6 @@ FShaderCompilingManager::FShaderCompilingManager() :
 	bCompilingDuringGame(false),
 	NumExternalJobs(0),
 	NumSingleThreadedRunsBeforeRetry(GSingleThreadedRunsIdle),
-	NumDumpedShaderSources(0),
 #if PLATFORM_MAC
 	ShaderCompileWorkerName(FPaths::EngineDir() / TEXT("Binaries/Mac/ShaderCompileWorker")),
 #elif PLATFORM_LINUX
@@ -3887,7 +3856,14 @@ FShaderCompilingManager::FShaderCompilingManager() :
 	FPaths::NormalizeDirectoryName(AbsoluteBaseDirectory);
 	AbsoluteShaderBaseWorkingDirectory = AbsoluteBaseDirectory + TEXT("/");
 
-	FString AbsoluteDebugInfoDirectory = IFileManager::Get().ConvertToAbsolutePathForExternalAppForWrite(*(FPaths::ProjectSavedDir() / TEXT("ShaderDebugInfo")));
+	// Build machines should dump to the AutomationTool/Saved/Logs directory and they will upload as build artifacts via the AutomationTool.
+	FString BaseDebugInfoPath = FPaths::ProjectSavedDir();
+	if (GIsBuildMachine)
+	{
+		BaseDebugInfoPath = FPaths::Combine(*FPaths::EngineDir(), TEXT("Programs"), TEXT("AutomationTool"), TEXT("Saved"), TEXT("Logs"));
+	}
+
+	FString AbsoluteDebugInfoDirectory = IFileManager::Get().ConvertToAbsolutePathForExternalAppForWrite(*(BaseDebugInfoPath / TEXT("ShaderDebugInfo")));
 	const FString OverrideShaderDebugDir = CVarShaderOverrideDebugDir.GetValueOnAnyThread();
 	if (!OverrideShaderDebugDir.IsEmpty())
 	{
@@ -4194,14 +4170,23 @@ bool FShaderCompilingManager::ShouldRecompileToDumpShaderDebugInfo(const FShader
 	if (Input.DumpDebugInfoPath.IsEmpty())
 	{
 		const EDumpShaderDebugInfo DumpShaderDebugInfo = GetDumpShaderDebugInfo();
+		const bool bErrors = !bSucceeded;
+		const bool bWarnings = Output.Errors.Num() > 0;
+
+		bool bShouldDump = true;
+		if (GIsBuildMachine)
+		{
+			// Build machines dump these as build artifacts and they should only upload so many due to size constraints.
+			bShouldDump = (NumDumpedShaderSources < GMaxNumDumpedShaderSources);
+		}
 
 		if (DumpShaderDebugInfo == EDumpShaderDebugInfo::OnError)
 		{
-			return !bSucceeded;
+			return bShouldDump && (bErrors);
 		}
 		else if (DumpShaderDebugInfo == EDumpShaderDebugInfo::OnErrorOrWarning)
 		{
-			return !bSucceeded || Output.Errors.Num() > 0;
+			return bShouldDump && (bErrors || bWarnings);
 		}
 	}
 
@@ -4752,7 +4737,6 @@ void FShaderCompilingManager::ProcessCompiledShaderMaps(
 			TArray<TRefCountPtr<FMaterial>>& MaterialDependencies = CompilingShaderMap->CompilingMaterialDependencies;
 
 			TArray<FString> Errors;
-			FString DumpedSource;
 
 			bool bSuccess = true;
 			for (int32 JobIndex = 0; JobIndex < FinishedJobs.Num(); JobIndex++)
@@ -4761,7 +4745,7 @@ void FShaderCompilingManager::ProcessCompiledShaderMaps(
 
 				if (FShaderCompileJob* SingleJob = CurrentJob.GetSingleShaderJob())
 				{
-					const bool bCheckSucceeded = CheckSingleJob(*SingleJob, Errors, (NumDumpedShaderSources < GMaxNumDumpedShaderSources ? &DumpedSource : nullptr));
+					const bool bCheckSucceeded = CheckSingleJob(*SingleJob, Errors);
 					bSuccess = bCheckSucceeded && bSuccess;
 				}
 				else
@@ -4769,7 +4753,7 @@ void FShaderCompilingManager::ProcessCompiledShaderMaps(
 					FShaderPipelineCompileJob* PipelineJob = CurrentJob.GetShaderPipelineJob();
 					for (int32 Index = 0; Index < PipelineJob->StageJobs.Num(); ++Index)
 					{
-						const bool bCheckSucceeded = CheckSingleJob(*PipelineJob->StageJobs[Index], Errors, (NumDumpedShaderSources < GMaxNumDumpedShaderSources ? &DumpedSource : nullptr));
+						const bool bCheckSucceeded = CheckSingleJob(*PipelineJob->StageJobs[Index], Errors);
 						bSuccess = PipelineJob->StageJobs[Index]->bSucceeded && bCheckSucceeded && bSuccess;
 					}
 				}
@@ -4859,13 +4843,6 @@ void FShaderCompilingManager::ProcessCompiledShaderMaps(
 
 					MaterialsToUpdate.Add(Material, nullptr);
 
-					if (!DumpedSource.IsEmpty() && NumDumpedShaderSources < GMaxNumDumpedShaderSources)
-					{
-						// Limit number of preprocessed shaders to dump to the log as they are quite large
-						DumpShaderSourceToLog(DumpedSource);
-						++NumDumpedShaderSources;
-					}
-
 					if (Material->IsDefaultMaterial())
 					{
 						FString ErrorString;
@@ -5156,7 +5133,7 @@ bool FShaderCompilingManager::HandlePotentialRetryOnError(TMap<int32, FShaderMap
 						const auto* SingleJob = CurrentJob.GetSingleShaderJob();
 						if (SingleJob)
 						{
-							ProcessErrors(*SingleJob, UniqueErrors, ErrorString, NumDumpedShaderSources);
+							ProcessErrors(*SingleJob, UniqueErrors, ErrorString);
 						}
 						else
 						{
@@ -5164,7 +5141,7 @@ bool FShaderCompilingManager::HandlePotentialRetryOnError(TMap<int32, FShaderMap
 							check(PipelineJob);
 							for (auto CommonJob : PipelineJob->StageJobs)
 							{
-								ProcessErrors(*CommonJob, UniqueErrors, ErrorString, NumDumpedShaderSources);
+								ProcessErrors(*CommonJob, UniqueErrors, ErrorString);
 							}
 						}
 					}
diff --git a/Engine/Source/Runtime/Engine/Public/ShaderCompiler.h b/Engine/Source/Runtime/Engine/Public/ShaderCompiler.h
--- a/Engine/Source/Runtime/Engine/Public/ShaderCompiler.h
+++ b/Engine/Source/Runtime/Engine/Public/ShaderCompiler.h
@@ -894,8 +894,9 @@ private:
 	int32 MaxShaderJobBatchSize;
 	/** Number of runs through single-threaded compiling before we can retry to compile through workers. -1 if not used. */
 	int32 NumSingleThreadedRunsBeforeRetry;
-	/** Number of preprocessed shader sources that are dumped due to a crash of the shader compiler. Only used for D3DCompile exceptions at the moment. */
-	int32 NumDumpedShaderSources;
+	/** Number of preprocessed shader sources that are dumped due to a crash of the shader compiler. */
+	/** This state is updated when a job is completed, which happens on a worker thread.  The state is read on the Game Thread to see if we should dump the source to begin with. */
+	std::atomic<int32> NumDumpedShaderSources = 0;
 	/** Process Id of UE. */
 	uint32 ProcessId;
 	/** Whether to allow compiling shaders through the worker application, which allows multiple cores to be used. */
@@ -1067,6 +1068,12 @@ public:
 	ENGINE_API bool ShouldRecompileToDumpShaderDebugInfo(const FShaderCompileJob& Job) const;
 	ENGINE_API bool ShouldRecompileToDumpShaderDebugInfo(const FShaderCompilerInput& Input, const FShaderCompilerOutput& Output, bool bSucceeded) const;
 
+
+	void IncrementNumDumpedShaderSources()
+	{
+		NumDumpedShaderSources++;
+	}
+
 	const FString& GetAbsoluteShaderDebugInfoDirectory() const
 	{
 		return AbsoluteShaderDebugInfoDirectory;
diff --git a/Engine/Source/Runtime/RenderCore/Public/ShaderCompilerCore.h b/Engine/Source/Runtime/RenderCore/Public/ShaderCompilerCore.h
--- a/Engine/Source/Runtime/RenderCore/Public/ShaderCompilerCore.h
+++ b/Engine/Source/Runtime/RenderCore/Public/ShaderCompilerCore.h
@@ -20,7 +20,7 @@ class Error;
 // this is for the protocol, not the data, bump if FShaderCompilerInput or ProcessInputFromArchive changes.
 inline const int32 ShaderCompileWorkerInputVersion = 17;
 // this is for the protocol, not the data, bump if FShaderCompilerOutput or WriteToOutputArchive changes.
-inline const int32 ShaderCompileWorkerOutputVersion = 9;
+inline const int32 ShaderCompileWorkerOutputVersion = 10;
 // this is for the protocol, not the data.
 inline const int32 ShaderCompileWorkerSingleJobHeader = 'S';
 // this is for the protocol, not the data.
@@ -470,7 +470,6 @@ struct FShaderCompilerOutput
 	bool bUsedHLSLccCompiler;
 	TArray<FString> UsedAttributes;
 
-	FString OptionalPreprocessedShaderSource;
 	FString OptionalFinalShaderSource;
 
 	TArray<uint8> PlatformDebugData;
@@ -488,7 +487,6 @@ struct FShaderCompilerOutput
 		Ar << Output.bFailedRemovingUnused << Output.bSupportsQueryingUsedAttributes << Output.UsedAttributes;
 		Ar << Output.CompileTime;
 		Ar << Output.PreprocessTime;
-		Ar << Output.OptionalPreprocessedShaderSource;
 		Ar << Output.OptionalFinalShaderSource;
 		Ar << Output.PlatformDebugData;
 
