Sanitize data to ensure stale entries do not linger around

diff --git a/Engine/Source/Runtime/Engine/Private/Animation/BlendProfile.cpp b/Engine/Source/Runtime/Engine/Private/Animation/BlendProfile.cpp
--- a/Engine/Source/Runtime/Engine/Private/Animation/BlendProfile.cpp
+++ b/Engine/Source/Runtime/Engine/Private/Animation/BlendProfile.cpp
@@ -121,6 +121,12 @@ void UBlendProfile::SetSkeleton(USkeleton* InSkeleton)
 			Entry.BoneReference.Initialize(OwningSkeleton);
 		}
 	}
+
+	// Remove any entries for bones that aren't mapped
+	ProfileEntries.RemoveAll([](const FBlendProfileBoneEntry& Current)
+		{
+			return Current.BoneReference.BoneIndex == INDEX_NONE;
+		});
 }
 
 void UBlendProfile::PostLoad()
@@ -135,6 +141,14 @@ void UBlendProfile::PostLoad()
 			Entry.BoneReference.Initialize(OwningSkeleton);
 		}
 	}
+
+#if WITH_EDITOR
+	// Remove any entries for bones that aren't mapped
+	ProfileEntries.RemoveAll([](const FBlendProfileBoneEntry& Current)
+		{
+			return Current.BoneReference.BoneIndex == INDEX_NONE;
+		});
+#endif
 }
 
 int32 UBlendProfile::GetEntryIndex(const int32 InBoneIdx) const
