Avoid loading editor only material in FWorldPartitionRuntimeSpatialHashGridPreviewer when cooking.

diff --git a/Engine/Source/Runtime/Engine/Private/WorldPartition/WorldPartitionRuntimeSpatialHashGridPreviewer.cpp b/Engine/Source/Runtime/Engine/Private/WorldPartition/WorldPartitionRuntimeSpatialHashGridPreviewer.cpp
--- a/Engine/Source/Runtime/Engine/Private/WorldPartition/WorldPartitionRuntimeSpatialHashGridPreviewer.cpp
+++ b/Engine/Source/Runtime/Engine/Private/WorldPartition/WorldPartitionRuntimeSpatialHashGridPreviewer.cpp
@@ -19,7 +19,10 @@ FWorldPartitionRuntimeSpatialHashGridPreviewer::FWorldPartitionRuntimeSpatialHas
 #endif
 {
 #if WITH_EDITORONLY_DATA
-	Material = LoadObject<UMaterial>(nullptr, TEXT("/Engine/EditorMaterials/WorldPartition/WorldPartitionSpatialHashGridPreviewMaterial"));
+	if (!IsRunningCookCommandlet())
+	{
+		Material = LoadObject<UMaterial>(nullptr, TEXT("/Engine/EditorMaterials/WorldPartition/WorldPartitionSpatialHashGridPreviewMaterial"));
+	}
 #endif
 }
 
