Fix issue when an attached actor is created, the absolute scaling is sent to the client instead of the relative scale to its parent. This causes issues because that scale would then get treated as the relative scale, causing the scaling to be applied twice.

diff --git a/Engine/Source/Runtime/Engine/Private/PackageMapClient.cpp b/Engine/Source/Runtime/Engine/Private/PackageMapClient.cpp
--- a/Engine/Source/Runtime/Engine/Private/PackageMapClient.cpp
+++ b/Engine/Source/Runtime/Engine/Private/PackageMapClient.cpp
@@ -520,6 +520,16 @@ bool UPackageMapClient::SerializeNewActor(FArchive& Ar, class UActorChannel *Cha
 				Location = FRepMovement::RebaseOntoZeroOrigin(Actor->GetActorLocation(), Actor);
 				Rotation = Actor->GetActorRotation();
 				Scale = Actor->GetActorScale();
+
+				if (USceneComponent* AttachParent = RootComponent->GetAttachParent())
+				{
+					// If this actor is attached, when the scale is serialized on the client, the attach parent property won't be set yet.
+					// USceneComponent::SetWorldScale3D (which got called by AActor::SetActorScale3D, which we used to do but no longer).
+					// would perform this transformation so that what is sent is relative to the parent. If we don't do this, we will
+					// apply the world scale on the client, which will then get applied a second time when the attach parent property is received.
+					FTransform ParentToWorld = AttachParent->GetSocketTransform(RootComponent->GetAttachSocketName());
+					Scale = Scale * ParentToWorld.GetSafeScaleReciprocal(ParentToWorld.GetScale3D());
+				}
 				Velocity = Actor->GetVelocity();
 			}
 		}
@@ -668,7 +678,7 @@ bool UPackageMapClient::SerializeNewActor(FArchive& Ar, class UActorChannel *Cha
 							// Scale was serialized by the server
 							if (bSerializeScale)
 							{
-								Actor->SetActorScale3D(Scale);
+								Actor->SetActorRelativeScale3D(Scale);
 							}
 
 							GuidCache->RegisterNetGUID_Client(NetGUID, Actor);
