fix crash when accessing nullptr when adding geometry collection to external spatial acceleration due to using the internal physics object interface instead of the external one

diff --git a/Engine/Source/Runtime/Engine/Private/PhysicsEngine/Experimental/PhysScene_Chaos.cpp b/Engine/Source/Runtime/Engine/Private/PhysicsEngine/Experimental/PhysScene_Chaos.cpp
--- a/Engine/Source/Runtime/Engine/Private/PhysicsEngine/Experimental/PhysScene_Chaos.cpp
+++ b/Engine/Source/Runtime/Engine/Private/PhysicsEngine/Experimental/PhysScene_Chaos.cpp
@@ -13,7 +13,7 @@
 #include "PhysicsReplication.h"
 #include "PhysicsEngine/ConstraintInstance.h"
 #include "PhysicsEngine/PhysicsCollisionHandler.h"
-#include "Chaos/PhysicsObjectInternalInterface.h"
+#include "PhysicsEngine/PhysicsObjectExternalInterface.h"
 
 #include "ChaosSolversModule.h"
 
@@ -1769,16 +1769,16 @@ void FPhysScene_Chaos::OnSyncBodies(Chaos::FPhysicsSolverBase* Solver)
 
 		auto GeometryCollectionLambda = [this](FGeometryCollectionPhysicsProxy* Proxy)
 		{
-			// Use the interface interface since we're already locked.
-			Chaos::FWritePhysicsObjectInterface_Internal Interface = Chaos::FPhysicsObjectInternalInterface::GetWrite();
+			// Use don't pass in anything here so we don't end up locking anything because we can assume the scene is already locked.
+			FLockedWritePhysicsObjectExternalInterface Interface = FPhysicsObjectExternalInterface::LockWrite({});
 			TArray<Chaos::FPhysicsObjectHandle> Handles = Proxy->GetAllPhysicsObjects().FilterByPredicate(
 				[&Interface](Chaos::FPhysicsObjectHandle Handle)
 				{
-					return !Interface.AreAllDisabled({ &Handle, 1 });
+					return !Interface->AreAllDisabled({ &Handle, 1 });
 				}
 			);
 
-			Interface.AddToSpatialAcceleration(Handles, GetSpacialAcceleration());
+			Interface->AddToSpatialAcceleration(Handles, GetSpacialAcceleration());
 		};
 
 		Solver->PullPhysicsStateForEachDirtyProxy_External(RigidLambda, ConstraintLambda, GeometryCollectionLambda);
