Fixed issue where sim proxy characters on dynamic, rotated movement bases would be in the wrong location while playing animated root motion.

diff --git a/Engine/Source/Runtime/Engine/Private/Character.cpp b/Engine/Source/Runtime/Engine/Private/Character.cpp
--- a/Engine/Source/Runtime/Engine/Private/Character.cpp
+++ b/Engine/Source/Runtime/Engine/Private/Character.cpp
@@ -1407,8 +1407,9 @@ bool ACharacter::RestoreReplicatedMove(const FSimulatedRootMotionReplicatedMove&
 			FVector BaseLocation;
 			FQuat BaseRotation;
 			MovementBaseUtility::GetMovementBaseTransform(ServerBase, ServerBaseBoneName, BaseLocation, BaseRotation);
-
-			const FVector ServerLocation = BaseLocation + RootMotionRepMove.RootMotion.Location;
+			const FTransform BaseTransform(BaseRotation, BaseLocation);
+			
+			const FVector ServerLocation = BaseTransform.TransformPositionNoScale(RootMotionRepMove.RootMotion.Location);
 			FRotator ServerRotation;
 			if (RootMotionRepMove.RootMotion.bRelativeRotation)
 			{
