Improve material analysis to add support for analyzing:
- Which exact texture coordinates are used
- Which exact shading models are possible
- Custom outputs (just like material properties)
- lightmapUVs (when determining if vertex data is required)

(Based on merger of PRs #8486 #8487 #8488 #8489 #9040)

diff --git a/Engine/Source/Runtime/Engine/Classes/Materials/MaterialInterface.h b/Engine/Source/Runtime/Engine/Classes/Materials/MaterialInterface.h
--- a/Engine/Source/Runtime/Engine/Classes/Materials/MaterialInterface.h
+++ b/Engine/Source/Runtime/Engine/Classes/Materials/MaterialInterface.h
@@ -64,6 +64,7 @@ struct FParameterChannelNames;
 #endif
 enum EShaderPlatform : uint16;
 struct FStrataCompilationConfig;
+class UMaterialExpressionCustomOutput;
 
 typedef TArray<FMaterialResource*> FMaterialResourceDeferredDeletionArray;
 
@@ -212,6 +213,21 @@ struct FMaterialInheritanceChain
 	inline const FMaterialCachedExpressionData& GetCachedExpressionData() const { checkSlow(CachedExpressionData); return *CachedExpressionData; }
 };
 
+/**
+ * Holds data about what is used in the shader graph of a specific material property or custom output
+ */
+struct FMaterialAnalysisResult
+{
+	/** The texture coordinates used */
+	TBitArray<> TextureCoordinates;
+
+	/** The shading models used (only relevant when analyzing property MP_ShadingModel) */
+	FMaterialShadingModelField ShadingModels;
+
+	/** Whether any vertex data is used */
+	bool bRequiresVertexData = false;
+};
+
 UCLASS(Optional)
 class UMaterialInterfaceEditorOnlyData : public UObject
 {
@@ -995,6 +1011,15 @@ public:
 	/** Return number of used texture coordinates and whether or not the Vertex data is used in the shader graph */
 	ENGINE_API void AnalyzeMaterialProperty(EMaterialProperty InProperty, int32& OutNumTextureCoordinates, bool& bOutRequiresVertexData);
 
+	/** Return insight on what (e.g. texture coordinates, vertex data, etc) is used in the shader graph of a material property */
+	ENGINE_API void AnalyzeMaterialPropertyEx(EMaterialProperty InProperty, FMaterialAnalysisResult& OutResult);
+
+	/** Return insight on what (e.g. texture coordinates, vertex data, etc) is used in the shader graph of a material custom output */
+	ENGINE_API void AnalyzeMaterialCustomOutput(UMaterialExpressionCustomOutput* InCustomOutput, int32 InOutputIndex, FMaterialAnalysisResult& OutResult);
+
+	/** Return insight on what (e.g. texture coordinates, vertex data, etc) is used in the shader graph compiled by a callback */
+	ENGINE_API void AnalyzeMaterialCompilationInCallback(TFunctionRef<void (FMaterialCompiler*)> InCompilationCallback, FMaterialAnalysisResult& OutResult);
+
 #if WITH_EDITOR
 	/** Checks to see if the given property references the texture */
 	ENGINE_API bool IsTextureReferencedByProperty(EMaterialProperty InProperty, const UTexture* InTexture);
diff --git a/Engine/Source/Runtime/Engine/Private/Materials/MaterialShared.cpp b/Engine/Source/Runtime/Engine/Private/Materials/MaterialShared.cpp
--- a/Engine/Source/Runtime/Engine/Private/Materials/MaterialShared.cpp
+++ b/Engine/Source/Runtime/Engine/Private/Materials/MaterialShared.cpp
@@ -4425,6 +4425,39 @@ int32 UMaterialInterface::CompileProperty(FMaterialCompiler* Compiler, EMaterial
 
 void UMaterialInterface::AnalyzeMaterialProperty(EMaterialProperty InProperty, int32& OutNumTextureCoordinates, bool& bOutRequiresVertexData)
 {
+#if WITH_EDITORONLY_DATA
+	FMaterialAnalysisResult Result;
+	AnalyzeMaterialPropertyEx(InProperty, Result);
+
+	OutNumTextureCoordinates = Result.TextureCoordinates.FindLast(true) + 1;
+	bOutRequiresVertexData = Result.bRequiresVertexData;
+#endif
+}
+
+void UMaterialInterface::AnalyzeMaterialPropertyEx(EMaterialProperty InProperty, FMaterialAnalysisResult& OutResult)
+{
+#if WITH_EDITORONLY_DATA
+	AnalyzeMaterialCompilationInCallback([InProperty, this](FMaterialCompiler* Compiler)
+	{
+		Compiler->SetMaterialProperty(InProperty);
+		this->CompileProperty(Compiler, InProperty);
+	}, OutResult);
+#endif
+}
+
+void UMaterialInterface::AnalyzeMaterialCustomOutput(UMaterialExpressionCustomOutput* InCustomOutput, int32 InOutputIndex, FMaterialAnalysisResult& OutResult)
+{
+#if WITH_EDITORONLY_DATA
+	AnalyzeMaterialCompilationInCallback([InCustomOutput, InOutputIndex](FMaterialCompiler* Compiler)
+	{
+		Compiler->SetMaterialProperty(MP_MAX, InCustomOutput->GetShaderFrequency());
+		InCustomOutput->Compile(Compiler, InOutputIndex);
+	}, OutResult);
+#endif
+}
+
+void UMaterialInterface::AnalyzeMaterialCompilationInCallback(TFunctionRef<void (FMaterialCompiler*)> InCompilationCallback, FMaterialAnalysisResult& OutResult)
+{
 #if WITH_EDITORONLY_DATA
 	// FHLSLMaterialTranslator collects all required information during translation, but these data are protected. Needs to
 	// derive own class from it to get access to these data.
@@ -4434,34 +4467,15 @@ void UMaterialInterface::AnalyzeMaterialProperty(EMaterialProperty InProperty, i
 		FMaterialAnalyzer(FMaterial* InMaterial, FMaterialCompilationOutput& InMaterialCompilationOutput, const FStaticParameterSet& StaticParameters, EShaderPlatform InPlatform, EMaterialQualityLevel::Type InQualityLevel, ERHIFeatureLevel::Type InFeatureLevel)
 			: FHLSLMaterialTranslator(InMaterial, InMaterialCompilationOutput, StaticParameters, InPlatform, InQualityLevel, InFeatureLevel)
 		{}
-		int32 GetTextureCoordsCount() const
-		{
-			return GetNumUserTexCoords();
-		}
-		bool UsesVertexColor() const
-		{
-			return bUsesVertexColor;
-		}
-
-		bool UsesTransformVector() const
-		{
-			return bUsesTransformVector;
-		}
-
-		bool UsesWorldPositionExcludingShaderOffsets() const
-		{
-			return bNeedsWorldPositionExcludingShaderOffsets;
-		}
 
-		bool UsesPrecomputedAOMask() const
-		{
-			return bUsesAOMaterialMask;
-		}
-
-		bool UsesVertexPosition() const 
-		{
-			return bUsesVertexPosition;
-		}
+		using FHLSLMaterialTranslator::AllocatedUserTexCoords;
+		using FHLSLMaterialTranslator::ShadingModelsFromCompilation;
+		using FHLSLMaterialTranslator::bUsesVertexColor;
+		using FHLSLMaterialTranslator::bUsesTransformVector;
+		using FHLSLMaterialTranslator::bNeedsWorldPositionExcludingShaderOffsets;
+		using FHLSLMaterialTranslator::bUsesAOMaterialMask;
+		using FHLSLMaterialTranslator::bUsesLightmapUVs;
+		using FHLSLMaterialTranslator::bUsesVertexPosition;
 	};
 
 	FMaterialCompilationOutput TempOutput;
@@ -4477,12 +4491,19 @@ void UMaterialInterface::AnalyzeMaterialProperty(EMaterialProperty InProperty, i
 	FStaticParameterSet StaticParamSet;
 	MaterialResource->GetStaticParameterSet(GMaxRHIShaderPlatform, StaticParamSet);
 	FMaterialAnalyzer MaterialTranslator(MaterialResource, TempOutput, StaticParamSet, GMaxRHIShaderPlatform, MaterialResource->GetQualityLevel(), GMaxRHIFeatureLevel);
-	
-	static_cast<FMaterialCompiler*>(&MaterialTranslator)->SetMaterialProperty(InProperty); // FHLSLMaterialTranslator hides this interface, so cast to parent
-	CompileProperty(&MaterialTranslator, InProperty);
+
+	InCompilationCallback(&MaterialTranslator);
+
 	// Request data from translator
-	OutNumTextureCoordinates = MaterialTranslator.GetTextureCoordsCount();
-	bOutRequiresVertexData = MaterialTranslator.UsesVertexColor() || MaterialTranslator.UsesTransformVector() || MaterialTranslator.UsesWorldPositionExcludingShaderOffsets() || MaterialTranslator.UsesPrecomputedAOMask() || MaterialTranslator.UsesVertexPosition();
+	OutResult.TextureCoordinates = MaterialTranslator.AllocatedUserTexCoords;
+	OutResult.ShadingModels = MaterialTranslator.ShadingModelsFromCompilation;
+	OutResult.bRequiresVertexData =
+		MaterialTranslator.bUsesVertexColor ||
+		MaterialTranslator.bUsesTransformVector ||
+		MaterialTranslator.bNeedsWorldPositionExcludingShaderOffsets ||
+		MaterialTranslator.bUsesAOMaterialMask ||
+		MaterialTranslator.bUsesLightmapUVs ||
+		MaterialTranslator.bUsesVertexPosition;
 #endif
 }
 
