- Make sure we gracefully handle ::BeginDestroy for an AnimSequence with in-flight compression tasks
- Also make sure we handle ::BeginDestroy for a USkeleton, who might be targetted by in-flight AnimSequence compression tasks
- React to GC marking an AnimSequence, with in-flight compression, as unreachable
- Make BoneContainer accept Unreachable (but not yet destroyed) UObject references through WeakPtr access in WITH_EDITOR builds, this can happen when a Skeleton is marked Unreachable during GC - while async compression task is trying to sample raw animation data (poses)

diff --git a/Engine/Plugins/Animation/AnimationData/Source/AnimationData/Private/AnimSequencerDataModel.cpp b/Engine/Plugins/Animation/AnimationData/Source/AnimationData/Private/AnimSequencerDataModel.cpp
--- a/Engine/Plugins/Animation/AnimationData/Source/AnimationData/Private/AnimSequencerDataModel.cpp
+++ b/Engine/Plugins/Animation/AnimationData/Source/AnimationData/Private/AnimSequencerDataModel.cpp
@@ -1118,7 +1118,7 @@ void UAnimationSequencerDataModel::GeneratePoseData(UControlRig* ControlRig, FAn
 					const int32 BoneIndex = RefSkeleton.FindBoneIndex(BoneName);
 					if (BoneIndex != INDEX_NONE)
 					{
-						const int32 SkeletonBoneIndex = RequiredBones.GetSkeletonAsset()->GetReferenceSkeleton().FindBoneIndex(BoneName);
+						const int32 SkeletonBoneIndex = RefSkeleton.FindBoneIndex(BoneName);
 						const FCompactPoseBoneIndex CompactPoseBoneIndex = RequiredBones.GetCompactPoseIndexFromSkeletonIndex(SkeletonBoneIndex);
 						if (CompactPoseBoneIndex != INDEX_NONE)
 						{
@@ -1133,7 +1133,8 @@ void UAnimationSequencerDataModel::GeneratePoseData(UControlRig* ControlRig, FAn
 
 				if (Curve.IsValid())
 				{
-					const FSmartNameMapping* SmartNameMapping = RequiredBones.GetSkeletonAsset()->GetSmartNameContainer(USkeleton::AnimCurveMappingName);
+					// Called during compression that can occur while GC is in progress, marking weakptrs as unreachable temporarily
+					const FSmartNameMapping* SmartNameMapping = RequiredBones.GetSkeletonAsset(true)->GetSmartNameContainer(USkeleton::AnimCurveMappingName);
 					RigHierarchy->ForEach<FRigCurveElement>([this, &RigHierarchy, &Curve, SmartNameMapping](const FRigCurveElement* CurveElement) -> bool
 					{
 						const FName& CurveName = CurveElement->GetName();
diff --git a/Engine/Source/Runtime/Engine/Private/Animation/AnimCompressionTypes.cpp b/Engine/Source/Runtime/Engine/Private/Animation/AnimCompressionTypes.cpp
--- a/Engine/Source/Runtime/Engine/Private/Animation/AnimCompressionTypes.cpp
+++ b/Engine/Source/Runtime/Engine/Private/Animation/AnimCompressionTypes.cpp
@@ -177,7 +177,7 @@ static bool IsNewKeyDifferent(const FRichCurveKey& LastKey, float NewValue)
 void FCompressibleAnimData::BakeOutAdditiveIntoRawData(const FFrameRate& SampleRate, TArray<FBoneAnimationTrack>& ResampledTrackData, TArray<FFloatCurve>& FloatCurves)
 {
 	TRACE_CPUPROFILER_EVENT_SCOPE(FCompressibleAnimData::BakeOutAdditiveIntoRawData);
-	const UAnimSequence* AnimSequence = WeakSequence.Get();
+	const UAnimSequence* AnimSequence = WeakSequence.GetEvenIfUnreachable();
 	checkf(AnimSequence, TEXT("Animation Sequence is invalid"));
 
 	const USkeleton* Skeleton = AnimSequence->GetSkeleton();
@@ -386,7 +386,7 @@ static void FindAnimatedVirtualBones(const TArray<FBoneAnimationTrack>& Animated
 
 void FCompressibleAnimData::ResampleAnimationTrackData(const FFrameRate& SampleRate, TArray<FBoneAnimationTrack>& ResampledTrackData) const
 {
-	UAnimSequence* AnimSequence = WeakSequence.Get();
+	UAnimSequence* AnimSequence = WeakSequence.GetEvenIfUnreachable();
 	checkf(AnimSequence, TEXT("Animation Sequence is invalid"));
 	
 	QUICK_SCOPE_CYCLE_COUNTER(STAT_ResampleAnimationTrackData);
@@ -729,11 +729,16 @@ void FCompressibleAnimData::FetchData(const ITargetPlatform* InPlatform)
 	if (bDataFetched)
 		return;
 
+	checkf(InPlatform, TEXT("Invalid target platform while trying to fetch to-compress animation data"));
+
 	bDataFetched = true;
 	TRACE_CPUPROFILER_EVENT_SCOPE(FCompressibleAnimData::FetchData);
 
-	const UAnimSequence* AnimSequence = WeakSequence.Get();
+	const UAnimSequence* AnimSequence = WeakSequence.GetEvenIfUnreachable();
+	checkf(AnimSequence, TEXT("Invalid animation sequence while trying to fetch to-compress animation data"));
 	USkeleton* Skeleton = AnimSequence->GetSkeleton();
+	checkf(Skeleton, TEXT("Invalid Skeleton while trying to fetch to-compress animation data, %s"), *Skeleton->GetPathName());
+	
 	FAnimationUtils::BuildSkeletonMetaData(Skeleton, BoneData);
 
 	const FFrameRate DefaultSamplingFrameRate = AnimSequence->GetSamplingFrameRate();
diff --git a/Engine/Source/Runtime/Engine/Private/Animation/AnimSequence.cpp b/Engine/Source/Runtime/Engine/Private/Animation/AnimSequence.cpp
--- a/Engine/Source/Runtime/Engine/Private/Animation/AnimSequence.cpp
+++ b/Engine/Source/Runtime/Engine/Private/Animation/AnimSequence.cpp
@@ -449,6 +449,8 @@ void UAnimSequence::AddReferencedObjects(UObject* This, FReferenceCollector& Col
 void UAnimSequence::WillNeverCacheCookedPlatformDataAgain()
 {
 	Super::WillNeverCacheCookedPlatformDataAgain();
+
+	UE::Anim::FAnimSequenceCompilingManager::Get().FinishCompilation({this});	
 	ClearCompressedCurveData();
 	ClearCompressedBoneData();
 
@@ -962,7 +964,7 @@ void UAnimSequence::BeginDestroy()
 {
 #if WITH_EDITOR
 	// Could already be compressing
-	UE::Anim::FAnimSequenceCompilingManager::Get().FinishCompilation({this});
+	WaitOnExistingCompression(false);
 #endif // WITH_EDITOR
 
 	Super::BeginDestroy();
@@ -4996,6 +4998,8 @@ FIoHash UAnimSequence::BeginCacheDerivedData(const ITargetPlatform* TargetPlatfo
 	check(IsInGameThread());
 	check(!FPlatformProperties::RequiresCookedData());
 	
+	check(!IsUnreachable());
+	
 	if(bBlockCompressionRequests)
 	{
 		UE_LOG(LogAnimation, Warning, TEXT("Animation Compression request for %s was blocked, bBlockCompressionRequests == true."), *GetName());
@@ -5192,10 +5196,10 @@ void UAnimSequence::FinishAsyncTasks()
 		COOK_STAT(Timer.TrackCyclesOnly());
 	
 #if WITH_EDITOR
-		SynchronousAnimatedBoneAttributesCompression();
 		//This is only safe during sync anim compression
 		if (GetSkeleton())
 		{
+			SynchronousAnimatedBoneAttributesCompression();
 			SetSkeletonVirtualBoneGuid(GetSkeleton()->GetVirtualBoneGuid());
 		}
 #endif
diff --git a/Engine/Source/Runtime/Engine/Private/Animation/AnimationCompressionDerivedData.cpp b/Engine/Source/Runtime/Engine/Private/Animation/AnimationCompressionDerivedData.cpp
--- a/Engine/Source/Runtime/Engine/Private/Animation/AnimationCompressionDerivedData.cpp
+++ b/Engine/Source/Runtime/Engine/Private/Animation/AnimationCompressionDerivedData.cpp
@@ -8,6 +8,7 @@
 #include "Animation/AnimationSequenceCompiler.h"
 #include "Animation/AnimCompress.h"
 #include "DerivedDataCache.h"
+#include "Animation/Skeleton.h"
 
 namespace UE::Anim
 {	
@@ -31,6 +32,9 @@ FAnimationSequenceAsyncCacheTask::FAnimationSequenceAsyncCacheTask(const FIoHash
 		, CompressibleAnimPtr(InCompressibleAnimPtr)
 		, TargetPlatform(InTargetPlatform)
 {
+	check(!InAnimSequence.IsUnreachable());
+	check(!InAnimSequence.GetSkeleton()->IsUnreachable());
+		
 	BeginCache(InKeyHash);
 }
 
@@ -193,6 +197,7 @@ bool FAnimationSequenceAsyncCacheTask::BuildData() const
 	TRACE_CPUPROFILER_EVENT_SCOPE_TEXT(*(FString(TEXT("FAnimationSequenceAsyncCacheTask::BuildData ") + CompressibleAnimPtr->Name)));
 	UE_LOG(LogAnimationCompression, Display, TEXT("Building compressed animation data for %s"), *CompressibleAnimPtr->FullName);
 
+	check(CompressibleAnimPtr.IsValid());
 	FCompressibleAnimData& DataToCompress = *CompressibleAnimPtr.Get();
 	FCompressedAnimSequence& OutData = *CompressedData;
 
diff --git a/Engine/Source/Runtime/Engine/Private/Animation/AnimationSequenceCompiler.cpp b/Engine/Source/Runtime/Engine/Private/Animation/AnimationSequenceCompiler.cpp
--- a/Engine/Source/Runtime/Engine/Private/Animation/AnimationSequenceCompiler.cpp
+++ b/Engine/Source/Runtime/Engine/Private/Animation/AnimationSequenceCompiler.cpp
@@ -10,6 +10,9 @@
 
 #include "ObjectCacheContext.h"
 #include "UObject/Package.h"
+#include "Animation/Skeleton.h"
+#include "Animation/AnimSequence.h"
+#include "UObject/StrongObjectPtr.h"
 
 #define LOCTEXT_NAMESPACE "AnimSequenceCompilingManager"
 
@@ -24,6 +27,7 @@ namespace UE::Anim
 	FAnimSequenceCompilingManager::FAnimSequenceCompilingManager()
 		: Notification(FAnimSequenceCompilingManager::GetAssetNameFormat())
 	{
+		PostReachabilityAnalysisHandle = FCoreUObjectDelegates::PostReachabilityAnalysis.AddRaw(this, &FAnimSequenceCompilingManager::OnPostReachabilityAnalysis);
 	}
 	
 	TRACE_DECLARE_INT_COUNTER(QueuedAnimationSequenceCompilation, TEXT("AsyncCompilation/QueuedAnimationSequences"));
@@ -74,6 +78,8 @@ namespace UE::Anim
 
 			FinishCompilation(PendingAnimationSequences);
 		}
+		
+		FCoreUObjectDelegates::PostReachabilityAnalysis.Remove(PostReachabilityAnalysisHandle);
 	}
 
 	FName FAnimSequenceCompilingManager::GetAssetTypeName() const
@@ -169,7 +175,7 @@ namespace UE::Anim
 					return true;
 				}
 
-				UAnimSequence* AnimSequence;
+				TStrongObjectPtr<UAnimSequence> AnimSequence;
 				virtual FName GetName() override { return AnimSequence->GetOutermost()->GetFName(); }
 			};
 
@@ -182,7 +188,7 @@ namespace UE::Anim
 				LogAnimation,
 				[this](AsyncCompilationHelpers::ICompilable* Object)
 				{
-					UAnimSequence* AnimSequence = static_cast<FCompilableAnimationSequence*>(Object)->AnimSequence;
+					UAnimSequence* AnimSequence = static_cast<FCompilableAnimationSequence*>(Object)->AnimSequence.Get();
 					ApplyCompilation(AnimSequence);
 					RegisteredAnimSequences.Remove(AnimSequence);
 				}
@@ -193,6 +199,23 @@ namespace UE::Anim
 		
 	}
 
+	void FAnimSequenceCompilingManager::FinishCompilation(TArrayView<USkeleton* const> InSkeletons)
+	{
+		TArray<UAnimSequence*> PendingAnimationSequences;	
+		for (const TWeakObjectPtr<UAnimSequence>& WeakAnimSequence : RegisteredAnimSequences)
+		{
+			if (UAnimSequence* AnimSequence = WeakAnimSequence.GetEvenIfUnreachable())
+			{
+				if (InSkeletons.Contains(AnimSequence->GetSkeleton()))
+				{
+					PendingAnimationSequences.Add(AnimSequence);
+				}
+			}
+		}
+
+		FinishCompilation(PendingAnimationSequences);
+	}
+
 	void FAnimSequenceCompilingManager::ProcessAsyncTasks(bool bLimitExecutionTime)
 	{
 		FObjectCacheContextScope ObjectCacheScope;
@@ -281,8 +304,11 @@ namespace UE::Anim
 		if (!IsEngineExitRequested())
 		{
 			check(IsInGameThread());
-			TRACE_CPUPROFILER_EVENT_SCOPE(PostCompilation);			
-			InAnimSequence->FinishAsyncTasks();
+			TRACE_CPUPROFILER_EVENT_SCOPE(PostCompilation);
+			if (InAnimSequence && !InAnimSequence->IsUnreachable() && InAnimSequence->GetSkeleton() && !InAnimSequence->GetSkeleton()->IsUnreachable())
+			{
+				InAnimSequence->FinishAsyncTasks();
+			}
 		}
 	}
 	void FAnimSequenceCompilingManager::UpdateCompilationNotification()
@@ -290,6 +316,41 @@ namespace UE::Anim
 		TRACE_COUNTER_SET(QueuedAnimationSequenceCompilation, GetNumRemainingAssets());
 		Notification.Update(GetNumRemainingAssets());
 	}
+
+	void FAnimSequenceCompilingManager::OnPostReachabilityAnalysis()
+	{
+		if (GetNumRemainingAssets())
+		{
+			TRACE_CPUPROFILER_EVENT_SCOPE(FNaniteDisplacedMeshCompilingManager::CancelUnreachableMeshes);
+
+			TArray<UAnimSequence*> PendingAnimationSequences;
+			PendingAnimationSequences.Reserve(GetNumRemainingAssets());
+
+			for (auto Iterator = RegisteredAnimSequences.CreateIterator(); Iterator; ++Iterator)
+			{
+				if (UAnimSequence* AnimSequence = Iterator->GetEvenIfUnreachable())
+				{
+					const bool bAnimSequenceUnreachable = AnimSequence->IsUnreachable();
+					const bool bNullSkeleton = AnimSequence->GetSkeleton() == nullptr;
+					const bool bSkeletonUnreachable = bNullSkeleton ? true : AnimSequence->GetSkeleton()->IsUnreachable();
+					
+					if(bAnimSequenceUnreachable||bNullSkeleton||bSkeletonUnreachable)
+					{
+						if (AnimSequence->TryCancelAsyncTasks())
+						{
+							Iterator.RemoveCurrent();
+						}
+						else
+						{
+							PendingAnimationSequences.Add(AnimSequence);
+						}
+					}
+				}
+			}
+
+			FinishCompilation(PendingAnimationSequences);
+		}
+	}
 }
 
 #undef LOCTEXT_NAMESPACE //"AnimSequenceCompilingManager"
diff --git a/Engine/Source/Runtime/Engine/Private/Animation/AnimationSequenceCompiler.h b/Engine/Source/Runtime/Engine/Private/Animation/AnimationSequenceCompiler.h
--- a/Engine/Source/Runtime/Engine/Private/Animation/AnimationSequenceCompiler.h
+++ b/Engine/Source/Runtime/Engine/Private/Animation/AnimationSequenceCompiler.h
@@ -8,6 +8,7 @@
 #include "AsyncCompilationHelpers.h"
 
 class UAnimSequence;
+class USkeleton;
 class FQueuedThreadPool;
 enum class EQueuedWorkPriority : uint8;
 
@@ -32,6 +33,7 @@ namespace UE::Anim
 
 		void AddAnimSequences(TArrayView<UAnimSequence* const> InAnimSequences);
 		void FinishCompilation(TArrayView<UAnimSequence* const> InAnimSequences);
+		void FinishCompilation(TArrayView<USkeleton* const> InSkeletons);
 
 	protected:
 		virtual void ProcessAsyncTasks(bool bLimitExecutionTime = false) override;
@@ -41,11 +43,14 @@ namespace UE::Anim
 		void ApplyCompilation(UAnimSequence* InAnimSequence);
 
 		void UpdateCompilationNotification();
+
+		void OnPostReachabilityAnalysis();
 	private:
 		friend class FAssetCompilingManager;
 	
 		TSet<TWeakObjectPtr<UAnimSequence>> RegisteredAnimSequences;
 		FAsyncCompilationNotification Notification;
+		FDelegateHandle PostReachabilityAnalysisHandle;
 	};
 }
 
diff --git a/Engine/Source/Runtime/Engine/Private/Animation/Skeleton.cpp b/Engine/Source/Runtime/Engine/Private/Animation/Skeleton.cpp
--- a/Engine/Source/Runtime/Engine/Private/Animation/Skeleton.cpp
+++ b/Engine/Source/Runtime/Engine/Private/Animation/Skeleton.cpp
@@ -5,6 +5,8 @@
 =============================================================================*/ 
 
 #include "Animation/Skeleton.h"
+
+#include "AnimationSequenceCompiler.h"
 #include "Animation/AnimData/AnimDataModel.h"
 #include "UObject/LinkerLoad.h"
 #include "Engine/AssetUserData.h"
@@ -313,6 +315,10 @@ USkeleton::USkeleton(const FObjectInitializer& ObjectInitializer)
 
 void USkeleton::BeginDestroy()
 {
+#if WITH_EDITOR
+	UE::Anim::FAnimSequenceCompilingManager::Get().FinishCompilation({this});
+#endif // WITH_EDITOR
+	
 	Super::BeginDestroy();
 	if (HasAnyFlags(RF_ClassDefaultObject))
 	{
diff --git a/Engine/Source/Runtime/Engine/Private/BoneContainer.cpp b/Engine/Source/Runtime/Engine/Private/BoneContainer.cpp
--- a/Engine/Source/Runtime/Engine/Private/BoneContainer.cpp
+++ b/Engine/Source/Runtime/Engine/Private/BoneContainer.cpp
@@ -116,7 +116,14 @@ void FBoneContainer::Initialize(const FCurveEvaluationOption& CurveEvalOption)
 {
 	LLM_SCOPE_BYNAME(TEXT("Animation/BoneContainer"));
 	RefSkeleton = nullptr;
-	UObject* AssetObj = Asset.Get();
+
+	UObject* AssetObj =
+#if WITH_EDITOR
+		Asset.GetEvenIfUnreachable();
+#else
+		Asset.Get();
+#endif
+	
 	USkeletalMesh* AssetSkeletalMeshObj = Cast<USkeletalMesh>(AssetObj);
 	USkeleton* AssetSkeletonObj = nullptr;
 
diff --git a/Engine/Source/Runtime/Engine/Public/BoneContainer.h b/Engine/Source/Runtime/Engine/Public/BoneContainer.h
--- a/Engine/Source/Runtime/Engine/Public/BoneContainer.h
+++ b/Engine/Source/Runtime/Engine/Public/BoneContainer.h
@@ -273,7 +273,12 @@ public:
 	/** Returns true if FBoneContainer is Valid. Needs an Asset, a RefPoseArray, and a RequiredBonesArray. */
 	const bool IsValid() const
 	{
-		return (Asset.IsValid() && (RefSkeleton != NULL) && (BoneIndicesArray.Num() > 0));
+		return (Asset.IsValid(
+#if WITH_EDITOR
+			true, true
+#endif
+			)
+				&& (RefSkeleton != nullptr) && (BoneIndicesArray.Num() > 0));
 	}
 
 	/** Get Asset this BoneContainer was made for. Typically a SkeletalMesh, but could also be a USkeleton. */
@@ -289,9 +294,13 @@ public:
 	}
 
 	/** Get Skeleton Asset. Could either be the SkeletalMesh's Skeleton, or the Skeleton this BoneContainer was made for. Is non NULL is BoneContainer is valid. */
-	USkeleton* GetSkeletonAsset() const
+	USkeleton* GetSkeletonAsset(bool bEvenIfUnreachable = false) const
 	{
-		return AssetSkeleton.Get();
+		return
+#if WITH_EDITOR
+		bEvenIfUnreachable ? AssetSkeleton.GetEvenIfUnreachable() : 
+#endif
+		AssetSkeleton.Get();
 	}
 
 	/** Disable Retargeting for debugging. */
