Strata - Fixed compilation failure when the an output FStrataData is not output 0 of a material function.

diff --git a/Engine/Source/Runtime/Engine/Private/Materials/HLSLMaterialTranslator.cpp b/Engine/Source/Runtime/Engine/Private/Materials/HLSLMaterialTranslator.cpp
--- a/Engine/Source/Runtime/Engine/Private/Materials/HLSLMaterialTranslator.cpp
+++ b/Engine/Source/Runtime/Engine/Private/Materials/HLSLMaterialTranslator.cpp
@@ -799,7 +799,7 @@ bool FHLSLMaterialTranslator::Translate()
 			
 				FExpressionInput* SurfaceThickness = Material->IsThinSurface() && Material->GetMaterialInterface() ? &Material->GetMaterialInterface()->GetMaterial()->GetEditorOnlyData()->SurfaceThickness : nullptr;
 				StrataThicknessStackPush(nullptr, SurfaceThickness);
-				FrontMaterialExpr->StrataGenerateMaterialTopologyTree(this, nullptr, 0);
+				FrontMaterialExpr->StrataGenerateMaterialTopologyTree(this, nullptr, FrontMaterialInput->OutputIndex);
 				StrataThicknessStackPop();
 
 				check(StrataThicknessStack.Num() == 0);
