Fix to the DeriveNormalZ material node that saturates the input to the square root to avoid negative values.

This is could be caused for instance by extreme normals combined with texture compression.

From UDN https://udn.unrealengine.com/s/question/0D54z00008anTs9CAE/normalderivez-returning-nans?fromCase=1

diff --git a/Engine/Source/Runtime/Engine/Private/Materials/MaterialExpressions.cpp b/Engine/Source/Runtime/Engine/Private/Materials/MaterialExpressions.cpp
--- a/Engine/Source/Runtime/Engine/Private/Materials/MaterialExpressions.cpp
+++ b/Engine/Source/Runtime/Engine/Private/Materials/MaterialExpressions.cpp
@@ -12813,7 +12813,8 @@ int32 UMaterialExpressionDeriveNormalZ::Compile(class FMaterialCompiler* Compile
 	int32 InputVector = Compiler->ForceCast(InXY.Compile(Compiler), MCT_Float2);
 	int32 DotResult = Compiler->Dot(InputVector, InputVector);
 	int32 InnerResult = Compiler->Sub(Compiler->Constant(1), DotResult);
-	int32 DerivedZ = Compiler->SquareRoot(InnerResult);
+	int32 SaturatedInnerResult = Compiler->Saturate(InnerResult);
+	int32 DerivedZ = Compiler->SquareRoot(SaturatedInnerResult);
 	int32 AppendedResult = Compiler->ForceCast(Compiler->AppendVector(InputVector, DerivedZ), MCT_Float3);
 
 	return AppendedResult;
