Removed the ResetToDefault button from the Parent Class blueprint property

diff --git a/Engine/Source/Runtime/Engine/Classes/Engine/Blueprint.h b/Engine/Source/Runtime/Engine/Classes/Engine/Blueprint.h
--- a/Engine/Source/Runtime/Engine/Classes/Engine/Blueprint.h
+++ b/Engine/Source/Runtime/Engine/Classes/Engine/Blueprint.h
@@ -472,7 +472,7 @@ class ENGINE_API UBlueprint : public UBlueprintCore, public IBlueprintPropertyGu
 	 * one such case can be created by creating a blueprint (A) based on another blueprint (B), shutting down the editor, and
 	 * deleting the parent blueprint. Exported as Alphabetical in GetAssetRegistryTags
 	 */
-	UPROPERTY()
+	UPROPERTY(meta=(NoResetToDefault))
 	TSubclassOf<UObject> ParentClass;
 
 	/** The type of this blueprint */
