Fix for index mismatch/out-of-bounds access when drawing bones of a pose watch that doesn't use the full set of bones.
* All parent bone indices are now copied instead of only the required bones.
* Now the correct bone index is returned when accessing the parent of a bone that exists after a bone that is not being drawn.

diff --git a/Engine/Source/Runtime/Engine/Classes/Engine/PoseWatchRenderData.h b/Engine/Source/Runtime/Engine/Classes/Engine/PoseWatchRenderData.h
--- a/Engine/Source/Runtime/Engine/Classes/Engine/PoseWatchRenderData.h
+++ b/Engine/Source/Runtime/Engine/Classes/Engine/PoseWatchRenderData.h
@@ -50,9 +50,9 @@ public:
 
 	FVector GetViewportOffset() const;
 
-	TArray<int32> GetViewportAllowList() const;
+	const TArray<int32>& GetViewportAllowList() const;
 
-	TArray<int32> GetParentIndices() const;
+	const TArray<int32>& GetParentIndices() const;
 
 private:
 	FTransform WorldTransform;
diff --git a/Engine/Source/Runtime/Engine/Private/Engine/PoseWatchRenderData.cpp b/Engine/Source/Runtime/Engine/Private/Engine/PoseWatchRenderData.cpp
--- a/Engine/Source/Runtime/Engine/Private/Engine/PoseWatchRenderData.cpp
+++ b/Engine/Source/Runtime/Engine/Private/Engine/PoseWatchRenderData.cpp
@@ -29,11 +29,15 @@ void FAnimNodePoseWatch::CopyPoseWatchData(const FReferenceSkeleton& RefSkeleton
 	ParentIndices.Empty();
 	ViewportMaskAllowedList.Empty();
 
+	ParentIndices.AddUninitialized(RefSkeleton.GetNum());
+	for (uint16 BoneIndex = 0; BoneIndex < RefSkeleton.GetNum(); ++BoneIndex)
+	{
+		ParentIndices[BoneIndex] = RefSkeleton.GetParentIndex(BoneIndex);
+	}
+
 	const UBlendProfile* Mask = PoseWatchPoseElement->ViewportMask;
 	for (const FBoneIndexType& BoneIndex : RequiredBones)
 	{
-		ParentIndices.Add(RefSkeleton.GetParentIndex(BoneIndex));
-
 		if (Mask)
 		{
 			const FName& BoneName = RefSkeleton.GetBoneName(BoneIndex);
@@ -79,12 +83,12 @@ FVector FAnimNodePoseWatch::GetViewportOffset() const
 	return ViewportOffset;
 }
 
-TArray<int32> FAnimNodePoseWatch::GetViewportAllowList() const
+const TArray<int32>& FAnimNodePoseWatch::GetViewportAllowList() const
 { 
 	return ViewportMaskAllowedList;
 }
 
-TArray<int32> FAnimNodePoseWatch::GetParentIndices() const
+const TArray<int32>& FAnimNodePoseWatch::GetParentIndices() const
 { 
 	return ParentIndices; 
 }
diff --git a/Engine/Source/Runtime/Engine/Private/SkeletalDebugRendering.cpp b/Engine/Source/Runtime/Engine/Private/SkeletalDebugRendering.cpp
--- a/Engine/Source/Runtime/Engine/Private/SkeletalDebugRendering.cpp
+++ b/Engine/Source/Runtime/Engine/Private/SkeletalDebugRendering.cpp
@@ -186,30 +186,28 @@ void DrawBonesFromPoseWatch(
 	const FTransform WorldTransform = bUseWorldTransform ? PoseWatch.GetWorldTransform() : FTransform::Identity;
 	const FVector RelativeOffset = WorldTransform.GetRotation().RotateVector(PoseWatch.GetViewportOffset());
 
+	const TArray<int32>& ViewportMaskAllowList = PoseWatch.GetViewportAllowList();
+	const TArray<int32>& ParentIndices = PoseWatch.GetParentIndices();
+
 	TArray<FTransform> UseWorldTransforms;
-	UseWorldTransforms.AddDefaulted(InRequiredBones.Num());
+	UseWorldTransforms.AddDefaulted(ParentIndices.Num());
 	
 	TArray<FBoneIndexType> UseRequiredBones;
 	UseRequiredBones.Reserve(InRequiredBones.Num());
 
-	TArray<int32> ViewportMaskAllowList = PoseWatch.GetViewportAllowList();
-	TArray<int32> ParentIndices = PoseWatch.GetParentIndices();
-	
-	for (int32 Index = 0; Index < InRequiredBones.Num(); ++Index)
+	for (const FBoneIndexType& BoneIndex : InRequiredBones)
 	{
-		const FBoneIndexType& BoneIndex = InRequiredBones[Index];
-
 		check(ParentIndices.IsValidIndex(BoneIndex));
 		const int32 ParentIndex = ParentIndices[BoneIndex];
 
 		if (ParentIndex == INDEX_NONE)
 		{
-			UseWorldTransforms[BoneIndex] = InBoneTransforms[Index] * WorldTransform;
+			UseWorldTransforms[BoneIndex] = InBoneTransforms[BoneIndex] * WorldTransform;
 			UseWorldTransforms[BoneIndex].AddToTranslation(RelativeOffset);
 		}
 		else
 		{
-			UseWorldTransforms[BoneIndex] = InBoneTransforms[Index] * UseWorldTransforms[ParentIndex];
+			UseWorldTransforms[BoneIndex] = InBoneTransforms[BoneIndex] * UseWorldTransforms[ParentIndex];
 		}
 
 		if (!ViewportMaskAllowList.Contains(BoneIndex))
