Crash fix for cleaning up invalid Blend Profile entries leading to OOB access when more than one entry is removed at once

diff --git a/Engine/Source/Runtime/Engine/Classes/Animation/BlendProfile.h b/Engine/Source/Runtime/Engine/Classes/Animation/BlendProfile.h
--- a/Engine/Source/Runtime/Engine/Classes/Animation/BlendProfile.h
+++ b/Engine/Source/Runtime/Engine/Classes/Animation/BlendProfile.h
@@ -95,6 +95,9 @@ public:
 	/** Ensures the bone names match the skeleton indices by using the bone name as our lookup key. */
 	void RefreshBoneEntriesFromName();
 
+	/** Removes entries with bone references to invalid bones */
+	void CleanupBoneEntries();
+
 	/**
 	 * Get the bone entry by entry index.
 	 * @param[in] InEntryIdx The index to the bone entry in range [0, GetNumBlendEntries()-1].
diff --git a/Engine/Source/Runtime/Engine/Private/Animation/BlendProfile.cpp b/Engine/Source/Runtime/Engine/Private/Animation/BlendProfile.cpp
--- a/Engine/Source/Runtime/Engine/Private/Animation/BlendProfile.cpp
+++ b/Engine/Source/Runtime/Engine/Private/Animation/BlendProfile.cpp
@@ -74,6 +74,14 @@ void UBlendProfile::RefreshBoneEntriesFromName()
 	}
 }
 
+void UBlendProfile::CleanupBoneEntries()
+{
+	ProfileEntries.RemoveAll([](const FBlendProfileBoneEntry& Entry)
+		{
+			return !Entry.BoneReference.HasValidSetup();
+		});
+}
+
 const FBlendProfileBoneEntry& UBlendProfile::GetEntry(const int32 InEntryIdx) const
 {
 	return ProfileEntries[InEntryIdx];
diff --git a/Engine/Source/Runtime/Engine/Private/Animation/Skeleton.cpp b/Engine/Source/Runtime/Engine/Private/Animation/Skeleton.cpp
--- a/Engine/Source/Runtime/Engine/Private/Animation/Skeleton.cpp
+++ b/Engine/Source/Runtime/Engine/Private/Animation/Skeleton.cpp
@@ -1391,14 +1391,8 @@ void USkeleton::HandleSkeletonHierarchyChange()
 	// Remove entries from Blend Profiles for bones that no longer exists
 	for (UBlendProfile* Profile : BlendProfiles)
 	{
-		for (int32 EntryIndex = Profile->ProfileEntries.Num() - 1; EntryIndex >= 0; --EntryIndex)
-		{
-			const FBlendProfileBoneEntry& Entry = Profile->ProfileEntries[EntryIndex];
-			if (ReferenceSkeleton.FindBoneIndex(Entry.BoneReference.BoneName) == INDEX_NONE)
-			{
-				Profile->RemoveEntry(Entry.BoneReference.BoneIndex);
-			}
-		}
+		Profile->RefreshBoneEntriesFromName();
+		Profile->CleanupBoneEntries();
 	}
 
 	OnSkeletonHierarchyChanged.Broadcast();
