DataDrivenShaderPlatform Info Material Node

diff --git a/Engine/Source/Runtime/Engine/Classes/Materials/MaterialExpressionDataDrivenShaderPlatformInfoSwitch.h b/Engine/Source/Runtime/Engine/Classes/Materials/MaterialExpressionDataDrivenShaderPlatformInfoSwitch.h
new file mode 100644
--- /dev/null
+++ b/Engine/Source/Runtime/Engine/Classes/Materials/MaterialExpressionDataDrivenShaderPlatformInfoSwitch.h
@@ -0,0 +1,72 @@
+// Copyright Epic Games, Inc. All Rights Reserved.
+
+
+#pragma once
+
+#include "CoreMinimal.h"
+#include "UObject/ObjectMacros.h"
+#include "MaterialExpressionIO.h"
+#include "MaterialValueType.h"
+#include "Materials/MaterialExpression.h"
+#include "RHIFeatureLevel.h"
+#include "MaterialExpressionDataDrivenShaderPlatformInfoSwitch.generated.h"
+
+UENUM()
+enum EDataDrivenShaderPlatformInfoCondition : int
+{
+	COND_True UMETA(DisplayName = "Property must be true"),
+	COND_False UMETA(DisplayName = "Property must be false"),
+};
+
+USTRUCT()
+struct FDataDrivenShaderPlatformInfoInput
+{
+	GENERATED_USTRUCT_BODY()
+
+		UPROPERTY(EditAnywhere, Category = FDataDrivenShaderPlatformInfoInput, meta = (GetOptions = "GetNameOptions"))
+		FName InputName;
+
+		UPROPERTY(EditAnywhere, Category = FDataDrivenShaderPlatformInfoInput)
+		TEnumAsByte<EDataDrivenShaderPlatformInfoCondition> PropertyCondition = EDataDrivenShaderPlatformInfoCondition::COND_True;
+};
+
+UCLASS(collapsecategories, hidecategories = Object, MinimalAPI)
+class UMaterialExpressionDataDrivenShaderPlatformInfoSwitch : public UMaterialExpression
+{
+	GENERATED_UCLASS_BODY()
+		UPROPERTY()
+		FExpressionInput InputTrue;
+		
+		UPROPERTY()
+		FExpressionInput InputFalse;
+
+		UPROPERTY(EditAnywhere, Category = DataDrivenShaderPlatformInfoInput)
+		TArray<FDataDrivenShaderPlatformInfoInput> DDSPIPropertyNames;
+
+		UFUNCTION()
+		TArray<FString> GetNameOptions() const;
+		
+		UPROPERTY()
+		uint32 bContainsInvalidProperty : 1;
+
+		//~ Begin UObject Interface
+		virtual void PostLoad() override;
+#if WITH_EDITOR
+		virtual void PostEditChangeProperty(FPropertyChangedEvent& PropertyChangedEvent) override;
+#endif
+		//~ End UObject Interface
+		// 
+	//~ Begin UMaterialExpression Interface
+#if WITH_EDITOR
+	virtual int32 Compile(class FMaterialCompiler* Compiler, int32 OutputIndex) override;
+	virtual void GetCaption(TArray<FString>& OutCaptions) const override;
+	virtual bool IsInputConnectionRequired(int32 InputIndex) const override;
+	virtual const TArray<FExpressionInput*> GetInputs() override;
+	virtual FExpressionInput* GetInput(int32 InputIndex) override;
+	virtual FName GetInputName(int32 InputIndex) const override;
+	virtual bool IsResultMaterialAttributes(int32 OutputIndex) override;
+	virtual uint32 GetInputType(int32 InputIndex) override { return MCT_Unknown; }
+	virtual uint32 GetOutputType(int32 OutputIndex) override { return MCT_Unknown; }
+#endif // WITH_EDITOR
+	//~ End UMaterialExpression Interface
+};
diff --git a/Engine/Source/Runtime/Engine/Private/Materials/MaterialExpressions.cpp b/Engine/Source/Runtime/Engine/Private/Materials/MaterialExpressions.cpp
--- a/Engine/Source/Runtime/Engine/Private/Materials/MaterialExpressions.cpp
+++ b/Engine/Source/Runtime/Engine/Private/Materials/MaterialExpressions.cpp
@@ -98,6 +98,7 @@
 #include "Materials/MaterialExpressionEyeAdaptation.h"
 #include "Materials/MaterialExpressionEyeAdaptationInverse.h"
 #include "Materials/MaterialExpressionFeatureLevelSwitch.h"
+#include "Materials/MaterialExpressionDataDrivenShaderPlatformInfoSwitch.h"
 #include "Materials/MaterialExpressionFloor.h"	
 #include "Materials/MaterialExpressionFmod.h"
 #include "Materials/MaterialExpressionFontSample.h"
@@ -9653,6 +9654,222 @@ void UMaterialExpressionFeatureLevelSwitch::Serialize(FStructuredArchive::FRecor
 	}
 }
 
+//
+//	UMaterialExpressionDataDrivenShaderPlatformInfoSwitch
+//
+
+UMaterialExpressionDataDrivenShaderPlatformInfoSwitch::UMaterialExpressionDataDrivenShaderPlatformInfoSwitch(const FObjectInitializer& ObjectInitializer)
+: Super(ObjectInitializer)
+{
+#if WITH_EDITORONLY_DATA
+	// Structure to hold one-time initialization
+	struct FConstructorStatics
+	{
+		FText NAME_Utility;
+		FConstructorStatics()
+			: NAME_Utility(LOCTEXT("Utility", "Utility"))
+		{
+		}
+	};
+	static FConstructorStatics ConstructorStatics;
+
+	MenuCategories.Add(ConstructorStatics.NAME_Utility);
+	bCollapsed = false;
+	bContainsInvalidProperty = false;
+#endif // WITH_EDITORONLY_DATA
+}
+
+TArray<FString> UMaterialExpressionDataDrivenShaderPlatformInfoSwitch::GetNameOptions() const
+{
+	TArray<FString> Output;
+#if WITH_EDITOR
+	for (const auto& DDSPINames : FGenericDataDrivenShaderPlatformInfo::PropertyToShaderPlatformFunctionMap)
+	{
+		Output.Add(DDSPINames.Key);
+	}
+#endif
+	return Output;
+}
+
+#if WITH_EDITOR
+bool IsDataDrivenShaderPlatformInfoSwitchValid(TArray<FDataDrivenShaderPlatformInfoInput>& DDSPIPropertyNames, const UMaterial* Material)
+{
+	for (const FDataDrivenShaderPlatformInfoInput& DDSPIInput : DDSPIPropertyNames)
+	{
+		if (DDSPIInput.InputName == NAME_None)
+		{
+			continue;
+		}
+
+		bool PropertyExists = FGenericDataDrivenShaderPlatformInfo::PropertyToShaderPlatformFunctionMap.Find(DDSPIInput.InputName.ToString()) != nullptr;
+		if (!PropertyExists)
+		{
+			return true;
+		}
+	}
+
+	return false;
+}
+#endif
+
+void UMaterialExpressionDataDrivenShaderPlatformInfoSwitch::PostLoad()
+{
+	Super::PostLoad();
+#if WITH_EDITOR
+	bContainsInvalidProperty = IsDataDrivenShaderPlatformInfoSwitchValid(DDSPIPropertyNames, Material);
+#endif
+}
+
+#if WITH_EDITOR
+
+int32 UMaterialExpressionDataDrivenShaderPlatformInfoSwitch::Compile(class FMaterialCompiler* Compiler, int32 OutputIndex)
+{
+	if (bContainsInvalidProperty || DDSPIPropertyNames.IsEmpty())
+	{
+		return CompilerError(Compiler, *FString::Printf(TEXT("%s is using a DataDrivenShaderPlatformInfoSwitch whose condition is invalid. Default material will be used until this is fixed"), Material ? *(Material->GetName()) : TEXT("Unknown")));
+	}
+
+	const EShaderPlatform ShaderPlatform = Compiler->GetShaderPlatform();
+	check(FDataDrivenShaderPlatformInfo::IsValid(ShaderPlatform));
+
+	bool bAllNamesAreNone = true;
+	bool bCheck = true;
+	for (const FDataDrivenShaderPlatformInfoInput& DDSPIInput : DDSPIPropertyNames)
+	{
+		if (DDSPIInput.InputName == NAME_None)
+		{
+			continue;
+		}
+
+		bAllNamesAreNone = false;
+		bool bCheckProperty = FGenericDataDrivenShaderPlatformInfo::PropertyToShaderPlatformFunctionMap[DDSPIInput.InputName.ToString()](ShaderPlatform);
+		if (DDSPIInput.PropertyCondition == EDataDrivenShaderPlatformInfoCondition::COND_True)
+		{
+			bCheck = bCheck && bCheckProperty;
+		}
+		else
+		{
+			bCheck = bCheck && !bCheckProperty;
+		}
+	}
+
+	if (bAllNamesAreNone)
+	{
+		return CompilerError(Compiler, *FString::Printf(TEXT("%s is using a DataDrivenShaderPlatformInfoSwitch whose condition is empty. Default material will be used until this is fixed"), Material ? *(Material->GetName()) : TEXT("Unknown")));
+	}
+	else if (bCheck)
+	{
+		return InputTrue.Compile(Compiler);
+	}
+	else
+	{
+		return InputFalse.Compile(Compiler);
+	}
+}
+
+void UMaterialExpressionDataDrivenShaderPlatformInfoSwitch::GetCaption(TArray<FString>& OutCaptions) const
+{
+	OutCaptions.Add(FString(TEXT("DataDrivenShaderPlatformInfo Switch")));
+}
+
+bool UMaterialExpressionDataDrivenShaderPlatformInfoSwitch::IsInputConnectionRequired(int32 InputIndex) const
+{
+	return true;
+}
+
+const TArray<FExpressionInput*> UMaterialExpressionDataDrivenShaderPlatformInfoSwitch::GetInputs()
+{
+	TArray<FExpressionInput*> Result;
+	Result.Add(&InputTrue);
+	Result.Add(&InputFalse);
+	return Result;
+}
+
+FExpressionInput* UMaterialExpressionDataDrivenShaderPlatformInfoSwitch::GetInput(int32 InputIndex)
+{
+	if (InputIndex == 0)
+	{
+		return &InputTrue;
+	}
+	else if (InputIndex == 1)
+	{
+		return &InputFalse;
+	}
+	return nullptr;
+}
+
+FName UMaterialExpressionDataDrivenShaderPlatformInfoSwitch::GetInputName(int32 InputIndex) const
+{
+	TStringBuilder<128> Condition;
+	bool bIsFirst = true;
+	for (const FDataDrivenShaderPlatformInfoInput& DDSPIInput : DDSPIPropertyNames)
+	{
+		if (DDSPIInput.InputName == NAME_None)
+		{
+			continue;
+		}
+
+		if (!bIsFirst)
+		{
+			Condition.Append(TEXT(" && "));
+		}
+
+		if (DDSPIInput.PropertyCondition == EDataDrivenShaderPlatformInfoCondition::COND_False)
+		{
+			Condition.Append(TEXT("!"));
+		}
+
+		Condition.Append(DDSPIInput.InputName.ToString());
+		bIsFirst = false;
+	}
+
+	const FString ConditionString = Condition.ToString();
+	FString NegateConditionString = TEXT("!(") + ConditionString + TEXT(")");
+
+	if (InputIndex == 0)
+	{
+		return *ConditionString;
+	}
+	else if (InputIndex == 1)
+	{
+		return *NegateConditionString;
+	}
+	return NAME_None;
+}
+
+bool UMaterialExpressionDataDrivenShaderPlatformInfoSwitch::IsResultMaterialAttributes(int32 OutputIndex)
+{
+	check(OutputIndex == 0);
+	TArray<FExpressionInput*> ExpressionInputs = GetInputs();
+
+	for (FExpressionInput* ExpressionInput : ExpressionInputs)
+	{
+		// If there is a loop anywhere in this expression's inputs then we can't risk checking them
+		if (ExpressionInput->Expression && !ExpressionInput->Expression->ContainsInputLoop() && ExpressionInput->Expression->IsResultMaterialAttributes(ExpressionInput->OutputIndex))
+		{
+			return true;
+		}
+	}
+
+	return false;
+}
+
+void UMaterialExpressionDataDrivenShaderPlatformInfoSwitch::PostEditChangeProperty(FPropertyChangedEvent& PropertyChangedEvent)
+{
+	if (PropertyChangedEvent.MemberProperty && GraphNode)
+	{
+		const FName PropertyName = PropertyChangedEvent.MemberProperty->GetFName();
+		if (PropertyName == GET_MEMBER_NAME_CHECKED(UMaterialExpressionDataDrivenShaderPlatformInfoSwitch, DDSPIPropertyNames))
+		{
+			bContainsInvalidProperty = IsDataDrivenShaderPlatformInfoSwitchValid(DDSPIPropertyNames, Material);
+			GraphNode->ReconstructNode();
+		}
+	}
+
+	Super::PostEditChangeProperty(PropertyChangedEvent);
+}
+#endif // WITH_EDITOR
+
 //
 //	UMaterialExpressionShadingPathSwitch
 //
diff --git a/Engine/Source/Runtime/RHI/Private/DataDrivenShaderPlatformInfo.cpp b/Engine/Source/Runtime/RHI/Private/DataDrivenShaderPlatformInfo.cpp
--- a/Engine/Source/Runtime/RHI/Private/DataDrivenShaderPlatformInfo.cpp
+++ b/Engine/Source/Runtime/RHI/Private/DataDrivenShaderPlatformInfo.cpp
@@ -14,6 +14,9 @@ const FName LANGUAGE_Nintendo("Nintendo");
 
 FGenericDataDrivenShaderPlatformInfo FGenericDataDrivenShaderPlatformInfo::Infos[SP_NumPlatforms];
 
+#if WITH_EDITOR
+TMap < FString, TFunction<bool(const FStaticShaderPlatform Platform)>> FGenericDataDrivenShaderPlatformInfo::PropertyToShaderPlatformFunctionMap;
+#endif
 TMap<FName, EShaderPlatform> PlatformNameToShaderPlatformMap;
 
 EShaderPlatform ParseShaderPlatform(const TCHAR* String)
@@ -147,8 +150,16 @@ void FGenericDataDrivenShaderPlatformInfo::ParseDataDrivenShaderInfo(const FConf
 	ShaderPropertiesString += TEXT(#SettingName); \
 	ShaderPropertiesString += FString::Printf(TEXT("_%d"), SettingValue);
 
+#if WITH_EDITOR
+	#define ADD_PROPERTY_TO_SHADERPLATFORM_FUNCTIONMAP(SettingName) \
+		FGenericDataDrivenShaderPlatformInfo::PropertyToShaderPlatformFunctionMap.FindOrAdd(#SettingName) = [](const FStaticShaderPlatform Platform)->bool{ return Infos[Platform].SettingName;};
+#else
+	#define ADD_PROPERTY_TO_SHADERPLATFORM_FUNCTIONMAP(SettingName)
+#endif
+
 #define GET_SECTION_BOOL_HELPER(SettingName)	\
 	Info.SettingName = GetSectionBool(Section, #SettingName, Info.SettingName);	\
+	ADD_PROPERTY_TO_SHADERPLATFORM_FUNCTIONMAP(SettingName)	\
 	ADD_TO_PROPERTIES_STRING(SettingName, Info.SettingName)
 
 #define GET_SECTION_INT_HELPER(SettingName)	\
@@ -255,6 +266,7 @@ void FGenericDataDrivenShaderPlatformInfo::ParseDataDrivenShaderInfo(const FConf
 #undef GET_SECTION_INT_HELPER
 #undef GET_SECTION_SUPPORT_HELPER
 #undef ADD_TO_PROPERTIES_STRING
+#undef ADD_PROPERTY_TO_SHADERPLATFORM_FUNCTIONMAP
 
 	Info.ShaderPropertiesHash = GetTypeHash(ShaderPropertiesString);
 #if WITH_EDITOR
diff --git a/Engine/Source/Runtime/RHI/Public/DataDrivenShaderPlatformInfo.h b/Engine/Source/Runtime/RHI/Public/DataDrivenShaderPlatformInfo.h
--- a/Engine/Source/Runtime/RHI/Public/DataDrivenShaderPlatformInfo.h
+++ b/Engine/Source/Runtime/RHI/Public/DataDrivenShaderPlatformInfo.h
@@ -750,6 +750,11 @@ private:
 	static FGenericDataDrivenShaderPlatformInfo Infos[SP_NumPlatforms];
 
 public:
+
+#if WITH_EDITOR
+	static TMap < FString, TFunction<bool(const FStaticShaderPlatform Platform)>> PropertyToShaderPlatformFunctionMap;
+#endif
+
 	static bool IsValid(const FStaticShaderPlatform Platform)
 	{
 		return Infos[Platform].bContainsValidPlatformInfo;
