Fix race condition where clients would be sent new actors before seamless travel finished
It was only clearing the client's current map on an exact match which would fail for PIE, and it was failing to copy the travel count when the player controller was recreated on travel

diff --git a/Engine/Source/Runtime/Engine/Private/PlayerController.cpp b/Engine/Source/Runtime/Engine/Private/PlayerController.cpp
--- a/Engine/Source/Runtime/Engine/Private/PlayerController.cpp
+++ b/Engine/Source/Runtime/Engine/Private/PlayerController.cpp
@@ -3585,6 +3585,10 @@ void APlayerController::SeamlessTravelFrom(APlayerController* OldPC)
 		OldPC->PlayerState->Destroy();
 		OldPC->PlayerState = NULL;
 	}
+
+	// Copy seamless travel state
+	SeamlessTravelCount = OldPC->SeamlessTravelCount;
+	LastCompletedSeamlessTravelCount = OldPC->LastCompletedSeamlessTravelCount;
 }
 
 void APlayerController::PostSeamlessTravel()
diff --git a/Engine/Source/Runtime/Engine/Private/World.cpp b/Engine/Source/Runtime/Engine/Private/World.cpp
--- a/Engine/Source/Runtime/Engine/Private/World.cpp
+++ b/Engine/Source/Runtime/Engine/Private/World.cpp
@@ -6803,20 +6803,17 @@ bool FSeamlessTravelHandler::StartTravel(UWorld* InCurrentWorld, const FURL& InU
 			}
 			else
 			{
-				if (CurrentMapName == DestinationMapName)
+				UNetDriver* const NetDriver = CurrentWorld->GetNetDriver();
+				if (NetDriver)
 				{
-					UNetDriver* const NetDriver = CurrentWorld->GetNetDriver();
-					if (NetDriver)
+					for (int32 ClientIdx = 0; ClientIdx < NetDriver->ClientConnections.Num(); ClientIdx++)
 					{
-						for (int32 ClientIdx = 0; ClientIdx < NetDriver->ClientConnections.Num(); ClientIdx++)
+						UNetConnection* Connection = NetDriver->ClientConnections[ClientIdx];
+						if (Connection)
 						{
-							UNetConnection* Connection = NetDriver->ClientConnections[ClientIdx];
-							if (Connection)
-							{
-								// Empty the current map name in case we are going A -> transition -> A and the server loads fast enough
-								// that the clients are not on the transition map yet causing the server to think its loaded
-								Connection->SetClientWorldPackageName(NAME_None);
-							}
+							// Empty the current map name on all transitions because the server could try to spawn actors 
+							// before the client starts the transfer causing the server to think its loaded
+							Connection->SetClientWorldPackageName(NAME_None);
 						}
 					}
 				}
