Dependency cleanup around RenderCore.h includes.

diff --git a/Engine/Source/Editor/UnrealEd/Private/UnrealEdPrivatePCH.h b/Engine/Source/Editor/UnrealEd/Private/UnrealEdPrivatePCH.h
--- a/Engine/Source/Editor/UnrealEd/Private/UnrealEdPrivatePCH.h
+++ b/Engine/Source/Editor/UnrealEd/Private/UnrealEdPrivatePCH.h
@@ -432,7 +432,6 @@
 // From RenderCore:
 #include "RenderCommandFence.h"
 #include "RenderResource.h"
-#include "RenderCore.h"
 #include "RenderingThread.h"
 #include "UniformBuffer.h"
 #include "PackedNormal.h"
diff --git a/Engine/Source/Editor/UnrealEd/Public/UnrealEdSharedPCH.h b/Engine/Source/Editor/UnrealEd/Public/UnrealEdSharedPCH.h
--- a/Engine/Source/Editor/UnrealEd/Public/UnrealEdSharedPCH.h
+++ b/Engine/Source/Editor/UnrealEd/Public/UnrealEdSharedPCH.h
@@ -430,7 +430,9 @@
 // From RenderCore:
 #include "RenderCommandFence.h"
 #include "RenderResource.h"
+#if UE_ENABLE_INCLUDE_ORDER_DEPRECATED_IN_5_2
 #include "RenderCore.h"
+#endif
 #include "RenderingThread.h"
 #include "UniformBuffer.h"
 #include "PackedNormal.h"
diff --git a/Engine/Source/Runtime/Engine/Private/EnginePrivatePCH.h b/Engine/Source/Runtime/Engine/Private/EnginePrivatePCH.h
--- a/Engine/Source/Runtime/Engine/Private/EnginePrivatePCH.h
+++ b/Engine/Source/Runtime/Engine/Private/EnginePrivatePCH.h
@@ -355,7 +355,6 @@
 // From RenderCore:
 #include "RenderCommandFence.h"
 #include "RenderResource.h"
-#include "RenderCore.h"
 #include "RenderingThread.h"
 #include "UniformBuffer.h"
 #include "PackedNormal.h"
diff --git a/Engine/Source/Runtime/Engine/Private/Materials/MaterialRenderProxy.cpp b/Engine/Source/Runtime/Engine/Private/Materials/MaterialRenderProxy.cpp
--- a/Engine/Source/Runtime/Engine/Private/Materials/MaterialRenderProxy.cpp
+++ b/Engine/Source/Runtime/Engine/Private/Materials/MaterialRenderProxy.cpp
@@ -693,6 +693,11 @@ void FMaterialRenderProxy::UpdateDeferredCachedUniformExpressions()
 	DeferredUniformExpressionCacheRequests.Reset();
 }
 
+bool FMaterialRenderProxy::HasDeferredUniformExpressionCacheRequests()
+{
+	return DeferredUniformExpressionCacheRequests.Num() > 0;
+}
+
 #if WITH_EDITOR
 TSet<FMaterialRenderProxy*> FMaterialRenderProxy::MaterialRenderProxyMap;
 FCriticalSection FMaterialRenderProxy::MaterialRenderProxyMapLock;
diff --git a/Engine/Source/Runtime/Engine/Public/EngineSharedPCH.h b/Engine/Source/Runtime/Engine/Public/EngineSharedPCH.h
--- a/Engine/Source/Runtime/Engine/Public/EngineSharedPCH.h
+++ b/Engine/Source/Runtime/Engine/Public/EngineSharedPCH.h
@@ -436,7 +436,9 @@
 // From RenderCore:
 #include "RenderCommandFence.h"
 #include "RenderResource.h"
+#if UE_ENABLE_INCLUDE_ORDER_DEPRECATED_IN_5_2
 #include "RenderCore.h"
+#endif
 #include "RenderingThread.h"
 #include "UniformBuffer.h"
 #include "PackedNormal.h"
diff --git a/Engine/Source/Runtime/Engine/Public/Materials/MaterialRenderProxy.h b/Engine/Source/Runtime/Engine/Public/Materials/MaterialRenderProxy.h
--- a/Engine/Source/Runtime/Engine/Public/Materials/MaterialRenderProxy.h
+++ b/Engine/Source/Runtime/Engine/Public/Materials/MaterialRenderProxy.h
@@ -3,6 +3,8 @@
 #pragma once
 
 #include "CoreTypes.h"
+#include "Math/Color.h"
+#include "Math/Vector2D.h"
 #include "RenderResource.h"
 #include "RHIImmutableSamplerState.h"
 
@@ -206,10 +208,7 @@ public:
 
 	static void UpdateDeferredCachedUniformExpressions();
 
-	static inline bool HasDeferredUniformExpressionCacheRequests()
-	{
-		return DeferredUniformExpressionCacheRequests.Num() > 0;
-	}
+	static bool HasDeferredUniformExpressionCacheRequests();
 
 	int32 GetExpressionCacheSerialNumber() const { return UniformExpressionCacheSerialNumber; }
 
diff --git a/Engine/Source/Runtime/Engine/Public/Rendering/MorphTargetVertexInfoBuffers.h b/Engine/Source/Runtime/Engine/Public/Rendering/MorphTargetVertexInfoBuffers.h
--- a/Engine/Source/Runtime/Engine/Public/Rendering/MorphTargetVertexInfoBuffers.h
+++ b/Engine/Source/Runtime/Engine/Public/Rendering/MorphTargetVertexInfoBuffers.h
@@ -4,6 +4,7 @@
 
 #include "RenderResource.h"
 #include "Containers/DynamicRHIResourceArray.h"
+#include "Math/Vector4.h"
 
 struct FSkelMeshRenderSection;
 class UMorphTarget;
diff --git a/Engine/Source/Runtime/RenderCore/Private/GlobalRenderResources.cpp b/Engine/Source/Runtime/RenderCore/Private/GlobalRenderResources.cpp
--- a/Engine/Source/Runtime/RenderCore/Private/GlobalRenderResources.cpp
+++ b/Engine/Source/Runtime/RenderCore/Private/GlobalRenderResources.cpp
@@ -3,6 +3,7 @@
 #include "GlobalRenderResources.h"
 #include "RenderGraphUtils.h"
 #include "Containers/ResourceArray.h"
+#include "RenderCore.h"
 
 // The maximum number of transient vertex buffer bytes to allocate before we start panic logging who is doing the allocations
 int32 GMaxVertexBytesAllocatedPerFrame = 32 * 1024 * 1024;
diff --git a/Engine/Source/Runtime/RenderCore/Private/RenderResource.cpp b/Engine/Source/Runtime/RenderCore/Private/RenderResource.cpp
--- a/Engine/Source/Runtime/RenderCore/Private/RenderResource.cpp
+++ b/Engine/Source/Runtime/RenderCore/Private/RenderResource.cpp
@@ -13,6 +13,7 @@
 #include "RayTracingGeometryManager.h"
 #include "RenderGraphResources.h"
 #include "Containers/ResourceArray.h"
+#include "RenderCore.h"
 
 /** Whether to enable mip-level fading or not: +1.0f if enabled, -1.0f if disabled. */
 float GEnableMipLevelFading = 1.0f;
diff --git a/Engine/Source/Runtime/RenderCore/Private/RenderUtils.cpp b/Engine/Source/Runtime/RenderCore/Private/RenderUtils.cpp
--- a/Engine/Source/Runtime/RenderCore/Private/RenderUtils.cpp
+++ b/Engine/Source/Runtime/RenderCore/Private/RenderUtils.cpp
@@ -14,6 +14,7 @@
 #include "Shader.h"
 #include "ShaderPlatformCachedIniValue.h"
 #include "DataDrivenShaderPlatformInfo.h"
+#include "RenderCore.h"
 
 #if WITH_EDITOR
 #include "Interfaces/ITargetPlatform.h"
diff --git a/Engine/Source/Runtime/RenderCore/Public/RenderCore.h b/Engine/Source/Runtime/RenderCore/Public/RenderCore.h
--- a/Engine/Source/Runtime/RenderCore/Public/RenderCore.h
+++ b/Engine/Source/Runtime/RenderCore/Public/RenderCore.h
@@ -6,6 +6,10 @@
 
 #pragma once
 
+#include "HAL/Platform.h"
+#include "Stats/Stats2.h"
+
+#if UE_ENABLE_INCLUDE_ORDER_DEPRECATED_IN_5_2
 #include "CoreMinimal.h"
 #include "HAL/PlatformTime.h"
 #include "Logging/LogMacros.h"
@@ -15,10 +19,7 @@
 #include "Math/Vector4.h"
 #include "RHIDefinitions.h"
 #include "Stats/Stats.h"
-#include "Stats/Stats2.h"
 #include "RenderTimer.h"
-
-#if UE_ENABLE_INCLUDE_ORDER_DEPRECATED_IN_5_2
 #include "HDRHelper.h"
 #endif
 
