Rename variable

#rnx

diff --git a/Engine/Source/Runtime/Engine/Private/Streaming/AsyncTextureStreaming.cpp b/Engine/Source/Runtime/Engine/Private/Streaming/AsyncTextureStreaming.cpp
--- a/Engine/Source/Runtime/Engine/Private/Streaming/AsyncTextureStreaming.cpp
+++ b/Engine/Source/Runtime/Engine/Private/Streaming/AsyncTextureStreaming.cpp
@@ -220,13 +220,13 @@ void FAsyncRenderAssetStreamingData::UpdatePerfectWantedMips_Async(FStreamingRen
 			&& MaxSize != FLT_MAX
 			&& MaxSize_VisibleOnly != FLT_MAX)
 		{
-			const float CumBoostFactor = StreamingRenderAsset.BoostFactor * StreamingRenderAsset.DynamicBoostFactor;
+			const float FinalBoostFactor = StreamingRenderAsset.BoostFactor * StreamingRenderAsset.DynamicBoostFactor;
 
 			// If there is not enough resolution in the texture to fix the required quality, save this information to prevent degrading this texture before other ones.
-			bLooksLowRes = FMath::Max3(MaxSize_VisibleOnly, MaxSize, MaxAllowedSize) / MaxAllowedSize >= CumBoostFactor * 2.f;
+			bLooksLowRes = FMath::Max3(MaxSize_VisibleOnly, MaxSize, MaxAllowedSize) / MaxAllowedSize >= FinalBoostFactor * 2.f;
 
-			MaxSize *=  CumBoostFactor;
-			MaxSize_VisibleOnly *= CumBoostFactor;
+			MaxSize *= FinalBoostFactor;
+			MaxSize_VisibleOnly *= FinalBoostFactor;
 		}
 
 		// Last part checks that it has been used since the last reference was removed.
