Fix resolve conflict

#rnx

diff --git a/Engine/Source/Runtime/Engine/Private/LocalVertexFactory.cpp b/Engine/Source/Runtime/Engine/Private/LocalVertexFactory.cpp
--- a/Engine/Source/Runtime/Engine/Private/LocalVertexFactory.cpp
+++ b/Engine/Source/Runtime/Engine/Private/LocalVertexFactory.cpp
@@ -463,6 +463,8 @@ void FLocalVertexFactory::InitRHI()
 	}
 
 	FLocalVertexFactoryLooseParameters LooseParameters;
+	LooseParameters.FrameNumber = -1;
+	LooseParameters.GPUSkinPassThroughPositionBuffer = GNullVertexBuffer.VertexBufferSRV;
 	LooseParameters.GPUSkinPassThroughPreviousPositionBuffer = GNullVertexBuffer.VertexBufferSRV;
 	LooseParametersUniformBuffer = TUniformBufferRef<FLocalVertexFactoryLooseParameters>::CreateUniformBufferImmediate(LooseParameters, UniformBuffer_MultiFrame);
 
@@ -557,26 +559,6 @@ void FLocalVertexFactory::GetVertexElements(
 			Elements.Add(AccessStreamComponent(Data.TextureCoordinates[0], 15, InOutStreams));
 		}
 	}
-
-	check(Streams.Num() > 0);
-
-	InitDeclaration(Elements);
-	check(IsValidRef(GetDeclaration()));
-
-	const int32 DefaultBaseVertexIndex = 0;
-	const int32 DefaultPreSkinBaseVertexIndex = 0;
-
-	if (RHISupportsManualVertexFetch(GMaxRHIShaderPlatform) || bCanUseGPUScene)
-	{
-		SCOPED_LOADTIMER(FLocalVertexFactory_InitRHI_CreateLocalVFUniformBuffer);
-		UniformBuffer = CreateLocalVFUniformBuffer(this, Data.LODLightmapDataIndex, nullptr, DefaultBaseVertexIndex, DefaultPreSkinBaseVertexIndex);
-	}
-
-	FLocalVertexFactoryLooseParameters LooseParameters;
-	LooseParameters.GPUSkinPassThroughPreviousPositionBuffer = GNullVertexBuffer.VertexBufferSRV;
-	LooseParametersUniformBuffer = TUniformBufferRef<FLocalVertexFactoryLooseParameters>::CreateUniformBufferImmediate(LooseParameters, UniformBuffer_MultiFrame);
-
-	check(IsValidRef(GetDeclaration()));
 }
 
 IMPLEMENT_VERTEX_FACTORY_PARAMETER_TYPE(FLocalVertexFactory, SF_Vertex, FLocalVertexFactoryShaderParameters);
