Fix vertex & color vertex buffer AppendVertices functions potential integer overflow

diff --git a/Engine/Source/Runtime/Engine/Private/PositionVertexBuffer.cpp b/Engine/Source/Runtime/Engine/Private/PositionVertexBuffer.cpp
--- a/Engine/Source/Runtime/Engine/Private/PositionVertexBuffer.cpp
+++ b/Engine/Source/Runtime/Engine/Private/PositionVertexBuffer.cpp
@@ -108,8 +108,14 @@ void FPositionVertexBuffer::Init(const TArray<FVector3f>& InPositions, bool bInN
 	}
 }
 
-void FPositionVertexBuffer::AppendVertices( const FStaticMeshBuildVertex* Vertices, const uint32 NumVerticesToAppend )
+bool FPositionVertexBuffer::AppendVertices( const FStaticMeshBuildVertex* Vertices, const uint32 NumVerticesToAppend )
 {
+	const uint64 TotalNumVertices = (uint64)NumVertices + (uint64)NumVerticesToAppend;
+	if (!ensureMsgf(TotalNumVertices < INT32_MAX, TEXT("FPositionVertexBuffer::AppendVertices adding %u to %u vertices exceeds INT32_MAX limit"), NumVerticesToAppend, NumVertices))
+	{
+		return false;
+	}
+
 	if (VertexData == nullptr && NumVerticesToAppend > 0)
 	{
 		// Allocate the vertex data storage type if the buffer was never allocated before
@@ -138,6 +144,8 @@ void FPositionVertexBuffer::AppendVertices( const FStaticMeshBuildVertex* Vertic
 			}
 		}
 	}
+
+	return true;
 }
 
 /**
diff --git a/Engine/Source/Runtime/Engine/Private/Rendering/ColorVertexBuffer.cpp b/Engine/Source/Runtime/Engine/Private/Rendering/ColorVertexBuffer.cpp
--- a/Engine/Source/Runtime/Engine/Private/Rendering/ColorVertexBuffer.cpp
+++ b/Engine/Source/Runtime/Engine/Private/Rendering/ColorVertexBuffer.cpp
@@ -154,8 +154,14 @@ void FColorVertexBuffer::Init(const FColorVertexBuffer& InVertexBuffer, bool bNe
 	}
 }
 
-void FColorVertexBuffer::AppendVertices( const FStaticMeshBuildVertex* Vertices, const uint32 NumVerticesToAppend )
+bool FColorVertexBuffer::AppendVertices( const FStaticMeshBuildVertex* Vertices, const uint32 NumVerticesToAppend )
 {
+	const uint64 TotalNumVertices = (uint64)NumVertices + (uint64)NumVerticesToAppend;
+	if (!ensureMsgf(TotalNumVertices < INT32_MAX, TEXT("FColorVertexBuffer::AppendVertices adding %u to %u vertices exceeds INT32_MAX limit"), NumVerticesToAppend, NumVertices))
+	{
+		return false;
+	}
+
 	if (VertexData == nullptr && NumVerticesToAppend > 0)
 	{
 		check( NumVertices == 0 );
@@ -188,6 +194,8 @@ void FColorVertexBuffer::AppendVertices( const FStaticMeshBuildVertex* Vertices,
 			}
 		}
 	}
+
+	return true;
 }
 
 /**
diff --git a/Engine/Source/Runtime/Engine/Public/Rendering/ColorVertexBuffer.h b/Engine/Source/Runtime/Engine/Public/Rendering/ColorVertexBuffer.h
--- a/Engine/Source/Runtime/Engine/Public/Rendering/ColorVertexBuffer.h
+++ b/Engine/Source/Runtime/Engine/Public/Rendering/ColorVertexBuffer.h
@@ -48,8 +48,9 @@ public:
 	 *
 	 * @param	Vertices	The vertex data to be appended.  Must not be nullptr.
 	 * @param	NumVerticesToAppend		How many vertices should be added
+	 * @return	true if append operation is successful
 	 */
-	ENGINE_API void AppendVertices( const FStaticMeshBuildVertex* Vertices, const uint32 NumVerticesToAppend );
+	ENGINE_API bool AppendVertices( const FStaticMeshBuildVertex* Vertices, const uint32 NumVerticesToAppend );
 	
 	/**
 	* Serializer
diff --git a/Engine/Source/Runtime/Engine/Public/Rendering/PositionVertexBuffer.h b/Engine/Source/Runtime/Engine/Public/Rendering/PositionVertexBuffer.h
--- a/Engine/Source/Runtime/Engine/Public/Rendering/PositionVertexBuffer.h
+++ b/Engine/Source/Runtime/Engine/Public/Rendering/PositionVertexBuffer.h
@@ -54,8 +54,9 @@ public:
 	 *
 	 * @param	Vertices	The vertex data to be appended.  Must not be nullptr.
 	 * @param	NumVerticesToAppend		How many vertices should be added
+	 * @return	true if append operation is successful
 	 */
-	ENGINE_API void AppendVertices( const FStaticMeshBuildVertex* Vertices, const uint32 NumVerticesToAppend );
+	ENGINE_API bool AppendVertices( const FStaticMeshBuildVertex* Vertices, const uint32 NumVerticesToAppend );
 
 	/**
 	* Serializer
