Fixes for non unity build

#rnx

diff --git a/Engine/Source/Runtime/Engine/Private/Components/SkinnedMeshComponent.cpp b/Engine/Source/Runtime/Engine/Private/Components/SkinnedMeshComponent.cpp
--- a/Engine/Source/Runtime/Engine/Private/Components/SkinnedMeshComponent.cpp
+++ b/Engine/Source/Runtime/Engine/Private/Components/SkinnedMeshComponent.cpp
@@ -31,10 +31,14 @@
 #include "GPUSkinCache.h"
 #include "PSOPrecache.h"
 #include "MaterialDomain.h"
+#include "Materials/Material.h"
 #include "SkeletalRender.h"
 #include "UObject/UE5MainStreamObjectVersion.h"
 #include "UObject/UObjectIterator.h"
 #include "UObject/UnrealType.h"
+#include "HAL/LowLevelMemTracker.h"
+#include "HAL/LowLevelMemStats.h"
+#include "UObject/Package.h"
 
 DEFINE_LOG_CATEGORY_STATIC(LogSkinnedMeshComp, Log, All);
 
@@ -969,7 +973,7 @@ void USkinnedMeshComponent::CreateRenderState_Concurrent(FRegisterComponentConte
  	{
 		MeshDeformerInstance->AllocateResources();
 		
-		// Enqueue immediate execution of work here to ensure that we have some deformer outputs written for the next frame.
+			// Enqueue immediate execution of work here to ensure that we have some deformer outputs written for the next frame.
 		UMeshDeformerInstance::FEnqueueWorkDesc Desc;
 		Desc.Scene = GetScene();
 		Desc.ExecutionGroup = UMeshDeformerInstance::ExecutionGroup_Immediate;
@@ -1070,7 +1074,7 @@ void USkinnedMeshComponent::SendRenderDynamicData_Concurrent()
 			{
 				ActiveMorphTargets.Empty();
 			}
-			
+
 			MeshObject->Update(
 				UseLOD,
 				this,
diff --git a/Engine/Source/Runtime/Engine/Private/Components/SplineMeshComponent.cpp b/Engine/Source/Runtime/Engine/Private/Components/SplineMeshComponent.cpp
--- a/Engine/Source/Runtime/Engine/Private/Components/SplineMeshComponent.cpp
+++ b/Engine/Source/Runtime/Engine/Private/Components/SplineMeshComponent.cpp
@@ -21,6 +21,7 @@
 #include "SceneInterface.h"
 #include "UObject/UnrealType.h"
 #include "MaterialDomain.h"
+#include "Materials/Material.h"
 
 #if WITH_EDITOR
 #include "IHierarchicalLODUtilities.h"
diff --git a/Engine/Source/Runtime/Engine/Private/PSOPrecache.cpp b/Engine/Source/Runtime/Engine/Private/PSOPrecache.cpp
--- a/Engine/Source/Runtime/Engine/Private/PSOPrecache.cpp
+++ b/Engine/Source/Runtime/Engine/Private/PSOPrecache.cpp
@@ -9,6 +9,9 @@
 #include "SceneInterface.h"
 #include "PipelineStateCache.h"
 #include "VertexFactory.h"
+#include "UObject/StrongObjectPtr.h"
+#include "Materials/MaterialInterface.h"
+#include "MaterialShared.h"
 
 static TAutoConsoleVariable<int32> CVarPrecacheGlobalComputeShaders(
 	TEXT("r.PSOPrecache.GlobalComputeShaders"),
diff --git a/Engine/Source/Runtime/Renderer/Private/CustomDepthRendering.cpp b/Engine/Source/Runtime/Renderer/Private/CustomDepthRendering.cpp
--- a/Engine/Source/Runtime/Renderer/Private/CustomDepthRendering.cpp
+++ b/Engine/Source/Runtime/Renderer/Private/CustomDepthRendering.cpp
@@ -9,6 +9,7 @@
 #include "ScenePrivate.h"
 #include "Materials/Material.h"
 #include "MeshPassProcessor.inl"
+#include "UnrealEngine.h"
 
 static TAutoConsoleVariable<int32> CVarCustomDepth(
 	TEXT("r.CustomDepth"),
diff --git a/Engine/Source/Runtime/Renderer/Private/DepthRendering.cpp b/Engine/Source/Runtime/Renderer/Private/DepthRendering.cpp
--- a/Engine/Source/Runtime/Renderer/Private/DepthRendering.cpp
+++ b/Engine/Source/Runtime/Renderer/Private/DepthRendering.cpp
@@ -35,6 +35,7 @@
 #include "DebugProbeRendering.h"
 #include "RenderCore.h"
 #include "SimpleMeshDrawCommandPass.h"
+#include "UnrealEngine.h"
 
 static TAutoConsoleVariable<int32> CVarParallelPrePass(
 	TEXT("r.ParallelPrePass"),
diff --git a/Engine/Source/Runtime/Renderer/Private/ShadowDepthRendering.cpp b/Engine/Source/Runtime/Renderer/Private/ShadowDepthRendering.cpp
--- a/Engine/Source/Runtime/Renderer/Private/ShadowDepthRendering.cpp
+++ b/Engine/Source/Runtime/Renderer/Private/ShadowDepthRendering.cpp
@@ -49,6 +49,7 @@
 #include "Shadows/ShadowSceneRenderer.h"
 #include "RenderCore.h"
 #include "ShaderPlatformCachedIniValue.h"
+#include "UnrealEngine.h"
 
 DEFINE_GPU_DRAWCALL_STAT(ShadowDepths);
 
