Deprecate the legacy Level Streaming > Streaming Distance Factor property.

diff --git a/Engine/Source/Runtime/Engine/Classes/Engine/Engine.h b/Engine/Source/Runtime/Engine/Classes/Engine/Engine.h
--- a/Engine/Source/Runtime/Engine/Classes/Engine/Engine.h
+++ b/Engine/Source/Runtime/Engine/Classes/Engine/Engine.h
@@ -1545,8 +1545,8 @@ public:
 	UPROPERTY()
 	FColor C_BrushShape;
 
-	/** Fudge factor for tweaking the distance based miplevel determination */
-	UPROPERTY(EditAnywhere, Category=LevelStreaming, AdvancedDisplay)
+	/** Fudge factor for tweaking the distance based miplevel determination. No longer used. */
+	UE_DEPRECATED(5.2, "This setting is no longer used.")
 	float StreamingDistanceFactor;
 
 	/** The save directory for newly created screenshots */
