Engine FExec (exec commands) changes
- Add Exec_Dev and Exec_Editor virtual functions to allow implementing Exec for specific configurations.
- Add a UBT flag: bAllowExecCommandsInShipping as a master switch to allow or remove Exec in shipping.
- The default Exec implementation calls both Exec_Dev and Exec_Editor, when they are available.
- Blueprint ExecuteConsoleCommand will now ask the console manager to process the command before falling back to Player and GEngine Exec. This is to make sure that if exec commands are disabled the blueprint node is still able to execute console commands.

diff --git a/Engine/Source/Programs/UnrealBuildTool/Configuration/TargetRules.cs b/Engine/Source/Programs/UnrealBuildTool/Configuration/TargetRules.cs
--- a/Engine/Source/Programs/UnrealBuildTool/Configuration/TargetRules.cs
+++ b/Engine/Source/Programs/UnrealBuildTool/Configuration/TargetRules.cs
@@ -1227,6 +1227,12 @@ namespace UnrealBuildTool
 		[ConfigFile(ConfigHierarchyType.Engine, "/Script/BuildSettings.BuildSettings", "bCompileFreeType")]
 		public bool bCompileFreeType = true;
 
+		/// <summary>
+		/// Whether to turn allow exec commands for shipping builds.
+		/// </summary>
+		[RequiresUniqueBuildEnvironment]
+		public bool bUseExecCommandsInShipping = true;
+
 		/// <summary>
 		/// True if we want to favor optimizing size over speed.
 		/// </summary>
@@ -3110,6 +3116,11 @@ namespace UnrealBuildTool
 			get { return Inner.bCompileFreeType; }
 		}
 
+		public bool bUseExecCommandsInShipping
+		{
+			get { return Inner.bUseExecCommandsInShipping; }
+		}
+
 		[Obsolete("Deprecated in UE5.1 - Please use OptimizationLevel instead.")]
 		public bool bCompileForSize
 		{
diff --git a/Engine/Source/Programs/UnrealBuildTool/Configuration/UEBuildTarget.cs b/Engine/Source/Programs/UnrealBuildTool/Configuration/UEBuildTarget.cs
--- a/Engine/Source/Programs/UnrealBuildTool/Configuration/UEBuildTarget.cs
+++ b/Engine/Source/Programs/UnrealBuildTool/Configuration/UEBuildTarget.cs
@@ -4180,6 +4180,15 @@ namespace UnrealBuildTool
 				GlobalCompileEnvironment.Definitions.Add("USE_ESTIMATED_UTCNOW=0");
 			}
 
+			if (Rules.bUseExecCommandsInShipping)
+			{
+				GlobalCompileEnvironment.Definitions.Add("UE_ALLOW_EXEC_COMMANDS_IN_SHIPPING=1");
+			}
+			else
+			{
+				GlobalCompileEnvironment.Definitions.Add("UE_ALLOW_EXEC_COMMANDS_IN_SHIPPING=0");
+			}
+
 			if ((Rules.bCompileAgainstEditor && (Rules.Type == TargetType.Editor || Rules.Type == TargetType.Program)))
 			{
 				GlobalCompileEnvironment.Definitions.Add("WITH_EDITOR=1");
diff --git a/Engine/Source/Runtime/Core/Private/Misc/CoreMisc.cpp b/Engine/Source/Runtime/Core/Private/Misc/CoreMisc.cpp
--- a/Engine/Source/Runtime/Core/Private/Misc/CoreMisc.cpp
+++ b/Engine/Source/Runtime/Core/Private/Misc/CoreMisc.cpp
@@ -77,11 +77,30 @@ FStaticSelfRegisteringExec::FStaticSelfRegisteringExec(bool (*InStaticExecFunc)(
 :	StaticExecFunc(InStaticExecFunc)
 {}
 
+#if UE_ALLOW_EXEC_COMMANDS
 bool FStaticSelfRegisteringExec::Exec( UWorld* InWorld, const TCHAR* Cmd, FOutputDevice& Ar )
 {
 	return (*StaticExecFunc)( InWorld, Cmd,Ar);
 }
+#endif // UE_ALLOW_EXEC_COMMANDS
 
+FStaticSelfRegisteringExec_Dev::FStaticSelfRegisteringExec_Dev(bool (*InStaticExecFunc)(UWorld* Inworld, const TCHAR* Cmd,FOutputDevice& Ar))
+	: StaticExecFunc(InStaticExecFunc)
+{}
+
+bool FStaticSelfRegisteringExec_Dev::Exec_Dev(UWorld* InWorld, const TCHAR* Cmd, FOutputDevice& Ar)
+{
+	return (*StaticExecFunc)(InWorld, Cmd, Ar);
+}
+
+FStaticSelfRegisteringExec_Editor::FStaticSelfRegisteringExec_Editor(bool (*InStaticExecFunc)(UWorld* Inworld, const TCHAR* Cmd, FOutputDevice& Ar))
+	: StaticExecFunc(InStaticExecFunc)
+{}
+
+bool FStaticSelfRegisteringExec_Editor::Exec_Editor(UWorld* InWorld, const TCHAR* Cmd, FOutputDevice& Ar)
+{
+	return (*StaticExecFunc)(InWorld, Cmd, Ar);
+}
 
 // Remove old crash contexts
 
diff --git a/Engine/Source/Runtime/Core/Private/Misc/Exec.cpp b/Engine/Source/Runtime/Core/Private/Misc/Exec.cpp
--- a/Engine/Source/Runtime/Core/Private/Misc/Exec.cpp
+++ b/Engine/Source/Runtime/Core/Private/Misc/Exec.cpp
@@ -5,3 +5,20 @@
 FExec::~FExec()
 {
 }
+
+#if UE_ALLOW_EXEC_COMMANDS
+bool FExec::Exec( UWorld* InWorld, const TCHAR* Cmd, FOutputDevice& Ar )
+{
+	bool bExecSuccess = false;
+
+#if UE_ALLOW_EXEC_EDITOR
+	bExecSuccess = bExecSuccess || Exec_Editor( InWorld, Cmd, Ar );
+#endif
+
+#if UE_ALLOW_EXEC_DEV
+	bExecSuccess = bExecSuccess || Exec_Dev( InWorld, Cmd, Ar );
+#endif
+
+	return bExecSuccess;
+}
+#endif // UE_ALLOW_EXEC_COMMANDS
\ No newline at end of file
diff --git a/Engine/Source/Runtime/Core/Public/Misc/CoreMisc.h b/Engine/Source/Runtime/Core/Public/Misc/CoreMisc.h
--- a/Engine/Source/Runtime/Core/Public/Misc/CoreMisc.h
+++ b/Engine/Source/Runtime/Core/Public/Misc/CoreMisc.h
@@ -45,7 +45,26 @@ public:
 	FStaticSelfRegisteringExec(bool (*InStaticExecFunc)(UWorld* Inworld, const TCHAR* Cmd,FOutputDevice& Ar));
 
 	//~ Begin Exec Interface
-	virtual bool Exec( UWorld* InWorld, const TCHAR* Cmd, FOutputDevice& Ar );
+#if UE_ALLOW_EXEC_COMMANDS
+	virtual bool Exec( UWorld* InWorld, const TCHAR* Cmd, FOutputDevice& Ar ) override;
+#endif
+	//~ End Exec Interface
+
+private:
+
+	bool (*StaticExecFunc)(UWorld* Inworld, const TCHAR* Cmd,FOutputDevice& Ar);
+};
+
+/** Registers a static Exec_Dev function using FSelfRegisteringExec. */
+class CORE_API FStaticSelfRegisteringExec_Dev: public FSelfRegisteringExec
+{
+public:
+
+	/** Initialization constructor. */
+	explicit FStaticSelfRegisteringExec_Dev(bool (*InStaticExecFunc)(UWorld* Inworld, const TCHAR* Cmd, FOutputDevice& Ar));
+
+	//~ Begin Exec Interface
+	virtual bool Exec_Dev(UWorld* InWorld, const TCHAR* Cmd, FOutputDevice& Ar) override;
 	//~ End Exec Interface
 
 private:
@@ -53,6 +72,23 @@ private:
 	bool (*StaticExecFunc)(UWorld* Inworld, const TCHAR* Cmd,FOutputDevice& Ar);
 };
 
+/** Registers a static Exec_Editor function using FSelfRegisteringExec. */
+class CORE_API FStaticSelfRegisteringExec_Editor: public FSelfRegisteringExec
+{
+public:
+
+	/** Initialization constructor. */
+	explicit FStaticSelfRegisteringExec_Editor(bool (*InStaticExecFunc)(UWorld* Inworld, const TCHAR* Cmd, FOutputDevice& Ar));
+
+	//~ Begin Exec Interface
+	virtual bool Exec_Editor(UWorld* InWorld, const TCHAR* Cmd, FOutputDevice& Ar) override;
+	//~ End Exec Interface
+
+private:
+
+	bool (*StaticExecFunc)(UWorld* Inworld, const TCHAR* Cmd, FOutputDevice& Ar);
+};
+
 // Interface for returning a context string.
 class FContextSupplier
 {
diff --git a/Engine/Source/Runtime/Core/Public/Misc/Exec.h b/Engine/Source/Runtime/Core/Public/Misc/Exec.h
--- a/Engine/Source/Runtime/Core/Public/Misc/Exec.h
+++ b/Engine/Source/Runtime/Core/Public/Misc/Exec.h
@@ -6,6 +6,23 @@
 #include "Misc/AssertionMacros.h"
 
 class FOutputDevice;
+class UWorld;
+
+#ifndef UE_ALLOW_EXEC_COMMANDS
+	#if UE_BUILD_SHIPPING && !WITH_EDITOR
+		#define UE_ALLOW_EXEC_COMMANDS UE_ALLOW_EXEC_COMMANDS_IN_SHIPPING
+	#else
+		#define UE_ALLOW_EXEC_COMMANDS 1
+	#endif
+#endif
+
+#ifndef UE_ALLOW_EXEC_DEV
+	#define UE_ALLOW_EXEC_DEV !UE_BUILD_SHIPPING && UE_ALLOW_EXEC_COMMANDS
+#endif
+
+#ifndef UE_ALLOW_EXEC_EDITOR
+	#define UE_ALLOW_EXEC_EDITOR WITH_EDITOR && UE_ALLOW_EXEC_COMMANDS
+#endif
 
 // Any object that is capable of taking commands.
 class CORE_API FExec
@@ -13,16 +30,40 @@ class CORE_API FExec
 public:
 	virtual ~FExec();
 
+#if UE_ALLOW_EXEC_COMMANDS
+
+	/**
+	 * Exec handler
+	 *
+	 * @param	InWorld	World context
+	 * @param	Cmd		Command to parse
+	 * @param	Ar		Output device to log to
+	 *
+	 * @return	true if command was handled, false otherwise
+	 */
+	virtual bool Exec( UWorld* InWorld, const TCHAR* Cmd, FOutputDevice& Ar );
+#else // UE_ALLOW_EXEC_COMMANDS
+
+public:
+	/**
+	 * final override of Exec that asserts if called
+	 */
+	virtual bool Exec( UWorld* InWorld, const TCHAR* Cmd, FOutputDevice& Ar ) final { check(false); return false; }
+	
 	/**
-	* Exec handler
-	*
-	* @param	InWorld	World context
-	* @param	Cmd		Command to parse
-	* @param	Ar		Output device to log to
-	*
-	* @return	true if command was handled, false otherwise
-	*/
-	virtual bool Exec( class UWorld* InWorld, const TCHAR* Cmd, FOutputDevice& Ar ) PURE_VIRTUAL(FExec::Exec,return false;)
+	 * final override of Exec to replace overrides where a default value for Ar is provided
+	 */
+	virtual bool Exec( UWorld* InWorld, const TCHAR* Cmd ) final { check(false); return false; }
+
+#endif // !UE_ALLOW_EXEC_COMMANDS
+
+protected:
+	/** Implementation of Exec that is only called in non-shipping targets */
+	virtual bool Exec_Dev( UWorld* InWorld, const TCHAR* Cmd, FOutputDevice& Ar ) { return false; }
+
+	/** Implementation of Exec that is only called in editor */
+	virtual bool Exec_Editor( UWorld* InWorld, const TCHAR* Cmd, FOutputDevice& Ar ) { return false; }
+
 };
 
 
diff --git a/Engine/Source/Runtime/Engine/Classes/Engine/Engine.h b/Engine/Source/Runtime/Engine/Classes/Engine/Engine.h
--- a/Engine/Source/Runtime/Engine/Classes/Engine/Engine.h
+++ b/Engine/Source/Runtime/Engine/Classes/Engine/Engine.h
@@ -2157,9 +2157,17 @@ public:
 	void ParseCommandline();
 
 	//~ Begin FExec Interface
+public:
+#if UE_ALLOW_EXEC_COMMANDS
 	virtual bool Exec( UWorld* InWorld, const TCHAR* Cmd, FOutputDevice& Out=*GLog ) override;
+#endif
+
+protected:
+	virtual bool Exec_Dev( UWorld* InWorld, const TCHAR* Cmd, FOutputDevice& Out=*GLog ) override;
+	virtual bool Exec_Editor(UWorld* InWorld, const TCHAR* Cmd, FOutputDevice& Out = *GLog) override;
 	//~ End FExec Interface
 
+public:
 	/** 
 	 * Exec command handlers
 	 */
diff --git a/Engine/Source/Runtime/Engine/Private/KismetSystemLibrary.cpp b/Engine/Source/Runtime/Engine/Private/KismetSystemLibrary.cpp
--- a/Engine/Source/Runtime/Engine/Private/KismetSystemLibrary.cpp
+++ b/Engine/Source/Runtime/Engine/Private/KismetSystemLibrary.cpp
@@ -36,6 +36,7 @@
 #include "KismetTraceUtils.h"
 #include "Engine/AssetManager.h"
 #include "RHI.h"
+#include "HAL/IConsoleManager.h"
 #include "HAL/PlatformApplicationMisc.h"
 #include "UObject/FieldPathProperty.h"
 #include "Commandlets/Commandlet.h"
@@ -445,16 +446,23 @@ void UKismetSystemLibrary::SetWindowTitle(const FText& Title)
 
 void UKismetSystemLibrary::ExecuteConsoleCommand(const UObject* WorldContextObject, const FString& Command, APlayerController* Player)
 {
-	// First, try routing through the primary player
 	UWorld* World = GEngine->GetWorldFromContextObject(WorldContextObject, EGetWorldErrorMode::ReturnNull);
-	APlayerController* TargetPC = Player || !World ? Player : World->GetFirstPlayerController();
-	if (TargetPC)
-	{
-		TargetPC->ConsoleCommand(Command, true);
-	}
-	else
+
+	// First, try routing through the console manager directly.
+	// This is needed in case the Exec commands have been compiled out, meaning that GEngine->Exec wouldn't route to anywhere.
+	if (IConsoleManager::Get().ProcessUserConsoleInput(*Command, *GLog, World) == false)
 	{
-		GEngine->Exec(World, *Command);
+		APlayerController* TargetPC = Player || !World ? Player : World->GetFirstPlayerController();
+
+		// Second, try routing through the primary player
+		if (TargetPC)
+		{
+			TargetPC->ConsoleCommand(Command, true);
+		}
+		else
+		{
+			GEngine->Exec(World, *Command);
+		}
 	}
 }
 
diff --git a/Engine/Source/Runtime/Engine/Private/UnrealEngine.cpp b/Engine/Source/Runtime/Engine/Private/UnrealEngine.cpp
--- a/Engine/Source/Runtime/Engine/Private/UnrealEngine.cpp
+++ b/Engine/Source/Runtime/Engine/Private/UnrealEngine.cpp
@@ -454,7 +454,7 @@
 #if !UE_BUILD_SHIPPING
 class FDisplayCVarListExecHelper : public FSelfRegisteringExec
 {
-	virtual bool Exec(class UWorld* InWorld, const TCHAR* Cmd, FOutputDevice& Ar)
+	virtual bool Exec_Dev(class UWorld* InWorld, const TCHAR* Cmd, FOutputDevice& Ar)
 	{
 		if (FParse::Command(&Cmd, TEXT("DisplayCVarList")))
 		{
@@ -4607,6 +4607,7 @@
 	UE_LOG(LogEngine, Log, TEXT("BufferOverflowFunction BufferSize=%d LocalBuffer=%s"),(int32)BufferSize, ANSI_TO_TCHAR(LocalBuffer));
 }
 
+#if UE_ALLOW_EXEC_COMMANDS
 bool UEngine::Exec( UWorld* InWorld, const TCHAR* Cmd, FOutputDevice& Ar )
 {
 	// If we don't have a viewport specified to catch the stat commands, use to the game viewport
@@ -4615,6 +4616,10 @@
 		GStatProcessingViewportClient = GameViewport;
 	}
 
+	if (FExec::Exec( InWorld, Cmd, Ar ))
+	{
+		return true;
+	}
 
 	// See if any other subsystems claim the command.
 	if (StaticExec(InWorld, Cmd,Ar) == true)
@@ -4784,7 +4789,7 @@
 	{
 		return HandleProfileGPUCommand( Cmd, Ar );
 	}	
-#endif // #if !UE_BUILD_SHIPPING
+#endif
 
 #if WITH_DUMPGPU
 	else if (FParse::Command(&Cmd, TEXT("DUMPGPU")))
@@ -4793,16 +4798,48 @@
 	}
 #endif
 
-#if	!(UE_BUILD_SHIPPING || UE_BUILD_TEST) && WITH_HOT_RELOAD
-	else if( FParse::Command(&Cmd,TEXT("HotReload")) )
+	else if ( FParse::Command(&Cmd,TEXT("SCALABILITY")) )
 	{
-		return HandleHotReloadCommand( Cmd, Ar );
+		Scalability::ProcessCommand(Cmd, Ar);
+		return true;
 	}
-#endif // !(UE_BUILD_SHIPPING || UE_BUILD_TEST) && WITH_HOT_RELOAD
+	else if(IConsoleManager::Get().ProcessUserConsoleInput(Cmd, Ar, InWorld))
+	{
+		// console variable interaction (get value, set value or get help)
+		return true;
+	}
+	else if (!IStreamingManager::HasShutdown() && IStreamingManager::Get().Exec( InWorld, Cmd,Ar ))
+	{
+		// The streaming manager has handled the exec command.
+	}
+	else if( FParse::Command(&Cmd, TEXT("DUMPTICKS")) )
+	{
+		return HandleDumpTicksCommand( InWorld, Cmd, Ar );
+	}
+	else if (FParse::Command(&Cmd, TEXT("CANCELASYNCLOAD")))
+	{
+		CancelAsyncLoading();
+		return true;
+	}
+#if USE_NETWORK_PROFILER
+	else if( FParse::Command(&Cmd,TEXT("NETPROFILE")) )
+	{
+		GNetworkProfiler.Exec( InWorld, Cmd, Ar );
+	}
+#endif
+	else 
+	{
+		return false;
+	}
 
+	return true;
+}
+#endif // UE_ALLOW_EXEC_COMMANDS
 
+bool UEngine::Exec_Dev( UWorld* InWorld, const TCHAR* Cmd, FOutputDevice& Ar )
+{
 #if !UE_BUILD_SHIPPING
-	else if (FParse::Command(&Cmd, TEXT("DumpConsoleCommands")))
+	if (FParse::Command(&Cmd, TEXT("DumpConsoleCommands")))
 	{
 		return HandleDumpConsoleCommandsCommand( Cmd, Ar, InWorld );
 	}
@@ -4830,7 +4867,6 @@
 	{
 		return HandleFreezeAllCommand( Cmd, Ar, InWorld );
 	}
-
 	else if( FParse::Command(&Cmd,TEXT("ToggleRenderingThread")) )
 	{
 		return HandleToggleRenderingThreadCommand( Cmd, Ar );
@@ -4855,12 +4891,6 @@
 	{
 		return HandleDumpMaterialStatsCommand( Cmd, Ar );	
 	}
-#if WITH_EDITOR
-	else if( FParse::Command(&Cmd,TEXT("DumpShaderCompileStats")) )
-	{
-		HandleDumpShaderCompileStatsCommand(Cmd, Ar);
-	}
-#endif
 	else if (FParse::Command(&Cmd, TEXT("DumpShaderPipelineStats")))
 	{
 		return HandleDumpShaderPipelineStatsCommand(Cmd, Ar);
@@ -5041,41 +5071,26 @@
 	}
 #endif // !UE_BUILD_SHIPPING
 
-	else if ( FParse::Command(&Cmd,TEXT("SCALABILITY")) )
+#if	!(UE_BUILD_SHIPPING || UE_BUILD_TEST) && WITH_HOT_RELOAD
+	else if (FParse::Command(&Cmd, TEXT("HotReload")))
 	{
-		Scalability::ProcessCommand(Cmd, Ar);
-		return true;
+		return HandleHotReloadCommand(Cmd, Ar);
 	}
-	else if(IConsoleManager::Get().ProcessUserConsoleInput(Cmd, Ar, InWorld))
+#endif // !(UE_BUILD_SHIPPING || UE_BUILD_TEST) && WITH_HOT_RELOAD
+
+	return false;
+}
+
+bool UEngine::Exec_Editor(UWorld* InWorld, const TCHAR* Cmd, FOutputDevice& Ar)
+{
+#if WITH_EDITOR
+	if (FParse::Command(&Cmd, TEXT("DumpShaderCompileStats")))
 	{
-		// console variable interaction (get value, set value or get help)
-		return true;
+		return HandleDumpShaderCompileStatsCommand(Cmd, Ar);
 	}
-	else if (!IStreamingManager::HasShutdown() && IStreamingManager::Get().Exec( InWorld, Cmd,Ar ))
-	{
-		// The streaming manager has handled the exec command.
-	}
-	else if( FParse::Command(&Cmd, TEXT("DUMPTICKS")) )
-	{
-		return HandleDumpTicksCommand( InWorld, Cmd, Ar );
-	}
-	else if (FParse::Command(&Cmd, TEXT("CANCELASYNCLOAD")))
-	{
-		CancelAsyncLoading();
-		return true;
-	}
-#if USE_NETWORK_PROFILER
-	else if( FParse::Command(&Cmd,TEXT("NETPROFILE")) )
-	{
-		GNetworkProfiler.Exec( InWorld, Cmd, Ar );
-	}
 #endif
-	else 
-	{
-		return false;
-	}
 
-	return true;
+	return false;
 }
 
 bool UEngine::HandleFlushLogCommand( const TCHAR* Cmd, FOutputDevice& Ar )
@@ -13035,7 +13050,7 @@
 	}
 
 	/** Console commands, see embeded usage statement **/
-	virtual bool Exec( UWorld* InWorld, const TCHAR* Cmd, FOutputDevice& Ar ) override
+	virtual bool Exec_Dev( UWorld* InWorld, const TCHAR* Cmd, FOutputDevice& Ar ) override
 	{
 		if(FParse::Command(&Cmd,TEXT("CDODump")))
 		{
