Re-introduced material compilation error emitted when the compiler detects the same  material parameters being set to different values.

This check is wrapped by a new cvar "r.Material.PedanticErrorChecksEnabled" (defaulted to false) to avoid breaking compilation of many existing materials.

diff --git a/Engine/Source/Runtime/Engine/Private/Materials/HLSLMaterialTranslator.cpp b/Engine/Source/Runtime/Engine/Private/Materials/HLSLMaterialTranslator.cpp
--- a/Engine/Source/Runtime/Engine/Private/Materials/HLSLMaterialTranslator.cpp
+++ b/Engine/Source/Runtime/Engine/Private/Materials/HLSLMaterialTranslator.cpp
@@ -110,6 +110,12 @@ static FAutoConsoleVariableRef CVarLWCEnabled(
 	TEXT("Enable generation of LWC values in materials. If disabled, materials will perform all operations at float-precision")
 );
 
+static bool GPedanticErrorChecksEnabled = false;
+static FAutoConsoleVariableRef CVarPedanticErrorChecksEnabled(
+	TEXT("r.Material.PedanticErrorChecksEnabled"),
+	GPedanticErrorChecksEnabled,
+	TEXT("Enables material compilation pedantic error checking"));
+
 static inline bool IsAnalyticDerivEnabled()
 {
 	return GAnalyticDerivEnabled != 0;
@@ -774,6 +780,12 @@ bool FHLSLMaterialTranslator::Translate()
 		Material->CompileErrors.Empty();
 		Material->ErrorExpressions.Empty();
 
+		// If pedantic error checking is enabled, log out the pre-compilation errors
+		if (GPedanticErrorChecksEnabled)
+		{
+			bSuccess = Material->CheckInValidStateForCompilation(this);
+		}
+
 		bEnableExecutionFlow = Material->IsUsingControlFlow();
 		bCompileForComputeShader = Material->IsLightFunction();
 
diff --git a/Engine/Source/Runtime/Engine/Private/Materials/Material.cpp b/Engine/Source/Runtime/Engine/Private/Materials/Material.cpp
--- a/Engine/Source/Runtime/Engine/Private/Materials/Material.cpp
+++ b/Engine/Source/Runtime/Engine/Private/Materials/Material.cpp
@@ -6219,9 +6219,8 @@ bool UMaterial::CheckInValidStateForCompilation(class FMaterialCompiler* Compile
 	{
 		for (const TPair<TObjectPtr<UMaterialExpression>, FName>& Error : CachedExpressionData->DuplicateParameterErrors)
 		{
-			// Although this is an error, log it as a warning for now (it should instead be an error logged to the compiler).
-			// We need to first solve this error for all existing materials before we can turn it into a compilation blocking error.
-			UE_LOG(LogMaterial, Warning, TEXT("Parameter '%s' is set multiple times to different values. Make sure each parameter is set once or always to the same value (e.g. same texture)."), *Error.Get<1>().ToString());
+			FString ErrorMsg = FString::Format(TEXT("Parameter '{0}' is set multiple times to different values. Make sure each parameter is set once or always to the same value (e.g. same texture)."), { *Error.Get<1>().ToString() });
+			Compiler->AppendExpressionError(Error.Get<0>(), *ErrorMsg);
 		}
 
 		bSuccess = false;
