Fix an issue where cells with data layer(s) set as loaded on a server which prevents cell unloads would be stuck in the LoadedCells queue and never become active.

diff --git a/Engine/Source/Runtime/Engine/Private/WorldPartition/WorldPartitionStreamingPolicy.cpp b/Engine/Source/Runtime/Engine/Private/WorldPartition/WorldPartitionStreamingPolicy.cpp
--- a/Engine/Source/Runtime/Engine/Private/WorldPartition/WorldPartitionStreamingPolicy.cpp
+++ b/Engine/Source/Runtime/Engine/Private/WorldPartition/WorldPartitionStreamingPolicy.cpp
@@ -470,7 +470,7 @@ void UWorldPartitionStreamingPolicy::UpdateStreamingState()
 
 			auto AddServerFrameCell = [this, bCanServerDeactivateOrUnloadDataLayerCells, &EffectiveLoadedDataLayerNames, &EffectiveActiveDataLayerNames](const UWorldPartitionRuntimeCell* Cell)
 			{
-				// Keep Data Layer cells is their current state if server cannot deactivat/unload data layer cells
+				// Keep Data Layer cells is their current state if server cannot deactivate/unload data layer cells
 				if (!bCanServerDeactivateOrUnloadDataLayerCells && Cell->HasDataLayers())
 				{
 					if (ActivatedCells.Contains(Cell))
@@ -478,10 +478,15 @@ void UWorldPartitionStreamingPolicy::UpdateStreamingState()
 						FrameActivateCells.Add(Cell);
 						return;
 					}
-					else if (LoadedCells.Contains(Cell))
+					else
 					{
-						FrameLoadCells.Add(Cell);
-						return;
+						// Allow a cell with data layer(s) to switch from loaded to activated. Do not early return here in that case.
+						const bool bIsAnActivatedDataLayerCell = EffectiveActiveDataLayerNames.Num() && Cell->HasAnyDataLayer(EffectiveActiveDataLayerNames);
+						if (LoadedCells.Contains(Cell) && !bIsAnActivatedDataLayerCell)
+						{
+							FrameLoadCells.Add(Cell);
+							return;
+						}
 					}
 				}
 				
