Exec commands changes:
- Port some editor only exec commands to Exec_Editor.

diff --git a/Engine/Plugins/VirtualProduction/Takes/Source/TakeRecorderSources/Private/TakeRecorderSourcesModule.cpp b/Engine/Plugins/VirtualProduction/Takes/Source/TakeRecorderSources/Private/TakeRecorderSourcesModule.cpp
--- a/Engine/Plugins/VirtualProduction/Takes/Source/TakeRecorderSources/Private/TakeRecorderSourcesModule.cpp
+++ b/Engine/Plugins/VirtualProduction/Takes/Source/TakeRecorderSources/Private/TakeRecorderSourcesModule.cpp
@@ -706,7 +706,7 @@ public:
 #endif
 	}
 
-	virtual bool Exec(UWorld* InWorld, const TCHAR* Cmd, FOutputDevice& Ar) override
+	virtual bool Exec_Editor(UWorld* InWorld, const TCHAR* Cmd, FOutputDevice& Ar) override
 	{
 #if WITH_EDITOR
 		if (FParse::Command(&Cmd, TEXT("RecordTake")))
diff --git a/Engine/Source/Editor/Kismet/Private/FindInBlueprintManager.cpp b/Engine/Source/Editor/Kismet/Private/FindInBlueprintManager.cpp
--- a/Engine/Source/Editor/Kismet/Private/FindInBlueprintManager.cpp
+++ b/Engine/Source/Editor/Kismet/Private/FindInBlueprintManager.cpp
@@ -3654,7 +3654,7 @@ TStatId FFindInBlueprintSearchManager::GetStatId() const
 
 class FFiBDumpIndexCacheToFileExecHelper : public FSelfRegisteringExec
 {
-	virtual bool Exec(class UWorld* InWorld, const TCHAR* Cmd, FOutputDevice& Ar)
+	virtual bool Exec_Editor(class UWorld* InWorld, const TCHAR* Cmd, FOutputDevice& Ar)
 	{
 		if (FParse::Command(&Cmd, TEXT("DUMPFIBINDEXCACHE")))
 		{
diff --git a/Engine/Source/Editor/SequenceRecorder/Private/SequenceRecorderModule.cpp b/Engine/Source/Editor/SequenceRecorder/Private/SequenceRecorderModule.cpp
--- a/Engine/Source/Editor/SequenceRecorder/Private/SequenceRecorderModule.cpp
+++ b/Engine/Source/Editor/SequenceRecorder/Private/SequenceRecorderModule.cpp
@@ -264,7 +264,7 @@ class FSequenceRecorderModule : public ISequenceRecorder, private FSelfRegisteri
 	}
 
 	// FSelfRegisteringExec implementation
-	virtual bool Exec(UWorld* InWorld, const TCHAR* Cmd, FOutputDevice& Ar) override
+	virtual bool Exec_Editor(UWorld* InWorld, const TCHAR* Cmd, FOutputDevice& Ar) override
 	{
 #if WITH_EDITOR
 		if (FParse::Command(&Cmd, TEXT("RecordAnimation")))
diff --git a/Engine/Source/Editor/SkeletalMeshEditor/Private/MeshHierarchyCmd.cpp b/Engine/Source/Editor/SkeletalMeshEditor/Private/MeshHierarchyCmd.cpp
--- a/Engine/Source/Editor/SkeletalMeshEditor/Private/MeshHierarchyCmd.cpp
+++ b/Engine/Source/Editor/SkeletalMeshEditor/Private/MeshHierarchyCmd.cpp
@@ -12,7 +12,7 @@
 
 static FMeshHierarchyCmd MeshHierarchyCmdExec;
 
-bool FMeshHierarchyCmd::Exec(UWorld*, const TCHAR* Cmd, FOutputDevice& Ar)
+bool FMeshHierarchyCmd::Exec_Editor(UWorld*, const TCHAR* Cmd, FOutputDevice& Ar)
 {
 	bool bResult = false;
 	if (FParse::Command(&Cmd, TEXT("TMH")))
diff --git a/Engine/Source/Editor/SkeletalMeshEditor/Private/MeshHierarchyCmd.h b/Engine/Source/Editor/SkeletalMeshEditor/Private/MeshHierarchyCmd.h
--- a/Engine/Source/Editor/SkeletalMeshEditor/Private/MeshHierarchyCmd.h
+++ b/Engine/Source/Editor/SkeletalMeshEditor/Private/MeshHierarchyCmd.h
@@ -7,7 +7,7 @@
 
 class FMeshHierarchyCmd : private FSelfRegisteringExec
 {
-public:
+protected:
 	/** Console commands, see embeded usage statement **/
-	virtual bool Exec(UWorld*, const TCHAR* Cmd, FOutputDevice& Ar) override;
+	virtual bool Exec_Editor(UWorld*, const TCHAR* Cmd, FOutputDevice& Ar) override;
 };
diff --git a/Engine/Source/Editor/TurnkeySupport/Private/TurnkeySupportModule.cpp b/Engine/Source/Editor/TurnkeySupport/Private/TurnkeySupportModule.cpp
--- a/Engine/Source/Editor/TurnkeySupport/Private/TurnkeySupportModule.cpp
+++ b/Engine/Source/Editor/TurnkeySupport/Private/TurnkeySupportModule.cpp
@@ -2524,7 +2524,7 @@ void FTurnkeySupportModule::ClearDeviceStatus(FName PlatformName)
 
 
 #if WITH_EDITOR
-bool FTurnkeySupportModule::Exec( UWorld* InWorld, const TCHAR* Cmd, FOutputDevice& Ar )
+bool FTurnkeySupportModule::Exec_Editor( UWorld* InWorld, const TCHAR* Cmd, FOutputDevice& Ar )
 {
 	if (FParse::Command( &Cmd, TEXT("Turnkey")) )
 	{
diff --git a/Engine/Source/Editor/TurnkeySupport/Private/TurnkeySupportModule.h b/Engine/Source/Editor/TurnkeySupport/Private/TurnkeySupportModule.h
--- a/Engine/Source/Editor/TurnkeySupport/Private/TurnkeySupportModule.h
+++ b/Engine/Source/Editor/TurnkeySupport/Private/TurnkeySupportModule.h
@@ -58,7 +58,7 @@ private:
 
 #if WITH_EDITOR
 	// FSelfRegisteringExec interface
-	virtual bool Exec( class UWorld* InWorld, const TCHAR* Cmd, FOutputDevice& Ar ) override;
+	virtual bool Exec_Editor( class UWorld* InWorld, const TCHAR* Cmd, FOutputDevice& Ar ) override;
 #endif
 
 };
diff --git a/Engine/Source/Editor/UnrealEd/Classes/CookOnTheSide/CookOnTheFlyServer.h b/Engine/Source/Editor/UnrealEd/Classes/CookOnTheSide/CookOnTheFlyServer.h
--- a/Engine/Source/Editor/UnrealEd/Classes/CookOnTheSide/CookOnTheFlyServer.h
+++ b/Engine/Source/Editor/UnrealEd/Classes/CookOnTheSide/CookOnTheFlyServer.h
@@ -604,9 +604,6 @@ public:
 	ECookMode::Type GetCookMode() const { return CurrentCookMode; }
 	ECookInitializationFlags GetCookFlags() const { return CookFlags; }
 
-	// FExec interface used in the editor
-	virtual bool Exec(class UWorld* InWorld, const TCHAR* Cmd, FOutputDevice& Ar) override;
-
 	// ICookInfo interface
 	virtual UE::Cook::FInstigator GetInstigator(FName PackageName) override;
 	virtual TArray<UE::Cook::FInstigator> GetInstigatorChain(FName PackageName) override;
@@ -869,6 +866,10 @@ public:
 	/** Print detailed stats from the cook. */
 	void PrintDetailedCookStats();
 
+protected:
+	// FExec interface used in the editor
+	virtual bool Exec_Editor(class UWorld* InWorld, const TCHAR* Cmd, FOutputDevice& Ar) override;
+
 private:
 
 	/**
diff --git a/Engine/Source/Editor/UnrealEd/Classes/Editor/EditorEngine.h b/Engine/Source/Editor/UnrealEd/Classes/Editor/EditorEngine.h
--- a/Engine/Source/Editor/UnrealEd/Classes/Editor/EditorEngine.h
+++ b/Engine/Source/Editor/UnrealEd/Classes/Editor/EditorEngine.h
@@ -814,14 +814,14 @@ private:
 
 	/** Called when hotreload adds a new class to create volume factories */
 	void CreateVolumeFactoriesForNewClasses(const TArray<UClass*>& NewClasses);
-
-public:
 	//~ End UEngine Interface.
 	
 	//~ Begin FExec Interface
-	virtual bool Exec( UWorld* InWorld, const TCHAR* Cmd, FOutputDevice& Ar=*GLog ) override;
+protected:
+	virtual bool Exec_Editor( UWorld* InWorld, const TCHAR* Cmd, FOutputDevice& Ar=*GLog ) override;
 	//~ End FExec Interface
 
+public:
 	bool	CommandIsDeprecated( const TCHAR* CommandStr, FOutputDevice& Ar );
 	
 	/**
diff --git a/Engine/Source/Editor/UnrealEd/Classes/Editor/UnrealEdEngine.h b/Engine/Source/Editor/UnrealEd/Classes/Editor/UnrealEdEngine.h
--- a/Engine/Source/Editor/UnrealEd/Classes/Editor/UnrealEdEngine.h
+++ b/Engine/Source/Editor/UnrealEd/Classes/Editor/UnrealEdEngine.h
@@ -198,9 +198,11 @@ public:
 	//~ End UEditorEngine Interface 
 	
 	//~ Begin FExec Interface
-	virtual bool Exec( UWorld* InWorld, const TCHAR* Cmd, FOutputDevice& Ar=*GLog ) override;
+protected:
+	virtual bool Exec_Editor( UWorld* InWorld, const TCHAR* Cmd, FOutputDevice& Ar=*GLog ) override;
 	//~ End FExec Interface
 
+public:
 
 	//~ Begin UEngine Interface.
 	virtual void Init(IEngineLoop* InEngineLoop) override;
diff --git a/Engine/Source/Editor/UnrealEd/Private/CookOnTheFlyServer.cpp b/Engine/Source/Editor/UnrealEd/Private/CookOnTheFlyServer.cpp
--- a/Engine/Source/Editor/UnrealEd/Private/CookOnTheFlyServer.cpp
+++ b/Engine/Source/Editor/UnrealEd/Private/CookOnTheFlyServer.cpp
@@ -6071,7 +6071,7 @@ void UCookOnTheFlyServer::InitializeSession()
 	VirtualMemoryHistory.Initialize(FPlatformMemory::GetStats().UsedVirtual);
 }
 
-bool UCookOnTheFlyServer::Exec(class UWorld* InWorld, const TCHAR* Cmd, FOutputDevice& Ar)
+bool UCookOnTheFlyServer::Exec_Editor(class UWorld* InWorld, const TCHAR* Cmd, FOutputDevice& Ar)
 {
 	if (FParse::Command(&Cmd, TEXT("package")))
 	{
diff --git a/Engine/Source/Editor/UnrealEd/Private/DebugToolExec.cpp b/Engine/Source/Editor/UnrealEd/Private/DebugToolExec.cpp
--- a/Engine/Source/Editor/UnrealEd/Private/DebugToolExec.cpp
+++ b/Engine/Source/Editor/UnrealEd/Private/DebugToolExec.cpp
@@ -86,7 +86,7 @@ void FDebugToolExec::EditObject(UObject* Object, bool bShouldShowNonEditable)
  * @param Cmd	Command to parse
  * @param Ar	output device used for logging
  */
-bool FDebugToolExec::Exec( UWorld* InWorld, const TCHAR* Cmd, FOutputDevice& Ar )
+bool FDebugToolExec::Exec_Editor( UWorld* InWorld, const TCHAR* Cmd, FOutputDevice& Ar )
 {
 	// these commands are only allowed in standalone games
 #if UE_BUILD_SHIPPING || UE_BUILD_TEST
diff --git a/Engine/Source/Editor/UnrealEd/Private/EditorServer.cpp b/Engine/Source/Editor/UnrealEd/Private/EditorServer.cpp
--- a/Engine/Source/Editor/UnrealEd/Private/EditorServer.cpp
+++ b/Engine/Source/Editor/UnrealEd/Private/EditorServer.cpp
@@ -5689,7 +5689,7 @@ void ListMapPackageDependencies(const TCHAR* InStr)
 
 COREUOBJECT_API void DumpClassSchemas(const TCHAR* Str, FOutputDevice& Ar);
 
-bool UEditorEngine::Exec( UWorld* InWorld, const TCHAR* Stream, FOutputDevice& Ar )
+bool UEditorEngine::Exec_Editor( UWorld* InWorld, const TCHAR* Stream, FOutputDevice& Ar )
 {
 	TCHAR ErrorTemp[256]=TEXT("Setup: ");
 	bool bProcessed=false;
@@ -5954,7 +5954,7 @@ bool UEditorEngine::Exec( UWorld* InWorld, const TCHAR* Stream, FOutputDevice& A
 		// The level handled it.
 		bProcessed = true;
 	}
-	else if( UEngine::Exec( InWorld, Stream, Ar ) )
+	else if( UEngine::Exec_Editor( InWorld, Stream, Ar ) )
 	{
 		// The engine handled it.
 		bProcessed = true;
diff --git a/Engine/Source/Editor/UnrealEd/Private/UnrealEdSrv.cpp b/Engine/Source/Editor/UnrealEd/Private/UnrealEdSrv.cpp
--- a/Engine/Source/Editor/UnrealEd/Private/UnrealEdSrv.cpp
+++ b/Engine/Source/Editor/UnrealEd/Private/UnrealEdSrv.cpp
@@ -629,7 +629,7 @@ bool UUnrealEdEngine::HandleDisasmScriptCommand(const TCHAR* Str, FOutputDevice&
 	return true;
 }
 
-bool UUnrealEdEngine::Exec( UWorld* InWorld, const TCHAR* Stream, FOutputDevice& Ar )
+bool UUnrealEdEngine::Exec_Editor( UWorld* InWorld, const TCHAR* Stream, FOutputDevice& Ar )
 {
 	const TCHAR* Str = Stream;
 	// disallow set commands in the editor as that modifies the default object, affecting object serialization
@@ -658,7 +658,7 @@ bool UUnrealEdEngine::Exec( UWorld* InWorld, const TCHAR* Stream, FOutputDevice&
 		GStatProcessingViewportClient = GLastKeyLevelEditingViewportClient ? GLastKeyLevelEditingViewportClient : GCurrentLevelEditingViewportClient;
 	}
 
-	bool bExecSucceeded = UEditorEngine::Exec( InWorld, Stream, Ar );
+	bool bExecSucceeded = UEditorEngine::Exec_Editor( InWorld, Stream, Ar );
 
 	GStatProcessingViewportClient = NULL;
 
diff --git a/Engine/Source/Editor/UnrealEd/Public/DebugToolExec.h b/Engine/Source/Editor/UnrealEd/Public/DebugToolExec.h
--- a/Engine/Source/Editor/UnrealEd/Public/DebugToolExec.h
+++ b/Engine/Source/Editor/UnrealEd/Public/DebugToolExec.h
@@ -22,11 +22,9 @@ protected:
 	 * @param bShouldShowNonEditable	whether to show properties that are normally not editable under "None"
 	 */
 	void EditObject(UObject* Object, bool bShouldShowNonEditable);
-
-public:
 	
 	//~ Begin Exec Interface
-	UNREALED_API virtual bool Exec( UWorld* InWorld, const TCHAR* Cmd, FOutputDevice& Ar ) override;
+	virtual bool Exec_Editor( UWorld* InWorld, const TCHAR* Cmd, FOutputDevice& Ar ) override;
 	//~ End Exec Interface
 
 };
diff --git a/Engine/Source/Runtime/Engine/Classes/Engine/LocalPlayer.h b/Engine/Source/Runtime/Engine/Classes/Engine/LocalPlayer.h
--- a/Engine/Source/Runtime/Engine/Classes/Engine/LocalPlayer.h
+++ b/Engine/Source/Runtime/Engine/Classes/Engine/LocalPlayer.h
@@ -236,9 +236,15 @@ public:
 	// End of UObject interface
 
 	// FExec interface
+public:
+#if UE_ALLOW_EXEC_COMMANDS
 	virtual bool Exec(UWorld* InWorld, const TCHAR* Cmd,FOutputDevice& Ar) override;
+#endif
+protected:
+	virtual bool Exec_Editor(UWorld* InWorld, const TCHAR* Cmd, FOutputDevice& Ar) override;
 	// End of FExec interface
 
+public:
 	/** 
 	 * Exec command handlers
 	 */
diff --git a/Engine/Source/Runtime/Engine/Private/LocalPlayer.cpp b/Engine/Source/Runtime/Engine/Private/LocalPlayer.cpp
--- a/Engine/Source/Runtime/Engine/Private/LocalPlayer.cpp
+++ b/Engine/Source/Runtime/Engine/Private/LocalPlayer.cpp
@@ -1451,37 +1451,9 @@ bool ULocalPlayer::HandleToggleStreamingVolumesCommand( const TCHAR* Cmd, FOutpu
 	return true;
 }
 
+#if UE_ALLOW_EXEC_COMMANDS
 bool ULocalPlayer::Exec(UWorld* InWorld, const TCHAR* Cmd,FOutputDevice& Ar)
 {
-#if WITH_EDITOR
-	if (GIsEditor)
-	{
-		// Override a few commands in PIE
-		if( FParse::Command(&Cmd,TEXT("DN")) )
-		{
-			return HandleDNCommand( Cmd, Ar );
-		}
-
-		if( FParse::Command(&Cmd,TEXT("Exit"))
-		||	FParse::Command(&Cmd,TEXT("Quit")))
-		{
-			return HandleExitCommand( Cmd, Ar );
-		}
-
-		if( FParse::Command(&Cmd,TEXT("FocusNextPIEWindow")))
-		{
-			GEngine->FocusNextPIEWorld(InWorld);
-			return true;
-		}
-		if( FParse::Command(&Cmd,TEXT("FocusLastPIEWindow")))
-		{
-			GEngine->FocusNextPIEWorld(InWorld, true);
-			return true;
-		}
-
-	}
-#endif // WITH_EDITOR
-
 // NOTE: all of these can probably be #if !(UE_BUILD_SHIPPING || UE_BUILD_TEST) out
 
 	if( FParse::Command(&Cmd,TEXT("LISTMOVEBODY")) )
@@ -1576,6 +1548,45 @@ bool ULocalPlayer::Exec(UWorld* InWorld, const TCHAR* Cmd,FOutputDevice& Ar)
 		return false;
 	}
 }
+#endif // UE_ALLOW_EXEC_COMMANDS
+
+bool ULocalPlayer::Exec_Editor(UWorld* InWorld, const TCHAR* Cmd, FOutputDevice& Ar)
+{
+#if WITH_EDITOR
+	if (GIsEditor)
+	{
+		// Override a few commands in PIE
+		if (FParse::Command(&Cmd, TEXT("DN")))
+		{
+			return HandleDNCommand(Cmd, Ar);
+		}
+
+		if( FParse::Command(&Cmd,TEXT("Exit"))
+		||	FParse::Command(&Cmd,TEXT("Quit")))
+		{
+			return HandleExitCommand( Cmd, Ar );
+		}
+
+		if (FParse::Command(&Cmd, TEXT("FocusNextPIEWindow")))
+		{
+			GEngine->FocusNextPIEWorld(InWorld);
+			return true;
+		}
+		if (FParse::Command(&Cmd, TEXT("FocusLastPIEWindow")))
+		{
+			GEngine->FocusNextPIEWorld(InWorld, true);
+			return true;
+		}
+
+		if (Super::Exec_Editor(InWorld, Cmd, Ar))
+		{
+			return true;
+		}
+	}
+#endif // WITH_EDITOR
+
+	return false;
+}
 
 void ULocalPlayer::ExecMacro( const TCHAR* Filename, FOutputDevice& Ar )
 {
diff --git a/Engine/Source/Runtime/Engine/Private/Materials/MaterialShared.cpp b/Engine/Source/Runtime/Engine/Private/Materials/MaterialShared.cpp
--- a/Engine/Source/Runtime/Engine/Private/Materials/MaterialShared.cpp
+++ b/Engine/Source/Runtime/Engine/Private/Materials/MaterialShared.cpp
@@ -117,7 +117,7 @@ static FAutoConsoleCommand GFlushMaterialUniforms(
 #if WITH_EDITOR
 class FMaterialDumpDebugInfoExecHelper : public FSelfRegisteringExec
 {
-	virtual bool Exec(class UWorld* InWorld, const TCHAR* Cmd, FOutputDevice& Ar)
+	virtual bool Exec_Editor(class UWorld* InWorld, const TCHAR* Cmd, FOutputDevice& Ar) override
 	{
 		if (FParse::Command(&Cmd, TEXT("material dumpdebuginfo")))
 		{
diff --git a/Engine/Source/Runtime/Engine/Private/Player.cpp b/Engine/Source/Runtime/Engine/Private/Player.cpp
--- a/Engine/Source/Runtime/Engine/Private/Player.cpp
+++ b/Engine/Source/Runtime/Engine/Private/Player.cpp
@@ -93,6 +93,12 @@ APlayerController* UPlayer::GetPlayerController(const UWorld* const InWorld) con
 
 bool UPlayer::Exec( UWorld* InWorld, const TCHAR* Cmd,FOutputDevice& Ar)
 {
+	// Route through Exec_Dev and Exec_Editor first
+	if (FExec::Exec(InWorld, Cmd, Ar))
+	{
+		return true;
+	}
+
 	AActor* ExecActor = PlayerController;
 	if (!ExecActor)
 	{
