Added PRIMITIVE_SCENE_DATA_FLAG_HAS_DISTANCE_FIELD_REPRESENTATION to GPU scene data

diff --git a/Engine/Plugins/FX/Niagara/Source/Niagara/Private/NiagaraComponent.cpp b/Engine/Plugins/FX/Niagara/Source/Niagara/Private/NiagaraComponent.cpp
--- a/Engine/Plugins/FX/Niagara/Source/Niagara/Private/NiagaraComponent.cpp
+++ b/Engine/Plugins/FX/Niagara/Source/Niagara/Private/NiagaraComponent.cpp
@@ -346,6 +346,7 @@ TUniformBuffer<FPrimitiveUniformShaderParameters>* FNiagaraSceneProxy::GetCustom
 				.CastContactShadow(CastsContactShadow())
 				.CastShadow(CastsDynamicShadow())
 				.Holdout(Holdout())
+				.HasDistanceFieldRepresentation(HasDistanceFieldRepresentation())
 				.HasCapsuleRepresentation(HasDynamicIndirectShadowCasterRepresentation())
 				.UseVolumetricLightmap(bHasPrecomputedVolumetricLightmap)
 				.UseSingleSampleShadowFromStationaryLights(UseSingleSampleShadowFromStationaryLights());
diff --git a/Engine/Shaders/Private/SceneData.ush b/Engine/Shaders/Private/SceneData.ush
--- a/Engine/Shaders/Private/SceneData.ush
+++ b/Engine/Shaders/Private/SceneData.ush
@@ -67,29 +67,30 @@
 #define PRIMITIVE_SCENE_DATA_FLAG_CACHE_SHADOW_AS_STATIC				0x10
 #define PRIMITIVE_SCENE_DATA_FLAG_OUTPUT_VELOCITY						0x20
 #define PRIMITIVE_SCENE_DATA_FLAG_DETERMINANT_SIGN						0x40
-#define PRIMITIVE_SCENE_DATA_FLAG_HAS_CAPSULE_REPRESENTATION			0x80
-#define PRIMITIVE_SCENE_DATA_FLAG_HAS_CAST_CONTACT_SHADOW				0x100
-#define PRIMITIVE_SCENE_DATA_FLAG_HAS_PRIMITIVE_CUSTOM_DATA				0x200
-#define PRIMITIVE_SCENE_DATA_FLAG_LIGHTING_CHANNEL_0					0x400
-#define PRIMITIVE_SCENE_DATA_FLAG_LIGHTING_CHANNEL_1					0x800
-#define PRIMITIVE_SCENE_DATA_FLAG_LIGHTING_CHANNEL_2					0x1000
-#define PRIMITIVE_SCENE_DATA_FLAG_HAS_INSTANCE_LOCAL_BOUNDS				0x2000
-#define PRIMITIVE_SCENE_DATA_FLAG_HAS_NANITE_IMPOSTER					0x4000
-#define PRIMITIVE_SCENE_DATA_FLAG_VISIBLE_IN_GAME						0x8000
-#define PRIMITIVE_SCENE_DATA_FLAG_VISIBLE_IN_EDITOR						0x10000
-#define PRIMITIVE_SCENE_DATA_FLAG_VISIBLE_IN_REFLECTION_CAPTURES		0x20000
-#define PRIMITIVE_SCENE_DATA_FLAG_VISIBLE_IN_REAL_TIME_SKY_CAPTURES		0x40000
-#define PRIMITIVE_SCENE_DATA_FLAG_VISIBLE_IN_RAY_TRACING				0x80000
-#define PRIMITIVE_SCENE_DATA_FLAG_VISIBLE_IN_SCENE_CAPTURE_ONLY			0x100000
-#define PRIMITIVE_SCENE_DATA_FLAG_HIDDEN_IN_SCENE_CAPTURE				0x200000
-#define PRIMITIVE_SCENE_DATA_FLAG_FORCE_HIDDEN							0x400000
-#define PRIMITIVE_SCENE_DATA_FLAG_CAST_HIDDEN_SHADOW					0x800000
-#define PRIMITIVE_SCENE_DATA_FLAG_EVALUATE_WORLD_POSITION_OFFSET		0x1000000
-#define PRIMITIVE_SCENE_DATA_FLAG_INSTANCE_DRAW_DISTANCE_CULL			0x2000000
-#define PRIMITIVE_SCENE_DATA_FLAG_WPO_DISABLE_DISTANCE					0x4000000
-#define PRIMITIVE_SCENE_DATA_FLAG_WRITES_CUSTOM_DEPTH_STENCIL			0x8000000
-#define PRIMITIVE_SCENE_DATA_FLAG_REVERSE_CULLING						0x10000000
-#define PRIMITIVE_SCENE_DATA_FLAG_HOLDOUT                               0x20000000
+#define PRIMITIVE_SCENE_DATA_FLAG_HAS_DISTANCE_FIELD_REPRESENTATION		0x80
+#define PRIMITIVE_SCENE_DATA_FLAG_HAS_CAPSULE_REPRESENTATION			0x100
+#define PRIMITIVE_SCENE_DATA_FLAG_HAS_CAST_CONTACT_SHADOW				0x200
+#define PRIMITIVE_SCENE_DATA_FLAG_HAS_PRIMITIVE_CUSTOM_DATA				0x400
+#define PRIMITIVE_SCENE_DATA_FLAG_LIGHTING_CHANNEL_0					0x800
+#define PRIMITIVE_SCENE_DATA_FLAG_LIGHTING_CHANNEL_1					0x1000
+#define PRIMITIVE_SCENE_DATA_FLAG_LIGHTING_CHANNEL_2					0x2000
+#define PRIMITIVE_SCENE_DATA_FLAG_HAS_INSTANCE_LOCAL_BOUNDS				0x4000
+#define PRIMITIVE_SCENE_DATA_FLAG_HAS_NANITE_IMPOSTER					0x8000
+#define PRIMITIVE_SCENE_DATA_FLAG_VISIBLE_IN_GAME						0x10000
+#define PRIMITIVE_SCENE_DATA_FLAG_VISIBLE_IN_EDITOR						0x20000
+#define PRIMITIVE_SCENE_DATA_FLAG_VISIBLE_IN_REFLECTION_CAPTURES		0x40000
+#define PRIMITIVE_SCENE_DATA_FLAG_VISIBLE_IN_REAL_TIME_SKY_CAPTURES		0x80000
+#define PRIMITIVE_SCENE_DATA_FLAG_VISIBLE_IN_RAY_TRACING				0x100000
+#define PRIMITIVE_SCENE_DATA_FLAG_VISIBLE_IN_SCENE_CAPTURE_ONLY			0x200000
+#define PRIMITIVE_SCENE_DATA_FLAG_HIDDEN_IN_SCENE_CAPTURE				0x400000
+#define PRIMITIVE_SCENE_DATA_FLAG_FORCE_HIDDEN							0x800000
+#define PRIMITIVE_SCENE_DATA_FLAG_CAST_HIDDEN_SHADOW					0x1000000
+#define PRIMITIVE_SCENE_DATA_FLAG_EVALUATE_WORLD_POSITION_OFFSET		0x2000000
+#define PRIMITIVE_SCENE_DATA_FLAG_INSTANCE_DRAW_DISTANCE_CULL			0x4000000
+#define PRIMITIVE_SCENE_DATA_FLAG_WPO_DISABLE_DISTANCE					0x8000000
+#define PRIMITIVE_SCENE_DATA_FLAG_WRITES_CUSTOM_DEPTH_STENCIL			0x10000000
+#define PRIMITIVE_SCENE_DATA_FLAG_REVERSE_CULLING						0x20000000
+#define PRIMITIVE_SCENE_DATA_FLAG_HOLDOUT								0x40000000
 
 uint GetDefaultPrimitiveSceneDataFlags()
 {
diff --git a/Engine/Source/Runtime/Engine/Private/PrimitiveSceneProxy.cpp b/Engine/Source/Runtime/Engine/Private/PrimitiveSceneProxy.cpp
--- a/Engine/Source/Runtime/Engine/Private/PrimitiveSceneProxy.cpp
+++ b/Engine/Source/Runtime/Engine/Private/PrimitiveSceneProxy.cpp
@@ -538,6 +538,7 @@ void FPrimitiveSceneProxy::UpdateUniformBuffer()
 				.LightmapUVIndex(GetLightMapCoordinateIndex())
 				.SingleCaptureIndex(SingleCaptureIndex)
 				.CustomPrimitiveData(GetCustomPrimitiveData())
+				.HasDistanceFieldRepresentation(HasDistanceFieldRepresentation())
 				.HasCapsuleRepresentation(HasDynamicIndirectShadowCasterRepresentation())
 				.UseSingleSampleShadowFromStationaryLights(UseSingleSampleShadowFromStationaryLights())
 				.UseVolumetricLightmap(bHasPrecomputedVolumetricLightmap)
diff --git a/Engine/Source/Runtime/Engine/Private/PrimitiveUniformShaderParameters.cpp b/Engine/Source/Runtime/Engine/Private/PrimitiveUniformShaderParameters.cpp
--- a/Engine/Source/Runtime/Engine/Private/PrimitiveUniformShaderParameters.cpp
+++ b/Engine/Source/Runtime/Engine/Private/PrimitiveUniformShaderParameters.cpp
@@ -176,6 +176,7 @@ FPrimitiveSceneShaderData::FPrimitiveSceneShaderData(const FPrimitiveSceneProxy*
 		.NumInstanceSceneDataEntries(Proxy->GetPrimitiveSceneInfo()->GetNumInstanceSceneDataEntries())
 		.InstancePayloadDataOffset(Proxy->GetPrimitiveSceneInfo()->GetInstancePayloadDataOffset())
 		.InstancePayloadDataStride(Proxy->GetPrimitiveSceneInfo()->GetInstancePayloadDataStride())
+		.HasDistanceFieldRepresentation(Proxy->HasDistanceFieldRepresentation())
 		.HasCapsuleRepresentation(Proxy->HasDynamicIndirectShadowCasterRepresentation())
 		.UseSingleSampleShadowFromStationaryLights(Proxy->UseSingleSampleShadowFromStationaryLights())
 		.ReceivesDecals(Proxy->ReceivesDecals())
diff --git a/Engine/Source/Runtime/Engine/Public/PrimitiveUniformShaderParameters.h b/Engine/Source/Runtime/Engine/Public/PrimitiveUniformShaderParameters.h
--- a/Engine/Source/Runtime/Engine/Public/PrimitiveUniformShaderParameters.h
+++ b/Engine/Source/Runtime/Engine/Public/PrimitiveUniformShaderParameters.h
@@ -74,29 +74,30 @@ END_GLOBAL_SHADER_PARAMETER_STRUCT()
 #define PRIMITIVE_SCENE_DATA_FLAG_CACHE_SHADOW_AS_STATIC				0x10
 #define PRIMITIVE_SCENE_DATA_FLAG_OUTPUT_VELOCITY						0x20
 #define PRIMITIVE_SCENE_DATA_FLAG_DETERMINANT_SIGN						0x40
-#define PRIMITIVE_SCENE_DATA_FLAG_HAS_CAPSULE_REPRESENTATION			0x80
-#define PRIMITIVE_SCENE_DATA_FLAG_HAS_CAST_CONTACT_SHADOW				0x100
-#define PRIMITIVE_SCENE_DATA_FLAG_HAS_PRIMITIVE_CUSTOM_DATA				0x200
-#define PRIMITIVE_SCENE_DATA_FLAG_LIGHTING_CHANNEL_0					0x400
-#define PRIMITIVE_SCENE_DATA_FLAG_LIGHTING_CHANNEL_1					0x800
-#define PRIMITIVE_SCENE_DATA_FLAG_LIGHTING_CHANNEL_2					0x1000
-#define PRIMITIVE_SCENE_DATA_FLAG_HAS_INSTANCE_LOCAL_BOUNDS				0x2000
-#define PRIMITIVE_SCENE_DATA_FLAG_HAS_NANITE_IMPOSTER					0x4000
-#define PRIMITIVE_SCENE_DATA_FLAG_VISIBLE_IN_GAME						0x8000
-#define PRIMITIVE_SCENE_DATA_FLAG_VISIBLE_IN_EDITOR						0x10000
-#define PRIMITIVE_SCENE_DATA_FLAG_VISIBLE_IN_REFLECTION_CAPTURES		0x20000
-#define PRIMITIVE_SCENE_DATA_FLAG_VISIBLE_IN_REAL_TIME_SKY_CAPTURES		0x40000
-#define PRIMITIVE_SCENE_DATA_FLAG_VISIBLE_IN_RAY_TRACING				0x80000
-#define PRIMITIVE_SCENE_DATA_FLAG_VISIBLE_IN_SCENE_CAPTURE_ONLY			0x100000
-#define PRIMITIVE_SCENE_DATA_FLAG_HIDDEN_IN_SCENE_CAPTURE				0x200000
-#define PRIMITIVE_SCENE_DATA_FLAG_FORCE_HIDDEN							0x400000
-#define PRIMITIVE_SCENE_DATA_FLAG_CAST_HIDDEN_SHADOW					0x800000
-#define PRIMITIVE_SCENE_DATA_FLAG_EVALUATE_WORLD_POSITION_OFFSET		0x1000000
-#define PRIMITIVE_SCENE_DATA_FLAG_INSTANCE_DRAW_DISTANCE_CULL			0x2000000
-#define PRIMITIVE_SCENE_DATA_FLAG_WPO_DISABLE_DISTANCE					0x4000000
-#define PRIMITIVE_SCENE_DATA_FLAG_WRITES_CUSTOM_DEPTH_STENCIL			0x8000000
-#define PRIMITIVE_SCENE_DATA_FLAG_REVERSE_CULLING						0x10000000
-#define PRIMITIVE_SCENE_DATA_FLAG_HOLDOUT                               0x20000000
+#define PRIMITIVE_SCENE_DATA_FLAG_HAS_DISTANCE_FIELD_REPRESENTATION		0x80
+#define PRIMITIVE_SCENE_DATA_FLAG_HAS_CAPSULE_REPRESENTATION			0x100
+#define PRIMITIVE_SCENE_DATA_FLAG_HAS_CAST_CONTACT_SHADOW				0x200
+#define PRIMITIVE_SCENE_DATA_FLAG_HAS_PRIMITIVE_CUSTOM_DATA				0x400
+#define PRIMITIVE_SCENE_DATA_FLAG_LIGHTING_CHANNEL_0					0x800
+#define PRIMITIVE_SCENE_DATA_FLAG_LIGHTING_CHANNEL_1					0x1000
+#define PRIMITIVE_SCENE_DATA_FLAG_LIGHTING_CHANNEL_2					0x2000
+#define PRIMITIVE_SCENE_DATA_FLAG_HAS_INSTANCE_LOCAL_BOUNDS				0x4000
+#define PRIMITIVE_SCENE_DATA_FLAG_HAS_NANITE_IMPOSTER					0x8000
+#define PRIMITIVE_SCENE_DATA_FLAG_VISIBLE_IN_GAME						0x10000
+#define PRIMITIVE_SCENE_DATA_FLAG_VISIBLE_IN_EDITOR						0x20000
+#define PRIMITIVE_SCENE_DATA_FLAG_VISIBLE_IN_REFLECTION_CAPTURES		0x40000
+#define PRIMITIVE_SCENE_DATA_FLAG_VISIBLE_IN_REAL_TIME_SKY_CAPTURES		0x80000
+#define PRIMITIVE_SCENE_DATA_FLAG_VISIBLE_IN_RAY_TRACING				0x100000
+#define PRIMITIVE_SCENE_DATA_FLAG_VISIBLE_IN_SCENE_CAPTURE_ONLY			0x200000
+#define PRIMITIVE_SCENE_DATA_FLAG_HIDDEN_IN_SCENE_CAPTURE				0x400000
+#define PRIMITIVE_SCENE_DATA_FLAG_FORCE_HIDDEN							0x800000
+#define PRIMITIVE_SCENE_DATA_FLAG_CAST_HIDDEN_SHADOW					0x1000000
+#define PRIMITIVE_SCENE_DATA_FLAG_EVALUATE_WORLD_POSITION_OFFSET		0x2000000
+#define PRIMITIVE_SCENE_DATA_FLAG_INSTANCE_DRAW_DISTANCE_CULL			0x4000000
+#define PRIMITIVE_SCENE_DATA_FLAG_WPO_DISABLE_DISTANCE					0x8000000
+#define PRIMITIVE_SCENE_DATA_FLAG_WRITES_CUSTOM_DEPTH_STENCIL			0x10000000
+#define PRIMITIVE_SCENE_DATA_FLAG_REVERSE_CULLING						0x20000000
+#define PRIMITIVE_SCENE_DATA_FLAG_HOLDOUT								0x40000000
 
 struct FPrimitiveUniformShaderParametersBuilder
 {
@@ -121,6 +122,7 @@ public:
 		bUseVolumetricLightmap						= false;
 		bCacheShadowAsStatic						= false;
 		bOutputVelocity								= false;
+		bHasDistanceFieldRepresentation				= false;
 		bHasCapsuleRepresentation					= false;
 		bHasPreSkinnedLocalBounds					= false;
 		bHasPreviousLocalToWorld					= false;
@@ -173,6 +175,7 @@ public:
 	inline FPrimitiveUniformShaderParametersBuilder& VARIABLE_NAME(INPUT_TYPE In##VARIABLE_NAME) { b##VARIABLE_NAME = In##VARIABLE_NAME; return *this; }
 
 	PRIMITIVE_UNIFORM_BUILDER_FLAG_METHOD(bool,			ReceivesDecals);
+	PRIMITIVE_UNIFORM_BUILDER_FLAG_METHOD(bool,			HasDistanceFieldRepresentation);
 	PRIMITIVE_UNIFORM_BUILDER_FLAG_METHOD(bool,			HasCapsuleRepresentation);
 	PRIMITIVE_UNIFORM_BUILDER_FLAG_METHOD(bool,			HasInstanceLocalBounds);
 	PRIMITIVE_UNIFORM_BUILDER_FLAG_METHOD(bool,			CastContactShadow);
@@ -424,6 +427,7 @@ public:
 
 		Parameters.Flags = 0;
 		Parameters.Flags |= bReceivesDecals ? PRIMITIVE_SCENE_DATA_FLAG_DECAL_RECEIVER : 0u;
+		Parameters.Flags |= bHasDistanceFieldRepresentation ? PRIMITIVE_SCENE_DATA_FLAG_HAS_DISTANCE_FIELD_REPRESENTATION : 0u;
 		Parameters.Flags |= bHasCapsuleRepresentation ? PRIMITIVE_SCENE_DATA_FLAG_HAS_CAPSULE_REPRESENTATION : 0u;
 		Parameters.Flags |= bUseSingleSampleShadowFromStationaryLights ? PRIMITIVE_SCENE_DATA_FLAG_USE_SINGLE_SAMPLE_SHADOW_SL : 0u;
 		Parameters.Flags |= (bUseVolumetricLightmap && bUseSingleSampleShadowFromStationaryLights) ? PRIMITIVE_SCENE_DATA_FLAG_USE_VOLUMETRIC_LM_SHADOW_SL : 0u;
@@ -478,6 +482,7 @@ private:
 	uint32 bCastShadow : 1;
 	uint32 bCastContactShadow : 1;
 	uint32 bCastHiddenShadow : 1;
+	uint32 bHasDistanceFieldRepresentation : 1;
 	uint32 bHasCapsuleRepresentation : 1;
 	uint32 bHasPreSkinnedLocalBounds : 1;
 	uint32 bHasInstanceLocalBounds : 1;
