World Partition: content bundle code cleanup.

diff --git a/Engine/Source/Runtime/Engine/Private/WorldPartition/WorldPartitionRuntimeHash.cpp b/Engine/Source/Runtime/Engine/Private/WorldPartition/WorldPartitionRuntimeHash.cpp
--- a/Engine/Source/Runtime/Engine/Private/WorldPartition/WorldPartitionRuntimeHash.cpp
+++ b/Engine/Source/Runtime/Engine/Private/WorldPartition/WorldPartitionRuntimeHash.cpp
@@ -9,11 +9,22 @@
 
 #include UE_INLINE_GENERATED_CPP_BY_NAME(WorldPartitionRuntimeHash)
 
-#if WITH_EDITOR
-#endif
-
 #define LOCTEXT_NAMESPACE "WorldPartition"
 
+void URuntimeHashExternalStreamingObjectBase::ForEachStreamingCells(TFunctionRef<void(UWorldPartitionRuntimeCell&)> Func)
+{
+	TArray<UObject*> Objects;
+	GetObjectsWithOuter(this, Objects);
+
+	for (UObject* Object : Objects)
+	{
+		if (UWorldPartitionRuntimeCell* Cell = Cast<UWorldPartitionRuntimeCell>(Object))
+		{
+			Func(*Cell);
+		}
+	}
+}
+
 UWorldPartitionRuntimeHash::UWorldPartitionRuntimeHash(const FObjectInitializer& ObjectInitializer)
 	: Super(ObjectInitializer)
 {}
diff --git a/Engine/Source/Runtime/Engine/Private/WorldPartition/WorldPartitionRuntimeSpatialHash.cpp b/Engine/Source/Runtime/Engine/Private/WorldPartition/WorldPartitionRuntimeSpatialHash.cpp
--- a/Engine/Source/Runtime/Engine/Private/WorldPartition/WorldPartitionRuntimeSpatialHash.cpp
+++ b/Engine/Source/Runtime/Engine/Private/WorldPartition/WorldPartitionRuntimeSpatialHash.cpp
@@ -975,24 +975,6 @@ void URuntimeSpatialHashExternalStreamingObject::OnStreamingObjectLoaded()
 	}
 }
 
-void URuntimeSpatialHashExternalStreamingObject::ForEachStreamingCells(TFunctionRef<void(UWorldPartitionRuntimeCell&)> Func)
-{
-	for (FSpatialHashStreamingGrid& StreamingGrid : StreamingGrids)
-	{
-		for (FSpatialHashStreamingGridLevel& GridLevel : StreamingGrid.GridLevels)
-		{
-			for (FSpatialHashStreamingGridLayerCell& GridLayerCell : GridLevel.LayerCells)
-			{
-				for (UWorldPartitionRuntimeSpatialHashCell* Cell : GridLayerCell.GridCells)
-				{
-					Func(*Cell);
-				}
-			}
-		}
-	}
-
-}
-
 #if WITH_EDITOR
 void URuntimeSpatialHashExternalStreamingObject::PopulateGeneratorPackageForCook()
 {
diff --git a/Engine/Source/Runtime/Engine/Public/WorldPartition/WorldPartitionRuntimeHash.h b/Engine/Source/Runtime/Engine/Public/WorldPartition/WorldPartitionRuntimeHash.h
--- a/Engine/Source/Runtime/Engine/Public/WorldPartition/WorldPartitionRuntimeHash.h
+++ b/Engine/Source/Runtime/Engine/Public/WorldPartition/WorldPartitionRuntimeHash.h
@@ -49,7 +49,7 @@ public:
 
 	virtual class UWorld* GetWorld() const override { return GetOwningWorld(); }
 
-	virtual void ForEachStreamingCells(TFunctionRef<void(UWorldPartitionRuntimeCell&)> Func) {};
+	virtual void ForEachStreamingCells(TFunctionRef<void(UWorldPartitionRuntimeCell&)> Func);
 	
 	virtual void OnStreamingObjectLoaded() {};
 
@@ -141,7 +141,6 @@ public:
 	bool IsCellRelevantFor(bool bClientOnlyVisible) const;
 	EWorldPartitionStreamingPerformance GetStreamingPerformance(const TSet<const UWorldPartitionRuntimeCell*>& CellToActivate) const;
 
-	virtual URuntimeHashExternalStreamingObjectBase* StoreToExternalStreamingObject(UObject* Outer) { return nullptr; }
 	virtual bool InjectExternalStreamingObject(URuntimeHashExternalStreamingObjectBase* ExternalStreamingObject) { return false; }
 	virtual bool RemoveExternalStreamingObject(URuntimeHashExternalStreamingObjectBase* ExternalStreamingObject) { return false; }
 
diff --git a/Engine/Source/Runtime/Engine/Public/WorldPartition/WorldPartitionRuntimeSpatialHash.h b/Engine/Source/Runtime/Engine/Public/WorldPartition/WorldPartitionRuntimeSpatialHash.h
--- a/Engine/Source/Runtime/Engine/Public/WorldPartition/WorldPartitionRuntimeSpatialHash.h
+++ b/Engine/Source/Runtime/Engine/Public/WorldPartition/WorldPartitionRuntimeSpatialHash.h
@@ -238,8 +238,6 @@ public:
 	//~ Begin URuntimeHashExternalStreamingObjectBase Interface
 	virtual void OnStreamingObjectLoaded() override;
 
-	virtual void ForEachStreamingCells(TFunctionRef<void(UWorldPartitionRuntimeCell&)> Func) override;
-
 #if WITH_EDITOR
 	virtual void PopulateGeneratorPackageForCook() override;
 #endif
