Optim: remove linear search in ULevelStreaming::FindStreamingLevel for assert.

diff --git a/Engine/Source/Runtime/Engine/Private/LevelStreaming.cpp b/Engine/Source/Runtime/Engine/Private/LevelStreaming.cpp
--- a/Engine/Source/Runtime/Engine/Private/LevelStreaming.cpp
+++ b/Engine/Source/Runtime/Engine/Private/LevelStreaming.cpp
@@ -1217,21 +1217,6 @@ ULevelStreaming* ULevelStreaming::FindStreamingLevel(const ULevel* Level)
 			check(LevelAnnotation.LevelStreaming->GetLoadedLevel() == Level);
 			FoundLevelStreaming = LevelAnnotation.LevelStreaming;
 		}
-
-		if (!FoundLevelStreaming)
-		{
-			// fallback search
-			for (ULevelStreaming* CurStreamingLevel : Level->OwningWorld->GetStreamingLevels())
-			{
-				if (CurStreamingLevel && CurStreamingLevel->GetLoadedLevel() == Level)
-				{
-					FoundLevelStreaming = CurStreamingLevel;
-					break;
-				}
-			}
-			// we shouldn't have found a streaming level here
-			ensure(!FoundLevelStreaming);
-		}
 	}
 
 	return FoundLevelStreaming;
