Optim: avoid array pointer to TObjectPtr conversions in FStreamingLevelsToConsider.

diff --git a/Engine/Source/Runtime/Engine/Classes/Engine/World.h b/Engine/Source/Runtime/Engine/Classes/Engine/World.h
--- a/Engine/Source/Runtime/Engine/Classes/Engine/World.h
+++ b/Engine/Source/Runtime/Engine/Classes/Engine/World.h
@@ -815,7 +815,7 @@ private:
 
 public:
 
-	const TArray<ULevelStreaming*>& GetStreamingLevels() const { return StreamingLevels; }
+	const TArray<TObjectPtr<ULevelStreaming>>& GetStreamingLevels() const { return StreamingLevels; }
 
 	void AddReferencedObjects(UObject* InThis, FReferenceCollector& Collector);
 
