Made DBuffer pass only usable when high quality normal is not used on non console.
This is because this pass writes highquality normal though UAV uint2 which is not supported by all graphics card on PC.

diff --git a/Engine/Source/Runtime/Engine/Private/ShaderCompiler/ShaderCompiler.cpp b/Engine/Source/Runtime/Engine/Private/ShaderCompiler/ShaderCompiler.cpp
--- a/Engine/Source/Runtime/Engine/Private/ShaderCompiler/ShaderCompiler.cpp
+++ b/Engine/Source/Runtime/Engine/Private/ShaderCompiler/ShaderCompiler.cpp
@@ -6412,8 +6412,14 @@ void GlobalBeginCompileShader(
 			const uint32 StrataTileCoord8Bits = CVarStrataTileCoord8Bits.Get(Target.GetPlatform());
 			Input.Environment.SetDefine(TEXT("USE_8BIT_TILE_COORD"), StrataTileCoord8Bits > 0 ? 1 : 0);
 
+			// DBuffer pass is only available if high quality normal is disabled, or if we are on a console.
+			// That is because in high quality normals requires a uint2 UAV which is only supported on some graphic cards on PC 
+			// and this is an unknown when compiling shaders at this stage.
+			// !!! If this is changed, please update all sites reading r.Strata.DBufferPass !!!
+			const bool bDBufferPassSupported = StrataNormalQuality == 0 || (StrataNormalQuality > 0 && IsConsolePlatform(Target.GetPlatform()));
+
 			static FShaderPlatformCachedIniValue<int32> CVarStrataDBufferPass(TEXT("r.Strata.DBufferPass"));
-			const uint32 StrataDBufferPass = CVarStrataDBufferPass.Get(Target.GetPlatform());
+			const uint32 StrataDBufferPass = bDBufferPassSupported ? CVarStrataDBufferPass.Get(Target.GetPlatform()) : 0;
 			Input.Environment.SetDefine(TEXT("STRATA_USE_DBUFFER_PASS"), StrataDBufferPass > 0 ? 1 : 0);
 		}
 
diff --git a/Engine/Source/Runtime/Engine/Private/StrataVisualizationData.cpp b/Engine/Source/Runtime/Engine/Private/StrataVisualizationData.cpp
--- a/Engine/Source/Runtime/Engine/Private/StrataVisualizationData.cpp
+++ b/Engine/Source/Runtime/Engine/Private/StrataVisualizationData.cpp
@@ -13,8 +13,17 @@ bool Engine_IsStrataEnabled();
 
 static bool Engine_IsStrataDBufferPassEnabled()
 {
+	static const auto CVarStrataGBufferFormat = IConsoleManager::Get().FindTConsoleVariableDataInt(TEXT("r.GBufferFormat")); 
+	const uint32 StrataNormalQuality = CVarStrataGBufferFormat && CVarStrataGBufferFormat->GetValueOnAnyThread() > 1 ? 1 : 0;
+
+	// DBuffer pass is only available if high quality normal is disabled, or if we are on a console.
+	// That is because in high quality normals requires a uint2 UAV which is only supported on some graphic cards on PC 
+	// and this is an unknown when compiling shaders at this stage.
+	// !!! If this is changed, please update all sites reading r.Strata.DBufferPass !!!
+	const bool bDBufferPassSupported = StrataNormalQuality == 0;	// We do not have access to platform to check if we are on console here. Visualize is likely used on PC editor so we can only visualize DBuffer if NormalQuality is 0.
+
 	static const auto CVar = IConsoleManager::Get().FindTConsoleVariableDataInt(TEXT("r.Strata.DBufferPass"));
-	return CVar && CVar->GetValueOnAnyThread() > 0;
+	return CVar && CVar->GetValueOnAnyThread() > 0 && bDBufferPassSupported;
 }
 
 static bool Engine_IsStrataRoughRefractionEnabled()
diff --git a/Engine/Source/Runtime/RenderCore/Private/Shader.cpp b/Engine/Source/Runtime/RenderCore/Private/Shader.cpp
--- a/Engine/Source/Runtime/RenderCore/Private/Shader.cpp
+++ b/Engine/Source/Runtime/RenderCore/Private/Shader.cpp
@@ -1942,6 +1942,15 @@ void ShaderMapAppendKeyString(EShaderPlatform Platform, FString& KeyString)
 
 		if (bStrataEnabled)
 		{
+			static const auto CVarStrataGBufferFormat = IConsoleManager::Get().FindTConsoleVariableDataInt(TEXT("r.GBufferFormat"));
+			const uint32 StrataNormalQuality = CVarStrataGBufferFormat && CVarStrataGBufferFormat->GetValueOnAnyThread() > 1 ? 1 : 0;
+
+			// DBuffer pass is only available if high quality normal is disabled, or if we are on a console.
+			// That is because in high quality normals requires a uint2 UAV which is only supported on some graphic cards on PC 
+			// and this is an unknown when compiling shaders at this stage.
+			// !!! If this is changed, please update all sites reading r.Strata.DBufferPass !!!
+			const bool bDBufferPassSupported = StrataNormalQuality == 0 || (StrataNormalQuality > 0 && IsConsolePlatform(Platform));
+
 			static FShaderPlatformCachedIniValue<int32> CVarDBufferPass(TEXT("r.Strata.DBufferPass"));
 			const int32 StrataDBufferPass = CVarDBufferPass.Get(Platform);
 			if (StrataDBufferPass>0)
diff --git a/Engine/Source/Runtime/Renderer/Private/Strata/Strata.cpp b/Engine/Source/Runtime/Renderer/Private/Strata/Strata.cpp
--- a/Engine/Source/Runtime/Renderer/Private/Strata/Strata.cpp
+++ b/Engine/Source/Runtime/Renderer/Private/Strata/Strata.cpp
@@ -99,8 +99,8 @@ static TAutoConsoleVariable<int32> CVarStrataAsyncClassification(
 static TAutoConsoleVariable<int32> CVarStrataDBufferPass(
 	TEXT("r.Strata.DBufferPass"),
 	0,
-	TEXT("Apply DBuffer after the base-pass as a separate pass."),
-	ECVF_RenderThreadSafe);
+	TEXT("Apply DBuffer after the base-pass as a separate pass. Read only because when this is changed, it will require the recompilation of all shaders."),
+	ECVF_ReadOnly | ECVF_RenderThreadSafe);
 
 static TAutoConsoleVariable<int32> CVarStrataDBufferPassDedicatedTiles(
 	TEXT("r.Strata.DBufferPass.DedicatedTiles"),
