Fix shipping editor build
#rnx

diff --git a/Engine/Source/Runtime/Engine/Private/UnrealEngine.cpp b/Engine/Source/Runtime/Engine/Private/UnrealEngine.cpp
--- a/Engine/Source/Runtime/Engine/Private/UnrealEngine.cpp
+++ b/Engine/Source/Runtime/Engine/Private/UnrealEngine.cpp
@@ -5083,7 +5083,7 @@ bool UEngine::Exec_Dev( UWorld* InWorld, const TCHAR* Cmd, FOutputDevice& Ar )
 
 bool UEngine::Exec_Editor(UWorld* InWorld, const TCHAR* Cmd, FOutputDevice& Ar)
 {
-#if WITH_EDITOR
+#if WITH_EDITOR && !UE_BUILD_SHIPPING
 	if (FParse::Command(&Cmd, TEXT("DumpShaderCompileStats")))
 	{
 		return HandleDumpShaderCompileStatsCommand(Cmd, Ar);
