Revert Exec changes to AudioDevice as they are incompatible with the default parameter override.

diff --git a/Engine/Source/Runtime/Engine/Private/AudioDevice.cpp b/Engine/Source/Runtime/Engine/Private/AudioDevice.cpp
--- a/Engine/Source/Runtime/Engine/Private/AudioDevice.cpp
+++ b/Engine/Source/Runtime/Engine/Private/AudioDevice.cpp
@@ -2211,7 +2211,7 @@ void FAudioDevice::SetMixDebugState(EDebugState InDebugState)
 	DebugState = InDebugState;
 }
 
-bool FAudioDevice::Exec_Dev(UWorld* InWorld, const TCHAR* Cmd, FOutputDevice& Ar)
+bool FAudioDevice::Exec(UWorld* InWorld, const TCHAR* Cmd, FOutputDevice& Ar)
 {
 #if !UE_BUILD_SHIPPING
 	auto ParseAudioExecCmd = [](const TCHAR** InCmd, const TCHAR* InMatch)
diff --git a/Engine/Source/Runtime/Engine/Public/AudioDevice.h b/Engine/Source/Runtime/Engine/Public/AudioDevice.h
--- a/Engine/Source/Runtime/Engine/Public/AudioDevice.h
+++ b/Engine/Source/Runtime/Engine/Public/AudioDevice.h
@@ -414,13 +414,12 @@ DECLARE_MULTICAST_DELEGATE_OneParam(FOnAudioDevicePostRender, const FAudioDevice
 
 class ENGINE_API FAudioDevice : public FExec
 {
-protected:
+public:
 
 	//Begin FExec Interface
-	virtual bool Exec_Dev(UWorld* InWorld, const TCHAR* Cmd, FOutputDevice& Ar = *GLog) override;
+	virtual bool Exec(UWorld* InWorld, const TCHAR* Cmd, FOutputDevice& Ar = *GLog) override;
 	//End FExec Interface
 
-public:
 #if !UE_BUILD_SHIPPING
 	UAudioComponent* GetTestComponent(UWorld* InWorld);
 	void StopTestComponent();
