Add proper material expression for object *local* bounds. This meant to replace the existing "hackish" material function /Engine/Functions/Engine_MaterialFunctions02/ObjectLocalBounds that uses hardcoded custom material expressions.

The reasoning for this replacement is two-fold:
1. Consistency with existing expressions for world bounds and pre-skinned local bounds.
2. Custom expressions can't easily be caught or overriden via FProxyMaterialCompiler.

(based on PR #8531)

diff --git a/Engine/Source/Runtime/Engine/Classes/Materials/MaterialExpressionObjectLocalBounds.h b/Engine/Source/Runtime/Engine/Classes/Materials/MaterialExpressionObjectLocalBounds.h
new file mode 100644
--- /dev/null
+++ b/Engine/Source/Runtime/Engine/Classes/Materials/MaterialExpressionObjectLocalBounds.h
@@ -0,0 +1,27 @@
+// Copyright Epic Games, Inc. All Rights Reserved.
+
+#pragma once
+
+#include "CoreMinimal.h"
+#include "UObject/ObjectMacros.h"
+#include "Materials/MaterialExpression.h"
+#include "MaterialExpressionObjectLocalBounds.generated.h"
+
+UCLASS(collapsecategories, hidecategories = Object)
+class UMaterialExpressionObjectLocalBounds : public UMaterialExpression
+{
+	GENERATED_UCLASS_BODY()
+
+#if WITH_EDITOR
+	//~ Begin UMaterialExpression Interface
+	virtual int32 Compile(class FMaterialCompiler* Compiler, int32 OutputIndex) override;
+	virtual void GetCaption(TArray<FString>& OutCaptions) const override;
+	virtual void GetConnectorToolTip(int32 InputIndex, int32 OutputIndex, TArray<FString>& OutToolTip) override;
+	virtual void GetExpressionToolTip(TArray<FString>& OutToolTip) override;
+	//~ End UMaterialExpression Interface
+
+#if WITH_EDITORONLY_DATA
+	TArray<FString> OutputToolTips;
+#endif
+#endif
+};
\ No newline at end of file
diff --git a/Engine/Source/Runtime/Engine/Private/Materials/HLSLMaterialTranslator.cpp b/Engine/Source/Runtime/Engine/Private/Materials/HLSLMaterialTranslator.cpp
--- a/Engine/Source/Runtime/Engine/Private/Materials/HLSLMaterialTranslator.cpp
+++ b/Engine/Source/Runtime/Engine/Private/Materials/HLSLMaterialTranslator.cpp
@@ -5428,6 +5428,25 @@ int32 FHLSLMaterialTranslator::ObjectBounds()
 	return AddInlinedCodeChunk(MCT_Float3, TEXT("float3(GetPrimitiveData(Parameters).ObjectBoundsX, GetPrimitiveData(Parameters).ObjectBoundsY, GetPrimitiveData(Parameters).ObjectBoundsZ)"));
 }
 
+int32 FHLSLMaterialTranslator::ObjectLocalBounds(int32 OutputIndex)
+{
+	switch (OutputIndex)
+	{
+	case 0: // Half extents
+		return AddInlinedCodeChunk(MCT_Float3, TEXT("((GetPrimitiveData(Parameters).LocalObjectBoundsMax - GetPrimitiveData(Parameters).LocalObjectBoundsMin) / 2.0f)"));
+	case 1: // Full extents
+		return AddInlinedCodeChunk(MCT_Float3, TEXT("(GetPrimitiveData(Parameters).LocalObjectBoundsMax - GetPrimitiveData(Parameters).LocalObjectBoundsMin)"));
+	case 2: // Min point
+		return GetPrimitiveProperty(MCT_Float3, TEXT("ObjectLocalBounds"), TEXT("LocalObjectBoundsMin"));
+	case 3: // Max point
+		return GetPrimitiveProperty(MCT_Float3, TEXT("ObjectLocalBounds"), TEXT("LocalObjectBoundsMax"));
+	default:
+		check(false);
+	}
+
+	return INDEX_NONE; 
+}
+
 int32 FHLSLMaterialTranslator::PreSkinnedLocalBounds(int32 OutputIndex)
 {
 	switch (OutputIndex)
diff --git a/Engine/Source/Runtime/Engine/Private/Materials/HLSLMaterialTranslator.h b/Engine/Source/Runtime/Engine/Private/Materials/HLSLMaterialTranslator.h
--- a/Engine/Source/Runtime/Engine/Private/Materials/HLSLMaterialTranslator.h
+++ b/Engine/Source/Runtime/Engine/Private/Materials/HLSLMaterialTranslator.h
@@ -828,6 +828,7 @@ protected:
 	virtual int32 ObjectRadius() override;
 	virtual int32 ObjectBounds() override;
 
+	virtual int32 ObjectLocalBounds(int32 OutputIndex) override;
 	virtual int32 PreSkinnedLocalBounds(int32 OutputIndex) override;
 
 	virtual int32 DistanceCullFade() override;
diff --git a/Engine/Source/Runtime/Engine/Private/Materials/MaterialExpressions.cpp b/Engine/Source/Runtime/Engine/Private/Materials/MaterialExpressions.cpp
--- a/Engine/Source/Runtime/Engine/Private/Materials/MaterialExpressions.cpp
+++ b/Engine/Source/Runtime/Engine/Private/Materials/MaterialExpressions.cpp
@@ -132,6 +132,7 @@
 #include "Materials/MaterialExpressionNoise.h"
 #include "Materials/MaterialExpressionNormalize.h"
 #include "Materials/MaterialExpressionObjectBounds.h"
+#include "Materials/MaterialExpressionObjectLocalBounds.h"
 #include "Materials/MaterialExpressionObjectOrientation.h"
 #include "Materials/MaterialExpressionObjectPositionWS.h"
 #include "Materials/MaterialExpressionObjectRadius.h"
@@ -12777,6 +12778,75 @@ void UMaterialExpressionObjectBounds::GetCaption(TArray<FString>& OutCaptions) c
 }
 #endif // WITH_EDITOR
 
+///////////////////////////////////////////////////////////////////////////////
+// UMaterialExpressionObjectLocalBounds
+///////////////////////////////////////////////////////////////////////////////
+UMaterialExpressionObjectLocalBounds::UMaterialExpressionObjectLocalBounds(const FObjectInitializer& ObjectInitializer)
+	: Super(ObjectInitializer)
+{
+#if WITH_EDITORONLY_DATA
+	// Structure to hold one-time initialization
+	struct FConstructorStatics
+	{
+		FText NAME_Vectors;
+		FConstructorStatics()
+			: NAME_Vectors(LOCTEXT("Vectors", "Vectors"))
+		{
+		}
+	};
+	static FConstructorStatics ConstructorStatics;
+
+	MenuCategories.Add(ConstructorStatics.NAME_Vectors);
+
+	bShaderInputData = true;
+	bShowOutputNameOnPin = true;
+
+	Outputs.Reset();
+	Outputs.Add(FExpressionOutput(TEXT("Half Extents"), 1, 1, 1, 1, 0));
+	OutputToolTips.Add("Half the extent (width, depth and height) of the object bounding box. In local space.");
+	Outputs.Add(FExpressionOutput(TEXT("Extents"), 1, 1, 1, 1, 0));
+	OutputToolTips.Add("Full extent (width, depth and height) of the object bounding box. Same as 2x Half Extents. In local space.");
+	Outputs.Add(FExpressionOutput(TEXT("Min"), 1, 1, 1, 1, 0));
+	OutputToolTips.Add("Minimum 3D point of the object bounding box. In local space.");
+	Outputs.Add(FExpressionOutput(TEXT("Max"), 1, 1, 1, 1, 0));
+	OutputToolTips.Add("Maximum 3D point of the object bounding box. In local space.");
+#endif
+}
+
+#if WITH_EDITOR
+int32 UMaterialExpressionObjectLocalBounds::Compile(class FMaterialCompiler* Compiler, int32 OutputIndex)
+{
+	if (Material && Material->MaterialDomain == MD_DeferredDecal)
+	{
+		return CompilerError(Compiler, TEXT("Expression not available in the deferred decal material domain."));
+	}
+
+	return Compiler->ObjectLocalBounds(OutputIndex);
+}
+
+void UMaterialExpressionObjectLocalBounds::GetCaption(TArray<FString>& OutCaptions) const
+{
+	OutCaptions.Add(TEXT("Object Local Bounds"));
+}
+
+void UMaterialExpressionObjectLocalBounds::GetConnectorToolTip(int32 InputIndex, int32 OutputIndex, TArray<FString>& OutToolTip) 
+{
+#if WITH_EDITORONLY_DATA
+	if (OutputIndex >= 0 && OutputIndex < OutputToolTips.Num())
+	{
+		ConvertToMultilineToolTip(OutputToolTips[OutputIndex], 40, OutToolTip);
+	}
+#endif // WITH_EDITORONLY_DATA
+}
+
+void UMaterialExpressionObjectLocalBounds::GetExpressionToolTip(TArray<FString>& OutToolTip) 
+{
+	ConvertToMultilineToolTip(TEXT("Returns various info about the object local bounding box."
+		"Usable in vertex or pixel shader (no need to pipe this through vertex interpolators)."
+		"Hover the output pins for more information."), 40, OutToolTip);
+}
+#endif // WITH_EDITOR
+
 ///////////////////////////////////////////////////////////////////////////////
 // UMaterialExpressionPreSkinnedLocalBounds
 ///////////////////////////////////////////////////////////////////////////////
diff --git a/Engine/Source/Runtime/Engine/Public/MaterialCompiler.h b/Engine/Source/Runtime/Engine/Public/MaterialCompiler.h
--- a/Engine/Source/Runtime/Engine/Public/MaterialCompiler.h
+++ b/Engine/Source/Runtime/Engine/Public/MaterialCompiler.h
@@ -333,6 +333,7 @@ public:
 	virtual int32 ObjectWorldPosition() = 0;
 	virtual int32 ObjectRadius() = 0;
 	virtual int32 ObjectBounds() = 0;
+	virtual int32 ObjectLocalBounds(int32 OutputIndex) = 0;
 	virtual int32 PreSkinnedLocalBounds(int32 OutputIndex) = 0;
 	virtual int32 DistanceCullFade() = 0;
 	virtual int32 ActorWorldPosition() = 0;
@@ -852,6 +853,7 @@ public:
 	virtual int32 ObjectWorldPosition() override { return Compiler->ObjectWorldPosition(); }
 	virtual int32 ObjectRadius() override { return Compiler->ObjectRadius(); }
 	virtual int32 ObjectBounds() override { return Compiler->ObjectBounds(); }
+	virtual int32 ObjectLocalBounds(int32 OutputIndex) override { return Compiler->ObjectLocalBounds(OutputIndex); }
 	virtual int32 PreSkinnedLocalBounds(int32 OutputIndex) override { return Compiler->PreSkinnedLocalBounds(OutputIndex); }
 	virtual int32 DistanceCullFade() override { return Compiler->DistanceCullFade(); }
 	virtual int32 ActorWorldPosition() override { return Compiler->ActorWorldPosition(); }
