[shaders] enable STB preprocessor by default

diff --git a/Engine/Source/Runtime/Engine/Private/ShaderCompiler/ShaderCompiler.cpp b/Engine/Source/Runtime/Engine/Private/ShaderCompiler/ShaderCompiler.cpp
--- a/Engine/Source/Runtime/Engine/Private/ShaderCompiler/ShaderCompiler.cpp
+++ b/Engine/Source/Runtime/Engine/Private/ShaderCompiler/ShaderCompiler.cpp
@@ -1187,7 +1187,7 @@ static TAutoConsoleVariable<int32> CVarShadersRemoveDeadCode(
 // ShaderMapAppendKeyString (in Shader.cpp)
 static TAutoConsoleVariable<bool> CVarShadersUseLegacyPreprocessor(
 	TEXT("r.Shaders.UseLegacyPreprocessor"),
-	true,
+	false,
 	TEXT("Executes shader preprocessing via the legacy MCPP preprocessor (instead of the new STB preprocessor).\n")
 	TEXT("\ttrue: Enabled - preprocess with MCPP\n")
 	TEXT("\tfalse: Disabled - preprocess with STB\n"),
