World Partition Streaming Generation Log:
- Don't log on build machines to avoid filling storage.
- Append milliseconds to timestamp to avoid log collisions.

diff --git a/Engine/Source/Runtime/Engine/Private/WorldPartition/WorldPartitionStreamingGeneration.cpp b/Engine/Source/Runtime/Engine/Private/WorldPartition/WorldPartitionStreamingGeneration.cpp
--- a/Engine/Source/Runtime/Engine/Private/WorldPartition/WorldPartitionStreamingGeneration.cpp
+++ b/Engine/Source/Runtime/Engine/Private/WorldPartition/WorldPartitionStreamingGeneration.cpp
@@ -828,16 +828,15 @@ public:
 
 	static TUniquePtr<FArchive> CreateDumpStateLogArchive(const TCHAR* Suffix)
 	{
-		FString StateLogOutputFilename = FPaths::ProjectSavedDir() / TEXT("WorldPartition") / FString::Printf(TEXT("StreamingGeneration-%s"), Suffix);
-		
 		if (!GIsBuildMachine)
 		{
-			StateLogOutputFilename += FString::Printf(TEXT("-%08x-%s"), FPlatformProcess::GetCurrentProcessId(), *FDateTime::Now().ToString());
+			FString StateLogOutputFilename = FPaths::ProjectSavedDir() / TEXT("WorldPartition") / FString::Printf(TEXT("StreamingGeneration-%s"), Suffix);
+			StateLogOutputFilename += FString::Printf(TEXT("-%08x-%s"), FPlatformProcess::GetCurrentProcessId(), *FDateTime::Now().ToIso8601().Replace(TEXT(":"), TEXT(".")));
+			StateLogOutputFilename += TEXT(".log");
+			return TUniquePtr<FArchive>(IFileManager::Get().CreateFileWriter(*StateLogOutputFilename));
 		}
 
-		StateLogOutputFilename += TEXT(".log");
-
-		return TUniquePtr<FArchive>(IFileManager::Get().CreateFileWriter(*StateLogOutputFilename));
+		return MakeUnique<FArchive>();
 	}
 
 	void DumpStateLog(FHierarchicalLogArchive& Ar)
