Iris
- Temporarily disabled size asssert for FNetHitResultNetSerializer as it differs between server and editor build

diff --git a/Engine/Source/Runtime/Engine/Private/Engine/HitResultNetSerializer.cpp b/Engine/Source/Runtime/Engine/Private/Engine/HitResultNetSerializer.cpp
--- a/Engine/Source/Runtime/Engine/Private/Engine/HitResultNetSerializer.cpp
+++ b/Engine/Source/Runtime/Engine/Private/Engine/HitResultNetSerializer.cpp
@@ -466,10 +466,11 @@ void FHitResultNetSerializer::FNetSerializerRegistryDelegates::OnPostFreezeNetSe
 
 		const UStruct* HitResultStruct = FHitResult::StaticStruct();
 
-		if (HitResultStruct->GetStructureSize() != 240 || HitResultStruct->GetMinAlignment() != 8)
-		{
-			LowLevelFatalError(TEXT("%s Size: %d Alignment: %d"), TEXT("FHitResult layout has changed. Need to update FHitResultNetSerializer."), HitResultStruct->GetStructureSize(), HitResultStruct->GetMinAlignment());
-		}
+		// Had do comment this out as the size differs between editor and non-editor builds.
+		//if (HitResultStruct->GetStructureSize() != 240 || HitResultStruct->GetMinAlignment() != 8)
+		//{
+		//	LowLevelFatalError(TEXT("%s Size: %d Alignment: %d"), TEXT("FHitResult layout has changed. Need to update FHitResultNetSerializer."), HitResultStruct->GetStructureSize(), HitResultStruct->GetMinAlignment());
+		//}
 
 		StructNetSerializerConfigForHitResult.StateDescriptor = FReplicationStateDescriptorBuilder::CreateDescriptorForStruct(HitResultStruct, Params);
 		const FReplicationStateDescriptor* Descriptor = StructNetSerializerConfigForHitResult.StateDescriptor.GetReference();
