Disable DeferSkeletalDynamicDataUpdateUntilGDME when using a mesh deformer.

diff --git a/Engine/Source/Runtime/Engine/Private/SkeletalRenderGPUSkin.cpp b/Engine/Source/Runtime/Engine/Private/SkeletalRenderGPUSkin.cpp
--- a/Engine/Source/Runtime/Engine/Private/SkeletalRenderGPUSkin.cpp
+++ b/Engine/Source/Runtime/Engine/Private/SkeletalRenderGPUSkin.cpp
@@ -2583,6 +2583,8 @@ void FDynamicSkelMeshObjectDataGPUSkin::InitDynamicSkelMeshObjectDataGPUSkin(
 	bIsSkinCacheAllowed = InMeshComponent ? InMeshComponent->IsSkinCacheAllowed(InLODIndex) : false;
 	bHasMeshDeformer = InMeshComponent ? InMeshComponent->GetMeshDeformerInstance() != nullptr && LODIndex <= InMeshComponent->GetMeshDeformerMaxLOD() : false;
 	bForceUpdateDynamicDataImmediately = InMeshComponent ? InMeshComponent->GetForceUpdateDynamicDataImmediately() : false;
+	// Force immediate update when using mesh deformer.
+	bForceUpdateDynamicDataImmediately |= bHasMeshDeformer;
 
 	if (bIsSkinCacheAllowed && InMeshObject->FeatureLevel == ERHIFeatureLevel::ES3_1)
 	{
