GC clusters/particles should remove themselves from the external acceleration structure as well when disabled

diff --git a/Engine/Source/Runtime/Engine/Private/PhysicsEngine/Experimental/PhysScene_Chaos.cpp b/Engine/Source/Runtime/Engine/Private/PhysicsEngine/Experimental/PhysScene_Chaos.cpp
--- a/Engine/Source/Runtime/Engine/Private/PhysicsEngine/Experimental/PhysScene_Chaos.cpp
+++ b/Engine/Source/Runtime/Engine/Private/PhysicsEngine/Experimental/PhysScene_Chaos.cpp
@@ -1769,16 +1769,24 @@ void FPhysScene_Chaos::OnSyncBodies(Chaos::FPhysicsSolverBase* Solver)
 
 		auto GeometryCollectionLambda = [this](FGeometryCollectionPhysicsProxy* Proxy)
 		{
-			// Use don't pass in anything here so we don't end up locking anything because we can assume the scene is already locked.
+			// Don't pass in anything here so we don't end up locking anything because we can assume the scene is already locked.
 			FLockedWritePhysicsObjectExternalInterface Interface = FPhysicsObjectExternalInterface::LockWrite({});
-			TArray<Chaos::FPhysicsObjectHandle> Handles = Proxy->GetAllPhysicsObjects().FilterByPredicate(
+
+			TArray<Chaos::FPhysicsObjectHandle> ActiveHandles = Proxy->GetAllPhysicsObjects().FilterByPredicate(
 				[&Interface](Chaos::FPhysicsObjectHandle Handle)
 				{
 					return !Interface->AreAllDisabled({ &Handle, 1 });
 				}
 			);
 
-			Interface->AddToSpatialAcceleration(Handles, GetSpacialAcceleration());
+			TArray<Chaos::FPhysicsObjectHandle> DisabledHandles = Proxy->GetAllPhysicsObjects().FilterByPredicate(
+				[&Interface](Chaos::FPhysicsObjectHandle Handle)
+				{
+					return Interface->AreAllDisabled({ &Handle, 1 });
+				}
+			);
+			Interface->AddToSpatialAcceleration(ActiveHandles, GetSpacialAcceleration());
+			Interface->RemoveFromSpatialAcceleration(DisabledHandles, GetSpacialAcceleration());
 		};
 
 		Solver->PullPhysicsStateForEachDirtyProxy_External(RigidLambda, ConstraintLambda, GeometryCollectionLambda);
diff --git a/Engine/Source/Runtime/Experimental/Chaos/Public/Chaos/PhysicsObjectInterface.h b/Engine/Source/Runtime/Experimental/Chaos/Public/Chaos/PhysicsObjectInterface.h
--- a/Engine/Source/Runtime/Experimental/Chaos/Public/Chaos/PhysicsObjectInterface.h
+++ b/Engine/Source/Runtime/Experimental/Chaos/Public/Chaos/PhysicsObjectInterface.h
@@ -115,6 +115,21 @@ namespace Chaos
 			}
 		}
 
+		template<typename TPayloadType, typename T, int d>
+		void RemoveFromSpatialAcceleration(TArrayView<FPhysicsObjectHandle> InObjects, ISpatialAcceleration<TPayloadType, T, d>* SpatialAcceleration)
+		{
+			if (!SpatialAcceleration)
+			{
+				return;
+			}
+
+			for (FPhysicsObjectHandle Handle : InObjects)
+			{
+				FAccelerationStructureHandle AccelerationHandle = this->CreateAccelerationStructureHandle(Handle);
+				SpatialAcceleration->RemoveElementFrom(AccelerationHandle, this->GetSpatialIndex(Handle));
+			}
+		}
+
 		friend class FPhysicsObjectInterface;
 	protected:
 		FWritePhysicsObjectInterface() = default;
