Fix an issue where some NetGUIDs never seem to finish loading when there are multiple GUIDs referring to the same package

diff --git a/Engine/Source/Runtime/Engine/Classes/Engine/PackageMapClient.h b/Engine/Source/Runtime/Engine/Classes/Engine/PackageMapClient.h
--- a/Engine/Source/Runtime/Engine/Classes/Engine/PackageMapClient.h
+++ b/Engine/Source/Runtime/Engine/Classes/Engine/PackageMapClient.h
@@ -267,12 +267,31 @@ private:
 	struct FPendingAsyncLoadRequest
 	{
 		FPendingAsyncLoadRequest(const FNetworkGUID InNetGUID, const double InRequestStartTime):
-			NetGUID(InNetGUID),
 			RequestStartTime(InRequestStartTime)
 		{
+			NetGUIDs.Add(InNetGUID);
 		}
 
-		FNetworkGUID NetGUID;
+		void Merge(const FPendingAsyncLoadRequest& Other)
+		{
+			for (FNetworkGUID OtherGUID : Other.NetGUIDs)
+			{
+				NetGUIDs.AddUnique(OtherGUID);
+			}
+
+#if CSV_PROFILER
+			bWasRequestedByOwnerOrPawn |= Other.bWasRequestedByOwnerOrPawn;
+#endif
+		}
+
+		void Merge(FNetworkGUID InNetGUID)
+		{
+			NetGUIDs.AddUnique(InNetGUID);
+		}
+
+		// Net GUIDs that requested loading for the same UPackage
+		TArray<FNetworkGUID, TInlineAllocator<2>> NetGUIDs;
+
 		double RequestStartTime;
 
 #if CSV_PROFILER
diff --git a/Engine/Source/Runtime/Engine/Private/PackageMapClient.cpp b/Engine/Source/Runtime/Engine/Private/PackageMapClient.cpp
--- a/Engine/Source/Runtime/Engine/Private/PackageMapClient.cpp
+++ b/Engine/Source/Runtime/Engine/Private/PackageMapClient.cpp
@@ -3123,18 +3123,21 @@ void FNetGUIDCache::ValidateAsyncLoadingPackage(FNetGuidCacheObject& CacheObject
 	// re-loaded, it will likely be assigned a new NetGUID (since the TWeakObjectPtr to the old package
 	// in the cache object would have gone stale). During replay fast-forward, it's possible
 	// to see the new NetGUID before the previous one has finished loading, so here we fix up
-	// PendingAsyncPackages to refer to the new NewGUID.
+	// PendingAsyncPackages to refer to the new NewGUID. Also keep track of all the GUIDs referring
+	// to the same package so their CacheObjects can be properly updated later.
 	FPendingAsyncLoadRequest& PendingLoadRequest = PendingAsyncLoadRequests[CacheObject.PathName];
-	if (PendingLoadRequest.NetGUID != NetGUID)
+
+	PendingLoadRequest.Merge(NetGUID);
+	CacheObject.bIsPending = true;
+
+	if (PendingLoadRequest.NetGUIDs.Last() != NetGUID)
 	{
 		UE_LOG(LogNetPackageMap, Log, TEXT("ValidateAsyncLoadingPackage: Already async loading package with a different NetGUID. Path: %s, original NetGUID: %s, new NetGUID: %s"),
-			*CacheObject.PathName.ToString(), *PendingLoadRequest.NetGUID.ToString(), *NetGUID.ToString());
+			*CacheObject.PathName.ToString(), *PendingLoadRequest.NetGUIDs.Last().ToString(), *NetGUID.ToString());
 
 PRAGMA_DISABLE_DEPRECATION_WARNINGS
 		PendingAsyncPackages[CacheObject.PathName] = NetGUID;
 PRAGMA_ENABLE_DEPRECATION_WARNINGS
-
-		PendingLoadRequest.NetGUID = NetGUID;
 	}
 	else
 	{
@@ -3161,11 +3164,20 @@ PRAGMA_ENABLE_DEPRECATION_WARNINGS
 	LoadRequest.bWasRequestedByOwnerOrPawn = IsTrackingOwnerOrPawn();
 #endif
 
+	CacheObject.bIsPending = true;
+
+	FPendingAsyncLoadRequest* ExistingRequest = PendingAsyncLoadRequests.Find(CacheObject.PathName);
+	if (ExistingRequest)
+	{
+		// Same package name but a possibly different net GUID. Note down the GUID and wait for the async load completion callback
+		ExistingRequest->Merge(LoadRequest);
+		return;
+	}
+
 	PendingAsyncLoadRequests.Emplace(CacheObject.PathName, MoveTemp(LoadRequest));
 
 	DelinquentAsyncLoads.MaxConcurrentAsyncLoads = FMath::Max<uint32>(DelinquentAsyncLoads.MaxConcurrentAsyncLoads, PendingAsyncLoadRequests.Num());
 
-	CacheObject.bIsPending = true;
 	LoadPackageAsync(CacheObject.PathName.ToString(), FLoadPackageAsyncDelegate::CreateRaw(this, &FNetGUIDCache::AsyncPackageCallback));
 }
 
@@ -3179,37 +3191,40 @@ void FNetGUIDCache::AsyncPackageCallback(const FName& PackageName, UPackage * Pa
 	{
 		const bool bIsBroken = (Package == nullptr);
 
-		if (FNetGuidCacheObject* CacheObject = ObjectLookup.Find(PendingLoadRequest->NetGUID))
+		for (FNetworkGUID NetGUIDToProcess : PendingLoadRequest->NetGUIDs)
 		{
-			if (!CacheObject->bIsPending)
+			if (FNetGuidCacheObject* CacheObject = ObjectLookup.Find(NetGUIDToProcess))
 			{
-				UE_LOG(LogNetPackageMap, Error, TEXT("AsyncPackageCallback: Package wasn't pending. Path: %s, NetGUID: %s"), *PackageName.ToString(), *PendingLoadRequest->NetGUID.ToString());
-			}
+				if (!CacheObject->bIsPending)
+				{
+					UE_LOG(LogNetPackageMap, Error, TEXT("AsyncPackageCallback: Package wasn't pending. Path: %s, NetGUID: %s"), *PackageName.ToString(), *NetGUIDToProcess.ToString());
+				}
 
-			CacheObject->bIsPending = false;
+				CacheObject->bIsPending = false;
 
-			if (bIsBroken)
-			{
-				CacheObject->bIsBroken = true;
-				UE_LOG(LogNetPackageMap, Error, TEXT("AsyncPackageCallback: Package FAILED to load. Path: %s, NetGUID: %s"), *PackageName.ToString(), *PendingLoadRequest->NetGUID.ToString());
-			}
-
-			if (UObject* Object = CacheObject->Object.Get())
-			{
-				UpdateQueuedBunchObjectReference(PendingLoadRequest->NetGUID, Object);
+				if (bIsBroken)
+				{
+					CacheObject->bIsBroken = true;
+					UE_LOG(LogNetPackageMap, Error, TEXT("AsyncPackageCallback: Package FAILED to load. Path: %s, NetGUID: %s"), *PackageName.ToString(), *NetGUIDToProcess.ToString());
+				}
 
-				if (UWorld* World = Object->GetWorld())
+				if (UObject* Object = CacheObject->Object.Get())
 				{
-					if (AGameStateBase* GS = World->GetGameState())
+					UpdateQueuedBunchObjectReference(NetGUIDToProcess, Object);
+
+					if (UWorld* World = Object->GetWorld())
 					{
-						GS->AsyncPackageLoaded(Object);
+						if (AGameStateBase* GS = World->GetGameState())
+						{
+							GS->AsyncPackageLoaded(Object);
+						}
 					}
 				}
 			}
-		}
-		else
-		{
-			UE_LOG(LogNetPackageMap, Error, TEXT("AsyncPackageCallback: Could not find net guid. Path: %s, NetGUID: %s"), *PackageName.ToString(), *PendingLoadRequest->NetGUID.ToString());
+			else
+			{
+				UE_LOG(LogNetPackageMap, Error, TEXT("AsyncPackageCallback: Could not find net guid. Path: %s, NetGUID: %s"), *PackageName.ToString(), *NetGUIDToProcess.ToString());
+			}
 		}
 
 		// This won't be the exact amount of time that we spent loading the package, but should
