Licensee request to enable optional UAV access to render targets created by the Kismet Rendering Library for compute shader access.

diff --git a/Engine/Source/Runtime/Engine/Classes/Kismet/KismetRenderingLibrary.h b/Engine/Source/Runtime/Engine/Classes/Kismet/KismetRenderingLibrary.h
--- a/Engine/Source/Runtime/Engine/Classes/Kismet/KismetRenderingLibrary.h
+++ b/Engine/Source/Runtime/Engine/Classes/Kismet/KismetRenderingLibrary.h
@@ -46,19 +46,19 @@ class UKismetRenderingLibrary : public UBlueprintFunctionLibrary
 	 * Creates a new render target and initializes it to the specified dimensions
 	 */
 	UFUNCTION(BlueprintCallable, Category="Rendering", meta=(WorldContext="WorldContextObject"))
-	static ENGINE_API UTextureRenderTarget2D* CreateRenderTarget2D(UObject* WorldContextObject, int32 Width = 256, int32 Height = 256, ETextureRenderTargetFormat Format = RTF_RGBA16f, FLinearColor ClearColor = FLinearColor::Black, bool bAutoGenerateMipMaps = false);
+	static ENGINE_API UTextureRenderTarget2D* CreateRenderTarget2D(UObject* WorldContextObject, int32 Width = 256, int32 Height = 256, ETextureRenderTargetFormat Format = RTF_RGBA16f, FLinearColor ClearColor = FLinearColor::Black, bool bAutoGenerateMipMaps = false, bool bSupportUAVs = false);
 
 	/**
 	 * Creates a new render target array and initializes it to the specified dimensions
 	 */
 	UFUNCTION(BlueprintCallable, Category = "Rendering", meta = (WorldContext = "WorldContextObject"))
-	static ENGINE_API UTextureRenderTarget2DArray* CreateRenderTarget2DArray(UObject* WorldContextObject, int32 Width = 256, int32 Height = 256, int32 Slices = 1, ETextureRenderTargetFormat Format = RTF_RGBA16f, FLinearColor ClearColor = FLinearColor::Black, bool bAutoGenerateMipMaps = false);
+	static ENGINE_API UTextureRenderTarget2DArray* CreateRenderTarget2DArray(UObject* WorldContextObject, int32 Width = 256, int32 Height = 256, int32 Slices = 1, ETextureRenderTargetFormat Format = RTF_RGBA16f, FLinearColor ClearColor = FLinearColor::Black, bool bAutoGenerateMipMaps = false, bool bSupportUAVs = false);
 
 	/**
 	 * Creates a new volume render target and initializes it to the specified dimensions
 	 */
 	UFUNCTION(BlueprintCallable, Category = "Rendering", meta = (WorldContext = "WorldContextObject"))
-	static ENGINE_API UTextureRenderTargetVolume* CreateRenderTargetVolume(UObject* WorldContextObject, int32 Width = 16, int32 Height = 16, int32 Depth = 16, ETextureRenderTargetFormat Format = RTF_RGBA16f, FLinearColor ClearColor = FLinearColor::Black, bool bAutoGenerateMipMaps = false);
+	static ENGINE_API UTextureRenderTargetVolume* CreateRenderTargetVolume(UObject* WorldContextObject, int32 Width = 16, int32 Height = 16, int32 Depth = 16, ETextureRenderTargetFormat Format = RTF_RGBA16f, FLinearColor ClearColor = FLinearColor::Black, bool bAutoGenerateMipMaps = false, bool bSupportUAVs = false);
 
 	/**
 	 * Manually releases GPU resources of a render target. This is useful for blueprint creating a lot of render target that would
diff --git a/Engine/Source/Runtime/Engine/Private/KismetRenderingLibrary.cpp b/Engine/Source/Runtime/Engine/Private/KismetRenderingLibrary.cpp
--- a/Engine/Source/Runtime/Engine/Private/KismetRenderingLibrary.cpp
+++ b/Engine/Source/Runtime/Engine/Private/KismetRenderingLibrary.cpp
@@ -65,7 +65,7 @@ void UKismetRenderingLibrary::ClearRenderTarget2D(UObject* WorldContextObject, U
 	}
 }
 
-UTextureRenderTarget2D* UKismetRenderingLibrary::CreateRenderTarget2D(UObject* WorldContextObject, int32 Width, int32 Height, ETextureRenderTargetFormat Format, FLinearColor ClearColor, bool bAutoGenerateMipMaps)
+UTextureRenderTarget2D* UKismetRenderingLibrary::CreateRenderTarget2D(UObject* WorldContextObject, int32 Width, int32 Height, ETextureRenderTargetFormat Format, FLinearColor ClearColor, bool bAutoGenerateMipMaps, bool bSupportUAVs)
 {
 	UWorld* World = GEngine->GetWorldFromContextObject(WorldContextObject, EGetWorldErrorMode::LogAndReturnNull);
 
@@ -76,6 +76,7 @@ UTextureRenderTarget2D* UKismetRenderingLibrary::CreateRenderTarget2D(UObject* W
 		NewRenderTarget2D->RenderTargetFormat = Format;
 		NewRenderTarget2D->ClearColor = ClearColor;
 		NewRenderTarget2D->bAutoGenerateMips = bAutoGenerateMipMaps;
+		NewRenderTarget2D->bCanCreateUAV = bSupportUAVs;
 		NewRenderTarget2D->InitAutoFormat(Width, Height);	
 		NewRenderTarget2D->UpdateResourceImmediate(true);
 
@@ -85,7 +86,7 @@ UTextureRenderTarget2D* UKismetRenderingLibrary::CreateRenderTarget2D(UObject* W
 	return nullptr;
 }
 
-UTextureRenderTarget2DArray* UKismetRenderingLibrary::CreateRenderTarget2DArray(UObject* WorldContextObject, int32 Width, int32 Height, int32 Slices, ETextureRenderTargetFormat Format, FLinearColor ClearColor, bool bAutoGenerateMipMaps)
+UTextureRenderTarget2DArray* UKismetRenderingLibrary::CreateRenderTarget2DArray(UObject* WorldContextObject, int32 Width, int32 Height, int32 Slices, ETextureRenderTargetFormat Format, FLinearColor ClearColor, bool bAutoGenerateMipMaps, bool bSupportUAVs)
 {
 	UWorld* World = GEngine->GetWorldFromContextObject(WorldContextObject, EGetWorldErrorMode::LogAndReturnNull);
 
@@ -94,6 +95,7 @@ UTextureRenderTarget2DArray* UKismetRenderingLibrary::CreateRenderTarget2DArray(
 		UTextureRenderTarget2DArray* NewRenderTarget = NewObject<UTextureRenderTarget2DArray>(WorldContextObject);
 		check(NewRenderTarget);
 		NewRenderTarget->ClearColor = ClearColor;
+		NewRenderTarget->bCanCreateUAV = bSupportUAVs;
 		NewRenderTarget->Init(Width, Height, Slices, GetPixelFormatFromRenderTargetFormat(Format));
 		NewRenderTarget->UpdateResourceImmediate(true);
 		return NewRenderTarget;
@@ -102,7 +104,7 @@ UTextureRenderTarget2DArray* UKismetRenderingLibrary::CreateRenderTarget2DArray(
 	return nullptr;
 }
 
-UTextureRenderTargetVolume* UKismetRenderingLibrary::CreateRenderTargetVolume(UObject* WorldContextObject, int32 Width, int32 Height, int32 Depth, ETextureRenderTargetFormat Format, FLinearColor ClearColor, bool bAutoGenerateMipMaps)
+UTextureRenderTargetVolume* UKismetRenderingLibrary::CreateRenderTargetVolume(UObject* WorldContextObject, int32 Width, int32 Height, int32 Depth, ETextureRenderTargetFormat Format, FLinearColor ClearColor, bool bAutoGenerateMipMaps, bool bSupportUAVs)
 {
 	UWorld* World = GEngine->GetWorldFromContextObject(WorldContextObject, EGetWorldErrorMode::LogAndReturnNull);
 
@@ -111,6 +113,7 @@ UTextureRenderTargetVolume* UKismetRenderingLibrary::CreateRenderTargetVolume(UO
 		UTextureRenderTargetVolume* NewRenderTarget = NewObject<UTextureRenderTargetVolume>(WorldContextObject);
 		check(NewRenderTarget);
 		NewRenderTarget->ClearColor = ClearColor;
+		NewRenderTarget->bCanCreateUAV = bSupportUAVs;
 		NewRenderTarget->Init(Width, Height, Depth, GetPixelFormatFromRenderTargetFormat(Format));
 		NewRenderTarget->UpdateResourceImmediate(true);
 		return NewRenderTarget;
