FX Assets PSO precaching

Assumes that UFXAssetSystem will have performed the PSO collection beforehand. We call it JIT for scene proxy creation so that the PSOs are being boosted.

#rnx

diff --git a/Engine/Plugins/FX/Niagara/Source/Niagara/Private/NiagaraComponent.cpp b/Engine/Plugins/FX/Niagara/Source/Niagara/Private/NiagaraComponent.cpp
--- a/Engine/Plugins/FX/Niagara/Source/Niagara/Private/NiagaraComponent.cpp
+++ b/Engine/Plugins/FX/Niagara/Source/Niagara/Private/NiagaraComponent.cpp
@@ -2061,6 +2061,21 @@ FPrimitiveSceneProxy* UNiagaraComponent::CreateSceneProxy()
 	LLM_SCOPE(ELLMTag::Niagara);
 	SCOPE_CYCLE_COUNTER(STAT_NiagaraCreateSceneProxy);
 	SCOPE_CYCLE_COUNTER(STAT_NiagaraOverview_GT);
+    
+	if (Asset != nullptr)
+	{
+		// PSO request should have been handled by the Asset itself, so here we just ensure the request gets boosted
+		if (!bPSOPrecacheCalled)
+		{
+			PrecacheAssetPSOs(Asset);
+		}
+
+		if (IsPSOPrecaching())
+		{
+			UE_LOG(LogNiagara, Verbose, TEXT("Skipping CreateSceneProxy for UNiagaraComponent %s (UNiagaraSystem PSOs are still compiling)"), *GetFullName());
+			return nullptr;
+		}
+	}
 
 	// The constructor will set up the System renderers from the component.
 	FNiagaraSceneProxy* Proxy = new FNiagaraSceneProxy(this);
diff --git a/Engine/Plugins/FX/Niagara/Source/Niagara/Private/NiagaraSystem.cpp b/Engine/Plugins/FX/Niagara/Source/Niagara/Private/NiagaraSystem.cpp
--- a/Engine/Plugins/FX/Niagara/Source/Niagara/Private/NiagaraSystem.cpp
+++ b/Engine/Plugins/FX/Niagara/Source/Niagara/Private/NiagaraSystem.cpp
@@ -1004,12 +1004,6 @@ void UNiagaraSystem::PrecachePSOs()
 		return;
 	}
 
-	struct VFsPerMaterialData
-	{
-		UMaterialInterface* MaterialInterface;
-		bool bDisableBackfaceCulling;
-		TArray<const FVertexFactoryType*, TInlineAllocator<2>> VertexFactoryTypes;
-	};
 	TArray<VFsPerMaterialData, TInlineAllocator<2>> VFsPerMaterials;
 
 	for (const FNiagaraEmitterHandle& EmitterHandle : GetEmitterHandles())
@@ -1052,17 +1046,7 @@ void UNiagaraSystem::PrecachePSOs()
 		}
 	}
 
-	FPSOPrecacheParams PreCachePSOParams;
-	PreCachePSOParams.SetMobility(EComponentMobility::Movable);
-
-	for (VFsPerMaterialData& VFsPerMaterial : VFsPerMaterials)
-	{
-		if (VFsPerMaterial.MaterialInterface)
-		{
-			PreCachePSOParams.bDisableBackFaceCulling = VFsPerMaterial.bDisableBackfaceCulling;
-			VFsPerMaterial.MaterialInterface->PrecachePSOs(VFsPerMaterial.VertexFactoryTypes, PreCachePSOParams);
-		}
-	}
+	LaunchPSOPrecaching(VFsPerMaterials);
 }
 
 
diff --git a/Engine/Source/Runtime/Engine/Classes/Materials/MaterialInterface.h b/Engine/Source/Runtime/Engine/Classes/Materials/MaterialInterface.h
--- a/Engine/Source/Runtime/Engine/Classes/Materials/MaterialInterface.h
+++ b/Engine/Source/Runtime/Engine/Classes/Materials/MaterialInterface.h
@@ -560,7 +560,13 @@ public:
 		TArray<FMaterialPSOPrecacheRequestID> MaterialPSOPrecacheRequestIDs;
 		return PrecachePSOs(VertexFactoryTypes, PreCacheParams, MaterialPSOPrecacheRequestIDs);
 	}
-	FGraphEventArray PrecachePSOs(const TConstArrayView<const FVertexFactoryType*>& VertexFactoryTypes, const struct FPSOPrecacheParams& PreCacheParams, TArray<FMaterialPSOPrecacheRequestID>& OutMaterialPSORequestIDs) 
+
+	FGraphEventArray PrecachePSOs(const TConstArrayView<const FVertexFactoryType*>& VertexFactoryTypes, const struct FPSOPrecacheParams& PreCacheParams, TArray<FMaterialPSOPrecacheRequestID>& OutMaterialPSORequestIDs)
+	{
+		return PrecachePSOs(VertexFactoryTypes, PreCacheParams, EPSOPrecachePriority::Medium, OutMaterialPSORequestIDs);
+	}
+
+	FGraphEventArray PrecachePSOs(const TConstArrayView<const FVertexFactoryType*>& VertexFactoryTypes, const struct FPSOPrecacheParams& PreCacheParams, EPSOPrecachePriority PSOPrecachePriority, TArray<FMaterialPSOPrecacheRequestID>& OutMaterialPSORequestIDs)
 	{ 
 		FPSOPrecacheVertexFactoryDataList VertexFactoryDataList;
 		VertexFactoryDataList.SetNum(VertexFactoryTypes.Num());
@@ -568,7 +574,7 @@ public:
 		{
 			VertexFactoryDataList[i].VertexFactoryType = VertexFactoryTypes[i];
 		}
-		return PrecachePSOs(VertexFactoryDataList, PreCacheParams, EPSOPrecachePriority::Medium, OutMaterialPSORequestIDs);
+		return PrecachePSOs(VertexFactoryDataList, PreCacheParams, PSOPrecachePriority, OutMaterialPSORequestIDs);
 	}
 	virtual FGraphEventArray PrecachePSOs(const FPSOPrecacheVertexFactoryDataList& VertexFactoryDataList, const struct FPSOPrecacheParams& PreCacheParams, EPSOPrecachePriority Priority, TArray<FMaterialPSOPrecacheRequestID>& OutMaterialPSORequestIDs) { return FGraphEventArray(); }
 
diff --git a/Engine/Source/Runtime/Engine/Classes/Particles/ParticleSystem.h b/Engine/Source/Runtime/Engine/Classes/Particles/ParticleSystem.h
--- a/Engine/Source/Runtime/Engine/Classes/Particles/ParticleSystem.h
+++ b/Engine/Source/Runtime/Engine/Classes/Particles/ParticleSystem.h
@@ -8,6 +8,8 @@
 #include "UObject/Object.h"
 #include "UObject/ScriptMacros.h"
 #include "Particles/ParticlePerfStats.h"
+#include "Async/TaskGraphInterfaces.h"
+#include "RHIDefinitions.h"
 
 #include "ParticleSystem.generated.h"
 
@@ -108,6 +110,8 @@ struct FNamedEmitterMaterial
 	TObjectPtr<UMaterialInterface> Material;
 };
 
+using FMaterialPSOPrecacheRequestID = uint32;
+
 UCLASS(Abstract, MinimalAPI, BlueprintType)
 class UFXSystemAsset : public UObject
 {
@@ -149,7 +153,22 @@ public:
 #endif
 #endif
 
+	const FGraphEventArray& GetPrecachePSOsEvents() const { return PrecachePSOsEvents; }
+	const TArray<FMaterialPSOPrecacheRequestID>& GetMaterialPSOPrecacheRequestIDs() const { return MaterialPSOPrecacheRequestIDs; }
+
+protected:
+	struct VFsPerMaterialData
+	{
+		UMaterialInterface* MaterialInterface = nullptr;
+		EPrimitiveType PrimitiveType = PT_TriangleList; // must match FPSOPrecacheParams::PrimitiveType default value
+		bool bDisableBackfaceCulling = false;  // must match FPSOPrecacheParams::bDisableBackfaceCulling default value
+		TArray<const class FVertexFactoryType*, TInlineAllocator<2>> VertexFactoryTypes;
+	};
+
+	ENGINE_API void LaunchPSOPrecaching(TArrayView<VFsPerMaterialData> VFsPerMaterials);
 
+	FGraphEventArray PrecachePSOsEvents;
+	TArray<FMaterialPSOPrecacheRequestID> MaterialPSOPrecacheRequestIDs;
 };
 
 /**
diff --git a/Engine/Source/Runtime/Engine/Classes/Particles/ParticleSystemComponent.h b/Engine/Source/Runtime/Engine/Classes/Particles/ParticleSystemComponent.h
--- a/Engine/Source/Runtime/Engine/Classes/Particles/ParticleSystemComponent.h
+++ b/Engine/Source/Runtime/Engine/Classes/Particles/ParticleSystemComponent.h
@@ -449,6 +449,9 @@ public:
 #if WITH_PER_COMPONENT_PARTICLE_PERF_STATS
 	mutable FParticlePerfStats* ParticlePerfStats = nullptr;
 #endif
+
+protected:
+	void PrecacheAssetPSOs(UFXSystemAsset* FXSystemAsset);
 };
 
 
diff --git a/Engine/Source/Runtime/Engine/Private/Particles/ParticleComponents.cpp b/Engine/Source/Runtime/Engine/Private/Particles/ParticleComponents.cpp
--- a/Engine/Source/Runtime/Engine/Private/Particles/ParticleComponents.cpp
+++ b/Engine/Source/Runtime/Engine/Private/Particles/ParticleComponents.cpp
@@ -145,6 +145,22 @@ void UFXSystemAsset::PostInitProperties()
 #endif
 }
 
+void UFXSystemAsset::LaunchPSOPrecaching(TArrayView<VFsPerMaterialData> VFsPerMaterials)
+{
+	FPSOPrecacheParams PreCachePSOParams;
+	PreCachePSOParams.SetMobility(EComponentMobility::Movable);
+
+	for (VFsPerMaterialData& VFsPerMaterial : VFsPerMaterials)
+	{
+		if (VFsPerMaterial.MaterialInterface)
+		{
+			PreCachePSOParams.PrimitiveType = (EPrimitiveType)VFsPerMaterial.PrimitiveType;
+			PreCachePSOParams.bDisableBackFaceCulling = VFsPerMaterial.bDisableBackfaceCulling;
+			PrecachePSOsEvents.Append(VFsPerMaterial.MaterialInterface->PrecachePSOs(VFsPerMaterial.VertexFactoryTypes, PreCachePSOParams,EPSOPrecachePriority::Medium, MaterialPSOPrecacheRequestIDs));
+		}
+	}
+}
+
 //////////////////////////////////////////////////////////////////////////
 
 int32 GParticleLODBias = 0;
@@ -2634,12 +2650,6 @@ void UParticleSystem::PrecachePSOs()
 		return;
 	}
 
-	struct VFsPerMaterialData
-	{
-		UMaterialInterface* MaterialInterface;
-		EPrimitiveType PrimitiveType;
-		TArray<const FVertexFactoryType*, TInlineAllocator<2>> VertexFactoryTypes;
-	};
 	TArray<VFsPerMaterialData, TInlineAllocator<2>> VFsPerMaterials;
 
 	// No per component emitter materials known at this point in time
@@ -2685,17 +2695,7 @@ void UParticleSystem::PrecachePSOs()
 		}
 	}
 
-	FPSOPrecacheParams PreCachePSOParams;
-	PreCachePSOParams.SetMobility(EComponentMobility::Movable);
-
-	for (VFsPerMaterialData& VFsPerMaterial : VFsPerMaterials)
-	{
-		if (VFsPerMaterial.MaterialInterface)
-		{
-			PreCachePSOParams.PrimitiveType = VFsPerMaterial.PrimitiveType;
-			VFsPerMaterial.MaterialInterface->PrecachePSOs(VFsPerMaterial.VertexFactoryTypes, PreCachePSOParams);
-		}
-	}
+	LaunchPSOPrecaching(VFsPerMaterials);
 }
 
 void UParticleSystem::Serialize(FArchive& Ar)
@@ -3480,6 +3480,41 @@ bool UFXSystemComponent::RequiresLWCTileRecache(const FVector3f CurrentTile, con
 	return bNeedsRecache;
 }
 
+void UFXSystemComponent::PrecacheAssetPSOs(UFXSystemAsset* FXSystemAsset)
+{
+	if (!FApp::CanEverRender() || !IsComponentPSOPrecachingEnabled() || FXSystemAsset == nullptr)
+	{
+		return;
+	}
+
+	const FGraphEventArray& GraphEvents = FXSystemAsset->GetPrecachePSOsEvents();
+
+	check(IsInGameThread() || IsInParallelGameThread());
+
+	MaterialPSOPrecacheRequestIDs.Empty();
+	PSOPrecacheCompileEvent = nullptr;
+	bPSOPrecacheRequestBoosted = false;
+
+	// The asset will keep the Precache events alive, but these might be over. Avoid delaying scene proxy creation if everything is finished
+	bool bAllEventsDone = true;
+	for (FGraphEventRef GraphEvent : GraphEvents)
+	{
+		if (!GraphEvent->IsComplete())
+		{
+			bAllEventsDone = false;
+			break;
+		}
+	}
+
+	if (!bAllEventsDone)
+	{
+		MaterialPSOPrecacheRequestIDs.Append(FXSystemAsset->GetMaterialPSOPrecacheRequestIDs());
+		RequestRecreateRenderStateWhenPSOPrecacheFinished(GraphEvents);
+	}
+
+	bPSOPrecacheCalled = true;
+}
+
 FOnSystemPreActivationChange UParticleSystemComponent::OnSystemPreActivationChange;
 
 UParticleSystemComponent::UParticleSystemComponent(const FObjectInitializer& ObjectInitializer)
diff --git a/Engine/Source/Runtime/Engine/Private/Particles/ParticleSystemRender.cpp b/Engine/Source/Runtime/Engine/Private/Particles/ParticleSystemRender.cpp
--- a/Engine/Source/Runtime/Engine/Private/Particles/ParticleSystemRender.cpp
+++ b/Engine/Source/Runtime/Engine/Private/Particles/ParticleSystemRender.cpp
@@ -6923,6 +6923,17 @@ FPrimitiveSceneProxy* UParticleSystemComponent::CreateSceneProxy()
 	//@fixme Get non-instanced path working in ES!
 	if ((IsActive() == true)/** && (EmitterInstances.Num() > 0)*/ && Template)
 	{
+		if (!bPSOPrecacheCalled)
+		{
+			PrecacheAssetPSOs(Template);
+		}
+
+		if (IsPSOPrecaching())
+		{
+			UE_LOG(LogParticles, Verbose, TEXT("Skipping CreateSceneProxy for UParticleSystemComponent %s (UParticleSystem PSOs are still compiling)"), *GetFullName());
+			return nullptr;
+		}
+
 		FInGameScopedCycleCounter InGameCycleCounter(GetWorld(), EInGamePerfTrackers::VFXSignificance, EInGamePerfTrackerThreads::GameThread, bIsManagingSignificance);
 
 		UE_LOG(LogParticles,Verbose,
