UEngine::CorrectNegativeTimeDelta: Don't crash on negative delta-time, removed obsolete warning message

diff --git a/Engine/Source/Runtime/Engine/Private/UnrealEngine.cpp b/Engine/Source/Runtime/Engine/Private/UnrealEngine.cpp
--- a/Engine/Source/Runtime/Engine/Private/UnrealEngine.cpp
+++ b/Engine/Source/Runtime/Engine/Private/UnrealEngine.cpp
@@ -2355,13 +2355,7 @@ static FAutoConsoleCommand SetTimedMemReport(TEXT("TimedMemReport.Delay"), TEXT(
 
 double UEngine::CorrectNegativeTimeDelta(double DeltaRealTime)
 {
-#if PLATFORM_ANDROID
 	UE_LOG(LogEngine, Warning, TEXT("Detected negative delta time - ignoring"));
-#else
-	// AMD dual-core systems are a known issue that require AMD CPU drivers to be installed. Installer will take care of this for shipping.
-	UE_LOG(LogEngine, Fatal, TEXT("Detected negative delta time - on AMD systems please install http://files.aoaforums.com/I3199-setup.zip.html - DeltaTime:%f"),
-		DeltaRealTime);
-#endif
 	return 0.01;
 }
 
@@ -2636,12 +2630,7 @@ void UEngine::UpdateTimeAndHandleMaxTickRate()
 																		 // Negative delta time means something is wrong with the system. Error out so user can address issue.
 		if( FApp::GetDeltaTime() < 0 )
 		{
-#if PLATFORM_ANDROID
 			UE_LOG(LogEngine, Warning, TEXT("Detected negative delta time - ignoring"));
-#else
-			// AMD dual-core systems are a known issue that require AMD CPU drivers to be installed. Installer will take care of this for shipping.
-			UE_LOG(LogEngine, Warning, TEXT("Detected negative delta time - on AMD systems please install http://files.aoaforums.com/I3199-setup.zip.html"));
-#endif
 			FApp::SetDeltaTime(0.01);
 		}
 
