Remove compilation flush not necessary anymore now that REN_ForceNoResetLoaders is used to rename packages

#rnx

diff --git a/Engine/Source/Runtime/Engine/Private/WorldPartition/WorldPartitionLevelHelper.cpp b/Engine/Source/Runtime/Engine/Private/WorldPartition/WorldPartitionLevelHelper.cpp
--- a/Engine/Source/Runtime/Engine/Private/WorldPartition/WorldPartitionLevelHelper.cpp
+++ b/Engine/Source/Runtime/Engine/Private/WorldPartition/WorldPartitionLevelHelper.cpp
@@ -109,11 +109,6 @@ void FWorldPartitionLevelHelper::MoveExternalActorsToLevel(const TArray<FWorldPa
 {
 	TRACE_CPUPROFILER_EVENT_SCOPE(FWorldPartitionLevelHelper::MoveExternalActorsToLevel);
 
-	// We can't have async compilation still going on while we move actors as this is going to ResetLoaders which will move bulkdata around that
-	// might still be used by async compilation. 
-	// #TODO_DC Revisit once virtualbulkdata are enabled
-	FAssetCompilingManager::Get().FinishAllCompilation();
-
 	check(InLevel);
 	UPackage* LevelPackage = InLevel->GetPackage();
 	
