Strata - fixed menu names

diff --git a/Engine/Source/Runtime/Engine/Private/StrataVisualizationData.cpp b/Engine/Source/Runtime/Engine/Private/StrataVisualizationData.cpp
--- a/Engine/Source/Runtime/Engine/Private/StrataVisualizationData.cpp
+++ b/Engine/Source/Runtime/Engine/Private/StrataVisualizationData.cpp
@@ -149,7 +149,7 @@ void FStrataVisualizationData::Initialize()
 		AddVisualizationMode(
 			AllModeMap,
 			TEXT("RoughRefractionClassification"),
-			LOCTEXT("RoughRefractionClassification", "Material Count"),
+			LOCTEXT("RoughRefractionClassification", "Rough Refraction Classification"),
 			LOCTEXT("RoughRefractionClassificationDesc", "Visualizes Strata rough refraction classification"),
 			FViewMode::RoughRefractionClassification,
 			true,
