ShaderPlatformInfo Material Node
- Limit the scope to just 3 properties for the moment(isMobile, isConsole, isPc)
- Create a compile-time link between the properties and the material in case they are removed/renamed
- Cleaning code

diff --git a/Engine/Source/Runtime/Engine/Private/Materials/MaterialExpressions.cpp b/Engine/Source/Runtime/Engine/Private/Materials/MaterialExpressions.cpp
--- a/Engine/Source/Runtime/Engine/Private/Materials/MaterialExpressions.cpp
+++ b/Engine/Source/Runtime/Engine/Private/Materials/MaterialExpressions.cpp
@@ -9663,18 +9663,8 @@ UMaterialExpressionDataDrivenShaderPlatformInfoSwitch::UMaterialExpressionDataDr
 : Super(ObjectInitializer)
 {
 #if WITH_EDITORONLY_DATA
-	// Structure to hold one-time initialization
-	struct FConstructorStatics
-	{
-		FText NAME_Utility;
-		FConstructorStatics()
-			: NAME_Utility(LOCTEXT("Utility", "Utility"))
-		{
-		}
-	};
-	static FConstructorStatics ConstructorStatics;
-
-	MenuCategories.Add(ConstructorStatics.NAME_Utility);
+	FText NAME_Utility(LOCTEXT("Utility", "Utility"));
+	MenuCategories.Add(NAME_Utility);
 	bCollapsed = false;
 	bContainsInvalidProperty = false;
 #endif // WITH_EDITORONLY_DATA
@@ -9746,11 +9736,11 @@ int32 UMaterialExpressionDataDrivenShaderPlatformInfoSwitch::Compile(class FMate
 		bool bCheckProperty = FGenericDataDrivenShaderPlatformInfo::PropertyToShaderPlatformFunctionMap[DDSPIInput.InputName.ToString()](ShaderPlatform);
 		if (DDSPIInput.PropertyCondition == EDataDrivenShaderPlatformInfoCondition::COND_True)
 		{
-			bCheck = bCheck && bCheckProperty;
+			bCheck &= bCheckProperty;
 		}
 		else
 		{
-			bCheck = bCheck && !bCheckProperty;
+			bCheck &= !bCheckProperty;
 		}
 	}
 
@@ -9770,7 +9760,7 @@ int32 UMaterialExpressionDataDrivenShaderPlatformInfoSwitch::Compile(class FMate
 
 void UMaterialExpressionDataDrivenShaderPlatformInfoSwitch::GetCaption(TArray<FString>& OutCaptions) const
 {
-	OutCaptions.Add(FString(TEXT("DataDrivenShaderPlatformInfo Switch")));
+	OutCaptions.Add(FString(TEXT("ShaderPlatformInfo Switch")));
 }
 
 bool UMaterialExpressionDataDrivenShaderPlatformInfoSwitch::IsInputConnectionRequired(int32 InputIndex) const
@@ -9846,7 +9836,8 @@ bool UMaterialExpressionDataDrivenShaderPlatformInfoSwitch::IsResultMaterialAttr
 	for (FExpressionInput* ExpressionInput : ExpressionInputs)
 	{
 		// If there is a loop anywhere in this expression's inputs then we can't risk checking them
-		if (ExpressionInput->Expression && !ExpressionInput->Expression->ContainsInputLoop() && ExpressionInput->Expression->IsResultMaterialAttributes(ExpressionInput->OutputIndex))
+		TObjectPtr<class UMaterialExpression> Expression = ExpressionInput->Expression;
+		if (Expression && !Expression->ContainsInputLoop() && ExpressionInput->Expression->IsResultMaterialAttributes(ExpressionInput->OutputIndex))
 		{
 			return true;
 		}
diff --git a/Engine/Source/Runtime/RHI/Private/DataDrivenShaderPlatformInfo.cpp b/Engine/Source/Runtime/RHI/Private/DataDrivenShaderPlatformInfo.cpp
--- a/Engine/Source/Runtime/RHI/Private/DataDrivenShaderPlatformInfo.cpp
+++ b/Engine/Source/Runtime/RHI/Private/DataDrivenShaderPlatformInfo.cpp
@@ -152,15 +152,14 @@ void FGenericDataDrivenShaderPlatformInfo::ParseDataDrivenShaderInfo(const FConf
 	ShaderPropertiesString += FString::Printf(TEXT("_%d"), SettingValue);
 
 #if WITH_EDITOR
-	#define ADD_PROPERTY_TO_SHADERPLATFORM_FUNCTIONMAP(SettingName) \
-		FGenericDataDrivenShaderPlatformInfo::PropertyToShaderPlatformFunctionMap.FindOrAdd(#SettingName) = [](const FStaticShaderPlatform Platform)->bool{ return Infos[Platform].SettingName;};
+	#define ADD_PROPERTY_TO_SHADERPLATFORM_FUNCTIONMAP(SettingName, FunctionName) \
+	FGenericDataDrivenShaderPlatformInfo::PropertyToShaderPlatformFunctionMap.FindOrAdd(#SettingName) = &FunctionName;
 #else
-	#define ADD_PROPERTY_TO_SHADERPLATFORM_FUNCTIONMAP(SettingName)
+	#define ADD_PROPERTY_TO_SHADERPLATFORM_FUNCTIONMAP(SettingName, FunctionName)
 #endif
 
 #define GET_SECTION_BOOL_HELPER(SettingName)	\
 	Info.SettingName = GetSectionBool(Section, #SettingName, Info.SettingName);	\
-	ADD_PROPERTY_TO_SHADERPLATFORM_FUNCTIONMAP(SettingName)	\
 	ADD_TO_PROPERTIES_STRING(SettingName, Info.SettingName)
 
 #define GET_SECTION_INT_HELPER(SettingName)	\
@@ -175,6 +174,13 @@ void FGenericDataDrivenShaderPlatformInfo::ParseDataDrivenShaderInfo(const FConf
 	Info.SettingName = GetSectionBindlessSupport(Section, #SettingName, Info.SettingName); \
 	ADD_TO_PROPERTIES_STRING(SettingName, Info.SettingName)
 
+	// These properties will be exposed to the MaterialEditor that use the ShaderPlatformInfo Node
+	// If you remove/rename a property be sure to address this in UMaterialExpressionDataDrivenShaderPlatformInfoSwitch serialization
+	// If you add a property that you think needs to be exposed to materials, it is enough to just call the macro here with the property name and the function it will call
+	ADD_PROPERTY_TO_SHADERPLATFORM_FUNCTIONMAP(IsMobile, GetIsMobile);
+	ADD_PROPERTY_TO_SHADERPLATFORM_FUNCTIONMAP(IsPC, GetIsPC);
+	ADD_PROPERTY_TO_SHADERPLATFORM_FUNCTIONMAP(IsConsole, GetIsConsole);
+
 	GET_SECTION_BOOL_HELPER(bIsMobile);
 	GET_SECTION_BOOL_HELPER(bIsMetalMRT);
 	GET_SECTION_BOOL_HELPER(bIsPC);
