Updated description text for non-functional AsyncCharacterMovement feature CVar to steer licensees away from using it.  No functional changes.

#rnx

diff --git a/Engine/Source/Runtime/Engine/Private/Components/CharacterMovementComponent.cpp b/Engine/Source/Runtime/Engine/Private/Components/CharacterMovementComponent.cpp
--- a/Engine/Source/Runtime/Engine/Private/Components/CharacterMovementComponent.cpp
+++ b/Engine/Source/Runtime/Engine/Private/Components/CharacterMovementComponent.cpp
@@ -237,7 +237,7 @@ namespace CharacterMovementCVars
 	int32 AsyncCharacterMovement = 0;
 	FAutoConsoleVariableRef CVarAsyncCharacterMovement(
 		TEXT("p.AsyncCharacterMovement"),
-		AsyncCharacterMovement, TEXT("1 enables asynchronous simulation of character movement on physics thread. Toggling this at runtime is not recommended."));
+		AsyncCharacterMovement, TEXT("1 enables asynchronous simulation of character movement on physics thread. Toggling this at runtime is not recommended. This feature is not fully developed, and its use is discouraged."));
 
 	int32 BasedMovementMode = 2;
 	FAutoConsoleVariableRef CVarBasedMovementMode(
