Scale cloth mapping influence distances when skeletal mesh actor is scaled

diff --git a/Engine/Shaders/Private/GpuSkinCacheComputeShader.usf b/Engine/Shaders/Private/GpuSkinCacheComputeShader.usf
--- a/Engine/Shaders/Private/GpuSkinCacheComputeShader.usf
+++ b/Engine/Shaders/Private/GpuSkinCacheComputeShader.usf
@@ -73,6 +73,8 @@ uint ClothBufferOffset;
 float ClothBlendWeight;
 /** Transform from cloth space (relative to cloth root bone) to local/component space */
 float4x4 ClothToLocal;
+/** Scale of the owner actor */
+float3 WorldScale;
 
 void GetClothBuffer(uint VertexIndex, uint ClothIndexOffset, out FVertexTriangleInfluence Influences[NUM_INFLUENCES_PER_VERTEX])
 {
@@ -421,21 +423,21 @@ void SkinCacheUpdateBatchCS(uint3 GroupID : SV_GroupID,
                 SumWeights = 1.0f;	   
 			}
 
-			NormalPosition += Weight * (Influence.NormalBaryCoordsAndDist.x * (A + NA * Influence.NormalBaryCoordsAndDist.w)
-								      + Influence.NormalBaryCoordsAndDist.y * (B + NB * Influence.NormalBaryCoordsAndDist.w)
-								      + Influence.NormalBaryCoordsAndDist.z * (C + NC * Influence.NormalBaryCoordsAndDist.w));
+			NormalPosition += Weight * (Influence.NormalBaryCoordsAndDist.x * (A + NA * Influence.NormalBaryCoordsAndDist.w * WorldScale.x)
+								      + Influence.NormalBaryCoordsAndDist.y * (B + NB * Influence.NormalBaryCoordsAndDist.w * WorldScale.y)
+								      + Influence.NormalBaryCoordsAndDist.z * (C + NC * Influence.NormalBaryCoordsAndDist.w * WorldScale.z));
 
-		    TangentPosition += Weight * (Influence.TangentBaryCoordsAndDist.x*(A+NA*Influence.TangentBaryCoordsAndDist.w)
-								       + Influence.TangentBaryCoordsAndDist.y*(B+NB*Influence.TangentBaryCoordsAndDist.w)
-								       + Influence.TangentBaryCoordsAndDist.z * (C + NC * Influence.TangentBaryCoordsAndDist.w));
+		    TangentPosition += Weight * (Influence.TangentBaryCoordsAndDist.x * (A + NA * Influence.TangentBaryCoordsAndDist.w * WorldScale.x)
+								       + Influence.TangentBaryCoordsAndDist.y * (B + NB * Influence.TangentBaryCoordsAndDist.w * WorldScale.y)
+								       + Influence.TangentBaryCoordsAndDist.z * (C + NC * Influence.TangentBaryCoordsAndDist.w * WorldScale.z));
 
 			float3 TriangleBary = float3(Influence.PositionBaryCoordsAndDist.x,
                                          Influence.PositionBaryCoordsAndDist.y,
 										 1.0f - Influence.PositionBaryCoordsAndDist.x - Influence.PositionBaryCoordsAndDist.y);
 
-            float3 SimPosition = TriangleBary.x*(A+NA*Influence.PositionBaryCoordsAndDist.w)
-					           + TriangleBary.y*(B+NB*Influence.PositionBaryCoordsAndDist.w)
-					           + TriangleBary.z*(C+NC*Influence.PositionBaryCoordsAndDist.w);
+            float3 SimPosition = TriangleBary.x*(A+NA*Influence.PositionBaryCoordsAndDist.w * WorldScale.x)
+					           + TriangleBary.y*(B+NB*Influence.PositionBaryCoordsAndDist.w * WorldScale.y)
+					           + TriangleBary.z*(C+NC*Influence.PositionBaryCoordsAndDist.w * WorldScale.z);
 
 	        Intermediates.SimulatedPosition += Weight * SimPosition;
 	    }
diff --git a/Engine/Shaders/Private/GpuSkinVertexFactory.ush b/Engine/Shaders/Private/GpuSkinVertexFactory.ush
--- a/Engine/Shaders/Private/GpuSkinVertexFactory.ush
+++ b/Engine/Shaders/Private/GpuSkinVertexFactory.ush
@@ -63,6 +63,8 @@ Buffer<float2> ClothSimulVertsPositionsNormals;
 float4x4 ClothToLocal;
 /** blend weight between simulated positions and original key-framed animation */
 float ClothBlendWeight;
+/** Scale of the owner actor */
+float3 WorldScale;
 #if GPUSKIN_APEX_CLOTH_PREVIOUS
 Buffer<float2> PreviousClothSimulVertsPositionsNormals;
 float4x4 PreviousClothToLocal;
@@ -452,9 +454,9 @@ float4 ClothingPosition(FClothVertex ClothVertexInfluences[NUM_INFLUENCES_PER_VE
                                          Influence.BaryCoordPos.y,
 										 1.0f - Influence.BaryCoordPos.x - Influence.BaryCoordPos.y);
 
-			float3 SimPosition = TriangleBary.x * (A + NA * Influence.BaryCoordPos.w)
-					           + TriangleBary.y * (B + NB * Influence.BaryCoordPos.w)
-					           + TriangleBary.z * (C + NC * Influence.BaryCoordPos.w);
+			float3 SimPosition = TriangleBary.x * (A + NA * Influence.BaryCoordPos.w * WorldScale.x)
+					           + TriangleBary.y * (B + NB * Influence.BaryCoordPos.w * WorldScale.y)
+					           + TriangleBary.z * (C + NC * Influence.BaryCoordPos.w * WorldScale.z);
 
 			AveragedSimulatedPosition += Weight * SimPosition;
 	    }
@@ -784,13 +786,13 @@ float3x3 SkinTangents( FVertexFactoryInput Input, FVertexFactoryIntermediates In
                     SumWeights = 1.0f;	   
 			    }
 
-			    NormalPosition += Weight * (Influence.BaryCoordNormal.x * (A + NA * Influence.BaryCoordNormal.w)
-								          + Influence.BaryCoordNormal.y * (B + NB * Influence.BaryCoordNormal.w)
-								          + Influence.BaryCoordNormal.z * (C + NC * Influence.BaryCoordNormal.w));
+			    NormalPosition += Weight * (Influence.BaryCoordNormal.x * (A + NA * Influence.BaryCoordNormal.w * WorldScale.x)
+								          + Influence.BaryCoordNormal.y * (B + NB * Influence.BaryCoordNormal.w * WorldScale.y)
+								          + Influence.BaryCoordNormal.z * (C + NC * Influence.BaryCoordNormal.w * WorldScale.z));
 
-		        TangentPosition += Weight * (Influence.BaryCoordTangent.x * (A+NA*Influence.BaryCoordTangent.w)
-								           + Influence.BaryCoordTangent.y * (B+NB*Influence.BaryCoordTangent.w)
-								           + Influence.BaryCoordTangent.z * (C + NC * Influence.BaryCoordTangent.w));
+		        TangentPosition += Weight * (Influence.BaryCoordTangent.x * (A + NA * Influence.BaryCoordTangent.w * WorldScale.x)
+								           + Influence.BaryCoordTangent.y * (B + NB * Influence.BaryCoordTangent.w * WorldScale.y)
+								           + Influence.BaryCoordTangent.z * (C + NC * Influence.BaryCoordTangent.w * WorldScale.z));
 	        }
         }
 
diff --git a/Engine/Source/Runtime/Engine/Private/GPUSkinCache.cpp b/Engine/Source/Runtime/Engine/Private/GPUSkinCache.cpp
--- a/Engine/Source/Runtime/Engine/Private/GPUSkinCache.cpp
+++ b/Engine/Source/Runtime/Engine/Private/GPUSkinCache.cpp
@@ -349,6 +349,7 @@ public:
         float ClothBlendWeight = 0.0f;
 		uint32 ClothNumInfluencesPerVertex = 1;
         FMatrix44f ClothToLocal = FMatrix44f::Identity;
+		FVector3f WorldScale = FVector3f::OneVector;
 
 		// triangle index buffer (input for the RecomputeSkinTangents, might need special index buffer unique to position and normal, not considering UV/vertex color)
 		uint32 IndexBufferOffsetValue = 0;
@@ -655,6 +656,7 @@ public:
 		ClothBlendWeight.Bind(Initializer.ParameterMap, TEXT("ClothBlendWeight"));
 		ClothToLocal.Bind(Initializer.ParameterMap, TEXT("ClothToLocal"));
 		ClothNumInfluencesPerVertex.Bind(Initializer.ParameterMap, TEXT("ClothNumInfluencesPerVertex"));
+		WorldScale.Bind(Initializer.ParameterMap, TEXT("WorldScale"));
 	}
 
 	void SetParameters(
@@ -704,6 +706,7 @@ public:
 			SetShaderValue(RHICmdList, ShaderRHI, ClothBlendWeight, DispatchData.ClothBlendWeight);
 			SetShaderValue(RHICmdList, ShaderRHI, ClothToLocal, DispatchData.ClothToLocal);
 			SetShaderValue(RHICmdList, ShaderRHI, ClothNumInfluencesPerVertex, DispatchData.ClothNumInfluencesPerVertex);
+			SetShaderValue(RHICmdList, ShaderRHI, WorldScale, DispatchData.WorldScale);
 		}
 
 		SetShaderValue(RHICmdList, ShaderRHI, SkinCacheDebug, CVarGPUSkinCacheDebug.GetValueOnRenderThread());
@@ -751,6 +754,7 @@ private:
 	LAYOUT_FIELD(FShaderParameter, ClothBlendWeight)
 	LAYOUT_FIELD(FShaderParameter, ClothToLocal)
 	LAYOUT_FIELD(FShaderParameter, ClothNumInfluencesPerVertex)
+	LAYOUT_FIELD(FShaderParameter, WorldScale)
 };
 
 /** Compute shader that skins a batch of vertices. */
@@ -1457,7 +1461,8 @@ bool FGPUSkinCache::ProcessEntry(
 	const FSkeletalMeshVertexClothBuffer* ClothVertexBuffer, 
 	const FClothSimulData* SimData,
 	const FMatrix44f& ClothToLocal,
-	float ClothBlendWeight, 
+	float ClothBlendWeight,
+	FVector3f WorldScale,
 	uint32 RevisionNumber, 
 	int32 Section,
 	int32 LODIndex,
@@ -1643,6 +1648,7 @@ bool FGPUSkinCache::ProcessEntry(
 
         InOutEntry->DispatchData[Section].ClothBlendWeight = ClothBlendWeight;
         InOutEntry->DispatchData[Section].ClothToLocal = ClothToLocal;
+		InOutEntry->DispatchData[Section].WorldScale = WorldScale;
     }
     InOutEntry->DispatchData[Section].SkinType = ClothVertexBuffer && InOutEntry->DispatchData[Section].ClothPositionsAndNormalsBuffer ? 2 : (bMorph ? 1 : 0);
 
diff --git a/Engine/Source/Runtime/Engine/Private/GPUSkinVertexFactory.cpp b/Engine/Source/Runtime/Engine/Private/GPUSkinVertexFactory.cpp
--- a/Engine/Source/Runtime/Engine/Private/GPUSkinVertexFactory.cpp
+++ b/Engine/Source/Runtime/Engine/Private/GPUSkinVertexFactory.cpp
@@ -978,6 +978,7 @@ public:
 		ClothToLocalParameter.Bind(ParameterMap, TEXT("ClothToLocal"));
 		PreviousClothToLocalParameter.Bind(ParameterMap, TEXT("PreviousClothToLocal"));
 		ClothBlendWeightParameter.Bind(ParameterMap, TEXT("ClothBlendWeight"));
+		WorldScaleParameter.Bind(ParameterMap, TEXT("WorldScale"));
 		GPUSkinApexClothParameter.Bind(ParameterMap, TEXT("GPUSkinApexCloth"));
 		GPUSkinApexClothStartIndexOffsetParameter.Bind(ParameterMap, TEXT("GPUSkinApexClothStartIndexOffset"));
 		ClothNumInfluencesPerVertexParameter.Bind(ParameterMap, TEXT("ClothNumInfluencesPerVertex"));
@@ -1007,6 +1008,7 @@ public:
 		ShaderBindings.Add(ClothSimulVertsPositionsNormalsParameter, ClothShaderData.GetClothBufferForReading(false).VertexBufferSRV);
 		ShaderBindings.Add(ClothToLocalParameter, ClothShaderData.GetClothToLocalForReading(false));
 		ShaderBindings.Add(ClothBlendWeightParameter,ClothShaderData.ClothBlendWeight);
+		ShaderBindings.Add(WorldScaleParameter, ClothShaderData.WorldScale);
 		ShaderBindings.Add(ClothNumInfluencesPerVertexParameter, ClothShaderData.NumInfluencesPerVertex);
 
 		// Mobile doesn't support motion blur, no need to feed the previous frame cloth data
@@ -1031,6 +1033,7 @@ protected:
 	LAYOUT_FIELD(FShaderParameter, ClothToLocalParameter);
 	LAYOUT_FIELD(FShaderParameter, PreviousClothToLocalParameter);
 	LAYOUT_FIELD(FShaderParameter, ClothBlendWeightParameter);
+	LAYOUT_FIELD(FShaderParameter, WorldScaleParameter);
 	LAYOUT_FIELD(FShaderResourceParameter, GPUSkinApexClothParameter);
 	LAYOUT_FIELD(FShaderParameter, GPUSkinApexClothStartIndexOffsetParameter);
 	LAYOUT_FIELD(FShaderParameter, ClothNumInfluencesPerVertexParameter);
diff --git a/Engine/Source/Runtime/Engine/Private/SkeletalRender.cpp b/Engine/Source/Runtime/Engine/Private/SkeletalRender.cpp
--- a/Engine/Source/Runtime/Engine/Private/SkeletalRender.cpp
+++ b/Engine/Source/Runtime/Engine/Private/SkeletalRender.cpp
@@ -66,6 +66,7 @@ FSkeletalMeshObject::FSkeletalMeshObject(USkinnedMeshComponent* InMeshComponent,
 ,	StatId(InMeshComponent->GetSkinnedAsset()->GetStatID(true))
 ,	FeatureLevel(InFeatureLevel)
 ,	ComponentId(InMeshComponent->ComponentId.PrimIDValue)
+,	WorldScale(InMeshComponent->GetComponentScale())
 #if RHI_ENABLE_RESOURCE_INFO
 ,	AssetPathName(InMeshComponent->GetSkinnedAsset()->GetPathName())
 #endif
diff --git a/Engine/Source/Runtime/Engine/Private/SkeletalRenderCPUSkin.cpp b/Engine/Source/Runtime/Engine/Private/SkeletalRenderCPUSkin.cpp
--- a/Engine/Source/Runtime/Engine/Private/SkeletalRenderCPUSkin.cpp
+++ b/Engine/Source/Runtime/Engine/Private/SkeletalRenderCPUSkin.cpp
@@ -23,7 +23,7 @@
 struct FMorphTargetDelta;
 
 template<typename VertexType, int32 NumberOfUVs>
-static void SkinVertices(FFinalSkinVertex* DestVertex, FMatrix44f* ReferenceToLocal, int32 LODIndex, FSkeletalMeshLODRenderData& LOD, FSkinWeightVertexBuffer& WeightBuffer, const FMorphTargetWeightMap& ActiveMorphTargets, const TArray<float>& MorphTargetWeights, const TMap<int32, FClothSimulData>& ClothSimulUpdateData, float ClothBlendWeight, const FMatrix& WorldToLocal);
+static void SkinVertices(FFinalSkinVertex* DestVertex, FMatrix44f* ReferenceToLocal, int32 LODIndex, FSkeletalMeshLODRenderData& LOD, FSkinWeightVertexBuffer& WeightBuffer, const FMorphTargetWeightMap& ActiveMorphTargets, const TArray<float>& MorphTargetWeights, const TMap<int32, FClothSimulData>& ClothSimulUpdateData, float ClothBlendWeight, const FMatrix& WorldToLocal, const FVector& WorldScale);
 
 #define INFLUENCE_0		0
 #define INFLUENCE_1		1
@@ -206,16 +206,16 @@ void FSkeletalMeshObjectCPUSkin::UpdateDynamicData_RenderThread(FRHICommandListI
 	switch( NumUVs )\
     {\
         case 1:\
-			SkinVertices<VertexType<1>, 1>( DestVertex, ReferenceToLocal, DynamicData->LODIndex, LOD, *MeshLOD.MeshObjectWeightBuffer, DynamicData->ActiveMorphTargets, DynamicData->MorphTargetWeights, DynamicData->ClothSimulUpdateData, DynamicData->ClothBlendWeight, DynamicData->WorldToLocal); \
+			SkinVertices<VertexType<1>, 1>( DestVertex, ReferenceToLocal, DynamicData->LODIndex, LOD, *MeshLOD.MeshObjectWeightBuffer, DynamicData->ActiveMorphTargets, DynamicData->MorphTargetWeights, DynamicData->ClothSimulUpdateData, DynamicData->ClothBlendWeight, DynamicData->WorldToLocal, WorldScale); \
           	break;\
         case 2:\
-			SkinVertices<VertexType<2>, 2>( DestVertex, ReferenceToLocal, DynamicData->LODIndex, LOD, *MeshLOD.MeshObjectWeightBuffer, DynamicData->ActiveMorphTargets, DynamicData->MorphTargetWeights, DynamicData->ClothSimulUpdateData, DynamicData->ClothBlendWeight, DynamicData->WorldToLocal); \
+			SkinVertices<VertexType<2>, 2>( DestVertex, ReferenceToLocal, DynamicData->LODIndex, LOD, *MeshLOD.MeshObjectWeightBuffer, DynamicData->ActiveMorphTargets, DynamicData->MorphTargetWeights, DynamicData->ClothSimulUpdateData, DynamicData->ClothBlendWeight, DynamicData->WorldToLocal, WorldScale); \
           	break;\
         case 3:\
-			SkinVertices<VertexType<3>, 3>( DestVertex, ReferenceToLocal, DynamicData->LODIndex, LOD, *MeshLOD.MeshObjectWeightBuffer, DynamicData->ActiveMorphTargets, DynamicData->MorphTargetWeights, DynamicData->ClothSimulUpdateData, DynamicData->ClothBlendWeight, DynamicData->WorldToLocal); \
+			SkinVertices<VertexType<3>, 3>( DestVertex, ReferenceToLocal, DynamicData->LODIndex, LOD, *MeshLOD.MeshObjectWeightBuffer, DynamicData->ActiveMorphTargets, DynamicData->MorphTargetWeights, DynamicData->ClothSimulUpdateData, DynamicData->ClothBlendWeight, DynamicData->WorldToLocal, WorldScale); \
           	break;\
         case 4:\
-			SkinVertices<VertexType<4>, 4>( DestVertex, ReferenceToLocal, DynamicData->LODIndex, LOD, *MeshLOD.MeshObjectWeightBuffer, DynamicData->ActiveMorphTargets, DynamicData->MorphTargetWeights, DynamicData->ClothSimulUpdateData, DynamicData->ClothBlendWeight, DynamicData->WorldToLocal); \
+			SkinVertices<VertexType<4>, 4>( DestVertex, ReferenceToLocal, DynamicData->LODIndex, LOD, *MeshLOD.MeshObjectWeightBuffer, DynamicData->ActiveMorphTargets, DynamicData->MorphTargetWeights, DynamicData->ClothSimulUpdateData, DynamicData->ClothBlendWeight, DynamicData->WorldToLocal, WorldScale); \
           	break;\
         default:\
           	checkf(false, TEXT("Invalid number of UV sets.  Must be between 1 and 4") );\
@@ -712,7 +712,8 @@ static void SkinVertexSection(
 	const FMatrix44f* RESTRICT ReferenceToLocal, 
 	const FClothSimulData* ClothSimData, 
 	float ClothBlendWeight, 
-	const FMatrix& WorldToLocal)
+	const FMatrix& WorldToLocal,
+	const FVector& WorldScale )
 {
 	static constexpr VectorRegister VECTOR_INV_65535 = MakeVectorRegisterDoubleConstant(1.0 / 65535, 1.0 / 65535, 1.0 / 65535, 1.0 / 65535);
 	
@@ -955,14 +956,14 @@ static void SkinVertexSection(
 						return FVector(0, 0, 1);
 					}
 
-					FORCEINLINE static FVector ClothingPosition(const FMeshToMeshVertData& InClothVertData, const FClothSimulData& InClothSimData)
+					FORCEINLINE static FVector ClothingPosition(const FMeshToMeshVertData& InClothVertData, const FClothSimulData& InClothSimData, const FVector& InWorldScale)
 					{
-						return    InClothVertData.PositionBaryCoordsAndDist.X * (GetClothSimulPosition(InClothSimData, InClothVertData.SourceMeshVertIndices[0]) + GetClothSimulNormal(InClothSimData, InClothVertData.SourceMeshVertIndices[0]) * InClothVertData.PositionBaryCoordsAndDist.W)
-								+ InClothVertData.PositionBaryCoordsAndDist.Y * (GetClothSimulPosition(InClothSimData, InClothVertData.SourceMeshVertIndices[1]) + GetClothSimulNormal(InClothSimData, InClothVertData.SourceMeshVertIndices[1]) * InClothVertData.PositionBaryCoordsAndDist.W)
-								+ InClothVertData.PositionBaryCoordsAndDist.Z * (GetClothSimulPosition(InClothSimData, InClothVertData.SourceMeshVertIndices[2]) + GetClothSimulNormal(InClothSimData, InClothVertData.SourceMeshVertIndices[2]) * InClothVertData.PositionBaryCoordsAndDist.W);
+						return    InClothVertData.PositionBaryCoordsAndDist.X * (GetClothSimulPosition(InClothSimData, InClothVertData.SourceMeshVertIndices[0]) + GetClothSimulNormal(InClothSimData, InClothVertData.SourceMeshVertIndices[0]) * InClothVertData.PositionBaryCoordsAndDist.W * InWorldScale.X)
+								+ InClothVertData.PositionBaryCoordsAndDist.Y * (GetClothSimulPosition(InClothSimData, InClothVertData.SourceMeshVertIndices[1]) + GetClothSimulNormal(InClothSimData, InClothVertData.SourceMeshVertIndices[1]) * InClothVertData.PositionBaryCoordsAndDist.W * InWorldScale.Y)
+								+ InClothVertData.PositionBaryCoordsAndDist.Z * (GetClothSimulPosition(InClothSimData, InClothVertData.SourceMeshVertIndices[2]) + GetClothSimulNormal(InClothSimData, InClothVertData.SourceMeshVertIndices[2]) * InClothVertData.PositionBaryCoordsAndDist.W * InWorldScale.Z);
 					}
 
-					FORCEINLINE static void ClothingTangents(const FMeshToMeshVertData& InClothVertData, const FClothSimulData& InClothSimData, const FVector& InSimulatedPosition, const FMatrix& InWorldToLocal, FVector& OutTangentX, FVector& OutTangentZ)
+					FORCEINLINE static void ClothingTangents(const FMeshToMeshVertData& InClothVertData, const FClothSimulData& InClothSimData, const FVector& InSimulatedPosition, const FMatrix& InWorldToLocal, const FVector& InWorldScale, FVector& OutTangentX, FVector& OutTangentZ)
 					{
 						FVector A = GetClothSimulPosition(InClothSimData, InClothVertData.SourceMeshVertIndices[0]);
 						FVector B = GetClothSimulPosition(InClothSimData, InClothVertData.SourceMeshVertIndices[1]);
@@ -972,13 +973,13 @@ static void SkinVertexSection(
 						FVector NB = GetClothSimulNormal(InClothSimData, InClothVertData.SourceMeshVertIndices[1]);
 						FVector NC = GetClothSimulNormal(InClothSimData, InClothVertData.SourceMeshVertIndices[2]);
 
-						FVector NormalPosition = InClothVertData.NormalBaryCoordsAndDist.X*(A + NA*InClothVertData.NormalBaryCoordsAndDist.W)
-												+ InClothVertData.NormalBaryCoordsAndDist.Y*(B + NB*InClothVertData.NormalBaryCoordsAndDist.W)
-												+ InClothVertData.NormalBaryCoordsAndDist.Z*(C + NC*InClothVertData.NormalBaryCoordsAndDist.W);
+						FVector NormalPosition = InClothVertData.NormalBaryCoordsAndDist.X*(A + NA*InClothVertData.NormalBaryCoordsAndDist.W * InWorldScale.X)
+												+ InClothVertData.NormalBaryCoordsAndDist.Y*(B + NB*InClothVertData.NormalBaryCoordsAndDist.W * InWorldScale.Y)
+												+ InClothVertData.NormalBaryCoordsAndDist.Z*(C + NC*InClothVertData.NormalBaryCoordsAndDist.W * InWorldScale.Z);
 
-						FVector TangentPosition = InClothVertData.TangentBaryCoordsAndDist.X*(A + NA*InClothVertData.TangentBaryCoordsAndDist.W)
-												+ InClothVertData.TangentBaryCoordsAndDist.Y*(B + NB*InClothVertData.TangentBaryCoordsAndDist.W)
-												+ InClothVertData.TangentBaryCoordsAndDist.Z*(C + NC*InClothVertData.TangentBaryCoordsAndDist.W);
+						FVector TangentPosition = InClothVertData.TangentBaryCoordsAndDist.X*(A + NA*InClothVertData.TangentBaryCoordsAndDist.W * InWorldScale.X)
+												+ InClothVertData.TangentBaryCoordsAndDist.Y*(B + NB*InClothVertData.TangentBaryCoordsAndDist.W * InWorldScale.Y)
+												+ InClothVertData.TangentBaryCoordsAndDist.Z*(C + NC*InClothVertData.TangentBaryCoordsAndDist.W * InWorldScale.Z);
 
 						OutTangentX = (TangentPosition - InSimulatedPosition).GetUnsafeNormal();
 						OutTangentZ = (NormalPosition - InSimulatedPosition).GetUnsafeNormal();
@@ -990,7 +991,7 @@ static void SkinVertexSection(
 				};
 
 				// build sim position (in world space)
-				FVector SimulatedPositionWorld = ClothCPU::ClothingPosition(*ClothVertData, *ClothSimData);
+				FVector SimulatedPositionWorld = ClothCPU::ClothingPosition(*ClothVertData, *ClothSimData, WorldScale);
 
 				// transform back to local space
 				FVector3f SimulatedPosition = (FVector4f)WorldToLocal.TransformPosition(SimulatedPositionWorld);
@@ -1003,7 +1004,7 @@ static void SkinVertexSection(
 				// recompute tangent & normal
 				FVector TangentX;
 				FVector TangentZ;
-				ClothCPU::ClothingTangents(*ClothVertData, *ClothSimData, SimulatedPositionWorld, WorldToLocal, TangentX, TangentZ);
+				ClothCPU::ClothingTangents(*ClothVertData, *ClothSimData, SimulatedPositionWorld, WorldToLocal, WorldScale, TangentX, TangentZ);
 
 				// Lerp between skinned and simulated tangents
 				FVector SkinnedTangentX = DestVertex->TangentX.ToFVector();
@@ -1034,7 +1035,8 @@ static void SkinVertices(
 	const TArray<float>& MorphTargetWeights, 
 	const TMap<int32, FClothSimulData>& ClothSimulUpdateData, 
 	float ClothBlendWeight, 
-	const FMatrix& WorldToLocal)
+	const FMatrix& WorldToLocal,
+	const FVector& WorldScale)
 {
 	uint32 StatusRegister = VectorGetControlRegister();
 	VectorSetControlRegister( StatusRegister | VECTOR_ROUND_TOWARD_ZERO );
@@ -1062,7 +1064,7 @@ static void SkinVertices(
 
 		const FClothSimulData* ClothSimData = ClothSimulUpdateData.Find(Section.CorrespondClothAssetIndex);
 
-		SkinVertexSection<VertexType, NumberOfUVs>(DestVertex, MorphEvalInfos, MorphTargetWeights, Section, LOD, WeightBuffer, VertexBufferBaseIndex, NumValidMorphs, CurBaseVertIdx, LODIndex, ReferenceToLocal, ClothSimData, ClothBlendWeight, WorldToLocal);
+		SkinVertexSection<VertexType, NumberOfUVs>(DestVertex, MorphEvalInfos, MorphTargetWeights, Section, LOD, WeightBuffer, VertexBufferBaseIndex, NumValidMorphs, CurBaseVertIdx, LODIndex, ReferenceToLocal, ClothSimData, ClothBlendWeight, WorldToLocal, WorldScale);
 	}
 
 	VectorSetControlRegister( StatusRegister );
diff --git a/Engine/Source/Runtime/Engine/Private/SkeletalRenderGPUSkin.cpp b/Engine/Source/Runtime/Engine/Private/SkeletalRenderGPUSkin.cpp
--- a/Engine/Source/Runtime/Engine/Private/SkeletalRenderGPUSkin.cpp
+++ b/Engine/Source/Runtime/Engine/Private/SkeletalRenderGPUSkin.cpp
@@ -703,6 +703,7 @@ void FSkeletalMeshObjectGPUSkin::ProcessUpdatedDynamicData(EGPUSkinCacheEntryMod
 			{
 				FGPUBaseSkinAPEXClothVertexFactory::ClothShaderType& ClothShaderData = VertexFactoryData.ClothVertexFactories[SectionIdx]->GetClothShaderData();
 				ClothShaderData.ClothBlendWeight = DynamicData->ClothBlendWeight;
+				ClothShaderData.WorldScale = (FVector3f)WorldScale;
 
 				bNeedFence = ClothShaderData.UpdateClothSimulData(RHICmdList, SimData->Positions, SimData->Normals, RevisionNumber, FeatureLevel, DynamicData->bForceUpdateDynamicDataImmediately, OwnerName) || bNeedFence;
 				// Transform from cloth space to local space. Cloth space is relative to cloth root bone, local space is component space.
@@ -729,6 +730,7 @@ void FSkeletalMeshObjectGPUSkin::ProcessUpdatedDynamicData(EGPUSkinCacheEntryMod
 						bClothFactory ? DynamicData->ClothingSimData.Find(Section.CorrespondClothAssetIndex) : 0,
 						ClothToLocal,
 						DynamicData->ClothBlendWeight,
+						(FVector3f)WorldScale,
 						RevisionNumber,
 						SectionIdx,
 						LODIndex,
diff --git a/Engine/Source/Runtime/Engine/Public/GPUSkinCache.h b/Engine/Source/Runtime/Engine/Public/GPUSkinCache.h
--- a/Engine/Source/Runtime/Engine/Public/GPUSkinCache.h
+++ b/Engine/Source/Runtime/Engine/Public/GPUSkinCache.h
@@ -139,6 +139,7 @@ public:
 		const FClothSimulData* SimData,
 		const FMatrix44f& ClothToLocal,
 		float ClothBlendWeight, 
+		FVector3f ClothScale,
 		uint32 RevisionNumber, 
 		int32 Section,
 		int32 LOD,
diff --git a/Engine/Source/Runtime/Engine/Public/GPUSkinVertexFactory.h b/Engine/Source/Runtime/Engine/Public/GPUSkinVertexFactory.h
--- a/Engine/Source/Runtime/Engine/Public/GPUSkinVertexFactory.h
+++ b/Engine/Source/Runtime/Engine/Public/GPUSkinVertexFactory.h
@@ -540,6 +540,8 @@ public:
 		 * if ClothBlendWeight is 1.0, it shows only simulated positions and if it is 0.0, it shows only key-framed animation
 		 */
 		float ClothBlendWeight = 1.0f;
+		/** Scale of the owner actor */
+		FVector3f WorldScale = FVector3f::OneVector;
 		uint32 NumInfluencesPerVertex = 1;
 
 	private:
diff --git a/Engine/Source/Runtime/Engine/Public/SkeletalRenderPublic.h b/Engine/Source/Runtime/Engine/Public/SkeletalRenderPublic.h
--- a/Engine/Source/Runtime/Engine/Public/SkeletalRenderPublic.h
+++ b/Engine/Source/Runtime/Engine/Public/SkeletalRenderPublic.h
@@ -353,6 +353,8 @@ protected:
 	/** Component ID to which belong this  mesh object  */
 	uint32 ComponentId;
 
+	FVector WorldScale = FVector::OneVector;
+
 #if RHI_ENABLE_RESOURCE_INFO
 	FName AssetPathName;
 #endif
