fix physics constraint component crash when it unregisters after the geometry collection component when they're both components of the same actor

diff --git a/Engine/Source/Runtime/Engine/Classes/PhysicsEngine/PhysicsConstraintComponent.h b/Engine/Source/Runtime/Engine/Classes/PhysicsEngine/PhysicsConstraintComponent.h
--- a/Engine/Source/Runtime/Engine/Classes/PhysicsEngine/PhysicsConstraintComponent.h
+++ b/Engine/Source/Runtime/Engine/Classes/PhysicsEngine/PhysicsConstraintComponent.h
@@ -81,6 +81,7 @@ public:
 	virtual void CheckForErrors() override;
 	virtual void OnRegister() override;
 #endif // WITH_EDITOR
+	virtual void EndPlay(const EEndPlayReason::Type EndPlayReason) override;
 	virtual void OnUnregister() override;
 	virtual void InitializeComponent() override;
 	//End ActorComponent interface
diff --git a/Engine/Source/Runtime/Engine/Private/PhysicsEngine/PhysicsConstraintComponent.cpp b/Engine/Source/Runtime/Engine/Private/PhysicsEngine/PhysicsConstraintComponent.cpp
--- a/Engine/Source/Runtime/Engine/Private/PhysicsEngine/PhysicsConstraintComponent.cpp
+++ b/Engine/Source/Runtime/Engine/Private/PhysicsEngine/PhysicsConstraintComponent.cpp
@@ -374,6 +374,12 @@ void UPhysicsConstraintComponent::OnRegister()
 }
 #endif
 
+void UPhysicsConstraintComponent::EndPlay(const EEndPlayReason::Type EndPlayReason)
+{
+	Super::EndPlay(EndPlayReason);
+	TermComponentConstraint();
+}
+
 void UPhysicsConstraintComponent::OnUnregister()
 {
 	Super::OnUnregister();
