Implements FRHIGPUUsageFractions interface

diff --git a/Engine/Source/Runtime/Engine/Private/UnrealClient.cpp b/Engine/Source/Runtime/Engine/Private/UnrealClient.cpp
--- a/Engine/Source/Runtime/Engine/Private/UnrealClient.cpp
+++ b/Engine/Source/Runtime/Engine/Private/UnrealClient.cpp
@@ -477,6 +477,20 @@ int32 FStatUnitData::DrawStat(FViewport* InViewport, FCanvas* InCanvas, int32 In
 		const uint32 GPUCycles = RHIGetGPUFrameCycles(GPUIndex);
 		RawGPUFrameTime[GPUIndex] = FPlatformTime::ToMilliseconds(GPUCycles);
 		GPUFrameTime[GPUIndex] = 0.9 * GPUFrameTime[GPUIndex] + 0.1 * RawGPUFrameTime[GPUIndex];
+
+		if (GRHISupportsGPUUsage)
+		{
+			FRHIGPUUsageFractions GPUUsageFractions = RHIGetGPUUsage(GPUIndex);
+
+			RawGPUClockFraction[GPUIndex] = GPUUsageFractions.ClockScaling;
+			GPUClockFraction[GPUIndex] = 0.9 * GPUClockFraction[GPUIndex] + 0.1 * RawGPUClockFraction[GPUIndex];
+
+			RawGPUUsageFraction[GPUIndex] = GPUUsageFractions.CurrentProcess;
+			GPUUsageFraction[GPUIndex] = 0.9 * GPUUsageFraction[GPUIndex] + 0.1 * RawGPUUsageFraction[GPUIndex];
+
+			RawGPUExternalUsageFraction[GPUIndex] = GPUUsageFractions.ExternalProcesses;
+			GPUExternalUsageFraction[GPUIndex] = 0.9 * GPUExternalUsageFraction[GPUIndex] + 0.1 * RawGPUExternalUsageFraction[GPUIndex];
+		}
 	}
 
 	SET_FLOAT_STAT(STAT_UnitFrame, FrameTime);
@@ -576,113 +590,237 @@ int32 FStatUnitData::DrawStat(FViewport* InViewport, FCanvas* InCanvas, int32 In
 
 	// Draw unit.
 	{
+		const FColor NoUnitGraphColor(220, 220, 220);
+
+		#define STATUNIT_FORMAT_AVGTIME TEXT("%3.2f ms")
+		#define STATUNIT_FORMAT_MAXTIME TEXT("%4.2f ms")
+		#define STATUNIT_FORMAT_PERCEMT TEXT("%3.2f %%")
+
+		const int32 AvgUnitColumnId = 0;
+		const int32 MaxUnitColumnId = 1;
+
+		const int32 ColumnWidth = Font->GetStringSize(TEXT(" 0000.00 ms "));
+		const int32 ColumnCount = bShowUnitMaxTimes ? 2 : 1;
+
 		int32 X3 = InX * (bStereoRendering ? 0.5f : 1.0f);
-		if (bShowUnitMaxTimes)
-		{
-			X3 -= (int32)((float)Font->GetStringSize(TEXT(" 000.00 ms ")));
-		}
+		int32 X2 = X3 - ColumnWidth * (ColumnCount - 1);
 
-		int32 X2 = bShowUnitMaxTimes ? X3 - (int32)((float)Font->GetStringSize(TEXT(" 000.00 ms "))) : X3;
 		const int32 RowHeight = FMath::TruncToInt(Font->GetMaxCharHeight() * 1.1f);
 
-		auto DrawTitleString = [&](const TCHAR* Title, const FColor& UnitGraphColor)
+		/* Draw a cell on raw that have custom number of columns. */
+		auto DrawCell = [&](int32 RowId, int32 ColumnId, int32 ColumnCount, const FString& CellText, const FColor& CellColor)
+		{
+			check(ColumnId < ColumnCount);
+			int32 CellCoordX = X3 - ColumnWidth * (ColumnCount - ColumnId - 1);
+			int32 CellCoordY = InY + RowId * RowHeight;
+			InCanvas->DrawShadowedString(CellCoordX, CellCoordY, *CellText, Font, CellColor);
+		};
+
+		/* Draw a row's title that have a custom number of columns. */
+		auto DrawRowTitle = [&](int32 RowId, int32 InColumnCount, const TCHAR* Title, const FColor& UnitGraphColor)
 		{
 			FString FullTitle = FString::Printf(TEXT("%s:   "), Title);
 			int32 TitleSize = Font->GetStringSize(*FullTitle);
-			InCanvas->DrawShadowedString(X2 - TitleSize, InY, *FullTitle, Font, bShowUnitTimeGraph ? UnitGraphColor : FColor::White);
+
+			int32 TitleCoordX = X3 - ColumnWidth * (InColumnCount - 1);
+			int32 TitleCoordY = InY + RowId * RowHeight;
+			InCanvas->DrawShadowedString(TitleCoordX - TitleSize, TitleCoordY, *FullTitle, Font, bShowUnitTimeGraph ? UnitGraphColor : FColor::White);
+		};
+
+		/* Push drawing of rows toward the bottom of the screen. */
+		auto PushRows = [&](int32 RowsCount)
+		{
+			InY += RowHeight * RowsCount;
+		};
+
+		/* Draw a default row's title. */
+		auto DrawTitleString = [&](const TCHAR* Title, const FColor& UnitGraphColor)
+		{
+			DrawRowTitle(/* RowId = */ 0, ColumnCount, Title, UnitGraphColor);
+		};
+
+		/* Draw a default row's avg. */
+		auto DrawDefaultAvgCell = [&](const FString& CellText, const FColor& UnitGraphColor)
+		{
+			DrawCell(/* RowId = */ 0, AvgUnitColumnId, ColumnCount, CellText, UnitGraphColor);
+		};
+		auto DrawDefaultMaxCell = [&](const FString& CellText, const FColor& UnitGraphColor)
+		{
+			check(bShowUnitMaxTimes);
+			DrawCell(/* RowId = */ 0, MaxUnitColumnId, ColumnCount, CellText, UnitGraphColor);
 		};
 
 		{
-			const FColor FrameTimeAverageColor = GEngine->GetFrameTimeDisplayColor(FrameTime);
 			DrawTitleString(TEXT("Frame"), /* UnitGraphColor = */ FColor(100, 255, 100));
-			InCanvas->DrawShadowedString(X2, InY, *FString::Printf(TEXT("%3.2f ms"), FrameTime), Font, FrameTimeAverageColor);
+			{
+				const FColor FrameTimeAverageColor = GEngine->GetFrameTimeDisplayColor(FrameTime);
+				DrawDefaultAvgCell(FString::Printf(STATUNIT_FORMAT_AVGTIME, FrameTime), FrameTimeAverageColor);
+			}
 			if (bShowUnitMaxTimes)
 			{
 				const FColor MaxFrameTimeColor = GEngine->GetFrameTimeDisplayColor(Max_FrameTime);
-				InCanvas->DrawShadowedString(X3, InY, *FString::Printf(TEXT("%4.2f ms"), Max_FrameTime), Font, MaxFrameTimeColor);
+				DrawDefaultMaxCell(FString::Printf(STATUNIT_FORMAT_MAXTIME, Max_FrameTime), MaxFrameTimeColor);
 			}
-			InY += RowHeight;
+			PushRows(/* RowsCount = */ 1);
 		}
 
 		{
-			const FColor GameThreadAverageColor = GEngine->GetFrameTimeDisplayColor(GameThreadTime);
 			DrawTitleString(TEXT("Game"), /* UnitGraphColor = */ FColor(255, 100, 100));
-			InCanvas->DrawShadowedString(X2, InY, *FString::Printf(TEXT("%3.2f ms"), GameThreadTime), Font, GameThreadAverageColor);
+			{
+				const FColor GameThreadAverageColor = GEngine->GetFrameTimeDisplayColor(GameThreadTime);
+				DrawDefaultAvgCell(FString::Printf(STATUNIT_FORMAT_AVGTIME, GameThreadTime), GameThreadAverageColor);
+			}
 			if (bShowUnitMaxTimes)
 			{
 				const FColor GameThreadMaxColor = GEngine->GetFrameTimeDisplayColor(Max_GameThreadTime);
-				InCanvas->DrawShadowedString(X3, InY, *FString::Printf(TEXT("%4.2f ms"), Max_GameThreadTime), Font, GameThreadMaxColor);
+				DrawDefaultMaxCell(FString::Printf(STATUNIT_FORMAT_MAXTIME, Max_GameThreadTime), GameThreadMaxColor);
 			}
-			InY += RowHeight;
+			PushRows(/* RowsCount = */ 1);
 		}
 
 		{
-			const FColor RenderThreadAverageColor = GEngine->GetFrameTimeDisplayColor(RenderThreadTime);
 			DrawTitleString(TEXT("Draw"), /* UnitGraphColor = */ FColor(100, 100, 255));
-			InCanvas->DrawShadowedString(X2, InY, *FString::Printf(TEXT("%3.2f ms"), RenderThreadTime), Font, RenderThreadAverageColor);
+			{
+				const FColor RenderThreadAverageColor = GEngine->GetFrameTimeDisplayColor(RenderThreadTime);
+				DrawDefaultAvgCell(FString::Printf(STATUNIT_FORMAT_AVGTIME, RenderThreadTime), RenderThreadAverageColor);
+			}
 			if (bShowUnitMaxTimes)
 			{
 				const FColor RenderThreadMaxColor = GEngine->GetFrameTimeDisplayColor(Max_RenderThreadTime);
-				InCanvas->DrawShadowedString(X3, InY, *FString::Printf(TEXT("%4.2f ms"), Max_RenderThreadTime), Font, RenderThreadMaxColor);
+				DrawDefaultMaxCell(FString::Printf(STATUNIT_FORMAT_MAXTIME, Max_RenderThreadTime), RenderThreadMaxColor);
 			}
-			InY += RowHeight;
+			PushRows(/* RowsCount = */ 1);
 		}
 
-		for (uint32 GPUIndex : FRHIGPUMask::All())
-		{
-			if (bHaveGPUData[GPUIndex])
-			{
-				const FColor GPUAverageColor = GEngine->GetFrameTimeDisplayColor(GPUFrameTime[GPUIndex]);
-				FString GPUString = GNumExplicitGPUsForRendering > 1 ? FString::Printf(TEXT("GPU%u"), GPUIndex) : TEXT("GPU");
-				DrawTitleString(*GPUString, /* UnitGraphColor = */ FColor(255, 255, 100));
-				InCanvas->DrawShadowedString(X2, InY, *FString::Printf(TEXT("%3.2f ms"), GPUFrameTime[GPUIndex]), Font, GPUAverageColor);
-				if (bShowUnitMaxTimes)
-				{
-					const FColor GPUMaxColor = GEngine->GetFrameTimeDisplayColor(Max_GPUFrameTime[GPUIndex]);
-					InCanvas->DrawShadowedString(X3, InY, *FString::Printf(TEXT("%4.2f ms"), Max_GPUFrameTime[GPUIndex]), Font, GPUMaxColor);
-				}
-				InY += RowHeight;
-			}
-		}
 		if (IsRunningRHIInSeparateThread())
 		{
-			const FColor RenderThreadAverageColor = GEngine->GetFrameTimeDisplayColor(RHITTime);
 			DrawTitleString(TEXT("RHIT"), /* UnitGraphColor = */ FColor(255, 100, 255));
-			InCanvas->DrawShadowedString(X2, InY, *FString::Printf(TEXT("%3.2f ms"), RHITTime), Font, RenderThreadAverageColor);
+			{
+				const FColor RenderThreadAverageColor = GEngine->GetFrameTimeDisplayColor(RHITTime);
+				DrawDefaultAvgCell(FString::Printf(STATUNIT_FORMAT_AVGTIME, RHITTime), RenderThreadAverageColor);
+			}
 			if (bShowUnitMaxTimes)
 			{
 				const FColor RenderThreadMaxColor = GEngine->GetFrameTimeDisplayColor(Max_RHITTime);
-				InCanvas->DrawShadowedString(X3, InY, *FString::Printf(TEXT("%4.2f ms"), Max_RHITTime), Font, RenderThreadMaxColor);
+				DrawDefaultMaxCell(FString::Printf(STATUNIT_FORMAT_MAXTIME, Max_RHITTime), RenderThreadMaxColor);
 			}
-			InY += RowHeight;
+			PushRows(/* RowsCount = */ 1);
 		}
+
+		// Draw all GPU informations
+		{
+			const int32 ColumnPerGPU = ColumnCount;
+			const int32 GPURowCount = 1 + (GRHISupportsGPUUsage ? 3 : 0) + (GNumExplicitGPUsForRendering > 1 ? 1 : 0);
+			const int32 GPUColumnCount = ColumnCount * GNumExplicitGPUsForRendering;
+
+			// Draw the different timings:
+			{
+				int32 GPURowId = GNumExplicitGPUsForRendering > 1 ? 1 : 0;
+				DrawRowTitle(GPURowId++, GPUColumnCount, TEXT("GPU Time"), /* UnitGraphColor = */ FColor(255, 255, 100));
+
+				if (GRHISupportsGPUUsage)
+				{
+					DrawRowTitle(GPURowId++, GPUColumnCount, TEXT("GPU Clock"), NoUnitGraphColor);
+					DrawRowTitle(GPURowId++, GPUColumnCount, TEXT("GPU Usage"), NoUnitGraphColor);
+					DrawRowTitle(GPURowId++, GPUColumnCount, TEXT("GPU External"), NoUnitGraphColor);
+				}
+				check(GPURowId == GPURowCount);
+			}
+
+			// Draw each GPU
+			for (uint32 GPUIndex : FRHIGPUMask::All())
+			{
+				int32 GPURowId = 0;
+				if (GNumExplicitGPUsForRendering > 1)
+				{
+					DrawCell(
+						GPURowId++, GPUIndex * ColumnPerGPU + AvgUnitColumnId, GPUColumnCount,
+						FString::Printf(TEXT("GPU%d"), GPUIndex), NoUnitGraphColor);
+				}
+
+				if (!bHaveGPUData[GPUIndex])
+				{
+					continue;
+				}
+
+				{
+					{
+						const FColor GPUAverageColor = GEngine->GetFrameTimeDisplayColor(GPUFrameTime[GPUIndex]);
+						DrawCell(
+							GPURowId, GPUIndex * ColumnPerGPU + AvgUnitColumnId, GPUColumnCount,
+							FString::Printf(STATUNIT_FORMAT_AVGTIME, GPUFrameTime[GPUIndex]), GPUAverageColor);
+					}
+
+					if (bShowUnitMaxTimes)
+					{
+						const FColor GPUMaxColor = GEngine->GetFrameTimeDisplayColor(Max_GPUFrameTime[GPUIndex]);
+						DrawCell(
+							GPURowId, GPUIndex* ColumnPerGPU + MaxUnitColumnId, GPUColumnCount,
+							FString::Printf(STATUNIT_FORMAT_MAXTIME, Max_GPUFrameTime[GPUIndex]), GPUMaxColor);
+					}
+
+					GPURowId++;
+				}
+
+				if (GRHISupportsGPUUsage)
+				{
+					{
+						const FColor Color = (GPUClockFraction[GPUIndex] < 0.5f) ? StatRed : ((GPUClockFraction[GPUIndex] < 0.75f) ? StatOrange : StatGreen);
+						DrawCell(
+							GPURowId++, GPUIndex* ColumnPerGPU + AvgUnitColumnId, GPUColumnCount,
+							FString::Printf(STATUNIT_FORMAT_PERCEMT, 100.0f * GPUClockFraction[GPUIndex]), Color);
+					}
+
+					{
+						const FColor Color = (GPUUsageFraction[GPUIndex] < 0.5f) ? StatRed : ((GPUUsageFraction[GPUIndex] < 0.75f) ? StatOrange : StatGreen);
+						DrawCell(
+							GPURowId++, GPUIndex* ColumnPerGPU + AvgUnitColumnId, GPUColumnCount,
+							FString::Printf(STATUNIT_FORMAT_PERCEMT, 100.0f * GPUUsageFraction[GPUIndex]), Color);
+					}
+
+					{
+						const FColor Color = (GPUExternalUsageFraction[GPUIndex] > 0.2f) ? StatRed : ((GPUExternalUsageFraction[GPUIndex] > 0.1f) ? StatOrange : StatGreen);
+						DrawCell(
+							GPURowId++, GPUIndex* ColumnPerGPU + AvgUnitColumnId, GPUColumnCount,
+							FString::Printf(STATUNIT_FORMAT_PERCEMT, 100.0f * GPUExternalUsageFraction[GPUIndex]), Color);
+					}
+				}
+				check(GPURowId == GPURowCount);
+			}
+
+			PushRows(/* RowsCount = */ GPURowCount);
+		}
+
 		if (bHaveInputLatencyData)
 		{
 			const float ReasonableInputLatencyFactor = 2.5f;
-			const FColor InputLatencyAverageColor = GEngine->GetFrameTimeDisplayColor(InputLatencyTime / ReasonableInputLatencyFactor);
 			DrawTitleString(TEXT("Input"), /* UnitGraphColor = */ FColor(255, 255, 100));
-			InCanvas->DrawShadowedString(X2, InY, *FString::Printf(TEXT("%3.2f ms"), InputLatencyTime), Font, InputLatencyAverageColor);
+			{
+				const FColor InputLatencyAverageColor = GEngine->GetFrameTimeDisplayColor(InputLatencyTime / ReasonableInputLatencyFactor);
+				DrawDefaultAvgCell(FString::Printf(STATUNIT_FORMAT_AVGTIME, InputLatencyTime), InputLatencyAverageColor);
+			}
 			if (bShowUnitMaxTimes)
 			{
 				const FColor InputLatencyMaxColor = GEngine->GetFrameTimeDisplayColor(Max_InputLatencyTime / ReasonableInputLatencyFactor);
-				InCanvas->DrawShadowedString(X3, InY, *FString::Printf(TEXT("%4.2f ms"), Max_InputLatencyTime), Font, InputLatencyMaxColor);
+				DrawDefaultMaxCell(FString::Printf(STATUNIT_FORMAT_MAXTIME, Max_InputLatencyTime), InputLatencyMaxColor);
 			}
-			InY += RowHeight;
+			PushRows(/* RowsCount = */ 1);
 		}
+
 		{
 			if (bShowUnitMaxTimes)
 			{
 				FPlatformMemoryStats Stats = FPlatformMemory::GetStats();
 
-				DrawTitleString(TEXT("Mem"), /* UnitGraphColor = */ FColor(100, 100, 255));
-				InCanvas->DrawShadowedString(X2, InY, *GetMemoryString(Stats.UsedPhysical), Font, StatGreen);
-				InCanvas->DrawShadowedString(X3, InY, *GetMemoryString(Stats.PeakUsedPhysical), Font, StatGreen);
-				InY += RowHeight;
+				DrawTitleString(TEXT("Mem"), NoUnitGraphColor);
+				DrawDefaultAvgCell(GetMemoryString(Stats.UsedPhysical), StatGreen);
+				DrawDefaultMaxCell(GetMemoryString(Stats.PeakUsedPhysical), StatGreen);
+				PushRows(/* RowsCount = */ 1);
 				
-				DrawTitleString(TEXT("VMem"), /* UnitGraphColor = */ FColor(100, 100, 255));
-				InCanvas->DrawShadowedString(X2, InY, *GetMemoryString(Stats.UsedVirtual), Font, StatGreen);
-				InCanvas->DrawShadowedString(X3, InY, *GetMemoryString(Stats.PeakUsedVirtual), Font, StatGreen);
-				InY += RowHeight;
+				DrawTitleString(TEXT("VMem"), NoUnitGraphColor);
+				DrawDefaultAvgCell(GetMemoryString(Stats.UsedVirtual), StatGreen);
+				DrawDefaultMaxCell(GetMemoryString(Stats.PeakUsedVirtual), StatGreen);
+				PushRows(/* RowsCount = */ 1);
 			}
 			else
 			{
@@ -690,9 +828,9 @@ int32 FStatUnitData::DrawStat(FViewport* InViewport, FCanvas* InCanvas, int32 In
 				if (MemoryUsed > 0)
 				{
 					// print out currently used memory
-					DrawTitleString(TEXT("Mem"), /* UnitGraphColor = */ FColor(100, 100, 255));
-					InCanvas->DrawShadowedString(X2, InY, *GetMemoryString(MemoryUsed), Font, StatGreen);
-					InY += RowHeight;
+					DrawTitleString(TEXT("Mem"), NoUnitGraphColor);
+					DrawDefaultAvgCell(GetMemoryString(MemoryUsed), StatGreen);
+					PushRows(/* RowsCount = */ 1);
 				}
 			}
 		}
@@ -706,29 +844,29 @@ int32 FStatUnitData::DrawStat(FViewport* InViewport, FCanvas* InCanvas, int32 In
 			if (DynamicResolutionStateInfos.Status == EDynamicResolutionStatus::Enabled)
 			{
 				FColor Color = (ResolutionFraction < AlertResolutionFraction) ? StatRed : ((ResolutionFraction < FMath::Min(ResolutionFraction * 0.97f, 1.0f)) ? StatOrange : StatGreen);
-				InCanvas->DrawShadowedString(X2, InY, *FString::Printf(TEXT("%3.1f%% x %3.1f%%"), ScreenPercentage, ScreenPercentage), Font, Color);
+				DrawDefaultAvgCell(FString::Printf(TEXT("%3.1f%% x %3.1f%%"), ScreenPercentage, ScreenPercentage), Color);
 			}
 			else if (DynamicResolutionStateInfos.Status == EDynamicResolutionStatus::DebugForceEnabled)
 			{
-				InCanvas->DrawShadowedString(X2, InY, *FString::Printf(TEXT("%3.1f%% x %3.1f%%"), ScreenPercentage, ScreenPercentage), Font, StatMagenda);
+				DrawDefaultAvgCell(FString::Printf(TEXT("%3.1f%% x %3.1f%%"), ScreenPercentage, ScreenPercentage), StatMagenda);
 			}
 			else if (DynamicResolutionStateInfos.Status == EDynamicResolutionStatus::Paused)
 			{
-				InCanvas->DrawShadowedString(X2, InY, TEXT("Paused"), Font, StatMagenda);
+				DrawDefaultAvgCell(TEXT("Paused"), StatMagenda);
 			}
 			else if (DynamicResolutionStateInfos.Status == EDynamicResolutionStatus::Disabled)
 			{
-				InCanvas->DrawShadowedString(X2, InY, TEXT("OFF"), Font, FColor(160, 160, 160));
+				DrawDefaultAvgCell(TEXT("OFF"), FColor(160, 160, 160));
 			}
 			else if (DynamicResolutionStateInfos.Status == EDynamicResolutionStatus::Unsupported)
 			{
-				InCanvas->DrawShadowedString(X2, InY, TEXT("Unsupported"), Font, FColor(160, 160, 160));
+				DrawDefaultAvgCell(TEXT("Unsupported"), FColor(160, 160, 160));
 			}
 			else
 			{
 				check(0);
 			}
-			InY += RowHeight;
+			PushRows(/* RowsCount = */ 1);
 		}
 
 		// Other dynamic render scalings
@@ -755,13 +893,13 @@ int32 FStatUnitData::DrawStat(FViewport* InViewport, FCanvas* InCanvas, int32 In
 				DrawTitleString(*DisplayName, /* UnitGraphColor = */ FColor::White);
 				if (HeuristicSettings.Model == DynamicRenderScaling::EHeuristicModel::Quadratic)
 				{
-					InCanvas->DrawShadowedString(X2, InY, *FString::Printf(TEXT("%3.1f%% x %3.1f%%"), ScreenPercentage, ScreenPercentage), Font, Color);
+					DrawDefaultAvgCell(FString::Printf(TEXT("%3.1f%% x %3.1f%%"), ScreenPercentage, ScreenPercentage), Color);
 				}
 				else
 				{
-					InCanvas->DrawShadowedString(X2, InY, *FString::Printf(TEXT("%3.1f%%"), ScreenPercentage), Font, Color);
+					DrawDefaultAvgCell(FString::Printf(TEXT("%3.1f%%"), ScreenPercentage), Color);
 				}
-				InY += RowHeight;
+				PushRows(/* RowsCount = */ 1);
 			}
 		}
 
@@ -769,27 +907,27 @@ int32 FStatUnitData::DrawStat(FViewport* InViewport, FCanvas* InCanvas, int32 In
 		{
 			// Assume we don't have more than 1 GPU in mobile.
 			int32 NumDrawCalls = GNumDrawCallsRHI[0];
-			DrawTitleString(TEXT("Draws"), /* UnitGraphColor = */ FColor(100, 100, 255));
-			InCanvas->DrawShadowedString(X2, InY, *FString::Printf(TEXT("%d"), NumDrawCalls), Font, StatGreen);
-			InY += RowHeight;
+			DrawTitleString(TEXT("Draws"), NoUnitGraphColor);
+			DrawDefaultAvgCell(FString::Printf(TEXT("%d"), NumDrawCalls), StatGreen);
+			PushRows(/* RowsCount = */ 1);
 		}
 			
 		// Primitives
 		{
 			// Assume we don't have more than 1 GPU in mobile.
 			int32 NumPrimitives = GNumPrimitivesDrawnRHI[0];
-			DrawTitleString(TEXT("Prims"), /* UnitGraphColor = */ FColor(100, 100, 255));
+			DrawTitleString(TEXT("Prims"), NoUnitGraphColor);
 			if (NumPrimitives < 10000)
 			{
-				InCanvas->DrawShadowedString(X2, InY, *FString::Printf(TEXT("%d"), NumPrimitives), Font, StatGreen);
+				DrawDefaultAvgCell(FString::Printf(TEXT("%d"), NumPrimitives), StatGreen);
 			}
 			else
 			{
 				float NumPrimitivesK = NumPrimitives/1000.f;
-				InCanvas->DrawShadowedString(X2, InY, *FString::Printf(TEXT("%.1fK"), NumPrimitivesK), Font, StatGreen);
+				DrawDefaultAvgCell(FString::Printf(TEXT("%.1fK"), NumPrimitivesK), StatGreen);
 			}
-				
-			InY += RowHeight;
+
+			PushRows(/* RowsCount = */ 1);
 		}
 	}
 
diff --git a/Engine/Source/Runtime/Engine/Public/UnrealClient.h b/Engine/Source/Runtime/Engine/Public/UnrealClient.h
--- a/Engine/Source/Runtime/Engine/Public/UnrealClient.h
+++ b/Engine/Source/Runtime/Engine/Public/UnrealClient.h
@@ -239,6 +239,9 @@ struct FStatUnitData
 	float RenderThreadTime;
 	float GameThreadTime;
 	float GPUFrameTime[MAX_NUM_GPUS];
+	float GPUClockFraction[MAX_NUM_GPUS];
+	float GPUUsageFraction[MAX_NUM_GPUS];
+	float GPUExternalUsageFraction[MAX_NUM_GPUS];
 	float FrameTime;
 	float RHITTime;
 	float InputLatencyTime;
@@ -247,6 +250,9 @@ struct FStatUnitData
 	float RawRenderThreadTime;
 	float RawGameThreadTime;
 	float RawGPUFrameTime[MAX_NUM_GPUS];
+	float RawGPUClockFraction[MAX_NUM_GPUS];
+	float RawGPUUsageFraction[MAX_NUM_GPUS];
+	float RawGPUExternalUsageFraction[MAX_NUM_GPUS];
 	float RawFrameTime;
 	float RawRHITTime;
 	float RawInputLatencyTime;
@@ -271,12 +277,18 @@ struct FStatUnitData
 		: RenderThreadTime(0.0f)
 		, GameThreadTime(0.0f)
 		, GPUFrameTime{ 0.0f }
+		, GPUClockFraction{ 0.0f }
+		, GPUUsageFraction{ 0.0f }
+		, GPUExternalUsageFraction{ 0.0f }
 		, FrameTime(0.0f)
 		, RHITTime(0.0f)
 		, InputLatencyTime(0.0f)
 		, RawRenderThreadTime(0.0f)
 		, RawGameThreadTime(0.0f)
 		, RawGPUFrameTime{ 0.0f }
+		, RawGPUClockFraction{ 0.0f }
+		, RawGPUUsageFraction{ 0.0f }
+		, RawGPUExternalUsageFraction{ 0.0f }
 		, RawFrameTime(0.0f)
 		, RawRHITTime(0.0f)
 		, RawInputLatencyTime(0.0f)
diff --git a/Engine/Source/Runtime/RHI/Private/DynamicRHI.cpp b/Engine/Source/Runtime/RHI/Private/DynamicRHI.cpp
--- a/Engine/Source/Runtime/RHI/Private/DynamicRHI.cpp
+++ b/Engine/Source/Runtime/RHI/Private/DynamicRHI.cpp
@@ -34,6 +34,7 @@ static_assert(sizeof(FRayTracingGeometryInstance) <= 104,
 
 // Globals.
 FDynamicRHI* GDynamicRHI = NULL;
+RHIGetGPUUsageType RHIGetGPUUsage = nullptr;
 
 static TAutoConsoleVariable<int32> CVarWarnOfBadDrivers(
 	TEXT("r.WarnOfBadDrivers"),
diff --git a/Engine/Source/Runtime/RHI/Private/RHI.cpp b/Engine/Source/Runtime/RHI/Private/RHI.cpp
--- a/Engine/Source/Runtime/RHI/Private/RHI.cpp
+++ b/Engine/Source/Runtime/RHI/Private/RHI.cpp
@@ -1247,6 +1247,7 @@ bool GSupportsPixelLocalStorage = false;
 bool GSupportsTimestampRenderQueries = false;
 bool GRHISupportsGPUTimestampBubblesRemoval = false;
 bool GRHISupportsFrameCyclesBubblesRemoval = false;
+bool GRHISupportsGPUUsage = false;
 bool GHardwareHiddenSurfaceRemoval = false;
 bool GRHISupportsAsyncTextureCreation = false;
 bool GRHISupportsQuadTopology = false;
diff --git a/Engine/Source/Runtime/RHI/Public/DynamicRHI.h b/Engine/Source/Runtime/RHI/Public/DynamicRHI.h
--- a/Engine/Source/Runtime/RHI/Public/DynamicRHI.h
+++ b/Engine/Source/Runtime/RHI/Public/DynamicRHI.h
@@ -1619,6 +1619,31 @@ inline bool appIsPoolTexture(FRHITexture* TextureRHI) { return false; }
 */
 inline void appDumpTextureMemoryStats(const TCHAR* /*Message*/) {}
 
+/** Structs that describe how intensively a GPU is being used. */
+struct FRHIGPUUsageFractions
+{
+	// Fraction on how much much the GPU clocks has been scaled down by driver for energy savings.
+	float ClockScaling = 1.0f;
+
+	// Fraction of GPU resource dedicated for our own process at current clock scaling.
+	float CurrentProcess = 0.0f;
+
+	// Fraction of GPU resource dedicated for other processes at current clock scaling.
+	float ExternalProcesses = 0.0f;
+
+	// Remaining fraction of GPU resource that is idle.
+	inline float GetUnused()
+	{
+		return FMath::Clamp(1.0f - CurrentProcess - ExternalProcesses, 0.0f, 1.0f);
+	}
+};
+
+/** Get how much the GPU is getting used.
+*
+* Requires GRHISupportsGPUUsage=true before use.
+*/
+typedef FRHIGPUUsageFractions (*RHIGetGPUUsageType)(uint32);
+extern RHI_API RHIGetGPUUsageType RHIGetGPUUsage;
 
 /** Defines the interface of a module implementing a dynamic RHI. */
 class IDynamicRHIModule : public IModuleInterface
diff --git a/Engine/Source/Runtime/RHI/Public/RHI.h b/Engine/Source/Runtime/RHI/Public/RHI.h
--- a/Engine/Source/Runtime/RHI/Public/RHI.h
+++ b/Engine/Source/Runtime/RHI/Public/RHI.h
@@ -236,6 +236,9 @@ extern RHI_API bool GRHISupportsGPUTimestampBubblesRemoval;
 /** true if RHIGetGPUFrameCycles removes CPu generated bubbles. */
 extern RHI_API bool GRHISupportsFrameCyclesBubblesRemoval;
 
+/** true if RHIGetGPUUsage() is supported. */
+extern RHI_API bool GRHISupportsGPUUsage;
+
 /** true if the GPU supports hidden surface removal in hardware. */
 extern RHI_API bool GHardwareHiddenSurfaceRemoval;
 
