Fix comment referencing the wrong function

diff --git a/Engine/Source/Runtime/Engine/Classes/Components/ActorComponent.h b/Engine/Source/Runtime/Engine/Classes/Components/ActorComponent.h
--- a/Engine/Source/Runtime/Engine/Classes/Components/ActorComponent.h
+++ b/Engine/Source/Runtime/Engine/Classes/Components/ActorComponent.h
@@ -527,7 +527,7 @@ public:
 	* Unregister a SubObject to stop replicating it's properties to clients.
 	* This does not remove or delete it from connections where it was already replicated.
 	* By default a replicated subobject gets deleted on clients when the original pointer on the authority gets invalidated.
-	* If you want to immediately remove it from client use the DeleteReplicatedSubObject or TearOffReplicatedSubObject functions instead of this one.
+	* If you want to immediately remove it from client use the DestroyReplicatedSubObjectOnRemotePeers or TearOffReplicatedSubObject functions instead of this one.
 	* @param SubObject The SubObject to remove
 	*/
 	void RemoveReplicatedSubObject(UObject* SubObject);
diff --git a/Engine/Source/Runtime/Engine/Classes/GameFramework/Actor.h b/Engine/Source/Runtime/Engine/Classes/GameFramework/Actor.h
--- a/Engine/Source/Runtime/Engine/Classes/GameFramework/Actor.h
+++ b/Engine/Source/Runtime/Engine/Classes/GameFramework/Actor.h
@@ -3746,7 +3746,7 @@ public:
 	* Unregister a SubObject to stop replicating it's properties to clients.
 	* This does not remove or delete it from connections where it was already replicated.
 	* By default a replicated subobject gets deleted on clients when the original pointer on the authority becomes invalid.
-	* If you want to immediately remove it from client use the DeleteReplicatedSubObject or TearOffReplicatedSubObject functions instead of this one.
+	* If you want to immediately remove it from client use the DestroyReplicatedSubObjectOnRemotePeers or TearOffReplicatedSubObject functions instead of this one.
 	* @param SubObject The SubObject to remove
 	*/
 	void RemoveReplicatedSubObject(UObject* SubObject);
