[global shaders] handle potential for getting a shutdown call before RHI initialization (in this case the global vars used in the global shader shutdown function may have invalid values causing an out-of-bounds read on the GGlobalShaderMap array)

diff --git a/Engine/Source/Runtime/Engine/Private/ShaderCompiler/ShaderCompiler.cpp b/Engine/Source/Runtime/Engine/Private/ShaderCompiler/ShaderCompiler.cpp
--- a/Engine/Source/Runtime/Engine/Private/ShaderCompiler/ShaderCompiler.cpp
+++ b/Engine/Source/Runtime/Engine/Private/ShaderCompiler/ShaderCompiler.cpp
@@ -7819,6 +7819,12 @@ void CompileGlobalShaderMap(bool bRefreshShaderMap)
 
 void ShutdownGlobalShaderMap()
 {
+	// handle edge case where we get a shutdown before fully initialized (the globals used below are not in a valid state)
+	if (!GIsRHIInitialized)
+	{
+		return;
+	}
+
 	// at the point this function is called (during the shutdown process) we do not expect any outstanding work that could potentially be still referencing
 	// global shaders, so we are not deferring the deletion (via GGlobalShaderMap_DeferredDeleteCopy) like we do during the shader recompilation.
 	EShaderPlatform Platform = GShaderPlatformForFeatureLevel[GMaxRHIFeatureLevel];
