change uses of SerializeCompressed to SerializeCompressedNew with Oodle
provide SerializeCompressedNew with default args for easier use

diff --git a/Engine/Source/Runtime/Core/Private/Serialization/Archive.cpp b/Engine/Source/Runtime/Core/Private/Serialization/Archive.cpp
--- a/Engine/Source/Runtime/Core/Private/Serialization/Archive.cpp
+++ b/Engine/Source/Runtime/Core/Private/Serialization/Archive.cpp
@@ -693,6 +693,11 @@ void FArchive::SerializeCompressed(void* V, int64 Length, FName CompressionForma
 	SerializeCompressedNew(V,Length,CompressionFormatCannotChange,CompressionFormatCannotChange,Flags,bTreatBufferAsFileReader,nullptr);
 }
 
+void FArchive::SerializeCompressedNew(void* V, int64 Length)
+{
+	SerializeCompressedNew(V,Length,NAME_Oodle,NAME_Zlib);
+}
+
 void FArchive::SerializeCompressedNew(void* V, int64 Length, FName CompressionFormatToEncode, FName CompressionFormatToDecodeOldV1Files,  ECompressionFlags Flags, bool bTreatBufferAsFileReader, int64 * OutPartialReadLength)
 {
 	TRACE_CPUPROFILER_EVENT_SCOPE(FArchive::SerializeCompressed);
diff --git a/Engine/Source/Runtime/Core/Public/Serialization/Archive.h b/Engine/Source/Runtime/Core/Public/Serialization/Archive.h
--- a/Engine/Source/Runtime/Core/Public/Serialization/Archive.h
+++ b/Engine/Source/Runtime/Core/Public/Serialization/Archive.h
@@ -1724,13 +1724,9 @@ public:
 
 	/**
 	 * Serializes and compresses/ uncompresses data. This is a shared helper function for compression support. 
-	 
-	 // @@ old comment ? remove me ?
-	 The data is saved in a way compatible with FIOSystem::LoadCompressedData.
-
-	 *
-	 * prefer SerializeCompressedNew instead.
 	 *
+	 * Do not use SerializeCompressed in new code, prefer SerializeCompressedNew instead.
+	 * SerializeCompressedNew can be dropped in to any existing use of SerializeCompressed.
 	 *
 	 * @param	V		Data pointer to serialize data from/ to
 	 * @param	Length	Length of source data if we're saving, unused otherwise
@@ -1745,6 +1741,10 @@ public:
 	 * Serializes and compresses/ uncompresses data. This is a shared helper function for compression support. 
 	 *
 	 * call SerializeCompressedNew instead of SerializeCompressed
+	 * Typically you should not serializing data compressed if it will be packaged or stored in the DDC.
+	 * Prefer to allow the package/iostore system to do the compression for you instead.
+	 *
+	 * SerializeCompressedNew can read existing data written by old SerializeCompressed calls.
 	 *
 	 * @param	V		Data pointer to serialize data from/ to
 	 * @param	Length	Length of source data if we're saving, unused otherwise
@@ -1758,6 +1758,17 @@ public:
 		ECompressionFlags Flags=COMPRESS_NoFlags, bool bTreatBufferAsFileReader=false,
 		int64 * OutPartialReadLength=nullptr);
 		
+	/**
+	 * Serializes and compresses/ uncompresses data with default compressor choices.
+	 *
+	 * Default compressors are Oodle for new data and Zlib when loading legacy data.
+	 *
+	 * @param	V		Data pointer to serialize data from/ to
+	 * @param	Length	Length of source data if we're saving, unused otherwise
+	 */
+	void SerializeCompressedNew(void* V, int64 Length);
+
+
 	using FArchiveState::IsByteSwapping;
 
 	/** Used to do byte swapping on small items. This does not happen usually, so we don't want it inline. */
diff --git a/Engine/Source/Runtime/Engine/Private/SparseVolumeTexture/SparseVolumeTexture.cpp b/Engine/Source/Runtime/Engine/Private/SparseVolumeTexture/SparseVolumeTexture.cpp
--- a/Engine/Source/Runtime/Engine/Private/SparseVolumeTexture/SparseVolumeTexture.cpp
+++ b/Engine/Source/Runtime/Engine/Private/SparseVolumeTexture/SparseVolumeTexture.cpp
@@ -323,7 +323,7 @@ void UStaticSparseVolumeTexture::GenerateOrLoadDDCRuntimeData()
 		Ar << UncompressedSize;
 
 		uint8* DecompressionBuffer = reinterpret_cast<uint8*>(FMemory::Malloc(UncompressedSize));
-		Ar.SerializeCompressed(DecompressionBuffer, UncompressedSize, NAME_Zlib);
+		Ar.SerializeCompressedNew(DecompressionBuffer, UncompressedSize);
 
 		FLargeMemoryReader LargeMemReader(DecompressionBuffer, UncompressedSize, ELargeMemoryReaderFlags::Persistent | ELargeMemoryReaderFlags::TakeOwnership);
 
@@ -347,7 +347,7 @@ void UStaticSparseVolumeTexture::GenerateOrLoadDDCRuntimeData()
 			FMemoryWriter CompressedArchive(DerivedData, true);
 
 			CompressedArchive << UncompressedSize; // needed for allocating decompression buffer
-			CompressedArchive.SerializeCompressed(LargeMemWriter.GetData(), UncompressedSize, NAME_Zlib);
+			CompressedArchive.SerializeCompressedNew(LargeMemWriter.GetData(), UncompressedSize);
 
 			GetDerivedDataCacheRef().Put(*DerivedDataKey, DerivedData, GetPathName());
 		}
@@ -657,7 +657,7 @@ void UAnimatedSparseVolumeTexture::GenerateOrLoadDDCRuntimeData(int32 FrameIndex
 		Ar << UncompressedSize;
 
 		uint8* DecompressionBuffer = reinterpret_cast<uint8*>(FMemory::Malloc(UncompressedSize));
-		Ar.SerializeCompressed(DecompressionBuffer, UncompressedSize, NAME_Zlib);
+		Ar.SerializeCompressedNew(DecompressionBuffer, UncompressedSize);
 
 		FLargeMemoryReader LargeMemReader(DecompressionBuffer, UncompressedSize, ELargeMemoryReaderFlags::Persistent | ELargeMemoryReaderFlags::TakeOwnership);
 
@@ -681,7 +681,7 @@ void UAnimatedSparseVolumeTexture::GenerateOrLoadDDCRuntimeData(int32 FrameIndex
 			FMemoryWriter CompressedArchive(DerivedData, true);
 
 			CompressedArchive << UncompressedSize; // needed for allocating decompression buffer
-			CompressedArchive.SerializeCompressed(LargeMemWriter.GetData(), UncompressedSize, NAME_Zlib);
+			CompressedArchive.SerializeCompressedNew(LargeMemWriter.GetData(), UncompressedSize);
 
 			GetDerivedDataCacheRef().Put(*DerivedDataKey, DerivedData, GetPathName());
 		}
