Fix for failure to play audio when ForceInline loading behavior selected and encoding changes.
#rnx

diff --git a/Engine/Source/Runtime/Engine/Private/SoundWave.cpp b/Engine/Source/Runtime/Engine/Private/SoundWave.cpp
--- a/Engine/Source/Runtime/Engine/Private/SoundWave.cpp
+++ b/Engine/Source/Runtime/Engine/Private/SoundWave.cpp
@@ -1565,6 +1565,8 @@ void USoundWave::InvalidateCompressedData(bool bFreeResources, bool bRebuildStre
 	CompressedDataGuid = FGuid::NewGuid();
 	SoundWaveDataPtr->DiscardZerothChunkData();
 	SoundWaveDataPtr->CompressedFormatData.FlushData();
+	SoundWaveDataPtr->ResourceSize = 0;
+	SoundWaveDataPtr->ResourceData.Empty();
 
 	if (bFreeResources)
 	{
