Fixed regression causing active AnimNotifies from all the animations within an AnimMontage to be queued as standard notifies, essentialy skipping the slot relevancy check before triggering the animation events ( FAnimNotifyQueue::ApplyMontageNotifies was never being called ).

diff --git a/Engine/Source/Runtime/Engine/Private/Animation/AnimMontage.cpp b/Engine/Source/Runtime/Engine/Private/Animation/AnimMontage.cpp
--- a/Engine/Source/Runtime/Engine/Private/Animation/AnimMontage.cpp
+++ b/Engine/Source/Runtime/Engine/Private/Animation/AnimMontage.cpp
@@ -2563,7 +2563,7 @@ void FAnimMontageInstance::HandleEvents(float PreviousTrackPos, float CurrentTra
 		return;
 	}
 
-	// now get active Notifies based on how it advanced
+	// Now get active Notifies based on how it advanced
 	if (AnimInstance.IsValid())
 	{
 		FAnimTickRecord TickRecord;
@@ -2572,21 +2572,38 @@ void FAnimMontageInstance::HandleEvents(float PreviousTrackPos, float CurrentTra
 		TickRecord.MakeContextData<UE::Anim::FAnimNotifyMontageInstanceContext>(InstanceID);
 
 		FAnimNotifyContext NotifyContext(TickRecord);
-		// We already break up AnimMontage update to handle looping, so we guarantee that PreviousPos and CurrentPos are contiguous.
-		Montage->GetAnimNotifiesFromDeltaPositions(PreviousTrackPos, CurrentTrackPos, NotifyContext);
 
-		// For Montage only, remove notifies marked as 'branching points'. They are not queued and are handled separately.
-		Montage->FilterOutNotifyBranchingPoints(NotifyContext.ActiveNotifies);
-
-		// now trigger notifies for all animations within montage
-		// we'll do this for all slots for now
-		for (auto SlotTrack = Montage->SlotAnimTracks.CreateIterator(); SlotTrack; ++SlotTrack)
+		// Queue all notifies fired from the AnimMontage's Notify Track.
 		{
-			SlotTrack->AnimTrack.GetAnimNotifiesFromTrackPositions(PreviousTrackPos, CurrentTrackPos, NotifyContext);
+			// We already break up AnimMontage update to handle looping, so we guarantee that PreviousPos and CurrentPos are contiguous.
+			Montage->GetAnimNotifiesFromDeltaPositions(PreviousTrackPos, CurrentTrackPos, NotifyContext);
+
+			// For Montage only, remove notifies marked as 'branching points'. They are not queued and are handled separately.
+			Montage->FilterOutNotifyBranchingPoints(NotifyContext.ActiveNotifies);
+
+			// Queue active non-'branching point' notifies.
+			AnimInstance->NotifyQueue.AddAnimNotifies(NotifyContext.ActiveNotifies, NotifyWeight);
 		}
 
-		// Queue all these notifies.
-		AnimInstance->NotifyQueue.AddAnimNotifies(NotifyContext.ActiveNotifies, NotifyWeight);
+		// Queue all notifies fired by all the animations within the AnimMontage. We'll do this for all slot tracks.
+		{
+			TMap<FName, TArray<FAnimNotifyEventReference>> NotifyMap;
+			
+			for (auto SlotTrack = Montage->SlotAnimTracks.CreateIterator(); SlotTrack; ++SlotTrack)
+			{
+				TArray<FAnimNotifyEventReference>& CurrentSlotNotifies = NotifyMap.FindOrAdd(SlotTrack->SlotName);
+
+				// Queue active notifies from current slot.
+				{
+					NotifyContext.ActiveNotifies.Reset();
+					SlotTrack->AnimTrack.GetAnimNotifiesFromTrackPositions(PreviousTrackPos, CurrentTrackPos, NotifyContext);
+					Swap(CurrentSlotNotifies, NotifyContext.ActiveNotifies);
+				}
+			}
+
+			// Queue active unfiltered notifies from slot tracks.
+			AnimInstance->NotifyQueue.AddAnimNotifies(NotifyMap, NotifyWeight);	
+		}
 	}
 
 	// Update active state branching points, before we handle the immediate tick marker.
