Removed old stale code that was preventing LightingOnly from working correctly with static meshes.

The removed condition was forbidding meshes with materials that add some form of displacement to support LightingOnly mode. The comment indicates that the intention was to forbid meshes with tessellation enabled. Since tessellation is completely been removed, the only two other cases that may mark the material as modifying mesh position is either WorldPositionOffset or PixelDepthOffset. Without the if, meshes with materials that write to either these two attributes will now support LightingOnly, although they will lose any displacements when this view mode is on.

diff --git a/Engine/Source/Runtime/Engine/Private/PrimitiveDrawingUtils.cpp b/Engine/Source/Runtime/Engine/Private/PrimitiveDrawingUtils.cpp
--- a/Engine/Source/Runtime/Engine/Private/PrimitiveDrawingUtils.cpp
+++ b/Engine/Source/Runtime/Engine/Private/PrimitiveDrawingUtils.cpp
@@ -1642,18 +1642,14 @@ void ApplyViewModeOverrides(
 
 			if (bTextureMapped == false)
 			{
-				// Tessellated geometry cannot use engine debug LevelColorationLitMaterial as it doesn't support tessellation
-				if (!bMaterialModifiesMeshPosition)
-				{
-					FMaterialRenderProxy* RenderProxy = GEngine->LevelColorationLitMaterial->GetRenderProxy();
-					auto LightingOnlyMaterialInstance = new FColoredMaterialRenderProxy(
-						RenderProxy,
-						GEngine->LightingOnlyBrightness
-					);
+				FMaterialRenderProxy* RenderProxy = GEngine->LevelColorationLitMaterial->GetRenderProxy();
+				auto LightingOnlyMaterialInstance = new FColoredMaterialRenderProxy(
+					RenderProxy,
+					GEngine->LightingOnlyBrightness
+				);
 
-					Mesh.MaterialRenderProxy = LightingOnlyMaterialInstance;
-					Collector.RegisterOneFrameMaterialProxy(LightingOnlyMaterialInstance);
-				}
+				Mesh.MaterialRenderProxy = LightingOnlyMaterialInstance;
+				Collector.RegisterOneFrameMaterialProxy(LightingOnlyMaterialInstance);
 			}
 			else
 			{
