Reflection Captures for mobile - addressing lighting scenarios

diff --git a/Engine/Source/Runtime/Engine/Private/Components/ReflectionCaptureComponent.cpp b/Engine/Source/Runtime/Engine/Private/Components/ReflectionCaptureComponent.cpp
--- a/Engine/Source/Runtime/Engine/Private/Components/ReflectionCaptureComponent.cpp
+++ b/Engine/Source/Runtime/Engine/Private/Components/ReflectionCaptureComponent.cpp
@@ -33,6 +33,7 @@
 #include "Math/PackedVector.h"
 #include "GlobalRenderResources.h"
 #include "UObject/UObjectAnnotation.h"
+#include "DataDrivenShaderPlatformInfo.h"
 
 #include UE_INLINE_GENERATED_CPP_BY_NAME(ReflectionCaptureComponent)
 
@@ -115,10 +116,30 @@ FReflectionCaptureMapBuildData* UReflectionCaptureComponent::GetMapBuildData() c
 	return NULL;
 }
 
+bool IsEncodedHDRCubemapTextureRequired(EShaderPlatform ShaderPlatform)
+{
+	const bool bEncodedHDRCubemapTextureRequired = (GIsEditor || IsMobilePlatform(ShaderPlatform))
+		// mobile forward renderer or translucensy in mobile deferred need encoded reflection texture when clustered reflections disabled
+		&& !MobileForwardEnableClusteredReflections(ShaderPlatform);
+	return bEncodedHDRCubemapTextureRequired;
+}
+
+
 void UReflectionCaptureComponent::PropagateLightingScenarioChange()
 {
-	// GetMapBuildData has changed, re-upload
-	MarkDirtyForRecaptureOrUpload();
+	const FSceneInterface* Scene = GetWorld()->Scene;
+	const EShaderPlatform  ShaderPlatform = Scene ? Scene->GetShaderPlatform() : GMaxRHIShaderPlatform;
+	const bool bEncodedDataRequired = IsEncodedHDRCubemapTextureRequired(ShaderPlatform);
+
+	if (bEncodedDataRequired && EncodedHDRCubemapTexture == nullptr)
+	{
+		ReregisterComponent();
+	}
+	else
+	{
+		// GetMapBuildData has changed, re-upload
+		MarkDirtyForRecaptureOrUpload();
+	}
 }
 
 AReflectionCapture::AReflectionCapture(const FObjectInitializer& ObjectInitializer)
@@ -727,9 +748,9 @@ void UReflectionCaptureComponent::SafeReleaseEncodedHDRCubemapTexture()
 
 void UReflectionCaptureComponent::OnRegister()
 {
-	const bool bEncodedHDRCubemapTextureRequired = (GIsEditor || GMaxRHIFeatureLevel == ERHIFeatureLevel::ES3_1) 
-		// mobile forward renderer or translucensy in mobile deferred need encoded reflection texture when clustered reflections disabled
-		&& !MobileForwardEnableClusteredReflections(GMaxRHIShaderPlatform);
+	const FSceneInterface* Scene = GetWorld()->Scene;
+	const EShaderPlatform  ShaderPlatform = Scene ? Scene->GetShaderPlatform() : GMaxRHIShaderPlatform;
+	const bool bEncodedHDRCubemapTextureRequired = IsEncodedHDRCubemapTextureRequired(ShaderPlatform);
 
 	if (bEncodedHDRCubemapTextureRequired)
 	{
@@ -743,7 +764,7 @@ void UReflectionCaptureComponent::OnRegister()
 				EncodedHDRCubemapTexture = new FReflectionTextureCubeResource();
 				TArray<uint8> EncodedHDRCapturedData;
 				EPixelFormat EncodedHDRCubemapTextureFormat = PF_Unknown;
-				if (IsMobileDeferredShadingEnabled(GMaxRHIShaderPlatform))
+				if (IsMobileDeferredShadingEnabled(ShaderPlatform))
 				{
 					// make a copy, FullHDR data will still be needed later
 					EncodedHDRCapturedData = MapBuildData->FullHDRCapturedData;
