Add method to extract all possible cvar values from device profiles that select fragments via matched rules.
Swap PerQualityLevelProperties over to use the new method to resolve warnings during cook.

diff --git a/Engine/Source/Runtime/Engine/Classes/DeviceProfiles/DeviceProfileManager.h b/Engine/Source/Runtime/Engine/Classes/DeviceProfiles/DeviceProfileManager.h
--- a/Engine/Source/Runtime/Engine/Classes/DeviceProfiles/DeviceProfileManager.h
+++ b/Engine/Source/Runtime/Engine/Classes/DeviceProfiles/DeviceProfileManager.h
@@ -231,6 +231,7 @@ public:
 	{
 		DPM_SetCVars,
 		DPM_CacheValues,
+		DPM_CacheValuesIgnoreMatchingRules, // GatherDeviceProfileCVars will not return any matching rules cvars. Used when entire matching rules state is applied later (see GetAllReferencedDeviceProfileCVars).
 	};
 
 	/**
@@ -238,6 +239,10 @@ public:
 	 */
 	static TMap<FName, FString> GatherDeviceProfileCVars(const FString& DeviceProfileName, EDeviceProfileMode GatherMode);
 
+	/**
+	 * Gather all the cvars from the static device profile and then add all of the possible cvars+values from the dynamic/matched rule fragments.
+	 */
+	static TMap<FName, TSet<FString>> GetAllReferencedDeviceProfileCVars(UDeviceProfile* DeviceProfile);
 
 private:
 	/**
@@ -268,6 +273,9 @@ private:
 	/** Read and process all of the fragment matching rules. Returns an array containing the names of fragments selected. */
 	static TArray<FSelectedFragmentProperties> FindMatchingFragments(const FString& ParentDP, class FConfigCacheIni* PreviewConfigSystem);
 
+	/** Read and discover all of the possible fragments referenced by the matching rules. Returns an array containing the names of fragments selected. */
+	static TArray<FString> FindAllReferencedFragmentsFromMatchedRules(const FString& ParentDP, FConfigCacheIni* ConfigSystem);
+
 	/** Convert a FSelectedFragmentProperties array to a string */
 	static const FString FragmentPropertyArrayToFragmentString(const TArray<FSelectedFragmentProperties>& FragmentProperties, bool bEnabledOnly, bool bIncludeTags, bool bAlphaSort);
 
diff --git a/Engine/Source/Runtime/Engine/Private/DeviceProfiles/DeviceProfileManager.cpp b/Engine/Source/Runtime/Engine/Private/DeviceProfiles/DeviceProfileManager.cpp
--- a/Engine/Source/Runtime/Engine/Private/DeviceProfiles/DeviceProfileManager.cpp
+++ b/Engine/Source/Runtime/Engine/Private/DeviceProfiles/DeviceProfileManager.cpp
@@ -169,7 +169,7 @@ TMap<FName, FString> UDeviceProfileManager::GatherDeviceProfileCVars(const FStri
 
 	EPlatformMemorySizeBucket MemBucket = FPlatformMemory::GetMemorySizeBucket();
 	// if caching (for another platfomr), then we use the DP's PreviewMemoryBucket, instead of querying for it
-	if (GatherMode == EDeviceProfileMode::DPM_CacheValues)
+	if (GatherMode == EDeviceProfileMode::DPM_CacheValues || GatherMode == EDeviceProfileMode::DPM_CacheValuesIgnoreMatchingRules )
 	{
 #if ALLOW_OTHER_PLATFORM_CONFIG
 		// caching is not done super early, so we can assume DPs have been found now
@@ -183,8 +183,10 @@ TMap<FName, FString> UDeviceProfileManager::GatherDeviceProfileCVars(const FStri
 		// use the DP's platform's configs, NOT the running platform
 		ConfigSystem = FConfigCacheIni::ForPlatform(*Profile->DeviceType);
 		MemBucket = Profile->GetPreviewMemorySizeBucket();
-
-		FragmentsSelected = FindMatchingFragments(DeviceProfileName, ConfigSystem);
+		if(GatherMode == EDeviceProfileMode::DPM_CacheValues)
+		{
+			FragmentsSelected = FindMatchingFragments(DeviceProfileName, ConfigSystem);
+		}
 #else
 		checkNoEntry();
 #endif
@@ -340,6 +342,57 @@ TMap<FName, FString> UDeviceProfileManager::GatherDeviceProfileCVars(const FStri
 	return DeviceProfileCVars;
 }
 
+TMap<FName, TSet<FString>> UDeviceProfileManager::GetAllReferencedDeviceProfileCVars(UDeviceProfile* DeviceProfile)
+{
+	FConfigCacheIni* ConfigSystem = GConfig;
+
+#if ALLOW_OTHER_PLATFORM_CONFIG
+	// use the DP's platform's configs, NOT the running platform
+	ConfigSystem = FConfigCacheIni::ForPlatform(*DeviceProfile->DeviceType);
+#endif
+
+	TMap<FName, FString> DeviceProfileCVars = GatherDeviceProfileCVars(DeviceProfile->GetName(), EDeviceProfileMode::DPM_CacheValuesIgnoreMatchingRules);
+
+	// gather all referenced fragments, note that this just a dumb traverse of the the matched rules so it may contain
+	// fragments that a device cannot ultimately select.
+	TArray<FString> AllReferencedMatchedFragments = FindAllReferencedFragmentsFromMatchedRules(DeviceProfile->GetName(), ConfigSystem);
+
+	TMap<FName, TSet<FString>> AllCVarsAndValues;
+
+	for (const auto& Pair : DeviceProfileCVars)
+	{
+		AllCVarsAndValues.FindOrAdd(Pair.Key).Add(Pair.Value);
+	}
+
+	for (const FString& Fragment : AllReferencedMatchedFragments)
+	{
+		TArray<FString> FragmentCVars;
+		GetFragmentCVars(Fragment, TEXT("CVars"), FragmentCVars, ConfigSystem);
+		for (const FString& FragCVar : FragmentCVars)
+		{
+			FString CVarKey, CVarValue;
+			if (FragCVar.Split(TEXT("="), &CVarKey, &CVarValue))
+			{
+				if (CVarKey.StartsWith(TEXT("sg.")))
+				{
+					TMap<FString, FString> ScalabilityCVars;
+					ExpandScalabilityCVar(ConfigSystem, CVarKey, CVarValue, ScalabilityCVars, true);
+					for (const auto& ScalabilityPair : ScalabilityCVars)
+					{
+						AllCVarsAndValues.FindOrAdd(*ScalabilityPair.Key).Add(ScalabilityPair.Value);
+					}
+				}
+				else
+				{
+					AllCVarsAndValues.FindOrAdd(FName(CVarKey)).Add(CVarValue);
+				}
+			}
+		}
+	}
+
+	return AllCVarsAndValues;
+}
+
 void UDeviceProfileManager::SetDeviceProfileCVars(const FString& DeviceProfileName)
 {
 	// walk over the parent chain, gathering the cvars this DP inherits and contains
diff --git a/Engine/Source/Runtime/Engine/Private/DeviceProfiles/DeviceProfileMatching.cpp b/Engine/Source/Runtime/Engine/Private/DeviceProfiles/DeviceProfileMatching.cpp
--- a/Engine/Source/Runtime/Engine/Private/DeviceProfiles/DeviceProfileMatching.cpp
+++ b/Engine/Source/Runtime/Engine/Private/DeviceProfiles/DeviceProfileMatching.cpp
@@ -316,15 +316,8 @@ static bool EvaluateMatch(IDeviceProfileSelectorModule* DPSelectorModule, const
 		{
 			SelectedFragmentsOUT += TEXT(",");
 		}
-		if (rule.AppendFragments.Contains(TEXT("[clear]")))
-		{
-			SelectedFragmentsOUT.Empty();
-			SelectedFragmentsOUT += rule.AppendFragments.Replace(TEXT("[clear]"), TEXT(""));
-		}
-		else
-		{
-			SelectedFragmentsOUT += rule.AppendFragments;
-		}
+
+		SelectedFragmentsOUT += rule.AppendFragments;
 	}
 
 	if (!rule.SetUserVar.IsEmpty())
@@ -581,34 +574,38 @@ public:
 
 	virtual void Serialize(const TCHAR* V, ELogVerbosity::Type Verbosity, const class FName& Category) override
 	{
-		if (IsRunningCookCommandlet())
-		{
-			UE_LOG(LogInit, Warning, TEXT("DeviceProfileMatching: Error while parsing Matching Rules (%s) : %s"), *Stage, V);
-		}
-		else
-		{
-			UE_LOG(LogInit, Error, TEXT("DeviceProfileMatching: Error while parsing Matching Rules (%s) : %s"), *Stage, V);
-		}
+		UE_LOG(LogInit, Error, TEXT("DeviceProfileMatching: Error while parsing Matching Rules (%s) : %s"), *Stage, V);
 		NumErrors++;
 	}
 };
 
-static bool DPHasMatchingRules(const FString& ParentDP, FConfigCacheIni* ConfigSystem)
+
+static TArray<FString> GetMatchedRulesArray(const FString& ParentDP, FConfigCacheIni* ConfigSystem)
 {
 	FString SectionSuffix = *FString::Printf(TEXT(" %s"), *UDeviceProfile::StaticClass()->GetName());
 	FString CurrentSectionName = ParentDP + SectionSuffix;
 	FString ArrayName(TEXT("MatchingRules"));
 	TArray<FString> MatchingRulesArray;
-	return ConfigSystem->GetArray(*CurrentSectionName, *ArrayName, MatchingRulesArray, GDeviceProfilesIni) != 0;
+	ConfigSystem->GetArray(*CurrentSectionName, *ArrayName, MatchingRulesArray, GDeviceProfilesIni);
+
+	return MatchingRulesArray;
+}
+
+static bool DPHasMatchingRules(const FString& ParentDP, FConfigCacheIni* ConfigSystem)
+{
+	return GetMatchedRulesArray(ParentDP, ConfigSystem).Num() > 0;
+}
+
+static FDPMatchingRulestruct ParseToRuleStruct(const FString& RuleText, FDeviceProfileMatchingErrorContext& DPMatchingErrorOutput)
+{
+	FDPMatchingRulestruct RuleStruct;
+	FDPMatchingRulestruct::StaticStruct()->ImportText(*RuleText, &RuleStruct, nullptr, EPropertyPortFlags::PPF_None, &DPMatchingErrorOutput, FDPMatchingRulestruct::StaticStruct()->GetName(), true);
+	return RuleStruct;
 }
 
 static FString LoadAndProcessMatchingRulesFromConfig(const FString& ParentDP, IDeviceProfileSelectorModule* DPSelector, FConfigCacheIni* ConfigSystem)
 {
-	FString SectionSuffix = *FString::Printf(TEXT(" %s"), *UDeviceProfile::StaticClass()->GetName());
-	FString CurrentSectionName = ParentDP + SectionSuffix;
-	FString ArrayName(TEXT("MatchingRules"));
-	TArray<FString> MatchingRulesArray;
-	ConfigSystem->GetArray(*CurrentSectionName, *ArrayName, MatchingRulesArray, GDeviceProfilesIni);
+	const TArray<FString> MatchingRulesArray = GetMatchedRulesArray(ParentDP, ConfigSystem);
 
 	FRulePropertiesContainer RuleProperties;
 	FString SelectedFragments;
@@ -618,8 +615,7 @@ static FString LoadAndProcessMatchingRulesFromConfig(const FString& ParentDP, ID
 	{
 		DPMatchingErrorOutput.Stage = FString::Printf(TEXT("%s rule #%d"), *ParentDP, Count);
 		Count++;
-		FDPMatchingRulestruct RuleStruct;
-		FDPMatchingRulestruct::StaticStruct()->ImportText(*RuleText, &RuleStruct, nullptr, EPropertyPortFlags::PPF_None, &DPMatchingErrorOutput, FDPMatchingRulestruct::StaticStruct()->GetName(), true);
+		FDPMatchingRulestruct RuleStruct = ParseToRuleStruct(RuleText, DPMatchingErrorOutput);
 		static const FString NoName("<unnamed>");
 		const FString& RuleName = RuleStruct.RuleName.IsEmpty() ? NoName : RuleStruct.RuleName;
 		EvaluateMatch(DPSelector, RuleName, RuleStruct, RuleProperties, SelectedFragments, &DPMatchingErrorOutput);
@@ -652,6 +648,50 @@ static FString RemoveAllWhiteSpace(const FString& StringIN)
 	return Ret;
 }
 
+ static void FindAllReferencedMatchedFragments(FString& AllAppendableFragments, const FDPMatchingRulestructBase* RuleStruct, FDeviceProfileMatchingErrorContext& DPMatchingErrorOutput)
+ {
+	if (!RuleStruct)
+		return;
+
+	FindAllReferencedMatchedFragments(AllAppendableFragments, RuleStruct->GetOnTrue(), DPMatchingErrorOutput);
+	FindAllReferencedMatchedFragments(AllAppendableFragments, RuleStruct->GetOnFalse(), DPMatchingErrorOutput);
+
+	if (!AllAppendableFragments.IsEmpty() && !RuleStruct->AppendFragments.IsEmpty())
+	{
+		AllAppendableFragments += TEXT(",");
+	}
+	AllAppendableFragments += RuleStruct->AppendFragments;
+ }
+
+TArray<FString> UDeviceProfileManager::FindAllReferencedFragmentsFromMatchedRules(const FString& ParentDP, FConfigCacheIni* ConfigSystem)
+{
+	TArray<FString> MatchingRulesArray = GetMatchedRulesArray(ParentDP, ConfigSystem);
+
+	FDeviceProfileMatchingErrorContext DPMatchingErrorOutput;
+	int Count = 0;
+	FString AllAppendableFragments;
+	for (const FString& RuleText : MatchingRulesArray)
+	{
+		DPMatchingErrorOutput.Stage = FString::Printf(TEXT("%s rule #%d"), *ParentDP, Count++);
+		FDPMatchingRulestruct RuleStruct = ParseToRuleStruct(RuleText, DPMatchingErrorOutput);
+		static const FString NoName("<unnamed>");
+		FindAllReferencedMatchedFragments(AllAppendableFragments, &RuleStruct, DPMatchingErrorOutput);
+	}
+	UE_CLOG(DPMatchingErrorOutput.NumErrors > 0, LogInit, Error, TEXT("DeviceProfileMatching: %d Error(s) encountered while processing FindAllReferencedFragments for %s"), DPMatchingErrorOutput.NumErrors, *ParentDP);
+
+	AllAppendableFragments = RemoveAllWhiteSpace(AllAppendableFragments);
+	TArray<FSelectedFragmentProperties> AllFragmentsSeen = FragmentStringToFragmentProperties(AllAppendableFragments);
+	TArray<FString> UniqueMatchedFragments;
+
+	// Return a unique list of the referenced fragments (including tagged/disabled)
+	for (const FSelectedFragmentProperties& FragmentProperties : AllFragmentsSeen)
+	{ 
+		UniqueMatchedFragments.AddUnique(FragmentProperties.Fragment);
+	}
+
+	return UniqueMatchedFragments;
+}
+
 TArray<FSelectedFragmentProperties> UDeviceProfileManager::FindMatchingFragments(const FString& ParentDP, FConfigCacheIni* ConfigSystem)
 {
 	FString SelectedFragments;
diff --git a/Engine/Source/Runtime/Engine/Private/PerQualityLevelProperties.cpp b/Engine/Source/Runtime/Engine/Private/PerQualityLevelProperties.cpp
--- a/Engine/Source/Runtime/Engine/Private/PerQualityLevelProperties.cpp
+++ b/Engine/Source/Runtime/Engine/Private/PerQualityLevelProperties.cpp
@@ -264,15 +264,19 @@ FSupportedQualityLevelArray FPerQualityLevelProperty<StructType, ValueType, _Bas
 		return (DeviceProfile->DeviceType == InPlatformName);
 	});
 
+	const FName CVarFName(CVarName);
 	// iterate through the DP sections to find the ones relevant to the platform
 	for (TObjectPtr<UDeviceProfile>& DeviceProfile : DeviceProfiles)
 	{
-		const TMap<FString, FString>& ExpandedCVars = DeviceProfile->GetAllExpandedCVars();
-		const FString* const CVarValue = ExpandedCVars.Find(CVarName);
-		if (CVarValue)
+		const TMap<FName, TSet<FString>> AllPossibleCVars = UDeviceProfileManager::GetAllReferencedDeviceProfileCVars(DeviceProfile);
+		const TSet<FString>* CVarValues = AllPossibleCVars.Find(CVarFName);
+		if (CVarValues)
 		{
-			int32 Result = FCString::Atoi(*(* CVarValue));
-			CookingQualitLevelInfo.Add(Result);
+			for(const FString& CVarValue : *CVarValues)
+			{
+				int32 Result = FCString::Atoi(*CVarValue);
+				CookingQualitLevelInfo.Add(Result);
+			}
 		}
 	}
 
