Remove FVertexElementTypeSupportInfo and GVertexElementTypeSupport
 - Old unused code path (from 2007!!!, CL 182619).
 - IsSupported() always returns true since nothing ever calls SetSupported().

diff --git a/Engine/Source/Editor/Persona/Private/PersonaMeshDetails.cpp b/Engine/Source/Editor/Persona/Private/PersonaMeshDetails.cpp
--- a/Engine/Source/Editor/Persona/Private/PersonaMeshDetails.cpp
+++ b/Engine/Source/Editor/Persona/Private/PersonaMeshDetails.cpp
@@ -2082,18 +2082,11 @@ void FSkeletalMeshBuildSettingsLayout::OnUseFullPrecisionUVsChanged(ECheckBoxSta
 	const bool bUseFullPrecisionUVs = (NewState == ECheckBoxState::Checked) ? true : false;
 	if (BuildSettings.bUseFullPrecisionUVs != bUseFullPrecisionUVs)
 	{
-		if (!bUseFullPrecisionUVs && !GVertexElementTypeSupport.IsSupported(VET_Half2))
-		{
-			UE_LOG(LogSkeletalMeshPersonaMeshDetail, Warning, TEXT("16 bit UVs not supported. Reverting to 32 bit UVs"));
-		}
-		else
-		{
-			FText TransactionText = FText::Format(LOCTEXT("PersonaChangedOnFullPrecisionUVsLOD", "LOD{0} build settings: use full precision UVs changed"), LODIndex);
-			FScopedTransaction Transaction(TransactionText);
-			ModifyMeshLODSettingsDelegate.ExecuteIfBound(LODIndex);
+		FText TransactionText = FText::Format(LOCTEXT("PersonaChangedOnFullPrecisionUVsLOD", "LOD{0} build settings: use full precision UVs changed"), LODIndex);
+		FScopedTransaction Transaction(TransactionText);
+		ModifyMeshLODSettingsDelegate.ExecuteIfBound(LODIndex);
 
-			BuildSettings.bUseFullPrecisionUVs = bUseFullPrecisionUVs;
-		}
+		BuildSettings.bUseFullPrecisionUVs = bUseFullPrecisionUVs;
 	}
 }
 void FSkeletalMeshBuildSettingsLayout::OnUseBackwardsCompatibleF16TruncUVsChanged(ECheckBoxState NewState)
diff --git a/Engine/Source/Runtime/Engine/Private/HierarchicalInstancedStaticMesh.cpp b/Engine/Source/Runtime/Engine/Private/HierarchicalInstancedStaticMesh.cpp
--- a/Engine/Source/Runtime/Engine/Private/HierarchicalInstancedStaticMesh.cpp
+++ b/Engine/Source/Runtime/Engine/Private/HierarchicalInstancedStaticMesh.cpp
@@ -323,8 +323,7 @@ protected:
 	{
 		// build new instance buffer
 		FRandomStream RandomStream = FRandomStream(InstancingRandomSeed);
-		bool bHalfFloat = GVertexElementTypeSupport.IsSupported(VET_Half2);
-		BuiltInstanceData = MakeUnique<FStaticMeshInstanceData>(bHalfFloat);
+		BuiltInstanceData = MakeUnique<FStaticMeshInstanceData>(/*bInUseHalfFloat = */true);
 		
 		int32 NumInstances = Result->InstanceReorderTable.Num();
 		int32 NumRenderInstances = Result->SortedInstances.Num();
@@ -2823,7 +2822,7 @@ void UHierarchicalInstancedStaticMeshComponent::ApplyEmpty()
 	InstanceUpdateCmdBuffer.Reset();
 	if (PerInstanceRenderData.IsValid())
 	{
-		TUniquePtr<FStaticMeshInstanceData> BuiltInstanceData = MakeUnique<FStaticMeshInstanceData>(GVertexElementTypeSupport.IsSupported(VET_Half2));
+		TUniquePtr<FStaticMeshInstanceData> BuiltInstanceData = MakeUnique<FStaticMeshInstanceData>(/*bInUseHalfFloat = */true);
 		PerInstanceRenderData->UpdateFromPreallocatedData(*BuiltInstanceData);
 		PerInstanceRenderData->HitProxies.Empty();
 		MarkRenderStateDirty();
diff --git a/Engine/Source/Runtime/Engine/Private/InstancedStaticMesh.cpp b/Engine/Source/Runtime/Engine/Private/InstancedStaticMesh.cpp
--- a/Engine/Source/Runtime/Engine/Private/InstancedStaticMesh.cpp
+++ b/Engine/Source/Runtime/Engine/Private/InstancedStaticMesh.cpp
@@ -2343,7 +2343,7 @@ void UInstancedStaticMeshComponent::FlushInstanceUpdateCommands(bool bFlushInsta
 
 	if (PerInstanceRenderData)
 	{
-		FStaticMeshInstanceData RenderInstanceData = FStaticMeshInstanceData(GVertexElementTypeSupport.IsSupported(VET_Half2));
+		FStaticMeshInstanceData RenderInstanceData = FStaticMeshInstanceData(/*bInUseHalfFloat = */ true);
 		BuildRenderData(RenderInstanceData, PerInstanceRenderData->HitProxies);
 		PerInstanceRenderData->UpdateFromPreallocatedData(RenderInstanceData);
 	}
@@ -4312,7 +4312,7 @@ void UInstancedStaticMeshComponent::InitPerInstanceRenderData(bool InitializeFro
 	else
 	{
 		TArray<TRefCountPtr<HHitProxy>> HitProxies;
-		FStaticMeshInstanceData InstanceBufferData = FStaticMeshInstanceData(GVertexElementTypeSupport.IsSupported(VET_Half2));
+		FStaticMeshInstanceData InstanceBufferData = FStaticMeshInstanceData(/*bInUseHalfFloat = */ true);
 		
 		if (InitializeFromCurrentData)
 		{
diff --git a/Engine/Source/Runtime/Engine/Private/Rendering/StaticMeshVertexBuffer.cpp b/Engine/Source/Runtime/Engine/Private/Rendering/StaticMeshVertexBuffer.cpp
--- a/Engine/Source/Runtime/Engine/Private/Rendering/StaticMeshVertexBuffer.cpp
+++ b/Engine/Source/Runtime/Engine/Private/Rendering/StaticMeshVertexBuffer.cpp
@@ -18,7 +18,7 @@ FStaticMeshVertexBuffer::FStaticMeshVertexBuffer() :
 	TexcoordDataPtr(nullptr),
 	NumTexCoords(0),
 	NumVertices(0),
-	bUseFullPrecisionUVs(!GVertexElementTypeSupport.IsSupported(VET_Half2)),
+	bUseFullPrecisionUVs(false),
 	bUseHighPrecisionTangentBasis(false)
 {}
 
@@ -108,17 +108,9 @@ void FStaticMeshVertexBuffer::Init(const FStaticMeshVertexBuffer& InVertexBuffer
 			check(GetNumTexCoords() == InVertexBuffer.GetNumTexCoords());
 			const uint8* InData = InVertexBuffer.TexcoordDataPtr;
 
-			// convert half float data to full float if the HW requires it.
-			if (!GetUseFullPrecisionUVs() && !GVertexElementTypeSupport.IsSupported(VET_Half2))
-			{
-				ConvertHalfTexcoordsToFloat(InData);
-			}
-			else
-			{
-				TexcoordData->ResizeBuffer(NumVertices * GetNumTexCoords());
-				TexcoordDataPtr = TexcoordData->GetDataPointer();
-				FMemory::Memcpy(TexcoordDataPtr, InData, TexcoordData->GetStride() * NumVertices * GetNumTexCoords());
-			}
+			TexcoordData->ResizeBuffer(NumVertices * GetNumTexCoords());
+			TexcoordDataPtr = TexcoordData->GetDataPointer();
+			FMemory::Memcpy(TexcoordDataPtr, InData, TexcoordData->GetStride() * NumVertices * GetNumTexCoords());
 		}
 	}
 }
@@ -232,12 +224,6 @@ void FStaticMeshVertexBuffer::Serialize(FArchive& Ar, bool bNeedsCPUAccess)
 
 			// Make a copy of the vertex data pointer.
 			TexcoordDataPtr = NumVertices ? TexcoordData->GetDataPointer() : nullptr;
-
-			// convert half float data to full float if the HW requires it.
-			if (NumVertices && !GetUseFullPrecisionUVs() && !GVertexElementTypeSupport.IsSupported(VET_Half2))
-			{
-				ConvertHalfTexcoordsToFloat(nullptr);
-			}
 		}
 	}
 }
@@ -254,7 +240,7 @@ void FStaticMeshVertexBuffer::SerializeMetaData(FArchive& Ar)
 void FStaticMeshVertexBuffer::ClearMetaData()
 {
 	NumTexCoords = NumVertices = 0;
-	bUseFullPrecisionUVs = !GVertexElementTypeSupport.IsSupported(VET_Half2);
+	bUseFullPrecisionUVs = false;
 	bUseHighPrecisionTangentBasis = false;
 	TangentsStride = TexcoordStride = 0;
 }
diff --git a/Engine/Source/Runtime/Engine/Private/SkeletalMesh.cpp b/Engine/Source/Runtime/Engine/Private/SkeletalMesh.cpp
--- a/Engine/Source/Runtime/Engine/Private/SkeletalMesh.cpp
+++ b/Engine/Source/Runtime/Engine/Private/SkeletalMesh.cpp
@@ -1971,9 +1971,6 @@ void USkeletalMesh::ExecuteBuildInternal(FSkinnedAssetBuildContext& Context)
 	// rebuild render data from imported model
 	CacheDerivedData(&Context);
 
-	// Do not need to fix up 16-bit UVs here, as we assume all editor platforms support them.
-	ensure(GVertexElementTypeSupport.IsSupported(VET_Half2));
-
 	GetSamplingInfoInternal().BuildRegions(this);
 	GetSamplingInfoInternal().BuildWholeMesh(this);
 
@@ -4437,7 +4434,6 @@ void USkeletalMesh::ClearAllCachedCookedPlatformData()
 static void SerializeLODInfoForDDC(USkeletalMesh* SkeletalMesh, FString& KeySuffix)
 {
 	TArray<FSkeletalMeshLODInfo>& LODInfos = SkeletalMesh->GetLODInfoArray();
-	const bool bIs16BitfloatBufferSupported = GVertexElementTypeSupport.IsSupported(VET_Half2);
 	for (int32 LODIndex = 0; LODIndex < SkeletalMesh->GetLODNum(); ++LODIndex)
 	{
 		check(LODInfos.IsValidIndex(LODIndex));
@@ -4470,7 +4466,7 @@ FString USkeletalMesh::BuildDerivedDataKey(const ITargetPlatform* TargetPlatform
 		const FSkeletalMeshLODInfo* BaseLODInfo = GetLODInfo(0);
 		bool bUseFullPrecisionUVs = BaseLODInfo ? BaseLODInfo->BuildSettings.bUseFullPrecisionUVs : false;
 		KeySuffix += GetImportedModel()->GetIdString();
-		KeySuffix += (bUseFullPrecisionUVs || !GVertexElementTypeSupport.IsSupported(VET_Half2)) ? "1" : "0";
+		KeySuffix += bUseFullPrecisionUVs ? "1" : "0";
 
 		//Dummy call to update the FSkeletalMeshLODModel::BuildStringID members.
 		GetImportedModel()->GetLODModelIdString();
diff --git a/Engine/Source/Runtime/Engine/Private/SkeletalMeshMerge.cpp b/Engine/Source/Runtime/Engine/Private/SkeletalMeshMerge.cpp
--- a/Engine/Source/Runtime/Engine/Private/SkeletalMeshMerge.cpp
+++ b/Engine/Source/Runtime/Engine/Private/SkeletalMeshMerge.cpp
@@ -232,7 +232,7 @@ bool FSkeletalMeshMerge::FinalizeMesh()
 		for (int32 LODIdx = 0; LODIdx < MaxNumLODs; LODIdx++)
 		{
 			const FSkeletalMeshLODInfo* LODInfoPtr = MergeMesh->GetLODInfo(LODIdx);
-			bool bUseFullPrecisionUVs = LODInfoPtr ? LODInfoPtr->BuildSettings.bUseFullPrecisionUVs : GVertexElementTypeSupport.IsSupported(VET_Half2) ? false : true;
+			bool bUseFullPrecisionUVs = LODInfoPtr ? LODInfoPtr->BuildSettings.bUseFullPrecisionUVs : false;
 			if (!bUseFullPrecisionUVs)
 			{
 				GENERATE_LOD_MODEL(TGPUSkinVertexFloat16Uvs, PerLODNumUVSets[LODIdx]);
diff --git a/Engine/Source/Runtime/Engine/Private/SkeletalMeshRenderData.cpp b/Engine/Source/Runtime/Engine/Private/SkeletalMeshRenderData.cpp
--- a/Engine/Source/Runtime/Engine/Private/SkeletalMeshRenderData.cpp
+++ b/Engine/Source/Runtime/Engine/Private/SkeletalMeshRenderData.cpp
@@ -345,7 +345,7 @@ void FSkeletalMeshRenderData::Cache(const ITargetPlatform* TargetPlatform, USkin
 					const bool bUseHighPrecisionTangentBasis = LODInfo->BuildSettings.bUseHighPrecisionTangentBasis;
 					const bool bUseBackwardsCompatibleF16TruncUVs = LODInfo->BuildSettings.bUseBackwardsCompatibleF16TruncUVs;
 					const bool bUseHighPrecisionWeights = LODInfo->BuildSettings.bUseHighPrecisionSkinWeights;
-					if (bUseFullPrecisionUVs || !GVertexElementTypeSupport.IsSupported(VET_Half2))
+					if (bUseFullPrecisionUVs)
 					{
 						VertexBufferBuildFlags |= ESkeletalMeshVertexFlags::UseFullPrecisionUVs;
 					}
diff --git a/Engine/Source/Runtime/Engine/Private/SkinnedAssetCommon.cpp b/Engine/Source/Runtime/Engine/Private/SkinnedAssetCommon.cpp
--- a/Engine/Source/Runtime/Engine/Private/SkinnedAssetCommon.cpp
+++ b/Engine/Source/Runtime/Engine/Private/SkinnedAssetCommon.cpp
@@ -80,8 +80,6 @@ static void SerializeBuildSettingsForDDC(FArchive& Ar, FSkeletalMeshBuildSetting
 
 FGuid FSkeletalMeshLODInfo::ComputeDeriveDataCacheKey(const FSkeletalMeshLODGroupSettings* SkeletalMeshLODGroupSettings)
 {
-	const bool bIs16BitfloatBufferSupported = GVertexElementTypeSupport.IsSupported(VET_Half2);
-
 	// Serialize the LOD info members, the BuildSettings and the ReductionSettings into a temporary array.
 	TArray<uint8> TempBytes;
 	TempBytes.Reserve(64);
@@ -117,10 +115,6 @@ FGuid FSkeletalMeshLODInfo::ComputeDeriveDataCacheKey(const FSkeletalMeshLODGrou
 		RealReductionSettings = SkeletalMeshLODGroupSettings->GetReductionSettings();
 	}
 
-	if (!BuildSettings.bUseFullPrecisionUVs && !bIs16BitfloatBufferSupported)
-	{
-		BuildSettings.bUseFullPrecisionUVs = true;
-	}
 	SerializeBuildSettingsForDDC(Ar, BuildSettings);
 	SerializeReductionSettingsForDDC(Ar, RealReductionSettings);
 
diff --git a/Engine/Source/Runtime/Engine/Private/StaticMesh.cpp b/Engine/Source/Runtime/Engine/Private/StaticMesh.cpp
--- a/Engine/Source/Runtime/Engine/Private/StaticMesh.cpp
+++ b/Engine/Source/Runtime/Engine/Private/StaticMesh.cpp
@@ -2694,7 +2694,7 @@ static FString BuildStaticMeshDerivedDataKeySuffix(const ITargetPlatform* Target
 		FMemoryWriter Ar(TempBytes, /*bIsPersistent=*/ true);
 		SerializeBuildSettingsForDDC(Ar, SrcModel.BuildSettings);
 
-		ANSICHAR Flag[2] = { (SrcModel.BuildSettings.bUseFullPrecisionUVs || !GVertexElementTypeSupport.IsSupported(VET_Half2)) ? '1' : '0', '\0' };
+		ANSICHAR Flag[2] = { SrcModel.BuildSettings.bUseFullPrecisionUVs ? '1' : '0', '\0' };
 		Ar.Serialize(Flag, 1);
 
 		FMeshReductionSettings FinalReductionSettings = LODGroup.GetSettings(SrcModel.ReductionSettings, LODIndex);
diff --git a/Engine/Source/Runtime/Landscape/Private/LandscapeGrass.cpp b/Engine/Source/Runtime/Landscape/Private/LandscapeGrass.cpp
--- a/Engine/Source/Runtime/Landscape/Private/LandscapeGrass.cpp
+++ b/Engine/Source/Runtime/Landscape/Private/LandscapeGrass.cpp
@@ -1922,7 +1922,7 @@ struct FAsyncGrassBuilder : public FGrassBuilderBase
 		, RequireCPUAccess(GrassVariety.bKeepInstanceBufferCPUCopy)
 
 		// output
-		, InstanceBuffer(/*bSupportsVertexHalfFloat*/ GVertexElementTypeSupport.IsSupported(VET_Half2))
+		, InstanceBuffer(/*bSupportsVertexHalfFloat*/ true)
 		, ClusterTree()
 		, OutOcclusionLayerNum(0)
 	{		
diff --git a/Engine/Source/Runtime/RHI/Private/RHI.cpp b/Engine/Source/Runtime/RHI/Private/RHI.cpp
--- a/Engine/Source/Runtime/RHI/Private/RHI.cpp
+++ b/Engine/Source/Runtime/RHI/Private/RHI.cpp
@@ -1388,8 +1388,6 @@ bool GRHISupportsSeparateDepthStencilCopyAccess = true;
 ERHIBindlessSupport GRHIBindlessSupport = ERHIBindlessSupport::Unsupported;
 bool GRHISupportsBindless = false;
 
-FVertexElementTypeSupportInfo GVertexElementTypeSupport;
-
 RHI_API int32 volatile GCurrentTextureMemorySize = 0;
 RHI_API int32 volatile GCurrentRendertargetMemorySize = 0;
 RHI_API int64 GTexturePoolSize = 0 * 1024 * 1024;
diff --git a/Engine/Source/Runtime/RHI/Public/RHI.h b/Engine/Source/Runtime/RHI/Public/RHI.h
--- a/Engine/Source/Runtime/RHI/Public/RHI.h
+++ b/Engine/Source/Runtime/RHI/Public/RHI.h
@@ -449,9 +449,6 @@ extern RHI_API int32 GDrawUPVertexCheckCount;
  */
 extern RHI_API int32 GDrawUPIndexCheckCount;
 
-/** true for each VET that is supported. One-to-one mapping with EVertexElementType */
-extern RHI_API class FVertexElementTypeSupportInfo GVertexElementTypeSupport;
-
 #include "MultiGPU.h" // IWYU pragma: export
 
 /** Whether the next frame should profile the GPU. */
@@ -729,18 +726,6 @@ inline EShaderPlatform GetFeatureLevelShaderPlatform(const FStaticFeatureLevel I
 	return GShaderPlatformForFeatureLevel[InFeatureLevel];
 }
 
-/** Info for supporting the vertex element types */
-class FVertexElementTypeSupportInfo
-{
-public:
-	FVertexElementTypeSupportInfo() { for(int32 i=0; i<VET_MAX; i++) ElementCaps[i]=true; }
-	FORCEINLINE bool IsSupported(EVertexElementType ElementType) { return ElementCaps[ElementType]; }
-	FORCEINLINE void SetSupported(EVertexElementType ElementType,bool bIsSupported) { ElementCaps[ElementType]=bIsSupported; }
-private:
-	/** cap bit set for each VET. One-to-one mapping based on EVertexElementType */
-	bool ElementCaps[VET_MAX];
-};
-
 struct FVertexElement
 {
 	uint8 StreamIndex;
