Cook optimization GetTextureDerivedDataKeySuffix: Remove allocation costs by using a StringBuilder.
#rnx

diff --git a/Engine/Source/Runtime/Engine/Private/TextureDerivedData.cpp b/Engine/Source/Runtime/Engine/Private/TextureDerivedData.cpp
--- a/Engine/Source/Runtime/Engine/Private/TextureDerivedData.cpp
+++ b/Engine/Source/Runtime/Engine/Private/TextureDerivedData.cpp
@@ -340,6 +340,7 @@ static void SerializeForKey(FArchive& Ar, const FTextureBuildSettings& Settings)
 void GetTextureDerivedDataKeySuffix(const UTexture& Texture, const FTextureBuildSettings* BuildSettingsPerLayer, FString& OutKeySuffix)
 {
 	uint16 Version = 0;
+	TStringBuilder<1024> KeyBuilder;
 
 	// Build settings for layer0 (used by default)
 	const FTextureBuildSettings& BuildSettings = BuildSettingsPerLayer[0];
@@ -368,7 +369,7 @@ void GetTextureDerivedDataKeySuffix(const UTexture& Texture, const FTextureBuild
 	}
 
 	// build the key, but don't use include the version if it's 0 to be backwards compatible
-	OutKeySuffix = FString::Printf(TEXT("%s_%s%s%s_%02u_%s"),
+	KeyBuilder.Appendf(TEXT("%s_%s%s%s_%02u_%s"),
 		*BuildSettings.TextureFormatName.GetPlainNameString(),
 		Version == 0 ? TEXT("") : *FString::Printf(TEXT("%d_"), Version),
 		*Texture.Source.GetIdString(),
@@ -393,35 +394,35 @@ void GetTextureDerivedDataKeySuffix(const UTexture& Texture, const FTextureBuild
 		{
 			LayerVersion = LayerTextureFormat->GetVersion(LayerBuildSettings.TextureFormatName, &LayerBuildSettings);
 		}
-		OutKeySuffix.Append(FString::Printf(TEXT("%s%d%s_"),
+		KeyBuilder.Appendf(TEXT("%s%d%s_"),
 			*LayerBuildSettings.TextureFormatName.GetPlainNameString(),
 			LayerVersion,
-			(LayerTextureFormat == NULL) ? TEXT("") : *LayerTextureFormat->GetDerivedDataKeyString(LayerBuildSettings)));
+			(LayerTextureFormat == NULL) ? TEXT("") : *LayerTextureFormat->GetDerivedDataKeyString(LayerBuildSettings));
 	}
 
 	if (BuildSettings.bVirtualStreamable)
 	{
 		// Additional GUID for virtual textures, make it easier to force these to rebuild while developing
-		OutKeySuffix.Append(FString::Printf(TEXT("VT%s_"), TEXTURE_VT_DERIVEDDATA_VER));
+		KeyBuilder.Appendf(TEXT("VT%s_"), TEXTURE_VT_DERIVEDDATA_VER);
 	}
 
 #if PLATFORM_CPU_ARM_FAMILY
 	// Separate out arm keys as x64 and arm64 clang do not generate the same data for a given
 	// input. Add the arm specifically so that a) we avoid rebuilding the current DDC and
 	// b) we can remove it once we get arm64 to be consistent.
-	OutKeySuffix.Append(TEXT("_arm64"));
+	KeyBuilder.Append(TEXT("_arm64"));
 #endif
 
 	if (BuildSettings.bAffectedBySharedLinearEncoding)
 	{
-		OutKeySuffix.Append(LexToString(GTextureSLEDerivedDataVer));
+		GTextureSLEDerivedDataVer.AppendString(KeyBuilder, EGuidFormats::Digits);
 	}
 
 	// Serialize the compressor settings into a temporary array. The archive
 	// is flagged as persistent so that machines of different endianness produce
 	// identical binary results.
 	TArray<uint8> TempBytes; 
-	TempBytes.Reserve(64);
+	TempBytes.Reserve(1024);
 	FMemoryWriter Ar(TempBytes, /*bIsPersistent=*/ true);
 	SerializeForKey(Ar, BuildSettings);
 
@@ -439,7 +440,8 @@ void GetTextureDerivedDataKeySuffix(const UTexture& Texture, const FTextureBuild
 
 	// Now convert the raw bytes to a string.
 	const uint8* SettingsAsBytes = TempBytes.GetData();
-	OutKeySuffix.Reserve(OutKeySuffix.Len() + TempBytes.Num());
+	OutKeySuffix.Reset(KeyBuilder.Len() + TempBytes.Num()*2 /* 2 hex characters per byte*/);
+	OutKeySuffix.Append(KeyBuilder.ToView());
 	for (int32 ByteIndex = 0; ByteIndex < TempBytes.Num(); ++ByteIndex)
 	{
 		ByteToHex(SettingsAsBytes[ByteIndex], OutKeySuffix);
