Immediately release static mesh offline ray tracing data if ray tracing is known to be disabled at run-time due to user settings.

diff --git a/Engine/Source/Runtime/Engine/Private/StaticMesh.cpp b/Engine/Source/Runtime/Engine/Private/StaticMesh.cpp
--- a/Engine/Source/Runtime/Engine/Private/StaticMesh.cpp
+++ b/Engine/Source/Runtime/Engine/Private/StaticMesh.cpp
@@ -575,6 +575,11 @@ void FStaticMeshLODResources::SerializeBuffers(FArchive& Ar, UStaticMesh* OwnerS
 	{
 		RayTracingGeometry.RawData.BulkSerialize(Ar);
 		AccumRayTracingGeometrySize(RayTracingGeometry, OutBuffersSize.SerializedBuffersSize);
+		if (Ar.IsLoading() && !IsRayTracingAllowed())
+		{
+			// Immediately release serialized offline BLAS data if it won't be used anyway due to rendering settings.
+			RayTracingGeometry.RawData.Discard();
+		}
 		bHasRayTracingGeometry = RayTracingGeometry.RawData.Num() != 0;
 	}
 
