World Partition: Fix Level Instance embedding on Level Streamed world partition in PIE

diff --git a/Engine/Source/Runtime/Engine/Private/WorldPartition/WorldPartitionLevelHelper.cpp b/Engine/Source/Runtime/Engine/Private/WorldPartition/WorldPartitionLevelHelper.cpp
--- a/Engine/Source/Runtime/Engine/Private/WorldPartition/WorldPartitionLevelHelper.cpp
+++ b/Engine/Source/Runtime/Engine/Private/WorldPartition/WorldPartitionLevelHelper.cpp
@@ -204,15 +204,15 @@ FString FWorldPartitionLevelHelper::AddActorContainerIDToSubPathString(const FAc
 	return InSubPathString;
 }
 
-FString FWorldPartitionLevelHelper::GetContainerPackage(const FActorContainerID& InContainerID, const FString& InPackageName, int32 InPIEInstanceID)
+FString FWorldPartitionLevelHelper::GetContainerPackage(const FActorContainerID& InContainerID, const FString& InPackageName, const FString& InDestLevelPackageName)
 {
-	FString ContainerPackageName = InPackageName;
-	if (InPIEInstanceID != INDEX_NONE)
+	uint64 DestLevelID = 0;
+	if (!InDestLevelPackageName.IsEmpty())
 	{
-		ContainerPackageName = UWorld::ConvertToPIEPackageName(ContainerPackageName, InPIEInstanceID);
+		DestLevelID = CityHash64((const char*)*InDestLevelPackageName, InDestLevelPackageName.Len() * sizeof(TCHAR));
 	}
 
-	return FString::Printf(TEXT("/Temp%s_%s"), *ContainerPackageName, *InContainerID.ToString());
+	return FString::Printf(TEXT("/Temp%s_%s_%016llx"), *InPackageName, *InContainerID.ToString(), DestLevelID);
 }
 
 bool FWorldPartitionLevelHelper::RemapActorPath(const FActorContainerID& InContainerID, const FString& InActorPath, FString& OutActorPath)
@@ -222,7 +222,7 @@ bool FWorldPartitionLevelHelper::RemapActorPath(const FActorContainerID& InConta
 		const FSoftObjectPath SoftObjectPath(InActorPath);
 		const FString LongPackageName = SoftObjectPath.GetLongPackageName();
 			
-		const FString ContainerPackageString = GetContainerPackage(InContainerID, LongPackageName, INDEX_NONE);	
+		const FString ContainerPackageString = GetContainerPackage(InContainerID, LongPackageName);	
 		const FString NewSubPathString = FWorldPartitionLevelHelper::AddActorContainerIDToSubPathString(InContainerID, SoftObjectPath.GetSubPathString());
 
 		FNameBuilder NewAssetPathBuilder;
@@ -328,8 +328,8 @@ bool FWorldPartitionLevelHelper::LoadActors(UWorld* InOwningWorld, ULevel* InDes
 		{
 			check(!PackageObjectMapping.ContainerID.IsMainContainer());
 		
-			const int32 PIEInstanceID = InDestLevel ? InDestLevel->GetPackage()->GetPIEInstanceID() : INDEX_NONE;
-			const FName ContainerPackageInstanceName(GetContainerPackage(PackageObjectMapping.ContainerID, PackageObjectMapping.ContainerPackage.ToString(), PIEInstanceID));
+			const FString DestLevelPackageName = InDestLevel ? InDestLevel->GetPackage()->GetName() : FString();
+			const FName ContainerPackageInstanceName(GetContainerPackage(PackageObjectMapping.ContainerID, PackageObjectMapping.ContainerPackage.ToString(), DestLevelPackageName));
 
 			FLinkerInstancingContext& NewContext = LinkerInstancingContexts.Add(PackageObjectMapping.ContainerID);
 			NewContext.AddTag(ULevel::DontLoadExternalObjectsTag);
diff --git a/Engine/Source/Runtime/Engine/Public/WorldPartition/WorldPartitionLevelHelper.h b/Engine/Source/Runtime/Engine/Public/WorldPartition/WorldPartitionLevelHelper.h
--- a/Engine/Source/Runtime/Engine/Public/WorldPartition/WorldPartitionLevelHelper.h
+++ b/Engine/Source/Runtime/Engine/Public/WorldPartition/WorldPartitionLevelHelper.h
@@ -39,7 +39,7 @@ public:
 	
 	static bool RemapActorPath(const FActorContainerID& InContainerID, const FString& InActorPath, FString& OutActorPath);
 	static FString AddActorContainerIDToSubPathString(const FActorContainerID& InContainerID, const FString& InSubPathString);
-	static FString GetContainerPackage(const FActorContainerID& InContainerID, const FString& InPackageName, int32 InPIEInstanceID);
+	static FString GetContainerPackage(const FActorContainerID& InContainerID, const FString& InPackageName, const FString& InDestLevelPackageName = FString());
 private:
 	FWorldPartitionLevelHelper();
 
