Fixing a crash with null asset user data in the array when saving tags.

diff --git a/Engine/Source/Runtime/Engine/Private/AnimationAsset.cpp b/Engine/Source/Runtime/Engine/Private/AnimationAsset.cpp
--- a/Engine/Source/Runtime/Engine/Private/AnimationAsset.cpp
+++ b/Engine/Source/Runtime/Engine/Private/AnimationAsset.cpp
@@ -728,7 +728,10 @@ void UAnimationAsset::GetAssetRegistryTags(TArray<FAssetRegistryTag>& OutTags) c
 	
 	for (const UAssetUserData* UserData : AssetUserData)
 	{
-		UserData->GetAssetRegistryTags(OutTags);
+		if (UserData)
+		{
+			UserData->GetAssetRegistryTags(OutTags);	
+		}
 	}
 	
 	OutTags.Add( FAssetRegistryTag("HasParentAsset", HasParentAsset() ? TEXT("True") : TEXT("False"), FAssetRegistryTag::TT_Hidden) );
