Rename the UPlatformInputDeviceMapperLibrary blueprint library to the UInputDeviceLibrary

#rnx

diff --git a/Engine/Config/BaseEngine.ini b/Engine/Config/BaseEngine.ini
--- a/Engine/Config/BaseEngine.ini
+++ b/Engine/Config/BaseEngine.ini
@@ -1098,6 +1098,7 @@ DefaultWorldPartitionHLODLayer=/Engine/Maps/Templates/HLODs/HLODLayer_Instanced.
 +PropertyRedirects=(OldName="AnimationSettings.BoneNamesWithCustomAttributes",NewName="BoneNamesWithAnimationAttributes")
 
 ; 5.2
++ClassRedirects=(OldName="/Script/Engine.PlatformInputDeviceMapperLibrary",NewName="/Script/Engine.InputDeviceLibrary")
 +StructRedirects=(OldName="/Script/AssetTools.RevisionInfo", NewName="/Script/AssetDefinition.RevisionInfo")
 +PropertyRedirects=(OldName="DecalComponent.SetFadeIn.Duaration", NewName="DecalComponent.SetFadeIn.Duration")
 +PropertyRedirects=(OldName="BlendSpace.PerBoneBlend", NewName="BlendSpace.ManualPerBoneOverrides")
diff --git a/Engine/Source/Runtime/Engine/Classes/GameFramework/InputDeviceLibrary.h b/Engine/Source/Runtime/Engine/Classes/GameFramework/InputDeviceLibrary.h
new file mode 100644
--- /dev/null
+++ b/Engine/Source/Runtime/Engine/Classes/GameFramework/InputDeviceLibrary.h
@@ -0,0 +1,157 @@
+// Copyright Epic Games, Inc. All Rights Reserved.
+
+#pragma once
+
+#include "CoreMinimal.h"
+#include "Kismet/BlueprintFunctionLibrary.h"
+
+#include "InputDeviceLibrary.generated.h"
+
+enum class EInputDeviceConnectionState : uint8;
+class APlayerController;
+
+ENGINE_API DECLARE_LOG_CATEGORY_EXTERN(LogInputDevices, Log, All);
+
+/**
+ * A static BP library that exposes Input Device data to blueprints
+ *
+ * @see IPlatformInputDeviceMapper
+ * @note Keep any function comments up to date with those in GenericPlatformInputDeviceMapper.h!
+ */
+UCLASS()
+class ENGINE_API UInputDeviceLibrary : public UBlueprintFunctionLibrary
+{
+	GENERATED_BODY()
+
+public:
+
+	/** Get the player controller who has the given Platform User ID. */
+	UFUNCTION(BlueprintCallable, Category = "Input Devices")
+	static APlayerController* GetPlayerControllerFromPlatformUser(const FPlatformUserId UserId);
+
+	/** Get the player controller who owns the given input device id */
+	UFUNCTION(BlueprintCallable, Category = "Input Devices")
+	static APlayerController* GetPlayerControllerFromInputDevice(const FInputDeviceId DeviceId);
+
+	/**
+	 * Populates the OutInputDevices array with any InputDeviceID's that are mapped to the given platform user
+	 *
+	 * @param UserId				The Platform User to gather the input devices of.
+	 * @param OutInputDevices		Array of input device ID's that will be populated with the mapped devices.
+	 * @return						The number of mapped devices, INDEX_NONE if the user was not found.
+	 */
+	UFUNCTION(BlueprintCallable, Category = "PlatformInputDevice")
+	static int32 GetAllInputDevicesForUser(const FPlatformUserId UserId, TArray<FInputDeviceId>& OutInputDevices);
+
+	/**
+	 * Get all mapped input devices on this platform regardless of their connection state.
+	 * 
+	 * @param OutInputDevices	Array of input devices to populate
+	 * @return					The number of connected input devices
+	 */
+	UFUNCTION(BlueprintCallable, Category = "PlatformInputDevice")
+	static int32 GetAllInputDevices(TArray<FInputDeviceId>& OutInputDevices);
+
+	/**
+	 * Gather all currently connected input devices
+	 * 
+	 * @param OutInputDevices	Array of input devices to populate
+	 * @return					The number of connected input devices
+	 */
+	UFUNCTION(BlueprintCallable, Category = "PlatformInputDevice")
+	static int32 GetAllConnectedInputDevices(TArray<FInputDeviceId>& OutInputDevices);
+
+	/**
+	 * Get all currently active platform ids, anyone who has a mapped input device
+	 *
+	 * @param OutUsers		Array that will be populated with the platform users.
+	 * @return				The number of active platform users
+	 */
+	UFUNCTION(BlueprintCallable, Category = "PlatformInputDevice")
+	static int32 GetAllActiveUsers(TArray<FPlatformUserId>& OutUsers);
+
+	/**
+	 * Returns the platform user id that is being used for unmapped input devices.
+	 * Will be PLATFORMUSERID_NONE if platform does not support this (this is the default behavior)
+	 */
+	UFUNCTION(BlueprintPure, Category = "PlatformInputDevice")
+	static FPlatformUserId GetUserForUnpairedInputDevices();
+
+	/**
+	 * Returns the 'Primary' Platform user for this platform.
+	 * This typically has an internal ID of '0' and is used as the default platform user to
+	 * map devices such as the keyboard and mouse that don't get assigned unique ID's from their
+	 * owning platform code.
+	 */
+	UFUNCTION(BlueprintPure, Category = "PlatformInputDevice")
+	static FPlatformUserId GetPrimaryPlatformUser();
+
+	/** Returns true if the given Platform User Id is the user for unpaired input devices on this platform. */
+	UFUNCTION(BlueprintPure, Category = "PlatformInputDevice")
+	static bool IsUnpairedUserId(const FPlatformUserId PlatformId);
+
+	/** Returns true if the given input device is mapped to the unpaired platform user id. */
+	UFUNCTION(BlueprintPure, Category = "PlatformInputDevice")
+	static bool IsInputDeviceMappedToUnpairedUser(const FInputDeviceId InputDevice);
+
+	/** Returns the default device id used for things like keyboard/mouse input */
+	UFUNCTION(BlueprintPure, Category = "PlatformInputDevice")
+	static FInputDeviceId GetDefaultInputDevice();
+
+	/** Returns the platform user attached to this input device, or PLATFORMUSERID_NONE if invalid */
+	UFUNCTION(BlueprintCallable, Category = "PlatformInputDevice")
+	static FPlatformUserId GetUserForInputDevice(FInputDeviceId DeviceId);
+
+	/** Returns the primary input device used by a specific player, or INPUTDEVICEID_NONE if invalid */
+	UFUNCTION(BlueprintCallable, Category = "PlatformInputDevice")
+	static FInputDeviceId GetPrimaryInputDeviceForUser(FPlatformUserId UserId);
+
+	/**
+	 * Gets the connection state of the given input device.
+	 * 
+	 * @param DeviceId		The device to get the connection state of
+	 * @return				The connection state of the given device. EInputDeviceConnectionState::Unknown if the device is not mapped
+	 */
+	UFUNCTION(BlueprintCallable, Category = "PlatformInputDevice")
+	static EInputDeviceConnectionState GetInputDeviceConnectionState(const FInputDeviceId DeviceId);
+
+	/**
+	 * Check if the given input device is valid
+	 * 
+	 * @return	True if the given input device is valid (it has been assigned an ID by the PlatformInputDeviceMapper)
+	 */
+	UFUNCTION(BlueprintPure, Category = "PlatformInputDevice")
+	static bool IsValidInputDevice(FInputDeviceId DeviceId);
+
+	/**
+	 * Check if the given platform ID is valid
+	 * 
+	 * @return	True if the platform ID is valid (it has been assigned by the PlatformInputDeviceMapper)
+	 */
+	UFUNCTION(BlueprintPure, Category = "PlatformInputDevice")
+	static bool IsValidPlatformId(FPlatformUserId UserId);
+
+	/** Static invalid platform user */
+	UFUNCTION(BlueprintPure, meta = (ScriptConstant = "None", ScriptConstantHost = "/Script/CoreUObject.PlatformUserId"), Category = "PlatformInputDevice")
+	static FPlatformUserId PlatformUserId_None();
+
+	/** Static invalid input device */
+	UFUNCTION(BlueprintPure, meta = (ScriptConstant = "None", ScriptConstantHost = "/Script/CoreUObject.InputDeviceId"), Category = "PlatformInputDevice")
+	static FInputDeviceId InputDeviceId_None();
+
+	/** Returns true if PlatformUserId A is equal to PlatformUserId B (A == B) */
+	UFUNCTION(BlueprintPure, meta = (DisplayName = "Equal (PlatformUserId)", CompactNodeTitle = "==", ScriptMethod = "Equals", ScriptOperator = "==", Keywords = "== equal"), Category = "PlatformInputDevice")
+	static bool EqualEqual_PlatformUserId(FPlatformUserId A, FPlatformUserId B);
+
+	/** Returns true if PlatformUserId A is not equal to PlatformUserId B (A != B) */
+	UFUNCTION(BlueprintPure, meta = (DisplayName = "Not Equal (PlatformUserId)", CompactNodeTitle = "!=", ScriptMethod = "NotEqual", ScriptOperator = "!=", Keywords = "!= not equal"), Category = "PlatformInputDevice")
+	static bool NotEqual_PlatformUserId(FPlatformUserId A, FPlatformUserId B);
+	
+	/** Returns true if InputDeviceId A is equal to InputDeviceId B (A == B) */
+	UFUNCTION(BlueprintPure, meta = (DisplayName = "Equal (InputDeviceId)", CompactNodeTitle = "==", ScriptMethod = "Equals", ScriptOperator = "==", Keywords = "== equal"), Category = "PlatformInputDevice")
+	static bool EqualEqual_InputDeviceId(FInputDeviceId A, FInputDeviceId B);
+	
+	/** Returns true if InputDeviceId A is not equal to InputDeviceId B (A != B) */
+	UFUNCTION(BlueprintPure, meta = (DisplayName = "Not Equal (InputDeviceId)", CompactNodeTitle = "!=", ScriptMethod = "NotEqual", ScriptOperator = "!=", Keywords = "!= not equal"), Category = "PlatformInputDevice")
+	static bool NotEqual_InputDeviceId(FInputDeviceId A, FInputDeviceId B);
+};
\ No newline at end of file
diff --git a/Engine/Source/Runtime/Engine/Classes/GameFramework/PlatformInputDeviceMapperLibrary.h b/Engine/Source/Runtime/Engine/Classes/GameFramework/PlatformInputDeviceMapperLibrary.h
deleted file mode 100644
diff --git a/Engine/Source/Runtime/Engine/Private/InputDeviceLibrary.cpp b/Engine/Source/Runtime/Engine/Private/InputDeviceLibrary.cpp
new file mode 100644
--- /dev/null
+++ b/Engine/Source/Runtime/Engine/Private/InputDeviceLibrary.cpp
@@ -0,0 +1,144 @@
+// Copyright Epic Games, Inc. All Rights Reserved.
+
+#include "GameFramework/InputDeviceLibrary.h"
+#include "GameFramework/PlayerController.h"
+#include "GameFramework/InputDeviceSubsystem.h"
+#include "Engine/World.h"
+#include "GenericPlatform/GenericPlatformInputDeviceMapper.h"
+
+#include UE_INLINE_GENERATED_CPP_BY_NAME(InputDeviceLibrary)
+
+DEFINE_LOG_CATEGORY(LogInputDevices);
+
+APlayerController* UInputDeviceLibrary::GetPlayerControllerFromPlatformUser(const FPlatformUserId UserId)
+{
+	if (UserId.IsValid())
+	{
+		if (const UInputDeviceSubsystem* System = UInputDeviceSubsystem::Get())
+		{
+			if (const UWorld* World = System->GetTickableGameObjectWorld())
+			{
+				for (FConstPlayerControllerIterator Iterator = World->GetPlayerControllerIterator(); Iterator; ++Iterator)
+				{
+					APlayerController* PlayerController = Iterator->Get();
+					const FPlatformUserId PlayerControllerUserID = PlayerController ? PlayerController->GetPlatformUserId() : PLATFORMUSERID_NONE;
+					if (PlayerControllerUserID.IsValid() && PlayerControllerUserID == UserId)
+					{
+						return PlayerController;
+					}
+				}
+			}
+		}
+	}
+	else
+	{
+		UE_LOG(LogInputDevices, Warning, TEXT("Attempting to find the player controller for an invalid Platform User Id"));
+	}
+
+	return nullptr;
+}
+
+APlayerController* UInputDeviceLibrary::GetPlayerControllerFromInputDevice(const FInputDeviceId DeviceId)
+{
+	return GetPlayerControllerFromPlatformUser(IPlatformInputDeviceMapper::Get().GetUserForInputDevice(DeviceId));
+}
+
+int32 UInputDeviceLibrary::GetAllInputDevicesForUser(const FPlatformUserId UserId, TArray<FInputDeviceId>& OutInputDevices)
+{
+	return IPlatformInputDeviceMapper::Get().GetAllInputDevicesForUser(UserId, OutInputDevices);
+}
+
+int32 UInputDeviceLibrary::GetAllInputDevices(TArray<FInputDeviceId>& OutInputDevices)
+{
+	return IPlatformInputDeviceMapper::Get().GetAllInputDevices(OutInputDevices);
+}
+
+int32 UInputDeviceLibrary::GetAllConnectedInputDevices(TArray<FInputDeviceId>& OutInputDevices)
+{
+	return IPlatformInputDeviceMapper::Get().GetAllConnectedInputDevices(OutInputDevices);
+}
+
+int32 UInputDeviceLibrary::GetAllActiveUsers(TArray<FPlatformUserId>& OutUsers)
+{
+	return IPlatformInputDeviceMapper::Get().GetAllActiveUsers(OutUsers);
+}
+
+FPlatformUserId UInputDeviceLibrary::GetUserForUnpairedInputDevices()
+{
+	return IPlatformInputDeviceMapper::Get().GetUserForUnpairedInputDevices();
+}
+
+FPlatformUserId UInputDeviceLibrary::GetPrimaryPlatformUser()
+{
+	return IPlatformInputDeviceMapper::Get().GetPrimaryPlatformUser();
+}
+
+bool UInputDeviceLibrary::IsUnpairedUserId(const FPlatformUserId PlatformId)
+{
+	return IPlatformInputDeviceMapper::Get().IsUnpairedUserId(PlatformId);
+}
+
+bool UInputDeviceLibrary::IsInputDeviceMappedToUnpairedUser(const FInputDeviceId InputDevice)
+{
+	return IPlatformInputDeviceMapper::Get().IsInputDeviceMappedToUnpairedUser(InputDevice);
+}
+
+FInputDeviceId UInputDeviceLibrary::GetDefaultInputDevice()
+{
+	return IPlatformInputDeviceMapper::Get().GetDefaultInputDevice();
+}
+
+FPlatformUserId UInputDeviceLibrary::GetUserForInputDevice(FInputDeviceId DeviceId)
+{
+	return IPlatformInputDeviceMapper::Get().GetUserForInputDevice(DeviceId);
+}
+
+FInputDeviceId UInputDeviceLibrary::GetPrimaryInputDeviceForUser(FPlatformUserId UserId)
+{
+	return IPlatformInputDeviceMapper::Get().GetPrimaryInputDeviceForUser(UserId);
+}
+
+EInputDeviceConnectionState UInputDeviceLibrary::GetInputDeviceConnectionState(const FInputDeviceId DeviceId)
+{
+	return IPlatformInputDeviceMapper::Get().GetInputDeviceConnectionState(DeviceId);
+}
+
+bool UInputDeviceLibrary::IsValidInputDevice(FInputDeviceId DeviceId)
+{
+	return DeviceId.IsValid();
+}
+
+bool UInputDeviceLibrary::IsValidPlatformId(FPlatformUserId UserId)
+{
+	return UserId.IsValid();
+}
+
+FPlatformUserId UInputDeviceLibrary::PlatformUserId_None()
+{
+	return PLATFORMUSERID_NONE;
+}
+
+FInputDeviceId UInputDeviceLibrary::InputDeviceId_None()
+{
+	return INPUTDEVICEID_NONE;
+}
+
+bool UInputDeviceLibrary::EqualEqual_PlatformUserId(FPlatformUserId A, FPlatformUserId B)
+{
+	return A == B;
+}
+
+bool UInputDeviceLibrary::NotEqual_PlatformUserId(FPlatformUserId A, FPlatformUserId B)
+{
+	return A != B;
+}
+
+bool UInputDeviceLibrary::EqualEqual_InputDeviceId(FInputDeviceId A, FInputDeviceId B)
+{
+	return A == B;
+}
+
+bool UInputDeviceLibrary::NotEqual_InputDeviceId(FInputDeviceId A, FInputDeviceId B)
+{
+	return A != B;
+}
diff --git a/Engine/Source/Runtime/Engine/Private/PlatformInputDeviceMapperLibrary.cpp b/Engine/Source/Runtime/Engine/Private/PlatformInputDeviceMapperLibrary.cpp
deleted file mode 100644
