Update metadata on Input Device properties to make the editor experience better

#rnx
#pf skip

diff --git a/Engine/Source/Runtime/Engine/Classes/GameFramework/InputDeviceProperties.h b/Engine/Source/Runtime/Engine/Classes/GameFramework/InputDeviceProperties.h
--- a/Engine/Source/Runtime/Engine/Classes/GameFramework/InputDeviceProperties.h
+++ b/Engine/Source/Runtime/Engine/Classes/GameFramework/InputDeviceProperties.h
@@ -131,6 +131,9 @@ protected:
 	/**
 	* The duration that this device property should last. Override this if your property has any dynamic curves 
 	* to be the max time range.
+	* 
+	* A duration of 0 means that the device property will be treated as a "One Shot" effect, being applied once
+	* before being removed by the Input Device Subsystem.
 	*/
 	UPROPERTY(VisibleAnywhere, BlueprintReadWrite, Category = "Info")
 	float PropertyDuration = 0.0f;
@@ -181,7 +184,7 @@ struct FDeviceColorData
 * NOTE: This property has platform specific implementations and may behave differently per platform.
 * See the docs for more details on each platform.
 */
-UCLASS(Blueprintable, BlueprintType)
+UCLASS(Blueprintable, BlueprintType, meta = (DisplayName = "Device Color (Static)"))
 class UColorInputDeviceProperty : public UInputDeviceProperty
 {
 	GENERATED_BODY()
@@ -235,7 +238,7 @@ struct FDeviceColorCurveData
 * NOTE: This property has platform specific implementations and may behave differently per platform.
 * See the docs for more details on each platform.
 */
-UCLASS(Blueprintable, BlueprintType)
+UCLASS(Blueprintable, BlueprintType, meta = (DisplayName = "Device Color (Curve)"))
 class UColorInputDeviceCurveProperty : public UInputDeviceProperty
 {
 	GENERATED_BODY()
@@ -272,16 +275,16 @@ struct FDeviceTriggerBaseData
 	GENERATED_BODY()
 
 	/** Which trigger this property should effect */
-	UPROPERTY(EditAnywhere, BlueprintReadWrite, Category = "DeviceProperty")
+	UPROPERTY(EditAnywhere, BlueprintReadWrite, Category = "Triggers")
 	EInputDeviceTriggerMask AffectedTriggers = EInputDeviceTriggerMask::None;
 
 	/** True if the triggers should be reset after the duration of this device property */
-	UPROPERTY(EditAnywhere, BlueprintReadWrite, Category = "DeviceProperty")
+	UPROPERTY(EditAnywhere, BlueprintReadWrite, Category = "Triggers")
 	bool bResetUponCompletion = true;
 };
 
 /** A property that effect the triggers on a gamepad */
-UCLASS(Abstract, Blueprintable, meta = (DisplayName = "Reset Trigger Device Properties"))
+UCLASS(Abstract, meta = (DisplayName = "Base Trigger Effect"))
 class ENGINE_API UInputDeviceTriggerEffect : public UInputDeviceProperty
 {
 	GENERATED_BODY()
@@ -291,7 +294,7 @@ public:
 	virtual FInputDeviceProperty* GetInternalDeviceProperty() override;
 	virtual void ResetDeviceProperty_Implementation(const FPlatformUserId PlatformUser, const FInputDeviceId DeviceId) override;
 
-	UPROPERTY(EditAnywhere, BlueprintReadWrite, Category = "DeviceProperty")
+	UPROPERTY(EditAnywhere, BlueprintReadWrite, Category = "Triggers")
 	FDeviceTriggerBaseData BaseTriggerData;
 
 protected:
@@ -323,7 +326,7 @@ struct FDeviceTriggerFeedbackData
 * NOTE: This property has platform specific implementations and may behave differently per platform.
 * See the docs for more details on each platform.
 */
-UCLASS(Blueprintable)
+UCLASS(Blueprintable, meta = (DisplayName = "Trigger Feedback"))
 class UInputDeviceTriggerFeedbackProperty : public UInputDeviceTriggerEffect
 {
 	GENERATED_BODY()
@@ -362,17 +365,17 @@ struct FDeviceTriggerTriggerResistanceData
 	GENERATED_BODY()
 
 	/** The position that the trigger should start providing resistance */
-	UPROPERTY(EditAnywhere, BlueprintReadWrite, Category = "DeviceProperty")
+	UPROPERTY(EditAnywhere, BlueprintReadWrite, Category = "DeviceProperty", meta = (UIMin = "0"))
 	int32 StartPosition = 0;
 
 	/** How strong the resistance is */
-	UPROPERTY(EditAnywhere, BlueprintReadWrite, Category = "DeviceProperty")
+	UPROPERTY(EditAnywhere, BlueprintReadWrite, Category = "DeviceProperty", meta = (UIMin = "0"))
 	int32 StartStrengh = 0;
 
-	UPROPERTY(EditAnywhere, BlueprintReadWrite, Category = "DeviceProperty")
+	UPROPERTY(EditAnywhere, BlueprintReadWrite, Category = "DeviceProperty", meta = (UIMin = "0"))
 	int32 EndPosition = 0;
 
-	UPROPERTY(EditAnywhere, BlueprintReadWrite, Category = "DeviceProperty")
+	UPROPERTY(EditAnywhere, BlueprintReadWrite, Category = "DeviceProperty", meta = (UIMin = "0"))
 	int32 EndStrengh = 0;
 };
 
@@ -382,7 +385,7 @@ struct FDeviceTriggerTriggerResistanceData
 * NOTE: This property has platform specific implementations and may behave differently per platform.
 * See the docs for more details on each platform.
 */
-UCLASS(Blueprintable)
+UCLASS(Blueprintable, meta = (DisplayName = "Trigger Resistance"))
 class UInputDeviceTriggerResistanceProperty : public UInputDeviceTriggerEffect
 {
 	GENERATED_BODY()
@@ -439,7 +442,7 @@ struct FDeviceTriggerTriggerVibrationData
 * NOTE: This property has platform specific implementations and may behave differently per platform.
 * See the docs for more details on each platform.
 */
-UCLASS(Blueprintable)
+UCLASS(Blueprintable, meta = (DisplayName = "Trigger Vibration"))
 class UInputDeviceTriggerVibrationProperty : public UInputDeviceTriggerEffect
 {
 	GENERATED_BODY()
