When push model fast array tracking is enabled, don't skip if there are unacked property retirements

diff --git a/Engine/Source/Runtime/Engine/Private/DataReplication.cpp b/Engine/Source/Runtime/Engine/Private/DataReplication.cpp
--- a/Engine/Source/Runtime/Engine/Private/DataReplication.cpp
+++ b/Engine/Source/Runtime/Engine/Private/DataReplication.cpp
@@ -1752,7 +1752,7 @@ bool FObjectReplicator::CanSkipUpdate(FReplicationFlags RepFlags)
 			SendingRepState.LastChangelistIndex == RepChangelistState.HistoryEnd &&
 			Connection->ResendAllDataState == EResendAllDataState::None &&
 			!OwningChannel->bForceCompareProperties &&
-			(!GbPushModelSkipUndirtiedFastArrays || (RepChangelistState.CustomDeltaChangeIndex == SendingRepState.CustomDeltaChangeIndex)) &&
+			(!GbPushModelSkipUndirtiedFastArrays || ((RepChangelistState.CustomDeltaChangeIndex == SendingRepState.CustomDeltaChangeIndex) && !SendingRepState.HasAnyPendingRetirements())) &&
 			!RepChangelistState.HasAnyDirtyProperties())
 		))
 	{
diff --git a/Engine/Source/Runtime/Engine/Private/RepLayout.cpp b/Engine/Source/Runtime/Engine/Private/RepLayout.cpp
--- a/Engine/Source/Runtime/Engine/Private/RepLayout.cpp
+++ b/Engine/Source/Runtime/Engine/Private/RepLayout.cpp
@@ -7148,6 +7148,19 @@ TStaticBitArray<COND_Max> FSendingRepState::BuildConditionMapFromRepFlags(const
 	return ConditionMap;
 }
 
+bool FSendingRepState::HasAnyPendingRetirements() const
+{
+	for (const FPropertyRetirement& PropRet : Retirement)
+	{
+		if (PropRet.Next != nullptr)
+		{
+			return true;
+		}
+	}
+
+	return false;
+}
+
 void FRepLayout::RebuildConditionalProperties(
 	FSendingRepState* RESTRICT RepState,
 	const FReplicationFlags& RepFlags) const
diff --git a/Engine/Source/Runtime/Engine/Public/Net/RepLayout.h b/Engine/Source/Runtime/Engine/Public/Net/RepLayout.h
--- a/Engine/Source/Runtime/Engine/Public/Net/RepLayout.h
+++ b/Engine/Source/Runtime/Engine/Public/Net/RepLayout.h
@@ -572,6 +572,8 @@ public:
 	 */
 	static TStaticBitArray<COND_Max> BuildConditionMapFromRepFlags(const FReplicationFlags InFlags);
 
+	bool HasAnyPendingRetirements() const;
+
 	/** Whether or not FRepLayout::OpenAcked has been called with this FRepState. */
 	bool bOpenAckedCalled;
 
