Added stat for AddParameterCollectionInstance.
- Refactored duplicate code by adding SetupWorldParameterCollectionInstances.
- Renamed some stats for consistency.

diff --git a/Engine/Source/Editor/UnrealEd/Private/Factories/EditorFactories.cpp b/Engine/Source/Editor/UnrealEd/Private/Factories/EditorFactories.cpp
--- a/Engine/Source/Editor/UnrealEd/Private/Factories/EditorFactories.cpp
+++ b/Engine/Source/Editor/UnrealEd/Private/Factories/EditorFactories.cpp
@@ -655,11 +655,7 @@ UObject* UMaterialParameterCollectionFactoryNew::FactoryCreateNew(UClass* Class,
 	
 	if (MaterialParameterCollection)
 	{
-		for (TObjectIterator<UWorld> It; It; ++It)
-		{
-			UWorld* CurrentWorld = *It;
-			CurrentWorld->AddParameterCollectionInstance(MaterialParameterCollection, true);
-		}
+		MaterialParameterCollection->SetupWorldParameterCollectionInstances();
 	}
 
 	return MaterialParameterCollection;
diff --git a/Engine/Source/Runtime/Engine/Classes/Materials/MaterialParameterCollection.h b/Engine/Source/Runtime/Engine/Classes/Materials/MaterialParameterCollection.h
--- a/Engine/Source/Runtime/Engine/Classes/Materials/MaterialParameterCollection.h
+++ b/Engine/Source/Runtime/Engine/Classes/Materials/MaterialParameterCollection.h
@@ -173,6 +173,8 @@ class UMaterialParameterCollection : public UObject
 		return *UniformBufferStruct;
 	}
 
+	ENGINE_API void SetupWorldParameterCollectionInstances();
+
 private:
 	virtual ENGINE_API void FinishDestroy() override;
 	virtual ENGINE_API bool IsReadyForFinishDestroy() override;
diff --git a/Engine/Source/Runtime/Engine/Private/Materials/ParameterCollection.cpp b/Engine/Source/Runtime/Engine/Private/Materials/ParameterCollection.cpp
--- a/Engine/Source/Runtime/Engine/Private/Materials/ParameterCollection.cpp
+++ b/Engine/Source/Runtime/Engine/Private/Materials/ParameterCollection.cpp
@@ -47,15 +47,18 @@ void UMaterialParameterCollection::PostLoad()
 	}
 
 	CreateBufferStruct();
+	SetupWorldParameterCollectionInstances();
+	UpdateDefaultResource(true);
+}
 
+void UMaterialParameterCollection::SetupWorldParameterCollectionInstances()
+{
 	// Create an instance for this collection in every world
 	for (TObjectIterator<UWorld> It; It; ++It)
 	{
 		UWorld* CurrentWorld = *It;
 		CurrentWorld->AddParameterCollectionInstance(this, true);
 	}
-
-	UpdateDefaultResource(true);
 }
 
 void UMaterialParameterCollection::BeginDestroy()
diff --git a/Engine/Source/Runtime/Engine/Private/World.cpp b/Engine/Source/Runtime/Engine/Private/World.cpp
--- a/Engine/Source/Runtime/Engine/Private/World.cpp
+++ b/Engine/Source/Runtime/Engine/Private/World.cpp
@@ -1359,7 +1359,7 @@ UWorld* UWorld::GetWorld() const
 
 void UWorld::SetupParameterCollectionInstances()
 {
-	QUICK_SCOPE_CYCLE_COUNTER(Stat_World_SetupParameterCollectionInstances);
+	QUICK_SCOPE_CYCLE_COUNTER(STAT_World_SetupParameterCollectionInstances);
 
 	// Create an instance for each parameter collection in memory
 	for (UMaterialParameterCollection* CurrentCollection : TObjectRange<UMaterialParameterCollection>())
@@ -1372,6 +1372,8 @@ void UWorld::SetupParameterCollectionInstances()
 
 void UWorld::AddParameterCollectionInstance(UMaterialParameterCollection* Collection, bool bUpdateScene)
 {
+	QUICK_SCOPE_CYCLE_COUNTER(STAT_World_AddParameterCollectionInstance);
+
 	int32 ExistingIndex = INDEX_NONE;
 
 	for (int32 InstanceIndex = 0; InstanceIndex < ParameterCollectionInstances.Num(); InstanceIndex++)
@@ -2937,7 +2939,7 @@ void UWorld::AddToWorld( ULevel* Level, const FTransform& LevelTransform, bool b
 		{
 			if (!bConsiderTimeLimit)
 			{
-				QUICK_SCOPE_CYCLE_COUNTER(UWorld_AddToWorld_WaitFor_AsyncPreRegisterLevelStreamingTasks);
+				QUICK_SCOPE_CYCLE_COUNTER(STAT_World_AddToWorld_WaitFor_AsyncPreRegisterLevelStreamingTasks);
 
 				for (TSharedPtr<FAsyncPreRegisterDDCRequest>& Request : AsyncPreRegisterDDCRequests)
 				{
@@ -3506,7 +3508,7 @@ UWorld* UWorld::GetDuplicatedWorldForPIE(UWorld* InWorld, UPackage* InPIEackage,
 UWorld* UWorld::DuplicateWorldForPIE(const FString& PackageName, UWorld* OwningWorld)
 {
 #if WITH_EDITOR
-	QUICK_SCOPE_CYCLE_COUNTER(UWorld_DuplicateWorldForPIE);
+	QUICK_SCOPE_CYCLE_COUNTER(STAT_World_DuplicateWorldForPIE);
 	FScopeCycleCounterUObject Context(OwningWorld);
 
 	FName PackageFName(*PackageName);
@@ -3601,7 +3603,7 @@ UWorld* UWorld::DuplicateWorldForPIE(const FString& PackageName, UWorld* OwningW
 			&& PIELevel->Model == EditorLevel->Model
 			&& PIELevel->ModelComponents.Num() == EditorLevel->ModelComponents.Num() )
 		{
-			QUICK_SCOPE_CYCLE_COUNTER(UWorld_DuplicateWorldForPIE_UpdateModelComponents);
+			QUICK_SCOPE_CYCLE_COUNTER(STAT_World_DuplicateWorldForPIE_UpdateModelComponents);
 
 			PIELevel->Model->ClearLocalMaterialIndexBuffersData();
 			for (int32 ComponentIndex = 0; ComponentIndex < PIELevel->ModelComponents.Num(); ++ComponentIndex)
