Expose Access to Curves, TickRecords and Graph Data From The Proxy for use in Debug Animation Tracking

Contributed by [REDACTED_PERSON]

diff --git a/Engine/Source/Runtime/Engine/Classes/Animation/AnimInstance.h b/Engine/Source/Runtime/Engine/Classes/Animation/AnimInstance.h
--- a/Engine/Source/Runtime/Engine/Classes/Animation/AnimInstance.h
+++ b/Engine/Source/Runtime/Engine/Classes/Animation/AnimInstance.h
@@ -17,6 +17,7 @@
 #endif
 #include "Animation/AnimNotifyQueue.h"
 #include "Animation/AnimSubsystemInstance.h"
+#include "Animation/AnimSync.h"
 #include "Animation/AnimNotifies/AnimNotify.h"
 #include "AnimInstance.generated.h"
 
@@ -47,6 +48,7 @@ struct FBlendedHeapCurve;
 struct FBoneContainer;
 struct FSmartNameMapping;
 struct FAnimNode_LinkedAnimLayer;
+struct FNodeDebugData;
 enum class ETransitionRequestQueueMode : uint8;
 enum class ETransitionRequestOverwriteMode : uint8;
 
@@ -432,6 +434,18 @@ public:
 	 */
 	UFUNCTION(BlueprintCallable, Category = "Animation|Pose")
 	virtual void SnapshotPose(UPARAM(ref) FPoseSnapshot& Snapshot);
+	
+	/** Get the sync group we are currently reading from */
+	const TMap<FName, FAnimGroupInstance>& GetSyncGroupMapRead() const;
+
+	/** Get the ungrouped active player we are currently reading from */
+	const TArray<FAnimTickRecord>& GetUngroupedActivePlayersRead();
+
+	/** Get the current value of all animation curves **/
+	const TMap<FName, float>& GetAnimationCurves(EAnimCurveType InCurveType) const;
+
+	/** Gather debug data from this instance proxy and the blend tree for display */
+	void GatherDebugData(FNodeDebugData& DebugData);
 
 	// Can this animation instance run Update or Evaluation work in parallel
 	virtual bool CanRunParallelWork() const { return true; }
diff --git a/Engine/Source/Runtime/Engine/Private/Animation/AnimInstance.cpp b/Engine/Source/Runtime/Engine/Private/Animation/AnimInstance.cpp
--- a/Engine/Source/Runtime/Engine/Private/Animation/AnimInstance.cpp
+++ b/Engine/Source/Runtime/Engine/Private/Animation/AnimInstance.cpp
@@ -192,6 +192,30 @@ void UAnimInstance::SnapshotPose(FPoseSnapshot& Snapshot)
 	}
 }
 
+const TMap<FName, FAnimGroupInstance>& UAnimInstance::GetSyncGroupMapRead() const
+{
+	const FAnimInstanceProxy& Proxy = GetProxyOnGameThread<FAnimInstanceProxy>();
+	return Proxy.GetSyncGroupMapRead();
+}
+
+const TArray<FAnimTickRecord>& UAnimInstance::GetUngroupedActivePlayersRead()
+{
+	FAnimInstanceProxy& Proxy = GetProxyOnGameThread<FAnimInstanceProxy>();
+	return Proxy.GetUngroupedActivePlayersRead();
+}
+
+const TMap<FName, float>& UAnimInstance::GetAnimationCurves(EAnimCurveType InCurveType) const
+{
+	const FAnimInstanceProxy& Proxy = GetProxyOnGameThread<FAnimInstanceProxy>();
+	return Proxy.GetAnimationCurves(InCurveType);
+}
+
+void UAnimInstance::GatherDebugData(FNodeDebugData& DebugData)
+{
+	FAnimInstanceProxy& Proxy = GetProxyOnGameThread<FAnimInstanceProxy>();
+	Proxy.GatherDebugData(DebugData);
+}
+
 USkeletalMeshComponent* UAnimInstance::GetOwningComponent() const
 {
 	return GetSkelMeshComponent();
