- Ensure that LinkedBone(references) for Curve LODing are updated with new Skeleton when merging
- Implement remapping of BoneTrackIndices for PoseAssets between Skeletons, to support evaluating them with _remapping only_ skeletons

diff --git a/Engine/Plugins/Runtime/SkeletalMerging/Source/SkeletalMerging/Private/SkeletalMergingLibrary.cpp b/Engine/Plugins/Runtime/SkeletalMerging/Source/SkeletalMerging/Private/SkeletalMergingLibrary.cpp
--- a/Engine/Plugins/Runtime/SkeletalMerging/Source/SkeletalMerging/Private/SkeletalMergingLibrary.cpp
+++ b/Engine/Plugins/Runtime/SkeletalMerging/Source/SkeletalMerging/Private/SkeletalMergingLibrary.cpp
@@ -345,7 +345,12 @@ void USkeletalMergingLibrary::AddCurveNames(USkeleton* InSkeleton, const TMap<FN
 	{
 		if (CurveMetaDataPair.Value)
 		{
-			*InSkeleton->GetCurveMetaData(CurveMetaDataPair.Key) = *CurveMetaDataPair.Value;
+			FCurveMetaData& SkeletonCurveMetaData = *InSkeleton->GetCurveMetaData(CurveMetaDataPair.Key);
+			SkeletonCurveMetaData = *CurveMetaDataPair.Value;
+			for (FBoneReference& BoneReference : SkeletonCurveMetaData.LinkedBones)
+			{
+				BoneReference.Initialize(InSkeleton);
+			}
 		}
 	}
 }
diff --git a/Engine/Source/Runtime/Engine/Private/Animation/PoseAsset.cpp b/Engine/Source/Runtime/Engine/Private/Animation/PoseAsset.cpp
--- a/Engine/Source/Runtime/Engine/Private/Animation/PoseAsset.cpp
+++ b/Engine/Source/Runtime/Engine/Private/Animation/PoseAsset.cpp
@@ -18,6 +18,8 @@
 #include "UObject/ObjectSaveContext.h"
 #include "UObject/Package.h"
 #include "UObject/UnrealType.h"
+#include "SkeletonRemappingRegistry.h"
+#include "Animation/SkeletonRemapping.h"
 
 #include UE_INLINE_GENERATED_CPP_BY_NAME(PoseAsset)
 
@@ -400,9 +402,10 @@ void UPoseAsset::GetBaseAnimationPose(FAnimationPoseData& OutAnimationPoseData)
         const int32 TrackNum = PoseContainer.Tracks.Num();
 		BoneIndices.SetNumUninitialized(TrackNum, false);
 
+		const FSkeletonRemapping& SkeletonRemapping = UE::Anim::FSkeletonRemappingRegistry::Get().GetRemapping(GetSkeleton(), RequiredBones.GetSkeletonAsset());
 		for(int32 TrackIndex = 0; TrackIndex < TrackNum; ++TrackIndex)
 		{
-			const int32 SkeletonBoneIndex = PoseContainer.TrackBoneIndices[TrackIndex];
+			const int32 SkeletonBoneIndex = SkeletonRemapping.IsValid() ? SkeletonRemapping.GetTargetSkeletonBoneIndex(PoseContainer.TrackBoneIndices[TrackIndex]) : PoseContainer.TrackBoneIndices[TrackIndex];
 			const FCompactPoseBoneIndex PoseBoneIndex = RequiredBones.GetCompactPoseIndexFromSkeletonIndex(SkeletonBoneIndex);
 			// we add even if it's invalid because we want it to match with track index
 			BoneIndices[TrackIndex].SkeletonBoneIndex = SkeletonBoneIndex;
@@ -589,6 +592,7 @@ bool UPoseAsset::GetAnimationPose(struct FAnimationPoseData& OutAnimationPoseDat
 
 		const int32 TrackNum = PoseContainer.Tracks.Num();
 
+		const FSkeletonRemapping& SkeletonRemapping = UE::Anim::FSkeletonRemappingRegistry::Get().GetRemapping(GetSkeleton(), RequiredBones.GetSkeletonAsset());
 		// Single pose optimized evaluation path - explicitly used by PoseByName Animation Node
 		if (ExtractionContext.PoseCurves.Num() == 1)
 		{
@@ -609,7 +613,7 @@ bool UPoseAsset::GetAnimationPose(struct FAnimationPoseData& OutAnimationPoseDat
 					// Per-track (bone) transform
 					for (int32 TrackIndex = 0; TrackIndex < TrackNum; ++TrackIndex)
 					{
-						const FSkeletonPoseBoneIndex SkeletonBoneIndex = FSkeletonPoseBoneIndex(PoseContainer.TrackBoneIndices[TrackIndex]);
+						const FSkeletonPoseBoneIndex SkeletonBoneIndex = FSkeletonPoseBoneIndex(SkeletonRemapping.IsValid() ? SkeletonRemapping.GetTargetSkeletonBoneIndex(PoseContainer.TrackBoneIndices[TrackIndex]) : PoseContainer.TrackBoneIndices[TrackIndex]);
 
 						UE_CLOG(!RequiredBones.IsSkeletonPoseIndexValid(SkeletonBoneIndex), LogAnimation, Warning, TEXT("PoseAsset [%s] with skeleton [%s] has bones not present in the evaluation container. Bone index(%d) not found."), *GetPathName(), MySkeleton ? *MySkeleton->GetPathName() : TEXT("<Skeleton Not Found>"), SkeletonBoneIndex.GetInt());
 
@@ -673,7 +677,7 @@ bool UPoseAsset::GetAnimationPose(struct FAnimationPoseData& OutAnimationPoseDat
 
 			for(int32 TrackIndex = 0; TrackIndex < TrackNum; ++TrackIndex)
 			{
-				const FSkeletonPoseBoneIndex SkeletonBoneIndex = FSkeletonPoseBoneIndex(PoseContainer.TrackBoneIndices[TrackIndex]);
+				const FSkeletonPoseBoneIndex SkeletonBoneIndex = FSkeletonPoseBoneIndex(SkeletonRemapping.IsValid() ? SkeletonRemapping.GetTargetSkeletonBoneIndex(PoseContainer.TrackBoneIndices[TrackIndex]) : PoseContainer.TrackBoneIndices[TrackIndex]);
 
 				UE_CLOG(!RequiredBones.IsSkeletonPoseIndexValid(SkeletonBoneIndex), LogAnimation, Warning, TEXT("PoseAsset [%s] with skeleton [%s] has bones not present in the evaluation container. Bone index(%d) not found."), *GetPathName(), MySkeleton ? *MySkeleton->GetPathName() : TEXT("<Skeleton Not Found>"), SkeletonBoneIndex.GetInt());
 
