Fix issues with alternate skin profiles and high-precision skin weights

#rnx

diff --git a/Engine/Source/Developer/MeshUtilities/Private/MeshUtilities.cpp b/Engine/Source/Developer/MeshUtilities/Private/MeshUtilities.cpp
--- a/Engine/Source/Developer/MeshUtilities/Private/MeshUtilities.cpp
+++ b/Engine/Source/Developer/MeshUtilities/Private/MeshUtilities.cpp
@@ -5753,7 +5753,11 @@ void FMeshUtilities::CalculateOverlappingCorners(const TArray<FVector3f>& InVert
 }
 
 
-void FMeshUtilities::GenerateRuntimeSkinWeightData(const FSkeletalMeshLODModel* ImportedModel, const TArray<FRawSkinWeight>& InRawSkinWeights, FRuntimeSkinWeightProfileData& InOutSkinWeightOverrideData) const
+void FMeshUtilities::GenerateRuntimeSkinWeightData(
+	const FSkeletalMeshLODModel* ImportedModel,
+	const TArray<FRawSkinWeight>& InRawSkinWeights,
+	bool bInUseHighPrecisionWeights,
+	FRuntimeSkinWeightProfileData& InOutSkinWeightOverrideData) const
 {
 	const FSkeletalMeshLODModel& TargetLODModel = *ImportedModel;
 
@@ -5792,13 +5796,13 @@ void FMeshUtilities::GenerateRuntimeSkinWeightData(const FSkeletalMeshLODModel*
 			if (bIsDifferent)
 			{
 				// Check whether or not there is already an override store which matches the new skin weight data
-				int32 OverrideIndex = UniqueWeights.IndexOfByPredicate([SourceSkinWeight, NumInfluences](const FRawSkinWeight Override)
+				int32 OverrideIndex = UniqueWeights.IndexOfByPredicate([SourceSkinWeight, NumInfluences](const FRawSkinWeight& InOverride)
 				{
 					bool bSame = true;
 					for (int32 InfluenceIndex = 0; InfluenceIndex < NumInfluences; ++InfluenceIndex)
 					{
-						bSame &= (Override.InfluenceBones[InfluenceIndex] == SourceSkinWeight.InfluenceBones[InfluenceIndex]);
-						bSame &= (Override.InfluenceWeights[InfluenceIndex] == SourceSkinWeight.InfluenceWeights[InfluenceIndex]);
+						bSame &= (InOverride.InfluenceBones[InfluenceIndex] == SourceSkinWeight.InfluenceBones[InfluenceIndex]);
+						bSame &= (InOverride.InfluenceWeights[InfluenceIndex] == SourceSkinWeight.InfluenceWeights[InfluenceIndex]);
 					}
 
 					return bSame;
@@ -5813,20 +5817,29 @@ void FMeshUtilities::GenerateRuntimeSkinWeightData(const FSkeletalMeshLODModel*
 					for (int32 InfluenceIndex = 0; InfluenceIndex < NumInfluences; ++InfluenceIndex)
 					{
 						const FBoneIndexType Index = SourceSkinWeight.InfluenceBones[InfluenceIndex];
-						const uint8 Weight = SourceSkinWeight.InfluenceWeights[InfluenceIndex];
+						const uint16 Weight = SourceSkinWeight.InfluenceWeights[InfluenceIndex];
+						TArray<uint8> &BoneIDs = InOutSkinWeightOverrideData.BoneIDs;
+						TArray<uint8> &BoneWeights = InOutSkinWeightOverrideData.BoneWeights;
 
 						if (b16BitBoneIndices)
 						{
-							InOutSkinWeightOverrideData.BoneIDs.AddZeroed(2);
-							FBoneIndexType* BoneIndex = (FBoneIndexType*)&InOutSkinWeightOverrideData.BoneIDs[InOutSkinWeightOverrideData.BoneIDs.Num() - 2];
-							*BoneIndex = Index;
+							BoneIDs.AddZeroed(2);
+							*reinterpret_cast<FBoneIndexType*>(&BoneIDs[BoneIDs.Num() - 2]) = Index;
 						}
 						else
 						{
-							InOutSkinWeightOverrideData.BoneIDs.Add((uint8)Index);
+							InOutSkinWeightOverrideData.BoneIDs.Add(static_cast<uint8>(Index));
 						}
 
-						InOutSkinWeightOverrideData.BoneWeights.Add(Weight);
+						if (bInUseHighPrecisionWeights)
+						{
+							BoneWeights.AddZeroed(2);
+							*reinterpret_cast<uint16*>(&BoneWeights[BoneWeights.Num() - 2]) = Weight;
+						}
+						else
+						{
+							BoneWeights.Add(static_cast<uint8>(Weight >> 8));
+						}
 					}
 
 					UniqueWeights.Add(SourceSkinWeight);
@@ -6011,8 +6024,8 @@ void FMeshUtilities::CreateImportDataFromLODModel(USkeletalMesh* SkeletalMesh) c
 
 				for (int32 InfluenceIndex = 0; InfluenceIndex < MAX_TOTAL_INFLUENCES; ++InfluenceIndex)
 				{
-					float Weight = static_cast<float>(Vertex.InfluenceWeights[InfluenceIndex]) / 65535.0f;
-					if (FMath::IsNearlyZero(Weight))
+					// Have we reached the end of valid weights?
+					if (Vertex.InfluenceWeights[InfluenceIndex] == 0)
 					{
 						break;
 					}
@@ -6111,11 +6124,12 @@ void FMeshUtilities::CreateImportDataFromLODModel(USkeletalMesh* SkeletalMesh) c
 					const FRawSkinWeight& SkinWeight = ImportedSkinWeightProfileData.SkinWeights[VerticeIndex];
 					for (int32 InfluenceIndex = 0; InfluenceIndex < MAX_TOTAL_INFLUENCES; ++InfluenceIndex)
 					{
-						float Weight = static_cast<float>(SkinWeight.InfluenceWeights[InfluenceIndex]) / 65535.0f;
-						if (FMath::IsNearlyZero(Weight))
+						const uint16 RawWeight = SkinWeight.InfluenceWeights[InfluenceIndex]; 
+						if (RawWeight == 0)
 						{
 							break;
 						}
+						const float Weight = static_cast<float>(RawWeight) / 65535.0f;
 						SkeletalMeshImportData::FRawBoneInfluence& Influence = AlternateInfluence.Influences.AddZeroed_GetRef();
 						Influence.VertexIndex = VerticeIndex;
 						Influence.BoneIndex = Section.BoneMap[SkinWeight.InfluenceBones[InfluenceIndex]];
diff --git a/Engine/Source/Developer/MeshUtilities/Private/MeshUtilitiesPrivate.h b/Engine/Source/Developer/MeshUtilities/Private/MeshUtilitiesPrivate.h
--- a/Engine/Source/Developer/MeshUtilities/Private/MeshUtilitiesPrivate.h
+++ b/Engine/Source/Developer/MeshUtilities/Private/MeshUtilitiesPrivate.h
@@ -154,7 +154,7 @@ private:
 	virtual void CalculateNormals(const TArray<FVector3f>& InVertices, const TArray<uint32>& InIndices, const TArray<FVector2f>& InUVs, const TArray<uint32>& InSmoothingGroupIndices, const uint32 InTangentOptions, TArray<FVector3f>& OutNormals) const override;
 	virtual void CalculateOverlappingCorners(const TArray<FVector3f>& InVertices, const TArray<uint32>& InIndices, bool bIgnoreDegenerateTriangles, FOverlappingCorners& OutOverlappingCorners) const override;
 
-	virtual void GenerateRuntimeSkinWeightData(const FSkeletalMeshLODModel* ImportedModel, const TArray<FRawSkinWeight>& InRawSkinWeights, FRuntimeSkinWeightProfileData& InOutSkinWeightOverrideData) const override;
+	virtual void GenerateRuntimeSkinWeightData(const FSkeletalMeshLODModel* ImportedModel, const TArray<FRawSkinWeight>& InRawSkinWeights, bool bInUseHighPrecisionWeights, FRuntimeSkinWeightProfileData& InOutSkinWeightOverrideData) const override;
 
 	virtual void CreateImportDataFromLODModel(USkeletalMesh* SkeletalMesh) const override;
 
diff --git a/Engine/Source/Developer/MeshUtilities/Public/MeshUtilities.h b/Engine/Source/Developer/MeshUtilities/Public/MeshUtilities.h
--- a/Engine/Source/Developer/MeshUtilities/Public/MeshUtilities.h
+++ b/Engine/Source/Developer/MeshUtilities/Public/MeshUtilities.h
@@ -333,7 +333,11 @@ public:
 	virtual void FindOverlappingCorners(FOverlappingCorners& OutOverlappingCorners, const TArray<FVector3f>& InVertices, const TArray<uint32>& InIndices, float ComparisonThreshold) const = 0;
 
 	/** Used to generate runtime skin weight data from Editor-only data */
-	virtual void GenerateRuntimeSkinWeightData(const FSkeletalMeshLODModel* ImportedModel, const TArray<FRawSkinWeight>& InRawSkinWeights, struct FRuntimeSkinWeightProfileData& InOutSkinWeightOverrideData) const = 0;
+	virtual void GenerateRuntimeSkinWeightData(
+		const FSkeletalMeshLODModel* ImportedModel,
+		const TArray<FRawSkinWeight>& InRawSkinWeights,
+		bool bInUseHighPrecisionWeights,
+		struct FRuntimeSkinWeightProfileData& InOutSkinWeightOverrideData) const = 0;
 
 	/*
 	 * This function create the import data using the LODModel. You can call this function if you load an asset that was not re-import since the build refactor and the chunking is more agressive than the bake data in the LODModel.
diff --git a/Engine/Source/Runtime/AnimationCore/Public/BoneWeights.h b/Engine/Source/Runtime/AnimationCore/Public/BoneWeights.h
--- a/Engine/Source/Runtime/AnimationCore/Public/BoneWeights.h
+++ b/Engine/Source/Runtime/AnimationCore/Public/BoneWeights.h
@@ -252,9 +252,10 @@ public:
 	/** Returns the maximum number of weights for these settings */
 	int32 GetMaxWeightCount() const { return MaxWeightCount; }
 
-	/** Sets the minimum influence allowed. Any bone weight value below this limit will be
-	  * ignored. The threshold value is clamped to the half-closed interval (0, 1] since
+	/** Sets the minimum influence allowed. Any bone weight value below this threshold value is
+	  * discarded. The threshold value is clamped to the half-closed interval (0, 1] since
 	  * weight values of zero indicate no influence at all and are ignored completely.
+	  * By default, all non-zero weights are allowed.
 	  */
 	FBoneWeightsSettings& SetWeightThreshold(float InWeightThreshold)
 	{
@@ -327,7 +328,7 @@ public:
 private:
 	EBoneWeightNormalizeType NormalizeType = EBoneWeightNormalizeType::Always;
 	int32 MaxWeightCount = MaxInlineBoneWeightCount;
-	uint16 WeightThreshold = 257; // == uint8(1) converted to 16 bit using 'v | v << 8'.
+	uint16 WeightThreshold = 1; 
 	FBoneIndexType DefaultBoneIndex = static_cast<FBoneIndexType>(0);
 	bool bHasDefaultBoneIndex = false;
 	bool bBlendZeroInfluence = true; 
diff --git a/Engine/Source/Runtime/Engine/Private/Animation/SkinWeightProfile.cpp b/Engine/Source/Runtime/Engine/Private/Animation/SkinWeightProfile.cpp
--- a/Engine/Source/Runtime/Engine/Private/Animation/SkinWeightProfile.cpp
+++ b/Engine/Source/Runtime/Engine/Private/Animation/SkinWeightProfile.cpp
@@ -345,7 +345,7 @@ void FSkinWeightProfilesData::SetDynamicDefaultSkinWeightProfile(USkeletalMesh*
 					if (bSerialization)
 					{
 						// During serialization the CPU copy of the weight should still be available
-						const void* BaseBufferData = BaseBuffer->GetDataVertexBuffer()->GetWeightData();
+						const uint8* BaseBufferData = BaseBuffer->GetDataVertexBuffer()->GetWeightData();
 						
 						if (ensure(BaseBufferData))
 						{
@@ -361,11 +361,6 @@ void FSkinWeightProfilesData::SetDynamicDefaultSkinWeightProfile(USkeletalMesh*
 								OverrideBuffer->CopyMetaData(*BaseBuffer);
 								ProfileNameToBuffer.Add(ProfileName, OverrideBuffer);
 								
-								const uint32 NumInfluences = BaseBuffer->GetMaxBoneInfluences();
-								OverrideBuffer->SetMaxBoneInfluences(NumInfluences);
-								const bool bUse16BitBoneIndex = BaseBuffer->Use16BitBoneIndex();
-								OverrideBuffer->SetUse16BitBoneIndex(bUse16BitBoneIndex);
-								
 								ProfilePtr->ApplyOverrides(OverrideBuffer, BaseBufferData, BaseBuffer->GetNumVertices());
 
 								DefaultOverrideSkinWeightBuffer = OverrideBuffer;
@@ -714,7 +709,8 @@ void FSkinWeightProfilesData::InitialiseProfileBuffer(const FName& ProfileName)
 
 	if (BaseBuffer)
 	{
-		const void* BaseBufferData = nullptr;
+		const uint8* BaseBufferData;
+		
 		if (BaseBuffer->GetNeedsCPUAccess())
 		{
 			BaseBufferData = BaseBuffer->GetDataVertexBuffer()->GetWeightData();
@@ -727,18 +723,11 @@ void FSkinWeightProfilesData::InitialiseProfileBuffer(const FName& ProfileName)
 		
 		if (ensure(BaseBufferData))
 		{
-			FSkinWeightVertexBuffer* OverrideBuffer = nullptr;
-
-			OverrideBuffer = new FSkinWeightVertexBuffer();
+			FSkinWeightVertexBuffer* OverrideBuffer = new FSkinWeightVertexBuffer();
 			OverrideBuffer->CopyMetaData(*BaseBuffer);
 			ProfileNameToBuffer.Add(ProfileName, OverrideBuffer);
 
-			const uint32 NumInfluences = BaseBuffer->GetMaxBoneInfluences();
-			OverrideBuffer->SetMaxBoneInfluences(NumInfluences);
-			const bool bUse16BitBoneIndex = BaseBuffer->Use16BitBoneIndex();
-			OverrideBuffer->SetUse16BitBoneIndex(bUse16BitBoneIndex);
-			const bool bNeedsCPUAccess = BaseBuffer->GetNeedsCPUAccess();
-			OverrideBuffer->SetNeedsCPUAccess(bNeedsCPUAccess);
+			OverrideBuffer->SetNeedsCPUAccess(BaseBuffer->GetNeedsCPUAccess());
 			
 			const FRuntimeSkinWeightProfileData* ProfilePtr = OverrideData.Find(ProfileName);
 			if (ProfilePtr)
@@ -757,10 +746,7 @@ void FSkinWeightProfilesData::InitialiseProfileBuffer(const FName& ProfileName)
 
 void FSkinWeightProfilesData::ApplyOverrideProfile(FSkinWeightVertexBuffer* OverrideBuffer, const FName& ProfileName)
 {
-	const uint32 NumInfluences = BaseBuffer->GetMaxBoneInfluences();
-	OverrideBuffer->SetMaxBoneInfluences(NumInfluences);
-	const bool bUse16BitBoneIndex = BaseBuffer->Use16BitBoneIndex();
-	OverrideBuffer->SetUse16BitBoneIndex(bUse16BitBoneIndex);
+	OverrideBuffer->CopyMetaData(*BaseBuffer);
 
 	const FRuntimeSkinWeightProfileData* ProfilePtr = OverrideData.Find(ProfileName);
 	if (ProfilePtr)
@@ -797,16 +783,15 @@ void FSkinWeightProfilesData::ReleaseRHIForStreaming(FRHIResourceUpdateBatcher&
 	}
 }
 
-void FRuntimeSkinWeightProfileData::ApplyOverrides(FSkinWeightVertexBuffer* OverrideBuffer, const void* DataBuffer, const int32 NumVerts) const
+void FRuntimeSkinWeightProfileData::ApplyOverrides(FSkinWeightVertexBuffer* OverrideBuffer, const uint8* DataBuffer, const int32 NumVerts) const
 {
 	if (DataBuffer)
 	{
 		if (OverrideBuffer)
 		{
-			const FSkinWeightInfo* BaseSkinWeights = (const FSkinWeightInfo*)DataBuffer;
-			OverrideBuffer->CopySkinWeightInfoData(TArrayView<const FSkinWeightInfo>(BaseSkinWeights, NumVerts));
+			OverrideBuffer->CopySkinWeightRawDataFromBuffer(DataBuffer, NumVerts);
 
-			uint8* TargetSkinWeightData = (uint8*)OverrideBuffer->GetDataVertexBuffer()->GetWeightData();
+			uint8* TargetSkinWeightData = OverrideBuffer->GetDataVertexBuffer()->GetWeightData();
 			
 			const uint8 VertexStride = OverrideBuffer->GetConstantInfluencesVertexStride();
 			const uint8 WeightDataOffset = OverrideBuffer->GetBoneIndexByteSize() * OverrideBuffer->GetMaxBoneInfluences();
@@ -817,8 +802,8 @@ void FRuntimeSkinWeightProfileData::ApplyOverrides(FSkinWeightVertexBuffer* Over
 				const uint32 VertexIndex = VertexIndexOverridePair.Key;
 				const uint32 InfluenceOffset = VertexIndexOverridePair.Value;
 				
-				const uint8* BoneData = TargetSkinWeightData + (VertexIndex * VertexStride);
-				const uint8* WeightData = BoneData + WeightDataOffset;
+				uint8* BoneData = TargetSkinWeightData + (VertexIndex * VertexStride);
+				uint8* WeightData = BoneData + WeightDataOffset;
 
 #if !UE_BUILD_SHIPPING
 				uint32 VertexOffset = 0;
@@ -829,8 +814,10 @@ void FRuntimeSkinWeightProfileData::ApplyOverrides(FSkinWeightVertexBuffer* Over
 				check(b16BitBoneIndices == OverrideBuffer->Use16BitBoneIndex());
 #endif
 				// BoneIDs either contains FBoneIndexType entries spanning (2) uint8 values, or single uint8 bone indices (1)
-				FMemory::Memcpy((void*)BoneData, &BoneIDs[InfluenceOffset * NumWeightsPerVertex * OverrideBuffer->GetBoneIndexByteSize()], OverrideBuffer->GetBoneIndexByteSize() * NumWeightsPerVertex);
-				FMemory::Memcpy((void*)WeightData, &BoneWeights[InfluenceOffset * NumWeightsPerVertex], sizeof(uint8) * NumWeightsPerVertex);
+				const uint32 BoneIndexByteSize = OverrideBuffer->GetBoneIndexByteSize();
+				const uint32 BoneWeightByteSize = OverrideBuffer->GetBoneWeightByteSize();
+				FMemory::Memcpy(BoneData, &BoneIDs[InfluenceOffset * NumWeightsPerVertex * BoneIndexByteSize], BoneIndexByteSize * NumWeightsPerVertex);
+				FMemory::Memcpy(WeightData, &BoneWeights[InfluenceOffset * NumWeightsPerVertex * BoneWeightByteSize], BoneWeightByteSize * NumWeightsPerVertex);
 			}
 		}
 	}
diff --git a/Engine/Source/Runtime/Engine/Private/Rendering/SkinWeightVertexBuffer.cpp b/Engine/Source/Runtime/Engine/Private/Rendering/SkinWeightVertexBuffer.cpp
--- a/Engine/Source/Runtime/Engine/Private/Rendering/SkinWeightVertexBuffer.cpp
+++ b/Engine/Source/Runtime/Engine/Private/Rendering/SkinWeightVertexBuffer.cpp
@@ -622,10 +622,10 @@ void FSkinWeightDataVertexBuffer::ResetVertexBoneWeights(uint32 VertexWeightOffs
 	}
 }
 
-void FSkinWeightDataVertexBuffer::CopyDataFromBuffer(const TArrayView<const FSkinWeightInfo>& SkinWeightData)
+void FSkinWeightDataVertexBuffer::CopyDataFromBuffer(const uint8* InSkinWeightData, uint32 InNumVertices)
 {
-	Init(SkinWeightData.Num() * GetMaxBoneInfluences(), SkinWeightData.Num());
-	FMemory::Memcpy(Data, SkinWeightData.GetData(), GetVertexDataSize());
+	Init(InNumVertices * GetMaxBoneInfluences(), InNumVertices);
+	FMemory::Memcpy(Data, InSkinWeightData, GetVertexDataSize());
 }
 
 /*-----------------------------------------------------------------------------
@@ -948,7 +948,7 @@ FSkinWeightInfo FSkinWeightVertexBuffer::GetVertexSkinWeights(uint32 VertexIndex
 	}
 	return OutVertex;
 }
-void FSkinWeightVertexBuffer::CopySkinWeightInfoData(const TArrayView<const FSkinWeightInfo>& SkinWeightData)
+void FSkinWeightVertexBuffer::CopySkinWeightRawDataFromBuffer(const uint8* InSkinWeightData, uint32 InNumVertices)
 {
-	DataVertexBuffer.CopyDataFromBuffer(SkinWeightData);
+	DataVertexBuffer.CopyDataFromBuffer(InSkinWeightData, InNumVertices);
 }
diff --git a/Engine/Source/Runtime/Engine/Private/SkeletalMeshLODRenderData.cpp b/Engine/Source/Runtime/Engine/Private/SkeletalMeshLODRenderData.cpp
--- a/Engine/Source/Runtime/Engine/Private/SkeletalMeshLODRenderData.cpp
+++ b/Engine/Source/Runtime/Engine/Private/SkeletalMeshLODRenderData.cpp
@@ -476,7 +476,7 @@ void FSkeletalMeshLODRenderData::BuildFromLODModel(const FSkeletalMeshLODModel*
 	for (const auto& Pair : ImportedModel->SkinWeightProfiles)
 	{
 		FRuntimeSkinWeightProfileData& Override = SkinWeightProfilesData.AddOverrideData(Pair.Key);
-		MeshUtilities.GenerateRuntimeSkinWeightData(ImportedModel, Pair.Value.SkinWeights, Override);
+		MeshUtilities.GenerateRuntimeSkinWeightData(ImportedModel, Pair.Value.SkinWeights, bUseHighPrecisionWeights, Override);
 	}
 
 	ActiveBoneIndices = ImportedModel->ActiveBoneIndices;
diff --git a/Engine/Source/Runtime/Engine/Public/Animation/SkinWeightProfile.h b/Engine/Source/Runtime/Engine/Public/Animation/SkinWeightProfile.h
--- a/Engine/Source/Runtime/Engine/Public/Animation/SkinWeightProfile.h
+++ b/Engine/Source/Runtime/Engine/Public/Animation/SkinWeightProfile.h
@@ -95,7 +95,7 @@ struct FRuntimeSkinWeightProfileData
 		friend FArchive& operator<<(FArchive& Ar, FSkinWeightOverrideInfo& OverrideInfo);
 	};
 
-	void ApplyOverrides(FSkinWeightVertexBuffer* OverrideBuffer, const void* DataBuffer, const int32 NumVerts) const;	
+	void ApplyOverrides(FSkinWeightVertexBuffer* OverrideBuffer, const uint8* DataBuffer, const int32 NumVerts) const;	
 	void ApplyDefaultOverride(FSkinWeightVertexBuffer* Buffer) const;
 
 #if WITH_EDITORONLY_DATA
diff --git a/Engine/Source/Runtime/Engine/Public/Rendering/SkinWeightVertexBuffer.h b/Engine/Source/Runtime/Engine/Public/Rendering/SkinWeightVertexBuffer.h
--- a/Engine/Source/Runtime/Engine/Public/Rendering/SkinWeightVertexBuffer.h
+++ b/Engine/Source/Runtime/Engine/Public/Rendering/SkinWeightVertexBuffer.h
@@ -297,7 +297,7 @@ public:
 	ENGINE_API void SetBoneWeight(uint32 VertexWeightOffset, uint32 VertexInfluenceCount, uint32 InfluenceIndex, uint16 BoneWeight);
 	ENGINE_API void ResetVertexBoneWeights(uint32 VertexWeightOffset, uint32 VertexInfluenceCount);
 
-	ENGINE_API void CopyDataFromBuffer(const TArrayView<const FSkinWeightInfo>& SkinWeightData);
+	ENGINE_API void CopyDataFromBuffer(const uint8* InSkinWeightData, uint32 InNumVertices);
 
 	/** Create an RHI vertex buffer with CPU data. CPU data may be discarded after creation (see TResourceArray::Discard) */
 	FBufferRHIRef CreateRHIBuffer_RenderThread();
@@ -394,7 +394,7 @@ public:
 	ENGINE_API void GetSkinWeights(TArray<FSkinWeightInfo>& OutVertices) const;
 	ENGINE_API FSkinWeightInfo GetVertexSkinWeights(uint32 VertexIndex) const;
 
-	ENGINE_API void CopySkinWeightInfoData(const TArrayView<const FSkinWeightInfo>& SkinWeightData);
+	ENGINE_API void CopySkinWeightRawDataFromBuffer(const uint8* InSkinWeightData, uint32 InNumVertices);
 
 	friend FArchive& operator<<(FArchive& Ar, FSkinWeightVertexBuffer& VertexBuffer);
 
@@ -454,6 +454,9 @@ public:
 	uint32 GetConstantInfluencesBoneWeightsOffset() const
 	{ return DataVertexBuffer.GetConstantInfluencesBoneWeightsOffset(); }
 
+	FSkinWeightDataVertexBuffer* GetDataVertexBuffer()
+	{ return &DataVertexBuffer; }
+	
 	const FSkinWeightDataVertexBuffer* GetDataVertexBuffer() const
 	{ return &DataVertexBuffer; }
 
