Add insight tags on long world partition streaming operations

#rnx

diff --git a/Engine/Source/Runtime/Engine/Private/WorldPartition/WorldPartitionStreamingGeneration.cpp b/Engine/Source/Runtime/Engine/Private/WorldPartition/WorldPartitionStreamingGeneration.cpp
--- a/Engine/Source/Runtime/Engine/Private/WorldPartition/WorldPartitionStreamingGeneration.cpp
+++ b/Engine/Source/Runtime/Engine/Private/WorldPartition/WorldPartitionStreamingGeneration.cpp
@@ -446,6 +446,8 @@ class FWorldPartitionStreamingGenerator
 	 */
 	void CreateActorContainers(const UActorDescContainer* InContainer)
 	{
+		TRACE_CPUPROFILER_EVENT_SCOPE(FWorldPartitionStreamingGenerator::CreateActorContainers);
+
 		CreateActorDescriptorViewsRecursive(InContainer, FTransform::Identity, NAME_None, false, TSet<FName>(), FActorContainerID(), FActorContainerID(), EContainerClusterMode::Partitioned, TEXT("MainContainer"));
 
 		// Update container instances bounds
@@ -841,6 +843,8 @@ public:
 
 	void DumpStateLog(FHierarchicalLogArchive& Ar)
 	{
+		TRACE_CPUPROFILER_EVENT_SCOPE(FWorldPartitionStreamingGenerator::DumpStateLog);
+
 		// Build the containers tree representation
 		TMultiMap<FActorContainerID, FActorContainerID> InvertedContainersHierarchy;
 		for (auto& [ChildContainerID, ParentContainerID] : ContainersHierarchy)
@@ -957,6 +961,8 @@ bool UWorldPartition::GenerateStreaming(TArray<FString>* OutPackagesToGenerate)
 
 bool UWorldPartition::GenerateContainerStreaming(const UActorDescContainer* InActorDescContainer, TArray<FString>* OutPackagesToGenerate /* = nullptr */)
 {
+	TRACE_CPUPROFILER_EVENT_SCOPE(UWorldPartition::GenerateContainerStreaming);
+
 	FActorDescList* ModifiedActorsDescList = nullptr;
 
 	FStreamingGenerationLogErrorHandler LogErrorHandler;
