World Partition Editor: auto focus on bugitgo region when loading is enabled.

diff --git a/Engine/Source/Editor/WorldPartitionEditor/Private/WorldPartition/SWorldPartitionEditor.cpp b/Engine/Source/Editor/WorldPartitionEditor/Private/WorldPartition/SWorldPartitionEditor.cpp
--- a/Engine/Source/Editor/WorldPartitionEditor/Private/WorldPartition/SWorldPartitionEditor.cpp
+++ b/Engine/Source/Editor/WorldPartitionEditor/Private/WorldPartition/SWorldPartitionEditor.cpp
@@ -60,6 +60,11 @@ void SWorldPartitionEditor::Reconstruct()
 	ContentParent->SetContent(ConstructContentWidget());
 }
 
+void SWorldPartitionEditor::FocusBox(const FBox& Box) const
+{
+	GridView->FocusBox(Box);
+}
+
 void SWorldPartitionEditor::OnBrowseWorld(UWorld* InWorld)
 {
 	World = InWorld;
diff --git a/Engine/Source/Editor/WorldPartitionEditor/Public/WorldPartition/SWorldPartitionEditor.h b/Engine/Source/Editor/WorldPartitionEditor/Public/WorldPartition/SWorldPartitionEditor.h
--- a/Engine/Source/Editor/WorldPartitionEditor/Public/WorldPartition/SWorldPartitionEditor.h
+++ b/Engine/Source/Editor/WorldPartitionEditor/Public/WorldPartition/SWorldPartitionEditor.h
@@ -23,6 +23,7 @@ public:
 	// IWorldPartitionEditor interface
 	virtual void Refresh() override;
 	virtual void Reconstruct() override;
+	virtual void FocusBox(const FBox& Box) const override;
 
 private:
 	void OnBrowseWorld(UWorld* InWorld);
diff --git a/Engine/Source/Editor/WorldPartitionEditor/Public/WorldPartition/SWorldPartitionEditorGrid.h b/Engine/Source/Editor/WorldPartitionEditor/Public/WorldPartition/SWorldPartitionEditorGrid.h
--- a/Engine/Source/Editor/WorldPartitionEditor/Public/WorldPartition/SWorldPartitionEditorGrid.h
+++ b/Engine/Source/Editor/WorldPartitionEditor/Public/WorldPartition/SWorldPartitionEditorGrid.h
@@ -64,6 +64,7 @@ public:
 	bool GetObserverView(FVector& Location, FRotator& Rotation) const;
 
 	void Refresh();
+	virtual void FocusBox(const FBox& Box) const {}
 
 	UWorld* World = nullptr;
 	UWorldPartition* WorldPartition = nullptr;
diff --git a/Engine/Source/Editor/WorldPartitionEditor/Public/WorldPartition/SWorldPartitionEditorGrid2D.h b/Engine/Source/Editor/WorldPartitionEditor/Public/WorldPartition/SWorldPartitionEditorGrid2D.h
--- a/Engine/Source/Editor/WorldPartitionEditor/Public/WorldPartition/SWorldPartitionEditorGrid2D.h
+++ b/Engine/Source/Editor/WorldPartitionEditor/Public/WorldPartition/SWorldPartitionEditorGrid2D.h
@@ -98,7 +98,7 @@ protected:
 		ArrangedChildren.AddWidget(AllottedGeometry.MakeChild(ChildSlot.GetWidget(), FVector2D::ZeroVector, AllottedGeometry.GetLocalSize()));
 	}
 
-	void FocusBox(const FBox& Box) const;
+	virtual void FocusBox(const FBox& Box) const override;
 
 	mutable float Scale;
 	mutable FVector2D Trans;
diff --git a/Engine/Source/Runtime/Engine/Private/WorldPartition/WorldPartition.cpp b/Engine/Source/Runtime/Engine/Private/WorldPartition/WorldPartition.cpp
--- a/Engine/Source/Runtime/Engine/Private/WorldPartition/WorldPartition.cpp
+++ b/Engine/Source/Runtime/Engine/Private/WorldPartition/WorldPartition.cpp
@@ -792,15 +792,20 @@ void UWorldPartition::OnPostBugItGoCalled(const FVector& Loc, const FRotator& Ro
 #if WITH_EDITOR
 	if (GetMutableDefault<UWorldPartitionEditorPerProjectUserSettings>()->GetBugItGoLoadRegion())
 	{
+		const FVector LoadExtent(UWorldPartition::LoadingRangeBugItGo, UWorldPartition::LoadingRangeBugItGo, HALF_WORLD_MAX);
+		const FBox LoadCellsBox(Loc - LoadExtent, Loc + LoadExtent);
+
 		IWorldPartitionEditorModule& WorldPartitionEditorModule = FModuleManager::LoadModuleChecked<IWorldPartitionEditorModule>("WorldPartitionEditor");
 		if (!WorldPartitionEditorModule.GetDisableLoadingInEditor())
 		{
-			const FVector LoadExtent(UWorldPartition::LoadingRangeBugItGo, UWorldPartition::LoadingRangeBugItGo, HALF_WORLD_MAX);
-			const FBox LoadCellsBox(Loc - LoadExtent, Loc + LoadExtent);
-
 			UWorldPartitionEditorLoaderAdapter* EditorLoaderAdapter = CreateEditorLoaderAdapter<FLoaderAdapterShape>(World, LoadCellsBox, TEXT("BugItGo"));
 			EditorLoaderAdapter->GetLoaderAdapter()->Load();
 		}
+
+		if (WorldPartitionEditor)
+		{
+			WorldPartitionEditor->FocusBox(LoadCellsBox);
+		}
 	}
 #endif
 }
diff --git a/Engine/Source/Runtime/Engine/Public/WorldPartition/WorldPartition.h b/Engine/Source/Runtime/Engine/Public/WorldPartition/WorldPartition.h
--- a/Engine/Source/Runtime/Engine/Public/WorldPartition/WorldPartition.h
+++ b/Engine/Source/Runtime/Engine/Public/WorldPartition/WorldPartition.h
@@ -62,6 +62,7 @@ struct ENGINE_API IWorldPartitionEditor
 {
 	virtual void Refresh() {}
 	virtual void Reconstruct() {}
+	virtual void FocusBox(const FBox& Box) const {}
 };
 
 class ENGINE_API ISourceControlHelper
