Allows for a register content bundle client to inject its content into a targetet world.
Content Bundle for the targeted world becomes active. For other world they remain registered
Support editing content bundle active while a plugin is registered by adding its content into the editor world.

#test edit, pie, cook

diff --git a/Engine/Source/Runtime/Engine/Private/WorldPartition/ContentBundle/ContentBundleClient.cpp b/Engine/Source/Runtime/Engine/Private/WorldPartition/ContentBundle/ContentBundleClient.cpp
--- a/Engine/Source/Runtime/Engine/Private/WorldPartition/ContentBundle/ContentBundleClient.cpp
+++ b/Engine/Source/Runtime/Engine/Private/WorldPartition/ContentBundle/ContentBundleClient.cpp
@@ -6,6 +6,7 @@
 #include "WorldPartition/ContentBundle/ContentBundleLog.h"
 #include "Engine/Engine.h"
 #include "WorldPartition/ContentBundle/ContentBundleStatus.h"
+#include "WorldPartition/ContentBundle/ContentBundleWorldSubsystem.h"
 
 TSharedPtr<FContentBundleClient> FContentBundleClient::CreateClient(const UContentBundleDescriptor* InContentBundleDescriptor, FString const& InDisplayName)
 {
@@ -38,7 +39,7 @@ void FContentBundleClient::RequestRemoveContent()
 {
 	if (State == EContentBundleClientState::ContentInjectionRequested)
 	{
-		if (HasInjectedAnyContent())
+		if (HasContentToRemove())
 		{
 			SetState(EContentBundleClientState::ContentRemovalRequested);
 			GEngine->GetEngineSubsystem<UContentBundleEngineSubsystem>()->RequestContentRemoval(*this);
@@ -69,60 +70,60 @@ void FContentBundleClient::RequestUnregister()
 	}
 }
 
-bool FContentBundleClient::HasInjectedAnyContent() const
+bool FContentBundleClient::ShouldInjectContent(UWorld* World) const
 {
-	for (auto& WorldState : WorldContentStates)
-	{
-		if (WorldState.Value != EWorldContentState::NoContent)
-		{
-			return true;
-		}
-	}
+	bool bIsClientInjecting = GetState() == EContentBundleClientState::ContentInjectionRequested;
 
-	return false;
+#if WITH_EDITOR
+	bool bIsClientForceInjectingWorld = GetState() != EContentBundleClientState::Unregistered && ForceInjectedWorlds.Contains(World);
+#else
+	bool bIsClientForceInjectingWorld = false;
+#endif
+	
+	return (bIsClientInjecting || bIsClientForceInjectingWorld) && GetWorldContentState(World) == EWorldContentState::NoContent;
+}
+
+bool FContentBundleClient::ShouldRemoveContent(UWorld* World) const
+{
+	bool bIsClientInjecting = GetState() == EContentBundleClientState::ContentInjectionRequested;
+
+#if WITH_EDITOR
+	bool bIsClientForceInjectingWorld = GetState() != EContentBundleClientState::Unregistered && ForceInjectedWorlds.Contains(World);
+#else
+	bool bIsClientForceInjectingWorld = false;
+#endif
+
+	return !bIsClientInjecting && !bIsClientForceInjectingWorld && GetWorldContentState(World) != EWorldContentState::NoContent;
 }
 
 void FContentBundleClient::OnContentInjectedInWorld(EContentBundleStatus InjectionStatus, UWorld* InjectedWorld)
 {
-	check(State == EContentBundleClientState::ContentInjectionRequested);
+	check(ShouldInjectContent(InjectedWorld));
+	check(InjectionStatus == EContentBundleStatus::ContentInjected
+		|| InjectionStatus == EContentBundleStatus::ReadyToInject
+		|| InjectionStatus == EContentBundleStatus::FailedToInject);
 
-	if (InjectionStatus == EContentBundleStatus::ContentInjected || InjectionStatus == EContentBundleStatus::ReadyToInject)
-	{
-		SetWorldContentState(InjectedWorld, EWorldContentState::ContentBundleInjected);
-	}
-	else if (InjectionStatus == EContentBundleStatus::FailedToInject)
-	{
-		SetWorldContentState(InjectedWorld, EWorldContentState::ContentBundleInjected);
-	}
-	else
-	{
-		// Injection status unhandled
-		check(0);
-	}
+	SetWorldContentState(InjectedWorld, EWorldContentState::ContentBundleInjected);
 }
 
 void FContentBundleClient::OnContentRemovedFromWorld(EContentBundleStatus RemovalStatus, UWorld* InjectedWorld)
 {
-	if (RemovalStatus == EContentBundleStatus::Registered)
-	{
-		SetWorldContentState(InjectedWorld, EWorldContentState::NoContent);
+	// Removal of content bundle can come from the client or WorldUnload
+	check(ShouldRemoveContent(InjectedWorld) || !InjectedWorld->ContentBundleManager->CanInject());
+	check(RemovalStatus == EContentBundleStatus::Registered);
 
-		if (!HasInjectedAnyContent())
-		{
-			WorldContentStates.Empty();
-
-			// If the client did not request content removal but removed everything keep the client state. We just changed levels. The next level will have its content bundle injected.
-			// If the client request content removal and removed everything set the client state as registered.
-			if (GetState() == EContentBundleClientState::ContentRemovalRequested)
-			{
-				SetState(EContentBundleClientState::Registered);
-			}
-		}
-	}
-	else
+	WorldContentStates.Remove(InjectedWorld);
+
+#if WITH_EDITOR
+	ForceInjectedWorlds.Remove(InjectedWorld);
+#endif
+
+	// If the client request content removal and removed everything set the client state as registered.
+	// If the client did not request content removal but removed everything keep the client state. We just changed levels. The next level will have its content bundle injected.
+	bool bHasProcessAllRemoveRequest = !HasContentToRemove() && GetState() == EContentBundleClientState::ContentRemovalRequested;
+	if (bHasProcessAllRemoveRequest)
 	{
-		// Removal status unhandled
-		check(0);
+		SetState(EContentBundleClientState::Registered);
 	}
 
 }
@@ -140,6 +141,37 @@ void FContentBundleClient::SetWorldContentState(UWorld* World, EWorldContentStat
 	WorldContentStates.Add(World, NewState);
 }
 
+EWorldContentState FContentBundleClient::GetWorldContentState(UWorld* World) const
+{
+	const EWorldContentState* WorldState = WorldContentStates.Find(World);
+	if (WorldState != nullptr)
+	{
+		return *WorldState;
+	}
+
+	return EWorldContentState::NoContent;
+}
+
+bool FContentBundleClient::HasContentToRemove() const
+{
+	for (auto& WorldState : WorldContentStates)
+	{
+#if WITH_EDITOR
+		if (ForceInjectedWorlds.Contains(WorldState.Key))
+		{
+			continue;
+		}
+#endif
+
+		if (WorldState.Value != EWorldContentState::NoContent)
+		{
+			return true;
+		}
+	}
+
+	return false;
+}
+
 void FContentBundleClient::SetState(EContentBundleClientState NewState)
 {
 	check(NewState != State);
@@ -147,3 +179,33 @@ void FContentBundleClient::SetState(EContentBundleClientState NewState)
 	UE_LOG(LogContentBundle, Log, TEXT("[CB: %s] Client State changing from %s to %s"), *GetDescriptor()->GetDisplayName(), *UEnum::GetDisplayValueAsText(State).ToString(), *UEnum::GetDisplayValueAsText(NewState).ToString());
 	State = NewState;
 }
+
+#if WITH_EDITOR
+
+void FContentBundleClient::RequestForceInject(UWorld* WorldToInject)
+{
+	if (UContentBundleManager* ContentBundleManager = WorldToInject->ContentBundleManager)
+	{
+		ForceInjectedWorlds.Add(WorldToInject);
+	
+		if (ShouldInjectContent(WorldToInject))
+		{
+			ContentBundleManager->TryInject(*this);
+		}
+	}
+}
+
+void FContentBundleClient::RequestRemoveForceInjectedContent(UWorld* WorldToInject)
+{
+	if (UContentBundleManager* ContentBundleManager = WorldToInject->ContentBundleManager)
+	{
+		ForceInjectedWorlds.Remove(WorldToInject);
+
+		if (ShouldRemoveContent(WorldToInject))
+		{
+			ContentBundleManager->Remove(*this);
+		}
+	}
+}
+
+#endif
diff --git a/Engine/Source/Runtime/Engine/Private/WorldPartition/ContentBundle/ContentBundleContainer.cpp b/Engine/Source/Runtime/Engine/Private/WorldPartition/ContentBundle/ContentBundleContainer.cpp
--- a/Engine/Source/Runtime/Engine/Private/WorldPartition/ContentBundle/ContentBundleContainer.cpp
+++ b/Engine/Source/Runtime/Engine/Private/WorldPartition/ContentBundle/ContentBundleContainer.cpp
@@ -266,14 +266,54 @@ void FContentBundleContainer::DeinitializeContentBundle(FContentBundleBase& Cont
 	GetGameContentBundles().RemoveAtSwap(Index);
 }
 
-void FContentBundleContainer::InjectContentBundle(FContentBundleBase& ContentBundle)
+bool FContentBundleContainer::InjectContentBundle(FContentBundleClient& ContentBundleClient)
 {
-	ContentBundle.InjectContent();
+	if (FContentBundleBase* ContentBundle = GetContentBundle(ContentBundleClient))
+	{
+		return InjectContentBundle(*ContentBundle);
+	}
+
+	return false;
+}
+
+bool FContentBundleContainer::InjectContentBundle(FContentBundleBase& ContentBundle)
+{
+	if (ContentBundle.GetStatus() == EContentBundleStatus::Registered)
+	{
+		TSharedPtr<FContentBundleClient> ContentBundleClient = ContentBundle.GetClient().Pin();
+		if (ContentBundleClient != nullptr && ContentBundleClient->ShouldInjectContent(InjectedWorld))
+		{
+			ContentBundle.InjectContent();
+			return true;
+		}
+	}
+	
+	return false;
+}
+
+bool FContentBundleContainer::RemoveContentBundle(FContentBundleClient& ContentBundleClient)
+{
+	if (FContentBundleBase* ContentBundle = GetContentBundle(ContentBundleClient))
+	{
+		return RemoveContentBundle(*ContentBundle);
+	}
+
+	return false;
 }
 
-void FContentBundleContainer::RemoveContentBundle(FContentBundleBase& ContentBundle)
+bool FContentBundleContainer::RemoveContentBundle(FContentBundleBase& ContentBundle)
 {
-	ContentBundle.RemoveContent();
+	if ((ContentBundle.GetStatus() == EContentBundleStatus::ContentInjected || ContentBundle.GetStatus() == EContentBundleStatus::ReadyToInject))
+	{
+		TSharedPtr<FContentBundleClient> ContentBundleClient = ContentBundle.GetClient().Pin();
+		if (ContentBundleClient == nullptr || ContentBundleClient->ShouldRemoveContent(InjectedWorld) || !InjectedWorld->ContentBundleManager->CanInject())
+		{
+			ContentBundle.RemoveContent();
+			return true;
+		}
+	}
+	
+	return false;
 }
 
 void FContentBundleContainer::OnContentBundleClientRegistered(TSharedPtr<FContentBundleClient>& ContentBundleClient)
@@ -291,18 +331,12 @@ void FContentBundleContainer::OnContentBundleClientUnregistered(FContentBundleCl
 
 void FContentBundleContainer::OnContentBundleClientContentInjectionRequested(FContentBundleClient& ContentBundleClient)
 {
-	if (FContentBundleBase* ContentBundle = GetContentBundle(ContentBundleClient))
-	{
-		InjectContentBundle(*ContentBundle);
-	}
+	InjectContentBundle(ContentBundleClient);
 }
 
 void FContentBundleContainer::OnContentBundleClientContentRemovalRequested(FContentBundleClient& ContentBundleClient)
 {
-	if (FContentBundleBase* ContentBundle = GetContentBundle(ContentBundleClient))
-	{
-		RemoveContentBundle(*ContentBundle);
-	}
+	RemoveContentBundle(ContentBundleClient);
 }
 
 void FContentBundleContainer::RegisterContentBundleClientEvents()
@@ -323,11 +357,6 @@ void FContentBundleContainer::UnregisterContentBundleClientEvents()
 	ContentBundleEngineSubsystem->OnContentBundleClientRequestedContentRemoval.RemoveAll(this);
 }
 
-bool FContentBundleContainer::ShouldInjectClientContent(const TSharedPtr<FContentBundleClient>& ContentBundleClient) const
-{
-	return ContentBundleClient->GetState() == EContentBundleClientState::ContentInjectionRequested;
-}
-
 void FContentBundleContainer::InitializeContentBundlesForegisteredClients()
 {
 	UContentBundleEngineSubsystem* ContentBundleEngineSubsystem = UContentBundleEngineSubsystem::Get();
@@ -340,11 +369,7 @@ void FContentBundleContainer::InitializeContentBundlesForegisteredClients()
 		for (TSharedPtr<FContentBundleClient>& ContentBundleClient : ContentBundleClients)
 		{
 			FContentBundleBase& ContentBundle = InitializeContentBundle(ContentBundleClient);
-
-			if (ShouldInjectClientContent(ContentBundleClient))
-			{
-				InjectContentBundle(ContentBundle);
-			}
+			InjectContentBundle(ContentBundle);
 		}
 
 		UE_LOG(LogContentBundle, Log, TEXT("[Container: %s] End initializing ContentBundles."), *GetInjectedWorld()->GetName());
diff --git a/Engine/Source/Runtime/Engine/Private/WorldPartition/ContentBundle/ContentBundleWorldSubsystem.cpp b/Engine/Source/Runtime/Engine/Private/WorldPartition/ContentBundle/ContentBundleWorldSubsystem.cpp
--- a/Engine/Source/Runtime/Engine/Private/WorldPartition/ContentBundle/ContentBundleWorldSubsystem.cpp
+++ b/Engine/Source/Runtime/Engine/Private/WorldPartition/ContentBundle/ContentBundleWorldSubsystem.cpp
@@ -58,6 +58,16 @@ void UContentBundleManager::Deinitialize()
 #endif
 }
 
+bool UContentBundleManager::CanInject() const
+{
+	if (UWorldPartition* WorldPartition = GetTypedOuter<UWorld>()->GetWorldPartition())
+	{
+		return WorldPartition->IsInitialized();
+	}
+
+	return false;
+}
+
 void UContentBundleManager::AddReferencedObjects(UObject* InThis, FReferenceCollector& Collector)
 {
 	UContentBundleManager* Subsystem = CastChecked<UContentBundleManager>(InThis);
@@ -158,6 +168,27 @@ TSharedPtr<FContentBundleEditor> UContentBundleManager::GetEditorContentBundle(c
 	return nullptr;
 }
 
+bool UContentBundleManager::TryInject(FContentBundleClient& Client)
+{
+	if (CanInject())
+	{
+		if (TUniquePtr<FContentBundleContainer>* Container = GetContentBundleContainer(GetTypedOuter<UWorld>()))
+		{
+			return (*Container)->InjectContentBundle(Client);
+		}
+	}
+
+	return false;
+}
+
+void UContentBundleManager::Remove(FContentBundleClient& Client)
+{
+	if (TUniquePtr<FContentBundleContainer>* Container = GetContentBundleContainer(GetTypedOuter<UWorld>()))
+	{
+		(*Container)->RemoveContentBundle(Client);
+	}
+}
+
 #endif
 
 uint32 UContentBundleManager::GetContentBundleContainerIndex(const UWorld* InjectedWorld) const
@@ -179,6 +210,11 @@ const TUniquePtr<FContentBundleContainer>* UContentBundleManager::GetContentBund
 	return nullptr;
 }
 
+TUniquePtr<FContentBundleContainer>* UContentBundleManager::GetContentBundleContainer(const UWorld* InjectedWorld)
+{
+	return const_cast<TUniquePtr<FContentBundleContainer>*>(const_cast<const UContentBundleManager*>(this)->GetContentBundleContainer(InjectedWorld));
+}
+
 void UContentBundleManager::OnWorldPartitionInitialized(UWorldPartition* InWorldPartition)
 {
 	if (InWorldPartition->GetWorld() != InWorldPartition->GetTypedOuter<UWorld>())
diff --git a/Engine/Source/Runtime/Engine/Public/WorldPartition/ContentBundle/ContentBundleClient.h b/Engine/Source/Runtime/Engine/Public/WorldPartition/ContentBundle/ContentBundleClient.h
--- a/Engine/Source/Runtime/Engine/Public/WorldPartition/ContentBundle/ContentBundleClient.h
+++ b/Engine/Source/Runtime/Engine/Public/WorldPartition/ContentBundle/ContentBundleClient.h
@@ -28,8 +28,7 @@ UENUM()
 enum class EWorldContentState
 {
 	NoContent,
-	ContentBundleInjected,
-	ContentBundleInjectionFailed
+	ContentBundleInjected
 };
 
 class ENGINE_API FContentBundleClient
@@ -50,15 +49,24 @@ public:
 	
 	void RequestUnregister();
 
-	EContentBundleClientState GetState() const { return State; }
+#if WITH_EDITOR
+	void RequestForceInject(UWorld* WorldToInject);
+	void RequestRemoveForceInjectedContent(UWorld *WorldToInject);
+#endif 
 
-	bool HasInjectedAnyContent() const;
+	EContentBundleClientState GetState() const { return State; }
 
 	FString const& GetDisplayName() const { return DisplayName; }
 
+	bool ShouldInjectContent(UWorld* World) const;
+	bool ShouldRemoveContent(UWorld* World) const;
+
 private:
+	bool HasContentToRemove() const;
+
 	void SetState(EContentBundleClientState State);
 	void SetWorldContentState(UWorld* World, EWorldContentState State);
+	EWorldContentState GetWorldContentState(UWorld* World) const;
 
 	void OnContentInjectedInWorld(EContentBundleStatus InjectionStatus, UWorld* InjectedWorld);
 	void OnContentRemovedFromWorld(EContentBundleStatus RemovalStatus, UWorld* InjectedWorld);
@@ -67,6 +75,10 @@ private:
 	
 	TMap<TWeakObjectPtr<UWorld>, EWorldContentState> WorldContentStates;
 
+#if WITH_EDITOR
+	TSet<TWeakObjectPtr<UWorld>> ForceInjectedWorlds;
+#endif
+
 	FString DisplayName;
 
 	EContentBundleClientState State;
diff --git a/Engine/Source/Runtime/Engine/Public/WorldPartition/ContentBundle/ContentBundleContainer.h b/Engine/Source/Runtime/Engine/Public/WorldPartition/ContentBundle/ContentBundleContainer.h
--- a/Engine/Source/Runtime/Engine/Public/WorldPartition/ContentBundle/ContentBundleContainer.h
+++ b/Engine/Source/Runtime/Engine/Public/WorldPartition/ContentBundle/ContentBundleContainer.h
@@ -48,6 +48,9 @@ public:
 	const TArray<TUniquePtr<FContentBundle>>& GetGameContentBundles() const;
 	TArray<TUniquePtr<FContentBundle>>& GetGameContentBundles();
 
+	bool InjectContentBundle(FContentBundleClient& ContentBundleClient);
+	bool RemoveContentBundle(FContentBundleClient& ContentBundleClient);
+
 #if WITH_EDITOR
 	TSharedPtr<FContentBundleEditor> GetEditorContentBundle(const UContentBundleDescriptor* Descriptor) const;
 	TSharedPtr<FContentBundleEditor> GetEditorContentBundle(const FGuid& ContentBundleGuid) const;
@@ -62,8 +65,6 @@ public:
 	void ForEachContentBundleBreakable(TFunctionRef<bool(FContentBundleBase*)> Func) const;
 
 private:
-	bool ShouldInjectClientContent(const TSharedPtr<FContentBundleClient>& ContentBundleClient) const;
-
 	FContentBundleBase* GetContentBundle(const FContentBundleClient& ContentBundleClient);
 
 	void RegisterContentBundleClientEvents();
@@ -80,8 +81,8 @@ private:
 	void InitializeContentBundlesForegisteredClients();
 	void DeinitializeContentBundles();
 
-	void InjectContentBundle(FContentBundleBase& ContentBundle);
-	void RemoveContentBundle(FContentBundleBase& ContentBundle);
+	bool InjectContentBundle(FContentBundleBase& ContentBundle);
+	bool RemoveContentBundle(FContentBundleBase& ContentBundle);
 
 #if WITH_EDITOR
 	void OnPreGenerateStreaming(TArray<FString>* OutPackageToGenerate);
diff --git a/Engine/Source/Runtime/Engine/Public/WorldPartition/ContentBundle/ContentBundleWorldSubsystem.h b/Engine/Source/Runtime/Engine/Public/WorldPartition/ContentBundle/ContentBundleWorldSubsystem.h
--- a/Engine/Source/Runtime/Engine/Public/WorldPartition/ContentBundle/ContentBundleWorldSubsystem.h
+++ b/Engine/Source/Runtime/Engine/Public/WorldPartition/ContentBundle/ContentBundleWorldSubsystem.h
@@ -26,12 +26,16 @@ class ENGINE_API UContentBundleManager : public UObject
 {
 	GENERATED_BODY()
 
+	friend class FContentBundleClient;
+
 public:
 	UContentBundleManager();
 
 	void Initialize();
 	void Deinitialize();
 
+	bool CanInject() const;
+
 	static void AddReferencedObjects(UObject* InThis, FReferenceCollector& Collector);
 
 #if WITH_EDITOR
@@ -46,8 +50,14 @@ public:
 	void DrawContentBundlesStatus(const UWorld* InWorld, UCanvas* Canvas, FVector2D& Offset) const;
 
 private:
+#if WITH_EDITOR
+	bool TryInject(FContentBundleClient& Client);
+	void Remove(FContentBundleClient& Client);
+#endif
+
 	uint32 GetContentBundleContainerIndex(const UWorld* InjectedWorld) const;
 	const TUniquePtr<FContentBundleContainer>* GetContentBundleContainer(const UWorld* InjectedWorld) const;
+	TUniquePtr<FContentBundleContainer>* GetContentBundleContainer(const UWorld* InjectedWorld);
 
 	void OnWorldPartitionInitialized(UWorldPartition* WorldPartition);
 	void OnWorldPartitionUninitialized(UWorldPartition* WorldPartition);
