Avoid running deferred construction scripts for actors past their world cleanup

#rnx

diff --git a/Engine/Source/Runtime/Engine/Private/ActorDeferredScriptManager.cpp b/Engine/Source/Runtime/Engine/Private/ActorDeferredScriptManager.cpp
--- a/Engine/Source/Runtime/Engine/Private/ActorDeferredScriptManager.cpp
+++ b/Engine/Source/Runtime/Engine/Private/ActorDeferredScriptManager.cpp
@@ -10,6 +10,7 @@
 #include "Engine/StaticMesh.h"
 #include "Logging/LogMacros.h"
 #include "ProfilingDebugging/CountersTrace.h"
+#include "Engine/World.h"
 
 #define LOCTEXT_NAMESPACE "ActorDeferredScriptManager"
 
@@ -28,11 +29,18 @@ void FActorDeferredScriptManager::Shutdown()
 		FAssetCompilingManager::Get().OnAssetPostCompileEvent().Remove(OnAssetChangeDelegateHandle);
 		OnAssetChangeDelegateHandle.Reset();
 	}
+
+	if (OnWorldCleanupDelegateHandle.IsValid())
+	{
+		FWorldDelegates::OnWorldCleanup.Remove(OnWorldCleanupDelegateHandle);
+		OnWorldCleanupDelegateHandle.Reset();
+	}
 }
 
 FActorDeferredScriptManager::FActorDeferredScriptManager()
 	: Notification(GetAssetNameFormat())
 {
+	OnWorldCleanupDelegateHandle = FWorldDelegates::OnWorldCleanup.AddRaw(this, &FActorDeferredScriptManager::OnWorldCleanup);
 }
 
 FName FActorDeferredScriptManager::GetStaticAssetTypeName()
@@ -40,6 +48,37 @@ FName FActorDeferredScriptManager::GetStaticAssetTypeName()
 	return TEXT("UE-ActorConstructionScripts");
 }
 
+void FActorDeferredScriptManager::OnWorldCleanup(class UWorld* InWorld, bool bInSessionEnded, bool bInCleanupResources)
+{
+	int32 RemovedCount = 
+		PendingConstructionScriptActors.RemoveAll(
+			[InWorld](TWeakObjectPtr<AActor>& InActor)
+			{
+				AActor* Actor = InActor.Get();
+				if (Actor == nullptr)
+				{
+					UE_LOG(LogActorDeferredScriptManager, VeryVerbose, TEXT("Removed an actor from deferred construction script because its not valid anymore"));
+					return true;
+				}
+				else if (Actor->GetWorld() == InWorld)
+				{
+					UE_LOG(LogActorDeferredScriptManager, VeryVerbose, TEXT("Removed %s from deferred construction script because its world is being cleaned up"), *InActor->GetPathName());
+					return true;
+				}
+
+				return false;
+			}
+		);
+
+	if (RemovedCount > 0)
+	{
+		// Whenever an actor is removed, reset the count so we go over all actors on next ProcessAsyncTasks
+		NumLeftToProcess = PendingConstructionScriptActors.Num();
+
+		UpdateCompilationNotification();
+	}
+}
+
 void FActorDeferredScriptManager::AddActor(AActor* InActor)
 {
 	UE_LOG(LogActorDeferredScriptManager, VeryVerbose, TEXT("Adding actor %s for deferred construction script"), *InActor->GetPathName());
diff --git a/Engine/Source/Runtime/Engine/Public/ActorDeferredScriptManager.h b/Engine/Source/Runtime/Engine/Public/ActorDeferredScriptManager.h
--- a/Engine/Source/Runtime/Engine/Public/ActorDeferredScriptManager.h
+++ b/Engine/Source/Runtime/Engine/Public/ActorDeferredScriptManager.h
@@ -47,6 +47,7 @@ private:
 	void UpdateCompilationNotification();
 
 	void OnAssetPostCompile(const TArray<FAssetCompileData>& CompiledAssets);
+	void OnWorldCleanup(class UWorld* InWorld, bool bInSessionEnded, bool bInCleanupResources);
 
 	void EnsureEventRegistered();
 
@@ -59,6 +60,9 @@ private:
 	/** We need to monitor when static meshes have finished compiling to only work when necessary */
 	FDelegateHandle OnAssetChangeDelegateHandle;
 
+	/** We need to monitor when worlds are cleaned up to avoid triggering any construction script past this point */
+	FDelegateHandle OnWorldCleanupDelegateHandle;
+
 	/** Where to start on the next ProcessAsyncTasks iteration. */
 	int32 NextIndexToProcess = 0;
 	
