Fixed crash when stopping PIE and loading instanced partitioned worlds.
Discard Cell's LevelStreaming notification when once WorldPartition is unitialized (can happen for instanced WorldPartition)

diff --git a/Engine/Source/Runtime/Engine/Private/WorldPartition/WorldPartition.cpp b/Engine/Source/Runtime/Engine/Private/WorldPartition/WorldPartition.cpp
--- a/Engine/Source/Runtime/Engine/Private/WorldPartition/WorldPartition.cpp
+++ b/Engine/Source/Runtime/Engine/Private/WorldPartition/WorldPartition.cpp
@@ -1496,7 +1496,8 @@ bool UWorldPartition::IsStreamingCompleted(EWorldPartitionRuntimeCellState Query
 
 void UWorldPartition::OnCellShown(const UWorldPartitionRuntimeCell* InCell)
 {
-	if (GetWorld()->IsGameWorld())
+	// Discard Cell's LevelStreaming notification when once WorldPartition is unitialized (can happen for instanced WorldPartition)
+	if (GetWorld()->IsGameWorld() && IsInitialized())
 	{
 		if (IsStreamingEnabled())
 		{
@@ -1508,7 +1509,8 @@ void UWorldPartition::OnCellShown(const UWorldPartitionRuntimeCell* InCell)
 
 void UWorldPartition::OnCellHidden(const UWorldPartitionRuntimeCell* InCell)
 {
-	if (GetWorld()->IsGameWorld())
+	// Discard Cell's LevelStreaming notification when once WorldPartition is unitialized (can happen for instanced WorldPartition)
+	if (GetWorld()->IsGameWorld() && IsInitialized())
 	{
 		if (IsStreamingEnabled())
 		{
