Fix crash in shipping configuration editor when changing to modeling mode

diff --git a/Engine/Source/Runtime/Engine/Private/UnrealEngine.cpp b/Engine/Source/Runtime/Engine/Private/UnrealEngine.cpp
--- a/Engine/Source/Runtime/Engine/Private/UnrealEngine.cpp
+++ b/Engine/Source/Runtime/Engine/Private/UnrealEngine.cpp
@@ -3018,7 +3018,7 @@ void UEngine::InitializeObjectReferences()
 	LoadSpecialMaterial(TEXT("InvalidLightmapSettingsMaterialName"), InvalidLightmapSettingsMaterialName.ToString(), InvalidLightmapSettingsMaterial, false);
 	LoadSpecialMaterial(TEXT("ArrowMaterialName"), ArrowMaterialName.ToString(), ArrowMaterial, false);
 
-#if !UE_BUILD_SHIPPING
+#if !UE_BUILD_SHIPPING || UE_BUILD_SHIPPING_WITH_EDITOR
 	ArrowMaterialYellow = UMaterialInstanceDynamic::Create(ArrowMaterial, nullptr);
 	ArrowMaterialYellow->SetVectorParameterValue("GizmoColor", FLinearColor::Yellow);
 
