Manual integration after robomerge failure: original change submitted to 23.20 in CL 23780423.
See thread: https://epicgames.slack.com/archives/C04553YFKPD/p1674162800144419

Speculative fix to revert UWorld's FOnActorRemovedFromWorld delegate timing changes introduced in CL 23585822. This change is the prime suspect at the moment:
* Updated UWorld:
** Restored the former FOnActorRemovedFromWorld used by ULevelSaveRecord::RenameActorAfterRemovedFromWorld.
** Renamed FOnActorRemovedFromWorld to FOnPostRegisterAllActorComponents.
** Added FOnPreUnregisterAllActorComponents.
* Updated UPlayspaceComponent_ActorMemoryTracker to use FOnPostRegisterAllActorComponents and FOnPreUnregisterAllActorComponents delegates.

#localization none
#tests this is an integration

diff --git a/Engine/Source/Runtime/Engine/Classes/Engine/World.h b/Engine/Source/Runtime/Engine/Classes/Engine/World.h
--- a/Engine/Source/Runtime/Engine/Classes/Engine/World.h
+++ b/Engine/Source/Runtime/Engine/Classes/Engine/World.h
@@ -96,7 +96,8 @@ ENGINE_API DECLARE_LOG_CATEGORY_EXTERN(LogSpawn, Warning, All);
 
 DECLARE_MULTICAST_DELEGATE_OneParam(FOnActorSpawned, AActor*);
 DECLARE_MULTICAST_DELEGATE_OneParam(FOnActorDestroyed, AActor*);
-DECLARE_MULTICAST_DELEGATE_OneParam(FOnActorAddedToWorld, AActor*);
+DECLARE_MULTICAST_DELEGATE_OneParam(FOnPostRegisterAllActorComponents, AActor*);
+DECLARE_MULTICAST_DELEGATE_OneParam(FOnPreUnregisterAllActorComponents, AActor*);
 DECLARE_MULTICAST_DELEGATE_OneParam(FOnActorRemovedFromWorld, AActor*);
 DECLARE_MULTICAST_DELEGATE_OneParam(FOnFeatureLevelChanged, ERHIFeatureLevel::Type);
 DECLARE_MULTICAST_DELEGATE_OneParam(FOnMovieSceneSequenceTick, float);
@@ -1532,13 +1533,19 @@ private:
 	mutable FOnActorDestroyed OnActorDestroyed;
 
 	/**
-	 * Broadcasts after an actor has been fully added to a world and its components have been registered.
+	 * Broadcasts after an actor has registered all its components.
 	 * This is called for both spawned and loaded actors.
 	 */
-	mutable FOnActorAddedToWorld OnActorAddedToWorld;
+	mutable FOnPostRegisterAllActorComponents OnPostRegisterAllActorComponents;
 
 	/**
-	 * Broadcasts when an actor has been removed from the world and its components have been unregistered.
+	 * Broadcasts before an actor unregisters all its components.
+	 * This is called for both spawned and loaded actors.
+	 */
+	mutable FOnPreUnregisterAllActorComponents OnPreUnregisterAllActorComponents;
+
+	/**
+	 * Broadcasts when an actor has been removed from the world.
 	 * This event is the earliest point where an actor can be safely renamed without affecting replication.
 	 */
 	mutable FOnActorRemovedFromWorld OnActorRemovedFromWorld;
@@ -2791,30 +2798,36 @@ public:
 	/** Remove a listener for OnActorDestroyed events */
 	void RemoveOnActorDestroyededHandler(FDelegateHandle InHandle) const;
 
-	/** Add a listener for OnActorAddedToWorld events */
-	FDelegateHandle AddOnActorAddedToWorldHandler(const FOnActorAddedToWorld::FDelegate& InHandler) const;
+	/** Add a listener for OnPostRegisterAllActorComponents events */
+	FDelegateHandle AddOnPostRegisterAllActorComponentsHandler(const FOnPostRegisterAllActorComponents::FDelegate& InHandler) const;
+
+	/** Remove a listener for OnPostRegisterAllActorComponents events */
+	void RemoveOnPostRegisterAllActorComponentsHandler(FDelegateHandle InHandle) const;
 
 	/**
-	 * Broadcast an OnActorAddedToWorld event.
+	 * Broadcast an OnPostRegisterAllActorComponents event.
 	 * This method should only be called from internal actor and level code and never on inactive worlds.
 	 */
-	void NotifyActorAddedToWorld(AActor* Actor);
-
-	/** Remove a listener for OnActorAddedToWorld events */
-	void RemoveOnActorAddedToWorldHandler(FDelegateHandle InHandle) const;
+	void NotifyPostRegisterAllActorComponents(AActor* Actor);
 
-	/** Add a listener for OnActorRemovedFromWorld events */
-	FDelegateHandle AddOnActorRemovedFromWorldHandler(const FOnActorRemovedFromWorld::FDelegate& InHandler) const;
+	/** Add a listener for OnPreUnregisterAllActorComponents events */
+	FDelegateHandle AddOnPreUnregisterAllActorComponentsHandler(const FOnPreUnregisterAllActorComponents::FDelegate& InHandler) const;
 
-	/** Remove a listener for OnActorRemovedFromWorld events */
-	void RemoveOnActorRemovedFromWorldHandler(FDelegateHandle InHandle) const;
+	/** Remove a listener for OnPreUnregisterAllActorComponents events */
+	void RemoveOnPreUnregisterAllActorComponentsHandler(FDelegateHandle InHandle) const;
 
 	/**
-	 * Broadcast an OnActorRemovedFromWorld event.
+	 * Broadcast an OnPreUnregisterAllActorComponents event.
 	 * This method should only be called from internal actor and level code. Calls on inactive or GCing worlds are
 	 * ignored.
 	 */
-	void NotifyActorRemovedFromWorld(AActor* Actor);
+	void NotifyPreUnregisterAllActorComponents(AActor* Actor);
+
+	/** Add a listener for OnActorRemovedFromWorld events */
+	FDelegateHandle AddOnActorRemovedFromWorldHandler(const FOnActorRemovedFromWorld::FDelegate& InHandler) const;
+
+	/** Remove a listener for OnActorRemovedFromWorld events */
+	void RemoveOnActorRemovedFromWorldHandler(FDelegateHandle InHandle) const;
 
 	/**
 	 * Returns whether the passed in actor is part of any of the loaded levels actors array.
diff --git a/Engine/Source/Runtime/Engine/Private/Actor.cpp b/Engine/Source/Runtime/Engine/Private/Actor.cpp
--- a/Engine/Source/Runtime/Engine/Private/Actor.cpp
+++ b/Engine/Source/Runtime/Engine/Private/Actor.cpp
@@ -5240,7 +5240,7 @@ void AActor::UnregisterAllComponents(const bool bForReregister)
 	{
 		if (UWorld* OwningWorld = GetWorld())
 		{
-			OwningWorld->NotifyActorRemovedFromWorld(this);
+			OwningWorld->NotifyPreUnregisterAllActorComponents(this);
 		}
 	}
 
@@ -5393,7 +5393,7 @@ bool AActor::IncrementalRegisterComponents(int32 NumComponentsToRegister, FRegis
 		PostRegisterAllComponents();
 
 		// After all components have been registered the actor is considered fully added: notify the owning world.
-		World->NotifyActorAddedToWorld(this);
+		World->NotifyPostRegisterAllActorComponents(this);
 		return true;
 	}
 	
diff --git a/Engine/Source/Runtime/Engine/Private/LevelActor.cpp b/Engine/Source/Runtime/Engine/Private/LevelActor.cpp
--- a/Engine/Source/Runtime/Engine/Private/LevelActor.cpp
+++ b/Engine/Source/Runtime/Engine/Private/LevelActor.cpp
@@ -973,6 +973,8 @@ bool UWorld::DestroyActor( AActor* ThisActor, bool bNetForce, bool bShouldModify
 #endif
 	}
 
+	OnActorRemovedFromWorld.Broadcast(ThisActor);
+
 	// Clean up the actor's components.
 	ThisActor->UnregisterAllComponents();
 
diff --git a/Engine/Source/Runtime/Engine/Private/World.cpp b/Engine/Source/Runtime/Engine/Private/World.cpp
--- a/Engine/Source/Runtime/Engine/Private/World.cpp
+++ b/Engine/Source/Runtime/Engine/Private/World.cpp
@@ -5365,44 +5365,54 @@ void UWorld::RemoveOnActorDestroyededHandler(FDelegateHandle InHandle) const
 	OnActorDestroyed.Remove(InHandle);
 }
 
-FDelegateHandle UWorld::AddOnActorAddedToWorldHandler(const FOnActorAddedToWorld::FDelegate& InHandler) const
+FDelegateHandle UWorld::AddOnPostRegisterAllActorComponentsHandler(const FOnPostRegisterAllActorComponents::FDelegate& InHandler) const
 {
-	return OnActorAddedToWorld.Add(InHandler);
+	return OnPostRegisterAllActorComponents.Add(InHandler);
 }
 
-void UWorld::RemoveOnActorAddedToWorldHandler(FDelegateHandle InHandle) const
+void UWorld::RemoveOnPostRegisterAllActorComponentsHandler(FDelegateHandle InHandle) const
 {
-	OnActorAddedToWorld.Remove(InHandle);
+	OnPostRegisterAllActorComponents.Remove(InHandle);
 }
 
-void UWorld::NotifyActorAddedToWorld(AActor* Actor)
+void UWorld::NotifyPostRegisterAllActorComponents(AActor* Actor)
 {
 	// This may be called on inactive worlds (for example, while cooking), if so ignore it.
 	if (WorldType != EWorldType::Inactive)
 	{
-		OnActorAddedToWorld.Broadcast(Actor);
+		OnPostRegisterAllActorComponents.Broadcast(Actor);
 	}
 }
 
-FDelegateHandle UWorld::AddOnActorRemovedFromWorldHandler(const FOnActorRemovedFromWorld::FDelegate& InHandler) const
+FDelegateHandle UWorld::AddOnPreUnregisterAllActorComponentsHandler(const FOnPreUnregisterAllActorComponents::FDelegate& InHandler) const
 {
-	return OnActorRemovedFromWorld.Add(InHandler);
+	return OnPreUnregisterAllActorComponents.Add(InHandler);
 }
 
-void UWorld::RemoveOnActorRemovedFromWorldHandler(FDelegateHandle InHandle) const
+void UWorld::RemoveOnPreUnregisterAllActorComponentsHandler(FDelegateHandle InHandle) const
 {
-	OnActorRemovedFromWorld.Remove(InHandle);
+	OnPreUnregisterAllActorComponents.Remove(InHandle);
 }
 
-void UWorld::NotifyActorRemovedFromWorld(AActor* Actor)
+void UWorld::NotifyPreUnregisterAllActorComponents(AActor* Actor)
 {
 	// This may be called on inactive/GCing worlds, if so ignore it.
 	if (WorldType != EWorldType::Inactive && IsValid(this))
 	{
-		OnActorRemovedFromWorld.Broadcast(Actor);
+		OnPreUnregisterAllActorComponents.Broadcast(Actor);
 	}
 }
 
+FDelegateHandle UWorld::AddOnActorRemovedFromWorldHandler(const FOnActorRemovedFromWorld::FDelegate& InHandler) const
+{
+	return OnActorRemovedFromWorld.Add(InHandler);
+}
+
+void UWorld::RemoveOnActorRemovedFromWorldHandler(FDelegateHandle InHandle) const
+{
+	OnActorRemovedFromWorld.Remove(InHandle);
+}
+
 FDelegateHandle UWorld::AddMovieSceneSequenceTickHandler(const FOnMovieSceneSequenceTick::FDelegate& InHandler)
 {
 	return MovieSceneSequenceTick.Add(InHandler);
