Misc. cleanup for the Input Device Subsystem:

- Remove the GetPlayerControllerFromPlatformUser function from the input device subsystem in favor of using the BP library. This function hasn't been released yet so there is no need to deprecate it.
- Rename FSetDevicePropertyParams to FActivateDevicePropertyParams to be in line with the current function
- Move the IsDevicePropertyHandleValid function from the subsystem to the library.
- Update some comments and metadata on various properties and functions in the subsystem to make them more user friendly
- Move some default value initalizations from the header file to the .cpp so that we can change it in hotfixes if we need to, and so that we can make them uint8 : 1 flags instead of bools

#rnx

diff --git a/Engine/Config/BaseEngine.ini b/Engine/Config/BaseEngine.ini
--- a/Engine/Config/BaseEngine.ini
+++ b/Engine/Config/BaseEngine.ini
@@ -1099,6 +1099,12 @@ DefaultWorldPartitionHLODLayer=/Engine/Maps/Templates/HLODs/HLODLayer_Instanced.
 
 ; 5.2
 +ClassRedirects=(OldName="/Script/Engine.PlatformInputDeviceMapperLibrary",NewName="/Script/Engine.InputDeviceLibrary")
++ClassRedirects=(OldName="/Script/Engine.PlatformInputDeviceMapperLibrary",NewName="/Script/Engine.InputDeviceLibrary")
++StructRedirects=(OldName="/Script/Engine.SetDevicePropertyParams", NewName="/Script/Engine.ActivateDevicePropertyParams")
++FunctionRedirects=(OldName="InputDeviceSubsystem.GetPlayerControllerFromInputDevice", NewName="InputDeviceLibrary.GetPlayerControllerFromInputDevice")
++FunctionRedirects=(OldName="InputDeviceSubsystem.GetPlayerControllerFromPlatformUser", NewName="InputDeviceLibrary.GetPlayerControllerFromPlatformUser")
++FunctionRedirects=(OldName="InputDeviceSubsystem.IsDevicePropertyHandleValid", NewName="InputDeviceLibrary.IsDevicePropertyHandleValid")
+
 +StructRedirects=(OldName="/Script/AssetTools.RevisionInfo", NewName="/Script/AssetDefinition.RevisionInfo")
 +PropertyRedirects=(OldName="DecalComponent.SetFadeIn.Duaration", NewName="DecalComponent.SetFadeIn.Duration")
 +PropertyRedirects=(OldName="BlendSpace.PerBoneBlend", NewName="BlendSpace.ManualPerBoneOverrides")
diff --git a/Engine/Plugins/Runtime/GameplayAbilities/Source/GameplayAbilities/Private/GameplayCueNotifyTypes.cpp b/Engine/Plugins/Runtime/GameplayAbilities/Source/GameplayAbilities/Private/GameplayCueNotifyTypes.cpp
--- a/Engine/Plugins/Runtime/GameplayAbilities/Source/GameplayAbilities/Private/GameplayCueNotifyTypes.cpp
+++ b/Engine/Plugins/Runtime/GameplayAbilities/Source/GameplayAbilities/Private/GameplayCueNotifyTypes.cpp
@@ -860,7 +860,7 @@ bool FGameplayCueNotify_InputDevicePropertyInfo::SetDeviceProperties(const FGame
 		// Apply any device properties from this gameplay cue
 		if (UInputDeviceSubsystem* System = UInputDeviceSubsystem::Get())
 		{
-			FSetDevicePropertyParams Params = {};
+			FActivateDevicePropertyParams Params = {};
 			Params.UserId = TargetPC->GetPlatformUserId();
 			
 			if (ensure(Params.UserId.IsValid()))
diff --git a/Engine/Source/Runtime/Engine/Classes/GameFramework/InputDeviceLibrary.h b/Engine/Source/Runtime/Engine/Classes/GameFramework/InputDeviceLibrary.h
--- a/Engine/Source/Runtime/Engine/Classes/GameFramework/InputDeviceLibrary.h
+++ b/Engine/Source/Runtime/Engine/Classes/GameFramework/InputDeviceLibrary.h
@@ -4,6 +4,7 @@
 
 #include "CoreMinimal.h"
 #include "Kismet/BlueprintFunctionLibrary.h"
+#include "GameFramework/InputDevicePropertyHandle.h"
 
 #include "InputDeviceLibrary.generated.h"
 
@@ -33,6 +34,10 @@ public:
 	UFUNCTION(BlueprintCallable, Category = "Input Devices")
 	static APlayerController* GetPlayerControllerFromInputDevice(const FInputDeviceId DeviceId);
 
+	/** Returns true if the given handle is valid */
+	UFUNCTION(BlueprintCallable, BlueprintPure, Category = "Input Devices", meta = (ReturnDisplayName = "Is Valid"))
+	static bool IsDevicePropertyHandleValid(const FInputDevicePropertyHandle& InHandle);
+	
 	/**
 	 * Populates the OutInputDevices array with any InputDeviceID's that are mapped to the given platform user
 	 *
diff --git a/Engine/Source/Runtime/Engine/Classes/GameFramework/InputDeviceSubsystem.h b/Engine/Source/Runtime/Engine/Classes/GameFramework/InputDeviceSubsystem.h
--- a/Engine/Source/Runtime/Engine/Classes/GameFramework/InputDeviceSubsystem.h
+++ b/Engine/Source/Runtime/Engine/Classes/GameFramework/InputDeviceSubsystem.h
@@ -12,18 +12,17 @@
 ENGINE_API DECLARE_LOG_CATEGORY_EXTERN(LogInputDeviceProperties, Log, All);
 
 class UInputDeviceProperty;
-class APlayerController;
 
-/** Parameters for the UInputDeviceSubsystem::SetDeviceProperty function */
+/** Parameters for the UInputDeviceSubsystem::ActivateDeviceProperty function */
 USTRUCT(BlueprintType)
-struct ENGINE_API FSetDevicePropertyParams
+struct ENGINE_API FActivateDevicePropertyParams
 {
 	GENERATED_BODY()
 
-	FSetDevicePropertyParams();
+	FActivateDevicePropertyParams();
 	
 	/** The Platform User whose device's should receive the device property */
-	UPROPERTY(EditAnywhere, BlueprintReadWrite, Category = "Input Devices")
+	UPROPERTY(EditAnywhere, BlueprintReadWrite, Category = "Activation Options")
 	FPlatformUserId UserId;
 
 	/** 
@@ -32,37 +31,32 @@ struct ENGINE_API FSetDevicePropertyParams
 	* 
 	* The default input device is obtained from IPlatformInputDeviceMapper::GetPrimaryInputDeviceForUser
 	*/
-	UPROPERTY(EditAnywhere, BlueprintReadWrite, Category = "Input Devices")
+	UPROPERTY(EditAnywhere, BlueprintReadWrite, Category = "Activation Options")
 	FInputDeviceId DeviceId;
 
 	/**
 	* If true, then the input device property will not be removed after it's evaluation time has completed.
 	* Instead, it will remain active until manually removed with a RemoveDeviceProperty call.
 	*/
-	UPROPERTY(EditAnywhere, BlueprintReadWrite, Category = "Input Devices")
-	bool bLooping = false;
+	UPROPERTY(EditAnywhere, BlueprintReadWrite, Category = "Activation Options")
+	uint8 bLooping : 1;
 
 	/** If true, then this device property will ignore dilated delta time and use the Applications delta time instead */
-	UPROPERTY(EditAnywhere, BlueprintReadWrite, Category = "Input Devices")
-	bool bIgnoreTimeDilation = false;
+	UPROPERTY(EditAnywhere, BlueprintReadWrite, Category = "Activation Options")
+	uint8 bIgnoreTimeDilation : 1;
 
 	/** If true, then this device property will be played even if the game world is paused. */
-	UPROPERTY(EditAnywhere, BlueprintReadWrite, Category = "Input Devices")
-	bool bPlayWhilePaused = false;
+	UPROPERTY(EditAnywhere, BlueprintReadWrite, Category = "Activation Options")
+	uint8 bPlayWhilePaused : 1;
 };
 
-/**
- * Delegate called when a user changed the hardware they are using for input.
- *
- * @param UserId		The Platform user whose device has changed
- */
-DECLARE_DYNAMIC_MULTICAST_DELEGATE_TwoParams(FHardwareInputDeviceChanged, const FPlatformUserId, UserId, const FInputDeviceId, DeviceId);
-
 /** Contains a pointer to an active device property and keeps track of how long it has been evaluated for */
 USTRUCT()
 struct ENGINE_API FActiveDeviceProperty
 {
 	GENERATED_BODY()
+
+	FActiveDeviceProperty();
 	
 	/** Active properties can just use the hash of their FInputDevicePropertyHandle for a fast and unique lookup */
 	ENGINE_API friend uint32 GetTypeHash(const FActiveDeviceProperty& InProp);
@@ -92,13 +86,13 @@ struct ENGINE_API FActiveDeviceProperty
 	* If true, then the input device property will not be removed after it's evaluation time has completed.
 	* Instead, it will remain active until manually removed with a RemoveDeviceProperty call.
 	*/
-	bool bLooping = false;
+	uint8 bLooping : 1;
 
 	/** If true, then this device property will ignore dilated delta time and use the Applications delta time instead */
-	bool bIgnoreTimeDilation = false;
+	uint8 bIgnoreTimeDilation : 1;
 
 	/** If true, then this device property will be played even if the game world is paused. */
-	bool bPlayWhilePaused = false;
+	uint8 bPlayWhilePaused : 1;
 
 	/**
 	 * This is set to true when this device property has been applied.
@@ -107,9 +101,16 @@ struct ENGINE_API FActiveDeviceProperty
 	 * I.e., your property is set to a duration of 0.1, but you get a delta time of .12 seconds for some reason
 	 * (choppy frames, low perf client, debugging, etc). 
 	 */
-	bool bHasBeenAppliedAtLeastOnce = false;
+	uint8 bHasBeenAppliedAtLeastOnce : 1;
 };
 
+/**
+ * Delegate called when a user changed the hardware they are using for input.
+ *
+ * @param UserId		The Platform user whose device has changed
+ */
+DECLARE_DYNAMIC_MULTICAST_DELEGATE_TwoParams(FHardwareInputDeviceChanged, const FPlatformUserId, UserId, const FInputDeviceId, DeviceId);
+
 /**
 * The input device subsystem provides an interface to allow users to set Input Device Properties
 * on any Platform User. 
@@ -131,9 +132,11 @@ public:
 	 */
 	static UInputDeviceSubsystem* Get();
 
+	//~ Begin UEngineSubsystem interface
 	virtual void Initialize(FSubsystemCollectionBase& Collection) override;
 	virtual void Deinitialize() override;
 	virtual bool ShouldCreateSubsystem(UObject* Outer) const;
+	//~ End UEngineSubsystem interface
 	
 	//~ Begin FTickableGameObject interface	
 	virtual UWorld* GetTickableGameObjectWorld() const override;
@@ -143,33 +146,25 @@ public:
 	virtual TStatId GetStatId() const override;
 	virtual void Tick(float InDeltaTime) override;
 	//~ End FTickableGameObject interface
-
-	/** Get the player controller who has the given Platform User ID. */
-	UFUNCTION(BlueprintCallable, Category = "Input Devices")
-	static APlayerController* GetPlayerControllerFromPlatformUser(const FPlatformUserId UserId);
-
-	/** Get the player controller who owns the given input device id */
-	UFUNCTION(BlueprintCallable, Category = "Input Devices")
-	static APlayerController* GetPlayerControllerFromInputDevice(const FInputDeviceId DeviceId);
-
+	
 	/**
 	 * Starts tracking the given device property as an "Active" property. This means that the property will be evaluted and applied to its platform user
 	 *
 	 * NOTE: This does NOT make a new instance of the given property. If you pass in the same object before it is completely
 	 * evaluated, then you see undesired effects.
 	 */
-	FInputDevicePropertyHandle ActivateDeviceProperty(UInputDeviceProperty* Property, const FSetDevicePropertyParams& Params);
+	FInputDevicePropertyHandle ActivateDeviceProperty(UInputDeviceProperty* Property, const FActivateDevicePropertyParams& Params);
 
 	/** Spawn a new instance of the given device property class and activate it. */
 	UFUNCTION(BlueprintCallable, Category = "Input Devices", meta = (AutoCreateRefTerm = "Params", ReturnDisplayName = "Device Property Handle"))
-	FInputDevicePropertyHandle ActivateDevicePropertyOfClass(TSubclassOf<UInputDeviceProperty> PropertyClass, const FSetDevicePropertyParams& Params);
+	FInputDevicePropertyHandle ActivateDevicePropertyOfClass(TSubclassOf<UInputDeviceProperty> PropertyClass, const FActivateDevicePropertyParams& Params);
 	
 	/** Returns a pointer to the active input device property with the given handle. Returns null if the property doesn't exist */
 	UFUNCTION(BlueprintCallable, Category = "Input Devices", meta=(ReturnDisplayName="Device Property"))
 	UInputDeviceProperty* GetActiveDeviceProperty(const FInputDevicePropertyHandle Handle) const;
 	
 	/** Returns true if the property associated with the given handle is currently active, and it is not pending removal */
-	UFUNCTION(BlueprintCallable, Category = "Input Devices", meta = (ReturnDisplayName = "Is Property Active"))
+	UFUNCTION(BlueprintCallable, Category = "Input Devices", meta = (ReturnDisplayName = "Is Active"))
 	bool IsPropertyActive(const FInputDevicePropertyHandle Handle) const;	
 
 	/**
@@ -196,10 +191,6 @@ public:
 	UFUNCTION(BlueprintCallable, Category = "Input Devices")
 	void RemoveAllDeviceProperties();
 
-	/** Returns true if the given handle is valid */
-	UFUNCTION(BlueprintCallable, BlueprintPure, Category = "Input Devices")
-	static bool IsDevicePropertyHandleValid(const FInputDevicePropertyHandle& InHandle);
-
 	/** Gets the most recently used hardware input device for the given platform user */
 	UFUNCTION(BlueprintCallable, Category = "Input Devices")
 	FHardwareDeviceIdentifier GetMostRecentlyUsedHardwareDevice(const FPlatformUserId InUserId) const;
@@ -208,12 +199,12 @@ public:
 	FHardwareDeviceIdentifier GetInputDeviceHardwareIdentifier(const FInputDeviceId InputDevice) const;
 
 	/** A delegate that is fired when a platform user changes what Hardware Input device they are using */
-	UPROPERTY(BlueprintAssignable, Category = "Input Device")
+	UPROPERTY(BlueprintAssignable, Category = "Input Devices")
 	FHardwareInputDeviceChanged OnInputHardwareDeviceChanged;
 	
 protected:
 
-	/** Set the most recently used hardware device */
+	/** Set the most recently used hardware device from the input processor */
 	void SetMostRecentlyUsedHardwareDevice(const FInputDeviceId InDeviceId, const FHardwareDeviceIdentifier& InHardwareId);
 
 	// Callbacks for when PIE is started/stopped. We will likely want to pause/resume input device properties
@@ -246,6 +237,6 @@ protected:
 	/** A map of platform user's to their most recent hardware device identifier */
 	TMap<FPlatformUserId, FHardwareDeviceIdentifier> LatestUserDeviceIdentifiers;
 
-	/** An input processor that is used to determine the current hardware input device */
+	/** An input processor that is used to determine the most recently used hardware device for each user*/
 	TSharedPtr<class FInputDeviceSubsystemProcessor> InputPreprocessor;
 };
\ No newline at end of file
diff --git a/Engine/Source/Runtime/Engine/Private/GameFramework/ForceFeedbackEffect.cpp b/Engine/Source/Runtime/Engine/Private/GameFramework/ForceFeedbackEffect.cpp
--- a/Engine/Source/Runtime/Engine/Private/GameFramework/ForceFeedbackEffect.cpp
+++ b/Engine/Source/Runtime/Engine/Private/GameFramework/ForceFeedbackEffect.cpp
@@ -179,7 +179,7 @@ void FActiveForceFeedbackEffect::ActivateDeviceProperties()
 			{
 				if (ensure(DeviceProp))
 				{
-					FSetDevicePropertyParams Params = {};
+					FActivateDevicePropertyParams Params = {};
 					Params.bIgnoreTimeDilation = Parameters.bIgnoreTimeDilation;
 					Params.UserId = PlatformUser;
 					Params.bPlayWhilePaused = Parameters.bPlayWhilePaused;
diff --git a/Engine/Source/Runtime/Engine/Private/GameFramework/InputDeviceSubsystem.cpp b/Engine/Source/Runtime/Engine/Private/GameFramework/InputDeviceSubsystem.cpp
--- a/Engine/Source/Runtime/Engine/Private/GameFramework/InputDeviceSubsystem.cpp
+++ b/Engine/Source/Runtime/Engine/Private/GameFramework/InputDeviceSubsystem.cpp
@@ -4,7 +4,6 @@
 #include "GameFramework/InputDeviceProperties.h"
 #include "Framework/Application/IInputProcessor.h"
 #include "Framework/Application/SlateApplication.h"			// For RegisterInputPreProcessor
-#include "GameFramework/PlayerController.h"
 #include "Misc/App.h"	// For FApp::GetDeltaTime()
 
 #include UE_INLINE_GENERATED_CPP_BY_NAME(InputDeviceSubsystem)
@@ -107,6 +106,18 @@ public:
 ////////////////////////////////////////////////////////
 // FActiveDeviceProperty
 
+FActiveDeviceProperty::FActiveDeviceProperty()
+	: Property(nullptr)
+	, EvaluatedDuration(0.0)
+	, PlatformUser(PLATFORMUSERID_NONE)
+	, DeviceId(INPUTDEVICEID_NONE)
+	, PropertyHandle(FInputDevicePropertyHandle::InvalidHandle)
+	, bLooping(false)
+	, bIgnoreTimeDilation(false)
+	, bPlayWhilePaused(false)
+	, bHasBeenAppliedAtLeastOnce(false)
+{ }
+
 uint32 GetTypeHash(const FActiveDeviceProperty& InProp)
 {
 	return InProp.PropertyHandle.GetTypeHash();
@@ -134,14 +145,17 @@ bool FActiveDeviceProperty::operator!=(const FActiveDeviceProperty& Other) const
 }
 
 ////////////////////////////////////////////////////////
-// FSetDevicePropertyParams
+// FActivateDevicePropertyParams
 
-FSetDevicePropertyParams::FSetDevicePropertyParams()
+FActivateDevicePropertyParams::FActivateDevicePropertyParams()
 	: UserId(FSlateApplicationBase::SlateAppPrimaryPlatformUser)
 	// Set this Device Id to NONE by default. The subsystem will detect this being invalid and fall back 
 	// to the default device for the given platform user. This will let the default behavior of the system be 
 	// consistent, while still allowing for setting specific input devices if desired.
 	, DeviceId(INPUTDEVICEID_NONE)
+	, bLooping(false)
+	, bIgnoreTimeDilation(false)
+	, bPlayWhilePaused(false)
 {
 }
 
@@ -332,40 +346,7 @@ void UInputDeviceSubsystem::Tick(float InDeltaTime)
 	PropertiesPendingRemoval.Empty();
 }
 
-APlayerController* UInputDeviceSubsystem::GetPlayerControllerFromPlatformUser(const FPlatformUserId UserId)
-{
-	if (UserId.IsValid())
-	{
-		if (const UInputDeviceSubsystem* System = UInputDeviceSubsystem::Get())
-		{
-			if (const UWorld* World = System->GetTickableGameObjectWorld())
-			{
-				for (FConstPlayerControllerIterator Iterator = World->GetPlayerControllerIterator(); Iterator; ++Iterator)
-				{
-					APlayerController* PlayerController = Iterator->Get();
-					const FPlatformUserId PlayerControllerUserID = PlayerController ? PlayerController->GetPlatformUserId() : PLATFORMUSERID_NONE;
-					if (PlayerControllerUserID.IsValid() && PlayerControllerUserID == UserId)
-					{
-						return PlayerController;
-					}
-				}
-			}
-		}
-	}
-	else
-	{
-		UE_LOG(LogInputDeviceProperties, Warning, TEXT("Attempting to find the player controller for an invalid Platform"));
-	}
-
-	return nullptr;
-}
-
-APlayerController* UInputDeviceSubsystem::GetPlayerControllerFromInputDevice(const FInputDeviceId DeviceId)
-{
-	return GetPlayerControllerFromPlatformUser(IPlatformInputDeviceMapper::Get().GetUserForInputDevice(DeviceId));
-}
-
-FInputDevicePropertyHandle UInputDeviceSubsystem::ActivateDeviceProperty(UInputDeviceProperty* Property, const FSetDevicePropertyParams& Params)
+FInputDevicePropertyHandle UInputDeviceSubsystem::ActivateDeviceProperty(UInputDeviceProperty* Property, const FActivateDevicePropertyParams& Params)
 {
 	if (!Property)
 	{
@@ -403,7 +384,7 @@ FInputDevicePropertyHandle UInputDeviceSubsystem::ActivateDeviceProperty(UInputD
 	return OutHandle;
 }
 
-FInputDevicePropertyHandle UInputDeviceSubsystem::ActivateDevicePropertyOfClass(TSubclassOf<UInputDeviceProperty> PropertyClass, const FSetDevicePropertyParams& Params)
+FInputDevicePropertyHandle UInputDeviceSubsystem::ActivateDevicePropertyOfClass(TSubclassOf<UInputDeviceProperty> PropertyClass, const FActivateDevicePropertyParams& Params)
 {
 	if (!PropertyClass)
 	{
@@ -453,11 +434,6 @@ void UInputDeviceSubsystem::RemoveDevicePropertyHandles(const TSet<FInputDeviceP
 	}
 }
 
-bool UInputDeviceSubsystem::IsDevicePropertyHandleValid(const FInputDevicePropertyHandle& InHandle)
-{
-	return InHandle.IsValid();
-}
-
 void UInputDeviceSubsystem::RemoveAllDeviceProperties()
 {
 	for (FActiveDeviceProperty& ActiveProperty : ActiveProperties)
diff --git a/Engine/Source/Runtime/Engine/Private/InputDeviceLibrary.cpp b/Engine/Source/Runtime/Engine/Private/InputDeviceLibrary.cpp
--- a/Engine/Source/Runtime/Engine/Private/InputDeviceLibrary.cpp
+++ b/Engine/Source/Runtime/Engine/Private/InputDeviceLibrary.cpp
@@ -43,6 +43,11 @@ APlayerController* UInputDeviceLibrary::GetPlayerControllerFromInputDevice(const
 	return GetPlayerControllerFromPlatformUser(IPlatformInputDeviceMapper::Get().GetUserForInputDevice(DeviceId));
 }
 
+bool UInputDeviceLibrary::IsDevicePropertyHandleValid(const FInputDevicePropertyHandle& InHandle)
+{
+	return InHandle.IsValid();
+}
+
 int32 UInputDeviceLibrary::GetAllInputDevicesForUser(const FPlatformUserId UserId, TArray<FInputDeviceId>& OutInputDevices)
 {
 	return IPlatformInputDeviceMapper::Get().GetAllInputDevicesForUser(UserId, OutInputDevices);
