Change FActiveDeviceProperty::Property to be a weak object pointer instead of a hard reference. The property mayu get GC'd during a world transition, so the subsystem should only have a weak pointer to it so that it doesn't prevent a world from being unloaded.

diff --git a/Engine/Source/Runtime/Engine/Classes/GameFramework/InputDeviceSubsystem.h b/Engine/Source/Runtime/Engine/Classes/GameFramework/InputDeviceSubsystem.h
--- a/Engine/Source/Runtime/Engine/Classes/GameFramework/InputDeviceSubsystem.h
+++ b/Engine/Source/Runtime/Engine/Classes/GameFramework/InputDeviceSubsystem.h
@@ -68,7 +68,7 @@ struct ENGINE_API FActiveDeviceProperty
 
 	/** The active device property */
 	UPROPERTY()
-	TObjectPtr<UInputDeviceProperty> Property = nullptr;
+	TWeakObjectPtr<UInputDeviceProperty> Property = nullptr;
 
 	/** How long this property has been evaluated for. DeltaTime is added to this on tick */
 	double EvaluatedDuration = 0.0;
diff --git a/Engine/Source/Runtime/Engine/Private/GameFramework/InputDeviceSubsystem.cpp b/Engine/Source/Runtime/Engine/Private/GameFramework/InputDeviceSubsystem.cpp
--- a/Engine/Source/Runtime/Engine/Private/GameFramework/InputDeviceSubsystem.cpp
+++ b/Engine/Source/Runtime/Engine/Private/GameFramework/InputDeviceSubsystem.cpp
@@ -293,11 +293,12 @@ void UInputDeviceSubsystem::Tick(float InDeltaTime)
 	for (auto It = ActiveProperties.CreateIterator(); It; ++It)
 	{
 		FActiveDeviceProperty& ActiveProp = *It;
-		if (!ActiveProp.Property)
+		if (!ActiveProp.Property.IsValid())
 		{
-			// Something has gone wrong if we get here... maybe the Property has been GC'd ?
-			// This is really just an emergency handling case
-			ensure(false);
+			// The property may have been GC'd if the world has changed and it was an instanced property
+			// on an actor in that world
+			It.RemoveCurrent();
+			PropertiesPendingRemoval.Remove(ActiveProp.PropertyHandle);
 			continue;
 		}
 		
@@ -405,7 +406,7 @@ UInputDeviceProperty* UInputDeviceSubsystem::GetActiveDeviceProperty(const FInpu
 		// just returns its FInputDevicePropertyHandle's GetTypeHash
 		if (const FActiveDeviceProperty* ExistingProp = ActiveProperties.FindByHash(Handle.GetTypeHash(), Handle))
 		{
-			return ExistingProp->Property;
+			return ExistingProp->Property.Get();
 		}
 	}
 	
@@ -438,7 +439,7 @@ void UInputDeviceSubsystem::RemoveAllDeviceProperties()
 {
 	for (FActiveDeviceProperty& ActiveProperty : ActiveProperties)
 	{
-		if (ActiveProperty.Property)
+		if (ActiveProperty.Property.IsValid())
 		{
 			ActiveProperty.Property->ResetDeviceProperty(ActiveProperty.PlatformUser, ActiveProperty.DeviceId);
 		}		
