Fix clumpID visualization.

diff --git a/Engine/Shaders/Private/MaterialTemplate.ush b/Engine/Shaders/Private/MaterialTemplate.ush
--- a/Engine/Shaders/Private/MaterialTemplate.ush
+++ b/Engine/Shaders/Private/MaterialTemplate.ush
@@ -1497,7 +1497,7 @@ float MaterialExpressionGetHairSeed(FMaterialPixelParameters Parameters)
 #endif
 }
 
-uint MaterialExpressionGetHairClumpId(FMaterialPixelParameters Parameters)
+float MaterialExpressionGetHairClumpID(FMaterialPixelParameters Parameters)
 {
 #if HAIR_STRAND_MESH_FACTORY
 	return GetHairStrandsClumpID(Parameters.HairPrimitiveId, Parameters.HairPrimitiveUV);
diff --git a/Engine/Source/Runtime/Engine/Private/Materials/HLSLMaterialTranslator.cpp b/Engine/Source/Runtime/Engine/Private/Materials/HLSLMaterialTranslator.cpp
--- a/Engine/Source/Runtime/Engine/Private/Materials/HLSLMaterialTranslator.cpp
+++ b/Engine/Source/Runtime/Engine/Private/Materials/HLSLMaterialTranslator.cpp
@@ -9451,7 +9451,7 @@ int32 FHLSLMaterialTranslator::GetHairSeed()
 
 int32 FHLSLMaterialTranslator::GetHairClumpID()
 {
-	return AddCodeChunk(MCT_UInt1, TEXT("MaterialExpressionGetHairClumpID(Parameters)"));
+	return AddCodeChunk(MCT_Float1, TEXT("MaterialExpressionGetHairClumpID(Parameters)"));
 }
 
 int32 FHLSLMaterialTranslator::GetHairTangent(bool bUseTangentSpace)
