Fix crash when opening a material instance, whose parents/elders are null with Strata enabled.

diff --git a/Engine/Source/Runtime/Engine/Private/Materials/MaterialShared.cpp b/Engine/Source/Runtime/Engine/Private/Materials/MaterialShared.cpp
--- a/Engine/Source/Runtime/Engine/Private/Materials/MaterialShared.cpp
+++ b/Engine/Source/Runtime/Engine/Private/Materials/MaterialShared.cpp
@@ -4258,6 +4258,11 @@ int32 UMaterialInterface::CompileProperty(FMaterialCompiler* Compiler, EMaterial
 		Result = FMaterialAttributeDefinitionMap::CompileDefaultExpression(Compiler, Property);
 	}
 
+	if (Result == INDEX_NONE && Property == MP_FrontMaterial && Engine_IsStrataEnabled())
+	{
+		Result = Compiler->StrataCreateAndRegisterNullMaterial();
+	}
+
 	if (Result != INDEX_NONE)
 	{
 		// Cast is always required to go between float and LWC
