Strata - fix crash to due blend mode with Strata materials.

diff --git a/Engine/Shaders/Shared/StrataDefinitions.h b/Engine/Shaders/Shared/StrataDefinitions.h
--- a/Engine/Shaders/Shared/StrataDefinitions.h
+++ b/Engine/Shaders/Shared/StrataDefinitions.h
@@ -8,7 +8,7 @@
 #pragma once
 
 // Change this to force recompilation of all strata dependent shaders (use https://www.random.org/cgi-bin/randbyte?nbytes=4&format=h)
-#define STRATA_SHADER_VERSION 0xcb0f647e1
+#define STRATA_SHADER_VERSION 0x084bf26a2 
 
 
 
diff --git a/Engine/Source/Runtime/Engine/Private/Materials/MaterialInterface.cpp b/Engine/Source/Runtime/Engine/Private/Materials/MaterialInterface.cpp
--- a/Engine/Source/Runtime/Engine/Private/Materials/MaterialInterface.cpp
+++ b/Engine/Source/Runtime/Engine/Private/Materials/MaterialInterface.cpp
@@ -394,7 +394,7 @@ FMaterialRelevance UMaterialInterface::GetRelevance_Internal(const UMaterial* Ma
 
 		const bool bIsMobile = InFeatureLevel <= ERHIFeatureLevel::ES3_1;
 		const bool bUsesSingleLayerWaterMaterial = MaterialResource->GetShadingModels().HasShadingModel(MSM_SingleLayerWater);
-		const bool IsSinglePassWaterTranslucent = bIsMobile && bUsesSingleLayerWaterMaterial;
+		const bool bIsSinglePassWaterTranslucent = bIsMobile && bUsesSingleLayerWaterMaterial;
 		const bool bIsMobilePixelProjectedTranslucent = MaterialResource->IsUsingPlanarForwardReflections() 
 														&& IsUsingMobilePixelProjectedReflection(GetFeatureLevelShaderPlatform(InFeatureLevel));
 
@@ -404,7 +404,7 @@ FMaterialRelevance UMaterialInterface::GetRelevance_Internal(const UMaterial* Ma
 			MaterialResource->MaterialUsesAnisotropy_GameThread();
 
 		const EBlendMode BlendMode = (EBlendMode)GetBlendMode();
-		const bool bIsTranslucent = IsTranslucentBlendMode(BlendMode) || IsSinglePassWaterTranslucent || bIsMobilePixelProjectedTranslucent; // We want meshes with water materials to be scheduled for translucent pass on mobile. And we also have to render the meshes used for mobile pixel projection reflection in translucent pass.
+		const bool bIsTranslucent = IsTranslucentBlendMode(BlendMode) || bIsSinglePassWaterTranslucent || bIsMobilePixelProjectedTranslucent; // We want meshes with water materials to be scheduled for translucent pass on mobile. And we also have to render the meshes used for mobile pixel projection reflection in translucent pass.
 
 		EMaterialDomain Domain = (EMaterialDomain)MaterialResource->GetMaterialDomain();
 		bool bDecal = (Domain == MD_DeferredDecal);
diff --git a/Engine/Source/Runtime/Engine/Private/Materials/MaterialShared.cpp b/Engine/Source/Runtime/Engine/Private/Materials/MaterialShared.cpp
--- a/Engine/Source/Runtime/Engine/Private/Materials/MaterialShared.cpp
+++ b/Engine/Source/Runtime/Engine/Private/Materials/MaterialShared.cpp
@@ -2291,9 +2291,14 @@ void FMaterial::SetupMaterialEnvironment(
 			OutEnvironment.SetDefine(TEXT("MATERIALBLENDING_TRANSLUCENT"), TEXT("1"));
 			break;
 		}
-		case BLEND_Translucent: OutEnvironment.SetDefine(TEXT("MATERIALBLENDING_TRANSLUCENT"),TEXT("1")); break;
+		case BLEND_TranslucentColoredTransmittance:
+		case BLEND_Translucent:
+			OutEnvironment.SetDefine(TEXT("MATERIALBLENDING_TRANSLUCENT"),TEXT("1")); 
+			break;
+
 		case BLEND_Additive: OutEnvironment.SetDefine(TEXT("MATERIALBLENDING_ADDITIVE"),TEXT("1")); break;
 		case BLEND_Modulate: OutEnvironment.SetDefine(TEXT("MATERIALBLENDING_MODULATE"),TEXT("1")); break;
+
 		default: 
 			UE_LOG(LogMaterial, Warning, TEXT("Unknown material blend mode: %u  Setting to BLEND_Opaque"),(int32)GetBlendMode());
 			OutEnvironment.SetDefine(TEXT("MATERIALBLENDING_SOLID"),TEXT("1"));
diff --git a/Engine/Source/Runtime/Renderer/Private/BasePassRendering.cpp b/Engine/Source/Runtime/Renderer/Private/BasePassRendering.cpp
--- a/Engine/Source/Runtime/Renderer/Private/BasePassRendering.cpp
+++ b/Engine/Source/Runtime/Renderer/Private/BasePassRendering.cpp
@@ -266,8 +266,7 @@ void SetTranslucentRenderState(FMeshPassProcessorRenderState& DrawRenderState, c
 			}
 			else
 			{
-				// STRATA_TODO_BLENDMODE_ADDITIVE
-				// We always use premultipled alpha for translucent rendering.
+				// We always use premultipled alpha for translucent rendering, that works for any surface with additive color.
 				// If a material was requesting dual source blending, the shader will use static platofm knowledge to convert colored transmittance to a grey scale transmittance.
 				DrawRenderState.SetBlendState(TStaticBlendState<CW_RGBA, BO_Add, BF_One, BF_InverseSourceAlpha, BO_Add, BF_Zero, BF_InverseSourceAlpha>::GetRHI());
 			}
@@ -324,6 +323,7 @@ void SetTranslucentRenderState(FMeshPassProcessorRenderState& DrawRenderState, c
 			// Masked materials are rendered together in the base pass, where the blend state is set at a higher level
 			break;
 		case BLEND_Translucent:
+		case BLEND_TranslucentColoredTransmittance:	// When Strata is disabled, this falls back to simple Translucency.
 			// Note: alpha channel used by separate translucency, storing how much of the background should be added when doing the final composite
 			// The Alpha channel is also used by non-separate translucency when rendering to scene captures, which store the final opacity
 			DrawRenderState.SetBlendState(TStaticBlendState<CW_RGBA, BO_Add, BF_SourceAlpha, BF_InverseSourceAlpha, BO_Add, BF_Zero, BF_InverseSourceAlpha>::GetRHI());
