Fixing bone rendering issue with being unable to visualize connection to parent joint when using "Show Selected" filter.

diff --git a/Engine/Plugins/Animation/ControlRig/Source/ControlRigEditor/Private/Editor/ControlRigEditorEditMode.cpp b/Engine/Plugins/Animation/ControlRig/Source/ControlRigEditor/Private/Editor/ControlRigEditorEditMode.cpp
--- a/Engine/Plugins/Animation/ControlRig/Source/ControlRigEditor/Private/Editor/ControlRigEditorEditMode.cpp
+++ b/Engine/Plugins/Animation/ControlRig/Source/ControlRigEditor/Private/Editor/ControlRigEditorEditMode.cpp
@@ -209,7 +209,7 @@ void FControlRigEditorEditMode::Render(const FSceneView* View, FViewport* Viewpo
 			bDrawAxesInsideBone);
 		if (Hierarchy->GetFirstParent(ElementIndex) == INDEX_NONE)
 		{
-			SkeletalDebugRendering::DrawRootCone(PDI, BoneTransform, FVector::Zero(), BoneRadius);
+			SkeletalDebugRendering::DrawConeConnection(PDI, BoneTransform.GetLocation(), FVector::Zero(), BoneRadius, FLinearColor::Red);
 		}
 		PDI->SetHitProxy(nullptr);
 	}
diff --git a/Engine/Source/Runtime/Engine/Private/SkeletalDebugRendering.cpp b/Engine/Source/Runtime/Engine/Private/SkeletalDebugRendering.cpp
--- a/Engine/Source/Runtime/Engine/Private/SkeletalDebugRendering.cpp
+++ b/Engine/Source/Runtime/Engine/Private/SkeletalDebugRendering.cpp
@@ -138,20 +138,20 @@ void DrawAxes(
 #endif
 }
 
-void DrawRootCone(
+void DrawConeConnection(
 	FPrimitiveDrawInterface* PDI,
-	const FTransform& InBoneTransform,
-	const FVector& ComponentOrigin,
-	const float SphereRadius)
+	const FVector& Start,
+	const FVector& End,
+	const float SphereRadius,
+	const FLinearColor& Color)
 {
 #if ENABLE_DRAW_DEBUG
 	// offset start/end based on bone radius
-	const FVector RadiusOffset = (ComponentOrigin - InBoneTransform.GetLocation()).GetSafeNormal() * SphereRadius;
-	const FVector Start = InBoneTransform.GetLocation() + RadiusOffset;
-	const FVector End = ComponentOrigin;
+	const FVector RadiusOffset = (End - Start).GetSafeNormal() * SphereRadius;
+	const FVector StartOffset = Start + RadiusOffset;
 			
 	// calc cone size
-	const FVector EndToStart = (Start - End);
+	const FVector EndToStart = (StartOffset - End);
 	const float ConeLength = EndToStart.Size();
 	const float Angle = FMath::RadiansToDegrees(FMath::Atan(SphereRadius / ConeLength));
 	TArray<FVector> Verts;
@@ -162,7 +162,7 @@ void DrawRootCone(
 		ConeLength,
 		Angle,
 		NumConeSides,
-		FLinearColor::Red,
+		Color,
 		ESceneDepthPriorityGroup::SDPG_Foreground,
 		0.0f,
 		1.0f);
@@ -254,9 +254,9 @@ void DrawBones(
 	const TArray<TRefCountPtr<HHitProxy>>& HitProxies,
 	const FSkelDebugDrawConfig& DrawConfig)
 {
+	// get parent indices of bones
 	TArray<int32> ParentIndices;
 	ParentIndices.AddUninitialized(RefSkeleton.GetNum());
-
 	for (int32 BoneIndex = 0; BoneIndex < RefSkeleton.GetNum(); ++BoneIndex)
 	{
 		ParentIndices[BoneIndex] = RefSkeleton.GetParentIndex(BoneIndex);
@@ -297,11 +297,18 @@ void DrawBonesInternal(
 	// first determine which bones to draw, and which to filter out
 	const int32 NumBones = ParentIndices.Num();
 	TBitArray<> BonesToDraw(false, NumBones);
+	const bool bDrawAll = DrawConfig.BoneDrawMode == EBoneDrawMode::All;
 	const bool bDrawSelected = DrawConfig.BoneDrawMode == EBoneDrawMode::Selected;
 	const bool bDrawSelectedAndParents = DrawConfig.BoneDrawMode == EBoneDrawMode::SelectedAndParents;
 	const bool bDrawSelectedAndChildren = DrawConfig.BoneDrawMode == EBoneDrawMode::SelectedAndChildren;
 	const bool bDrawSelectedAndParentsAndChildren = DrawConfig.BoneDrawMode == EBoneDrawMode::SelectedAndParentsAndChildren;
 
+	// draw all bones
+	if (bDrawAll)
+	{
+		BonesToDraw.Init(true, NumBones);
+	}
+
 	// add selected bones
 	if (bDrawSelected || bDrawSelectedAndParents || bDrawSelectedAndChildren || bDrawSelectedAndParentsAndChildren)
 	{
@@ -368,7 +375,7 @@ void DrawBonesInternal(
 		}
 
 		// skips bones that should not be drawn
-		const bool bDoDraw = DrawConfig.bForceDraw || DrawConfig.BoneDrawMode == EBoneDrawMode::All || BonesToDraw[BoneIndex];
+		const bool bDoDraw = DrawConfig.bForceDraw || BonesToDraw[BoneIndex];
 		if (!bDoDraw)
 		{
 			continue;
@@ -411,6 +418,8 @@ void DrawBonesInternal(
 
 		// Always set new hit proxy to prevent unintentionally using last drawn element's proxy
 		PDI->SetHitProxy(DrawConfig.bAddHitProxy ? HitProxies[BoneIndex] : nullptr);
+
+		// draw skeleton
 		SkeletalDebugRendering::DrawWireBoneAdvanced(
 			PDI,
 			BoneTransform,
@@ -420,10 +429,31 @@ void DrawBonesInternal(
 			SDPG_Foreground,
 			BoneRadius,
 			bDrawAxesInsideBone);
+		
+		// special case for root connection to origin
 		if (GetParentIndex(BoneIndex) == INDEX_NONE)
 		{
-			SkeletalDebugRendering::DrawRootCone(PDI, BoneTransform, ComponentOrigin, BoneRadius);
+			SkeletalDebugRendering::DrawConeConnection(PDI, BoneTransform.GetLocation(), ComponentOrigin, BoneRadius, FLinearColor::Red);
+		}
+		
+		// special case for forcing drawing connection to parent when:
+		// 1. parent is not selected AND
+		// 2. only drawing selected or children of selected bones
+		// In this case, the connection to the parent will not get drawn unless we force it here
+		if (bDrawSelected || bDrawSelectedAndChildren)
+		{
+			if (InSelectedBones.Contains(BoneIndex))
+			{
+				const int32 ParentIndex = GetParentIndex(BoneIndex);
+				if (!WorldTransforms.IsValidIndex(ParentIndex))
+				{
+					continue;
+				}
+				const FVector ParentPosition = WorldTransforms[ParentIndex].GetTranslation();
+				SkeletalDebugRendering::DrawConeConnection(PDI, ParentPosition, BoneTransform.GetLocation(), BoneRadius, DrawConfig.ParentOfSelectedBoneColor);
+			}
 		}
+		
 		PDI->SetHitProxy(nullptr);
 	}
 }
diff --git a/Engine/Source/Runtime/Engine/Public/SkeletalDebugRendering.h b/Engine/Source/Runtime/Engine/Public/SkeletalDebugRendering.h
--- a/Engine/Source/Runtime/Engine/Public/SkeletalDebugRendering.h
+++ b/Engine/Source/Runtime/Engine/Public/SkeletalDebugRendering.h
@@ -98,17 +98,20 @@ ENGINE_API	void DrawWireBoneAdvanced(
 	const bool bDrawAxes);
 
 /**
- * Draw a red cone showing offset of root bone from the component origin
- * @param	PDI					Primitive draw interface to use
- * @param	InBoneTransform		The bone transform
- * @param	ComponentOrigin		The position in world space of the component this bone lives within
- * @param	SphereRadius		The radius of the root bone
+ * Draw a cone showing offset from origin position to a given bone transform
+ * Used to draw the root cone (always in red)
+ * @param	PDI				Primitive draw interface to use
+ * @param	Start			The position in world space of the start of the cone
+ * @param	End				The position in world space of the tapered end of the cone
+ * @param	SphereRadius	The radius of the root bone
+ * @param	Color			The color to use for the cone
  */
-ENGINE_API	void DrawRootCone(
+ENGINE_API	void DrawConeConnection(
 	FPrimitiveDrawInterface* PDI,
-	const FTransform& InBoneTransform,
-	const FVector& ComponentOrigin,
-	const float SphereRadius);
+	const FVector& Start,
+	const FVector& End,
+	const float SphereRadius,
+	const FLinearColor& Color);
 
 #if WITH_EDITOR
 ENGINE_API void DrawBonesFromPoseWatch(
