WorldPartition's ActorDescContainer is now transient (no reason to persist the container)

diff --git a/Engine/Source/Runtime/Engine/Public/WorldPartition/WorldPartition.h b/Engine/Source/Runtime/Engine/Public/WorldPartition/WorldPartition.h
--- a/Engine/Source/Runtime/Engine/Public/WorldPartition/WorldPartition.h
+++ b/Engine/Source/Runtime/Engine/Public/WorldPartition/WorldPartition.h
@@ -315,7 +315,7 @@ private:
 public:
 	UActorDescContainer* GetActorDescContainer() const { return ActorDescContainer; }
 
-	UPROPERTY()
+	UPROPERTY(Transient)
 	TObjectPtr<UActorDescContainer> ActorDescContainer;
 
 	UPROPERTY()
