Allow creation of WorldPartitionSubsystem to be created for inactive worlds.
This will be necessary for DataLayerManager to be properly initialized and be used while duplicating partitioned worlds.

diff --git a/Engine/Source/Runtime/Engine/Private/WorldPartition/WorldPartitionSubsystem.cpp b/Engine/Source/Runtime/Engine/Private/WorldPartition/WorldPartitionSubsystem.cpp
--- a/Engine/Source/Runtime/Engine/Private/WorldPartition/WorldPartitionSubsystem.cpp
+++ b/Engine/Source/Runtime/Engine/Private/WorldPartition/WorldPartitionSubsystem.cpp
@@ -163,6 +163,11 @@ void UWorldPartitionSubsystem::Deinitialize()
 	Super::Deinitialize();
 }
 
+bool UWorldPartitionSubsystem::DoesSupportWorldType(const EWorldType::Type WorldType) const
+{
+	return Super::DoesSupportWorldType(WorldType) || WorldType == EWorldType::Inactive;
+}
+
 #if WITH_EDITOR
 void UWorldPartitionSubsystem::ForEachWorldPartition(TFunctionRef<bool(UWorldPartition*)> Func)
 {
diff --git a/Engine/Source/Runtime/Engine/Public/WorldPartition/WorldPartitionSubsystem.h b/Engine/Source/Runtime/Engine/Public/WorldPartition/WorldPartitionSubsystem.h
--- a/Engine/Source/Runtime/Engine/Public/WorldPartition/WorldPartitionSubsystem.h
+++ b/Engine/Source/Runtime/Engine/Public/WorldPartition/WorldPartitionSubsystem.h
@@ -20,6 +20,9 @@ class ENGINE_API UWorldPartitionSubsystem : public UTickableWorldSubsystem
 {
 	GENERATED_BODY()
 
+protected:
+	virtual bool DoesSupportWorldType(const EWorldType::Type WorldType) const override;
+
 public:
 	UWorldPartitionSubsystem();
 
