Unify logging messages of Content Bundle
Clarified which world is being injected
Add some more logging

diff --git a/Engine/Source/Runtime/Engine/Private/WorldPartition/ContentBundle/ContentBundle.cpp b/Engine/Source/Runtime/Engine/Private/WorldPartition/ContentBundle/ContentBundle.cpp
--- a/Engine/Source/Runtime/Engine/Private/WorldPartition/ContentBundle/ContentBundle.cpp
+++ b/Engine/Source/Runtime/Engine/Private/WorldPartition/ContentBundle/ContentBundle.cpp
@@ -36,12 +36,12 @@ void FContentBundle::DoInitialize()
 		}
 		else
 		{
-			UE_LOG(LogContentBundle, Error, TEXT("[CB: %s] No streaming object found in package %s."), *GetDescriptor()->GetDisplayName(), *GetExternalStreamingObjectPackagePath());
+			UE_LOG(LogContentBundle, Error, TEXT("%s No streaming object found in package %s."), *ContentBundle::Log::MakeDebugInfoString(*this), *GetExternalStreamingObjectPackagePath());
 		}
 	}
 	else
 	{
-		UE_LOG(LogContentBundle, Log, TEXT("[CB: %s] No streaming object found. No content will be injected."), *GetDescriptor()->GetDisplayName());
+		UE_LOG(LogContentBundle, Log, TEXT("%s No streaming object found. No content will be injected."), *ContentBundle::Log::MakeDebugInfoString(*this));
 	}
 #endif
 
@@ -62,20 +62,20 @@ void FContentBundle::DoInjectContent()
 	{
 		if (GetInjectedWorld()->GetWorldPartition()->RuntimeHash->InjectExternalStreamingObject(ExternalStreamingObject))
 		{
-			UE_LOG(LogContentBundle, Log, TEXT("[CB: %s] Streaming Object Injected."), *GetDescriptor()->GetDisplayName());
+			UE_LOG(LogContentBundle, Log, TEXT("%s Streaming Object Injected."), *ContentBundle::Log::MakeDebugInfoString(*this));
 			SetStatus(EContentBundleStatus::ContentInjected);
 
 			ContentBundlesEpoch++;
 		}
 		else
 		{
-			UE_LOG(LogContentBundle, Error, TEXT("[CB: %s] Failed to inject streaming object."), *GetDescriptor()->GetDisplayName());
+			UE_LOG(LogContentBundle, Error, TEXT("%s Failed to inject streaming object."), *ContentBundle::Log::MakeDebugInfoString(*this));
 			SetStatus(EContentBundleStatus::FailedToInject);
 		}
 	}
 	else
 	{
-		UE_LOG(LogContentBundle, Log, TEXT("[CB: %s] No streaming object to inject."), *GetDescriptor()->GetDisplayName());
+		UE_LOG(LogContentBundle, Log, TEXT("%s No streaming object to inject."), *ContentBundle::Log::MakeDebugInfoString(*this));
 		SetStatus(EContentBundleStatus::ContentInjected);
 	}
 }
@@ -90,7 +90,7 @@ void FContentBundle::DoRemoveContent()
 		}
 		else
 		{
-			UE_LOG(LogContentBundle, Error, TEXT("[CB: %s] Error while removing streaming object."), *GetDescriptor()->GetDisplayName());
+			UE_LOG(LogContentBundle, Error, TEXT("%s Error while removing streaming object."), *ContentBundle::Log::MakeDebugInfoString(*this));
 		}
 	}
 	
@@ -119,7 +119,7 @@ void FContentBundle::InitializeForPIE()
 		ExternalStreamingObject = PIEHelper->RetrieveContentBundleStreamingObject(*this);
 		if (ExternalStreamingObject == nullptr)
 		{
-			UE_LOG(LogContentBundle, Log, TEXT("[CB: %s] No streaming object found. There are %u existing streaming objects."), *GetDescriptor()->GetDisplayName(), PIEHelper->GetStreamingObjectCount());
+			UE_LOG(LogContentBundle, Log, TEXT("%s No streaming object found. There are %u existing streaming objects."), *ContentBundle::Log::MakeDebugInfoString(*this), PIEHelper->GetStreamingObjectCount());
 		}
 	}
 }
diff --git a/Engine/Source/Runtime/Engine/Private/WorldPartition/ContentBundle/ContentBundleBase.cpp b/Engine/Source/Runtime/Engine/Private/WorldPartition/ContentBundle/ContentBundleBase.cpp
--- a/Engine/Source/Runtime/Engine/Private/WorldPartition/ContentBundle/ContentBundleBase.cpp
+++ b/Engine/Source/Runtime/Engine/Private/WorldPartition/ContentBundle/ContentBundleBase.cpp
@@ -131,6 +131,7 @@ void FContentBundleBase::SetStatus(EContentBundleStatus NewStatus)
 {
 	check(NewStatus != Status);
 
-	UE_LOG(LogContentBundle, Log, TEXT("[CB: %s] State changing from %s to %s"), *GetDescriptor()->GetDisplayName(), *UEnum::GetDisplayValueAsText(Status).ToString(), *UEnum::GetDisplayValueAsText(NewStatus).ToString());
+	UE_LOG(LogContentBundle, Log, TEXT("%s State changing from %s to %s"), 
+		*ContentBundle::Log::MakeDebugInfoString(*this), *UEnum::GetDisplayValueAsText(Status).ToString(), *UEnum::GetDisplayValueAsText(NewStatus).ToString());
 	Status = NewStatus;
 }
diff --git a/Engine/Source/Runtime/Engine/Private/WorldPartition/ContentBundle/ContentBundleClient.cpp b/Engine/Source/Runtime/Engine/Private/WorldPartition/ContentBundle/ContentBundleClient.cpp
--- a/Engine/Source/Runtime/Engine/Private/WorldPartition/ContentBundle/ContentBundleClient.cpp
+++ b/Engine/Source/Runtime/Engine/Private/WorldPartition/ContentBundle/ContentBundleClient.cpp
@@ -132,12 +132,12 @@ void FContentBundleClient::SetWorldContentState(UWorld* World, EWorldContentStat
 {
 	if (EWorldContentState* OldState = WorldContentStates.Find(World))
 	{
-		UE_LOG(LogContentBundle, Log, TEXT("[CB: %s] Client WorldState for world %s changing from %s to %s"), *GetDescriptor()->GetDisplayName(), *World->GetName(), *UEnum::GetDisplayValueAsText(*OldState).ToString(), *UEnum::GetDisplayValueAsText(NewState).ToString());
+		UE_LOG(LogContentBundle, Log, TEXT("%s Client WorldState changing from %s to %s"), *ContentBundle::Log::MakeDebugInfoString(*this, World), *UEnum::GetDisplayValueAsText(*OldState).ToString(), *UEnum::GetDisplayValueAsText(NewState).ToString());
 		*OldState = NewState;
 		return;
 	}
 	
-	UE_LOG(LogContentBundle, Log, TEXT("[CB: %s] Client WorldState for world %s changing to %s"), *GetDescriptor()->GetDisplayName(), *World->GetName(), *UEnum::GetDisplayValueAsText(NewState).ToString());
+	UE_LOG(LogContentBundle, Log, TEXT("%s Client WorldState changing to %s"), *ContentBundle::Log::MakeDebugInfoString(*this, World), *UEnum::GetDisplayValueAsText(NewState).ToString());
 	WorldContentStates.Add(World, NewState);
 }
 
@@ -176,7 +176,7 @@ void FContentBundleClient::SetState(EContentBundleClientState NewState)
 {
 	check(NewState != State);
 
-	UE_LOG(LogContentBundle, Log, TEXT("[CB: %s] Client State changing from %s to %s"), *GetDescriptor()->GetDisplayName(), *UEnum::GetDisplayValueAsText(State).ToString(), *UEnum::GetDisplayValueAsText(NewState).ToString());
+	UE_LOG(LogContentBundle, Log, TEXT("%s Client State changing from %s to %s"), *ContentBundle::Log::MakeDebugInfoString(*this), *UEnum::GetDisplayValueAsText(State).ToString(), *UEnum::GetDisplayValueAsText(NewState).ToString());
 	State = NewState;
 }
 
@@ -186,6 +186,8 @@ void FContentBundleClient::RequestForceInject(UWorld* WorldToInject)
 {
 	if (UContentBundleManager* ContentBundleManager = WorldToInject->ContentBundleManager)
 	{
+		UE_LOG(LogContentBundle, Log, TEXT("%s Client requested a forced injection in world."), *ContentBundle::Log::MakeDebugInfoString(*this, WorldToInject));
+
 		ForceInjectedWorlds.Add(WorldToInject);
 	
 		if (ShouldInjectContent(WorldToInject))
@@ -193,12 +195,18 @@ void FContentBundleClient::RequestForceInject(UWorld* WorldToInject)
 			ContentBundleManager->TryInject(*this);
 		}
 	}
+	else
+	{
+		UE_LOG(LogContentBundle, Log, TEXT("%s Client cannot force inject in World. It does not support content bundles."), *ContentBundle::Log::MakeDebugInfoString(*this, WorldToInject));
+	}
 }
 
 void FContentBundleClient::RequestRemoveForceInjectedContent(UWorld* WorldToInject)
 {
 	if (UContentBundleManager* ContentBundleManager = WorldToInject->ContentBundleManager)
 	{
+		UE_LOG(LogContentBundle, Log, TEXT("%s Client requested removal of force injected content in world."), *ContentBundle::Log::MakeDebugInfoString(*this, WorldToInject));
+
 		ForceInjectedWorlds.Remove(WorldToInject);
 
 		if (ShouldRemoveContent(WorldToInject))
@@ -206,6 +214,10 @@ void FContentBundleClient::RequestRemoveForceInjectedContent(UWorld* WorldToInje
 			ContentBundleManager->Remove(*this);
 		}
 	}
+	else
+	{
+		UE_LOG(LogContentBundle, Log, TEXT("%s Client cannot remove force injected content from world. It does not support content bundles."), *ContentBundle::Log::MakeDebugInfoString(*this, WorldToInject));
+	}
 }
 
 #endif
diff --git a/Engine/Source/Runtime/Engine/Private/WorldPartition/ContentBundle/ContentBundleContainer.cpp b/Engine/Source/Runtime/Engine/Private/WorldPartition/ContentBundle/ContentBundleContainer.cpp
--- a/Engine/Source/Runtime/Engine/Private/WorldPartition/ContentBundle/ContentBundleContainer.cpp
+++ b/Engine/Source/Runtime/Engine/Private/WorldPartition/ContentBundle/ContentBundleContainer.cpp
@@ -35,7 +35,7 @@ UWorld* FContentBundleContainer::GetInjectedWorld() const
 void FContentBundleContainer::Initialize()
 {
 	TRACE_CPUPROFILER_EVENT_SCOPE(FContentBundleContainer::Initialize);
-	UE_LOG(LogContentBundle, Log, TEXT("[Container: %s] Creating new contrainer."), *GetInjectedWorld()->GetName());
+	UE_LOG(LogContentBundle, Log, TEXT("%s Creating new contrainer."), *ContentBundle::Log::MakeDebugInfoString(*this));
 
 #if WITH_EDITOR
 	if (UseEditorContentBundle())
@@ -64,7 +64,7 @@ void FContentBundleContainer::Deinitialize()
 
 	if (GetInjectedWorld())
 	{
-		UE_LOG(LogContentBundle, Log, TEXT("[Container: %s] Deleting container."), *GetInjectedWorld()->GetName());
+		UE_LOG(LogContentBundle, Log, TEXT("%s Deleting container."), *ContentBundle::Log::MakeDebugInfoString(*this));
 
 #if WITH_EDITOR
 		UWorldPartition* WorldPartition = GetInjectedWorld()->GetWorldPartition();
@@ -228,8 +228,8 @@ FContentBundleBase& FContentBundleContainer::InitializeContentBundle(TSharedPtr<
 
 	FContentBundleBase* ContentBundle = nullptr;
 
-	UE_LOG(LogContentBundle, Log, TEXT("[Container: %s] Creating new ContentBundle %s from client %s with state %s."), 
-		*GetInjectedWorld()->GetName(), *ContentBundleClient->GetDescriptor()->GetDisplayName(), *ContentBundleClient->GetDisplayName(), *UEnum::GetDisplayValueAsText(ContentBundleClient->GetState()).ToString());
+	UE_LOG(LogContentBundle, Log, TEXT("%s Creating new content bundle from client %s with client state %s."), 
+		*ContentBundle::Log::MakeDebugInfoString(*ContentBundleClient , GetInjectedWorld()), *ContentBundleClient->GetDisplayName(), *UEnum::GetDisplayValueAsText(ContentBundleClient->GetState()).ToString());
 
 #if WITH_EDITOR
 	if (UseEditorContentBundle())
@@ -273,6 +273,8 @@ bool FContentBundleContainer::InjectContentBundle(FContentBundleClient& ContentB
 		return InjectContentBundle(*ContentBundle);
 	}
 
+	UE_LOG(LogContentBundle, Log, TEXT("%s Failed to inject content bundle from client %s. It was not found in world."),
+		*ContentBundle::Log::MakeDebugInfoString(ContentBundleClient, InjectedWorld), *ContentBundleClient.GetDisplayName());
 	return false;
 }
 
@@ -298,6 +300,8 @@ bool FContentBundleContainer::RemoveContentBundle(FContentBundleClient& ContentB
 		return RemoveContentBundle(*ContentBundle);
 	}
 
+	UE_LOG(LogContentBundle, Log, TEXT("%s Failed to remove content bundle from client %s. It was not found in world."),
+		*ContentBundle::Log::MakeDebugInfoString(ContentBundleClient, InjectedWorld), *ContentBundleClient.GetDisplayName());
 	return false;
 }
 
@@ -311,6 +315,11 @@ bool FContentBundleContainer::RemoveContentBundle(FContentBundleBase& ContentBun
 			ContentBundle.RemoveContent();
 			return true;
 		}
+		else
+		{
+			UE_LOG(LogContentBundle, Log, TEXT("%s Client %s determined content bundle will not be removed from world."),
+				*ContentBundle::Log::MakeDebugInfoString(ContentBundle), *ContentBundleClient->GetDisplayName());
+		}
 	}
 	
 	return false;
@@ -364,7 +373,7 @@ void FContentBundleContainer::InitializeContentBundlesForegisteredClients()
 
 	if (!ContentBundleClients.IsEmpty())
 	{
-		UE_LOG(LogContentBundle, Log, TEXT("[Container: %s] Begin initializing ContentBundles from %u registered clients."), *GetInjectedWorld()->GetName(), ContentBundleClients.Num());
+		UE_LOG(LogContentBundle, Log, TEXT("%s Begin initializing ContentBundles from %u registered clients."), *ContentBundle::Log::MakeDebugInfoString(*this), ContentBundleClients.Num());
 
 		for (TSharedPtr<FContentBundleClient>& ContentBundleClient : ContentBundleClients)
 		{
@@ -372,7 +381,7 @@ void FContentBundleContainer::InitializeContentBundlesForegisteredClients()
 			InjectContentBundle(ContentBundle);
 		}
 
-		UE_LOG(LogContentBundle, Log, TEXT("[Container: %s] End initializing ContentBundles."), *GetInjectedWorld()->GetName());
+		UE_LOG(LogContentBundle, Log, TEXT("%s End initializing ContentBundles."), *ContentBundle::Log::MakeDebugInfoString(*this));
 	}
 }
 
@@ -395,7 +404,7 @@ void FContentBundleContainer::DeinitializeContentBundles()
 
 void FContentBundleContainer::OnPreGenerateStreaming(TArray<FString>* OutPackageToGenerate)
 {
-	UE_LOG(LogContentBundle, Log, TEXT("[Container: %s] Generating Streaming for %u Content Bundles."), *GetInjectedWorld()->GetName(), GetEditorContentBundles().Num());
+	UE_LOG(LogContentBundle, Log, TEXT("%s Generating Streaming for %u Content Bundles."), *ContentBundle::Log::MakeDebugInfoString(*this), GetEditorContentBundles().Num());
 
 	if (IsRunningCookCommandlet())
 	{
diff --git a/Engine/Source/Runtime/Engine/Private/WorldPartition/ContentBundle/ContentBundleEditor.cpp b/Engine/Source/Runtime/Engine/Private/WorldPartition/ContentBundle/ContentBundleEditor.cpp
--- a/Engine/Source/Runtime/Engine/Private/WorldPartition/ContentBundle/ContentBundleEditor.cpp
+++ b/Engine/Source/Runtime/Engine/Private/WorldPartition/ContentBundle/ContentBundleEditor.cpp
@@ -60,7 +60,7 @@ void FContentBundleEditor::DoInjectContent()
 		ActorDescContainer = WorldPartition->RegisterActorDescContainer(FName(*ActorDescContainerPackage));
 		if (ActorDescContainer.IsValid())
 		{
-			UE_LOG(LogContentBundle, Log, TEXT("[CB: %s] ExternalActors in %s found. %u actors were injected"), *GetDescriptor()->GetDisplayName(), *ActorDescContainer->GetExternalActorPath(), ActorDescContainer->GetActorDescCount());
+			UE_LOG(LogContentBundle, Log, TEXT("%s ExternalActors in %s found. %u actors were injected"), *ContentBundle::Log::MakeDebugInfoString(*this), *ActorDescContainer->GetExternalActorPath(), ActorDescContainer->GetActorDescCount());
 
 			check(GetDescriptor()->GetGuid().IsValid());
 			ActorDescContainer->SetContentBundleGuid(GetDescriptor()->GetGuid());
@@ -79,13 +79,13 @@ void FContentBundleEditor::DoInjectContent()
 		}
 		else
 		{
-			UE_LOG(LogContentBundle, Log, TEXT("[CB: %s] Failed to register actor desc container with %s"), *GetDescriptor()->GetDisplayName(), *ActorDescContainerPackage);
+			UE_LOG(LogContentBundle, Log, TEXT("%s Failed to register actor desc container with %s"), *ContentBundle::Log::MakeDebugInfoString(*this), *ActorDescContainerPackage);
 			SetStatus(EContentBundleStatus::FailedToInject);
 		}
 	}
 	else
 	{
-		UE_LOG(LogContentBundle, Error, TEXT("[CB: %s] Failed to build Container Package Path using %s"), *GetDescriptor()->GetDisplayName(), *GetInjectedWorld()->GetPackage()->GetName());
+		UE_LOG(LogContentBundle, Error, TEXT("%s Failed to build Container Package Path using %s"), *ContentBundle::Log::MakeDebugInfoString(*this), *GetInjectedWorld()->GetPackage()->GetName());
 		SetStatus(EContentBundleStatus::FailedToInject);
 	}
 
@@ -150,7 +150,7 @@ bool FContentBundleEditor::AddActor(AActor* InActor)
 
 	if (GetStatus() == EContentBundleStatus::ReadyToInject)
 	{
-		UE_LOG(LogContentBundle, Verbose, TEXT("[CB: %s] Adding first actor to content bundle."), *GetDescriptor()->GetDisplayName());
+		UE_LOG(LogContentBundle, Verbose, TEXT("%s Adding first actor to content bundle."), *ContentBundle::Log::MakeDebugInfoString(*this));
 		SetStatus(EContentBundleStatus::ContentInjected);
 	}
 
@@ -167,7 +167,7 @@ bool FContentBundleEditor::AddActor(AActor* InActor)
 
 	UnsavedActorMonitor->MonitorActor(InActor);
 
-	UE_LOG(LogContentBundle, Verbose, TEXT("[CB: %s] Added new actor %s, ActorCount:  %u. Package %s."), *GetDescriptor()->GetDisplayName(), *InActor->GetActorNameOrLabel(), GetActorCount(), *InActor->GetPackage()->GetName());
+	UE_LOG(LogContentBundle, Verbose, TEXT("%s Added new actor %s, ActorCount:  %u. Package %s."), *ContentBundle::Log::MakeDebugInfoString(*this), *InActor->GetActorNameOrLabel(), GetActorCount(), *InActor->GetPackage()->GetName());
 
 	return true;
 }
@@ -276,7 +276,7 @@ void FContentBundleEditor::StopEditing()
 void FContentBundleEditor::InjectBaseContent()
 {
 	check(GetStatus() == EContentBundleStatus::ReadyToInject);
-	UE_LOG(LogContentBundle, Log, TEXT("[CB: %s] Injecting Base Content"), *GetDescriptor()->GetDisplayName());
+	UE_LOG(LogContentBundle, Log, TEXT("%s Injecting Base Content"), *ContentBundle::Log::MakeDebugInfoString(*this));
 
 	FActorSpawnParameters SpawnParameters;
 	SpawnParameters.Name = BuildWorlDataLayersName();
@@ -307,7 +307,7 @@ void FContentBundleEditor::GenerateStreaming(TArray<FString>* OutPackageToGenera
 
 	if (GetStatus() != EContentBundleStatus::ContentInjected)
 	{
-		UE_LOG(LogContentBundle, Log, TEXT("[CB: %s] Skipping streaming generation. It's status is: %s."), *GetDescriptor()->GetDisplayName(), *UEnum::GetDisplayValueAsText(GetStatus()).ToString());
+		UE_LOG(LogContentBundle, Log, TEXT("%s Skipping streaming generation. It's status is: %s."), *ContentBundle::Log::MakeDebugInfoString(*this), *UEnum::GetDisplayValueAsText(GetStatus()).ToString());
 		return;
 	}
 
@@ -322,7 +322,7 @@ void FContentBundleEditor::GenerateStreaming(TArray<FString>* OutPackageToGenera
 		CellCount++;
 		return true;
 	});
-	UE_LOG(LogContentBundle, Log, TEXT("[CB: %s] Generated streaming cells. %u cells were generated"), *GetDescriptor()->GetDisplayName(), CellCount);
+	UE_LOG(LogContentBundle, Log, TEXT("%s Generated streaming cells. %u cells were generated"), *ContentBundle::Log::MakeDebugInfoString(*this), CellCount);
 
 	if (bIsPIE)
 	{
@@ -330,7 +330,7 @@ void FContentBundleEditor::GenerateStreaming(TArray<FString>* OutPackageToGenera
 		UContentBundleDuplicateForPIEHelper* DuplicateForPIEHelper = ContentBundleManager->GetPIEDuplicateHelper();
 		if (!DuplicateForPIEHelper->StoreContentBundleStreamingObect(*this, ExternalStreamingObject))
 		{
-			UE_LOG(LogContentBundle, Error, TEXT("[CB: %s] Failed to store streaming object for %s. PIE duplication will not work."), *GetDescriptor()->GetDisplayName());
+			UE_LOG(LogContentBundle, Error, TEXT("%s Failed to store streaming object. PIE duplication will not work."), *ContentBundle::Log::MakeDebugInfoString(*this));
 		}
 
 		// Clear streaming object. It will be kept alive & duplicated by UContentBundleDuplicateForPIEHelper.
@@ -365,7 +365,7 @@ bool FContentBundleEditor::GatherPackagesToCook(class IWorldPartitionCookPackage
 			}
 			else
 			{
-				UE_LOG(LogContentBundle, Error, TEXT("[CB: %s][Cook] Failed to add cell package %s in cook context."), *ContentBundlePaths::GetCookedContentBundleLevelFolder(*this), *RuntimeCell.GetPackageNameToCreate());
+				UE_LOG(LogContentBundle, Error, TEXT("%s[Cook] Failed to add cell package %s in cook context."), *ContentBundle::Log::MakeDebugInfoString(*this), *RuntimeCell.GetPackageNameToCreate());
 				bIsSuccess = false;
 			}
 		});
@@ -373,7 +373,7 @@ bool FContentBundleEditor::GatherPackagesToCook(class IWorldPartitionCookPackage
 		const FWorldPartitionCookPackage* CookPackage = CookContext.AddGenericPackageToGenerate(this, ContentBundlePaths::GetCookedContentBundleLevelFolder(*this), GetExternalStreamingObjectPackageName());
 		if (CookPackage == nullptr)
 		{
-			UE_LOG(LogContentBundle, Error, TEXT("[CB: %s][Cook] Failed to add streaming object package in cook context %s."), *GetDescriptor()->GetDisplayName());
+			UE_LOG(LogContentBundle, Error, TEXT("%s[Cook] Failed to add streaming object package in cook context."), *ContentBundle::Log::MakeDebugInfoString(*this));
 			bIsSuccess = false;
 		}
 	}
@@ -388,7 +388,7 @@ bool FContentBundleEditor::PopulateGeneratorPackageForCook(class IWorldPartition
 
 	if (HasCookedContent())
 	{
-		UE_LOG(LogContentBundle, Log, TEXT("[CB: %s][Cook] Populating Generator Package. %u Packages"), *GetDescriptor()->GetDisplayName(), PackagesToCook.Num());
+		UE_LOG(LogContentBundle, Log, TEXT("%s[Cook] Populating Generator Package. %u Packages"), *ContentBundle::Log::MakeDebugInfoString(*this), PackagesToCook.Num());
 
 		for (const FWorldPartitionCookPackage* CookPackage : PackagesToCook)
 		{
@@ -402,13 +402,13 @@ bool FContentBundleEditor::PopulateGeneratorPackageForCook(class IWorldPartition
 
 					if (!Cell->PrepareCellForCook(CookPackage->GetPackage()))
 					{
-						UE_LOG(LogContentBundle, Error, TEXT("[CB: %s][Cook] Failed to prepare cell with package %s for cook."), *GetDescriptor()->GetDisplayName(), *CookPackage->RelativePath);
+						UE_LOG(LogContentBundle, Error, TEXT("%s[Cook] Failed to prepare cell with package %s for cook."), *ContentBundle::Log::MakeDebugInfoString(*this), *CookPackage->RelativePath);
 						bIsSuccess = false;
 					}
 				}
 				else
 				{
-					UE_LOG(LogContentBundle, Error, TEXT("[CB: %s][Cook] Could not find cell for package %s while populating generator pacakges."), *GetDescriptor()->GetDisplayName(), *CookPackage->RelativePath);
+					UE_LOG(LogContentBundle, Error, TEXT("%s[Cook] Could not find cell for package %s while populating generator pacakges."), *ContentBundle::Log::MakeDebugInfoString(*this), *CookPackage->RelativePath);
 					bIsSuccess = false;
 				}
 			}
@@ -425,7 +425,7 @@ bool FContentBundleEditor::PopulateGeneratedPackageForCook(class IWorldPartition
 {
 	bool bIsSuccess = true;
 
-	UE_LOG(LogContentBundle, Log, TEXT("[CB: %s][Cook] Populating Generated Package %s"), *GetDescriptor()->GetDisplayName(), *PackageToCook.RelativePath);
+	UE_LOG(LogContentBundle, Log, TEXT("%s[Cook] Populating Generated Package %s"), *ContentBundle::Log::MakeDebugInfoString(*this), *PackageToCook.RelativePath);
 
 	if (PackageToCook.Type == FWorldPartitionCookPackage::EType::Level)
 	{
@@ -439,22 +439,22 @@ bool FContentBundleEditor::PopulateGeneratedPackageForCook(class IWorldPartition
 					TArray<UPackage*> ModifiedPackages;
 					if (!Cell->PopulateGeneratedPackageForCook(PackageToCook.GetPackage(), OutModifiedPackages))
 					{
-						UE_LOG(LogContentBundle, Error, TEXT("[CB: %s][Cook] Failed to populate cell package %s."), *GetDescriptor()->GetDisplayName(), *PackageToCook.RelativePath);
+						UE_LOG(LogContentBundle, Error, TEXT("%s[Cook] Failed to populate cell package %s."), *ContentBundle::Log::MakeDebugInfoString(*this), *PackageToCook.RelativePath);
 						bIsSuccess = false;
 					}
 					
 				}
 				else
 				{
-					UE_LOG(LogContentBundle, Warning, TEXT("[CB: %s][Cook] Cell %s is flagged always loaded. Content Bundles cells should never be always loaded. It will not be populated and be empty at runtime."),
-						*GetDescriptor()->GetDisplayName(), *PackageToCook.RelativePath);
+					UE_LOG(LogContentBundle, Warning, TEXT("%s[Cook] Cell %s is flagged always loaded. Content Bundles cells should never be always loaded. It will not be populated and be empty at runtime."),
+						*ContentBundle::Log::MakeDebugInfoString(*this), *PackageToCook.RelativePath);
 					bIsSuccess = false;
 				}
 			}
 		}
 		else
 		{
-			UE_LOG(LogContentBundle, Error, TEXT("[CB: %s][Cook] Could not find cell for package %s while populating generated pacakges."), *GetDescriptor()->GetDisplayName(), *PackageToCook.RelativePath);
+			UE_LOG(LogContentBundle, Error, TEXT("%s[Cook] Could not find cell for package %s while populating generated pacakges."), *ContentBundle::Log::MakeDebugInfoString(*this), *PackageToCook.RelativePath);
 			bIsSuccess = false;
 		}
 	}
@@ -462,7 +462,7 @@ bool FContentBundleEditor::PopulateGeneratedPackageForCook(class IWorldPartition
 	{
 		if (!ExternalStreamingObject->Rename(nullptr, PackageToCook.GetPackage(), REN_DontCreateRedirectors))
 		{
-			UE_LOG(LogContentBundle, Error, TEXT("[CB: %s][Cook] Failed to set streaming object package %s."), *GetDescriptor()->GetDisplayName());
+			UE_LOG(LogContentBundle, Error, TEXT("%s[Cook] Failed to rename streaming object package."), *ContentBundle::Log::MakeDebugInfoString(*this));
 			bIsSuccess = false;
 		}
 	}
@@ -513,8 +513,8 @@ void FContentBundleEditor::UnregisterDelegates()
 
 void FContentBundleEditor::OnActorDescAdded(FWorldPartitionActorDesc* ActorDesc)
 {
-	UE_LOG(LogContentBundle, Verbose, TEXT("[CB: %s] Added actor %s to container, ActorCount: %u. Package %s."), 
-		*GetDescriptor()->GetDisplayName(), *ActorDesc->GetActorLabelOrName().ToString(), GetActorCount(), *ActorDesc->GetActorPackage().ToString());
+	UE_LOG(LogContentBundle, Verbose, TEXT("%s Added actor %s to container, ActorCount: %u. Package %s."), 
+		*ContentBundle::Log::MakeDebugInfoString(*this), *ActorDesc->GetActorLabelOrName().ToString(), GetActorCount(), *ActorDesc->GetActorPackage().ToString());
 
 	AActor* Actor = ActorDesc->GetActor();
 	UnsavedActorMonitor->StopMonitoringActor(Actor);
@@ -522,8 +522,8 @@ void FContentBundleEditor::OnActorDescAdded(FWorldPartitionActorDesc* ActorDesc)
 
 void FContentBundleEditor::OnActorDescRemoved(FWorldPartitionActorDesc* ActorDesc)
 {
-	UE_LOG(LogContentBundle, Verbose, TEXT("[CB: %s] Removed actor %s from container, ActorCount:  %u. Package %s."), 
-		*GetDescriptor()->GetDisplayName(), *ActorDesc->GetActorLabelOrName().ToString(), GetActorCount(), *ActorDesc->GetActorPackage().ToString());
+	UE_LOG(LogContentBundle, Verbose, TEXT("%s Removed actor %s from container, ActorCount:  %u. Package %s."), 
+		*ContentBundle::Log::MakeDebugInfoString(*this), *ActorDesc->GetActorLabelOrName().ToString(), GetActorCount(), *ActorDesc->GetActorPackage().ToString());
 
 	if (!HasUserPlacedActors())
 	{
@@ -539,8 +539,8 @@ void FContentBundleEditor::OnActorDescRemoved(FWorldPartitionActorDesc* ActorDes
 
 void FContentBundleEditor::OnUnsavedActorDeleted(AActor* Actor)
 {
-	UE_LOG(LogContentBundle, Verbose, TEXT("[CB: %s] Removed unsaved actor %s, ActorCount: %u. Package %s."),
-		*GetDescriptor()->GetDisplayName(), *Actor->GetActorNameOrLabel(), GetActorCount(), *Actor->GetPackage()->GetName());
+	UE_LOG(LogContentBundle, Verbose, TEXT("%s Removed unsaved actor %s, ActorCount: %u. Package %s."),
+		*ContentBundle::Log::MakeDebugInfoString(*this), *Actor->GetActorNameOrLabel(), GetActorCount(), *Actor->GetPackage()->GetName());
 
 	if (!HasUserPlacedActors())
 	{
diff --git a/Engine/Source/Runtime/Engine/Private/WorldPartition/ContentBundle/ContentBundleLog.cpp b/Engine/Source/Runtime/Engine/Private/WorldPartition/ContentBundle/ContentBundleLog.cpp
--- a/Engine/Source/Runtime/Engine/Private/WorldPartition/ContentBundle/ContentBundleLog.cpp
+++ b/Engine/Source/Runtime/Engine/Private/WorldPartition/ContentBundle/ContentBundleLog.cpp
@@ -2,4 +2,61 @@
 
 #include "WorldPartition/ContentBundle/ContentBundleLog.h"
 
-DEFINE_LOG_CATEGORY(LogContentBundle);
\ No newline at end of file
+#include "Engine/World.h"
+#include "WorldPartition/ContentBundle/ContentBundleBase.h"
+#include "WorldPartition/ContentBundle/ContentBundleDescriptor.h"
+#include "WorldPartition/ContentBundle/ContentBundleClient.h"
+#include "WorldPartition/ContentBundle/ContentBundleContainer.h"
+
+DEFINE_LOG_CATEGORY(LogContentBundle);
+
+namespace ContentBundle
+{
+#if !NO_LOGGING
+
+	namespace LogPrivate
+	{
+		FString MakeDebugInfoString(const UContentBundleDescriptor& ContentBundleDescriptor)
+		{
+			return FString::Printf(TEXT("[CB: %s]"), *ContentBundleDescriptor.GetDisplayName());
+		}
+
+		FString MakeWorldDebugInfoString(UWorld* World)
+		{
+			return FString::Printf(TEXT("[%s(%s)]"), *World->GetName(), *GetDebugStringForWorld(World));
+		}
+
+		FString MakeDebugInfoString(const UContentBundleDescriptor& ContentBundleDescriptor, UWorld* World)
+		{
+			return FString::Printf(TEXT("%s%s"), *MakeDebugInfoString(ContentBundleDescriptor), *MakeWorldDebugInfoString(World));
+		}
+
+		
+	}
+
+FString Log::MakeDebugInfoString(const FContentBundleBase& ContentBundle)
+{
+	return LogPrivate::MakeDebugInfoString(*ContentBundle.GetDescriptor(), ContentBundle.GetInjectedWorld());
+}
+
+FString Log::MakeDebugInfoString(const FContentBundleClient& ContentBundleClient)
+{
+	return LogPrivate::MakeDebugInfoString(*ContentBundleClient.GetDescriptor());
+}
+
+FString Log::MakeDebugInfoString(const FContentBundleContainer& ContentBundleContainer)
+{
+	return LogPrivate::MakeWorldDebugInfoString(ContentBundleContainer.GetInjectedWorld());
+}
+
+FString Log::MakeDebugInfoString(const FContentBundleClient& ContentBundleClient, UWorld* World)
+{
+	return LogPrivate::MakeDebugInfoString(*ContentBundleClient.GetDescriptor(), World);
+}
+
+
+
+
+
+#endif
+}
\ No newline at end of file
diff --git a/Engine/Source/Runtime/Engine/Public/WorldPartition/ContentBundle/ContentBundleLog.h b/Engine/Source/Runtime/Engine/Public/WorldPartition/ContentBundle/ContentBundleLog.h
--- a/Engine/Source/Runtime/Engine/Public/WorldPartition/ContentBundle/ContentBundleLog.h
+++ b/Engine/Source/Runtime/Engine/Public/WorldPartition/ContentBundle/ContentBundleLog.h
@@ -3,4 +3,23 @@
 
 #include "CoreMinimal.h"
 
-ENGINE_API DECLARE_LOG_CATEGORY_EXTERN(LogContentBundle, Log, All);
\ No newline at end of file
+ENGINE_API DECLARE_LOG_CATEGORY_EXTERN(LogContentBundle, Log, All);
+
+class FContentBundleBase;	
+class FContentBundleClient;
+class FContentBundleContainer;
+class UContentBundleDescriptor;
+class UWorld;
+
+namespace ContentBundle
+{
+	namespace Log
+	{
+#if !NO_LOGGING
+		FString MakeDebugInfoString(const FContentBundleBase& ContentBundle);
+		FString MakeDebugInfoString(const FContentBundleClient& ContentBundleClient);
+		FString MakeDebugInfoString(const FContentBundleContainer& ContentBundleContainer);
+		FString MakeDebugInfoString(const FContentBundleClient& ContentBundleClient, UWorld* World);
+#endif
+	}
+}
\ No newline at end of file
