Add missing include

diff --git a/Engine/Source/Runtime/Engine/Private/WorldPartition/ContentBundle/ContentBundleLog.cpp b/Engine/Source/Runtime/Engine/Private/WorldPartition/ContentBundle/ContentBundleLog.cpp
--- a/Engine/Source/Runtime/Engine/Private/WorldPartition/ContentBundle/ContentBundleLog.cpp
+++ b/Engine/Source/Runtime/Engine/Private/WorldPartition/ContentBundle/ContentBundleLog.cpp
@@ -7,6 +7,7 @@
 #include "WorldPartition/ContentBundle/ContentBundleDescriptor.h"
 #include "WorldPartition/ContentBundle/ContentBundleClient.h"
 #include "WorldPartition/ContentBundle/ContentBundleContainer.h"
+#include "UnrealEngine.h"
 
 DEFINE_LOG_CATEGORY(LogContentBundle);
 
