Level Instance: Fix WP Level Instance actor transform when loading/unloading DataLayers

- Instead of doing transform fixup in LevelInstanceLevelStreaming, update the LevelStreaming/WP Transform so that AddLoadedActor adds it a the proper location

diff --git a/Engine/Source/Runtime/Engine/Private/LevelInstance/LevelInstanceComponent.cpp b/Engine/Source/Runtime/Engine/Private/LevelInstance/LevelInstanceComponent.cpp
--- a/Engine/Source/Runtime/Engine/Private/LevelInstance/LevelInstanceComponent.cpp
+++ b/Engine/Source/Runtime/Engine/Private/LevelInstance/LevelInstanceComponent.cpp
@@ -98,9 +98,9 @@ void ULevelInstanceComponent::UpdateEditorInstanceActor()
 		}
 	}
 
-	if (AActor* EditorInstanceActor = CachedEditorInstanceActorPtr.Get())
+	if (ALevelInstanceEditorInstanceActor* EditorInstanceActor = Cast<ALevelInstanceEditorInstanceActor>(CachedEditorInstanceActorPtr.Get()))
 	{
-		EditorInstanceActor->GetRootComponent()->SetWorldTransform(GetComponentTransform());
+		EditorInstanceActor->UpdateWorldTransform(GetComponentTransform());
 	}
 }
 
diff --git a/Engine/Source/Runtime/Engine/Private/LevelInstance/LevelInstanceEditorInstanceActor.cpp b/Engine/Source/Runtime/Engine/Private/LevelInstance/LevelInstanceEditorInstanceActor.cpp
--- a/Engine/Source/Runtime/Engine/Private/LevelInstance/LevelInstanceEditorInstanceActor.cpp
+++ b/Engine/Source/Runtime/Engine/Private/LevelInstance/LevelInstanceEditorInstanceActor.cpp
@@ -5,6 +5,10 @@
 #include "LevelInstance/LevelInstanceSubsystem.h"
 #include "Engine/Level.h"
 #include "Engine/World.h"
+#include "Engine/LevelStreaming.h"
+#include "LevelUtils.h"
+#include "GameFramework/WorldSettings.h"
+#include "WorldPartition/WorldPartition.h"
 
 #include UE_INLINE_GENERATED_CPP_BY_NAME(LevelInstanceEditorInstanceActor)
 
@@ -58,4 +62,23 @@ ALevelInstanceEditorInstanceActor* ALevelInstanceEditorInstanceActor::Create(ILe
 	return InstanceActor;
 }
 
+void ALevelInstanceEditorInstanceActor::UpdateWorldTransform(const FTransform& WorldTransform)
+{
+	GetRootComponent()->SetWorldTransform(WorldTransform);
+
+	const ULevel* Level = GetLevel();
+
+	ULevelStreaming* LevelStreaming = FLevelUtils::FindStreamingLevel(Level);
+	check(LevelStreaming);
+
+	AWorldSettings* WorldSettings = Level->GetWorldSettings();
+	check(WorldSettings);
+	LevelStreaming->LevelTransform = FTransform(WorldSettings->LevelInstancePivotOffset) * WorldTransform;
+
+	if (UWorldPartition* WorldPartition = WorldSettings->GetWorldPartition())
+	{
+		WorldPartition->SetInstanceTransform(LevelStreaming->LevelTransform);
+	}
+}
+
 #endif
diff --git a/Engine/Source/Runtime/Engine/Private/LevelInstance/LevelInstanceLevelStreaming.cpp b/Engine/Source/Runtime/Engine/Private/LevelInstance/LevelInstanceLevelStreaming.cpp
--- a/Engine/Source/Runtime/Engine/Private/LevelInstance/LevelInstanceLevelStreaming.cpp
+++ b/Engine/Source/Runtime/Engine/Private/LevelInstance/LevelInstanceLevelStreaming.cpp
@@ -153,16 +153,6 @@ void ULevelStreamingLevelInstance::PrepareLevelInstanceLoadedActor(AActor& InAct
 	{
 		if (InActor.GetAttachParentActor() == nullptr && !InActor.IsChildActor())
 		{
-			// Detect if LevelInstance has moved since it was loaded/created and move the child actor accordingly
-			if (!LevelTransform.Equals(LevelInstanceEditorInstanceActor->GetTransform()))
-			{
-				FTransform TransformToApply = LevelTransform.Inverse() * LevelInstanceEditorInstanceActor->GetTransform();
-				FLevelUtils::FApplyLevelTransformParams TransformParams(InActor.GetLevel(), TransformToApply);
-				TransformParams.Actor = &InActor;
-				TransformParams.bDoPostEditMove = true;
-				FLevelUtils::ApplyLevelTransform(TransformParams);
-			}
-
 			InActor.AttachToActor(LevelInstanceEditorInstanceActor.Get(), FAttachmentTransformRules::KeepWorldTransform);
 		}
 	}
diff --git a/Engine/Source/Runtime/Engine/Public/LevelInstance/LevelInstanceEditorInstanceActor.h b/Engine/Source/Runtime/Engine/Public/LevelInstance/LevelInstanceEditorInstanceActor.h
--- a/Engine/Source/Runtime/Engine/Public/LevelInstance/LevelInstanceEditorInstanceActor.h
+++ b/Engine/Source/Runtime/Engine/Public/LevelInstance/LevelInstanceEditorInstanceActor.h
@@ -32,6 +32,8 @@ public:
 	virtual AActor* GetSelectionParent() const override;
 
 private:
+	friend class ULevelInstanceComponent;
+	void UpdateWorldTransform(const FTransform& WorldTransform);
 
 	FLevelInstanceID LevelInstanceID;
 #endif
diff --git a/Engine/Source/Runtime/Engine/Public/WorldPartition/WorldPartition.h b/Engine/Source/Runtime/Engine/Public/WorldPartition/WorldPartition.h
--- a/Engine/Source/Runtime/Engine/Public/WorldPartition/WorldPartition.h
+++ b/Engine/Source/Runtime/Engine/Public/WorldPartition/WorldPartition.h
@@ -151,6 +151,7 @@ public:
 	//~ End UObject Interface
 
 #if WITH_EDITOR
+	void SetInstanceTransform(const FTransform& InInstanceTransform) { InstanceTransform = InInstanceTransform; }
 	FName GetWorldPartitionEditorName() const;
 
 	// Streaming generation
