Fix error in shipping build, UE_LOG is defined in shipping

diff --git a/Engine/Source/Runtime/Engine/Private/WorldPartition/ContentBundle/ContentBundleLog.cpp b/Engine/Source/Runtime/Engine/Private/WorldPartition/ContentBundle/ContentBundleLog.cpp
--- a/Engine/Source/Runtime/Engine/Private/WorldPartition/ContentBundle/ContentBundleLog.cpp
+++ b/Engine/Source/Runtime/Engine/Private/WorldPartition/ContentBundle/ContentBundleLog.cpp
@@ -14,7 +14,6 @@ DEFINE_LOG_CATEGORY(LogContentBundle);
 namespace ContentBundle
 {
 #if !NO_LOGGING
-
 	namespace LogPrivate
 	{
 		FString MakeDebugInfoString(const UContentBundleDescriptor& ContentBundleDescriptor)
@@ -30,34 +29,44 @@ namespace ContentBundle
 		FString MakeDebugInfoString(const UContentBundleDescriptor& ContentBundleDescriptor, UWorld* World)
 		{
 			return FString::Printf(TEXT("%s%s"), *MakeDebugInfoString(ContentBundleDescriptor), *MakeWorldDebugInfoString(World));
-		}
-
-		
+		}	
 	}
+#endif
 
 FString Log::MakeDebugInfoString(const FContentBundleBase& ContentBundle)
 {
+#if !NO_LOGGING
 	return LogPrivate::MakeDebugInfoString(*ContentBundle.GetDescriptor(), ContentBundle.GetInjectedWorld());
+#else
+	return FString();
+#endif
 }
 
 FString Log::MakeDebugInfoString(const FContentBundleClient& ContentBundleClient)
 {
+#if !NO_LOGGING
 	return LogPrivate::MakeDebugInfoString(*ContentBundleClient.GetDescriptor());
+#else
+	return FString();
+#endif
 }
 
 FString Log::MakeDebugInfoString(const FContentBundleContainer& ContentBundleContainer)
 {
+#if !NO_LOGGING
 	return LogPrivate::MakeWorldDebugInfoString(ContentBundleContainer.GetInjectedWorld());
+#else
+	return FString();
+#endif
 }
 
 FString Log::MakeDebugInfoString(const FContentBundleClient& ContentBundleClient, UWorld* World)
 {
+#if !NO_LOGGING
 	return LogPrivate::MakeDebugInfoString(*ContentBundleClient.GetDescriptor(), World);
+#else
+	return FString();
+#endif
 }
 
-
-
-
-
-#endif
 }
\ No newline at end of file
diff --git a/Engine/Source/Runtime/Engine/Public/WorldPartition/ContentBundle/ContentBundleLog.h b/Engine/Source/Runtime/Engine/Public/WorldPartition/ContentBundle/ContentBundleLog.h
--- a/Engine/Source/Runtime/Engine/Public/WorldPartition/ContentBundle/ContentBundleLog.h
+++ b/Engine/Source/Runtime/Engine/Public/WorldPartition/ContentBundle/ContentBundleLog.h
@@ -15,11 +15,9 @@ namespace ContentBundle
 {
 	namespace Log
 	{
-#if !NO_LOGGING
-		FString MakeDebugInfoString(const FContentBundleBase& ContentBundle);
-		FString MakeDebugInfoString(const FContentBundleClient& ContentBundleClient);
-		FString MakeDebugInfoString(const FContentBundleContainer& ContentBundleContainer);
-		FString MakeDebugInfoString(const FContentBundleClient& ContentBundleClient, UWorld* World);
-#endif
+		ENGINE_API FString MakeDebugInfoString(const FContentBundleBase& ContentBundle);
+		ENGINE_API FString MakeDebugInfoString(const FContentBundleClient& ContentBundleClient);
+		ENGINE_API FString MakeDebugInfoString(const FContentBundleContainer& ContentBundleContainer);
+		ENGINE_API FString MakeDebugInfoString(const FContentBundleClient& ContentBundleClient, UWorld* World);
 	}
 }
\ No newline at end of file
