Audio based vibration device properties

#rnx

diff --git a/Engine/Source/Runtime/Engine/Classes/GameFramework/InputDeviceProperties.h b/Engine/Source/Runtime/Engine/Classes/GameFramework/InputDeviceProperties.h
--- a/Engine/Source/Runtime/Engine/Classes/GameFramework/InputDeviceProperties.h
+++ b/Engine/Source/Runtime/Engine/Classes/GameFramework/InputDeviceProperties.h
@@ -470,4 +470,67 @@ private:
 
 	/** The internal property that represents this trigger feedback. */
 	FInputDeviceTriggerVibrationProperty InternalProperty;
+};
+
+///////////////////////////////////////////////////////////////////////
+// UInputDeviceSoundBasedVibrationProperty
+
+class USoundBase;
+class UEndpointSubmix;
+
+USTRUCT(BlueprintType)
+struct FAudioBasedVibrationData
+{
+	GENERATED_BODY()
+
+	FAudioBasedVibrationData();
+	
+	/** The sound to play on the gamepad */
+	UPROPERTY(EditAnywhere, BlueprintReadWrite, Category = "DeviceProperty")
+	TObjectPtr<USoundBase> Sound;
+
+	/** The audio endpoint that the sound will be played on for gamepad vibration. */
+	UPROPERTY(EditAnywhere, BlueprintReadWrite, Category = "Vibration Audio")
+	TObjectPtr<UEndpointSubmix> VibrationEndpoint;
+
+	/** The audio endpoint that the sound will be played on for gamepad audio. */
+	UPROPERTY(EditAnywhere, BlueprintReadWrite, Category = "Gamepad Audio")
+	TObjectPtr<UEndpointSubmix> GamepadAudioEndpoint;
+
+	/** The amount of sound to send to the Vibration endpoint submix */
+	UPROPERTY(EditAnywhere, BlueprintReadWrite, Category = "Vibration Audio", meta = (UIMin = "0.0"))
+	float VibrationSoundLevel;
+
+	/** The amount of sound to send to the Gamepad Audio endpoint submix */
+	UPROPERTY(EditAnywhere, BlueprintReadWrite, Category = "Gamepad Audio", meta = (UIMin = "0.0"))
+	float GamepadAudioSoundLevel;
+};
+
+/**
+ * Plays a sound to an input device's speaker. On platforms that support it, this sound will be played
+ * in the form of a vibration where the left and right audio channels represent the left and right side
+ * of the controller.
+ */
+UCLASS(Blueprintable, meta = (DisplayName = "Audio Based Vibration (Experimental)"))
+class UInputDeviceAudioBasedVibrationProperty : public UInputDeviceProperty
+{
+	GENERATED_BODY()
+
+public:
+	virtual void EvaluateDeviceProperty_Implementation(const FPlatformUserId PlatformUser, const FInputDeviceId DeviceId, const float DeltaTime, const float Duration) override;
+	virtual void ApplyDeviceProperty(const FPlatformUserId UserId, const FInputDeviceId DeviceId) override;
+	virtual FInputDeviceProperty* GetInternalDeviceProperty() override;
+	virtual float RecalculateDuration() override;
+	
+	UPROPERTY(EditAnywhere, BlueprintReadWrite, Category = "Trigger")
+	FAudioBasedVibrationData Data;
+
+	/** A map of device specific color data. If no overrides are specified, the Default hardware data will be used */
+	UPROPERTY(BlueprintReadWrite, EditAnywhere, Category = "Trigger", meta = (GetOptions = "Engine.InputPlatformSettings.GetAllHardwareDeviceNames"))
+	TMap<FName, FAudioBasedVibrationData> DeviceOverrideData;
+
+private:
+
+	/** Returns the data that is relevant to the current input device */
+	const FAudioBasedVibrationData* GetRelevantData(const FPlatformUserId UserId, const FInputDeviceId DeviceId) const;
 };
\ No newline at end of file
diff --git a/Engine/Source/Runtime/Engine/Private/GameFramework/InputDeviceProperties.cpp b/Engine/Source/Runtime/Engine/Private/GameFramework/InputDeviceProperties.cpp
--- a/Engine/Source/Runtime/Engine/Private/GameFramework/InputDeviceProperties.cpp
+++ b/Engine/Source/Runtime/Engine/Private/GameFramework/InputDeviceProperties.cpp
@@ -4,6 +4,7 @@
 #include "Framework/Application/SlateApplication.h"
 #include "Curves/CurveLinearColor.h"
 #include "Curves/CurveFloat.h"
+#include "GameFramework/InputDeviceLibrary.h"
 
 #include UE_INLINE_GENERATED_CPP_BY_NAME(InputDeviceProperties)
 
@@ -420,3 +421,86 @@ int32 UInputDeviceTriggerVibrationProperty::GetVibrationAmplitudeValue(const FDe
 	
 	return 0;
 }
+
+///////////////////////////////////////////////////////////////////////
+// UInputDeviceSoundBasedVibrationProperty
+
+#include "Sound/SoundBase.h"
+#include "Sound/SoundSubmix.h"
+#include "GameFramework/PlayerController.h"
+#include "GameFramework/InputDeviceSubsystem.h"
+
+FAudioBasedVibrationData::FAudioBasedVibrationData()
+	: Sound(nullptr)
+	, VibrationEndpoint(nullptr)
+	, GamepadAudioEndpoint(nullptr)
+	, VibrationSoundLevel(1.0f)
+	, GamepadAudioSoundLevel(1.0f)
+{
+}
+
+void UInputDeviceAudioBasedVibrationProperty::EvaluateDeviceProperty_Implementation(const FPlatformUserId PlatformUser, const FInputDeviceId DeviceId, const float DeltaTime, const float Duration)
+{
+	if (const FAudioBasedVibrationData* DataToUse = GetRelevantData(PlatformUser, DeviceId))
+	{
+		// Set the endpoints to use on this sound
+		if (DataToUse->Sound)
+		{
+			if (DataToUse->VibrationEndpoint)
+			{
+				FSoundSubmixSendInfo SendInfo;
+				SendInfo.SoundSubmix = DataToUse->VibrationEndpoint;
+				SendInfo.SendLevel = DataToUse->VibrationSoundLevel;
+				DataToUse->Sound->SoundSubmixSends.Emplace(MoveTemp(SendInfo));
+			}
+
+			if (DataToUse->GamepadAudioEndpoint)
+			{
+				FSoundSubmixSendInfo SendInfo;
+				SendInfo.SoundSubmix = DataToUse->GamepadAudioEndpoint;
+				SendInfo.SendLevel = DataToUse->GamepadAudioSoundLevel;
+				DataToUse->Sound->SoundSubmixSends.Emplace(MoveTemp(SendInfo));
+			}
+		}
+	}
+}
+
+void UInputDeviceAudioBasedVibrationProperty::ApplyDeviceProperty(const FPlatformUserId UserId, const FInputDeviceId DeviceId)
+{
+	// Audio Endpoints can already play vibrations and audio on an input device, if the device implements
+	// the interface for it. We can just do that! Woot
+	if (const FAudioBasedVibrationData* DataToUse = GetRelevantData(UserId, DeviceId))
+	{
+		if (DataToUse->Sound)
+		{
+			// Get the player controller and play the sound
+			if (APlayerController* PC = UInputDeviceLibrary::GetPlayerControllerFromPlatformUser(UserId))
+			{
+				// The sound endpoints will have been populated above in the Evaluate function
+				PC->ClientPlaySound(DataToUse->Sound);
+			}	
+		}		
+	}
+}
+
+FInputDeviceProperty* UInputDeviceAudioBasedVibrationProperty::GetInternalDeviceProperty()
+{
+	return nullptr;
+}
+
+float UInputDeviceAudioBasedVibrationProperty::RecalculateDuration()
+{
+	// Returning a durtion of 0 means that it will be played as a "one shot" effect
+	return 0.0f;
+}
+
+const FAudioBasedVibrationData* UInputDeviceAudioBasedVibrationProperty::GetRelevantData(const FPlatformUserId UserId, const FInputDeviceId DeviceId) const
+{
+	const FAudioBasedVibrationData* DataToUse = &Data;
+	
+	if (const FAudioBasedVibrationData* OverrideData = GetDeviceSpecificData<FAudioBasedVibrationData>(UserId, DeviceId, DeviceOverrideData))
+	{
+		DataToUse = OverrideData;
+	}
+	return DataToUse;
+}
\ No newline at end of file
