World Partitiojn: get rid of WorldPartitionRuntimeSpatialHashCell by splitting hash-especific data into it's own subobject.

#skipcis

diff --git a/Engine/Source/Runtime/Engine/Private/WorldPartition/WorldPartitionRuntimeCell.cpp b/Engine/Source/Runtime/Engine/Private/WorldPartition/WorldPartitionRuntimeCell.cpp
--- a/Engine/Source/Runtime/Engine/Private/WorldPartition/WorldPartitionRuntimeCell.cpp
+++ b/Engine/Source/Runtime/Engine/Private/WorldPartition/WorldPartitionRuntimeCell.cpp
@@ -1,6 +1,8 @@
 // Copyright Epic Games, Inc. All Rights Reserved.
 
 #include "WorldPartition/WorldPartitionRuntimeCell.h"
+#include "UObject/Package.h"
+#include "Engine/Level.h"
 #include "Misc/HierarchicalLogArchive.h"
 #include "WorldPartition/WorldPartitionDebugHelper.h"
 #include "WorldPartition/DataLayer/DataLayerSubsystem.h"
@@ -11,23 +13,39 @@
 
 #include UE_INLINE_GENERATED_CPP_BY_NAME(WorldPartitionRuntimeCell)
 
-int32 UWorldPartitionRuntimeCell::StreamingSourceCacheEpoch = 0;
-
 UWorldPartitionRuntimeCell::UWorldPartitionRuntimeCell(const FObjectInitializer& ObjectInitializer)
 	: Super(ObjectInitializer)
 	, bIsAlwaysLoaded(false)
-	, ContentBounds(ForceInit)
 	, Priority(0)
-	, CachedMinSourcePriority((uint8)EStreamingSourcePriority::Lowest)
-	, CachedSourceInfoEpoch(INT_MIN)
 #if !UE_BUILD_SHIPPING
 	, DebugStreamingPriority(-1.f)
-#endif	
+#endif
+	, RuntimeCellData(nullptr)
 {
 	check(HasAnyFlags(RF_ClassDefaultObject) || GetOuter()->Implements<UWorldPartitionRuntimeCellOwner>());
 }
 
 #if WITH_EDITOR
+void UWorldPartitionRuntimeCell::PostDuplicate(bool bDuplicateForPIE)
+{
+	Super::PostDuplicate(bDuplicateForPIE);
+
+	if (UnsavedActorsContainer)
+	{
+		// Make sure actor container isn't under PIE World so those template actors will never be considered part of the world.
+		UnsavedActorsContainer->Rename(nullptr, GetPackage());
+
+		for (auto& [ActorName, Actor] : UnsavedActorsContainer->Actors)
+		{
+			if (ensure(Actor))
+			{
+				// Don't use AActor::Rename here since the actor is not par of the world, it's only a duplication template.	
+				Actor->UObject::Rename(nullptr, UnsavedActorsContainer);
+			}
+		}
+	}
+}
+
 bool UWorldPartitionRuntimeCell::NeedsActorToCellRemapping() const
 {
 	// When cooking, always loaded cells content is moved to persistent level (see PopulateGeneratorPackageForCook)
@@ -92,38 +110,33 @@ void UWorldPartitionRuntimeCell::UpdateDebugName()
 void UWorldPartitionRuntimeCell::DumpStateLog(FHierarchicalLogArchive& Ar)
 {
 	Ar.Printf(TEXT("Actor Count: %d"), GetActorCount());
-	Ar.Printf(TEXT("MinMaxZ: %s"), *GetMinMaxZ().ToString());
 }
-
 #endif
 
 bool UWorldPartitionRuntimeCell::ShouldResetStreamingSourceInfo() const
 {
-	return CachedSourceInfoEpoch != StreamingSourceCacheEpoch;
+	return RuntimeCellData->ShouldResetStreamingSourceInfo();
 }
 
 void UWorldPartitionRuntimeCell::ResetStreamingSourceInfo() const
 {
-	check(ShouldResetStreamingSourceInfo());
-	CachedSourcePriorityWeights.Reset();
-	CachedMinSourcePriority = MAX_uint8;
-	CachedSourceInfoEpoch = StreamingSourceCacheEpoch;
+	RuntimeCellData->ResetStreamingSourceInfo();
 }
 
 void UWorldPartitionRuntimeCell::AppendStreamingSourceInfo(const FWorldPartitionStreamingSource& Source, const FSphericalSector& SourceShape) const
 {
-	static_assert((uint8)EStreamingSourcePriority::Lowest == 255);
-	CachedSourcePriorityWeights.Add(1.f - ((float)Source.Priority / 255.f));
-	CachedMinSourcePriority = FMath::Min((uint8)Source.Priority, CachedMinSourcePriority);
+	RuntimeCellData->AppendStreamingSourceInfo(Source, SourceShape);
 }
 
 void UWorldPartitionRuntimeCell::MergeStreamingSourceInfo() const
-{}
+{
+	RuntimeCellData->MergeStreamingSourceInfo();
+}
 
 int32 UWorldPartitionRuntimeCell::SortCompare(const UWorldPartitionRuntimeCell* Other, bool bCanUseSortingCache) const
 {
-	// Source priority (lower value is higher prio)
-	const int32 Comparison = bCanUseSortingCache ? ((int32)CachedMinSourcePriority - (int32)Other->CachedMinSourcePriority) : 0;
+	const int32 Comparison = RuntimeCellData->SortCompare(Other->RuntimeCellData, bCanUseSortingCache);
+
 	// Cell priority (lower value is higher prio)
 	return (Comparison != 0) ? Comparison : (Priority - Other->Priority);
 }
@@ -176,6 +189,16 @@ bool UWorldPartitionRuntimeCell::ContainsDataLayer(const UDataLayerInstance* Dat
 	return GetDataLayers().Contains(DataLayerInstance->GetDataLayerFName());
 }
 
+const FBox& UWorldPartitionRuntimeCell::GetContentBounds() const
+{
+	return RuntimeCellData->GetContentBounds();
+}
+
+FBox UWorldPartitionRuntimeCell::GetCellBounds() const
+{
+	return RuntimeCellData->GetCellBounds();
+}
+
 FName UWorldPartitionRuntimeCell::GetLevelPackageName() const
 { 
 #if WITH_EDITOR
diff --git a/Engine/Source/Runtime/Engine/Private/WorldPartition/WorldPartitionRuntimeCellData.cpp b/Engine/Source/Runtime/Engine/Private/WorldPartition/WorldPartitionRuntimeCellData.cpp
new file mode 100644
--- /dev/null
+++ b/Engine/Source/Runtime/Engine/Private/WorldPartition/WorldPartitionRuntimeCellData.cpp
@@ -0,0 +1,51 @@
+// Copyright Epic Games, Inc. All Rights Reserved.
+
+#include "WorldPartition/WorldPartitionRuntimeCellData.h"
+
+int32 UWorldPartitionRuntimeCellData::StreamingSourceCacheEpoch = 0;
+
+UWorldPartitionRuntimeCellData::UWorldPartitionRuntimeCellData(const FObjectInitializer& ObjectInitializer)
+	: Super(ObjectInitializer)
+	, CachedMinSourcePriority(MAX_uint8)
+	, CachedSourceInfoEpoch(MIN_int32)
+	, ContentBounds(ForceInit)
+{}
+
+bool UWorldPartitionRuntimeCellData::ShouldResetStreamingSourceInfo() const
+{
+	return CachedSourceInfoEpoch != StreamingSourceCacheEpoch;
+}
+
+void UWorldPartitionRuntimeCellData::ResetStreamingSourceInfo() const
+{
+	check(ShouldResetStreamingSourceInfo());
+	CachedSourcePriorityWeights.Reset();
+	CachedMinSourcePriority = MAX_uint8;
+	CachedSourceInfoEpoch = StreamingSourceCacheEpoch;
+}
+
+void UWorldPartitionRuntimeCellData::AppendStreamingSourceInfo(const FWorldPartitionStreamingSource& Source, const FSphericalSector& SourceShape) const
+{
+	static_assert((uint8)EStreamingSourcePriority::Lowest == 255);
+	CachedSourcePriorityWeights.Add(1.f - ((float)Source.Priority / 255.f));
+	CachedMinSourcePriority = FMath::Min((uint8)Source.Priority, CachedMinSourcePriority);
+}
+
+void UWorldPartitionRuntimeCellData::MergeStreamingSourceInfo() const
+{}
+
+int32 UWorldPartitionRuntimeCellData::SortCompare(const UWorldPartitionRuntimeCellData* Other, bool bCanUseSortingCache) const
+{
+	// Source priority (lower value is higher prio)
+	return bCanUseSortingCache ? ((int32)CachedMinSourcePriority - (int32)Other->CachedMinSourcePriority) : 0;
+}
+
+const FBox& UWorldPartitionRuntimeCellData::GetContentBounds() const
+{
+	return ContentBounds;
+}
+
+FBox UWorldPartitionRuntimeCellData::GetCellBounds() const
+{
+	return ContentBounds;
+}
diff --git a/Engine/Source/Runtime/Engine/Private/WorldPartition/WorldPartitionRuntimeCellDataSpatialHash.cpp b/Engine/Source/Runtime/Engine/Private/WorldPartition/WorldPartitionRuntimeCellDataSpatialHash.cpp
new file mode 100644
--- /dev/null
+++ b/Engine/Source/Runtime/Engine/Private/WorldPartition/WorldPartitionRuntimeCellDataSpatialHash.cpp
@@ -0,0 +1,165 @@
+// Copyright Epic Games, Inc. All Rights Reserved.
+
+#include "WorldPartition/WorldPartitionRuntimeCellDataSpatialHash.h"
+#include "Engine/Level.h"
+
+#include UE_INLINE_GENERATED_CPP_BY_NAME(WorldPartitionRuntimeCellDataSpatialHash)
+
+static float GRuntimeSpatialHashCellToSourceAngleContributionToCellImportance = 0.4f; // Value between [0, 1]
+static FAutoConsoleVariableRef CVarRuntimeSpatialHashCellToSourceAngleContributionToCellImportance(
+	TEXT("wp.Runtime.RuntimeSpatialHashCellToSourceAngleContributionToCellImportance"),
+	GRuntimeSpatialHashCellToSourceAngleContributionToCellImportance,
+	TEXT("Value between 0 and 1 that modulates the contribution of the angle between streaming source-to-cell vector and source-forward vector to the cell importance. The closest to 0, the less the angle will contribute to the cell importance."));
+
+UWorldPartitionRuntimeCellDataSpatialHash::UWorldPartitionRuntimeCellDataSpatialHash(const FObjectInitializer& ObjectInitializer)
+	: Super(ObjectInitializer)
+	, Extent(0)
+	, Level(0)
+	, bCachedIsBlockingSource(false)
+	, CachedMinSquareDistanceToBlockingSource(MAX_dbl)
+	, CachedMinSquareDistanceToSource(MAX_dbl)
+{}
+
+void UWorldPartitionRuntimeCellDataSpatialHash::ResetStreamingSourceInfo() const
+{
+	Super::ResetStreamingSourceInfo();
+
+	bCachedIsBlockingSource = false;
+	CachedMinSquareDistanceToBlockingSource = MAX_dbl;
+	CachedMinSquareDistanceToSource = MAX_dbl;
+	CachedSourceSquaredDistances.Reset();
+	CachedInstersectingShapes.Reset();
+}
+
+float UWorldPartitionRuntimeCellDataSpatialHash::ComputeSourceToCellAngleFactor(const FSphericalSector& SourceShape) const
+{
+	float AngleContribution = FMath::Clamp(GRuntimeSpatialHashCellToSourceAngleContributionToCellImportance, 0.f, 1.f);
+	float AngleFactor = 1.f;
+	if (!FMath::IsNearlyZero(AngleContribution))
+	{
+		const FBox Box(FVector(Position.X - Extent, Position.Y - Extent, 0.f), FVector(Position.X + Extent, Position.Y + Extent, 0.f));
+		const FVector2D SourcePos(SourceShape.GetCenter());
+		const FVector StartVert(SourcePos, 0.f);
+		const FVector EndVert(SourcePos + FVector2D(SourceShape.GetScaledAxis()), 0.f);
+
+		float Angle = 0.f;
+		if (!FMath::LineBoxIntersection(Box, StartVert, EndVert, EndVert - StartVert))
+		{
+			// Find smallest angle using 4 corners and center of cell bounds
+			const FVector2D Position2D(Position);
+			float MaxDot = 0.f;
+			FVector2D SourceForward(SourceShape.GetAxis());
+			SourceForward.Normalize();
+			FVector2D CellPoints[5];
+			CellPoints[0] = Position2D + FVector2D(-Extent, -Extent);
+			CellPoints[1] = Position2D + FVector2D(-Extent, Extent);
+			CellPoints[2] = Position2D + FVector2D(Extent, -Extent);
+			CellPoints[3] = Position2D + FVector2D(Extent, Extent);
+			CellPoints[4] = Position2D;
+			for (const FVector2D& CellPoint : CellPoints)
+			{
+				const FVector2D SourceToCell(CellPoint - SourcePos);
+				const float Dot = FVector2D::DotProduct(SourceForward, SourceToCell.GetSafeNormal());
+				MaxDot = FMath::Max(MaxDot, Dot);
+			}
+			Angle = FMath::Abs(FMath::Acos(MaxDot) / UE_PI);
+		}
+		const float NormalizedAngle = FMath::Clamp(Angle, UE_PI / 180.f, 1.f);
+		AngleFactor = FMath::Pow(NormalizedAngle, AngleContribution);
+	}
+	return AngleFactor;
+}
+
+void UWorldPartitionRuntimeCellDataSpatialHash::AppendStreamingSourceInfo(const FWorldPartitionStreamingSource& Source, const FSphericalSector& SourceShape) const
+{
+	Super::AppendStreamingSourceInfo(Source, SourceShape);
+
+	const double SquareDistance = FVector::DistSquared2D(SourceShape.GetCenter(), Position);
+
+	// Only consider blocking sources
+	if (Source.bBlockOnSlowLoading)
+	{
+		bCachedIsBlockingSource = true;
+		CachedMinSquareDistanceToBlockingSource = FMath::Min(SquareDistance, CachedMinSquareDistanceToBlockingSource);
+	}
+	CachedSourceSquaredDistances.Add(SquareDistance);
+	CachedInstersectingShapes.Add(SourceShape);
+}
+
+void UWorldPartitionRuntimeCellDataSpatialHash::MergeStreamingSourceInfo() const
+{
+	Super::MergeStreamingSourceInfo();
+
+	int32 Count = CachedSourceSquaredDistances.Num();
+	check(Count == CachedSourcePriorityWeights.Num());
+	check(Count == CachedInstersectingShapes.Num());
+
+	CachedMinSquareDistanceToSource = Count ? FMath::Min(CachedSourceSquaredDistances) : MAX_dbl;
+
+	if (Count)
+	{
+		double TotalSourcePriorityWeight = 0.f;
+		for (int32 i = 0; i < Count; ++i)
+		{
+			TotalSourcePriorityWeight += CachedSourcePriorityWeights[i];
+		}
+
+		double HighestPrioWeight = 0;
+		double HighestPrioMinModulatedDistance = 0;
+		double WeightedModulatedDistance = 0;
+		for (int32 i = 0; i < Count; ++i)
+		{
+			const float AngleFactor = ComputeSourceToCellAngleFactor(CachedInstersectingShapes[i]);
+			const double CurrentModulatedDistance = CachedSourceSquaredDistances[i] * AngleFactor * AngleFactor;
+			const double CurrentWeight = CachedSourcePriorityWeights[i];
+			WeightedModulatedDistance += CurrentModulatedDistance * CurrentWeight / TotalSourcePriorityWeight;
+			// Find highest priority source with the minimum modulated distance
+			if ((i == 0) || ((CurrentModulatedDistance < HighestPrioMinModulatedDistance) && (CurrentWeight >= HighestPrioWeight)))
+			{
+				HighestPrioMinModulatedDistance = CurrentModulatedDistance;
+				HighestPrioWeight = CurrentWeight;
+			}
+		}
+
+		// Sorting distance is the minimum between these:
+		// - the highest priority source with the minimum modulated distance 
+		// - the weighted modulated distance
+		CachedSourceSortingDistance = FMath::Min(HighestPrioMinModulatedDistance, WeightedModulatedDistance);
+	}
+}
+
+int32 UWorldPartitionRuntimeCellDataSpatialHash::SortCompare(const UWorldPartitionRuntimeCellData* InOther, bool bCanUseSortingCache) const
+{
+	int32 Result = Super::SortCompare(InOther, bCanUseSortingCache);
+	if (Result == 0)
+	{
+		const UWorldPartitionRuntimeCellDataSpatialHash* Other = (UWorldPartitionRuntimeCellDataSpatialHash*)InOther;
+		
+		// Level (higher value is higher prio)
+		Result = Other->Level - Level;
+		if (bCanUseSortingCache && (Result == 0))
+		{
+			// Closest distance (lower value is higher prio)
+			const double Diff = CachedSourceSortingDistance - Other->CachedSourceSortingDistance;
+			if (FMath::IsNearlyZero(Diff))
+			{
+				const double RawDistanceDiff = CachedMinSquareDistanceToSource - Other->CachedMinSquareDistanceToSource;
+				Result = RawDistanceDiff < 0 ? -1 : (RawDistanceDiff > 0.f ? 1 : 0);
+			}
+			else
+			{
+				Result = Diff < 0.f ? -1 : (Diff > 0.f ? 1 : 0);
+			}
+		}
+	}
+	return Result;
+}
+
+FBox UWorldPartitionRuntimeCellDataSpatialHash::GetCellBounds() const
+{
+	FBox Box = FBox::BuildAABB(Position, FVector(Extent));
+	// Use content bounds for the Z extent
+	Box.Min.Z = GetContentBounds().Min.Z;
+	Box.Max.Z = GetContentBounds().Max.Z;
+	return Box;
+}
diff --git a/Engine/Source/Runtime/Engine/Private/WorldPartition/WorldPartitionRuntimeHash.cpp b/Engine/Source/Runtime/Engine/Private/WorldPartition/WorldPartitionRuntimeHash.cpp
--- a/Engine/Source/Runtime/Engine/Private/WorldPartition/WorldPartitionRuntimeHash.cpp
+++ b/Engine/Source/Runtime/Engine/Private/WorldPartition/WorldPartitionRuntimeHash.cpp
@@ -214,7 +214,7 @@ void UWorldPartitionRuntimeHash::FStreamingSourceCells::AddCell(const UWorldPart
 
 void FWorldPartitionQueryCache::AddCellInfo(const UWorldPartitionRuntimeCell* Cell, const FSphericalSector& SourceShape)
 {
-	const double SquareDistance = FVector::DistSquared2D(SourceShape.GetCenter(), Cell->GetCellBounds().GetCenter());
+	const double SquareDistance = FVector::DistSquared2D(SourceShape.GetCenter(), Cell->GetContentBounds().GetCenter());
 	if (double* ExistingSquareDistance = CellToSourceMinSqrDistances.Find(Cell))
 	{
 		*ExistingSquareDistance = FMath::Min(*ExistingSquareDistance, SquareDistance);
diff --git a/Engine/Source/Runtime/Engine/Private/WorldPartition/WorldPartitionRuntimeLevelStreamingCell.cpp b/Engine/Source/Runtime/Engine/Private/WorldPartition/WorldPartitionRuntimeLevelStreamingCell.cpp
--- a/Engine/Source/Runtime/Engine/Private/WorldPartition/WorldPartitionRuntimeLevelStreamingCell.cpp
+++ b/Engine/Source/Runtime/Engine/Private/WorldPartition/WorldPartitionRuntimeLevelStreamingCell.cpp
@@ -1,15 +1,14 @@
 // Copyright Epic Games, Inc. All Rights Reserved.
 
 #include "WorldPartition/WorldPartitionRuntimeLevelStreamingCell.h"
-#include "Engine/LevelStreaming.h"
 #include "WorldPartition/WorldPartitionLevelStreamingDynamic.h"
+#include "WorldPartition/WorldPartitionActorDescView.h"
+#include "WorldPartition/WorldPartitionLevelStreamingPolicy.h"
+#include "Engine/LevelStreaming.h"
 #include "Engine/Level.h"
 #include "Misc/HierarchicalLogArchive.h"
 #include "Misc/Paths.h"
 
-#include "WorldPartition/WorldPartitionActorDescView.h"
-#include "WorldPartition/WorldPartitionLevelStreamingPolicy.h"
-
 #include UE_INLINE_GENERATED_CPP_BY_NAME(WorldPartitionRuntimeLevelStreamingCell)
 
 #if WITH_EDITOR
@@ -183,7 +182,7 @@ FLinearColor UWorldPartitionRuntimeLevelStreamingCell::GetDebugColor(EWorldParti
 		{
 			// Return streaming status color
 			FLinearColor Color = LevelStreaming ? ULevelStreaming::GetLevelStreamingStatusColor(GetStreamingStatus()) : FLinearColor::Black;
-			Color.A = 0.25f / (Level + 1);
+			Color.A = 0.25f;
 			return Color;
 		}
 		default:
diff --git a/Engine/Source/Runtime/Engine/Private/WorldPartition/WorldPartitionRuntimeSpatialHash.cpp b/Engine/Source/Runtime/Engine/Private/WorldPartition/WorldPartitionRuntimeSpatialHash.cpp
--- a/Engine/Source/Runtime/Engine/Private/WorldPartition/WorldPartitionRuntimeSpatialHash.cpp
+++ b/Engine/Source/Runtime/Engine/Private/WorldPartition/WorldPartitionRuntimeSpatialHash.cpp
@@ -11,6 +11,13 @@
 #include "Misc/HierarchicalLogArchive.h"
 #include "WorldPartition/WorldPartitionStreamingPolicy.h"
 #include "UObject/UnrealType.h"
+#include "Engine/Level.h"
+#include "Engine/Canvas.h"
+#include "DrawDebugHelpers.h"
+#include "GlobalRenderResources.h"
+#include "UObject/ObjectSaveContext.h"
+#include "Components/LineBatchComponent.h"
+#include "ProfilingDebugging/ScopedTimers.h"
 #include "WorldPartition/WorldPartitionDebugHelper.h"
 #include "WorldPartition/HLOD/HLODLayer.h"
 #include "WorldPartition/WorldPartitionLevelStreamingDynamic.h"
@@ -19,17 +26,11 @@
 #include "WorldPartition/ContentBundle/ContentBundleWorldSubsystem.h"
 #include "WorldPartition/ContentBundle/ContentBundleDescriptor.h"
 #include "WorldPartition/ContentBundle/ContentBundleBase.h"
-#include "ProfilingDebugging/ScopedTimers.h"
-#include "Engine/Level.h"
-#include "Engine/Canvas.h"
-#include "DrawDebugHelpers.h"
-#include "GlobalRenderResources.h"
-#include "UObject/ObjectSaveContext.h"
-#include "Components/LineBatchComponent.h"
 #include "WorldPartition/WorldPartitionLog.h"
 #include "WorldPartition/WorldPartitionRuntimeHash.h"
 #include "WorldPartition/WorldPartitionRuntimeSpatialHashGridPreviewer.h"
 #include "WorldPartition/WorldPartitionStreamingGeneration.h"
+#include "WorldPartition/WorldPartitionRuntimeCellDataSpatialHash.h"
 
 #include UE_INLINE_GENERATED_CPP_BY_NAME(WorldPartitionRuntimeSpatialHash)
 
@@ -166,7 +167,8 @@ void FSpatialHashStreamingGrid::GetCells(const FWorldPartitionStreamingQuerySour
 			{
 				ForEachRuntimeCell(Coords, [&](const UWorldPartitionRuntimeCell* Cell)
 				{
-					if (!bEnableZCulling || TRange<double>(Cell->GetMinMaxZ().X, Cell->GetMinMaxZ().Y).Overlaps(TRange<double>(Shape.GetCenter().Z - Shape.GetRadius(), Shape.GetCenter().Z + Shape.GetRadius())))
+					const FVector2D CellMinMaxZ(Cell->GetContentBounds().Min.Z, Cell->GetContentBounds().Max.Z);
+					if (!bEnableZCulling || TRange<double>(CellMinMaxZ.X, CellMinMaxZ.Y).Overlaps(TRange<double>(Shape.GetCenter().Z - Shape.GetRadius(), Shape.GetCenter().Z + Shape.GetRadius())))
 					{
 						if (ShouldAddCell(Cell, QuerySource))
 						{
@@ -240,7 +242,8 @@ void FSpatialHashStreamingGrid::GetCells(const TArray<FWorldPartitionStreamingSo
 				{
 					ForEachRuntimeCell(Coords, [&](const UWorldPartitionRuntimeCell* Cell)
 					{
-						if (!bEnableZCulling || TRange<double>(Cell->GetMinMaxZ().X, Cell->GetMinMaxZ().Y).Overlaps(TRange<double>(Shape.GetCenter().Z - Shape.GetRadius(), Shape.GetCenter().Z + Shape.GetRadius())))
+						const FVector2D CellMinMaxZ(Cell->GetContentBounds().Min.Z, Cell->GetContentBounds().Max.Z);
+						if (!bEnableZCulling || TRange<double>(CellMinMaxZ.X, CellMinMaxZ.Y).Overlaps(TRange<double>(Shape.GetCenter().Z - Shape.GetRadius(), Shape.GetCenter().Z + Shape.GetRadius())))
 						{
 							if (!Cell->HasDataLayers() || (DataLayerSubsystem && DataLayerSubsystem->IsAnyDataLayerInEffectiveRuntimeState(Cell->GetDataLayers(), EDataLayerRuntimeState::Activated)))
 							{
@@ -365,7 +368,7 @@ void FSpatialHashStreamingGrid::InjectExternalStreamingObjectGrid(const FSpatial
 				if (ensure(!ExternalObjectLayerCell.GridCells.IsEmpty()))
 				{
 					// Any cell position should be good to find the proper cell coord
-					const FVector& CellPosition = ExternalObjectLayerCell.GridCells[0]->Position;
+					const FVector& CellPosition = CastChecked<UWorldPartitionRuntimeCellDataSpatialHash>(ExternalObjectLayerCell.GridCells[0]->RuntimeCellData)->Position;
 					FGridCellCoord2 CellCoords;
 					const FSquare2DGridHelper::FGridLevel& InjectedGridHelperLevel = GetGridHelper().Levels[TargetGridLevel];
 					if (!InjectedGridHelperLevel.GetCellCoords(FVector2D(CellPosition.X, CellPosition.Y), CellCoords))
@@ -390,7 +393,7 @@ void FSpatialHashStreamingGrid::InjectExternalStreamingObjectGrid(const FSpatial
 						TargetInjectedGridLevel.LayerCellsMapping.Add(CoordKey, LayerCellIndex);
 					}
 
-					for (UWorldPartitionRuntimeSpatialHashCell* StreamingCell : ExternalObjectLayerCell.GridCells)
+					for (UWorldPartitionRuntimeCell* StreamingCell : ExternalObjectLayerCell.GridCells)
 					{
 						TargetInjectedGridLevel.LayerCells[LayerCellIndex].GridCells.Add(StreamingCell);
 					}
@@ -423,7 +426,7 @@ void FSpatialHashStreamingGrid::RemoveExternalStreamingObjectGrid(const FSpatial
 
 			for (const FSpatialHashStreamingGridLayerCell& ExternalObjectLayerCell : ExternalObjectGridLevel.LayerCells)
 			{
-				for (UWorldPartitionRuntimeSpatialHashCell* StreamingCell : ExternalObjectLayerCell.GridCells)
+				for (UWorldPartitionRuntimeCell* StreamingCell : ExternalObjectLayerCell.GridCells)
 				{
 					bool bFound = false;
 					for (int LayerCellIndex = 0; LayerCellIndex < LayerCells.Num(); ++ LayerCellIndex)
@@ -501,7 +504,7 @@ void FSpatialHashStreamingGrid::ForEachRuntimeCell(TFunctionRef<bool(const UWorl
 		{
 			for (const FSpatialHashStreamingGridLayerCell& LayerCell : GridLevel.LayerCells)
 			{
-				for (const UWorldPartitionRuntimeSpatialHashCell* Cell : LayerCell.GridCells)
+				for (const UWorldPartitionRuntimeCell* Cell : LayerCell.GridCells)
 				{
 					if (!Func(Cell))
 					{
@@ -621,7 +624,7 @@ void FSpatialHashStreamingGrid::Draw3D(UWorld* World, const TArray<FWorldPartiti
 					}
 
 					// Use cell MinMaxZ to compute effective cell bounds
-					FBox CellBox = Cell->GetCellBounds();
+					FBox CellBox = CastChecked<UWorldPartitionRuntimeCellDataSpatialHash>(Cell->RuntimeCellData)->GetCellBounds();
 
 					// Draw Cell using its debug color
 					FColor CellColor = Cell->GetDebugColor(VisualizeMode).ToFColor(false).WithAlpha(16);
@@ -1048,7 +1051,7 @@ URuntimeHashExternalStreamingObjectBase* UWorldPartitionRuntimeSpatialHash::Stor
 
 			for (FSpatialHashStreamingGridLayerCell& GridLayerCell : GridLevel.LayerCells)
 			{
-				for (UWorldPartitionRuntimeSpatialHashCell* Cell : GridLayerCell.GridCells)
+				for (UWorldPartitionRuntimeCell* Cell : GridLayerCell.GridCells)
 				{
 					// Do not dirty, otherwise it dirties the level previous outer (WorldPartition) which dirties the map. Occurs when entering PIE. 
 					// Do not reset loaders, the cell was just created it did not exists when loading initially.
@@ -1212,7 +1215,7 @@ void UWorldPartitionRuntimeSpatialHash::DumpStateLog(FHierarchicalLogArchive& Ar
 				for (FSpatialHashStreamingGridLayerCell& LayerCell : GridLevel.LayerCells)
 				{
 					LevelCellCount += LayerCell.GridCells.Num();
-					for (TObjectPtr<UWorldPartitionRuntimeSpatialHashCell>& Cell : LayerCell.GridCells)
+					for (TObjectPtr<UWorldPartitionRuntimeCell>& Cell : LayerCell.GridCells)
 					{
 						LevelActorCount += Cell->GetActorCount();
 					}
@@ -1241,7 +1244,7 @@ void UWorldPartitionRuntimeSpatialHash::DumpStateLog(FHierarchicalLogArchive& Ar
 
 				for (const FSpatialHashStreamingGridLayerCell& LayerCell : GridLevel.LayerCells)
 				{
-					for (const TObjectPtr<UWorldPartitionRuntimeSpatialHashCell>& Cell : LayerCell.GridCells)
+					for (const TObjectPtr<UWorldPartitionRuntimeCell>& Cell : LayerCell.GridCells)
 					{
 						FHierarchicalLogArchive::FIndentScope CellIndentScope = Ar.PrintfIndent(TEXT("Content of Cell %s"), *Cell->GetDebugName());
 						Cell->DumpStateLog(Ar);
@@ -1381,7 +1384,7 @@ bool UWorldPartitionRuntimeSpatialHash::CreateStreamingGrid(const FSpatialHashRu
 				const FString CellName = GetCellNameString(OuterWorld->GetPackage()->GetFName(), CurrentStreamingGrid.GridName, CellGlobalCoords, GridCellDataChunk.GetDataLayersID(), GridCellDataChunk.GetContentBundleID());
 				const FGuid CellGuid = GetCellGuid(CurrentStreamingGrid.GridName, CellGlobalCoords, GridCellDataChunk.GetDataLayersID(), GridCellDataChunk.GetContentBundleID());
 
-				UWorldPartitionRuntimeSpatialHashCell* StreamingCell = NewObject<UWorldPartitionRuntimeSpatialHashCell>(this, StreamingPolicy->GetRuntimeCellClass(), FName(CellName));
+				UWorldPartitionRuntimeCell* StreamingCell = NewObject<UWorldPartitionRuntimeCell>(this, StreamingPolicy->GetRuntimeCellClass(), FName(CellName));
 
 				int32 LayerCellIndex;
 				int32* LayerCellIndexPtr = GridLevel.LayerCellsMapping.Find(CellIndex);
@@ -1399,20 +1402,23 @@ bool UWorldPartitionRuntimeSpatialHash::CreateStreamingGrid(const FSpatialHashRu
 				StreamingCell->SetIsAlwaysLoaded(bIsCellAlwaysLoaded);
 				StreamingCell->SetDataLayers(GridCellDataChunk.GetDataLayers());
 				StreamingCell->SetContentBundleUID(GridCellDataChunk.GetContentBundleID());
-				StreamingCell->Level = Level;
 				StreamingCell->SetPriority(RuntimeGrid.Priority);
-				FBox2D Bounds;
-				verify(TempLevel.GetCellBounds(FGridCellCoord2(CellCoordX, CellCoordY), Bounds));
-				StreamingCell->Position = FVector(Bounds.GetCenter(), 0.f);
-				const double CellExtent = Bounds.GetExtent().X;
-				check(CellExtent < MAX_flt);
-				StreamingCell->Extent = CellExtent;
 				StreamingCell->SetDebugInfo(CellGlobalCoords.X, CellGlobalCoords.Y, CellGlobalCoords.Z, CurrentStreamingGrid.GridName);
 				StreamingCell->SetClientOnlyVisible(CurrentStreamingGrid.bClientOnlyVisible);
 				StreamingCell->SetBlockOnSlowLoading(CurrentStreamingGrid.bBlockOnSlowStreaming);
 				StreamingCell->SetIsHLOD(RuntimeGrid.HLODLayer ? true : false);
 				StreamingCell->SetGuid(CellGuid);
 
+				UWorldPartitionRuntimeCellDataSpatialHash* StreamingSourceInfoSpatialHash = NewObject<UWorldPartitionRuntimeCellDataSpatialHash>(StreamingCell);
+				StreamingSourceInfoSpatialHash->Level = Level;
+				FBox2D Bounds;
+				verify(TempLevel.GetCellBounds(FGridCellCoord2(CellCoordX, CellCoordY), Bounds));
+				StreamingSourceInfoSpatialHash->Position = FVector(Bounds.GetCenter(), 0.f);
+				const double CellExtent = Bounds.GetExtent().X;
+				check(CellExtent < MAX_flt);
+				StreamingSourceInfoSpatialHash->Extent = CellExtent;
+				StreamingCell->RuntimeCellData = StreamingSourceInfoSpatialHash;
+
 				check(!StreamingCell->UnsavedActorsContainer);
 				for (const IStreamingGenerationContext::FActorInstance& ActorInstance : FilteredActors)
 				{
@@ -1453,7 +1459,7 @@ bool UWorldPartitionRuntimeSpatialHash::CreateStreamingGrid(const FSpatialHashRu
 						}
 					}
 				}
-				StreamingCell->SetContentBounds(CellContentBounds);
+				StreamingSourceInfoSpatialHash->ContentBounds = CellContentBounds;
 
 				// Always loaded cell actors are transfered to World's Persistent Level (see UWorldPartitionRuntimeSpatialHash::PopulateGeneratorPackageForCook)
 				if (StreamingCell->GetActorCount() && !StreamingCell->IsAlwaysLoaded())
@@ -1818,7 +1824,7 @@ EWorldPartitionStreamingPerformance UWorldPartitionRuntimeSpatialHash::GetStream
 	check(Cell->GetBlockOnSlowLoading());
 	const double BlockOnSlowStreamingWarningRatio = GBlockOnSlowStreamingRatio * GBlockOnSlowStreamingWarningFactor;
 	
-	const UWorldPartitionRuntimeSpatialHashCell* StreamingCell = Cast<const UWorldPartitionRuntimeSpatialHashCell>(Cell);
+	const UWorldPartitionRuntimeCell* StreamingCell = Cast<const UWorldPartitionRuntimeCell>(Cell);
 	check(StreamingCell);
 
 	const FSpatialHashStreamingGrid* StreamingGrid = GetStreamingGridByName(StreamingCell->GetGridName());
@@ -1826,9 +1832,11 @@ EWorldPartitionStreamingPerformance UWorldPartitionRuntimeSpatialHash::GetStream
 	{
 		const float LoadingRange = StreamingGrid->LoadingRange;
 
-		if (StreamingCell->IsBlockingSource())
+		UWorldPartitionRuntimeCellDataSpatialHash* StreamingSourceInfoSpatialHash = CastChecked<UWorldPartitionRuntimeCellDataSpatialHash>(Cell->RuntimeCellData);
+
+		if (StreamingSourceInfoSpatialHash->IsBlockingSource())
 		{
-			const double Distance = FMath::Sqrt(StreamingCell->GetMinSquareDistanceToBlockingSource()) - ((double)StreamingGrid->GetCellSize(StreamingCell->Level) / 2);
+			const double Distance = FMath::Sqrt(StreamingSourceInfoSpatialHash->GetMinSquareDistanceToBlockingSource()) - ((double)StreamingGrid->GetCellSize(StreamingSourceInfoSpatialHash->Level) / 2);
 
 			const double Ratio = Distance / LoadingRange;
 
diff --git a/Engine/Source/Runtime/Engine/Private/WorldPartition/WorldPartitionRuntimeSpatialHashCell.cpp b/Engine/Source/Runtime/Engine/Private/WorldPartition/WorldPartitionRuntimeSpatialHashCell.cpp
deleted file mode 100644
diff --git a/Engine/Source/Runtime/Engine/Private/WorldPartition/WorldPartitionStreamingPolicy.cpp b/Engine/Source/Runtime/Engine/Private/WorldPartition/WorldPartitionStreamingPolicy.cpp
--- a/Engine/Source/Runtime/Engine/Private/WorldPartition/WorldPartitionStreamingPolicy.cpp
+++ b/Engine/Source/Runtime/Engine/Private/WorldPartition/WorldPartitionStreamingPolicy.cpp
@@ -413,7 +413,7 @@ void UWorldPartitionStreamingPolicy::UpdateStreamingState()
 			{
 				TRACE_CPUPROFILER_EVENT_SCOPE(UWorldPartitionStreamingPolicy::UpdateStreamingState_ForEachStreamingCellsSources);
 
-				UWorldPartitionRuntimeCell::DirtyStreamingSourceCacheEpoch();
+				UWorldPartitionRuntimeCellData::DirtyStreamingSourceCacheEpoch();
 
 				WorldPartition->RuntimeHash->ForEachStreamingCellsSources(StreamingSources, [this](const UWorldPartitionRuntimeCell* Cell, EStreamingSourceTargetState TargetState)
 				{
@@ -689,7 +689,7 @@ void UWorldPartitionStreamingPolicy::UpdateDebugCellsStreamingPriority(const TSe
 
 		for (const UWorldPartitionRuntimeCell* Cell : Cells)
 		{
-			Cell->MergeStreamingSourceInfo();
+			Cell->RuntimeCellData->MergeStreamingSourceInfo();
 		}
 
 		SortStreamingCellsByImportance(Cells);
diff --git a/Engine/Source/Runtime/Engine/Public/WorldPartition/WorldPartitionRuntimeCell.h b/Engine/Source/Runtime/Engine/Public/WorldPartition/WorldPartitionRuntimeCell.h
--- a/Engine/Source/Runtime/Engine/Public/WorldPartition/WorldPartitionRuntimeCell.h
+++ b/Engine/Source/Runtime/Engine/Public/WorldPartition/WorldPartitionRuntimeCell.h
@@ -7,6 +7,7 @@
 #include "WorldPartition/WorldPartitionActorDescView.h"
 #include "WorldPartition/WorldPartitionStreamingSource.h"
 #include "WorldPartition/WorldPartitionActorContainerID.h"
+#include "WorldPartition/WorldPartitionRuntimeCellData.h"
 #include "WorldPartition/WorldPartitionRuntimeCellInterface.h"
 #include "WorldPartition/WorldPartitionRuntimeCellOwner.h"
 #include "ProfilingDebugging/ProfilingHelpers.h"
@@ -152,6 +153,10 @@ class ENGINE_API UWorldPartitionRuntimeCell : public UObject, public IWorldParti
 {
 	GENERATED_UCLASS_BODY()
 
+#if WITH_EDITOR
+	virtual void PostDuplicate(bool bDuplicateForPIE) override;
+#endif
+
 	virtual void Load() const PURE_VIRTUAL(UWorldPartitionRuntimeCell::Load,);
 	virtual void Unload() const PURE_VIRTUAL(UWorldPartitionRuntimeCell::Unload,);
 	virtual bool CanUnload() const PURE_VIRTUAL(UWorldPartitionRuntimeCell::CanUnload, return true;);
@@ -170,7 +175,6 @@ class ENGINE_API UWorldPartitionRuntimeCell : public UObject, public IWorldParti
 	virtual bool IsLoading() const { return false; }
 	virtual const FString& GetDebugName() const { return DebugInfo.Name; }
 	virtual bool IsDebugShown() const;
-	virtual int32 SortCompare(const UWorldPartitionRuntimeCell* Other, bool bCanUseSortingCache = true) const;
 	virtual FName GetGridName() const { return DebugInfo.GridName; }
 	bool GetClientOnlyVisible() const { return bClientOnlyVisible; }
 	virtual FGuid const& GetContentBundleID() const { return ContentBundleID; }
@@ -180,11 +184,10 @@ class ENGINE_API UWorldPartitionRuntimeCell : public UObject, public IWorldParti
 
 	/** Caches information on streaming source that will be used later on to sort cell. */
 	bool ShouldResetStreamingSourceInfo() const;
-	virtual void ResetStreamingSourceInfo() const;
-	virtual void AppendStreamingSourceInfo(const FWorldPartitionStreamingSource& Source, const FSphericalSector& SourceShape) const;
-	virtual void MergeStreamingSourceInfo() const;
-
-	static void DirtyStreamingSourceCacheEpoch() { ++UWorldPartitionRuntimeCell::StreamingSourceCacheEpoch; }
+	void ResetStreamingSourceInfo() const;
+	void AppendStreamingSourceInfo(const FWorldPartitionStreamingSource& Source, const FSphericalSector& SourceShape) const;
+	void MergeStreamingSourceInfo() const;
+	int32 SortCompare(const UWorldPartitionRuntimeCell* Other, bool bCanUseSortingCache = true) const;
 
 	//~Begin IWorldPartitionCell Interface
 	virtual TArray<const UDataLayerInstance*> GetDataLayerInstances() const override;
@@ -196,13 +199,11 @@ class ENGINE_API UWorldPartitionRuntimeCell : public UObject, public IWorldParti
 	{
 		return Algo::AnyOf(DataLayers, [&InDataLayers](const FName& DataLayer) { return InDataLayers.Contains(DataLayer); });
 	}
-	virtual const FBox& GetContentBounds() const override { return ContentBounds; }
-	virtual FBox GetCellBounds() const override { return FBox(ForceInit); }
+	virtual const FBox& GetContentBounds() const override;
+	virtual FBox GetCellBounds() const override;
 	virtual FName GetLevelPackageName() const override;
 	//~End IWorldPartitionCell Interface
 
-	const FVector2D GetMinMaxZ() const { return FVector2D(ContentBounds.Min.Z, ContentBounds.Max.Z); }
-
 	bool GetBlockOnSlowLoading() const { return bBlockOnSlowLoading; }
 #if WITH_EDITOR
 	bool NeedsActorToCellRemapping() const;
@@ -210,7 +211,6 @@ class ENGINE_API UWorldPartitionRuntimeCell : public UObject, public IWorldParti
 	void SetClientOnlyVisible(bool bInClientOnlyVisible) { bClientOnlyVisible = bInClientOnlyVisible; }
 	void SetDataLayers(const TArray<const UDataLayerInstance*>& InDataLayerInstances);
 	void SetContentBundleUID(const FGuid& InContentBundleID) { ContentBundleID = InContentBundleID; }
-	void SetContentBounds(const FBox& InBounds) { ContentBounds = InBounds; }
 	void SetDebugInfo(int64 InCoordX, int64 InCoordY, int64 InCoordZ, FName InGridName);
 	void SetLevelPackageName(const FName& InLevelPackageName) { LevelPackageName = InLevelPackageName; }
 	
@@ -241,6 +241,11 @@ class ENGINE_API UWorldPartitionRuntimeCell : public UObject, public IWorldParti
 	void SetDebugStreamingPriority(float InDebugStreamingPriority) { DebugStreamingPriority = InDebugStreamingPriority; }
 #endif
 
+#if WITH_EDITORONLY_DATA
+	UPROPERTY()
+	TObjectPtr<UActorContainer> UnsavedActorsContainer;
+#endif
+
 protected:
 	FLinearColor GetDebugStreamingPriorityColor() const;
 
@@ -255,9 +260,6 @@ private:
 	UPROPERTY()
 	TArray<FName> DataLayers;
 
-	UPROPERTY()
-	FBox ContentBounds;
-
 	// Debug Info
 	UPROPERTY()
 	FWorldPartitionRuntimeCellDebugInfo DebugInfo;
@@ -282,17 +284,6 @@ protected:
 	UPROPERTY()
 	FGuid CellGuid;
 
-	// Source Priority
-	mutable uint8 CachedMinSourcePriority;
-
-	// Source Priorities
-	mutable TArray<float> CachedSourcePriorityWeights;
-
-	// Epoch used to dirty cache
-	mutable int32 CachedSourceInfoEpoch;
-
-	static int32 StreamingSourceCacheEpoch;
-
 #if !UE_BUILD_SHIPPING
 	// Represents the streaming priority relative to other cells
 	float DebugStreamingPriority;
@@ -301,4 +292,8 @@ protected:
 #if WITH_EDITOR
 	FName LevelPackageName;
 #endif
+
+public:
+	UPROPERTY()
+	TObjectPtr<UWorldPartitionRuntimeCellData> RuntimeCellData;
 };
diff --git a/Engine/Source/Runtime/Engine/Public/WorldPartition/WorldPartitionRuntimeCellData.h b/Engine/Source/Runtime/Engine/Public/WorldPartition/WorldPartitionRuntimeCellData.h
new file mode 100644
--- /dev/null
+++ b/Engine/Source/Runtime/Engine/Public/WorldPartition/WorldPartitionRuntimeCellData.h
@@ -0,0 +1,41 @@
+// Copyright Epic Games, Inc. All Rights Reserved.
+#pragma once
+
+#include "CoreMinimal.h"
+#include "WorldPartition/WorldPartitionStreamingSource.h"
+#include "WorldPartitionRuntimeCellData.generated.h"
+
+class UActorContainer;
+
+/** Caches information on streaming source that will be used later on to sort cell. */
+UCLASS()
+class ENGINE_API UWorldPartitionRuntimeCellData : public UObject
+{
+	GENERATED_UCLASS_BODY()
+
+	virtual ~UWorldPartitionRuntimeCellData() {}
+
+	bool ShouldResetStreamingSourceInfo() const;
+	virtual void ResetStreamingSourceInfo() const;
+	virtual void AppendStreamingSourceInfo(const FWorldPartitionStreamingSource& Source, const FSphericalSector& SourceShape) const;
+	virtual void MergeStreamingSourceInfo() const;
+	virtual int32 SortCompare(const UWorldPartitionRuntimeCellData* InOther, bool bCanUseSortingCache = true) const;
+
+	virtual const FBox& GetContentBounds() const;
+	virtual FBox GetCellBounds() const;
+
+	static int32 StreamingSourceCacheEpoch;
+	static void DirtyStreamingSourceCacheEpoch() { ++StreamingSourceCacheEpoch; }
+
+	// Source Priority
+	mutable uint8 CachedMinSourcePriority;
+
+	// Source Priorities
+	mutable TArray<float> CachedSourcePriorityWeights;
+
+	// Epoch used to dirty cache
+	mutable int32 CachedSourceInfoEpoch;
+
+	UPROPERTY()
+	FBox ContentBounds;
+};
\ No newline at end of file
diff --git a/Engine/Source/Runtime/Engine/Public/WorldPartition/WorldPartitionRuntimeCellDataSpatialHash.h b/Engine/Source/Runtime/Engine/Public/WorldPartition/WorldPartitionRuntimeCellDataSpatialHash.h
new file mode 100644
--- /dev/null
+++ b/Engine/Source/Runtime/Engine/Public/WorldPartition/WorldPartitionRuntimeCellDataSpatialHash.h
@@ -0,0 +1,61 @@
+// Copyright Epic Games, Inc. All Rights Reserved.
+#pragma once
+
+#include "CoreMinimal.h"
+#include "WorldPartition/WorldPartitionRuntimeCellData.h"
+#include "WorldPartitionRuntimeCellDataSpatialHash.generated.h"
+
+class UActorContainer;
+class UWorldPartitionRuntimeCell;
+
+UCLASS(Within = WorldPartitionRuntimeCell)
+class ENGINE_API UWorldPartitionRuntimeCellDataSpatialHash : public UWorldPartitionRuntimeCellData
+{
+	GENERATED_UCLASS_BODY()
+
+	//~Begin UWorldPartitionRuntimeCellData
+	virtual void ResetStreamingSourceInfo() const override;
+	virtual void AppendStreamingSourceInfo(const FWorldPartitionStreamingSource& Source, const FSphericalSector& SourceShape) const override;
+	virtual void MergeStreamingSourceInfo() const override;
+	virtual int32 SortCompare(const UWorldPartitionRuntimeCellData* InOther, bool bCanUseSortingCache = true) const override;
+	virtual FBox GetCellBounds() const override;
+	//~End UWorldPartitionRuntimeCellData
+
+	bool IsBlockingSource() const { return bCachedIsBlockingSource; }
+	double GetMinSquareDistanceToBlockingSource() const { return CachedMinSquareDistanceToBlockingSource; }
+
+	UPROPERTY()
+	FVector Position;
+
+	UPROPERTY()
+	float Extent;
+
+	UPROPERTY()
+	int32 Level;
+
+private:
+	float ComputeSourceToCellAngleFactor(const FSphericalSector& SourceShape) const;
+
+	// Used to determine if cell was requested by blocking source
+	mutable bool bCachedIsBlockingSource;
+
+	// Square distance from the cell to the closest blocking streaming source
+	mutable double CachedMinSquareDistanceToBlockingSource;
+
+	// Square distance from the cell to the closest streaming source
+	mutable double CachedMinSquareDistanceToSource;
+
+	// Modulated distance to the different streaming sources used to sort relative priority amongst streaming cells
+	// The value is affected by :
+	// - All sources intersecting the cell
+	// - The priority of each source
+	// - The distance between the cell and each source
+	// - The angle between the cell and each source orientation
+	mutable double CachedSourceSortingDistance;
+
+	// Square distance from the cell to  the intersecting streaming sources
+	mutable TArray<double> CachedSourceSquaredDistances;
+
+	// Intersecting streaming source shapes
+	mutable TArray<FSphericalSector> CachedInstersectingShapes;
+};
\ No newline at end of file
diff --git a/Engine/Source/Runtime/Engine/Public/WorldPartition/WorldPartitionRuntimeLevelStreamingCell.h b/Engine/Source/Runtime/Engine/Public/WorldPartition/WorldPartitionRuntimeLevelStreamingCell.h
--- a/Engine/Source/Runtime/Engine/Public/WorldPartition/WorldPartitionRuntimeLevelStreamingCell.h
+++ b/Engine/Source/Runtime/Engine/Public/WorldPartition/WorldPartitionRuntimeLevelStreamingCell.h
@@ -1,7 +1,8 @@
 // Copyright Epic Games, Inc. All Rights Reserved.
 #pragma once
 
-#include "WorldPartition/WorldPartitionRuntimeSpatialHashCell.h"
+#include "WorldPartition/WorldPartitionRuntimeCell.h"
+#include "WorldPartition/WorldPartitionRuntimeCellData.h"
 #include "ProfilingDebugging/ProfilingHelpers.h"
 #if WITH_EDITOR
 #include "WorldPartition/WorldPartitionLevelHelper.h"
@@ -9,7 +10,7 @@
 #include "WorldPartitionRuntimeLevelStreamingCell.generated.h"
 
 UCLASS()
-class UWorldPartitionRuntimeLevelStreamingCell : public UWorldPartitionRuntimeSpatialHashCell
+class ENGINE_API UWorldPartitionRuntimeLevelStreamingCell : public UWorldPartitionRuntimeCell
 {
 	GENERATED_UCLASS_BODY()
 
diff --git a/Engine/Source/Runtime/Engine/Public/WorldPartition/WorldPartitionRuntimeSpatialHash.h b/Engine/Source/Runtime/Engine/Public/WorldPartition/WorldPartitionRuntimeSpatialHash.h
--- a/Engine/Source/Runtime/Engine/Public/WorldPartition/WorldPartitionRuntimeSpatialHash.h
+++ b/Engine/Source/Runtime/Engine/Public/WorldPartition/WorldPartitionRuntimeSpatialHash.h
@@ -11,7 +11,6 @@
 #include "WorldPartition/DataLayer/DataLayersID.h"
 #include "WorldPartition/HLOD/HLODLayer.h"
 #include "WorldPartitionRuntimeHash.h"
-#include "WorldPartitionRuntimeSpatialHashCell.h"
 #include "WorldPartitionRuntimeSpatialHashGridPreviewer.h"
 #include "WorldPartitionRuntimeSpatialHash.generated.h"
 
@@ -34,7 +33,7 @@ struct FSpatialHashStreamingGridLayerCell
 	FSpatialHashStreamingGridLayerCell& operator=(FSpatialHashStreamingGridLayerCell&& Other) = default;
 
 	UPROPERTY()
-	TArray<TObjectPtr<UWorldPartitionRuntimeSpatialHashCell>> GridCells;
+	TArray<TObjectPtr<UWorldPartitionRuntimeCell>> GridCells;
 };
 
 // Disable copy of this USTRUCT by the autogenerated code.
diff --git a/Engine/Source/Runtime/Engine/Public/WorldPartition/WorldPartitionRuntimeSpatialHashCell.h b/Engine/Source/Runtime/Engine/Public/WorldPartition/WorldPartitionRuntimeSpatialHashCell.h
deleted file mode 100644
