Fix crash in shipping configuration editor when changing to modeling mode, using WITH_EDITORONLY_DATA this time

diff --git a/Engine/Source/Runtime/Engine/Private/UnrealEngine.cpp b/Engine/Source/Runtime/Engine/Private/UnrealEngine.cpp
--- a/Engine/Source/Runtime/Engine/Private/UnrealEngine.cpp
+++ b/Engine/Source/Runtime/Engine/Private/UnrealEngine.cpp
@@ -3018,7 +3018,7 @@ void UEngine::InitializeObjectReferences()
 	LoadSpecialMaterial(TEXT("InvalidLightmapSettingsMaterialName"), InvalidLightmapSettingsMaterialName.ToString(), InvalidLightmapSettingsMaterial, false);
 	LoadSpecialMaterial(TEXT("ArrowMaterialName"), ArrowMaterialName.ToString(), ArrowMaterial, false);
 
-#if !UE_BUILD_SHIPPING || UE_BUILD_SHIPPING_WITH_EDITOR
+#if !UE_BUILD_SHIPPING || WITH_EDITORONLY_DATA
 	ArrowMaterialYellow = UMaterialInstanceDynamic::Create(ArrowMaterial, nullptr);
 	ArrowMaterialYellow->SetVectorParameterValue("GizmoColor", FLinearColor::Yellow);
 
