Remove inline property initalization in favor of the constror in cpp

#rnx
#pf skip

diff --git a/Engine/Source/Runtime/Engine/Classes/GameFramework/InputDeviceSubsystem.h b/Engine/Source/Runtime/Engine/Classes/GameFramework/InputDeviceSubsystem.h
--- a/Engine/Source/Runtime/Engine/Classes/GameFramework/InputDeviceSubsystem.h
+++ b/Engine/Source/Runtime/Engine/Classes/GameFramework/InputDeviceSubsystem.h
@@ -68,19 +68,19 @@ struct ENGINE_API FActiveDeviceProperty
 
 	/** The active device property */
 	UPROPERTY()
-	TWeakObjectPtr<UInputDeviceProperty> Property = nullptr;
+	TWeakObjectPtr<UInputDeviceProperty> Property;
 
 	/** How long this property has been evaluated for. DeltaTime is added to this on tick */
-	double EvaluatedDuration = 0.0;
+	double EvaluatedDuration;
 
 	/** The platform user that is actively receiving this device property */
-	FPlatformUserId PlatformUser = PLATFORMUSERID_NONE;
+	FPlatformUserId PlatformUser;
 
 	/** The input device id that should receive this property. */
-	FInputDeviceId DeviceId = INPUTDEVICEID_NONE;
+	FInputDeviceId DeviceId;
 
 	/** The handle of this active property. */
-	FInputDevicePropertyHandle PropertyHandle = FInputDevicePropertyHandle::InvalidHandle;
+	FInputDevicePropertyHandle PropertyHandle;
 
 	/**
 	* If true, then the input device property will not be removed after it's evaluation time has completed.
