add nullptr check to LocalContent before accessing it

#rnx

diff --git a/Engine/Source/Runtime/Engine/Private/Materials/MaterialShader.cpp b/Engine/Source/Runtime/Engine/Private/Materials/MaterialShader.cpp
--- a/Engine/Source/Runtime/Engine/Private/Materials/MaterialShader.cpp
+++ b/Engine/Source/Runtime/Engine/Private/Materials/MaterialShader.cpp
@@ -2652,7 +2652,7 @@ FPSOPrecacheRequestResultArray FMaterialShaderMap::CollectPSOs(const FMaterialPS
 
 	// Has data for this VF type
 	const FMaterialShaderMapContent* LocalContent = GetContent();
-	if (!LocalContent->GetMeshShaderMap(PrecacheParams.VertexFactoryData.VertexFactoryType->GetHashedName()))
+	if (LocalContent == nullptr || !LocalContent->GetMeshShaderMap(PrecacheParams.VertexFactoryData.VertexFactoryType->GetHashedName()))
 	{
 		return FPSOPrecacheRequestResultArray();
 	}
