PR #6525: Force Feedback mapping clarity and consistency improvements (Contributed by [REDACTED_PERSON])

Update comments on the force feedback effect channel details to explictly say that it may depend on hardware configurations

#pf skip

diff --git a/Engine/Source/Runtime/Engine/Classes/GameFramework/ForceFeedbackEffect.h b/Engine/Source/Runtime/Engine/Classes/GameFramework/ForceFeedbackEffect.h
--- a/Engine/Source/Runtime/Engine/Classes/GameFramework/ForceFeedbackEffect.h
+++ b/Engine/Source/Runtime/Engine/Classes/GameFramework/ForceFeedbackEffect.h
@@ -18,15 +18,19 @@ struct FForceFeedbackChannelDetails
 {
 	GENERATED_USTRUCT_BODY()
 
+	/** Please note the final channel mapping depends on the software and hardware capabilities of the platform used to run the engine or the game. Refer to documentation for more information. */
 	UPROPERTY(EditAnywhere, Category="ChannelDetails")
 	uint32 bAffectsLeftLarge:1;
 
+	/** Please note the final channel mapping depends on the software and hardware capabilities of the platform used to run the engine or the game. Refer to documentation for more information. */
 	UPROPERTY(EditAnywhere, Category="ChannelDetails")
 	uint32 bAffectsLeftSmall:1;
 
+	/** Please note the final channel mapping depends on the software and hardware capabilities of the platform used to run the engine or the game. Refer to documentation for more information. */
 	UPROPERTY(EditAnywhere, Category="ChannelDetails")
 	uint32 bAffectsRightLarge:1;
 
+	/** Please note the final channel mapping depends on the software and hardware capabilities of the platform used to run the engine or the game. Refer to documentation for more information. */
 	UPROPERTY(EditAnywhere, Category="ChannelDetails")
 	uint32 bAffectsRightSmall:1;
 
